#ttp 294228 - Regression: EDITOR: VIMBLUEPRINT: Animation stops when you drag a slider to its limits

#proj UE4

#summary Added code to prevent viewport throttling and changed MetaData setting so that it updates input controls without the need for a compile

Originally by [REDACTED_PERSON]

diff --git a/Engine/Source/Editor/Kismet/Private/BlueprintDetailsCustomization.cpp b/Engine/Source/Editor/Kismet/Private/BlueprintDetailsCustomization.cpp
--- a/Engine/Source/Editor/Kismet/Private/BlueprintDetailsCustomization.cpp
+++ b/Engine/Source/Editor/Kismet/Private/BlueprintDetailsCustomization.cpp
@@ -270,8 +270,10 @@ void FBlueprintVarActionDetails::CustomizeDetails( IDetailLayoutBuilder& DetailL
 
 	TSharedPtr<SToolTip> SliderRangeTooltip = IDocumentation::Get()->CreateToolTip(LOCTEXT("SliderRange_Tooltip", "Allows setting the minimum and maximum values for the UI slider for this variable."), NULL, DocLink, TEXT("SliderRange"));
 
+	FName UIMin = TEXT("UIMin");
+	FName UIMax = TEXT("UIMax");
 	Category.AddCustomRow( TEXT("Slider Range") )
-	.Visibility(TAttribute<EVisibility>(this, &FBlueprintVarActionDetails::SliderVisibility))
+	.Visibility(TAttribute<EVisibility>(this, &FBlueprintVarActionDetails::RangeVisibility))
 	.NameContent()
 	[
 		SNew(STextBlock)
@@ -287,8 +289,8 @@ void FBlueprintVarActionDetails::CustomizeDetails( IDetailLayoutBuilder& DetailL
 		.FillWidth(1)
 		[
 			SNew(SEditableTextBox)
-			.Text( this, &FBlueprintVarActionDetails::OnGetSliderMinValue )
-			.OnTextCommitted( this, &FBlueprintVarActionDetails::OnSliderMinValueChanged )
+			.Text(this, &FBlueprintVarActionDetails::OnGetMetaKeyValue, UIMin)
+			.OnTextCommitted(this, &FBlueprintVarActionDetails::OnMetaKeyValueChanged, UIMin)
 			.Font( DetailFontInfo )
 		]
 		+SHorizontalBox::Slot()
@@ -302,12 +304,53 @@ void FBlueprintVarActionDetails::CustomizeDetails( IDetailLayoutBuilder& DetailL
 		.FillWidth(1)
 		[
 			SNew(SEditableTextBox)
-			.Text( this, &FBlueprintVarActionDetails::OnGetSliderMaxValue )
-			.OnTextCommitted(this, &FBlueprintVarActionDetails::OnSliderMaxValueChanged)
+			.Text(this, &FBlueprintVarActionDetails::OnGetMetaKeyValue, UIMax)
+			.OnTextCommitted(this, &FBlueprintVarActionDetails::OnMetaKeyValueChanged, UIMax)
 			.Font( DetailFontInfo )
 		]
 	];
 
+	TSharedPtr<SToolTip> ValueRangeTooltip = IDocumentation::Get()->CreateToolTip(LOCTEXT("ValueRangeLabel_Tooltip", "The range of values allowed by this variable. Values outside of this will be clamped to the range."), NULL, DocLink, TEXT("ValueRange"));
+
+	FName ClampMin = TEXT("ClampMin");
+	FName ClampMax = TEXT("ClampMax");
+	Category.AddCustomRow(TEXT("Value Range"))
+	.Visibility(TAttribute<EVisibility>(this, &FBlueprintVarActionDetails::RangeVisibility))
+	.NameContent()
+	[
+		SNew(STextBlock)
+		.Text(LOCTEXT("ValueRangeLabel", "Value Range"))
+		.ToolTipText(LOCTEXT("ValueRangeLabel_Tooltip", "The range of values allowed by this variable. Values outside of this will be clamped to the range."))
+		.ToolTip(ValueRangeTooltip)
+		.Font(DetailFontInfo)
+	]
+	.ValueContent()
+	[
+		SNew(SHorizontalBox)
+		+ SHorizontalBox::Slot()
+		.FillWidth(1)
+		[
+			SNew(SEditableTextBox)
+			.Text(this, &FBlueprintVarActionDetails::OnGetMetaKeyValue, ClampMin)
+			.OnTextCommitted(this, &FBlueprintVarActionDetails::OnMetaKeyValueChanged, ClampMin)
+			.Font(DetailFontInfo)
+		]
+		+ SHorizontalBox::Slot()
+		.AutoWidth()
+		[
+			SNew(STextBlock)
+			.Text(LOCTEXT("Min .. Max Separator", " .. "))
+			.Font(DetailFontInfo)
+		]
+		+ SHorizontalBox::Slot()
+		.FillWidth(1)
+		[
+			SNew(SEditableTextBox)
+			.Text(this, &FBlueprintVarActionDetails::OnGetMetaKeyValue, ClampMax)
+			.OnTextCommitted(this, &FBlueprintVarActionDetails::OnMetaKeyValueChanged, ClampMax)
+			.Font(DetailFontInfo)
+		]
+	];
 	ReplicationOptions.Empty();
 	ReplicationOptions.Add(MakeShareable(new FString("None")));
 	ReplicationOptions.Add(MakeShareable(new FString("Replicated")));
@@ -1177,57 +1220,32 @@ EVisibility FBlueprintVarActionDetails::ExposeToMatineeVisibility() const
 	return EVisibility::Collapsed;
 }
 
-FText FBlueprintVarActionDetails::OnGetSliderMinValue() const
-{
-	FName VarName = GetVariableName();
-	if (VarName != NAME_None)
-	{
-		FString Result;
-		FBlueprintEditorUtils::GetBlueprintVariableMetaData(GetBlueprintObj(), VarName, TEXT("UIMin"), /*out*/ Result);
-
-		return FText::FromString(Result);
-	}
-	return FText();
-}
-
-void FBlueprintVarActionDetails::OnSliderMinValueChanged(const FText& NewMinValue, ETextCommit::Type CommitInfo)
-{
-	FName VarName = GetVariableName();
-	if (VarName != NAME_None)
-	{
-		if ((CommitInfo == ETextCommit::OnEnter) || (CommitInfo == ETextCommit::OnUserMovedFocus))
-		{
-			FBlueprintEditorUtils::SetBlueprintVariableMetaData(GetBlueprintObj(), VarName, TEXT("UIMin"), NewMinValue.ToString());
-		}
-	}
-}
-
-FText FBlueprintVarActionDetails::OnGetSliderMaxValue() const
+FText FBlueprintVarActionDetails::OnGetMetaKeyValue(FName Key) const
 {
 	FName VarName = GetVariableName();
 	if (VarName != NAME_None)
 	{
 		FString Result;
-		FBlueprintEditorUtils::GetBlueprintVariableMetaData(GetBlueprintObj(), VarName, TEXT("UIMax"), /*out*/ Result);
+		FBlueprintEditorUtils::GetBlueprintVariableMetaData(GetBlueprintObj(), VarName, Key, /*out*/ Result);
 
 		return FText::FromString(Result);
 	}
 	return FText();
 }
 
-void FBlueprintVarActionDetails::OnSliderMaxValueChanged(const FText& NewMaxValue, ETextCommit::Type CommitInfo)
+void FBlueprintVarActionDetails::OnMetaKeyValueChanged(const FText& NewMinValue, ETextCommit::Type CommitInfo, FName Key)
 {
 	FName VarName = GetVariableName();
 	if (VarName != NAME_None)
 	{
 		if ((CommitInfo == ETextCommit::OnEnter) || (CommitInfo == ETextCommit::OnUserMovedFocus))
 		{
-			FBlueprintEditorUtils::SetBlueprintVariableMetaData(GetBlueprintObj(), VarName, TEXT("UIMax"), NewMaxValue.ToString());
+			FBlueprintEditorUtils::SetBlueprintVariableMetaData(GetBlueprintObj(), VarName, Key, NewMinValue.ToString());
 		}
 	}
 }
 
-EVisibility FBlueprintVarActionDetails::SliderVisibility() const
+EVisibility FBlueprintVarActionDetails::RangeVisibility() const
 {
 	UProperty* VariableProperty = SelectionAsProperty();
 	if (VariableProperty)
diff --git a/Engine/Source/Editor/Kismet/Private/BlueprintDetailsCustomization.h b/Engine/Source/Editor/Kismet/Private/BlueprintDetailsCustomization.h
--- a/Engine/Source/Editor/Kismet/Private/BlueprintDetailsCustomization.h
+++ b/Engine/Source/Editor/Kismet/Private/BlueprintDetailsCustomization.h
@@ -90,11 +90,9 @@ private:
 	void OnExposedToMatineeChanged(ESlateCheckBoxState::Type InNewState);
 	EVisibility ExposeToMatineeVisibility() const;
 
-	FText OnGetSliderMinValue() const;
-	void OnSliderMinValueChanged(const FText& NewMinValue, ETextCommit::Type CommitInfo);
-	FText OnGetSliderMaxValue() const;
-	void OnSliderMaxValueChanged(const FText& NewMaxValue, ETextCommit::Type CommitInfo);
-	EVisibility SliderVisibility() const;
+	FText OnGetMetaKeyValue(FName Key) const;
+	void OnMetaKeyValueChanged(const FText& NewMinValue, ETextCommit::Type CommitInfo, FName Key);
+	EVisibility RangeVisibility() const;
 	
 	TSharedPtr<FString> GetVariableReplicationType() const;
 	void OnChangeReplication(TSharedPtr<FString> ItemSelected, ESelectInfo::Type SelectInfo);
diff --git a/Engine/Source/Editor/PropertyEditor/Private/UserInterface/PropertyEditor/SPropertyEditorNumeric.h b/Engine/Source/Editor/PropertyEditor/Private/UserInterface/PropertyEditor/SPropertyEditorNumeric.h
--- a/Engine/Source/Editor/PropertyEditor/Private/UserInterface/PropertyEditor/SPropertyEditorNumeric.h
+++ b/Engine/Source/Editor/PropertyEditor/Private/UserInterface/PropertyEditor/SPropertyEditorNumeric.h
@@ -199,7 +199,7 @@ private:
 				}
 			}
 
-			// We don't create a transaction for each property change when usign the slider.  Only once when the slider first is moved
+			// We don't create a transaction for each property change when using the slider.  Only once when the slider first is moved
 			EPropertyValueSetFlags::Type Flags = (EPropertyValueSetFlags::InteractiveChange | EPropertyValueSetFlags::NotTransactable);
 			PropertyHandle->SetValue( NewValue, Flags );
 		}
diff --git a/Engine/Source/Editor/UnrealEd/Private/Kismet2/BlueprintEditorUtils.cpp b/Engine/Source/Editor/UnrealEd/Private/Kismet2/BlueprintEditorUtils.cpp
--- a/Engine/Source/Editor/UnrealEd/Private/Kismet2/BlueprintEditorUtils.cpp
+++ b/Engine/Source/Editor/UnrealEd/Private/Kismet2/BlueprintEditorUtils.cpp
@@ -2367,7 +2367,14 @@ void FBlueprintEditorUtils::SetBlueprintVariableMetaData(UBlueprint* Blueprint,
 		Blueprint->NewVariables[VarIndex].SetMetaData(MetaDataKey, MetaDataValue);
 		UProperty* Property = FindField<UProperty>(Blueprint->SkeletonGeneratedClass, VarName);
 		Property->SetMetaData(MetaDataKey, *MetaDataValue);
+		
+		Property = FindField<UProperty>(Blueprint->GeneratedClass, VarName);
+		if (Property)
+		{
+			Property->SetMetaData(MetaDataKey, *MetaDataValue);
+		}
 	}
+	FBlueprintEditorUtils::MarkBlueprintAsStructurallyModified(Blueprint);
 }
 
 bool FBlueprintEditorUtils::GetBlueprintVariableMetaData(UBlueprint* Blueprint, const FName& VarName, const FName& MetaDataKey, FString& OutMetaDataValue)
