#ttp 331323 UE4: BLUEPRINTS: Interacting with Instanced Static Mesh Components in the Component view (if they have a mesh, and an instance) crashes the editor

#proj Blueprints

#branch UE4

#summary Fixed a crash after adding or removing an instance from an InstancedStaticMeshComponent node and then attempting to use the transform widget to translate one of the components in Components mode within the Blueprint editor.

#add Added FInstancedStaticMeshSCSEditorCustomization::ValidateSelectedInstances() to ensure that the number of selected instances is in sync with the number of mesh data instances before working with it.
#change Modified FInstancedStaticMeshSCSEditorCustomization methods to call ValidateSelectedInstances() to ensure that the number of selected instances is up to date prior to iterating through the bit array.

diff --git a/Engine/Source/Editor/Kismet/Private/InstancedStaticMeshSCSEditorCustomization.cpp b/Engine/Source/Editor/Kismet/Private/InstancedStaticMeshSCSEditorCustomization.cpp
--- a/Engine/Source/Editor/Kismet/Private/InstancedStaticMeshSCSEditorCustomization.cpp
+++ b/Engine/Source/Editor/Kismet/Private/InstancedStaticMeshSCSEditorCustomization.cpp
@@ -11,6 +11,24 @@ TSharedRef<ISCSEditorCustomization> FInstancedStaticMeshSCSEditorCustomization::
 	return NewCustomization;
 }
 
+void FInstancedStaticMeshSCSEditorCustomization::ValidateSelectedInstances(UInstancedStaticMeshComponent* InComponent)
+{
+	check(InComponent);
+
+	// @TODO - revisit; this might be better done as post-edit handling in the InstancedStaticMesh class itself.
+	// Ensure that the number of selection bits matches the number of instances; if not, it's likely that the user has
+	// just added or removed an instance, in which case we'll reset the selection to select the last instance in the list.
+	if(InComponent->SelectedInstances.Num() != InComponent->PerInstanceSMData.Num())
+	{
+		InComponent->SelectedInstances.Init(false, InComponent->PerInstanceSMData.Num());
+		if(InComponent->PerInstanceSMData.Num() > 0)
+		{
+			InComponent->SelectInstance(true, InComponent->PerInstanceSMData.Num() - 1);
+			InComponent->MarkRenderStateDirty();
+		}
+	}
+}
+
 bool FInstancedStaticMeshSCSEditorCustomization::HandleViewportClick(const TSharedRef<FEditorViewportClient>& InViewportClient, class FSceneView& InView, class HHitProxy* InHitProxy, FKey InKey, EInputEvent InEvent, uint32 InHitX, uint32 InHitY)
 {
 	check(InHitProxy != NULL);
@@ -47,6 +65,9 @@ bool FInstancedStaticMeshSCSEditorCustomization::HandleViewportDrag(class UScene
 	// transform pivot into component's space
 	const FVector LocalPivot = InstancedStaticMeshComponentScene->GetComponentToWorld().InverseTransformPosition(InPivot);
 
+	// Ensure that selected instances are up-to-date
+	ValidateSelectedInstances(InstancedStaticMeshComponentScene);
+
 	bool bMovedInstance = false;
 	check(InstancedStaticMeshComponentScene->SelectedInstances.Num() == InstancedStaticMeshComponentScene->PerInstanceSMData.Num());
 	for(int32 InstanceIndex = 0; InstanceIndex < InstancedStaticMeshComponentScene->SelectedInstances.Num(); InstanceIndex++)
@@ -103,6 +124,10 @@ bool FInstancedStaticMeshSCSEditorCustomization::HandleGetWidgetLocation(class U
 	float SelectedInstanceCount = 0.0f;
 	FVector AverageLocation = FVector::ZeroVector;
 	UInstancedStaticMeshComponent* InstancedStaticMeshComponent = CastChecked<UInstancedStaticMeshComponent>(InSceneComponent);
+
+	// Ensure that selected instances are up-to-date
+	ValidateSelectedInstances(InstancedStaticMeshComponent);
+
 	for(int32 InstanceIndex = 0; InstanceIndex < InstancedStaticMeshComponent->SelectedInstances.Num(); InstanceIndex++)
 	{
 		if(InstancedStaticMeshComponent->SelectedInstances[InstanceIndex])
@@ -126,6 +151,10 @@ bool FInstancedStaticMeshSCSEditorCustomization::HandleGetWidgetTransform(class
 	// transform is first selected instance
 	bool bInstanceSelected = false;
 	UInstancedStaticMeshComponent* InstancedStaticMeshComponent = CastChecked<UInstancedStaticMeshComponent>(InSceneComponent);
+
+	// Ensure that selected instances are up-to-date
+	ValidateSelectedInstances(InstancedStaticMeshComponent);
+
 	for(int32 InstanceIndex = 0; InstanceIndex < InstancedStaticMeshComponent->SelectedInstances.Num(); InstanceIndex++)
 	{
 		if(InstancedStaticMeshComponent->SelectedInstances[InstanceIndex])
diff --git a/Engine/Source/Editor/Kismet/Private/InstancedStaticMeshSCSEditorCustomization.h b/Engine/Source/Editor/Kismet/Private/InstancedStaticMeshSCSEditorCustomization.h
--- a/Engine/Source/Editor/Kismet/Private/InstancedStaticMeshSCSEditorCustomization.h
+++ b/Engine/Source/Editor/Kismet/Private/InstancedStaticMeshSCSEditorCustomization.h
@@ -16,6 +16,10 @@ public:
 	virtual bool HandleGetWidgetLocation(class USceneComponent* InSceneComponent, FVector& OutWidgetLocation) OVERRIDE;
 	virtual bool HandleGetWidgetTransform(class USceneComponent* InSceneComponent, FMatrix& OutWidgetTransform) OVERRIDE;
 
+protected:
+	/** Ensure that selection bits are in sync w/ the number of instances */
+	void ValidateSelectedInstances(class UInstancedStaticMeshComponent* InComponent);
+
 private:
 	/** The blueprint editor we are bound to */
 	TWeakPtr<class IBlueprintEditor> BlueprintEditorPtr;
