#ttp 295999 UE4: BLUEPRINTS: SCS EDITOR: Should be able to copy-paste components from the SCS (either into the same BP or into another BP)

#proj Blueprints

#branch UE4

#summary Added support for copy-and-paste for component selections in Components mode within the Blueprint editor.

#add Added FSCSEditorComponentObjectTextFactory to handle creating new object instances from clipboard contents, and preserving the copied hierarchy for scene component nodes.
#add Added a new USimpleConstructionScript::CreateNode() variant that accepts an already-constructed template instance.
#change Updated SSCSEditor to provide delegates for general cut/copy/paste UI commands, and implemented their functionality.
#change Fixed up SSCSEditor tree view node expansion states so that nodes are now expanded by default and their collapsed state is preserved across tree updates, rather than vice-versa.

diff --git a/Engine/Source/Editor/Kismet/Private/SSCSEditor.cpp b/Engine/Source/Editor/Kismet/Private/SSCSEditor.cpp
--- a/Engine/Source/Editor/Kismet/Private/SSCSEditor.cpp
+++ b/Engine/Source/Editor/Kismet/Private/SSCSEditor.cpp
@@ -14,6 +14,7 @@
 
 #include "AssetSelection.h"
 #include "Editor/SceneOutliner/Private/SSocketChooser.h"
+#include "Factories.h"
 #include "ScopedTransaction.h"
 
 #include "DragAndDrop/AssetDragDropOp.h"
@@ -478,6 +479,30 @@ bool FSCSEditorTreeNode::CanEditDefaults() const
 	return bCanEdit;
 }
 
+FSCSEditorTreeNodePtrType FSCSEditorTreeNode::FindClosestParent(TArray<FSCSEditorTreeNodePtrType> InNodes)
+{
+	uint32 MinDepth = MAX_uint32;
+	FSCSEditorTreeNodePtrType ClosestParentNodePtr;
+
+	for(int32 i = 0; i < InNodes.Num() && MinDepth > 1; ++i)
+	{
+		if(InNodes[i].IsValid())
+		{
+			uint32 CurDepth = 0;
+			if(InNodes[i]->FindChild(GetComponentTemplate(), true, &CurDepth).IsValid())
+			{
+				if(CurDepth < MinDepth)
+				{
+					MinDepth = CurDepth;
+					ClosestParentNodePtr = InNodes[i];
+				}
+			}
+		}
+	}
+
+	return ClosestParentNodePtr;
+}
+
 void FSCSEditorTreeNode::AddChild(FSCSEditorTreeNodePtrType InChildNodePtr)
 {
 	USCS_Node* SCS_Node = GetSCSNode();
@@ -571,42 +596,62 @@ FSCSEditorTreeNodePtrType FSCSEditorTreeNode::AddChild(UActorComponent* InCompon
 	return ChildNodePtr;
 }
 
-FSCSEditorTreeNodePtrType FSCSEditorTreeNode::FindChild(const USCS_Node* InSCSNode) const
+FSCSEditorTreeNodePtrType FSCSEditorTreeNode::FindChild(const USCS_Node* InSCSNode, bool bRecursiveSearch, uint32* OutDepth) const
 {
+	FSCSEditorTreeNodePtrType Result;
+
 	// Ensure that the given SCS node is valid
 	if(InSCSNode != NULL)
 	{
 		// Look for a match in our set of child nodes
-		for(int32 ChildIndex = 0; ChildIndex < Children.Num(); ++ChildIndex)
+		for(int32 ChildIndex = 0; ChildIndex < Children.Num() && !Result.IsValid(); ++ChildIndex)
 		{
 			if(InSCSNode == Children[ChildIndex]->GetSCSNode())
 			{
-				return Children[ChildIndex];
+				Result = Children[ChildIndex];
+			}
+			else if(bRecursiveSearch)
+			{
+				Result = Children[ChildIndex]->FindChild(InSCSNode, true, OutDepth);
 			}
 		}
 	}
 
-	// Return an empty node reference if no match is found
-	return FSCSEditorTreeNodePtrType();
+	if(OutDepth && Result.IsValid())
+	{
+		*OutDepth++;
+	}
+
+	return Result;
 }
 
-FSCSEditorTreeNodePtrType FSCSEditorTreeNode::FindChild(const UActorComponent* InComponentTemplate) const
+FSCSEditorTreeNodePtrType FSCSEditorTreeNode::FindChild(const UActorComponent* InComponentTemplate, bool bRecursiveSearch, uint32* OutDepth) const
 {
+	FSCSEditorTreeNodePtrType Result;
+
 	// Ensure that the given component template is valid
 	if(InComponentTemplate != NULL)
 	{
 		// Look for a match in our set of child nodes
-		for(int32 ChildIndex = 0; ChildIndex < Children.Num(); ++ChildIndex)
+		for(int32 ChildIndex = 0; ChildIndex < Children.Num() && !Result.IsValid(); ++ChildIndex)
 		{
 			if(InComponentTemplate == Children[ChildIndex]->GetComponentTemplate())
 			{
-				return Children[ChildIndex];
+				Result = Children[ChildIndex];
+			}
+			else if(bRecursiveSearch)
+			{
+				Result = Children[ChildIndex]->FindChild(InComponentTemplate, true, OutDepth);
 			}
 		}
 	}
 
-	// Return an empty node reference if no match is found
-	return FSCSEditorTreeNodePtrType();
+	if(OutDepth && Result.IsValid())
+	{
+		*OutDepth++;
+	}
+
+	return Result;
 }
 
 void FSCSEditorTreeNode::RemoveChild(FSCSEditorTreeNodePtrType InChildNodePtr)
@@ -654,6 +699,94 @@ void FSCSEditorTreeNode::OnCompleteRename(const FText& InNewName)
 	}
 }
 
+//////////////////////////////////////////////////////////////////////////
+// FSCSEditorComponentObjectTextFactory
+
+struct FSCSEditorComponentObjectTextFactory : public FCustomizableTextObjectFactory
+{
+	// Child->Parent name map
+	TMap<FName, FName> ParentMap;
+
+	// Name->Instance object mapping
+	TMap<FName, UActorComponent*> NewObjectMap;
+
+	// Determine whether or not scene components in the new object set can be attached to the given scene root component
+	bool CanAttachComponentsTo(USceneComponent* InRootComponent)
+	{
+		check(InRootComponent);
+
+		// For each component in the set, check against the given root component and break if we fail to validate
+		bool bCanAttachToRoot = true;
+		for(auto NewComponentIt = NewObjectMap.CreateConstIterator(); NewComponentIt && bCanAttachToRoot; ++NewComponentIt)
+		{
+			// If this is a scene component, and it does not already have a parent within the set
+			USceneComponent* SceneComponent = Cast<USceneComponent>(NewComponentIt->Value);
+			if(SceneComponent != NULL && !ParentMap.Contains(SceneComponent->GetFName()))
+			{
+				// Determine if we are allowed to attach the scene component to the given root component
+				bCanAttachToRoot = InRootComponent->CanAttachAsChild(SceneComponent, NAME_None)
+					&& SceneComponent->Mobility >= InRootComponent->Mobility
+					&& (!InRootComponent->IsEditorOnly() || SceneComponent->IsEditorOnly());
+			}
+		}
+
+		return bCanAttachToRoot;
+	}
+
+	// Constructs a new object factory from the given text buffer
+	static TSharedRef<FSCSEditorComponentObjectTextFactory> Get(FString InTextBuffer)
+	{
+		// Construct a new instance
+		TSharedPtr<FSCSEditorComponentObjectTextFactory> FactoryPtr = MakeShareable(new FSCSEditorComponentObjectTextFactory());
+		check(FactoryPtr.IsValid());
+
+		// Create new objects if we're allowed to
+		if(FactoryPtr->CanCreateObjectsFromText(InTextBuffer))
+		{
+			// Use the transient package initially for creating the objects, since the variable name is used when copying
+			FactoryPtr->ProcessBuffer(GetTransientPackage(), RF_ArchetypeObject|RF_Transactional, InTextBuffer);
+		}
+
+		return FactoryPtr.ToSharedRef();
+	}
+
+protected:	
+	// Constructor; protected to only allow this class to instance itself
+	FSCSEditorComponentObjectTextFactory()
+		:FCustomizableTextObjectFactory(GWarn)
+	{
+	}
+
+	// FCustomizableTextObjectFactory implementation
+
+	virtual bool CanCreateClass(UClass* ObjectClass) const OVERRIDE
+	{
+		// Only allow actor component types to be created
+		return ObjectClass->IsChildOf(UActorComponent::StaticClass());
+	}
+
+	virtual void ProcessConstructedObject(UObject* NewObject) OVERRIDE
+	{
+		check(NewObject);
+
+		// Add it to the new object map
+		NewObjectMap.Add(NewObject->GetFName(), Cast<UActorComponent>(NewObject));
+
+		// If this is a scene component and it has a parent
+		USceneComponent* SceneComponent = Cast<USceneComponent>(NewObject);
+		if(SceneComponent && SceneComponent->AttachParent)
+		{
+			// Add an entry to the child->parent name map
+			ParentMap.Add(NewObject->GetFName(), SceneComponent->AttachParent->GetFName());
+
+			// Clear this so it isn't used when constructing the new SCS node
+			SceneComponent->AttachParent = NULL;
+		}
+	}
+
+	// FCustomizableTextObjectFactory (end)
+};
+
 //////////////////////////////////////////////////////////////////////////
 // SSCS_RowWidget
 
@@ -1773,6 +1906,18 @@ void SSCSEditor::Construct( const FArguments& InArgs, TSharedPtr<FBlueprintEdito
 	SCS = InSCS;
 
 	CommandList = MakeShareable( new FUICommandList );
+	CommandList->MapAction( FGenericCommands::Get().Cut,
+		FUIAction( FExecuteAction::CreateSP( this, &SSCSEditor::CutSelectedNodes ), 
+		FCanExecuteAction::CreateSP( this, &SSCSEditor::CanCutNodes ) ) 
+		);
+	CommandList->MapAction( FGenericCommands::Get().Copy,
+		FUIAction( FExecuteAction::CreateSP( this, &SSCSEditor::CopySelectedNodes ), 
+		FCanExecuteAction::CreateSP( this, &SSCSEditor::CanCopyNodes ) ) 
+		);
+	CommandList->MapAction( FGenericCommands::Get().Paste,
+		FUIAction( FExecuteAction::CreateSP( this, &SSCSEditor::PasteNodes ), 
+		FCanExecuteAction::CreateSP( this, &SSCSEditor::CanPasteNodes ) ) 
+		);
 	CommandList->MapAction( FGenericCommands::Get().Duplicate,
 		FUIAction( FExecuteAction::CreateSP( this, &SSCSEditor::OnDuplicateComponent ), 
 		FCanExecuteAction::CreateSP( this, &SSCSEditor::CanDuplicateComponent ) ) 
@@ -1905,42 +2050,52 @@ TSharedPtr< SWidget > SSCSEditor::CreateContextMenu()
 {
 	TArray<FSCSEditorTreeNodePtrType> SelectedNodes = SCSTreeWidget->GetSelectedItems();
 
-	if (SelectedNodes.Num() > 0)
+	if (SelectedNodes.Num() > 0 || CanPasteNodes())
 	{
 		const bool CloseAfterSelection = true;
 		FMenuBuilder MenuBuilder( CloseAfterSelection, CommandList );
 
 		MenuBuilder.BeginSection("ComponentActions", LOCTEXT("ComponentContextMenu", "Component Actions") );
 		{
-			MenuBuilder.AddMenuEntry( FGenericCommands::Get().Duplicate );
-			MenuBuilder.AddMenuEntry( FGenericCommands::Get().Delete );
-			MenuBuilder.AddMenuEntry( FGenericCommands::Get().Rename );
-
-			// Collect the classes of all selected objects
-			TArray<UClass*> SelectionClasses;
-			for( auto NodeIter = SelectedNodes.CreateConstIterator(); NodeIter; ++NodeIter )
+			if(SelectedNodes.Num() > 0)
 			{
-				auto TreeNode = *NodeIter;
-				if( TreeNode->GetComponentTemplate() )
+				MenuBuilder.AddMenuEntry( FGenericCommands::Get().Cut) ;
+				MenuBuilder.AddMenuEntry( FGenericCommands::Get().Copy );
+				MenuBuilder.AddMenuEntry( FGenericCommands::Get().Paste );
+				MenuBuilder.AddMenuEntry( FGenericCommands::Get().Duplicate );
+				MenuBuilder.AddMenuEntry( FGenericCommands::Get().Delete );
+				MenuBuilder.AddMenuEntry( FGenericCommands::Get().Rename );
+
+				// Collect the classes of all selected objects
+				TArray<UClass*> SelectionClasses;
+				for( auto NodeIter = SelectedNodes.CreateConstIterator(); NodeIter; ++NodeIter )
 				{
-					SelectionClasses.Add( TreeNode->GetComponentTemplate()->GetClass() );
+					auto TreeNode = *NodeIter;
+					if( TreeNode->GetComponentTemplate() )
+					{
+						SelectionClasses.Add( TreeNode->GetComponentTemplate()->GetClass() );
+					}
 				}
-			}
 
-			TSharedPtr<FBlueprintEditor> PinnedBlueprintEditor = Kismet2Ptr.Pin();
-			if ( PinnedBlueprintEditor.IsValid() && SelectionClasses.Num() )
-			{
-				// Find the common base class of all selected classes
-				UClass* SelectedClass = UClass::FindCommonBase( SelectionClasses );
-				// Build an event submenu if we can generate events
-				if( FBlueprintEditor::CanClassGenerateEvents( SelectedClass ))
+				TSharedPtr<FBlueprintEditor> PinnedBlueprintEditor = Kismet2Ptr.Pin();
+				if ( PinnedBlueprintEditor.IsValid() && SelectionClasses.Num() )
 				{
-					MenuBuilder.AddSubMenu(	LOCTEXT("AddEventSubMenu", "Add Event"), 
-											LOCTEXT("ActtionsSubMenu_ToolTip", "Add Event"), 
-											FNewMenuDelegate::CreateStatic( &SSCSEditor::BuildMenuEventsSection, PinnedBlueprintEditor, SelectedClass, 
-											FGetSelectedObjectsDelegate::CreateSP(this, &SSCSEditor::GetSelectedItemsForContextMenu)));
+					// Find the common base class of all selected classes
+					UClass* SelectedClass = UClass::FindCommonBase( SelectionClasses );
+					// Build an event submenu if we can generate events
+					if( FBlueprintEditor::CanClassGenerateEvents( SelectedClass ))
+					{
+						MenuBuilder.AddSubMenu(	LOCTEXT("AddEventSubMenu", "Add Event"), 
+							LOCTEXT("ActtionsSubMenu_ToolTip", "Add Event"), 
+							FNewMenuDelegate::CreateStatic( &SSCSEditor::BuildMenuEventsSection, PinnedBlueprintEditor, SelectedClass, 
+							FGetSelectedObjectsDelegate::CreateSP(this, &SSCSEditor::GetSelectedItemsForContextMenu)));
+					}
 				}
 			}
+			else
+			{
+				MenuBuilder.AddMenuEntry( FGenericCommands::Get().Paste );
+			}
 		}
 		MenuBuilder.EndSection();
 
@@ -2300,9 +2455,9 @@ void SSCSEditor::UpdateTree(bool bRegenerateTreeNodes)
 
 	if(bRegenerateTreeNodes)
 	{
-		// Obtain the set of expanded tree nodes
-		TSet<FSCSEditorTreeNodePtrType> ExpandedTreeNodes;
-		SCSTreeWidget->GetExpandedItems(ExpandedTreeNodes);
+		// Obtain the set of expandable tree nodes that are currently collapsed
+		TSet<FSCSEditorTreeNodePtrType> CollapsedTreeNodes;
+		GetCollapsedNodes(SceneRootNodePtr, CollapsedTreeNodes);
 
 		// Obtain the list of selected items
 		TArray<FSCSEditorTreeNodePtrType> SelectedTreeNodes = SCSTreeWidget->GetSelectedItems();
@@ -2383,14 +2538,14 @@ void SSCSEditor::UpdateTree(bool bRegenerateTreeNodes)
 		}
 
 		// Restore the previous expansion state on the new tree nodes
-		TArray<FSCSEditorTreeNodePtrType> ExpandedTreeNodeArray = ExpandedTreeNodes.Array();
-		for(int i = 0; i < ExpandedTreeNodeArray.Num(); ++i)
+		TArray<FSCSEditorTreeNodePtrType> CollapsedTreeNodeArray = CollapsedTreeNodes.Array();
+		for(int i = 0; i < CollapsedTreeNodeArray.Num(); ++i)
 		{
-			// Look for a component match in the new hierarchy; if found, mark it as expanded to match the previous setting
-			FSCSEditorTreeNodePtrType NodeToExpandPtr = FindTreeNode(ExpandedTreeNodeArray[i]->GetComponentTemplate());
+			// Look for a component match in the new hierarchy; if found, mark it as collapsed to match the previous setting
+			FSCSEditorTreeNodePtrType NodeToExpandPtr = FindTreeNode(CollapsedTreeNodeArray[i]->GetComponentTemplate());
 			if(NodeToExpandPtr.IsValid())
 			{
-				SCSTreeWidget->SetItemExpansion(NodeToExpandPtr, true);
+				SCSTreeWidget->SetItemExpansion(NodeToExpandPtr, false);
 			}
 		}
 
@@ -2475,7 +2630,7 @@ UActorComponent* SSCSEditor::AddNewComponent( UClass* NewComponentClass, UObject
 	return AddNewNode(NewNode, Asset, true);
 }
 
-UActorComponent* SSCSEditor::AddNewNode(USCS_Node* NewNode,  UObject* Asset, bool bMarkBlueprintModified)
+UActorComponent* SSCSEditor::AddNewNode(USCS_Node* NewNode,  UObject* Asset, bool bMarkBlueprintModified, bool bSetFocusToNewItem)
 {
 	UBlueprint* Blueprint = Kismet2Ptr.Pin()->GetBlueprintObj();
 	if(Asset)
@@ -2500,9 +2655,12 @@ UActorComponent* SSCSEditor::AddNewNode(USCS_Node* NewNode,  UObject* Asset, boo
 		FBlueprintEditorUtils::ValidateBlueprintChildVariables(Blueprint, NewNode->VariableName);
 	}
 	
-	// Select and request a rename on the new component
-	SCSTreeWidget->SetSelection(NewNodePtr);
-	OnRenameComponent(false);
+	if(bSetFocusToNewItem)
+	{
+		// Select and request a rename on the new component
+		SCSTreeWidget->SetSelection(NewNodePtr);
+		OnRenameComponent(false);
+	}
 
 	// Will call UpdateTree as part of OnBlueprintChanged handling
 	if(bMarkBlueprintModified)
@@ -2536,6 +2694,186 @@ void SSCSEditor::SetSelectionOverride(UPrimitiveComponent* PrimComponent) const
 	PrimComponent->PushSelectionToProxy();
 }
 
+bool SSCSEditor::CanCutNodes() const
+{
+	return CanCopyNodes() && CanDeleteNodes();
+}
+
+void SSCSEditor::CutSelectedNodes()
+{
+	TArray<FSCSEditorTreeNodePtrType> SelectedNodes = GetSelectedNodes();
+	const FScopedTransaction Transaction( SelectedNodes.Num() > 1 ? LOCTEXT("CutComponents", "Cut Components") : LOCTEXT("CutComponent", "Cut Component") );
+
+	CopySelectedNodes();
+	OnDeleteNodes();
+}
+
+bool SSCSEditor::CanCopyNodes() const
+{
+	return CanDuplicateComponent();
+}
+
+void SSCSEditor::CopySelectedNodes()
+{
+	FStringOutputDevice Archive;
+	const FExportObjectInnerContext Context;
+	TArray<FSCSEditorTreeNodePtrType> SelectedNodes = GetSelectedNodes();
+
+	// Clear the mark state for saving.
+	UnMarkAllObjects(EObjectMark(OBJECTMARK_TagExp|OBJECTMARK_TagImp));
+
+	// Duplicate the selected component templates into temporary objects that we can modify
+	TMap<FName, FName> ParentMap;
+	TMap<FName, UActorComponent*> ObjectMap;
+	for (int32 i = 0; i < SelectedNodes.Num(); ++i)
+	{
+		// Get the current selected node reference
+		FSCSEditorTreeNodePtrType SelectedNodePtr = SelectedNodes[i];
+		check(SelectedNodePtr.IsValid());
+
+		// Get the component template associated with the selected node
+		UObject* ObjectToCopy = SelectedNodePtr->GetComponentTemplate();
+		if(ObjectToCopy)
+		{
+			// If valid, duplicate the component template into a temporary object
+			ObjectToCopy = StaticDuplicateObject(ObjectToCopy, GetTransientPackage(), *SelectedNodePtr->GetVariableName().ToString(), RF_AllFlags & ~RF_ArchetypeObject);
+			if(ObjectToCopy)
+			{
+				// Get the closest parent node of the current node selection within the selected set
+				FSCSEditorTreeNodePtrType ParentNodePtr = SelectedNodePtr->FindClosestParent(SelectedNodes);
+				if(ParentNodePtr.IsValid())
+				{
+					// If valid, record the parent node's variable name into the node->parent map
+					ParentMap.Add(SelectedNodePtr->GetVariableName(), ParentNodePtr->GetVariableName());
+				}
+
+				// Record the temporary object into the name->object map
+				ObjectMap.Add(SelectedNodePtr->GetVariableName(), CastChecked<UActorComponent>(ObjectToCopy));
+			}
+		}
+	}
+
+	// Export the component object(s) to text for copying
+	for(auto ObjectIt = ObjectMap.CreateIterator(); ObjectIt; ++ObjectIt)
+	{
+		// Get the component object to be copied
+		UActorComponent* ComponentToCopy = ObjectIt->Value;
+		check(ComponentToCopy);
+
+		// If this component object had a parent within the selected set
+		if(ParentMap.Contains(ComponentToCopy->GetFName()))
+		{
+			// Get the name of the parent component
+			FName ParentName = ParentMap[ComponentToCopy->GetFName()];
+			if(ObjectMap.Contains(ParentName))
+			{
+				// Ensure that this component is a scene component
+				USceneComponent* SceneComponent = Cast<USceneComponent>(ComponentToCopy);
+				if(SceneComponent)
+				{
+					// Set the attach parent to the matching parent object in the temporary set. This allows us to preserve hierarchy in the copied set.
+					SceneComponent->AttachParent = Cast<USceneComponent>(ObjectMap[ParentName]);
+				}
+			}
+		}
+
+		// Export the component object to the given string
+		UExporter::ExportToOutputDevice(&Context, ComponentToCopy, NULL, Archive, TEXT("copy"), 0, PPF_ExportsNotFullyQualified|PPF_Copy|PPF_Delimited, false, ComponentToCopy->GetOuter());
+	}
+
+	// Copy text to clipboard
+	FString ExportedText = Archive;
+	FPlatformMisc::ClipboardCopy(*ExportedText);
+}
+
+bool SSCSEditor::CanPasteNodes() const
+{
+	if(!InEditingMode())
+	{
+		return false;
+	}
+
+	FString ClipboardContent;
+	FPlatformMisc::ClipboardPaste(ClipboardContent);
+
+	// Obtain the component object text factory for the clipboard content and return whether or not we can use it
+	TSharedRef<FSCSEditorComponentObjectTextFactory> Factory = FSCSEditorComponentObjectTextFactory::Get(ClipboardContent);
+	return Factory->NewObjectMap.Num() > 0
+		&& (SceneRootNodePtr->IsDefaultSceneRoot() || Factory->CanAttachComponentsTo(Cast<USceneComponent>(SceneRootNodePtr->GetComponentTemplate())));
+}
+
+void SSCSEditor::PasteNodes()
+{
+	const FScopedTransaction Transaction(LOCTEXT("PasteComponents", "Paste Component(s)"));
+
+	// Mark the Blueprint as about to be modified
+	UBlueprint* Blueprint = Kismet2Ptr.Pin()->GetBlueprintObj();
+	Blueprint->Modify();
+
+	// Mark the SCS and all SCS nodes as about to be modified
+	SaveSCSCurrentState(Blueprint->SimpleConstructionScript);
+
+	// Get the text from the clipboard
+	FString TextToImport;
+	FPlatformMisc::ClipboardPaste(TextToImport);
+
+	// Get a new component object factory for the clipboard content
+	TSharedRef<FSCSEditorComponentObjectTextFactory> Factory = FSCSEditorComponentObjectTextFactory::Get(TextToImport);
+
+	// Clear the current selection
+	SCSTreeWidget->ClearSelection();
+
+	// Create a new tree node for each new (pasted) component
+	TMap<FName, FSCSEditorTreeNodePtrType> NewNodeMap;
+	for(auto NewObjectIt = Factory->NewObjectMap.CreateIterator(); NewObjectIt; ++NewObjectIt)
+	{
+		// Get the component object instance
+		UActorComponent* NewActorComponent = NewObjectIt->Value;
+		check(NewActorComponent);
+
+		// Relocate the instance from the transient package to the BPGC and assign it a unique object name
+		NewActorComponent->Rename(NULL, Blueprint->GeneratedClass, REN_DontCreateRedirectors|REN_DoNotDirty);
+
+		// Create a new SCS node to contain the new component and add it to the tree
+		NewActorComponent = AddNewNode(SCS->CreateNode(NewActorComponent), NULL, false, false);
+		if(NewActorComponent)
+		{
+			// Locate the node that corresponds to the new component template
+			FSCSEditorTreeNodePtrType NewNodePtr = FindTreeNode(NewActorComponent);
+			if(NewNodePtr.IsValid())
+			{
+				// Add the new node to the node map
+				NewNodeMap.Add(NewObjectIt->Key, NewNodePtr);
+
+				// Update the selection to include the new node
+				SCSTreeWidget->SetItemSelection(NewNodePtr, true);
+			}
+		}
+	}
+
+	// Restore the node hierarchy from the original copy
+	for(auto NodeIt = NewNodeMap.CreateConstIterator(); NodeIt; ++NodeIt)
+	{
+		// If an entry exists in the set of known parent nodes for the current node
+		if(Factory->ParentMap.Contains(NodeIt->Key))
+		{
+			// Get the parent node name
+			FName ParentName = Factory->ParentMap[NodeIt->Key];
+			if(NewNodeMap.Contains(ParentName))
+			{
+				// Reattach the current node to the parent node (this will also handle detachment from the scene root node)
+				NewNodeMap[ParentName]->AddChild(NodeIt->Value);
+
+				// Ensure that the new node is expanded to show the child node(s)
+				SCSTreeWidget->SetItemExpansion(NewNodeMap[ParentName], true);
+			}
+		}
+	}
+
+	// Modify the Blueprint generated class structure (this will also call UpdateTree() as a result)
+	FBlueprintEditorUtils::MarkBlueprintAsStructurallyModified(Blueprint);
+}
+
 bool SSCSEditor::CanDeleteNodes() const
 {
 	if(!InEditingMode())
@@ -2735,6 +3073,9 @@ FSCSEditorTreeNodePtrType SSCSEditor::AddTreeNode(USCS_Node* InSCSNode, FSCSEdit
 				// move the proposed parent in as a child to the new node
 				NewNodePtr->AddChild(OldParentPtr);
 			} // if bParentIsEditorOnly...
+
+			// Expand parent nodes by default
+			SCSTreeWidget->SetItemExpansion(ParentPtr, true);
 		}
 		//else, if !SceneRootNodePtr.IsValid(), make it the scene root node if it has not been set yet
 		else 
@@ -2805,6 +3146,9 @@ FSCSEditorTreeNodePtrType SSCSEditor::AddTreeNode(USceneComponent* InSceneCompon
 		// Add a new tree node for the given scene component
 		check(ParentNodePtr.IsValid());
 		NewNodePtr = ParentNodePtr->AddChild(InSceneComponent);
+
+		// Expand parent nodes by default
+		SCSTreeWidget->SetItemExpansion(ParentNodePtr, true);
 	}
 	else
 	{
@@ -2931,6 +3275,25 @@ bool SSCSEditor::CanRenameComponent() const
 	return SCSTreeWidget->GetSelectedItems().Num() == 1 && SCSTreeWidget->GetSelectedItems()[0]->CanRename();
 }
 
+void SSCSEditor::GetCollapsedNodes(const FSCSEditorTreeNodePtrType& InNodePtr, TSet<FSCSEditorTreeNodePtrType>& OutCollapsedNodes) const
+{
+	if(InNodePtr.IsValid())
+	{
+		const TArray<FSCSEditorTreeNodePtrType>& Children = InNodePtr->GetChildren();
+		if(Children.Num() > 0)
+		{
+			if(!SCSTreeWidget->IsItemExpanded(InNodePtr))
+			{
+				OutCollapsedNodes.Add(InNodePtr);
+			}
+
+			for(int32 i = 0; i < Children.Num(); ++i)
+			{
+				GetCollapsedNodes(Children[i], OutCollapsedNodes);
+			}
+		}
+	}
+}
 
 #undef LOCTEXT_NAMESPACE
 
diff --git a/Engine/Source/Editor/Kismet/Private/SSCSEditor.h b/Engine/Source/Editor/Kismet/Private/SSCSEditor.h
--- a/Engine/Source/Editor/Kismet/Private/SSCSEditor.h
+++ b/Engine/Source/Editor/Kismet/Private/SSCSEditor.h
@@ -116,6 +116,14 @@ public:
 	 */
 	bool CanEditDefaults() const;
 
+	/**
+	 * Finds the closest ancestor node in the given node set.
+	 * 
+	 * @param InNodes The given node set.
+	 * @return One of the nodes from the set, or an invalid node reference if the set does not contain any ancestor nodes.
+	 */
+	FSCSEditorTreeNodePtrType FindClosestParent(TArray<FSCSEditorTreeNodePtrType> InNodes);
+
 	/** 
 	 * Adds the given node as a child node.
 	 *
@@ -144,17 +152,21 @@ public:
 	 * Attempts to find a reference to the child node that matches the given SCS node.
 	 *
 	 * @param InSCSNode The SCS node to match.
-	 * @return Index of the child node that matches the given SCS node, or INDEX_NONE if no match was found.
+	 * @param bRecursiveSearch Whether or not to recursively search child nodes (default == false).
+	 * @param OutDepth If non-NULL, the depth of the child node will be returned in this parameter on success (default == NULL).
+	 * @return The child node that matches the given SCS node, or an invalid node reference if no match was found.
 	 */
-	FSCSEditorTreeNodePtrType FindChild(const USCS_Node* InSCSNode) const;
+	FSCSEditorTreeNodePtrType FindChild(const USCS_Node* InSCSNode, bool bRecursiveSearch = false, uint32* OutDepth = NULL) const;
 
 	/** 
 	 * Attempts to find a reference to the child node that matches the given component template.
 	 *
 	 * @param InComponentTemplate The component template instance to match.
-	 * @return Index of the child node with a component template that matches the given component template instance, or INDEX_NONE if no match was found.
+	 * @param bRecursiveSearch Whether or not to recursively search child nodes (default == false).
+	 * @param OutDepth If non-NULL, the depth of the child node will be returned in this parameter on success (default == NULL).
+	 * @return The child node with a component template that matches the given component template instance, or an invalid node reference if no match was found.
 	 */
-	FSCSEditorTreeNodePtrType FindChild(const UActorComponent* InComponentTemplate) const;
+	FSCSEditorTreeNodePtrType FindChild(const UActorComponent* InComponentTemplate, bool bRecursiveSearch = false, uint32* OutDepth = NULL) const;
 
 	/** 
 	 * Removes the given node from the list of child nodes.
@@ -387,8 +399,10 @@ public:
 	/** Adds a new Node to the component Table
 	   @param NewNode	(In) The node to add
 	   @param Asset		(In) Optional asset to assign to the component
+	   @param bMarkBlueprintModified (In) Whether or not to mark the Blueprint as structurally modified
+	   @param bSetFocusToNewItem (In) Select the new item and activate the inline rename widget (default is true)
 	   @return The reference of the newly created ActorComponent */
-	UActorComponent* AddNewNode(USCS_Node* NewNode, UObject* Asset, bool bMarkBlueprintModified);
+	UActorComponent* AddNewNode(USCS_Node* NewNode, UObject* Asset, bool bMarkBlueprintModified, bool bSetFocusToNewItem = true);
 	
 	/** Returns true if the specified component is currently selected */
 	bool IsComponentSelected(const UPrimitiveComponent* PrimComponent) const;
@@ -396,6 +410,18 @@ public:
 	/** Assigns a selection override delegate to the specified component */
 	void SetSelectionOverride(UPrimitiveComponent* PrimComponent) const;
 
+	/** Cut selected node(s) */
+	void CutSelectedNodes();
+	bool CanCutNodes() const;
+
+	/** Copy selected node(s) */
+	void CopySelectedNodes();
+	bool CanCopyNodes() const;
+
+	/** Pastes previously copied node(s) */
+	void PasteNodes();
+	bool CanPasteNodes() const;
+
 	/** Removes existing selected component nodes from the SCS */
 	void OnDeleteNodes();
 	bool CanDeleteNodes() const;
@@ -526,6 +552,9 @@ protected:
 
 	/** Callback when a component item is scrolled into view */
 	void OnItemScrolledIntoView( FSCSEditorTreeNodePtrType InItem, const TSharedPtr<ITableRow>& InWidget);
+
+	/** Returns the set of expandable nodes that are currently collapsed in the UI */
+	void GetCollapsedNodes(const FSCSEditorTreeNodePtrType& InNodePtr, TSet<FSCSEditorTreeNodePtrType>& OutCollapsedNodes) const;
 public:
 	/** Pointer to the script that we are editing */
 	USimpleConstructionScript* SCS;
diff --git a/Engine/Source/Runtime/Engine/Classes/Engine/SimpleConstructionScript.h b/Engine/Source/Runtime/Engine/Classes/Engine/SimpleConstructionScript.h
--- a/Engine/Source/Runtime/Engine/Classes/Engine/SimpleConstructionScript.h
+++ b/Engine/Source/Runtime/Engine/Classes/Engine/SimpleConstructionScript.h
@@ -76,6 +76,9 @@ public:
 	/** Creates a new SCS node using the given class to create the component template */
 	ENGINE_API USCS_Node* CreateNode(class UClass* NewComponentClass, FName NewComponentVariableName = NAME_None);
 
+	/** Creates a new SCS node using the given component template instance */
+	ENGINE_API USCS_Node* CreateNode(UActorComponent* NewComponentTemplate, FName NewComponentVariableName = NAME_None);
+
 	/** Ensures that all nodes in the SCS have valid names for compilation/replication */
 	ENGINE_API void ValidateNodeVariableNames(class FCompilerResultsLog& MessageLog);
 
diff --git a/Engine/Source/Runtime/Engine/Private/SimpleConstructionScript.cpp b/Engine/Source/Runtime/Engine/Private/SimpleConstructionScript.cpp
--- a/Engine/Source/Runtime/Engine/Private/SimpleConstructionScript.cpp
+++ b/Engine/Source/Runtime/Engine/Private/SimpleConstructionScript.cpp
@@ -600,25 +600,35 @@ USCS_Node* USimpleConstructionScript::CreateNode(UClass* NewComponentClass, FNam
 	check(NewComponentClass->IsChildOf(UActorComponent::StaticClass()));
 
 	ensure(NULL != Cast<UBlueprintGeneratedClass>(Blueprint->GeneratedClass));
-	UActorComponent* NewTemplate = ConstructObject<UActorComponent>(NewComponentClass, Blueprint->GeneratedClass);
-	NewTemplate->SetFlags(RF_ArchetypeObject|RF_Transactional);
+	UActorComponent* NewComponentTemplate = ConstructObject<UActorComponent>(NewComponentClass, Blueprint->GeneratedClass);
+	NewComponentTemplate->SetFlags(RF_ArchetypeObject|RF_Transactional);
 
-	// Create a node for the script, and save a pointer to the template.
-	USCS_Node* NewNode = NewObject<USCS_Node>(this);
-	NewNode->SetFlags(RF_Transactional);
-	NewNode->ComponentTemplate = NewTemplate;
+	return CreateNode(NewComponentTemplate, NewComponentVariableName);
+}
+
+USCS_Node* USimpleConstructionScript::CreateNode(UActorComponent* NewComponentTemplate, FName NewComponentVariableName)
+{
+	if(NewComponentTemplate)
+	{
+		// Create a node for the script, and save a pointer to the template.
+		USCS_Node* NewNode = NewObject<USCS_Node>(this);
+		NewNode->SetFlags(RF_Transactional);
+		NewNode->ComponentTemplate = NewComponentTemplate;
 
-	// Get a list of names currently in use.
-	TArray<FName> CurrentNames;
-	NewNode->GenerateListOfExistingNames( CurrentNames );
+		// Get a list of names currently in use.
+		TArray<FName> CurrentNames;
+		NewNode->GenerateListOfExistingNames( CurrentNames );
 
-	// Now create a name for the new component.
-	NewNode->VariableName = NewNode->GenerateNewComponentName( CurrentNames, NewComponentVariableName );
+		// Now create a name for the new component.
+		NewNode->VariableName = NewNode->GenerateNewComponentName( CurrentNames, NewComponentVariableName );
 
-	// Note: This should match up with UEdGraphSchema_K2::VR_DefaultCategory
-	NewNode->CategoryName = TEXT("Default");
+		// Note: This should match up with UEdGraphSchema_K2::VR_DefaultCategory
+		NewNode->CategoryName = TEXT("Default");
 
-	return NewNode;
+		return NewNode;
+	}
+
+	return NULL;
 }
 
 void USimpleConstructionScript::ValidateNodeVariableNames(FCompilerResultsLog& MessageLog)
