#summary Can search for nodes in the palette menu in Blueprints with both the native name and the localized name.

diff --git a/Engine/Source/Editor/Kismet/Private/BPFunctionDragDropAction.cpp b/Engine/Source/Editor/Kismet/Private/BPFunctionDragDropAction.cpp
--- a/Engine/Source/Editor/Kismet/Private/BPFunctionDragDropAction.cpp
+++ b/Engine/Source/Editor/Kismet/Private/BPFunctionDragDropAction.cpp
@@ -96,8 +96,8 @@ void FKismetDragDropAction::HoverTargetChanged()
 	if (ActionWillShowExistingNode())
 	{
 		FSlateBrush const* ShowsExistingIcon = FEditorStyle::GetBrush(TEXT("Graph.ConnectorFeedback.ShowNode"));
-		FText DragingText = FText::Format(LOCTEXT("ShowExistingNode", "Show '{0}'"), FText::FromString(ActionNode->MenuDescription));
-		SetSimpleFeedbackMessage(ShowsExistingIcon, FLinearColor::White, DragingText.ToString());
+		FText DragingText = FText::Format(LOCTEXT("ShowExistingNode", "Show '{0}'"), ActionNode->MenuDescription);
+		SetSimpleFeedbackMessage(ShowsExistingIcon, FLinearColor::White, DragingText);
 	}
 	// it should be obvious that we can't drop on anything but a graph, so no need to point that out
 	else if ((HoveredGraph == NULL) || !CanBeDroppedDelegate.IsBound() || CanBeDroppedDelegate.Execute(ActionNode, HoveredGraph, CannotDropReason))
@@ -107,7 +107,7 @@ void FKismetDragDropAction::HoverTargetChanged()
 	else 
 	{
 		FSlateBrush const* DropPreventedIcon = FEditorStyle::GetBrush(TEXT("Graph.ConnectorFeedback.Error"));
-		SetSimpleFeedbackMessage(DropPreventedIcon, FLinearColor::White, CannotDropReason.ToString());
+		SetSimpleFeedbackMessage(DropPreventedIcon, FLinearColor::White, CannotDropReason);
 	}
 }
 
@@ -315,12 +315,12 @@ void FKismetMacroDragDropAction::HoverTargetChanged()
 	if ((HoveredGraph == NULL) || !CanBeDroppedDelegate.IsBound() || CanBeDroppedDelegate.Execute(ActionNode, HoveredGraph, CannotDropReason))
 	{
 		FSlateBrush const* DropPreventedIcon = FEditorStyle::GetBrush(TEXT("Graph.ConnectorFeedback.NewNode"));
-		SetSimpleFeedbackMessage(DropPreventedIcon, FLinearColor::White, MacroName.ToString());
+		SetSimpleFeedbackMessage(DropPreventedIcon, FLinearColor::White, FText::FromName(MacroName));
 	}
 	else 
 	{
 		FSlateBrush const* DropPreventedIcon = FEditorStyle::GetBrush(TEXT("Graph.ConnectorFeedback.Error"));
-		SetSimpleFeedbackMessage(DropPreventedIcon, FLinearColor::White, CannotDropReason.ToString());
+		SetSimpleFeedbackMessage(DropPreventedIcon, FLinearColor::White, CannotDropReason);
 	}
 }
 
diff --git a/Engine/Source/Editor/Kismet/Private/BPVariableDragDropAction.cpp b/Engine/Source/Editor/Kismet/Private/BPVariableDragDropAction.cpp
--- a/Engine/Source/Editor/Kismet/Private/BPVariableDragDropAction.cpp
+++ b/Engine/Source/Editor/Kismet/Private/BPVariableDragDropAction.cpp
@@ -49,7 +49,7 @@ void FKismetVariableDragDropAction::HoverTargetChanged()
 	// Icon/text to draw on tooltip
 	FSlateColor IconColor = FLinearColor::White;
 	const FSlateBrush* StatusSymbol = FEditorStyle::GetBrush(TEXT("NoBrush")); 
-	FString Message;
+	FText Message;
 
 	UEdGraphPin* PinUnderCursor = GetHoveredPin();
 
@@ -67,7 +67,7 @@ void FKismetVariableDragDropAction::HoverTargetChanged()
 	if (bBadSchema)
 	{
 		StatusSymbol = FEditorStyle::GetBrush(TEXT("Graph.ConnectorFeedback.Error"));
-		Message = LOCTEXT("CannotCreateInThisSchema", "Cannot access variables in this type of graph").ToString();
+		Message = LOCTEXT("CannotCreateInThisSchema", "Cannot access variables in this type of graph");
 	}
 	else if (PinUnderCursor != NULL)
 	{
@@ -82,17 +82,21 @@ void FKismetVariableDragDropAction::HoverTargetChanged()
 		Schema->ConvertPropertyToPinType(VariableProperty, VariablePinType);
 		const bool bTypeMatch = Schema->ArePinTypesCompatible(VariablePinType, PinUnderCursor->PinType);
 
+		FFormatNamedArguments Args;
+		Args.Add(TEXT("PinUnderCursor"), FText::FromString(PinUnderCursor->PinName));
+		Args.Add(TEXT("VariableName"), FText::FromString(VariableString));
+
 		if (bTypeMatch && bCanWriteIfNeeded)
 		{
 			StatusSymbol = FEditorStyle::GetBrush(TEXT("Graph.ConnectorFeedback.OK"));
 
 			if (bIsRead)
 			{
-				Message = FString::Printf( *LOCTEXT("MakeThisEqualThat", "Make %s = %s").ToString(), *PinUnderCursor->PinName, *VariableString);
+				Message = FText::Format(LOCTEXT("MakeThisEqualThat", "Make {PinUnderCursor} = {VariableName}"), Args);
 			}
 			else
 			{
-				Message = FString::Printf( *LOCTEXT("MakeThisEqualThat", "Make %s = %s").ToString(), *VariableString, *PinUnderCursor->PinName);
+				Message = FText::Format(LOCTEXT("MakeThisEqualThat", "Make {VariableName} = {PinUnderCursor}"), Args);
 			}
 		}
 		else
@@ -100,11 +104,11 @@ void FKismetVariableDragDropAction::HoverTargetChanged()
 			StatusSymbol = FEditorStyle::GetBrush(TEXT("Graph.ConnectorFeedback.Error"));
 			if (!bCanWriteIfNeeded)
 			{
-				Message = FString::Printf( *LOCTEXT("ReadOnlyVar_Error", "Cannot write to read-only variable '%s'").ToString(), *VariableString);
+				Message = FText::Format(LOCTEXT("ReadOnlyVar_Error", "Cannot write to read-only variable '{VariableName}'"), Args);
 			}
 			else
 			{
-				Message = FString::Printf( *LOCTEXT("NotCompatible_Error", "The type of '%s' is not compatible with %s").ToString(), *VariableString, *(PinUnderCursor->PinName));
+				Message = FText::Format(LOCTEXT("NotCompatible_Error", "The type of '{VariableName}' is not compatible with {PinUnderCursor}"), Args);
 			}
 		}
 	}
@@ -115,23 +119,27 @@ void FKismetVariableDragDropAction::HoverTargetChanged()
 		const bool bReadOnlyProperty = FBlueprintEditorUtils::IsPropertyReadOnlyInCurrentBlueprint(Blueprint, VariableProperty);
 		const bool bCanWriteIfNeeded = bIsRead || !bReadOnlyProperty;
 		
+		FFormatNamedArguments Args;
+		Args.Add(TEXT("VariableName"), FText::FromString(VariableString));
+
 		if (bCanWriteIfNeeded)
 		{
+			Args.Add(TEXT("ReadOrWrite"), bIsRead ? LOCTEXT("Read", "read") : LOCTEXT("Write", "write"));
 			if(WillBreakLinks(VarNodeUnderCursor, VariableProperty))
 			{
 				StatusSymbol = FEditorStyle::GetBrush(TEXT("Graph.ConnectorFeedback.OKWarn"));
-				Message = FString::Printf( *LOCTEXT("ChangeNodeToWarnBreakLinks", "Change node to %s '%s', WARNING this will break links!").ToString(), bIsRead ? TEXT("read") : TEXT("write"), *VariableString);
+				Message = FText::Format( LOCTEXT("ChangeNodeToWarnBreakLinks", "Change node to {ReadOrWrite} '{VariableName}', WARNING this will break links!"), Args);
 			}
 			else
 			{
 				StatusSymbol = FEditorStyle::GetBrush(TEXT("Graph.ConnectorFeedback.OK"));
-				Message = FString::Printf( *LOCTEXT("ChangeNodeTo", "Change node to %s '%s'").ToString(), bIsRead ? TEXT("read") : TEXT("write"), *VariableString);
+				Message = FText::Format( LOCTEXT("ChangeNodeTo", "Change node to {ReadOrWrite} '{VariableName}'"), Args);
 			}
 		}
 		else
 		{
 			StatusSymbol = FEditorStyle::GetBrush(TEXT("Graph.ConnectorFeedback.Error"));
-			Message = FString::Printf( *LOCTEXT("ReadOnlyVar_Error", "Cannot write to read-only variable '%s'").ToString(), *VariableString);
+			Message = FText::Format( LOCTEXT("ReadOnlyVar_Error", "Cannot write to read-only variable '{VariableName}'"), Args);
 		}
 	}
 	else if (HoveredCategoryName.Len() > 0)
@@ -150,20 +158,24 @@ void FKismetVariableDragDropAction::HoverTargetChanged()
 		UClass* OuterClass = CastChecked<UClass>(VariableProperty->GetOuter());
 		const bool bIsNativeVar = (OuterClass->ClassGeneratedBy == NULL);
 
+		FFormatNamedArguments Args;
+		Args.Add(TEXT("VariableName"), FText::FromString(VariableString));
+		Args.Add(TEXT("HoveredCategoryName"), FText::FromString(HoveredCategoryName));
+
 		if (bIsNativeVar)
 		{
 			StatusSymbol = FEditorStyle::GetBrush(TEXT("Graph.ConnectorFeedback.Error"));
-			Message = FString::Printf( *LOCTEXT("ChangingCatagoryNotThisVar", "Cannot change category for variable '%s'").ToString(), *VariableString );
+			Message = FText::Format( LOCTEXT("ChangingCatagoryNotThisVar", "Cannot change category for variable '{VariableName}'"), Args );
 		}
 		else if (Category == NewCategory)
 		{
 			StatusSymbol = FEditorStyle::GetBrush(TEXT("Graph.ConnectorFeedback.Error"));
-			Message = FString::Printf( *LOCTEXT("ChangingCatagoryAlreadyIn", "Variable '%s' is already in category '%s'").ToString(), *VariableString, *HoveredCategoryName );
+			Message = FText::Format( LOCTEXT("ChangingCatagoryAlreadyIn", "Variable '{VariableName}' is already in category '{HoveredCategoryName}'"), Args );
 		}
 		else
 		{
 			StatusSymbol = FEditorStyle::GetBrush(TEXT("Graph.ConnectorFeedback.OK"));
-			Message = FString::Printf( *LOCTEXT("ChangingCatagoryOk", "Move variable '%s' to category '%s'").ToString(),  *VariableString, *HoveredCategoryName );
+			Message = FText::Format( LOCTEXT("ChangingCatagoryOk", "Move variable '{VariableName}' to category '{HoveredCategoryName}'"), Args );
 		}
 	}
 	else if (HoveredAction.IsValid())
@@ -183,25 +195,29 @@ void FKismetVariableDragDropAction::HoverTargetChanged()
 				TargetVarIndex = FBlueprintEditorUtils::FindNewVariableIndex(Blueprint, TargetVarName);
 			}
 
+			FFormatNamedArguments Args;
+			Args.Add(TEXT("VariableName"), FText::FromString(VariableString));
+			Args.Add(TEXT("TargetVarName"), FText::FromName(TargetVarName));
+
 			if(MoveVarIndex == INDEX_NONE)
 			{
 				StatusSymbol = FEditorStyle::GetBrush(TEXT("Graph.ConnectorFeedback.Error"));
-				Message = FString::Printf( *LOCTEXT("MoveVarDiffClass", "Cannot reorder variable '%s'.").ToString(), *VariableString );
+				Message = FText::Format( LOCTEXT("MoveVarDiffClass", "Cannot reorder variable '{VariableName}'."), Args );
 			}
 			else if(TargetVarIndex == INDEX_NONE)
 			{
 				StatusSymbol = FEditorStyle::GetBrush(TEXT("Graph.ConnectorFeedback.Error"));
-				Message = FString::Printf( *LOCTEXT("MoveVarOther", "Cannot reorder variable '%s' before '%s'.").ToString(), *VariableString, *TargetVarName.ToString() );
+				Message = FText::Format( LOCTEXT("MoveVarOther", "Cannot reorder variable '{VariableName}' before '{TargetVarName}'."), Args );
 			}
 			else if(VariableName == TargetVarName)
 			{
 				StatusSymbol = FEditorStyle::GetBrush(TEXT("Graph.ConnectorFeedback.Error"));
-				Message = FString::Printf( *LOCTEXT("MoveVarYourself", "Cannot reorder variable '%s' before itself.").ToString(), *VariableString );
+				Message = FText::Format( LOCTEXT("MoveVarYourself", "Cannot reorder variable '{VariableName}' before itself."), Args );
 			}
 			else
 			{
 				StatusSymbol = FEditorStyle::GetBrush(TEXT("Graph.ConnectorFeedback.OK"));
-				Message = FString::Printf( *LOCTEXT("MoveVarOK", "Reorder variable '%s' before '%s'").ToString(), *VariableString, *TargetVarName.ToString() );
+				Message = FText::Format( LOCTEXT("MoveVarOK", "Reorder variable '{VariableName}' before '{TargetVarName}'"), Args );
 			}
 		}
 	}
@@ -209,7 +225,7 @@ void FKismetVariableDragDropAction::HoverTargetChanged()
 	else
 	{
 		StatusSymbol = FBlueprintEditor::GetVarIconAndColor(VariableSourceClass.Get(), VariableName, IconColor);
-		Message = VariableString;
+		Message = FText::FromString(VariableString);
 	}
 
 	SetSimpleFeedbackMessage(StatusSymbol, IconColor, Message);
diff --git a/Engine/Source/Editor/Kismet/Private/BlueprintDetailsCustomization.cpp b/Engine/Source/Editor/Kismet/Private/BlueprintDetailsCustomization.cpp
--- a/Engine/Source/Editor/Kismet/Private/BlueprintDetailsCustomization.cpp
+++ b/Engine/Source/Editor/Kismet/Private/BlueprintDetailsCustomization.cpp
@@ -4609,7 +4609,7 @@ FText FBlueprintGraphNodeDetails::OnGetName() const
 	FText Name;
 	if(GraphNodePtr.IsValid())
 	{
-		Name = FText::FromString( GraphNodePtr->GetNodeTitle( ENodeTitleType::EditableTitle ));
+		Name = GraphNodePtr->GetNodeTitle( ENodeTitleType::EditableTitle );
 	}
 	return Name;
 }
@@ -4620,7 +4620,7 @@ void FBlueprintGraphNodeDetails::OnNameChanged(const FText& InNewText)
 
 	if( GraphNodePtr.IsValid() && BlueprintEditorPtr.IsValid() )
 	{
-		FName NodeName( *GraphNodePtr->GetNodeTitle(ENodeTitleType::EditableTitle) );
+		FName NodeName( *GraphNodePtr->GetNodeTitle(ENodeTitleType::EditableTitle).ToString() );
 		TSharedPtr<INameValidatorInterface> NameValidator = GraphNodePtr->MakeNameValidator();
 
 		if( !NameValidator.IsValid() )
diff --git a/Engine/Source/Editor/Kismet/Private/BlueprintEditor.cpp b/Engine/Source/Editor/Kismet/Private/BlueprintEditor.cpp
--- a/Engine/Source/Editor/Kismet/Private/BlueprintEditor.cpp
+++ b/Engine/Source/Editor/Kismet/Private/BlueprintEditor.cpp
@@ -4694,7 +4694,7 @@ void FBlueprintEditor::CollapseNodesIntoGraph(UEdGraphNode* InGatewayNode, UK2No
 
 				ExtractEventTemplateForFunction(CustomEvent, InGatewayNode, InEntryNode, InResultNode, InCollapsableNodes);
 
-				FString GraphName = FBlueprintEditorUtils::GenerateUniqueGraphName(GetBlueprintObj(), CustomEvent->GetNodeTitle(ENodeTitleType::ListView)).ToString();
+				FString GraphName = FBlueprintEditorUtils::GenerateUniqueGraphName(GetBlueprintObj(), CustomEvent->GetNodeTitle(ENodeTitleType::ListView).ToString()).ToString();
 				FBlueprintEditorUtils::RenameGraph(InDestinationGraph, GraphName);
 
 				// Remove the node, it has no place in the new graph
diff --git a/Engine/Source/Editor/Kismet/Private/BlueprintEditorCommands.cpp b/Engine/Source/Editor/Kismet/Private/BlueprintEditorCommands.cpp
--- a/Engine/Source/Editor/Kismet/Private/BlueprintEditorCommands.cpp
+++ b/Engine/Source/Editor/Kismet/Private/BlueprintEditorCommands.cpp
@@ -172,6 +172,7 @@ public:
 
 		// the NodeTemplate UPROPERTY takes ownership of GraphNode's lifetime (hence it being a weak-pointer here)
 		NewActionNode->NodeTemplate = GraphNode.Get();
+		NewActionNode->SearchTitle = NewActionNode->NodeTemplate->GetNodeSearchTitle();
 
 		return NewActionNode;
 	}
@@ -289,7 +290,7 @@ void FBlueprintSpawnNodeCommands::RegisterCommands()
 			// The class name matches that of a UEdGraphNode, so setup a spawn info that can generate UEdGraphNode graph actions
 			UEdGraphNode* GraphNode = Cast<UEdGraphNode>(FoundClass->GetDefaultObject());
 
-			CommandLabel = GraphNode->GetNodeTitle(ENodeTitleType::ListView);
+			CommandLabel = GraphNode->GetNodeTitle(ENodeTitleType::ListView).ToString();
 
 			if(CommandLabel.Len() == 0)
 			{
diff --git a/Engine/Source/Editor/Kismet/Private/Debugging/SKismetDebuggingView.cpp b/Engine/Source/Editor/Kismet/Private/Debugging/SKismetDebuggingView.cpp
--- a/Engine/Source/Editor/Kismet/Private/Debugging/SKismetDebuggingView.cpp
+++ b/Engine/Source/Editor/Kismet/Private/Debugging/SKismetDebuggingView.cpp
@@ -26,9 +26,9 @@ namespace KismetDebugViewConstants
 //////////////////////////////////////////////////////////////////////////
 // FDebugLineItem
 
-FString FDebugLineItem::GetDisplayName() const
+FText FDebugLineItem::GetDisplayName() const
 {
-	return TEXT("");
+	return FText::GetEmpty();
 }
 
 FString FDebugLineItem::GetDescription() const
@@ -171,7 +171,7 @@ protected:
 protected:
 	virtual TSharedRef<SWidget> GenerateNameWidget() OVERRIDE;
 	virtual FString GetDescription() const OVERRIDE;
-	virtual FString GetDisplayName() const OVERRIDE;
+	virtual FText GetDisplayName() const OVERRIDE;
 	void OnNavigateToLatentNode();
 
 	class UEdGraphNode* FindAssociatedNode() const;
@@ -224,15 +224,19 @@ UEdGraphNode* FLatentActionLineItem::FindAssociatedNode() const
 	return NULL;
 }
 
-FString FLatentActionLineItem::GetDisplayName() const
+FText FLatentActionLineItem::GetDisplayName() const
 {
+	FFormatNamedArguments Args;
+	Args.Add(TEXT("ID"), UUID);
 	if (UK2Node* Node = Cast<UK2Node>(FindAssociatedNode()))
 	{
-		return FString::Printf( *LOCTEXT("ID", "%s (ID: %d)").ToString(), *Node->GetCompactNodeTitle(), UUID);
+		Args.Add(TEXT("Title"), Node->GetCompactNodeTitle());
+
+		return FText::Format(LOCTEXT("ID", "{Title} (ID: {ID})"), Args);
 	}
 	else
 	{
-		return FString::Printf( *LOCTEXT("LatentAction", "Latent action # %d").ToString(), UUID);
+		return FText::Format(LOCTEXT("LatentAction", "Latent action # {ID}"), Args);
 	}
 }
 
@@ -289,13 +293,13 @@ public:
 	}
 protected:
 	virtual FString GetDescription() const OVERRIDE;
-	virtual FString GetDisplayName() const OVERRIDE;
+	virtual FText GetDisplayName() const OVERRIDE;
 	virtual TSharedRef<SWidget> GenerateNameWidget() OVERRIDE;
 
 	void OnNavigateToWatchLocation();
 };
 
-FString FWatchLineItem::GetDisplayName() const
+FText FWatchLineItem::GetDisplayName() const
 {
 	if (UEdGraphPin* PinToWatch = Cast<UEdGraphPin>(ObjectRef.Get()))
 	{
@@ -303,15 +307,17 @@ FString FWatchLineItem::GetDisplayName() const
 		{
 			if (UProperty* Property = FKismetDebugUtilities::FindClassPropertyForPin(Blueprint, PinToWatch))
 			{
-				return UEditorEngine::GetFriendlyName(Property);
+				return FText::FromString(UEditorEngine::GetFriendlyName(Property));
 			}
 		}
 
-		return PinToWatch->GetName() + LOCTEXT("NoProp", " (no prop)").ToString();
+		FFormatNamedArguments Args;
+		Args.Add(TEXT("PinWatchName"), FText::FromString(PinToWatch->GetName()));
+		return FText::Format(LOCTEXT("DisplayNameNoProperty", "{PinWatchName} (no prop)"), Args);
 	}
 	else
 	{
-		return TEXT("");
+		return FText::GetEmpty();
 	}
 }
 
@@ -672,18 +678,18 @@ protected:
 		return new FParentLineItem(ObjectRef.Get());
 	}
 
-	virtual FString GetDisplayName() const OVERRIDE
+	virtual FText GetDisplayName() const OVERRIDE
 	{
 		UObject* Object = ObjectRef.Get();
 		AActor* Actor = Cast<AActor>(Object);
 
 		if (Actor != NULL)
 		{
-			 return Actor->GetActorLabel();
+			 return FText::FromString(Actor->GetActorLabel());
 		}
 		else
 		{
-			return (Object != NULL) ? Object->GetName() : TEXT("(null)");
+			return (Object != NULL) ? FText::FromString(Object->GetName()) : LOCTEXT("Null", "(null)");
 		}
 	}
 
@@ -768,7 +774,7 @@ protected:
 		return NULL;
 	}
 
-	virtual FString GetDisplayName() const OVERRIDE
+	virtual FText GetDisplayName() const OVERRIDE
 	{
 		UEdGraphNode* Node = GetNode();
 		if (Node != NULL)
@@ -777,7 +783,7 @@ protected:
 		}
 		else
 		{
-			return LOCTEXT("Unknown", "(unknown)").ToString();
+			return LOCTEXT("Unknown", "(unknown)");
 		}
 	}
 
@@ -907,9 +913,9 @@ public:
 		}
 	}
 protected:
-	virtual FString GetDisplayName() const
+	virtual FText GetDisplayName() const
 	{
-		return LOCTEXT("ExecutionTrace", "Execution Trace").ToString();
+		return LOCTEXT("ExecutionTrace", "Execution Trace");
 	}
 
 	virtual bool Compare(const FDebugLineItem* BaseOther) const OVERRIDE
diff --git a/Engine/Source/Editor/Kismet/Private/Debugging/SKismetDebuggingView.h b/Engine/Source/Editor/Kismet/Private/Debugging/SKismetDebuggingView.h
--- a/Engine/Source/Editor/Kismet/Private/Debugging/SKismetDebuggingView.h
+++ b/Engine/Source/Editor/Kismet/Private/Debugging/SKismetDebuggingView.h
@@ -67,7 +67,7 @@ protected:
 	virtual bool Compare(const FDebugLineItem* Other) const=0;
 
 	// @return The text to display in the name column, unless GenerateNameWidget is overridden
-	virtual FString GetDisplayName() const;
+	virtual FText GetDisplayName() const;
 
 	// @return The text to display in the value column, unless GenerateValueWidget is overridden
 	virtual FString GetDescription() const;
diff --git a/Engine/Source/Editor/Kismet/Private/FindInBlueprints.cpp b/Engine/Source/Editor/Kismet/Private/FindInBlueprints.cpp
--- a/Engine/Source/Editor/Kismet/Private/FindInBlueprints.cpp
+++ b/Engine/Source/Editor/Kismet/Private/FindInBlueprints.cpp
@@ -50,7 +50,7 @@ inline bool GetNode(UBlueprint* Blueprint, const FFindInBlueprintsResult& Result
 		for(auto It(Nodes.CreateConstIterator());It;++It)
 		{
 			NodeType* Node = *It;
-			const FString NodePath = FindInBlueprintsUtil::GetNodeTypePath(Node) + "::" + Node->GetNodeTitle(ENodeTitleType::ListView) ;
+			const FString NodePath = FindInBlueprintsUtil::GetNodeTypePath(Node) + "::" + Node->GetNodeTitle(ENodeTitleType::ListView).ToString() ;
 			if(NodePath == Result.Value)
 			{
 				if(FoundCount == Result.DuplicationIndex)
@@ -435,7 +435,7 @@ void SFindInBlueprints::MatchTokensWithinBlueprint(const TArray<FString>& Tokens
 	{
 		UK2Node* Node = *It;
 			
-		const FString NodeName = Node->GetNodeTitle(ENodeTitleType::ListView);
+		const FString NodeName = Node->GetNodeTitle(ENodeTitleType::ListView).ToString();
 		FSearchResult NodeResult(new FFindInBlueprintsResult(NodeName, RootSearchResult, Node));
 		UK2Node* K2Node = Cast<UK2Node>(Node);
 		AActor* NodeActor = K2Node ? K2Node->GetReferencedLevelActor() : NULL;
@@ -448,7 +448,7 @@ void SFindInBlueprints::MatchTokensWithinBlueprint(const TArray<FString>& Tokens
 		{
 			UEdGraphPin* Pin = *PinIt;
 			FString PinName = Pin->PinName;
-			FString PinSearchString = Pin->PinName + Pin->PinFriendlyName + Pin->DefaultValue + Pin->PinType.PinCategory + Pin->PinType.PinSubCategory + (Pin->PinType.PinSubCategoryObject.IsValid() ? Pin->PinType.PinSubCategoryObject.Get()->GetFullName() : TEXT(""));
+			FString PinSearchString = Pin->PinName + Pin->PinFriendlyName.ToString() + Pin->DefaultValue + Pin->PinType.PinCategory + Pin->PinType.PinSubCategory + (Pin->PinType.PinSubCategoryObject.IsValid() ? Pin->PinType.PinSubCategoryObject.Get()->GetFullName() : TEXT(""));
 			PinSearchString = PinSearchString.Replace( TEXT( " " ), TEXT( "" ) );
 			if ((Pin != NULL) && StringMatchesSearchTokens(Tokens, PinSearchString))
 			{
diff --git a/Engine/Source/Editor/Kismet/Private/SBlueprintActionMenu.cpp b/Engine/Source/Editor/Kismet/Private/SBlueprintActionMenu.cpp
--- a/Engine/Source/Editor/Kismet/Private/SBlueprintActionMenu.cpp
+++ b/Engine/Source/Editor/Kismet/Private/SBlueprintActionMenu.cpp
@@ -23,7 +23,7 @@ struct FBlueprintAction_PromoteVariable : public FEdGraphSchemaAction
 		: FEdGraphSchemaAction()
 	{
 		Category = TEXT("");
-		MenuDescription = LOCTEXT("PromoteToVariable", "Promote to variable").ToString();
+		MenuDescription = LOCTEXT("PromoteToVariable", "Promote to variable");
 		TooltipDescription = LOCTEXT("PromoteToVariable", "Promote to variable").ToString();
 		Grouping = 1;
 	}
diff --git a/Engine/Source/Editor/Kismet/Private/SBlueprintPalette.cpp b/Engine/Source/Editor/Kismet/Private/SBlueprintPalette.cpp
--- a/Engine/Source/Editor/Kismet/Private/SBlueprintPalette.cpp
+++ b/Engine/Source/Editor/Kismet/Private/SBlueprintPalette.cpp
@@ -257,7 +257,7 @@ static void GetSubGraphIcon(FEdGraphSchemaAction_K2Graph const* const ActionIn,
 static void GetPaletteItemIcon(TSharedPtr<FEdGraphSchemaAction> ActionIn, UBlueprint const* BlueprintIn, FSlateBrush const*& BrushOut, FSlateColor& ColorOut, FString& ToolTipOut, FString& DocLinkOut, FString& DocExcerptOut)
 {
 	// Default to tooltip based on action supplied
-	ToolTipOut = (ActionIn->TooltipDescription.Len() > 0) ? ActionIn->TooltipDescription : ActionIn->MenuDescription;
+	ToolTipOut = (ActionIn->TooltipDescription.Len() > 0) ? ActionIn->TooltipDescription : ActionIn->MenuDescription.ToString();
 
 	if (ActionIn->GetTypeId() == FEdGraphSchemaAction_K2TargetNode::StaticGetTypeId())
 	{
@@ -889,15 +889,20 @@ TSharedRef<SWidget> SBlueprintPaletteItem::CreateTextSlotWidget(const FSlateFont
 //------------------------------------------------------------------------------
 FText SBlueprintPaletteItem::GetDisplayText() const
 {
-	FString DisplayText = ActionPtr.Pin()->MenuDescription;
+	FText DisplayText;
 
 	TSharedPtr< FEdGraphSchemaAction > GraphAction = ActionPtr.Pin();
 	if (GraphAction->GetTypeId() == FEdGraphSchemaAction_K2Enum::StaticGetTypeId())
 	{
 		FEdGraphSchemaAction_K2Enum* EnumAction = (FEdGraphSchemaAction_K2Enum*)GraphAction.Get();
-		DisplayText = EnumAction->Enum->GetName();
+		DisplayText = FText::FromString(EnumAction->Enum->GetName());
 	}
-	return FText::FromString(DisplayText);
+	else
+	{
+		DisplayText = ActionPtr.Pin()->MenuDescription;
+	}
+
+	return DisplayText;
 }
 
 //------------------------------------------------------------------------------
@@ -991,7 +996,7 @@ void SBlueprintPaletteItem::OnNameTextCommitted(const FText& NewText, ETextCommi
 				GraphAction->EdGraph->GetSchema()->GetGraphDisplayInformation(*GraphAction->EdGraph, DisplayInfo);
 
 				// Check if the name is unchanged
-				if( NewText.EqualTo( DisplayInfo.DisplayName ) )
+				if( NewText.EqualTo( DisplayInfo.PlainName ) )
 				{
 					return;
 				}
@@ -1020,7 +1025,7 @@ void SBlueprintPaletteItem::OnNameTextCommitted(const FText& NewText, ETextCommi
 				Graph->GetSchema()->GetGraphDisplayInformation(*Graph, DisplayInfo);
 
 				// Check if the name is unchanged
-				if( NewText.EqualTo( DisplayInfo.DisplayName ) )
+				if( NewText.EqualTo( DisplayInfo.PlainName ) )
 				{
 					return;
 				}
@@ -1107,7 +1112,7 @@ FText SBlueprintPaletteItem::GetToolTipText() const
 	if (PaletteAction.IsValid())
 	{
 		// Default tooltip is taken from the action
-		ToolTipText = (PaletteAction->TooltipDescription.Len() > 0) ? PaletteAction->TooltipDescription : PaletteAction->MenuDescription;
+		ToolTipText = (PaletteAction->TooltipDescription.Len() > 0) ? PaletteAction->TooltipDescription : PaletteAction->MenuDescription.ToString();
 
 		if(PaletteAction->GetTypeId() == FEdGraphSchemaAction_K2AddComponent::StaticGetTypeId())
 		{
@@ -1121,6 +1126,12 @@ FText SBlueprintPaletteItem::GetToolTipText() const
 		}
 		else if (UK2Node const* const NodeTemplate = FK2SchemaActionUtils::ExtractNodeTemplateFromAction(PaletteAction))
 		{
+			// Display the native title of the node when alt is held
+			if(FSlateApplication::Get().GetModifierKeys().IsAltDown())
+			{
+				return FText::FromString(NodeTemplate->GetNodeNativeTitle(ENodeTitleType::ListView));
+			}
+
 			// If the node wants to create tooltip text, use that instead, because its probably more detailed
 			FString NodeToolTipText = NodeTemplate->GetTooltip();
 			if (NodeToolTipText.Len() > 0)
@@ -1254,7 +1265,44 @@ TSharedPtr<SToolTip> SBlueprintPaletteItem::ConstructToolTipWidget() const
 	TAttribute<FText> TextAttribute;
 	TextAttribute.Bind(this, &SBlueprintPaletteItem::GetToolTipText);
 
-	return IDocumentation::Get()->CreateToolTip(TextAttribute, NULL, DocLink, DocExcerptName);
+	TSharedRef< SToolTip > TooltipWidget = IDocumentation::Get()->CreateToolTip(TextAttribute, NULL, DocLink, DocExcerptName);
+
+	// English speakers have no real need to know this exists.
+	if(FInternationalization::GetCurrentCulture()->GetTwoLetterISOLanguageName() != TEXT("en"))
+	{
+		if (UK2Node const* const NodeTemplate = FK2SchemaActionUtils::ExtractNodeTemplateFromAction(PaletteAction))
+		{
+			struct Local
+			{
+				static EVisibility GetNativeNodeNameVisibility()
+				{
+					return FSlateApplication::Get().GetModifierKeys().IsAltDown()? EVisibility::Collapsed : EVisibility::Visible;
+				}
+			};
+
+			return 
+				SNew(SToolTip)
+				[
+					SNew(SVerticalBox)
+					+SVerticalBox::Slot()
+					[
+						TooltipWidget->GetContent()
+					]
+					+SVerticalBox::Slot()
+						.AutoHeight()
+						.HAlign( HAlign_Right )
+					[
+
+						SNew( STextBlock )
+						.ColorAndOpacity( FSlateColor::UseSubduedForeground() )
+						.Text( LOCTEXT( "NativeNodeName", "hold (Alt) for native node name" ) )
+						.TextStyle( &FEditorStyle::GetWidgetStyle<FTextBlockStyle>(TEXT("Documentation.SDocumentationTooltip")) )
+						.Visibility_Static(&Local::GetNativeNodeNameVisibility)
+					]
+				];
+		}
+	}
+	return TooltipWidget;
 }
 
 /*******************************************************************************
diff --git a/Engine/Source/Editor/Kismet/Private/SGraphTitleBar.cpp b/Engine/Source/Editor/Kismet/Private/SGraphTitleBar.cpp
--- a/Engine/Source/Editor/Kismet/Private/SGraphTitleBar.cpp
+++ b/Engine/Source/Editor/Kismet/Private/SGraphTitleBar.cpp
@@ -34,7 +34,7 @@ FText SGraphTitleBar::GetTitleForOneCrumb(const UEdGraph* Graph)
 	Schema->GetGraphDisplayInformation(*Graph, /*out*/ DisplayInfo);
 
 	FFormatNamedArguments Args;
-	Args.Add(TEXT("BreadcrumbDisplayName"), DisplayInfo.DisplayName);
+	Args.Add(TEXT("BreadcrumbDisplayName"), DisplayInfo.PlainName);
 	Args.Add(TEXT("BreadcrumbNotes"), FText::FromString(DisplayInfo.GetNotesAsString()));
 	return FText::Format(LOCTEXT("BreadcrumbTitle", "{BreadcrumbDisplayName} {BreadcrumbNotes}"), Args);
 }
diff --git a/Engine/Source/Editor/Kismet/Private/SKismetInspector.cpp b/Engine/Source/Editor/Kismet/Private/SKismetInspector.cpp
--- a/Engine/Source/Editor/Kismet/Private/SKismetInspector.cpp
+++ b/Engine/Source/Editor/Kismet/Private/SKismetInspector.cpp
@@ -113,7 +113,7 @@ FString SKismetInspector::GetContextualEditingWidgetTitle() const
 
 			if (UEdGraphNode* Node = Cast<UEdGraphNode>(Object))
 			{
-				Title = Node->GetNodeTitle(ENodeTitleType::ListView);
+				Title = Node->GetNodeTitle(ENodeTitleType::ListView).ToString();
 			}
 			else if (USCS_Node* SCSNode = Cast<USCS_Node>(Object))
 			{
diff --git a/Engine/Source/Editor/Kismet/Private/SMyBlueprint.cpp b/Engine/Source/Editor/Kismet/Private/SMyBlueprint.cpp
--- a/Engine/Source/Editor/Kismet/Private/SMyBlueprint.cpp
+++ b/Engine/Source/Editor/Kismet/Private/SMyBlueprint.cpp
@@ -64,25 +64,31 @@ public:
 	virtual void HoverTargetChanged() OVERRIDE
 	{
 		const FSlateBrush* StatusSymbol = FEditorStyle::GetBrush(TEXT("NoBrush")); 
-		FString Message = DraggedCategory;
+		FText Message = FText::FromString(DraggedCategory);
 
 		if (HoveredCategoryName.Len() > 0)
 		{
+			FFormatNamedArguments Args;
+			Args.Add(TEXT("DraggedCategory"), FText::FromString(DraggedCategory));
+
 			if(HoveredCategoryName == DraggedCategory)
 			{
 				StatusSymbol = FEditorStyle::GetBrush(TEXT("Graph.ConnectorFeedback.Error"));
-				Message = FString::Printf( *LOCTEXT("MoveCatOverSelf", "Cannot insert category '%s' before itself.").ToString(), *DraggedCategory );
+
+				
+				Message = FText::Format( LOCTEXT("MoveCatOverSelf", "Cannot insert category '{DraggedCategory}' before itself."), Args );
 			}
 			else
 			{
 				StatusSymbol = FEditorStyle::GetBrush(TEXT("Graph.ConnectorFeedback.OK"));
-				Message = FString::Printf( *LOCTEXT("MoveCatOK", "Move category '%s' before '%s'").ToString(),  *DraggedCategory, *HoveredCategoryName );
+				Args.Add(TEXT("HoveredCategory"), FText::FromString(HoveredCategoryName));
+				Message = FText::Format( LOCTEXT("MoveCatOK", "Move category '{DraggedCategory}' before '{HoveredCategory}'"), Args );
 			}
 		}
 		else if (HoveredAction.IsValid())
 		{
 			StatusSymbol = FEditorStyle::GetBrush(TEXT("Graph.ConnectorFeedback.Error"));
-			Message = FString::Printf( *LOCTEXT("MoveCatOverAction", "Can only insert before another category.").ToString() );
+			Message = LOCTEXT("MoveCatOverAction", "Can only insert before another category.");
 		}
 
 		SetSimpleFeedbackMessage(StatusSymbol, FLinearColor::White, Message);
@@ -494,12 +500,12 @@ void SMyBlueprint::GetChildGraphs(UEdGraph* EdGraph, FGraphActionListBuilderBase
 		FGraphDisplayInfo ChildGraphDisplayInfo;
 		ChildSchema->GetGraphDisplayInformation(*Graph, ChildGraphDisplayInfo);
 
-		FString DisplayString = ChildGraphDisplayInfo.DisplayName.ToString();
-		const FName DisplayName =  FName(*DisplayString);
+		FText DisplayText = ChildGraphDisplayInfo.DisplayName;
 
 		FString Category = ((ParentCategory.IsEmpty()) ? "" : ParentCategory + "|") + EdGraphDisplayName.ToString();
-		FString FunctionTooltip = DisplayString;
-		FString FunctionDesc = DisplayString;
+		FString FunctionTooltip = DisplayText.ToString();
+		FText FunctionDesc = DisplayText;
+		const FName DisplayName =  FName(*DisplayText.ToString());
 
 		TSharedPtr<FEdGraphSchemaAction_K2Graph> NewFuncAction = MakeShareable(new FEdGraphSchemaAction_K2Graph(EEdGraphSchemaAction_K2Graph::Subgraph, Category, FunctionDesc, FunctionTooltip, 1));
 		NewFuncAction->FuncName = DisplayName;
@@ -533,10 +539,11 @@ void SMyBlueprint::GetChildEvents(UEdGraph const* EdGraph, int32 const SectionId
 		UK2Node_Event* const EventNode = (*It);
 
 		FString const Tooltip = EventNode->GetTooltip();
-		FString const Description = EventNode->GetNodeTitle(ENodeTitleType::EditableTitle);
+		FText const Description = EventNode->GetNodeTitle(ENodeTitleType::EditableTitle);
 
 		TSharedPtr<FEdGraphSchemaAction_K2Event> EventNodeAction = MakeShareable(new FEdGraphSchemaAction_K2Event(ActionCategory, Description, Tooltip, 0));
 		EventNodeAction->NodeTemplate = EventNode;
+		EventNodeAction->SearchTitle = EventNodeAction->NodeTemplate->GetNodeSearchTitle();
 		EventNodeAction->SectionID = SectionId;
 		OutAllActions.AddAction(EventNodeAction);
 	}
@@ -559,10 +566,11 @@ void SMyBlueprint::GetLocalVariables(FGraphActionListBuilderBase& OutAllActions)
 		for (UK2Node_LocalVariable* const LocalVariable : LocalVariableNodes)
 		{
 			FString const Tooltip = LocalVariable->GetTooltip();
-			FString const Description = LocalVariable->GetNodeTitle(ENodeTitleType::EditableTitle);
+			FText const Description = LocalVariable->GetNodeTitle(ENodeTitleType::EditableTitle);
 
 			TSharedPtr<FEdGraphSchemaAction_K2TargetNode> TargetNodeAction = MakeShareable(new FEdGraphSchemaAction_K2TargetNode(ActionCategory, Description, Tooltip, 0));
 			TargetNodeAction->NodeTemplate = LocalVariable;
+			TargetNodeAction->SearchTitle = TargetNodeAction->NodeTemplate->GetNodeSearchTitle();
 			OutAllActions.AddAction(TargetNodeAction);
 		}
 	}
@@ -618,9 +626,9 @@ void SMyBlueprint::CollectAllActions(FGraphActionListBuilderBase& OutAllActions)
 			continue;
 		}
 
-		const FString PropertyTooltip = Property->GetToolTipText().ToString();
+		const FText PropertyTooltip = Property->GetToolTipText();
 		const FName PropertyName = Property->GetFName();
-		const FString PropertyDesc = PropertyName.ToString();
+		const FText PropertyDesc = FText::FromName(PropertyName);
 
 		FName CategoryName = FObjectEditorUtils::GetCategoryFName(Property);
 		FString PropertyCategory = FObjectEditorUtils::GetCategory(Property);
@@ -632,7 +640,7 @@ void SMyBlueprint::CollectAllActions(FGraphActionListBuilderBase& OutAllActions)
 
 		if (bShouldShowAsVar)
 		{
-			TSharedPtr<FEdGraphSchemaAction_K2Var> NewVarAction = MakeShareable(new FEdGraphSchemaAction_K2Var(PropertyCategory, PropertyDesc, PropertyTooltip, 0));
+			TSharedPtr<FEdGraphSchemaAction_K2Var> NewVarAction = MakeShareable(new FEdGraphSchemaAction_K2Var(PropertyCategory, PropertyDesc, PropertyTooltip.ToString(), 0));
 			NewVarAction->SetVariableInfo(PropertyName, Blueprint->SkeletonGeneratedClass);
 			NewVarAction->SectionID = NodeSectionID::VARIABLE;
 			OutAllActions.AddAction(NewVarAction);
@@ -643,7 +651,7 @@ void SMyBlueprint::CollectAllActions(FGraphActionListBuilderBase& OutAllActions)
 			// Delegate is visible in MyBlueprint when not-native or its category name is not empty.
 			if (Property->HasAllPropertyFlags(CPF_Edit) || !PropertyCategory.IsEmpty())
 			{
-				NewFuncAction = MakeShareable(new FEdGraphSchemaAction_K2Delegate(PropertyCategory, PropertyDesc, PropertyTooltip, 0));
+				NewFuncAction = MakeShareable(new FEdGraphSchemaAction_K2Delegate(PropertyCategory, PropertyDesc, PropertyTooltip.ToString(), 0));
 				NewFuncAction->SetDelegateInfo(PropertyName, Blueprint->SkeletonGeneratedClass);
 				NewFuncAction->EdGraph = NULL;
 				NewFuncAction->SectionID = NodeSectionID::DELEGATE;				
@@ -674,7 +682,7 @@ void SMyBlueprint::CollectAllActions(FGraphActionListBuilderBase& OutAllActions)
 		FGraphDisplayInfo DisplayInfo;
 		Graph->GetSchema()->GetGraphDisplayInformation(*Graph, DisplayInfo);
 
-		TSharedPtr<FEdGraphSchemaAction_K2Graph> NewFuncAction = MakeShareable(new FEdGraphSchemaAction_K2Graph(EEdGraphSchemaAction_K2Graph::Function, FString(), DisplayInfo.DisplayName.ToString(), DisplayInfo.Tooltip, bIsConstructionScript ? 2 : 1));
+		TSharedPtr<FEdGraphSchemaAction_K2Graph> NewFuncAction = MakeShareable(new FEdGraphSchemaAction_K2Graph(EEdGraphSchemaAction_K2Graph::Function, FString(), DisplayInfo.PlainName, DisplayInfo.Tooltip, bIsConstructionScript ? 2 : 1));
 		NewFuncAction->FuncName = Graph->GetFName();
 		NewFuncAction->EdGraph = Graph;
 		NewFuncAction->SectionID = NodeSectionID::FUNCTION;
@@ -703,7 +711,7 @@ void SMyBlueprint::CollectAllActions(FGraphActionListBuilderBase& OutAllActions)
 		FGraphDisplayInfo DisplayInfo;
 		Graph->GetSchema()->GetGraphDisplayInformation(*Graph, DisplayInfo);
 
-		TSharedPtr<FEdGraphSchemaAction_K2Graph> NewFuncAction = MakeShareable(new FEdGraphSchemaAction_K2Graph(EEdGraphSchemaAction_K2Graph::Macro, FString(), DisplayInfo.DisplayName.ToString(), DisplayInfo.Tooltip, 1));
+		TSharedPtr<FEdGraphSchemaAction_K2Graph> NewFuncAction = MakeShareable(new FEdGraphSchemaAction_K2Graph(EEdGraphSchemaAction_K2Graph::Macro, FString(), DisplayInfo.PlainName, DisplayInfo.Tooltip, 1));
 		NewFuncAction->SectionID = NodeSectionID::MACRO;
 		NewFuncAction->FuncName = FunctionName;
 		NewFuncAction->EdGraph = Graph;
@@ -724,8 +732,11 @@ void SMyBlueprint::CollectAllActions(FGraphActionListBuilderBase& OutAllActions)
 		if (UEdGraphSchema_K2::CanKismetOverrideFunction(Function) && !ImplementedFunctions.Contains(FunctionName) && !UEdGraphSchema_K2::FunctionCanBePlacedAsEvent(Function))
 		{
 			FString FunctionTooltip = Function->GetToolTipText().ToString();
-			FString FunctionDesc = Function->GetMetaData(TEXT("FriendlyName"));
-			if (FunctionDesc.IsEmpty()) {FunctionDesc = Function->GetName();}
+			FText FunctionDesc = FText::FromString(Function->GetMetaData(TEXT("FriendlyName")));
+			if (FunctionDesc.IsEmpty())
+			{
+				FunctionDesc = FText::FromString(Function->GetName());
+			}
 
 			FString FunctionCategory = Function->GetMetaData(FBlueprintMetadata::MD_FunctionCategory);
 
@@ -748,7 +759,7 @@ void SMyBlueprint::CollectAllActions(FGraphActionListBuilderBase& OutAllActions)
 			const FName FunctionName = Graph->GetFName();
 			FString FunctionTooltip = FunctionName.ToString();
 			FString FunctionDesc = FunctionName.ToString();
-			TSharedPtr<FEdGraphSchemaAction_K2Graph> NewFuncAction = MakeShareable(new FEdGraphSchemaAction_K2Graph(EEdGraphSchemaAction_K2Graph::Interface, FString(), FunctionDesc, FunctionTooltip, 1));
+			TSharedPtr<FEdGraphSchemaAction_K2Graph> NewFuncAction = MakeShareable(new FEdGraphSchemaAction_K2Graph(EEdGraphSchemaAction_K2Graph::Interface, FString(), FText::FromString(FunctionDesc), FunctionTooltip, 1));
 			NewFuncAction->FuncName = FunctionName;
 			NewFuncAction->EdGraph = Graph;
 			NewFuncAction->SectionID = NodeSectionID::INTERFACE;
@@ -781,7 +792,7 @@ void SMyBlueprint::CollectAllActions(FGraphActionListBuilderBase& OutAllActions)
 
 						FString FunctionCategory = Function->GetMetaData(FBlueprintMetadata::MD_FunctionCategory);
 
-						TSharedPtr<FEdGraphSchemaAction_K2Graph> NewFuncAction = MakeShareable(new FEdGraphSchemaAction_K2Graph(EEdGraphSchemaAction_K2Graph::Function, FunctionCategory, FunctionDesc, FunctionTooltip, 1));
+						TSharedPtr<FEdGraphSchemaAction_K2Graph> NewFuncAction = MakeShareable(new FEdGraphSchemaAction_K2Graph(EEdGraphSchemaAction_K2Graph::Function, FunctionCategory, FText::FromString(FunctionDesc), FunctionTooltip, 1));
 						NewFuncAction->FuncName = FunctionName;
 						OutAllActions.AddAction(NewFuncAction);
 					}
@@ -799,7 +810,7 @@ void SMyBlueprint::CollectAllActions(FGraphActionListBuilderBase& OutAllActions)
 		FGraphDisplayInfo DisplayInfo;
 		Graph->GetSchema()->GetGraphDisplayInformation(*Graph, DisplayInfo);
 
-		TSharedPtr<FEdGraphSchemaAction_K2Graph> NewFuncAction = MakeShareable(new FEdGraphSchemaAction_K2Graph(EEdGraphSchemaAction_K2Graph::Graph, FString(), DisplayInfo.DisplayName.ToString(), DisplayInfo.Tooltip, 2));
+		TSharedPtr<FEdGraphSchemaAction_K2Graph> NewFuncAction = MakeShareable(new FEdGraphSchemaAction_K2Graph(EEdGraphSchemaAction_K2Graph::Graph, FString(), DisplayInfo.PlainName, DisplayInfo.Tooltip, 2));
 		NewFuncAction->FuncName = Graph->GetFName();
 		NewFuncAction->EdGraph = Graph;
 		NewFuncAction->SectionID = NodeSectionID::GRAPH;
@@ -818,7 +829,7 @@ void SMyBlueprint::CollectAllActions(FGraphActionListBuilderBase& OutAllActions)
 		const FName FunctionName(*(FString(TEXT("$INTERMEDIATE$_")) + Graph->GetName()));
 		FString FunctionTooltip = FunctionName.ToString();
 		FString FunctionDesc = FunctionName.ToString();
-		TSharedPtr<FEdGraphSchemaAction_K2Graph> NewFuncAction = MakeShareable(new FEdGraphSchemaAction_K2Graph(EEdGraphSchemaAction_K2Graph::Graph, FString(), FunctionDesc, FunctionTooltip, 1));
+		TSharedPtr<FEdGraphSchemaAction_K2Graph> NewFuncAction = MakeShareable(new FEdGraphSchemaAction_K2Graph(EEdGraphSchemaAction_K2Graph::Graph, FString(), FText::FromString(FunctionDesc), FunctionTooltip, 1));
 		NewFuncAction->FuncName = FunctionName;
 		NewFuncAction->EdGraph = Graph;
 		OutAllActions.AddAction(NewFuncAction);
@@ -885,9 +896,9 @@ void SMyBlueprint::CollectAllActions(FGraphActionListBuilderBase& OutAllActions)
 	{
 		if(UEnum* Enum = *EnumIter)
 		{
-			const FString EnumNameStr = Enum->GetName();
+			const FText EnumNameStr = FText::FromString(Enum->GetName());
 			TSharedPtr<FEdGraphSchemaAction_K2Enum> NewEnumAction = MakeShareable(new FEdGraphSchemaAction_K2Enum( 
-				UserDefinedEnumCategoryNameStr(), EnumNameStr, EnumNameStr, -1));
+				UserDefinedEnumCategoryNameStr(), EnumNameStr, EnumNameStr.ToString(), -1));
 			NewEnumAction->Enum = Enum;
 			NewEnumAction->SectionID = NodeSectionID::USER_ENUM;
 			OutAllActions.AddAction(NewEnumAction);
@@ -1042,7 +1053,7 @@ void SMyBlueprint::OnActionSelected( const TArray< TSharedPtr<FEdGraphSchemaActi
 			{
 				FGraphDisplayInfo DisplayInfo;
 				GraphAction->EdGraph->GetSchema()->GetGraphDisplayInformation(*GraphAction->EdGraph, DisplayInfo);
-				Inspector->ShowDetailsForSingleObject(GraphAction->EdGraph, SKismetInspector::FShowDetailsOptions(DisplayInfo.DisplayName.ToString()));
+				Inspector->ShowDetailsForSingleObject(GraphAction->EdGraph, SKismetInspector::FShowDetailsOptions(DisplayInfo.PlainName.ToString()));
 			}
 		}
 		else if (InAction->GetTypeId() == FEdGraphSchemaAction_K2Delegate::StaticGetTypeId())
@@ -1074,7 +1085,7 @@ void SMyBlueprint::OnActionSelected( const TArray< TSharedPtr<FEdGraphSchemaActi
 		else if (InAction->GetTypeId() == FEdGraphSchemaAction_K2TargetNode::StaticGetTypeId())
 		{
 			FEdGraphSchemaAction_K2TargetNode* TargetNodeAction = (FEdGraphSchemaAction_K2TargetNode*)InAction.Get();
-			SKismetInspector::FShowDetailsOptions Options(TargetNodeAction->NodeTemplate->GetNodeTitle(ENodeTitleType::EditableTitle));
+			SKismetInspector::FShowDetailsOptions Options(TargetNodeAction->NodeTemplate->GetNodeTitle(ENodeTitleType::EditableTitle).ToString());
 			Inspector->ShowDetailsForSingleObject(TargetNodeAction->NodeTemplate, Options);
 		}
 		else
@@ -1490,7 +1501,7 @@ void SMyBlueprint::OnFindEntry()
 	}
 	else if (FEdGraphSchemaAction_K2Event* EventAction = SelectionAsEvent())
 	{
-		SearchTerm = EventAction->MenuDescription;
+		SearchTerm = EventAction->MenuDescription.ToString();
 	}
 
 	if(!SearchTerm.IsEmpty())
diff --git a/Engine/Source/Editor/KismetWidgets/Private/SKismetLinearExpression.cpp b/Engine/Source/Editor/KismetWidgets/Private/SKismetLinearExpression.cpp
--- a/Engine/Source/Editor/KismetWidgets/Private/SKismetLinearExpression.cpp
+++ b/Engine/Source/Editor/KismetWidgets/Private/SKismetLinearExpression.cpp
@@ -68,7 +68,10 @@ TSharedRef<SWidget> SKismetLinearExpression::MakeNodeWidget(const UEdGraphNode*
 	{
 		// The source node is impure or has multiple outputs, so cannot be directly part of this pure expression
 		// Instead show it as a special sort of variable get
-		const FString EffectiveVariableName = FString::Printf(TEXT("%s_%s"), *Node->GetNodeTitle(ENodeTitleType::ListView), FromPin->PinFriendlyName.IsEmpty() ? *FromPin->PinName : *FromPin->PinFriendlyName);
+		FFormatNamedArguments Args;
+		Args.Add(TEXT("NodeTitle"), Node->GetNodeTitle(ENodeTitleType::ListView));
+		Args.Add(TEXT("PinName"), FromPin->PinFriendlyName.IsEmpty() ? FText::FromString(FromPin->PinName) : FromPin->PinFriendlyName);
+		const FText EffectiveVariableName = FText::Format(LOCTEXT("NodeTitleWithPinName", "{NodeTitle}_{PinName}"), Args );
 
 		return SNew(SOverlay)
 			+ SOverlay::Slot()
diff --git a/Engine/Source/Editor/MaterialEditor/Private/FindInMaterial.cpp b/Engine/Source/Editor/MaterialEditor/Private/FindInMaterial.cpp
--- a/Engine/Source/Editor/MaterialEditor/Private/FindInMaterial.cpp
+++ b/Engine/Source/Editor/MaterialEditor/Private/FindInMaterial.cpp
@@ -263,7 +263,7 @@ void SFindInMaterial::MatchTokens(const TArray<FString> &Tokens)
 	{
 		UEdGraphNode* Node = *It;
 
-		const FString NodeName = Node->GetNodeTitle(ENodeTitleType::ListView);
+		const FString NodeName = Node->GetNodeTitle(ENodeTitleType::ListView).ToString();
 		FSearchResult NodeResult(new FFindInMaterialResult(NodeName, RootSearchResult, Node));
 
 		FString NodeSearchString = NodeName + Node->NodeComment;
@@ -310,10 +310,10 @@ void SFindInMaterial::MatchTokens(const TArray<FString> &Tokens)
 		for (TArray<UEdGraphPin*>::TIterator PinIt(Node->Pins); PinIt; ++PinIt)
 		{
 			UEdGraphPin* Pin = *PinIt;
-			if (Pin && Pin->PinFriendlyName != TEXT(" "))
+			if (Pin && Pin->PinFriendlyName.CompareTo(FText::FromString(TEXT(" "))) != 0)
 			{
 				FString PinName = Pin->GetSchema()->GetPinDisplayName(Pin);
-				FString PinSearchString = Pin->PinName + Pin->PinFriendlyName + Pin->DefaultValue + Pin->PinType.PinCategory + Pin->PinType.PinSubCategory + (Pin->PinType.PinSubCategoryObject.IsValid() ? Pin->PinType.PinSubCategoryObject.Get()->GetFullName() : TEXT(""));
+				FString PinSearchString = Pin->PinName + Pin->PinFriendlyName.ToString() + Pin->DefaultValue + Pin->PinType.PinCategory + Pin->PinType.PinSubCategory + (Pin->PinType.PinSubCategoryObject.IsValid() ? Pin->PinType.PinSubCategoryObject.Get()->GetFullName() : TEXT(""));
 				PinSearchString = PinSearchString.Replace(TEXT(" "), TEXT(""));
 				if (StringMatchesSearchTokens(Tokens, PinSearchString))
 				{
diff --git a/Engine/Source/Editor/MaterialEditor/Private/MaterialEditorUtilities.cpp b/Engine/Source/Editor/MaterialEditor/Private/MaterialEditorUtilities.cpp
--- a/Engine/Source/Editor/MaterialEditor/Private/MaterialEditorUtilities.cpp
+++ b/Engine/Source/Editor/MaterialEditor/Private/MaterialEditorUtilities.cpp
@@ -621,7 +621,7 @@ void FMaterialEditorUtilities::AddMaterialExpressionCategory(FGraphActionMenuBui
 				const FText ToolTip = FText::Format( LOCTEXT( "NewMaterialExpressionTooltip", "Adds a {Name} node here" ), Arguments );
 				TSharedPtr<FMaterialGraphSchemaAction_NewNode> NewNodeAction(new FMaterialGraphSchemaAction_NewNode(
 					CategoryName,
-					MaterialExpression.Name,
+					FText::FromString(MaterialExpression.Name),
 					ToolTip.ToString(), 0));
 				ActionMenuBuilder.AddAction(NewNodeAction);
 				NewNodeAction->MaterialExpressionClass = MaterialExpression.MaterialClass;
diff --git a/Engine/Source/Editor/OnlineBlueprintSupport/Private/K2Node_LeaderboardFlush.cpp b/Engine/Source/Editor/OnlineBlueprintSupport/Private/K2Node_LeaderboardFlush.cpp
--- a/Engine/Source/Editor/OnlineBlueprintSupport/Private/K2Node_LeaderboardFlush.cpp
+++ b/Engine/Source/Editor/OnlineBlueprintSupport/Private/K2Node_LeaderboardFlush.cpp
@@ -4,6 +4,8 @@
 #include "LeaderboardFlushCallbackProxy.h"
 #include "K2Node_LeaderboardFlush.h"
 
+#define LOCTEXT_NAMESPACE "K2Node"
+
 UK2Node_LeaderboardFlush::UK2Node_LeaderboardFlush(const FPostConstructInitializeProperties& PCIP)
 	: Super(PCIP)
 {
@@ -15,15 +17,17 @@ UK2Node_LeaderboardFlush::UK2Node_LeaderboardFlush(const FPostConstructInitializ
 
 FString UK2Node_LeaderboardFlush::GetTooltip() const
 {
-	return TEXT("Flushes leaderboards for a session");
+	return LOCTEXT("K2Node_LeaderboardFlush_Tooltip", "Flushes leaderboards for a session").ToString();
 }
 
-FString UK2Node_LeaderboardFlush::GetNodeTitle(ENodeTitleType::Type TitleType) const
+FText UK2Node_LeaderboardFlush::GetNodeTitle(ENodeTitleType::Type TitleType) const
 {
-	return TEXT("Flush Leaderboards");
+	return LOCTEXT("FlushLeaderboards", "Flush Leaderboards");
 }
 
 FString UK2Node_LeaderboardFlush::GetCategoryName()
 {
 	return TEXT("Online|Leaderboard");
 }
+
+#undef LOCTEXT_NAMESPACE
diff --git a/Engine/Source/Editor/OnlineBlueprintSupport/Private/K2Node_LeaderboardQuery.cpp b/Engine/Source/Editor/OnlineBlueprintSupport/Private/K2Node_LeaderboardQuery.cpp
--- a/Engine/Source/Editor/OnlineBlueprintSupport/Private/K2Node_LeaderboardQuery.cpp
+++ b/Engine/Source/Editor/OnlineBlueprintSupport/Private/K2Node_LeaderboardQuery.cpp
@@ -4,6 +4,8 @@
 #include "LeaderboardQueryCallbackProxy.h"
 #include "K2Node_LeaderboardQuery.h"
 
+#define LOCTEXT_NAMESPACE "K2Node"
+
 UK2Node_LeaderboardQuery::UK2Node_LeaderboardQuery(const FPostConstructInitializeProperties& PCIP)
 	: Super(PCIP)
 {
@@ -15,15 +17,17 @@ UK2Node_LeaderboardQuery::UK2Node_LeaderboardQuery(const FPostConstructInitializ
 
 FString UK2Node_LeaderboardQuery::GetTooltip() const
 {
-	return TEXT("Queries a leaderboard for an integer value");
+	return LOCTEXT("K2Node_LeaderboardQuery_Tooltip", "Queries a leaderboard for an integer value").ToString();
 }
 
-FString UK2Node_LeaderboardQuery::GetNodeTitle(ENodeTitleType::Type TitleType) const
+FText UK2Node_LeaderboardQuery::GetNodeTitle(ENodeTitleType::Type TitleType) const
 {
-	return TEXT("Read Leaderboard Integer");
+	return LOCTEXT("ReadLeaderboardInteger", "Read Leaderboard Integer");
 }
 
 FString UK2Node_LeaderboardQuery::GetCategoryName()
 {
 	return TEXT("Online|Leaderboard");
 }
+
+#undef LOCTEXT_NAMESPACE
diff --git a/Engine/Source/Editor/ReferenceViewer/Classes/EdGraphNode_Reference.h b/Engine/Source/Editor/ReferenceViewer/Classes/EdGraphNode_Reference.h
--- a/Engine/Source/Editor/ReferenceViewer/Classes/EdGraphNode_Reference.h
+++ b/Engine/Source/Editor/ReferenceViewer/Classes/EdGraphNode_Reference.h
@@ -16,7 +16,7 @@ class UEdGraphNode_Reference : public UEdGraphNode
 	class UEdGraph_ReferenceViewer* GetReferenceViewerGraph() const;
 
 	// UEdGraphNode implementation
-	virtual FString GetNodeTitle(ENodeTitleType::Type TitleType) const OVERRIDE;
+	virtual FText GetNodeTitle(ENodeTitleType::Type TitleType) const OVERRIDE;
 	virtual void AllocateDefaultPins() OVERRIDE;
 	virtual UObject* GetJumpTargetForDoubleClick() const OVERRIDE;
 	// End UEdGraphNode implementation
@@ -36,7 +36,7 @@ protected:
 
 private:
 	TArray<FName> PackageNames;
-	FString NodeTitle;
+	FText NodeTitle;
 
 	bool bUsesThumbnail;
 	FAssetData CachedAssetData;
diff --git a/Engine/Source/Editor/ReferenceViewer/Private/EdGraphNode_Reference.cpp b/Engine/Source/Editor/ReferenceViewer/Private/EdGraphNode_Reference.cpp
--- a/Engine/Source/Editor/ReferenceViewer/Private/EdGraphNode_Reference.cpp
+++ b/Engine/Source/Editor/ReferenceViewer/Private/EdGraphNode_Reference.cpp
@@ -28,12 +28,12 @@ void UEdGraphNode_Reference::SetupReferenceNode(const FIntPoint& NodeLoc, const
 	if ( NewPackageNames.Num() == 1 )
 	{
 		NodeComment = NewPackageNames[0].ToString();
-		NodeTitle = ShortPackageName;
+		NodeTitle = FText::FromString(ShortPackageName);
 	}
 	else
 	{
 		NodeComment = FText::Format(LOCTEXT("ReferenceNodeMultiplePackagesTitle", "{0} nodes"), FText::AsNumber(NewPackageNames.Num())).ToString();
-		NodeTitle = FText::Format(LOCTEXT("ReferenceNodeMultiplePackagesComment", "{0} and {1} others"), FText::FromString(ShortPackageName), FText::AsNumber(NewPackageNames.Num())).ToString();
+		NodeTitle = FText::Format(LOCTEXT("ReferenceNodeMultiplePackagesComment", "{0} and {1} others"), FText::FromString(ShortPackageName), FText::AsNumber(NewPackageNames.Num()));
 	}
 	
 	CacheAssetData(InAssetData);
@@ -48,7 +48,7 @@ void UEdGraphNode_Reference::SetReferenceNodeCollapsed(const FIntPoint& NodeLoc,
 	PackageNames.Empty();
 	NodeComment = FText::Format(LOCTEXT("ReferenceNodeCollapsedMessage", "{0} other nodes"), FText::AsNumber(InNumReferencesExceedingMax)).ToString();
 
-	NodeTitle = LOCTEXT("ReferenceNodeCollapsedTitle", "Collapsed nodes").ToString();
+	NodeTitle = LOCTEXT("ReferenceNodeCollapsedTitle", "Collapsed nodes");
 	CacheAssetData(FAssetData());
 	AllocateDefaultPins();
 }
@@ -80,7 +80,7 @@ UEdGraph_ReferenceViewer* UEdGraphNode_Reference::GetReferenceViewerGraph() cons
 	return Cast<UEdGraph_ReferenceViewer>( GetGraph() );
 }
 
-FString UEdGraphNode_Reference::GetNodeTitle(ENodeTitleType::Type TitleType) const
+FText UEdGraphNode_Reference::GetNodeTitle(ENodeTitleType::Type TitleType) const
 {
 	return NodeTitle;
 }
diff --git a/Engine/Source/Editor/SoundClassEditor/Private/SSoundClassActionMenu.cpp b/Engine/Source/Editor/SoundClassEditor/Private/SSoundClassActionMenu.cpp
--- a/Engine/Source/Editor/SoundClassEditor/Private/SSoundClassActionMenu.cpp
+++ b/Engine/Source/Editor/SoundClassEditor/Private/SSoundClassActionMenu.cpp
@@ -40,7 +40,7 @@ void SSoundClassActionMenuItem::Construct(const FArguments& InArgs, TSharedPtr<F
 	}
 	else
 	{
-		TSharedRef<SWidget> NewSoundClassWidget = CreateNewSoundClassWidget( InAction->MenuDescription, InAction->TooltipDescription, FSlateFontInfo( FPaths::EngineContentDir() / TEXT("Slate/Fonts/Roboto-Regular.ttf"), 9), InAction );
+		TSharedRef<SWidget> NewSoundClassWidget = CreateNewSoundClassWidget( InAction->MenuDescription.ToString(), InAction->TooltipDescription, FSlateFontInfo( FPaths::EngineContentDir() / TEXT("Slate/Fonts/Roboto-Regular.ttf"), 9), InAction );
 		// Promote requires 2 'slots'
 		this->ChildSlot
 			[
diff --git a/Engine/Source/Editor/UnrealEd/Classes/MaterialGraph/MaterialGraphSchema.h b/Engine/Source/Editor/UnrealEd/Classes/MaterialGraph/MaterialGraphSchema.h
--- a/Engine/Source/Editor/UnrealEd/Classes/MaterialGraph/MaterialGraphSchema.h
+++ b/Engine/Source/Editor/UnrealEd/Classes/MaterialGraph/MaterialGraphSchema.h
@@ -23,7 +23,7 @@ struct UNREALED_API FMaterialGraphSchemaAction_NewNode : public FEdGraphSchemaAc
 		, MaterialExpressionClass(NULL)
 	{}
 
-	FMaterialGraphSchemaAction_NewNode(const FString& InNodeCategory, const FString& InMenuDesc, const FString& InToolTip, const int32 InGrouping)
+	FMaterialGraphSchemaAction_NewNode(const FString& InNodeCategory, const FText& InMenuDesc, const FString& InToolTip, const int32 InGrouping)
 		: FEdGraphSchemaAction(InNodeCategory, InMenuDesc, InToolTip, InGrouping) 
 		, MaterialExpressionClass(NULL)
 	{}
@@ -59,7 +59,7 @@ struct UNREALED_API FMaterialGraphSchemaAction_NewFunctionCall : public FEdGraph
 		: FEdGraphSchemaAction()
 	{}
 
-	FMaterialGraphSchemaAction_NewFunctionCall(const FString& InNodeCategory, const FString& InMenuDesc, const FString& InToolTip, const int32 InGrouping)
+	FMaterialGraphSchemaAction_NewFunctionCall(const FString& InNodeCategory, const FText& InMenuDesc, const FString& InToolTip, const int32 InGrouping)
 		: FEdGraphSchemaAction(InNodeCategory, InMenuDesc, InToolTip, InGrouping)
 	{}
 
@@ -82,7 +82,7 @@ struct UNREALED_API FMaterialGraphSchemaAction_NewComment : public FEdGraphSchem
 		: FEdGraphSchemaAction()
 	{}
 
-	FMaterialGraphSchemaAction_NewComment(const FString& InNodeCategory, const FString& InMenuDesc, const FString& InToolTip, const int32 InGrouping)
+	FMaterialGraphSchemaAction_NewComment(const FString& InNodeCategory, const FText& InMenuDesc, const FString& InToolTip, const int32 InGrouping)
 		: FEdGraphSchemaAction(InNodeCategory, InMenuDesc, InToolTip, InGrouping)
 	{}
 
@@ -105,7 +105,7 @@ struct UNREALED_API FMaterialGraphSchemaAction_Paste : public FEdGraphSchemaActi
 		: FEdGraphSchemaAction()
 	{}
 
-	FMaterialGraphSchemaAction_Paste(const FString& InNodeCategory, const FString& InMenuDesc, const FString& InToolTip, const int32 InGrouping)
+	FMaterialGraphSchemaAction_Paste(const FString& InNodeCategory, const FText& InMenuDesc, const FString& InToolTip, const int32 InGrouping)
 		: FEdGraphSchemaAction(InNodeCategory, InMenuDesc, InToolTip, InGrouping)
 	{}
 
diff --git a/Engine/Source/Editor/UnrealEd/Classes/Niagara/EdGraphSchema_Niagara.h b/Engine/Source/Editor/UnrealEd/Classes/Niagara/EdGraphSchema_Niagara.h
--- a/Engine/Source/Editor/UnrealEd/Classes/Niagara/EdGraphSchema_Niagara.h
+++ b/Engine/Source/Editor/UnrealEd/Classes/Niagara/EdGraphSchema_Niagara.h
@@ -19,7 +19,7 @@ struct UNREALED_API FNiagaraSchemaAction_NewNode : public FEdGraphSchemaAction
 		, NodeTemplate(NULL)
 	{}
 
-	FNiagaraSchemaAction_NewNode(const FString& InNodeCategory, const FString& InMenuDesc, const FString& InToolTip, const int32 InGrouping)
+	FNiagaraSchemaAction_NewNode(const FString& InNodeCategory, const FText& InMenuDesc, const FString& InToolTip, const int32 InGrouping)
 		: FEdGraphSchemaAction(InNodeCategory, InMenuDesc, InToolTip, InGrouping) 
 		, NodeTemplate(NULL)
 	{}
diff --git a/Engine/Source/Editor/UnrealEd/Classes/SoundClassGraph/SoundClassGraphSchema.h b/Engine/Source/Editor/UnrealEd/Classes/SoundClassGraph/SoundClassGraphSchema.h
--- a/Engine/Source/Editor/UnrealEd/Classes/SoundClassGraph/SoundClassGraphSchema.h
+++ b/Engine/Source/Editor/UnrealEd/Classes/SoundClassGraph/SoundClassGraphSchema.h
@@ -17,7 +17,7 @@ struct UNREALED_API FSoundClassGraphSchemaAction_NewNode : public FEdGraphSchema
 		, NewSoundClassName(TEXT("ClassName"))
 	{}
 
-	FSoundClassGraphSchemaAction_NewNode(const FString& InNodeCategory, const FString& InMenuDesc, const FString& InToolTip, const int32 InGrouping)
+	FSoundClassGraphSchemaAction_NewNode(const FString& InNodeCategory, const FText& InMenuDesc, const FString& InToolTip, const int32 InGrouping)
 		: FEdGraphSchemaAction(InNodeCategory, InMenuDesc, InToolTip, InGrouping)
 		, NewSoundClassName(TEXT("ClassName"))
 	{}
diff --git a/Engine/Source/Editor/UnrealEd/Classes/SoundCueGraph/SoundCueGraphSchema.h b/Engine/Source/Editor/UnrealEd/Classes/SoundCueGraph/SoundCueGraphSchema.h
--- a/Engine/Source/Editor/UnrealEd/Classes/SoundCueGraph/SoundCueGraphSchema.h
+++ b/Engine/Source/Editor/UnrealEd/Classes/SoundCueGraph/SoundCueGraphSchema.h
@@ -22,7 +22,7 @@ struct UNREALED_API FSoundCueGraphSchemaAction_NewNode : public FEdGraphSchemaAc
 		, SoundNodeClass(NULL)
 	{}
 
-	FSoundCueGraphSchemaAction_NewNode(const FString& InNodeCategory, const FString& InMenuDesc, const FString& InToolTip, const int32 InGrouping)
+	FSoundCueGraphSchemaAction_NewNode(const FString& InNodeCategory, const FText& InMenuDesc, const FString& InToolTip, const int32 InGrouping)
 		: FEdGraphSchemaAction(InNodeCategory, InMenuDesc, InToolTip, InGrouping) 
 		, SoundNodeClass(NULL)
 	{}
@@ -46,7 +46,7 @@ struct UNREALED_API FSoundCueGraphSchemaAction_NewFromSelected : public FSoundCu
 		: FSoundCueGraphSchemaAction_NewNode()
 	{}
 
-	FSoundCueGraphSchemaAction_NewFromSelected(const FString& InNodeCategory, const FString& InMenuDesc, const FString& InToolTip, const int32 InGrouping)
+	FSoundCueGraphSchemaAction_NewFromSelected(const FString& InNodeCategory, const FText& InMenuDesc, const FString& InToolTip, const int32 InGrouping)
 		: FSoundCueGraphSchemaAction_NewNode(InNodeCategory, InMenuDesc, InToolTip, InGrouping) 
 	{}
 
@@ -65,7 +65,7 @@ struct UNREALED_API FSoundCueGraphSchemaAction_NewComment : public FEdGraphSchem
 		: FEdGraphSchemaAction()
 	{}
 
-	FSoundCueGraphSchemaAction_NewComment(const FString& InNodeCategory, const FString& InMenuDesc, const FString& InToolTip, const int32 InGrouping)
+	FSoundCueGraphSchemaAction_NewComment(const FString& InNodeCategory, const FText& InMenuDesc, const FString& InToolTip, const int32 InGrouping)
 		: FEdGraphSchemaAction(InNodeCategory, InMenuDesc, InToolTip, InGrouping)
 	{}
 
@@ -84,7 +84,7 @@ struct UNREALED_API FSoundCueGraphSchemaAction_Paste : public FEdGraphSchemaActi
 		: FEdGraphSchemaAction()
 	{}
 
-	FSoundCueGraphSchemaAction_Paste(const FString& InNodeCategory, const FString& InMenuDesc, const FString& InToolTip, const int32 InGrouping)
+	FSoundCueGraphSchemaAction_Paste(const FString& InNodeCategory, const FText& InMenuDesc, const FString& InToolTip, const int32 InGrouping)
 		: FEdGraphSchemaAction(InNodeCategory, InMenuDesc, InToolTip, InGrouping)
 	{}
 
diff --git a/Engine/Source/Editor/UnrealEd/Private/Kismet2/BlueprintEditorUtils.cpp b/Engine/Source/Editor/UnrealEd/Private/Kismet2/BlueprintEditorUtils.cpp
--- a/Engine/Source/Editor/UnrealEd/Private/Kismet2/BlueprintEditorUtils.cpp
+++ b/Engine/Source/Editor/UnrealEd/Private/Kismet2/BlueprintEditorUtils.cpp
@@ -3731,8 +3731,10 @@ void FBlueprintEditorUtils::ConformCallsToParentFunctions(UBlueprint* Blueprint)
 						if(TargetFunction->GetOwnerClass() != CallFunctionNode->FunctionReference.GetMemberParentClass(Blueprint->ParentClass))
 						{
 							// Emit something to the log to indicate that we're making a change
-							Blueprint->Message_Note(FString::Printf(*LOCTEXT("CallParentFunctionSignatureFixed_Note", "%s (%s) had an invalid function signature - it has now been fixed.").ToString(),
-								*CallFunctionNode->GetNodeTitle(ENodeTitleType::ListView), *CallFunctionNode->GetName()));
+							FFormatNamedArguments Args;
+							Args.Add(TEXT("NodeTitle"), CallFunctionNode->GetNodeTitle(ENodeTitleType::ListView));
+							Args.Add(TEXT("FunctionNodeName"), FText::FromString(CallFunctionNode->GetName()));
+							Blueprint->Message_Note(FText::Format(LOCTEXT("CallParentFunctionSignatureFixed_Note", "{NodeTitle} ({FunctionNodeName}) had an invalid function signature - it has now been fixed."), Args).ToString() );
 
 							// Redirect to the correct parent function
 							CallFunctionNode->SetFromFunction(TargetFunction);
@@ -3783,8 +3785,10 @@ void FBlueprintEditorUtils::ConformCallsToParentFunctions(UBlueprint* Blueprint)
 						}
 
 						// Emit something to the log to indicate that we're making a change
-						Blueprint->Message_Note(FString::Printf(*LOCTEXT("CallParentNodeRemoved_Note", "%s (%s) was not valid for this Blueprint - it has been removed.").ToString(),
-							*CallFunctionNode->GetNodeTitle(ENodeTitleType::ListView), *CallFunctionNode->GetName()));
+						FFormatNamedArguments Args;
+						Args.Add(TEXT("NodeTitle"), CallFunctionNode->GetNodeTitle(ENodeTitleType::ListView));
+						Args.Add(TEXT("FunctionNodeName"), FText::FromString(CallFunctionNode->GetName()));
+						Blueprint->Message_Note( FText::Format(LOCTEXT("CallParentNodeRemoved_Note", "{NodeTitle} ({FunctionNodeName}) was not valid for this Blueprint - it has been removed."), Args).ToString() );
 
 						// Destroy the existing parent function call node (this will also break pin links and remove it from the graph)
 						CallFunctionNode->DestroyNode();
@@ -3847,9 +3851,12 @@ void FBlueprintEditorUtils::ConformImplementedEvents(UBlueprint* Blueprint)
 							if(!Blueprint->GeneratedClass->IsChildOf(EventNode->EventSignatureClass)
 								&& !ImplementedInterfaceClasses.Contains(EventNode->EventSignatureClass))
 							{
+								FFormatNamedArguments Args;
+								Args.Add(TEXT("NodeTitle"), EventNode->GetNodeTitle(ENodeTitleType::ListView));
+								Args.Add(TEXT("EventNodeName"), FText::FromString(EventNode->GetName()));
+
 								// Emit something to the log to indicate that we've made a change
-								Blueprint->Message_Note(FString::Printf(*LOCTEXT("EventSignatureFixed_Note", "%s (%s) had an invalid function signature - it has now been fixed.").ToString(),
-									*EventNode->GetNodeTitle(ENodeTitleType::ListView), *EventNode->GetName()));
+								Blueprint->Message_Note( FText::Format(LOCTEXT("EventSignatureFixed_Note", "{NodeTitle ({EventNodeName}) had an invalid function signature - it has now been fixed."), Args).ToString() );
 
 								// Fix up the event signature
 								EventNode->EventSignatureClass = Blueprint->GeneratedClass;
@@ -3861,9 +3868,12 @@ void FBlueprintEditorUtils::ConformImplementedEvents(UBlueprint* Blueprint)
 							UEdGraphNode* CustomEventNode = CurrentGraph->GetSchema()->CreateSubstituteNode(EventNode, CurrentGraph, NULL);
 							if(CustomEventNode)
 							{
+								FFormatNamedArguments Args;
+								Args.Add(TEXT("NodeTitle"), EventNode->GetNodeTitle(ENodeTitleType::ListView));
+								Args.Add(TEXT("EventNodeName"), FText::FromString(EventNode->GetName()));
+
 								// Emit something to the log to indicate that we've made a change
-								Blueprint->Message_Note(FString::Printf(*LOCTEXT("EventNodeReplaced_Note", "%s (%s) was not valid for this Blueprint - it has been converted to a custom event.").ToString(),
-									*EventNode->GetNodeTitle(ENodeTitleType::ListView), *EventNode->GetName()));
+								Blueprint->Message_Note(FText::Format(LOCTEXT("EventNodeReplaced_Note", "{NodeTitle} ({EventNodeName}) was not valid for this Blueprint - it has been converted to a custom event."), Args).ToString());
 
 								// Destroy the old event node (this will also break all pin links and remove it from the graph)
 								EventNode->DestroyNode();
diff --git a/Engine/Source/Editor/UnrealEd/Private/Kismet2/CompilerResultsLog.cpp b/Engine/Source/Editor/UnrealEd/Private/Kismet2/CompilerResultsLog.cpp
--- a/Engine/Source/Editor/UnrealEd/Private/Kismet2/CompilerResultsLog.cpp
+++ b/Engine/Source/Editor/UnrealEd/Private/Kismet2/CompilerResultsLog.cpp
@@ -137,11 +137,11 @@ void FCompilerResultsLog::InternalLogMessage(const EMessageSeverity::Type& Sever
 					}
 					else if (Node != NULL)
 					{
-						ObjText = FText::FromString( Node->GetNodeTitle(ENodeTitleType::ListView) );
+						ObjText = Node->GetNodeTitle(ENodeTitleType::ListView);
 					}
 					else if (Pin != NULL)
 					{
-						ObjText = FText::FromString( Pin->PinFriendlyName.IsEmpty() ? Pin->PinName : Pin->PinFriendlyName );
+						ObjText = Pin->PinFriendlyName.IsEmpty() ? FText::FromString(Pin->PinName) : Pin->PinFriendlyName;
 					}
 					else
 					{
diff --git a/Engine/Source/Editor/UnrealEd/Private/Kismet2/Kismet2.cpp b/Engine/Source/Editor/UnrealEd/Private/Kismet2/Kismet2.cpp
--- a/Engine/Source/Editor/UnrealEd/Private/Kismet2/Kismet2.cpp
+++ b/Engine/Source/Editor/UnrealEd/Private/Kismet2/Kismet2.cpp
@@ -1361,7 +1361,7 @@ inline void AddSearchMetaInfo( TArray<UObject::FAssetRegistryTag> &OutTags, cons
 		{
 			continue;
 		}
-		Result += FString("::") + Node->GetNodeTitle(ENodeTitleType::ListView);
+		Result += FString("::") + Node->GetNodeTitle(ENodeTitleType::ListView).ToString();
 
 		if(int* Count = AllPaths.Find(Result))
 		{
diff --git a/Engine/Source/Editor/UnrealEd/Private/Kismet2/KismetDebugUtilities.cpp b/Engine/Source/Editor/UnrealEd/Private/Kismet2/KismetDebugUtilities.cpp
--- a/Engine/Source/Editor/UnrealEd/Private/Kismet2/KismetDebugUtilities.cpp
+++ b/Engine/Source/Editor/UnrealEd/Private/Kismet2/KismetDebugUtilities.cpp
@@ -106,9 +106,9 @@ void FKismetDebugUtilities::OnScriptException(const UObject* ActiveObject, const
 					{
 						TraceClassificationMessageLog = FTextToken::Create(FText::FromString(TEXT("from node")));
 
-						FString NodeTitle = BlueprintNode->GetNodeTitle(ENodeTitleType::ListView); // a more user friendly name
+						FText NodeTitle = BlueprintNode->GetNodeTitle(ENodeTitleType::ListView); // a more user friendly name
 						// link to the last executed node (the one throwing the exception, presumably)
-						TraceTargetMessageLog = FUObjectToken::Create(BlueprintNode, FText::FromString(NodeTitle))->OnMessageTokenActivated(FOnMessageTokenActivated::CreateStatic(&Local::OnMessageLogLinkActivated));
+						TraceTargetMessageLog = FUObjectToken::Create(BlueprintNode, NodeTitle)->OnMessageTokenActivated(FOnMessageTokenActivated::CreateStatic(&Local::OnMessageLogLinkActivated));
 					}
 				}
 #endif // WITH_EDITORONLY_DATA
diff --git a/Engine/Source/Editor/UnrealEd/Private/MaterialGraphNode.cpp b/Engine/Source/Editor/UnrealEd/Private/MaterialGraphNode.cpp
--- a/Engine/Source/Editor/UnrealEd/Private/MaterialGraphNode.cpp
+++ b/Engine/Source/Editor/UnrealEd/Private/MaterialGraphNode.cpp
@@ -119,11 +119,55 @@ bool UMaterialGraphNode::CanPasteHere(const UEdGraph* TargetGraph, const UEdGrap
 	return false;
 }
 
-FString UMaterialGraphNode::GetNodeTitle(ENodeTitleType::Type TitleType) const
+FText UMaterialGraphNode::GetNodeTitle(ENodeTitleType::Type TitleType) const
 {
 	TArray<FString> Captions;
 	MaterialExpression->GetCaption(Captions);
 
+	if (TitleType == ENodeTitleType::EditableTitle)
+	{
+		return FText::FromString(GetParameterName());
+	}
+	else if (TitleType == ENodeTitleType::ListView)
+	{
+		return FText::FromString(MaterialExpression->GetClass()->GetDescription());
+	}
+	else
+	{
+		// More useful to display multi line parameter captions in reverse order
+		// TODO: May have to choose order based on expression type if others need correct order
+		int32 CaptionIndex = Captions.Num() -1;
+
+		FTextBuilder NodeTitle;
+		NodeTitle.AppendLine(Captions[CaptionIndex]);
+
+		for (; CaptionIndex > 0; )
+		{
+			CaptionIndex--;
+			NodeTitle.AppendLine(Captions[CaptionIndex]);
+		}
+
+		if ( MaterialExpression->bShaderInputData && (MaterialExpression->bHidePreviewWindow || MaterialExpression->bCollapsed))
+		{
+			NodeTitle.AppendLine(LOCTEXT("InputData", "Input Data"));
+		}
+
+		if (bIsPreviewExpression)
+		{
+			NodeTitle.AppendLine(LOCTEXT("PreviewExpression", "\nPreviewing"));
+		}
+
+		return NodeTitle.ToText();
+	}
+}
+
+FString UMaterialGraphNode::GetNodeNativeTitle(ENodeTitleType::Type TitleType) const
+{
+	// Do not setup this function for localization, intentionally left unlocalized!
+	
+	TArray<FString> Captions;
+	MaterialExpression->GetCaption(Captions);
+
 	if (TitleType == ENodeTitleType::EditableTitle)
 	{
 		return GetParameterName();
@@ -153,7 +197,7 @@ FString UMaterialGraphNode::GetNodeTitle(ENodeTitleType::Type TitleType) const
 
 		if (bIsPreviewExpression)
 		{
-			NodeTitle += LOCTEXT("PreviewExpression", "\nPreviewing").ToString();
+			NodeTitle += TEXT("\nPreviewing");
 		}
 
 		return NodeTitle;
@@ -432,7 +476,7 @@ void UMaterialGraphNode::CreateInputPins()
 		{
 			// Makes sure pin has a name for lookup purposes but user will never see it
 			NewPin->PinName = CreateUniquePinName(TEXT("Input"));
-			NewPin->PinFriendlyName = TEXT(" ");
+			NewPin->PinFriendlyName = FText::FromString(TEXT(" "));
 		}
 	}
 }
@@ -482,7 +526,7 @@ void UMaterialGraphNode::CreateOutputPins()
 		{
 			// Makes sure pin has a name for lookup purposes but user will never see it
 			NewPin->PinName = CreateUniquePinName(TEXT("Output"));
-			NewPin->PinFriendlyName = TEXT(" ");
+			NewPin->PinFriendlyName = FText::FromString(TEXT(" "));
 		}
 	}
 }
diff --git a/Engine/Source/Editor/UnrealEd/Private/MaterialGraphNode_Root.cpp b/Engine/Source/Editor/UnrealEd/Private/MaterialGraphNode_Root.cpp
--- a/Engine/Source/Editor/UnrealEd/Private/MaterialGraphNode_Root.cpp
+++ b/Engine/Source/Editor/UnrealEd/Private/MaterialGraphNode_Root.cpp
@@ -18,9 +18,9 @@ UMaterialGraphNode_Root::UMaterialGraphNode_Root(const class FPostConstructIniti
 {
 }
 
-FString UMaterialGraphNode_Root::GetNodeTitle(ENodeTitleType::Type TitleType) const
+FText UMaterialGraphNode_Root::GetNodeTitle(ENodeTitleType::Type TitleType) const
 {
-	return FMaterialEditorUtilities::GetOriginalObjectName(this->GetGraph());
+	return FText::FromString(FMaterialEditorUtilities::GetOriginalObjectName(this->GetGraph()));
 }
 
 FLinearColor UMaterialGraphNode_Root::GetNodeTitleColor() const
diff --git a/Engine/Source/Editor/UnrealEd/Private/MaterialGraphSchema.cpp b/Engine/Source/Editor/UnrealEd/Private/MaterialGraphSchema.cpp
--- a/Engine/Source/Editor/UnrealEd/Private/MaterialGraphSchema.cpp
+++ b/Engine/Source/Editor/UnrealEd/Private/MaterialGraphSchema.cpp
@@ -163,7 +163,7 @@ void UMaterialGraphSchema::GetBreakLinkToSubMenuActions( class FMenuBuilder& Men
 	for(TArray<class UEdGraphPin*>::TConstIterator Links(InGraphPin->LinkedTo); Links; ++Links)
 	{
 		UEdGraphPin* Pin = *Links;
-		FString TitleString = Pin->GetOwningNode()->GetNodeTitle(ENodeTitleType::ListView);
+		FString TitleString = Pin->GetOwningNode()->GetNodeTitle(ENodeTitleType::ListView).ToString();
 		FText Title = FText::FromString( TitleString );
 		if ( Pin->PinName != TEXT("") )
 		{
@@ -318,7 +318,7 @@ void UMaterialGraphSchema::GetGraphContextActions(FGraphContextMenuBuilder& Cont
 	{
 		const FText PasteDesc = LOCTEXT("PasteDesc", "Paste Here");
 		const FText PasteToolTip = LOCTEXT("PasteToolTip", "Pastes copied items at this location.");
-		TSharedPtr<FMaterialGraphSchemaAction_Paste> PasteAction(new FMaterialGraphSchemaAction_Paste(TEXT(""), PasteDesc.ToString(), PasteToolTip.ToString(), 0));
+		TSharedPtr<FMaterialGraphSchemaAction_Paste> PasteAction(new FMaterialGraphSchemaAction_Paste(TEXT(""), PasteDesc, PasteToolTip.ToString(), 0));
 		ContextMenuBuilder.AddAction(PasteAction);
 	}
 }
@@ -729,7 +729,7 @@ void UMaterialGraphSchema::GetMaterialFunctionActions(FGraphActionMenuBuilder& A
 					const FString& CategoryName = LibraryCategories[CategoryIndex];
 					TSharedPtr<FMaterialGraphSchemaAction_NewFunctionCall> NewFunctionAction(new FMaterialGraphSchemaAction_NewFunctionCall(
 						CategoryName,
-						FunctionName,
+						FText::FromString(FunctionName),
 						Description, 0));
 					ActionMenuBuilder.AddAction(NewFunctionAction);
 					NewFunctionAction->FunctionPath = FunctionPathName;
@@ -748,7 +748,7 @@ void UMaterialGraphSchema::GetCommentAction(FGraphActionMenuBuilder& ActionMenuB
 		const FText MultiCommentDesc = LOCTEXT("MultiCommentDesc", "Create Comment from Selection");
 		const FText CommentToolTip = LOCTEXT("CommentToolTip", "Creates a comment.");
 		const FText MenuDescription = bIsManyNodesSelected ? MultiCommentDesc : CommentDesc;
-		TSharedPtr<FMaterialGraphSchemaAction_NewComment> NewAction(new FMaterialGraphSchemaAction_NewComment(TEXT(""), MenuDescription.ToString(), CommentToolTip.ToString(), 0));
+		TSharedPtr<FMaterialGraphSchemaAction_NewComment> NewAction(new FMaterialGraphSchemaAction_NewComment(TEXT(""), MenuDescription, CommentToolTip.ToString(), 0));
 		ActionMenuBuilder.AddAction( NewAction );
 	}
 }
diff --git a/Engine/Source/Editor/UnrealEd/Private/Niagara/EdGraphSchema_Niagara.cpp b/Engine/Source/Editor/UnrealEd/Private/Niagara/EdGraphSchema_Niagara.cpp
--- a/Engine/Source/Editor/UnrealEd/Private/Niagara/EdGraphSchema_Niagara.cpp
+++ b/Engine/Source/Editor/UnrealEd/Private/Niagara/EdGraphSchema_Niagara.cpp
@@ -97,7 +97,7 @@ UEdGraphSchema_Niagara::UEdGraphSchema_Niagara(const class FPostConstructInitial
 	PC_Float = TEXT("float");
 }
 
-TSharedPtr<FNiagaraSchemaAction_NewNode> AddNewNodeAction(FGraphContextMenuBuilder& ContextMenuBuilder, const FString& Category, const FString& MenuDesc, const FString& Tooltip)
+TSharedPtr<FNiagaraSchemaAction_NewNode> AddNewNodeAction(FGraphContextMenuBuilder& ContextMenuBuilder, const FString& Category, const FText& MenuDesc, const FString& Tooltip)
 {
 	TSharedPtr<FNiagaraSchemaAction_NewNode> NewAction = TSharedPtr<FNiagaraSchemaAction_NewNode>(new FNiagaraSchemaAction_NewNode(Category, MenuDesc, Tooltip, 0));
 	ContextMenuBuilder.AddAction( NewAction );
@@ -109,19 +109,22 @@ void UEdGraphSchema_Niagara::GetGraphContextActions(FGraphContextMenuBuilder& Co
 	const UNiagaraGraph* NiagaraGraph = CastChecked<UNiagaraGraph>(ContextMenuBuilder.CurrentGraph);
 	UNiagaraScriptSource* Source = NiagaraGraph->GetSource();
 
-	// First get attributes.
-	FString GetString(TEXT("Get "));
-
 	TArray<FName> InputAttributeNames;
 	Source->GetUpdateInputs(InputAttributeNames);
 	for(int32 i=0; i<InputAttributeNames.Num(); i++)
 	{
 		const FName AttrName = InputAttributeNames[i];
-		TSharedPtr<FNiagaraSchemaAction_NewNode> GetAttrAction = AddNewNodeAction(ContextMenuBuilder, TEXT("Get Attribute"), GetString + AttrName.ToString(), TEXT(""));
+
+		FFormatNamedArguments Args;
+		Args.Add(TEXT("Attribute"), FText::FromName(AttrName));
+		const FText MenuDesc = FText::Format(NSLOCTEXT("Niagara", "GetAttribute", "Get {Attribute}"), Args);
+
+		TSharedPtr<FNiagaraSchemaAction_NewNode> GetAttrAction = AddNewNodeAction(ContextMenuBuilder, TEXT("Get Attribute"), MenuDesc, TEXT(""));
 
 		UNiagaraNodeGetAttr* GetAttrNode = NewObject<UNiagaraNodeGetAttr>(ContextMenuBuilder.OwnerOfTemporaries);
 		GetAttrNode->AttrName = AttrName;
 		GetAttrAction->NodeTemplate = GetAttrNode;
+		GetAttrAction->SearchTitle = GetAttrAction->NodeTemplate->GetNodeSearchTitle();
 	}
 
 	// Then get ops.
@@ -131,11 +134,12 @@ void UEdGraphSchema_Niagara::GetGraphContextActions(FGraphContextMenuBuilder& Co
 		VectorVM::FVectorVMOpInfo const& Info = VectorVM::GetOpCodeInfo(OpIdx);
 		if(Info.IsImplemented())
 		{
-			TSharedPtr<FNiagaraSchemaAction_NewNode> AddOpAction = AddNewNodeAction(ContextMenuBuilder, TEXT("Add Operation"), Info.FriendlyName, TEXT(""));
+			TSharedPtr<FNiagaraSchemaAction_NewNode> AddOpAction = AddNewNodeAction(ContextMenuBuilder, TEXT("Add Operation"), FText::FromString(Info.FriendlyName), TEXT(""));
 
 			UNiagaraNodeOp* OpNode = NewObject<UNiagaraNodeOp>(ContextMenuBuilder.OwnerOfTemporaries);
 			OpNode->OpIndex = OpIdx;
 			AddOpAction->NodeTemplate = OpNode;
+			AddOpAction->SearchTitle = AddOpAction->NodeTemplate->GetNodeSearchTitle();
 		}
 	}
 }
diff --git a/Engine/Source/Editor/UnrealEd/Private/Niagara/NiagaraNodeGetAttr.cpp b/Engine/Source/Editor/UnrealEd/Private/Niagara/NiagaraNodeGetAttr.cpp
--- a/Engine/Source/Editor/UnrealEd/Private/Niagara/NiagaraNodeGetAttr.cpp
+++ b/Engine/Source/Editor/UnrealEd/Private/Niagara/NiagaraNodeGetAttr.cpp
@@ -14,8 +14,17 @@ void UNiagaraNodeGetAttr::AllocateDefaultPins()
 	CreatePin(EGPD_Output, Schema->PC_Float, TEXT(""), NULL, false, false, AttrName.ToString());
 }
 
-FString UNiagaraNodeGetAttr::GetNodeTitle(ENodeTitleType::Type TitleType) const
+FText UNiagaraNodeGetAttr::GetNodeTitle(ENodeTitleType::Type TitleType) const
 {
+	FFormatNamedArguments Args;
+	Args.Add(TEXT("Attribute"), FText::FromName(AttrName));
+	return FText::Format(NSLOCTEXT("Niagara", "GetAttribute", "Get {Attribute}"), Args);
+}
+
+FString UNiagaraNodeGetAttr::GetNodeNativeTitle(ENodeTitleType::Type TitleType) const
+{
+	// Do not setup this function for localization, intentionally left unlocalized!
+	
 	return FString::Printf(TEXT("Get %s"), *AttrName.ToString());
 }
 
diff --git a/Engine/Source/Editor/UnrealEd/Private/Niagara/NiagaraNodeOp.cpp b/Engine/Source/Editor/UnrealEd/Private/Niagara/NiagaraNodeOp.cpp
--- a/Engine/Source/Editor/UnrealEd/Private/Niagara/NiagaraNodeOp.cpp
+++ b/Engine/Source/Editor/UnrealEd/Private/Niagara/NiagaraNodeOp.cpp
@@ -36,10 +36,10 @@ void UNiagaraNodeOp::AllocateDefaultPins()
 	CreatePin(EGPD_Output, Schema->PC_Float, TEXT(""), NULL, false, false, TEXT("Result"));
 }
 
-FString UNiagaraNodeOp::GetNodeTitle(ENodeTitleType::Type TitleType) const
+FText UNiagaraNodeOp::GetNodeTitle(ENodeTitleType::Type TitleType) const
 {
 	VectorVM::FVectorVMOpInfo const& Info = VectorVM::GetOpCodeInfo(OpIndex);
-	return Info.FriendlyName;
+	return FText::FromString(Info.FriendlyName);
 }
 
 FLinearColor UNiagaraNodeOp::GetNodeTitleColor() const
diff --git a/Engine/Source/Editor/UnrealEd/Private/Niagara/NiagaraNodeOutputUpdate.cpp b/Engine/Source/Editor/UnrealEd/Private/Niagara/NiagaraNodeOutputUpdate.cpp
--- a/Engine/Source/Editor/UnrealEd/Private/Niagara/NiagaraNodeOutputUpdate.cpp
+++ b/Engine/Source/Editor/UnrealEd/Private/Niagara/NiagaraNodeOutputUpdate.cpp
@@ -25,9 +25,9 @@ void UNiagaraNodeOutputUpdate::AllocateDefaultPins()
 }
 
 
-FString UNiagaraNodeOutputUpdate::GetNodeTitle(ENodeTitleType::Type TitleType) const
+FText UNiagaraNodeOutputUpdate::GetNodeTitle(ENodeTitleType::Type TitleType) const
 {
-	return TEXT("Output");
+	return NSLOCTEXT("NiagaraNode", "Output", "Output");
 }
 
 FLinearColor UNiagaraNodeOutputUpdate::GetNodeTitleColor() const
diff --git a/Engine/Source/Editor/UnrealEd/Private/SoundClassGraphNode.cpp b/Engine/Source/Editor/UnrealEd/Private/SoundClassGraphNode.cpp
--- a/Engine/Source/Editor/UnrealEd/Private/SoundClassGraphNode.cpp
+++ b/Engine/Source/Editor/UnrealEd/Private/SoundClassGraphNode.cpp
@@ -93,11 +93,11 @@ bool USoundClassGraphNode::CanCreateUnderSpecifiedSchema(const UEdGraphSchema* S
 	return Schema->IsA(USoundClassGraphSchema::StaticClass());
 }
 
-FString USoundClassGraphNode::GetNodeTitle(ENodeTitleType::Type TitleType) const
+FText USoundClassGraphNode::GetNodeTitle(ENodeTitleType::Type TitleType) const
 {
 	if (SoundClass)
 	{
-		return SoundClass->GetName();
+		return FText::FromString(SoundClass->GetName());
 	}
 	else
 	{
diff --git a/Engine/Source/Editor/UnrealEd/Private/SoundClassGraphSchema.cpp b/Engine/Source/Editor/UnrealEd/Private/SoundClassGraphSchema.cpp
--- a/Engine/Source/Editor/UnrealEd/Private/SoundClassGraphSchema.cpp
+++ b/Engine/Source/Editor/UnrealEd/Private/SoundClassGraphSchema.cpp
@@ -40,7 +40,7 @@ void USoundClassGraphSchema::GetBreakLinkToSubMenuActions( class FMenuBuilder& M
 	for(TArray<class UEdGraphPin*>::TConstIterator Links(InGraphPin->LinkedTo); Links; ++Links)
 	{
 		UEdGraphPin* Pin = *Links;
-		FString TitleString = Pin->GetOwningNode()->GetNodeTitle(ENodeTitleType::ListView);
+		FString TitleString = Pin->GetOwningNode()->GetNodeTitle(ENodeTitleType::ListView).ToString();
 		FText Title = FText::FromString( TitleString );
 		if ( Pin->PinName != TEXT("") )
 		{
@@ -80,7 +80,7 @@ void USoundClassGraphSchema::GetGraphContextActions(FGraphContextMenuBuilder& Co
 	const FText Name = LOCTEXT("NewSoundClass", "New Sound Class");
 	const FText ToolTip = LOCTEXT("NewSoundClassTooltip", "Create a new sound class");
 	
-	TSharedPtr<FSoundClassGraphSchemaAction_NewNode> NewAction(new FSoundClassGraphSchemaAction_NewNode(TEXT(""), Name.ToString(), ToolTip.ToString(), 0));
+	TSharedPtr<FSoundClassGraphSchemaAction_NewNode> NewAction(new FSoundClassGraphSchemaAction_NewNode(TEXT(""), Name, ToolTip.ToString(), 0));
 
 	ContextMenuBuilder.AddAction( NewAction );
 }
diff --git a/Engine/Source/Editor/UnrealEd/Private/SoundCueGraphNode.cpp b/Engine/Source/Editor/UnrealEd/Private/SoundCueGraphNode.cpp
--- a/Engine/Source/Editor/UnrealEd/Private/SoundCueGraphNode.cpp
+++ b/Engine/Source/Editor/UnrealEd/Private/SoundCueGraphNode.cpp
@@ -44,7 +44,7 @@ void USoundCueGraphNode::PostLoad()
 			{
 				Pin->PinName = CreateUniquePinName(TEXT("Output"));
 			}
-			Pin->PinFriendlyName = TEXT(" ");
+			Pin->PinFriendlyName = FText::FromString(TEXT(" "));
 		}
 	}
 }
@@ -62,7 +62,7 @@ void USoundCueGraphNode::CreateInputPin()
 	{
 		// Makes sure pin has a name for lookup purposes but user will never see it
 		NewPin->PinName = CreateUniquePinName(TEXT("Input"));
-		NewPin->PinFriendlyName = TEXT(" ");
+		NewPin->PinFriendlyName = FText::FromString(TEXT(" "));
 	}
 }
 
@@ -112,7 +112,7 @@ void USoundCueGraphNode::RemoveInputPin(UEdGraphPin* InGraphPin)
 int32 USoundCueGraphNode::EstimateNodeWidth() const
 {
 	const int32 EstimatedCharWidth = 6;
-	FString NodeTitle = GetNodeTitle(ENodeTitleType::FullTitle);
+	FString NodeTitle = GetNodeTitle(ENodeTitleType::FullTitle).ToString();
 	UFont* Font = GetDefault<UEditorEngine>()->EditorFont;
 	int32 Result = NodeTitle.Len()*EstimatedCharWidth;
 
@@ -130,11 +130,11 @@ bool USoundCueGraphNode::CanAddInputPin() const
 	return SoundNode->ChildNodes.Num() < SoundNode->GetMaxChildNodes();
 }
 
-FString USoundCueGraphNode::GetNodeTitle(ENodeTitleType::Type TitleType) const
+FText USoundCueGraphNode::GetNodeTitle(ENodeTitleType::Type TitleType) const
 {
 	if (SoundNode)
 	{
-		return SoundNode->GetTitle();
+		return FText::FromString(SoundNode->GetTitle());
 	}
 	else
 	{
@@ -241,7 +241,7 @@ FString USoundCueGraphNode::GetTooltip() const
 	}
 	if (Tooltip.Len() == 0)
 	{
-		Tooltip = GetNodeTitle(ENodeTitleType::ListView);
+		Tooltip = GetNodeTitle(ENodeTitleType::ListView).ToString();
 	}
 	return Tooltip;
 }
diff --git a/Engine/Source/Editor/UnrealEd/Private/SoundCueGraphNode_Root.cpp b/Engine/Source/Editor/UnrealEd/Private/SoundCueGraphNode_Root.cpp
--- a/Engine/Source/Editor/UnrealEd/Private/SoundCueGraphNode_Root.cpp
+++ b/Engine/Source/Editor/UnrealEd/Private/SoundCueGraphNode_Root.cpp
@@ -22,9 +22,9 @@ FLinearColor USoundCueGraphNode_Root::GetNodeTitleColor() const
 	return GEditor->AccessEditorUserSettings().ResultNodeTitleColor;
 }
 
-FString USoundCueGraphNode_Root::GetNodeTitle(ENodeTitleType::Type TitleType) const
+FText USoundCueGraphNode_Root::GetNodeTitle(ENodeTitleType::Type TitleType) const
 {
-	return LOCTEXT("RootTitle", "Output").ToString();
+	return LOCTEXT("RootTitle", "Output");
 }
 
 FString USoundCueGraphNode_Root::GetTooltip() const
diff --git a/Engine/Source/Editor/UnrealEd/Private/SoundCueGraphSchema.cpp b/Engine/Source/Editor/UnrealEd/Private/SoundCueGraphSchema.cpp
--- a/Engine/Source/Editor/UnrealEd/Private/SoundCueGraphSchema.cpp
+++ b/Engine/Source/Editor/UnrealEd/Private/SoundCueGraphSchema.cpp
@@ -245,7 +245,7 @@ void USoundCueGraphSchema::GetGraphContextActions(FGraphContextMenuBuilder& Cont
 
 	if (!ContextMenuBuilder.FromPin && FSoundCueEditorUtilities::CanPasteNodes(ContextMenuBuilder.CurrentGraph))
 	{
-		TSharedPtr<FSoundCueGraphSchemaAction_Paste> NewAction( new FSoundCueGraphSchemaAction_Paste(TEXT(""), LOCTEXT("PasteHereAction", "Paste here").ToString(), TEXT(""), 0) );
+		TSharedPtr<FSoundCueGraphSchemaAction_Paste> NewAction( new FSoundCueGraphSchemaAction_Paste(TEXT(""), LOCTEXT("PasteHereAction", "Paste here"), TEXT(""), 0) );
 		ContextMenuBuilder.AddAction( NewAction );
 	}
 }
@@ -442,7 +442,7 @@ void USoundCueGraphSchema::GetAllSoundNodeActions(FGraphActionMenuBuilder& Actio
 				FFormatNamedArguments Arguments;
 				Arguments.Add(TEXT("Name"), Name);
 				const FText AddToolTip = FText::Format(LOCTEXT("NewSoundCueNodeTooltip", "Adds {Name} node here"), Arguments);
-				TSharedPtr<FSoundCueGraphSchemaAction_NewNode> NewNodeAction(new FSoundCueGraphSchemaAction_NewNode(LOCTEXT("SoundNodeAction", "Sound Node").ToString(), Name.ToString(), AddToolTip.ToString(), 0));
+				TSharedPtr<FSoundCueGraphSchemaAction_NewNode> NewNodeAction(new FSoundCueGraphSchemaAction_NewNode(LOCTEXT("SoundNodeAction", "Sound Node").ToString(), Name, AddToolTip.ToString(), 0));
 				ActionMenuBuilder.AddAction(NewNodeAction);
 				NewNodeAction->SoundNodeClass = SoundNodeClass;
 			}
@@ -455,7 +455,7 @@ void USoundCueGraphSchema::GetAllSoundNodeActions(FGraphActionMenuBuilder& Actio
 				const FText MenuDesc = FText::Format(LOCTEXT("NewSoundNodeRandom", "{Name}: {SelectedItems}"), Arguments);
 				const FText ToolTip = FText::Format(LOCTEXT("NewSoundNodeRandomTooltip", "Adds a {Name} node for {SelectedItems} here"), Arguments);
 				TSharedPtr<FSoundCueGraphSchemaAction_NewFromSelected> NewNodeAction(new FSoundCueGraphSchemaAction_NewFromSelected(TEXT("From Selected"),
-					MenuDesc.ToString(),
+					MenuDesc,
 					ToolTip.ToString(), 0));
 				ActionMenuBuilder.AddAction(NewNodeAction);
 				NewNodeAction->SoundNodeClass = (SoundNodeClass == USoundNodeWavePlayer::StaticClass() ? NULL : SoundNodeClass);
@@ -469,7 +469,7 @@ void USoundCueGraphSchema::GetCommentAction(FGraphActionMenuBuilder& ActionMenuB
 	if (!ActionMenuBuilder.FromPin)
 	{
 		const bool bIsManyNodesSelected = CurrentGraph ? (FSoundCueEditorUtilities::GetNumberOfSelectedNodes(CurrentGraph) > 0) : false;
-		const FString MenuDescription = bIsManyNodesSelected ? LOCTEXT("CreateCommentAction", "Create Comment from Selection").ToString() : LOCTEXT("AddCommentAction", "Add Comment...").ToString();
+		const FText MenuDescription = bIsManyNodesSelected ? LOCTEXT("CreateCommentAction", "Create Comment from Selection") : LOCTEXT("AddCommentAction", "Add Comment...");
 		const FString ToolTip = LOCTEXT("CreateCommentToolTip", "Creates a comment.").ToString();
 
 		TSharedPtr<FSoundCueGraphSchemaAction_NewComment> NewAction(new FSoundCueGraphSchemaAction_NewComment(TEXT(""), MenuDescription, ToolTip, 0));
diff --git a/Engine/Source/Editor/UnrealEd/Private/Tests/EditorAutomationTests.cpp b/Engine/Source/Editor/UnrealEd/Private/Tests/EditorAutomationTests.cpp
--- a/Engine/Source/Editor/UnrealEd/Private/Tests/EditorAutomationTests.cpp
+++ b/Engine/Source/Editor/UnrealEd/Private/Tests/EditorAutomationTests.cpp
@@ -1823,7 +1823,7 @@ bool FTraceAllTimelinesAutomationTest::RunTest(const FString& BlueprintName)
 				const bool bIsWarningEvent = WarningEventNodes.Contains(Event);
 				const TCHAR SymbolString = bIsBadEvent ? TEXT('-') : (bIsWarningEvent ? TEXT('?') : TEXT('+'));
 				const TCHAR* CosmeticString = bIsCosmeticEvent ? TEXT("Client") : TEXT("Server");
-				FString OutputText = FString::Printf(TEXT("TraceTimeline:   %c %sEvent '%s' -> %s"), SymbolString, CosmeticString, *Event->GetNodeTitle(ENodeTitleType::EditableTitle), *TimelinePtr->TimelineName.ToString());
+				FString OutputText = FString::Printf(TEXT("TraceTimeline:   %c %sEvent '%s' -> %s"), SymbolString, CosmeticString, *Event->GetNodeTitle(ENodeTitleType::EditableTitle).ToString(), *TimelinePtr->TimelineName.ToString());
 
 				if (bIsBadEvent)
 				{
diff --git a/Engine/Source/Runtime/Engine/Classes/AI/EnvironmentQuery/EnvQueryGenerator.h b/Engine/Source/Runtime/Engine/Classes/AI/EnvironmentQuery/EnvQueryGenerator.h
--- a/Engine/Source/Runtime/Engine/Classes/AI/EnvironmentQuery/EnvQueryGenerator.h
+++ b/Engine/Source/Runtime/Engine/Classes/AI/EnvironmentQuery/EnvQueryGenerator.h
@@ -17,8 +17,8 @@ class ENGINE_API UEnvQueryGenerator : public UObject
 	FGenerateItemsSignature GenerateDelegate;
 
 	/** get description of generator */
-	virtual FString GetDescriptionTitle() const;
-	virtual FString GetDescriptionDetails() const;
+	virtual FText GetDescriptionTitle() const;
+	virtual FText GetDescriptionDetails() const;
 
 #if WITH_EDITOR
 	virtual void PostEditChangeProperty(struct FPropertyChangedEvent& PropertyChangedEvent) OVERRIDE;
diff --git a/Engine/Source/Runtime/Engine/Classes/AI/EnvironmentQuery/EnvQueryOption.h b/Engine/Source/Runtime/Engine/Classes/AI/EnvironmentQuery/EnvQueryOption.h
--- a/Engine/Source/Runtime/Engine/Classes/AI/EnvironmentQuery/EnvQueryOption.h
+++ b/Engine/Source/Runtime/Engine/Classes/AI/EnvironmentQuery/EnvQueryOption.h
@@ -14,6 +14,6 @@ class ENGINE_API UEnvQueryOption : public UObject
 	UPROPERTY()
 	TArray<class UEnvQueryTest*> Tests;
 
-	FString GetDescriptionTitle() const;
-	FString GetDescriptionDetails() const;
+	FText GetDescriptionTitle() const;
+	FText GetDescriptionDetails() const;
 };
diff --git a/Engine/Source/Runtime/Engine/Classes/AI/EnvironmentQuery/EnvQueryTest.h b/Engine/Source/Runtime/Engine/Classes/AI/EnvironmentQuery/EnvQueryTest.h
--- a/Engine/Source/Runtime/Engine/Classes/AI/EnvironmentQuery/EnvQueryTest.h
+++ b/Engine/Source/Runtime/Engine/Classes/AI/EnvironmentQuery/EnvQueryTest.h
@@ -68,10 +68,10 @@ class ENGINE_API UEnvQueryTest : public UObject
 
 	/** get description of test */
 	virtual FString GetDescriptionTitle() const;
-	virtual FString GetDescriptionDetails() const;
+	virtual FText GetDescriptionDetails() const;
 
-	FString DescribeFloatTestParams() const;
-	FString DescribeBoolTestParams(const FString& ConditionDesc) const;
+	FText DescribeFloatTestParams() const;
+	FText DescribeBoolTestParams(const FString& ConditionDesc) const;
 
 #if WITH_EDITOR
 	virtual void PostEditChangeProperty(struct FPropertyChangedEvent& PropertyChangedEvent) OVERRIDE;
diff --git a/Engine/Source/Runtime/Engine/Classes/AI/EnvironmentQuery/EnvQueryTypes.h b/Engine/Source/Runtime/Engine/Classes/AI/EnvironmentQuery/EnvQueryTypes.h
--- a/Engine/Source/Runtime/Engine/Classes/AI/EnvironmentQuery/EnvQueryTypes.h
+++ b/Engine/Source/Runtime/Engine/Classes/AI/EnvironmentQuery/EnvQueryTypes.h
@@ -212,7 +212,7 @@ struct ENGINE_API FEnvDirection
 	UPROPERTY(EditDefaultsOnly, Category=Direction, meta=(DisplayName="Mode"))
 	TEnumAsByte<EEnvDirection::Type> DirMode;
 
-	FString ToString() const;
+	FText ToText() const;
 };
 
 USTRUCT()
@@ -292,7 +292,7 @@ struct ENGINE_API FEnvTraceData
 	UPROPERTY(EditDefaultsOnly, Category=Trace)
 	uint32 bCanProjectDown : 1;
 
-	FString ToString(EDescriptionMode DescMode) const;
+	FText ToText(EDescriptionMode DescMode) const;
 
 	void SetGeometryOnly();
 	void SetNavmeshOnly();
@@ -756,9 +756,9 @@ class ENGINE_API UEnvQueryTypes : public UObject
 	GENERATED_UCLASS_BODY()
 
 	UEnvQueryTypes(const class FPostConstructInitializeProperties& PCIP) : Super(PCIP) {}
-	static FString GetShortTypeName(const UObject* Ob);
+	static FText GetShortTypeName(const UObject* Ob);
 
-	static FString DescribeContext(TSubclassOf<class UEnvQueryContext> ContextClass);
+	static FText DescribeContext(TSubclassOf<class UEnvQueryContext> ContextClass);
 	static FString DescribeIntParam(const FEnvIntParam& Param);
 	static FString DescribeFloatParam(const FEnvFloatParam& Param);
 	static FString DescribeBoolParam(const FEnvBoolParam& Param);
diff --git a/Engine/Source/Runtime/Engine/Classes/AI/EnvironmentQuery/Generators/EnvQueryGenerator_OnCircle.h b/Engine/Source/Runtime/Engine/Classes/AI/EnvironmentQuery/Generators/EnvQueryGenerator_OnCircle.h
--- a/Engine/Source/Runtime/Engine/Classes/AI/EnvironmentQuery/Generators/EnvQueryGenerator_OnCircle.h
+++ b/Engine/Source/Runtime/Engine/Classes/AI/EnvironmentQuery/Generators/EnvQueryGenerator_OnCircle.h
@@ -42,8 +42,8 @@ class UEnvQueryGenerator_OnCircle : public UEnvQueryGenerator_ProjectedPoints
 
 	void GenerateItems(struct FEnvQueryInstance& QueryInstance); 
 
-	virtual FString GetDescriptionTitle() const OVERRIDE;
-	virtual FString GetDescriptionDetails() const OVERRIDE;
+	virtual FText GetDescriptionTitle() const OVERRIDE;
+	virtual FText GetDescriptionDetails() const OVERRIDE;
 
 #if WITH_EDITOR
 	virtual void PostEditChangeProperty( struct FPropertyChangedEvent& PropertyChangedEvent) OVERRIDE;
diff --git a/Engine/Source/Runtime/Engine/Classes/AI/EnvironmentQuery/Generators/EnvQueryGenerator_PathingGrid.h b/Engine/Source/Runtime/Engine/Classes/AI/EnvironmentQuery/Generators/EnvQueryGenerator_PathingGrid.h
--- a/Engine/Source/Runtime/Engine/Classes/AI/EnvironmentQuery/Generators/EnvQueryGenerator_PathingGrid.h
+++ b/Engine/Source/Runtime/Engine/Classes/AI/EnvironmentQuery/Generators/EnvQueryGenerator_PathingGrid.h
@@ -39,8 +39,8 @@ class UEnvQueryGenerator_PathingGrid : public UEnvQueryGenerator
 
 	void GenerateItems(struct FEnvQueryInstance& QueryInstance); 
 
-	virtual FString GetDescriptionTitle() const OVERRIDE;
-	virtual FString GetDescriptionDetails() const OVERRIDE;
+	virtual FText GetDescriptionTitle() const OVERRIDE;
+	virtual FText GetDescriptionDetails() const OVERRIDE;
 
 protected:
 #if WITH_RECAST
diff --git a/Engine/Source/Runtime/Engine/Classes/AI/EnvironmentQuery/Tests/EnvQueryTest_Distance.h b/Engine/Source/Runtime/Engine/Classes/AI/EnvironmentQuery/Tests/EnvQueryTest_Distance.h
--- a/Engine/Source/Runtime/Engine/Classes/AI/EnvironmentQuery/Tests/EnvQueryTest_Distance.h
+++ b/Engine/Source/Runtime/Engine/Classes/AI/EnvironmentQuery/Tests/EnvQueryTest_Distance.h
@@ -30,5 +30,5 @@ class UEnvQueryTest_Distance : public UEnvQueryTest
 	void RunTest(struct FEnvQueryInstance& QueryInstance);
 
 	virtual FString GetDescriptionTitle() const OVERRIDE;
-	virtual FString GetDescriptionDetails() const OVERRIDE;
+	virtual FText GetDescriptionDetails() const OVERRIDE;
 };
diff --git a/Engine/Source/Runtime/Engine/Classes/AI/EnvironmentQuery/Tests/EnvQueryTest_Dot.h b/Engine/Source/Runtime/Engine/Classes/AI/EnvironmentQuery/Tests/EnvQueryTest_Dot.h
--- a/Engine/Source/Runtime/Engine/Classes/AI/EnvironmentQuery/Tests/EnvQueryTest_Dot.h
+++ b/Engine/Source/Runtime/Engine/Classes/AI/EnvironmentQuery/Tests/EnvQueryTest_Dot.h
@@ -19,7 +19,7 @@ class UEnvQueryTest_Dot : public UEnvQueryTest
 	void RunTest(struct FEnvQueryInstance& QueryInstance);
 
 	virtual FString GetDescriptionTitle() const OVERRIDE;
-	virtual FString GetDescriptionDetails() const OVERRIDE;
+	virtual FText GetDescriptionDetails() const OVERRIDE;
 
 protected:
 
diff --git a/Engine/Source/Runtime/Engine/Classes/EdGraph/EdGraphNode.h b/Engine/Source/Runtime/Engine/Classes/EdGraph/EdGraphNode.h
--- a/Engine/Source/Runtime/Engine/Classes/EdGraph/EdGraphNode.h
+++ b/Engine/Source/Runtime/Engine/Classes/EdGraph/EdGraphNode.h
@@ -240,7 +240,14 @@ public:
 	virtual void PostPasteNode() {}
 
 	/** Gets the name of this node, shown in title bar */
-	virtual FString GetNodeTitle(ENodeTitleType::Type TitleType) const;
+	virtual FText GetNodeTitle(ENodeTitleType::Type TitleType) const;
+	
+	/** Gets the name of this node in the native language. This should always be overridden when a node's title is built from concatenation
+		By default this will return the source string from GetNodeTitle, which is not guaranteed to be correct */
+	virtual FString GetNodeNativeTitle(ENodeTitleType::Type TitleType) const;
+
+	/** Gets the searchable metadata of this node */
+	virtual FText GetNodeSearchTitle() const;
 
 	/** 
 	 * Gets the draw color of a node's title bar
diff --git a/Engine/Source/Runtime/Engine/Classes/EdGraph/EdGraphNode_Comment.h b/Engine/Source/Runtime/Engine/Classes/EdGraph/EdGraphNode_Comment.h
--- a/Engine/Source/Runtime/Engine/Classes/EdGraph/EdGraphNode_Comment.h
+++ b/Engine/Source/Runtime/Engine/Classes/EdGraph/EdGraphNode_Comment.h
@@ -46,7 +46,7 @@ public:
 	virtual void AllocateDefaultPins() OVERRIDE {}
 	ENGINE_API virtual FString GetTooltip() const OVERRIDE;
 	ENGINE_API virtual FLinearColor GetNodeCommentColor() const OVERRIDE;
-	ENGINE_API virtual FString GetNodeTitle(ENodeTitleType::Type TitleType) const OVERRIDE;
+	ENGINE_API virtual FText GetNodeTitle(ENodeTitleType::Type TitleType) const OVERRIDE;
 	virtual bool ShouldOverridePinNames() const OVERRIDE { return true; }
 	ENGINE_API virtual FString GetPinNameOverride(const UEdGraphPin& Pin) const OVERRIDE;
 	ENGINE_API virtual void ResizeNode(const FVector2D& NewSize) OVERRIDE;
diff --git a/Engine/Source/Runtime/Engine/Classes/EdGraph/EdGraphPin.h b/Engine/Source/Runtime/Engine/Classes/EdGraph/EdGraphPin.h
--- a/Engine/Source/Runtime/Engine/Classes/EdGraph/EdGraphPin.h
+++ b/Engine/Source/Runtime/Engine/Classes/EdGraph/EdGraphPin.h
@@ -123,7 +123,7 @@ class UEdGraphPin : public UObject
 
 	/** Used as the display name if set */
 	UPROPERTY()
-	FString PinFriendlyName;
+	FText PinFriendlyName;
 
 	/** The tool-tip describing this pin's purpose */
 	UPROPERTY()
@@ -234,7 +234,7 @@ public:
 		PinType.ResetToDefaults();
 
 		PinName.Empty();
-		PinFriendlyName.Empty();
+		PinFriendlyName = FText::GetEmpty();
 		DefaultValue.Empty();
 		AutogeneratedDefaultValue.Empty();
 		DefaultObject = NULL;
