#summary Creating an external graph will now present the user with an SDlgPickAssetPath window to pick the path and filename instead of forcing them to type out a full and valid path.

#remove Removed SCurveAssetWidget
#remove Removed STimelineEdTrack::SaveCurveAssetPath, GetSaveCurvePath, SetSaveCurvePath, OnCreateButtonClicked, CreateCurveAssetWidget to streamline the code.
#add Added CreateUniqueCurveAssetPathname to generate a pathname to a new external track.

#change Picking the path now uses an SDlgPickAssetPath window instead of SCurveAssetWidget.

#ttp 324906 - k2: timelines: create an external curve doesn't make it easy to specify a target location (images)

diff --git a/Engine/Source/Editor/Kismet/Private/STimelineEditor.cpp b/Engine/Source/Editor/Kismet/Private/STimelineEditor.cpp
--- a/Engine/Source/Editor/Kismet/Private/STimelineEditor.cpp
+++ b/Engine/Source/Editor/Kismet/Private/STimelineEditor.cpp
@@ -11,6 +11,8 @@
 #include "STimelineEditor.h"
 #include "PackageTools.h"
 #include "AssetRegistryModule.h"
+#include "AssetToolsModule.h"
+#include "DlgPickAssetPath.h"
 
 #define LOCTEXT_NAMESPACE "STimelineEditor"
 
@@ -37,88 +39,6 @@ namespace TimelineEditorHelpers
 		return NAME_None;
 	}
 }
-//////////////////////////////////////////////////////////////////////////
-// SCurveAssetWidget Implementation
-
-void SCurveAssetWidget::Construct(const FArguments& InArgs)
-{
-	VisibleText = InArgs._VisibleText;
-	OnCreateCurve = InArgs._OnCreateCurve;
-	OnTextChanged = InArgs._OnTextChanged;
-
-	this->ChildSlot
-		.HAlign(HAlign_Center)
-		.VAlign(VAlign_Center)
-		[
-			SNew(SVerticalBox)
-			+SVerticalBox::Slot()
-			.AutoHeight()
-			[
-				//Background image
-				SNew(SBorder)
-				.BorderImage( FEditorStyle::GetBrush( "Graph.Node.Body" ) )
-				.Padding(0)
-				[
-					SNew(SVerticalBox)
-					+SVerticalBox::Slot()
-						.AutoHeight()
-						.Padding(FMargin(4,4,4,4))
-						[
-							//Path label text
-							SNew(SHorizontalBox)
-							+SHorizontalBox::Slot()
-								.AutoWidth()
-								.HAlign(HAlign_Left)
-								.VAlign(VAlign_Center)
-								.Padding(4)
-								[
-									SNew(STextBlock)
-									.Text( LOCTEXT( "AssetPath", "Path:" ) )
-								]
-
-							//Editable text
-							+SHorizontalBox::Slot()
-								.FillWidth(1.0f)
-								.Padding(4)
-								[
-									SAssignNew(WidgetWithDefaultFocus, SEditableTextBox)
-									.Text(FText::FromString( VisibleText.Get() ))
-									.MinDesiredWidth(400)
-									.OnTextChanged(this, &SCurveAssetWidget::OnEditableTextChanged)
-								]
-						]
-
-					+SVerticalBox::Slot()
-						.AutoHeight()
-						[
-							//Create button
-							SNew(SHorizontalBox)
-							+SHorizontalBox::Slot()
-							.FillWidth(1.0f)
-							.HAlign(HAlign_Right)
-							.Padding(4)
-							[
-								SNew(SButton)
-								.HAlign(HAlign_Center)
-								.Text( LOCTEXT( "CreateAssetButton", "Create External Curve" ) )
-								.OnClicked(this, &SCurveAssetWidget::OnCreateButtonClicked)
-							]
-						]
-				]
-			]
-		];
-}
-
-FReply SCurveAssetWidget::OnCreateButtonClicked( )
-{
-	OnCreateCurve.ExecuteIfBound();
-	return FReply::Handled();
-}
-
-void SCurveAssetWidget::OnEditableTextChanged( const FText& NewString )
-{
-	OnTextChanged.ExecuteIfBound(NewString);
-}
 
 //////////////////////////////////////////////////////////////////////////
 // STimelineEdTrack
@@ -328,50 +248,23 @@ void STimelineEdTrack::Construct(const FArguments& InArgs, TSharedPtr<FTimelineE
 }
 END_SLATE_FUNCTION_BUILD_OPTIMIZATION
 
-void STimelineEdTrack::InitCurveAssetPath()
+FString STimelineEdTrack::CreateUniqueCurveAssetPathName()
 {
 	//Default path
-	SaveCurveAssetPath = FString(TEXT( "/Game/Unsorted" ));
+	FString BasePath = FString(TEXT( "/Game/Unsorted" ));
 
 	TSharedRef<STimelineEditor> TimelineRef = TimelineEdPtr.Pin().ToSharedRef();
 
 	//Get curve name from editable text box
-	SaveCurveAssetPath += FString(TEXT("/")) + TimelineEditorHelpers::GetTrackNameFromTimeline(TimelineEdPtr.Pin()->GetTimeline(), Track) .ToString();
-}
-
-
-TSharedRef<SWidget> STimelineEdTrack::CreateCurveAssetWidget()
-{
-	//set default asset path
-	InitCurveAssetPath();
-
-	return ( SNew( SCurveAssetWidget )
-		.VisibleText(this, &STimelineEdTrack::GetSaveCurvePath)
-		.OnCreateCurve(this, &STimelineEdTrack::OnCreateButtonClicked )
-		.OnTextChanged( this, &STimelineEdTrack::SetSaveCurvePath ));
-}
+	FAssetToolsModule& AssetToolsModule = FModuleManager::LoadModuleChecked<FAssetToolsModule>("AssetTools");
 
-void STimelineEdTrack::OnCreateExternalCurve()
-{
-	TSharedRef<SWidget> CurveWidget = CreateCurveAssetWidget();
-
-	// Show dialog to create external curve
-	AssetCreationWindow = FSlateApplication::Get().PushMenu(
-		SharedThis( this ),
-		CurveWidget,
-		FSlateApplication::Get().GetCursorPos(),
-		FPopupTransitionEffect( FPopupTransitionEffect::TypeInPopup)
-		);
-}
+	// Create a unique asset name so the user can instantly hit OK if they want to create the new asset
+	FString AssetName = TimelineEditorHelpers::GetTrackNameFromTimeline(TimelineEdPtr.Pin()->GetTimeline(), Track).ToString();
+	FString PackageName;
+	BasePath = BasePath + TEXT("/") + AssetName;
+	AssetToolsModule.Get().CreateUniqueAssetName(BasePath, TEXT(""), PackageName, AssetName);
 
-void STimelineEdTrack::SetSaveCurvePath(const FText& NewStr)
-{
-	SaveCurveAssetPath = NewStr.ToString();
-}
-
-FString STimelineEdTrack::GetSaveCurvePath() const
-{
-	return SaveCurveAssetPath;
+	return PackageName;
 }
 
 void STimelineEdTrack::OnCloseCreateCurveWindow()
@@ -385,7 +278,7 @@ void STimelineEdTrack::OnCloseCreateCurveWindow()
 	}
 }
 
-void STimelineEdTrack::OnCreateButtonClicked()
+void STimelineEdTrack::OnCreateExternalCurve()
 {
 	UCurveBase* NewCurveAsset = CreateCurveAsset();
 	if( NewCurveAsset )
@@ -630,120 +523,95 @@ UCurveBase* STimelineEdTrack::CreateCurveAsset()
 
 	if( TrackWidget.IsValid() )
 	{
-		FString Package(SaveCurveAssetPath);
-		FString Name	= FPackageName::GetLongPackageAssetName(Package);
-		FString Group(TEXT(""));
-
-		FText Reason;
-		if (!FName(*Name).IsValidObjectName( Reason )
-			||	!FPackageName::IsValidLongPackageName( Package, false, &Reason ) )
-		{
-			FMessageDialog::Open( EAppMsgType::Ok, Reason );
-			return NULL;
-		}
-
-		// Find (or create!) the desired package for this object
-		UPackage* Pkg = CreatePackage(NULL,*SaveCurveAssetPath);
-		UPackage* OutermostPkg = Pkg->GetOutermost();
-
-		TArray<UPackage*> TopLevelPackages;
-		TopLevelPackages.Add( OutermostPkg );
-		if( !PackageTools::HandleFullyLoadingPackages( TopLevelPackages, NSLOCTEXT("UnrealEd", "CreateANewObject", "Create a new object") ) )
-		{
-			// User aborted.
-			return NULL;
-		}
-
-		// We need to test again after fully loading.
-		if (!FName(*Name).IsValidObjectName( Reason )
-			||	!FPackageName::IsValidLongPackageName( Package, false, &Reason ) )
-		{
-			FMessageDialog::Open( EAppMsgType::Ok, Reason );
-			return NULL;
-		}
+		TSharedRef<SDlgPickAssetPath> NewLayerDlg = 
+			SNew(SDlgPickAssetPath)
+			.Title(LOCTEXT("CreateExternalCurve", "Create External Curve"))
+			.DefaultAssetPath(FText::FromString(CreateUniqueCurveAssetPathName()));
 
-		if( false == PromptUserIfExistingObject(Name, Package, Group, Pkg))
+		if (NewLayerDlg->ShowModal() != EAppReturnType::Cancel)
 		{
-			return NULL;
-		}
+			FString PackageName = NewLayerDlg->GetFullAssetPath().ToString();
+			FName AssetName = FName(*NewLayerDlg->GetAssetName().ToString());
 
-		//Get the curve class type
-		TSubclassOf<UCurveBase> CurveType;
-		if( Track->TrackType == FTimelineEdTrack::TT_Event || Track->TrackType == FTimelineEdTrack::TT_FloatInterp )
-		{
-			CurveType = UCurveFloat::StaticClass();
-		}
-		else if( Track->TrackType == FTimelineEdTrack::TT_LinearColorInterp )
-		{
-			CurveType = UCurveLinearColor::StaticClass();
-		}
-		else 
-		{
-			CurveType = UCurveVector::StaticClass();
-		}
-
-		//Create curve object
-		FName AssetName = *Name;
-		UObject* NewObj = TrackWidget->CreateCurveObject( CurveType, Pkg, AssetName );
-		if( NewObj )
-		{
-			//Copy curve data from current curve to newly create curve
-			if(  Track->TrackType == FTimelineEdTrack::TT_Event || Track->TrackType == FTimelineEdTrack::TT_FloatInterp )
+			UPackage* Package = CreatePackage(NULL, *PackageName);
+			
+			//Get the curve class type
+			TSubclassOf<UCurveBase> CurveType;
+			if( Track->TrackType == FTimelineEdTrack::TT_Event || Track->TrackType == FTimelineEdTrack::TT_FloatInterp )
 			{
-				UCurveFloat* DestCurve = CastChecked<UCurveFloat>(NewObj);
-
-				AssetCurve = DestCurve;
-				UCurveFloat* SourceCurve = CastChecked<UCurveFloat>(CurveBasePtr);
+				CurveType = UCurveFloat::StaticClass();
+			}
+			else if( Track->TrackType == FTimelineEdTrack::TT_LinearColorInterp )
+			{
+				CurveType = UCurveLinearColor::StaticClass();
+			}
+			else 
+			{
+				CurveType = UCurveVector::StaticClass();
+			}
 
-				if( SourceCurve && DestCurve )
+			//Create curve object
+			UObject* NewObj = TrackWidget->CreateCurveObject( CurveType, Package, AssetName );
+			if( NewObj )
+			{
+				//Copy curve data from current curve to newly create curve
+				if(  Track->TrackType == FTimelineEdTrack::TT_Event || Track->TrackType == FTimelineEdTrack::TT_FloatInterp )
 				{
-					CopyCurveData( &SourceCurve->FloatCurve, &DestCurve->FloatCurve );
-				}
+					UCurveFloat* DestCurve = CastChecked<UCurveFloat>(NewObj);
 
-				DestCurve->bIsEventCurve = ( Track->TrackType == FTimelineEdTrack::TT_Event ) ? true : false;
-			}
-			else if( Track->TrackType == FTimelineEdTrack::TT_VectorInterp)
-			{
-				UCurveVector* DestCurve = Cast<UCurveVector>(NewObj);
+					AssetCurve = DestCurve;
+					UCurveFloat* SourceCurve = CastChecked<UCurveFloat>(CurveBasePtr);
 
-				AssetCurve = DestCurve;
-				UCurveVector* SrcCurve = CastChecked<UCurveVector>(CurveBasePtr);
+					if( SourceCurve && DestCurve )
+					{
+						CopyCurveData( &SourceCurve->FloatCurve, &DestCurve->FloatCurve );
+					}
 
-				if( SrcCurve && DestCurve )
+					DestCurve->bIsEventCurve = ( Track->TrackType == FTimelineEdTrack::TT_Event ) ? true : false;
+				}
+				else if( Track->TrackType == FTimelineEdTrack::TT_VectorInterp)
 				{
-					for( int32 i=0; i<3; i++ )
+					UCurveVector* DestCurve = Cast<UCurveVector>(NewObj);
+
+					AssetCurve = DestCurve;
+					UCurveVector* SrcCurve = CastChecked<UCurveVector>(CurveBasePtr);
+
+					if( SrcCurve && DestCurve )
 					{
-						CopyCurveData( &SrcCurve->FloatCurves[i], &DestCurve->FloatCurves[i] );
+						for( int32 i=0; i<3; i++ )
+						{
+							CopyCurveData( &SrcCurve->FloatCurves[i], &DestCurve->FloatCurves[i] );
+						}
 					}
 				}
-			}
-			else if( Track->TrackType == FTimelineEdTrack::TT_LinearColorInterp)
-			{
-				UCurveLinearColor* DestCurve = Cast<UCurveLinearColor>(NewObj);
+				else if( Track->TrackType == FTimelineEdTrack::TT_LinearColorInterp)
+				{
+					UCurveLinearColor* DestCurve = Cast<UCurveLinearColor>(NewObj);
 
-				AssetCurve = DestCurve;
-				UCurveLinearColor* SrcCurve = CastChecked<UCurveLinearColor>(CurveBasePtr);
+					AssetCurve = DestCurve;
+					UCurveLinearColor* SrcCurve = CastChecked<UCurveLinearColor>(CurveBasePtr);
 
-				if( SrcCurve && DestCurve )
-				{
-					for( int32 i=0; i<4; i++ )
+					if( SrcCurve && DestCurve )
 					{
-						CopyCurveData( &SrcCurve->FloatCurves[i], &DestCurve->FloatCurves[i] );
+						for( int32 i=0; i<4; i++ )
+						{
+							CopyCurveData( &SrcCurve->FloatCurves[i], &DestCurve->FloatCurves[i] );
+						}
 					}
 				}
-			}
 
-			// Set the new objects as the sole selection.
-			USelection* SelectionSet = GEditor->GetSelectedObjects();
-			SelectionSet->DeselectAll();
-			SelectionSet->Select( NewObj );
+				// Set the new objects as the sole selection.
+				USelection* SelectionSet = GEditor->GetSelectedObjects();
+				SelectionSet->DeselectAll();
+				SelectionSet->Select( NewObj );
 
-			// Notify the asset registry
-			FAssetRegistryModule::AssetCreated(NewObj);
+				// Notify the asset registry
+				FAssetRegistryModule::AssetCreated(NewObj);
 
-			// Mark the package dirty...
-			OutermostPkg->MarkPackageDirty();
-			return AssetCurve;
+				// Mark the package dirty...
+				Package->GetOutermost()->MarkPackageDirty();
+				return AssetCurve;
+			}
 		}
 	}
 
diff --git a/Engine/Source/Editor/Kismet/Private/STimelineEditor.h b/Engine/Source/Editor/Kismet/Private/STimelineEditor.h
--- a/Engine/Source/Editor/Kismet/Private/STimelineEditor.h
+++ b/Engine/Source/Editor/Kismet/Private/STimelineEditor.h
@@ -4,44 +4,6 @@
 
 #include "BlueprintEditor.h"
 
-
-//////////////////////////////////////////////////////////////////////////
-// SCurveAssetWidget
-class SCurveAssetWidget : public SCompoundWidget
-{
-public:
-	SLATE_BEGIN_ARGS( SCurveAssetWidget ){}
-	SLATE_ATTRIBUTE( FString, VisibleText )
-		SLATE_EVENT( FSimpleDelegate, OnCreateCurve )
-		SLATE_EVENT( FOnTextChanged, OnTextChanged )
-	SLATE_END_ARGS()
-
-
-	void Construct(const FArguments& InArgs);
-
-private:
-
-	/**	Reference to editable text */
-	TAttribute< FString > VisibleText;
-
-	/**	Callback function pointer to execute curve asset creation command */
-	FSimpleDelegate OnCreateCurve;
-
-	/** Callback function pointer to notify editable text changes */
-	FOnTextChanged OnTextChanged;
-
-	/** Widget that we want to be focused when the popup is shown  */
-	TSharedPtr<SWidget> WidgetWithDefaultFocus;
-
-	/** Function to execute curve asset creation command */
-	FReply OnCreateButtonClicked( );
-
-
-	/** Function to notify editable text changes */
-	void OnEditableTextChanged( const FText& NewString );
-};
-
-
 //////////////////////////////////////////////////////////////////////////
 // FTimelineEdTrack
 
@@ -106,9 +68,6 @@ private:
 	/**	Pointer to a window which prompts user to save internal curve as an external asset */
 	TSharedPtr<SWindow> AssetCreationWindow;
 
-	/**	File path to be used to save the external curve asset */
-	FString SaveCurveAssetPath;
-
 	/**	Pointer to the curve */
 	UCurveBase* CurveBasePtr;
 
@@ -118,29 +77,17 @@ private:
 	/** String to display external curve path as tooltip*/
 	FString ExternalCurvePath;
 
-	/**	Get curve asset path used to save the asset*/
-	FString GetSaveCurvePath() const;
-
-	/** Set curve asset path used to save the asset*/
-	void SetSaveCurvePath(const FText& NewStr);
-
 	/**Function to destroy popup window*/
 	void OnCloseCreateCurveWindow();
 
-	/** Callback function to confirm button press in order to create curve asset */
-	void OnCreateButtonClicked();
-
 	/** Function to create curve asset inside the user specified package*/
 	UCurveBase* CreateCurveAsset();
 
 	/** Callback function to initiate the external curve asset creation process*/
 	void OnCreateExternalCurve();
 
-	/** Function to create dialog which prompts to create external curve*/
-	TSharedRef<SWidget> CreateCurveAssetWidget();
-
-	/** Initialize default asset path*/
-	void InitCurveAssetPath();
+	/** Creates default asset path*/
+	FString CreateUniqueCurveAssetPathName();
 
 	/** Get the current external curve name*/
 	FText GetExternalCurveName( ) const;
