#summary Fixed build issues as a result of CL# 2044506

diff --git a/Engine/Source/Editor/GraphEditor/Private/KismetNodes/SGraphNodeK2Base.cpp b/Engine/Source/Editor/GraphEditor/Private/KismetNodes/SGraphNodeK2Base.cpp
--- a/Engine/Source/Editor/GraphEditor/Private/KismetNodes/SGraphNodeK2Base.cpp
+++ b/Engine/Source/Editor/GraphEditor/Private/KismetNodes/SGraphNodeK2Base.cpp
@@ -191,7 +191,7 @@ TSharedPtr<SToolTip> SGraphNodeK2Base::GetComplexTooltip()
 	];
 
 	// English speakers have no real need to know this exists.
-	if(FInternationalization::GetCurrentCulture()->GetTwoLetterISOLanguageName() != TEXT("en"))
+	if(FInternationalization::Get().GetCurrentCulture()->GetTwoLetterISOLanguageName() != TEXT("en"))
 	{
 		struct Local
 		{
diff --git a/Engine/Source/Editor/Kismet/Private/SBlueprintPalette.cpp b/Engine/Source/Editor/Kismet/Private/SBlueprintPalette.cpp
--- a/Engine/Source/Editor/Kismet/Private/SBlueprintPalette.cpp
+++ b/Engine/Source/Editor/Kismet/Private/SBlueprintPalette.cpp
@@ -1268,7 +1268,7 @@ TSharedPtr<SToolTip> SBlueprintPaletteItem::ConstructToolTipWidget() const
 	TSharedRef< SToolTip > TooltipWidget = IDocumentation::Get()->CreateToolTip(TextAttribute, NULL, DocLink, DocExcerptName);
 
 	// English speakers have no real need to know this exists.
-	if(FInternationalization::GetCurrentCulture()->GetTwoLetterISOLanguageName() != TEXT("en"))
+	if(FInternationalization::Get().GetCurrentCulture()->GetTwoLetterISOLanguageName() != TEXT("en"))
 	{
 		if (UK2Node const* const NodeTemplate = FK2SchemaActionUtils::ExtractNodeTemplateFromAction(PaletteAction))
 		{
diff --git a/Engine/Source/Runtime/Engine/Private/EdGraph/EdGraph.cpp b/Engine/Source/Runtime/Engine/Private/EdGraph/EdGraph.cpp
--- a/Engine/Source/Runtime/Engine/Private/EdGraph/EdGraph.cpp
+++ b/Engine/Source/Runtime/Engine/Private/EdGraph/EdGraph.cpp
@@ -635,7 +635,7 @@ FText UEdGraphNode::GetNodeSearchTitle() const
 
 	FText LocalizedTitle = GetNodeTitle(ENodeTitleType::ListView);
 	FText NativeTitle = FText::FromString(GetNodeNativeTitle(ENodeTitleType::ListView));
-	if(FInternationalization::GetCurrentCulture()->GetName() == TEXT("en_US"))
+	if(FInternationalization::Get().GetCurrentCulture()->GetName() == TEXT("en_US"))
 	{
 		ensure( LocalizedTitle.CompareTo(NativeTitle) == 0 );
 	}
