#ue4 Custom Blueprint Function Library
ttp#320585 BLUEPRINTS: Verify support for extending UKismetLibrary in blueprints

diff --git a/Engine/Config/BaseEditor.ini b/Engine/Config/BaseEditor.ini
--- a/Engine/Config/BaseEditor.ini
+++ b/Engine/Config/BaseEditor.ini
@@ -249,6 +249,9 @@ SupportWebsiteURL="http://unrealengine.com/support"
 [UserDefinedStructure]
 bUseUserDefinedStructure=false
 
+[CustomBlueprintFunctionLibrary]
+bCanCreateNew=false
+
 [EditoronlyBP]
 bAllowClassAndBlueprintPinMatching=false
 bReplaceBlueprintWithClass=false
diff --git a/Engine/Source/Editor/BlueprintGraph/Private/K2ActionMenuBuilder.cpp b/Engine/Source/Editor/BlueprintGraph/Private/K2ActionMenuBuilder.cpp
--- a/Engine/Source/Editor/BlueprintGraph/Private/K2ActionMenuBuilder.cpp
+++ b/Engine/Source/Editor/BlueprintGraph/Private/K2ActionMenuBuilder.cpp
@@ -1454,8 +1454,9 @@ void FK2ActionMenuBuilder::GetFuncNodesForClass(FGraphActionListBuilderBase& Lis
 		for (TObjectIterator<UClass> ClassIt; ClassIt; ++ClassIt)
 		{
 			UClass* TestClass = *ClassIt;
+			const bool bIsSkeletonClass = TestClass->HasAnyFlags(RF_Transient) && TestClass->HasAnyClassFlags(CLASS_CompiledFromBlueprint);
 			if (TestClass->IsChildOf(UBlueprintFunctionLibrary::StaticClass()) &&
-				CanUseLibraryFunctionsForScope(TestClass, OwnerClass))
+				CanUseLibraryFunctionsForScope(TestClass, OwnerClass) && !bIsSkeletonClass)
 			{
 				GetFuncNodesForClass(ListBuilder, TestClass, OwnerClass, DestGraph, FunctionTypes, BaseCategory, TargetInfo);
 			}
@@ -2481,8 +2482,9 @@ void FK2ActionMenuBuilder::GetFuncNodesWithPinType(FBlueprintGraphActionListBuil
 		for (TObjectIterator<UClass> ClassIt; ClassIt; ++ClassIt)
 		{
 			UClass* TestClass = *ClassIt;
+			const bool bIsSkeletonClass = TestClass->HasAnyFlags(RF_Transient) && TestClass->HasAnyClassFlags(CLASS_CompiledFromBlueprint);
 			if (TestClass->IsChildOf(UBlueprintFunctionLibrary::StaticClass()) &&
-				CanUseLibraryFunctionsForScope(OwnerClass, TestClass))
+				CanUseLibraryFunctionsForScope(OwnerClass, TestClass) && !bIsSkeletonClass)
 			{
 				GetFuncNodesWithPinType(ContextMenuBuilder, TestClass, OwnerClass, DesiredPinType, bWantOutput, bPureOnly);
 			}
diff --git a/Engine/Source/Editor/Kismet/Private/BlueprintDetailsCustomization.cpp b/Engine/Source/Editor/Kismet/Private/BlueprintDetailsCustomization.cpp
--- a/Engine/Source/Editor/Kismet/Private/BlueprintDetailsCustomization.cpp
+++ b/Engine/Source/Editor/Kismet/Private/BlueprintDetailsCustomization.cpp
@@ -3979,7 +3979,8 @@ void FBlueprintGlobalOptionsDetails::CustomizeDetails(IDetailLayoutBuilder& Deta
 		const bool bIsInterfaceBP = FBlueprintEditorUtils::IsInterfaceBlueprint(Blueprint);
 		const bool bIsMacroLibrary = Blueprint->BlueprintType == BPTYPE_MacroLibrary;
 		const bool bIsLevelScriptBP = FBlueprintEditorUtils::IsLevelScriptBlueprint(Blueprint);
-		const bool bSupportsInterfaces = !bIsLevelScriptBP && !bIsInterfaceBP && !bIsMacroLibrary;
+		const bool bIsFunctionLibrary = Blueprint->BlueprintType == BPTYPE_FunctionLibrary;
+		const bool bSupportsInterfaces = !bIsLevelScriptBP && !bIsInterfaceBP && !bIsMacroLibrary && !bIsFunctionLibrary;
 
 		if (bSupportsInterfaces)
 		{
diff --git a/Engine/Source/Editor/Kismet/Private/BlueprintEditor.cpp b/Engine/Source/Editor/Kismet/Private/BlueprintEditor.cpp
--- a/Engine/Source/Editor/Kismet/Private/BlueprintEditor.cpp
+++ b/Engine/Source/Editor/Kismet/Private/BlueprintEditor.cpp
@@ -1086,6 +1086,13 @@ void FBlueprintEditor::InitBlueprintEditor(const EToolkitMode::Type Mode, const
 				MakeShareable(new FBlueprintMacroApplicationMode( SharedThis(this) )));
 			SetCurrentMode(FBlueprintEditorApplicationModes::BlueprintMacroMode);
 		}
+		else if (SingleBP->BlueprintType == BPTYPE_FunctionLibrary)
+		{
+			AddApplicationMode(
+				FBlueprintEditorApplicationModes::StandardBlueprintEditorMode,
+				MakeShareable(new FBlueprintEditorApplicationMode(SharedThis(this), FBlueprintEditorApplicationModes::StandardBlueprintEditorMode)));
+			SetCurrentMode(FBlueprintEditorApplicationModes::StandardBlueprintEditorMode);
+		}
 		else
 		{
 			AddApplicationMode(
@@ -1860,7 +1867,7 @@ bool FBlueprintEditor::ReparentBlueprint_IsVisible() const
 	if (Blueprint != NULL)
 	{
 		// Don't show the reparent option if it's an Interface or we're not in editing mode
-		return !FBlueprintEditorUtils::IsInterfaceBlueprint(Blueprint) && InEditingMode();
+		return !FBlueprintEditorUtils::IsInterfaceBlueprint(Blueprint) && InEditingMode() && (BPTYPE_FunctionLibrary != Blueprint->BlueprintType);
 	}
 	else
 	{
@@ -2515,7 +2522,8 @@ bool FBlueprintEditor::CanAccessComponentsMode() const
 		UBlueprint* Blueprint = GetBlueprintObj();
 		bCanAccess = Blueprint->SimpleConstructionScript != NULL		// An SCS must be present (otherwise there is nothing valid to edit)
 			&& FBlueprintEditorUtils::IsActorBased(Blueprint)			// Must be parented to an AActor-derived class (some older BPs may have an SCS but may not be Actor-based)
-			&& Blueprint->BlueprintType != BPTYPE_MacroLibrary;			// Must not be a macro-type Blueprint
+			&& Blueprint->BlueprintType != BPTYPE_MacroLibrary			// Must not be a macro-type Blueprint
+			&& Blueprint->BlueprintType != BPTYPE_FunctionLibrary;		// Must not be a function library
 	}
 	
 	return bCanAccess;
@@ -5441,7 +5449,9 @@ bool FBlueprintEditor::NewDocument_IsVisibleForType(ECreatedDocumentType GraphTy
 	switch (GraphType)
 	{
 	case CGT_NewVariable:
-		return (GetBlueprintObj()->BlueprintType != BPTYPE_Interface) && (GetBlueprintObj()->BlueprintType != BPTYPE_MacroLibrary);
+		return (GetBlueprintObj()->BlueprintType != BPTYPE_FunctionLibrary) 
+			&& (GetBlueprintObj()->BlueprintType != BPTYPE_Interface) 
+			&& (GetBlueprintObj()->BlueprintType != BPTYPE_MacroLibrary);
 	case CGT_NewFunctionGraph:
 		return (GetBlueprintObj()->BlueprintType != BPTYPE_MacroLibrary);
 	case CGT_NewMacroGraph:
@@ -5462,7 +5472,8 @@ bool FBlueprintEditor::AddNewDelegateIsVisible() const
 	const UBlueprint* Blueprint = GetBlueprintObj();
 	return (NULL != Blueprint)
 		&& (Blueprint->BlueprintType != BPTYPE_Interface) 
-		&& (Blueprint->BlueprintType != BPTYPE_MacroLibrary);
+		&& (Blueprint->BlueprintType != BPTYPE_MacroLibrary)
+		&& (Blueprint->BlueprintType != BPTYPE_FunctionLibrary);
 }
 
 void FBlueprintEditor::NotifyPreChange(UProperty* PropertyAboutToChange)
diff --git a/Engine/Source/Editor/UnrealEd/Classes/Factories/BlueprintFunctionLibraryFactory.h b/Engine/Source/Editor/UnrealEd/Classes/Factories/BlueprintFunctionLibraryFactory.h
new file mode 100644
--- /dev/null
+++ b/Engine/Source/Editor/UnrealEd/Classes/Factories/BlueprintFunctionLibraryFactory.h
@@ -0,0 +1,25 @@
+// Copyright 1998-2014 Epic Games, Inc. All Rights Reserved.
+
+/**
+ *
+ */
+
+#pragma once
+#include "BlueprintFunctionLibraryFactory.generated.h"
+
+UCLASS(hidecategories=Object, collapsecategories)
+class UBlueprintFunctionLibraryFactory : public UBlueprintFactory
+{
+	GENERATED_UCLASS_BODY()
+
+	// Begin UFactory Interface
+	virtual FText GetDisplayName() const OVERRIDE;
+	virtual FName GetNewAssetThumbnailOverride() const OVERRIDE;
+	virtual uint32 GetMenuCategories() const OVERRIDE;
+	virtual UObject* FactoryCreateNew(UClass* Class, UObject* InParent, FName Name, EObjectFlags Flags, UObject* Context, FFeedbackContext* Warn, FName CallingContext) OVERRIDE;
+	virtual bool ConfigureProperties() OVERRIDE;
+	// Begin UFactory Interface
+};
+
+
+
diff --git a/Engine/Source/Editor/UnrealEd/Private/Factories/EditorFactories.cpp b/Engine/Source/Editor/UnrealEd/Private/Factories/EditorFactories.cpp
--- a/Engine/Source/Editor/UnrealEd/Private/Factories/EditorFactories.cpp
+++ b/Engine/Source/Editor/UnrealEd/Private/Factories/EditorFactories.cpp
@@ -6258,6 +6258,65 @@ UObject* UBlueprintMacroFactory::FactoryCreateNew(UClass* Class, UObject* InPare
 	}
 }
 
+/*------------------------------------------------------------------------------
+BlueprintFunctionLibraryFactory implementation.
+------------------------------------------------------------------------------*/
+UBlueprintFunctionLibraryFactory::UBlueprintFunctionLibraryFactory(const class FPostConstructInitializeProperties& PCIP)
+: Super(PCIP)
+{
+	struct FCanCreateNewHelper
+	{
+		bool bCanCreateNew;
+		FCanCreateNewHelper() : bCanCreateNew(false)
+		{
+			GConfig->GetBool(TEXT("CustomBlueprintFunctionLibrary"), TEXT("bCanCreateNew"), bCanCreateNew, GEditorIni);
+		}
+	};
+	static const FCanCreateNewHelper Helper;
+	bCreateNew = Helper.bCanCreateNew;
+	bEditAfterNew = true;
+	SupportedClass = UBlueprint::StaticClass();
+	ParentClass = UBlueprintFunctionLibrary::StaticClass();
+}
+
+FText UBlueprintFunctionLibraryFactory::GetDisplayName() const
+{
+	return LOCTEXT("BlueprintFunctionLibraryFactoryDescription", "Blueprint Function Library");
+}
+
+FName UBlueprintFunctionLibraryFactory::GetNewAssetThumbnailOverride() const
+{
+	return TEXT("ClassThumbnail.BlueprintMacroLibrary");
+}
+
+uint32 UBlueprintFunctionLibraryFactory::GetMenuCategories() const
+{
+	// Force this factory into the misc category, since it does not belong in the top menu
+	return EAssetTypeCategories::Misc;
+}
+
+UObject* UBlueprintFunctionLibraryFactory::FactoryCreateNew(UClass* Class, UObject* InParent, FName Name, EObjectFlags Flags, UObject* Context, FFeedbackContext* Warn, FName CallingContext)
+{
+	// Make sure we are trying to factory a blueprint, then create and init one
+	check(Class->IsChildOf(UBlueprint::StaticClass()));
+
+	if ((ParentClass == NULL) || !FKismetEditorUtilities::CanCreateBlueprintOfClass(ParentClass))
+	{
+		FFormatNamedArguments Args;
+		Args.Add(TEXT("ClassName"), (ParentClass != NULL) ? FText::FromString(ParentClass->GetName()) : LOCTEXT("Null", "(null)"));
+		FMessageDialog::Open(EAppMsgType::Ok, FText::Format(LOCTEXT("CannotCreateBlueprintFromClass", "Cannot create a blueprint based on the class '{0}'."), Args));
+		return NULL;
+	}
+	else
+	{
+		return FKismetEditorUtilities::CreateBlueprint(ParentClass, InParent, Name, BPTYPE_FunctionLibrary, UBlueprint::StaticClass(), CallingContext);
+	}
+}
+
+bool UBlueprintFunctionLibraryFactory::ConfigureProperties()
+{
+	return true;
+}
 
 /*------------------------------------------------------------------------------
 	UBlueprintInterfaceFactory implementation.
diff --git a/Engine/Source/Editor/UnrealEd/Private/Kismet2/BlueprintEditorUtils.cpp b/Engine/Source/Editor/UnrealEd/Private/Kismet2/BlueprintEditorUtils.cpp
--- a/Engine/Source/Editor/UnrealEd/Private/Kismet2/BlueprintEditorUtils.cpp
+++ b/Engine/Source/Editor/UnrealEd/Private/Kismet2/BlueprintEditorUtils.cpp
@@ -1540,7 +1540,12 @@ void FBlueprintEditorUtils::AddFunctionGraph(UBlueprint* Blueprint, class UEdGra
 		if (bIsUserCreated)
 		{
 			// We need to flag the entry node to make sure that the compiled function is callable from Kismet2
-			K2Schema->AddExtraFunctionFlags(Graph, (FUNC_BlueprintCallable|FUNC_BlueprintEvent|FUNC_Public));
+			int32 ExtraFunctionFlags = (FUNC_BlueprintCallable | FUNC_BlueprintEvent | FUNC_Public);
+			if (BPTYPE_FunctionLibrary == Blueprint->BlueprintType)
+			{
+				ExtraFunctionFlags |= FUNC_Static;
+			}
+			K2Schema->AddExtraFunctionFlags(Graph, ExtraFunctionFlags);
 			K2Schema->MarkFunctionEntryAsEditable(Graph, true);
 		}
 	}
@@ -2170,7 +2175,8 @@ bool FBlueprintEditorUtils::SupportsConstructionScript(const UBlueprint* Bluepri
 			!FBlueprintEditorUtils::IsBlueprintConst(Blueprint) && 
 			!FBlueprintEditorUtils::IsLevelScriptBlueprint(Blueprint) && 
 			FBlueprintEditorUtils::IsActorBased(Blueprint)) &&
-			!(Blueprint->BlueprintType == BPTYPE_MacroLibrary);
+			!(Blueprint->BlueprintType == BPTYPE_MacroLibrary) &&
+			!(Blueprint->BlueprintType == BPTYPE_FunctionLibrary);
 }
 
 UEdGraph* FBlueprintEditorUtils::FindUserConstructionScript(const UBlueprint* Blueprint)
@@ -2224,7 +2230,8 @@ bool FBlueprintEditorUtils::DoesBlueprintContainField(const UBlueprint* Blueprin
 bool FBlueprintEditorUtils::DoesSupportOverridingFunctions(const UBlueprint* Blueprint)
 {
 	return Blueprint->BlueprintType != BPTYPE_MacroLibrary 
-		&& Blueprint->BlueprintType != BPTYPE_Interface;
+		&& Blueprint->BlueprintType != BPTYPE_Interface
+		&& Blueprint->BlueprintType != BPTYPE_FunctionLibrary;
 }
 
 bool FBlueprintEditorUtils::DoesSupportTimelines(const UBlueprint* Blueprint)
@@ -2242,9 +2249,10 @@ bool FBlueprintEditorUtils::DoesSupportEventGraphs(const UBlueprint* Blueprint)
 /** Returns whether or not the blueprint supports implementing interfaces */
 bool FBlueprintEditorUtils::DoesSupportImplementingInterfaces(const UBlueprint* Blueprint)
 {
-	return Blueprint->BlueprintType != BPTYPE_MacroLibrary 
-		&& Blueprint->BlueprintType != BPTYPE_Interface 
-		&& Blueprint->BlueprintType != BPTYPE_LevelScript;
+	return Blueprint->BlueprintType != BPTYPE_MacroLibrary
+		&& Blueprint->BlueprintType != BPTYPE_Interface
+		&& Blueprint->BlueprintType != BPTYPE_LevelScript
+		&& Blueprint->BlueprintType != BPTYPE_FunctionLibrary;
 }
 
 // Returns a descriptive name of the type of blueprint passed in
@@ -2262,6 +2270,7 @@ FString FBlueprintEditorUtils::GetBlueprintTypeDescription(const UBlueprint* Blu
 	case BPTYPE_Interface:
 		BlueprintTypeString = LOCTEXT("BlueprintType_Interface", "Interface").ToString();
 		break;
+	case BPTYPE_FunctionLibrary:
 	case BPTYPE_Normal:
 	case BPTYPE_Const:
 		BlueprintTypeString = Blueprint->GetClass()->GetName();
diff --git a/Engine/Source/Editor/UnrealEd/Private/Kismet2/Kismet2.cpp b/Engine/Source/Editor/UnrealEd/Private/Kismet2/Kismet2.cpp
--- a/Engine/Source/Editor/UnrealEd/Private/Kismet2/Kismet2.cpp
+++ b/Engine/Source/Editor/UnrealEd/Private/Kismet2/Kismet2.cpp
@@ -489,7 +489,9 @@ bool FKismetEditorUtilities::CanCreateBlueprintOfClass(const UClass* Class)
 		&& (!Class->ClassGeneratedBy || bAllowDerivedBlueprints);
 
 	const UEdGraphSchema_K2* K2Schema = GetDefault<UEdGraphSchema_K2>();
-	const bool bIsValidClass = Class->GetBoolMetaDataHierarchical(FBlueprintMetadata::MD_IsBlueprintBase) || (Class == UObject::StaticClass());
+	const bool bIsValidClass = Class->GetBoolMetaDataHierarchical(FBlueprintMetadata::MD_IsBlueprintBase) 
+		|| (Class == UObject::StaticClass()) 
+		|| (Class == UBlueprintFunctionLibrary::StaticClass());
 
 	return bCanCreateBlueprint && bIsValidClass;
 }
