#ue4 BP: User is notified while bp reparenting. When the blueprint implements already an interface, thet is implemented by the new parent.

diff --git a/Engine/Source/Editor/Kismet/Private/BlueprintEditor.cpp b/Engine/Source/Editor/Kismet/Private/BlueprintEditor.cpp
--- a/Engine/Source/Editor/Kismet/Private/BlueprintEditor.cpp
+++ b/Engine/Source/Editor/Kismet/Private/BlueprintEditor.cpp
@@ -1819,9 +1819,39 @@ void FBlueprintEditor::ReparentBlueprint_NewParentChosen(UClass* ChosenClass)
 
 	if ((BlueprintObj != NULL) && (ChosenClass != NULL) && (ChosenClass != BlueprintObj->ParentClass))
 	{
-		// If the chosen class differs hierarchically from the current class, warn that there may be data loss
+		// Notify user, about common interfaces
 		bool bReparent = true;
-		if ( !ChosenClass->GetDefaultObject()->IsA( BlueprintObj->ParentClass ) )
+		{
+			FString CommonInterfacesNames;
+			for (auto InterdaceDesc : BlueprintObj->ImplementedInterfaces)
+			{
+				if (ChosenClass->ImplementsInterface(*InterdaceDesc.Interface))
+				{
+					CommonInterfacesNames += InterdaceDesc.Interface->GetName();
+					CommonInterfacesNames += TCHAR('\n');
+				}
+			}
+			if (!CommonInterfacesNames.IsEmpty())
+			{
+				const FText Title = LOCTEXT("CommonInterfacesTitle", "Common interfaces");
+				const FText Message = FText::Format(
+					LOCTEXT("ReparentWarning", "Following interfaces are already implemented. Continue reparenting? \n {0}"),
+					FText::FromString(CommonInterfacesNames));
+
+				FSuppressableWarningDialog::FSetupInfo Info(Message, Title, "Warning_CommonInterfacesWhileReparenting");
+				Info.ConfirmText = LOCTEXT("ReparentYesButton", "Reparent");
+				Info.CancelText = LOCTEXT("ReparentNoButton", "Cancel");
+
+				FSuppressableWarningDialog ReparentBlueprintDlg(Info);
+				if (ReparentBlueprintDlg.ShowModal() == FSuppressableWarningDialog::Cancel)
+				{
+					bReparent = false;
+				}
+			}
+		}
+
+		// If the chosen class differs hierarchically from the current class, warn that there may be data loss
+		if (bReparent && !ChosenClass->GetDefaultObject()->IsA(BlueprintObj->ParentClass))
 		{
 			const FText Title = LOCTEXT("ReparentTitle", "Reparent Blueprint"); 
 			const FText Message = LOCTEXT("ReparentWarning", "Reparenting this blueprint may cause data loss.  Continue reparenting?"); 
