* GetGraphAppearance is now protected and virtual, subclasses of the FBlueprintEditor can now override it to return custom text on the graphs.  Updating Persona to do this, instead of having a hack in the FBlueprintEditor to do this.
* WorkflowTabManager now has a way to unregister tabs.

diff --git a/Engine/Source/Editor/Kismet/Private/BlueprintEditor.cpp b/Engine/Source/Editor/Kismet/Private/BlueprintEditor.cpp
--- a/Engine/Source/Editor/Kismet/Private/BlueprintEditor.cpp
+++ b/Engine/Source/Editor/Kismet/Private/BlueprintEditor.cpp
@@ -703,28 +703,20 @@ FGraphAppearanceInfo FBlueprintEditor::GetGraphAppearance() const
 {
 	// Create the appearance info
 	FGraphAppearanceInfo AppearanceInfo;
-	if (GetBlueprintObj()->IsA(UAnimBlueprint::StaticClass()))
+	switch ( GetBlueprintObj()->BlueprintType )
 	{
-		AppearanceInfo.CornerText = LOCTEXT("AppearanceCornerText_Animation", "ANIMATION").ToString();
-	}
-	else
-	{
-		switch (GetBlueprintObj()->BlueprintType)
-		{
-		case BPTYPE_LevelScript:
-			AppearanceInfo.CornerText = LOCTEXT("AppearanceCornerText_LevelScript", "LEVEL BLUEPRINT").ToString();
-			break;
-		case BPTYPE_MacroLibrary:
-			AppearanceInfo.CornerText = LOCTEXT("AppearanceCornerText_Macro", "MACRO").ToString();
-			break;
-		case BPTYPE_Interface:
-			AppearanceInfo.CornerText = LOCTEXT("AppearanceCornerText_Interface", "INTERFACE").ToString();
-			break;
-		default:
-			AppearanceInfo.CornerText = LOCTEXT("AppearanceCornerText_Blueprint", "BLUEPRINT").ToString();
-			break;
-		}
-
+	case BPTYPE_LevelScript:
+		AppearanceInfo.CornerText = LOCTEXT("AppearanceCornerText_LevelScript", "LEVEL BLUEPRINT").ToString();
+		break;
+	case BPTYPE_MacroLibrary:
+		AppearanceInfo.CornerText = LOCTEXT("AppearanceCornerText_Macro", "MACRO").ToString();
+		break;
+	case BPTYPE_Interface:
+		AppearanceInfo.CornerText = LOCTEXT("AppearanceCornerText_Interface", "INTERFACE").ToString();
+		break;
+	default:
+		AppearanceInfo.CornerText = LOCTEXT("AppearanceCornerText_Blueprint", "BLUEPRINT").ToString();
+		break;
 	}
 
 	AppearanceInfo.PIENotifyText = GetPIEStatus();
diff --git a/Engine/Source/Editor/Kismet/Private/WorkflowOrientedApp/WorkflowTabManager.cpp b/Engine/Source/Editor/Kismet/Private/WorkflowOrientedApp/WorkflowTabManager.cpp
--- a/Engine/Source/Editor/Kismet/Private/WorkflowOrientedApp/WorkflowTabManager.cpp
+++ b/Engine/Source/Editor/Kismet/Private/WorkflowOrientedApp/WorkflowTabManager.cpp
@@ -326,6 +326,12 @@ void FWorkflowAllowedTabSet::RegisterFactory(TSharedPtr<class FWorkflowTabFactor
 	Factories.Add(NewIdentifier, Factory);
 }
 
+void FWorkflowAllowedTabSet::UnregisterFactory(FName FactoryID)
+{
+	int32 Removed = Factories.Remove(FactoryID);
+	check(Removed != 0);
+}
+
 // Merges in a set of factories into this set
 void FWorkflowAllowedTabSet::MergeInSet(const FWorkflowAllowedTabSet& OtherSet)
 {
diff --git a/Engine/Source/Editor/Kismet/Public/BlueprintEditor.h b/Engine/Source/Editor/Kismet/Public/BlueprintEditor.h
--- a/Engine/Source/Editor/Kismet/Public/BlueprintEditor.h
+++ b/Engine/Source/Editor/Kismet/Public/BlueprintEditor.h
@@ -779,6 +779,9 @@ protected:
 
 	virtual void AnalyticsTrackNewNode( FName NodeClass, FName NodeType ) OVERRIDE;
 
+	/** Get graph appearance */
+	virtual FGraphAppearanceInfo GetGraphAppearance() const;
+
 private:
 
 	/* User wants to edit tunnel via function editor */
@@ -793,9 +796,6 @@ private:
 	/** Helper to move focused graph when clicking on graph breadcrumb */
 	void OnChangeBreadCrumbGraph( class UEdGraph* InGraph);
 
-	/** Get graph appearance */
-	FGraphAppearanceInfo GetGraphAppearance() const;
-
 	/** Function to check whether the give graph is a subgraph */
 	static bool IsASubGraph( const class UEdGraph* GraphPtr );
 
diff --git a/Engine/Source/Editor/Kismet/Public/WorkflowOrientedApp/WorkflowTabManager.h b/Engine/Source/Editor/Kismet/Public/WorkflowOrientedApp/WorkflowTabManager.h
--- a/Engine/Source/Editor/Kismet/Public/WorkflowOrientedApp/WorkflowTabManager.h
+++ b/Engine/Source/Editor/Kismet/Public/WorkflowOrientedApp/WorkflowTabManager.h
@@ -25,6 +25,9 @@ public:
 	// Registers a factory with this set - must be unique.
 	void RegisterFactory(TSharedPtr<class FWorkflowTabFactory> Factory);
 
+	// Unregisters a factory with this specified ID.
+	void UnregisterFactory(FName FactoryID);
+
 	// Merges in a set of factories into this set - duplicates are OK.
 	void MergeInSet(const FWorkflowAllowedTabSet& OtherSet);
 
diff --git a/Engine/Source/Editor/Persona/Private/Persona.cpp b/Engine/Source/Editor/Persona/Private/Persona.cpp
--- a/Engine/Source/Editor/Persona/Private/Persona.cpp
+++ b/Engine/Source/Editor/Persona/Private/Persona.cpp
@@ -1698,6 +1698,18 @@ void FPersona::CreateDefaultTabContents(const TArray<UBlueprint*>& InBlueprints)
 		. OnFinishedChangingProperties( FOnFinishedChangingProperties::FDelegate::CreateSP( this, &FPersona::OnFinishedChangingProperties ) );
 }
 
+FGraphAppearanceInfo FPersona::GetGraphAppearance() const
+{
+	FGraphAppearanceInfo AppearanceInfo = FBlueprintEditor::GetGraphAppearance();
+
+	if ( GetBlueprintObj()->IsA(UAnimBlueprint::StaticClass()) )
+	{
+		AppearanceInfo.CornerText = LOCTEXT("AppearanceCornerText_Animation", "ANIMATION").ToString();
+	}
+
+	return AppearanceInfo;
+}
+
 void FPersona::FocusWindow(UObject* ObjectToFocusOn)
 {
 	FBlueprintEditor::FocusWindow(ObjectToFocusOn);
diff --git a/Engine/Source/Editor/Persona/Private/Persona.h b/Engine/Source/Editor/Persona/Private/Persona.h
--- a/Engine/Source/Editor/Persona/Private/Persona.h
+++ b/Engine/Source/Editor/Persona/Private/Persona.h
@@ -234,6 +234,7 @@ protected:
 	virtual void OnOpenRelatedAsset() OVERRIDE;
 	virtual void GetCustomDebugObjects(TArray<FCustomDebugObject>& DebugList) const OVERRIDE;
 	virtual void CreateDefaultTabContents(const TArray<UBlueprint*>& InBlueprints) OVERRIDE;
+	virtual FGraphAppearanceInfo GetGraphAppearance() const OVERRIDE;
 	// End of FBlueprintEditor interface
 
 	// IAssetEditorInstance interface
