#ue4
- Allowed derived classes to be selected in the debug filter widget

diff --git a/Engine/Source/Editor/Kismet/Private/SBlueprintEditorSelectedDebugObjectWidget.cpp b/Engine/Source/Editor/Kismet/Private/SBlueprintEditorSelectedDebugObjectWidget.cpp
--- a/Engine/Source/Editor/Kismet/Private/SBlueprintEditorSelectedDebugObjectWidget.cpp
+++ b/Engine/Source/Editor/Kismet/Private/SBlueprintEditorSelectedDebugObjectWidget.cpp
@@ -375,10 +375,11 @@ void SBlueprintEditorSelectedDebugObjectWidget::GenerateDebugObjectNames(bool bR
 		}
 
 		const bool bPassesFlags = !TestObject->HasAnyFlags(RF_PendingKill | RF_ClassDefaultObject);
-		const bool bGeneratedByBlueprint = TestObject->GetClass()->ClassGeneratedBy == GetBlueprintObj();
-		if (bPassesFlags && bGeneratedByBlueprint)
-		{
+		const bool bGeneratedByAnyBlueprint = TestObject->GetClass()->ClassGeneratedBy != nullptr;
+		const bool bGeneratedByThisBlueprint = bGeneratedByAnyBlueprint && TestObject->IsA(GetBlueprintObj()->GeneratedClass);
 
+		if (bPassesFlags && bGeneratedByThisBlueprint)
+		{
 			UObject *ObjOuter = TestObject;
 			UWorld *ObjWorld = NULL;
 			while (ObjWorld == NULL && ObjOuter != NULL)
