ttp#332980 CRASH: FBlueprintEditor::ExpandNode() BaseSPMethodDelegateInstance_NoParams<FBlueprintEditor,void,0>::Execute()

One should never assume that CallFunction->GetThenPin() returns a pin, the function can be pure.

diff --git a/Engine/Source/Editor/Kismet/Private/BlueprintEditor.cpp b/Engine/Source/Editor/Kismet/Private/BlueprintEditor.cpp
--- a/Engine/Source/Editor/Kismet/Private/BlueprintEditor.cpp
+++ b/Engine/Source/Editor/Kismet/Private/BlueprintEditor.cpp
@@ -5124,9 +5124,10 @@ void FBlueprintEditor::ExpandNode(UEdGraphNode* InNodeToExpand, UEdGraph* InSour
 	UEdGraphPin* OutputExecPinReconnect = NULL;
 	if(UK2Node_CallFunction* CallFunction = Cast<UK2Node_CallFunction>(InNodeToExpand))
 	{
-		if(CallFunction->GetThenPin()->LinkedTo.Num())
+		auto ThenPin = CallFunction->GetThenPin();
+		if (ThenPin && ThenPin->LinkedTo.Num())
 		{
-			OutputExecPinReconnect = CallFunction->GetThenPin()->LinkedTo[0];
+			OutputExecPinReconnect = ThenPin->LinkedTo[0];
 		}
 	}
 
