Fix non-unity build

diff --git a/Engine/Source/Editor/Kismet/Private/SSCSEditorViewport.h b/Engine/Source/Editor/Kismet/Private/SSCSEditorViewport.h
--- a/Engine/Source/Editor/Kismet/Private/SSCSEditorViewport.h
+++ b/Engine/Source/Editor/Kismet/Private/SSCSEditorViewport.h
@@ -3,6 +3,7 @@
 #pragma once
 
 #include "PreviewScene.h"
+#include "SEditorViewport.h"
 
 /**
  * Implements the viewport widget that's hosted in the SCS editor tab.
