Merging UE4-Pretest @ 2042161 to UE4

diff --git a/Engine/Source/Editor/BlueprintGraph/Private/K2Node_MatineeController.cpp b/Engine/Source/Editor/BlueprintGraph/Private/K2Node_MatineeController.cpp
--- a/Engine/Source/Editor/BlueprintGraph/Private/K2Node_MatineeController.cpp
+++ b/Engine/Source/Editor/BlueprintGraph/Private/K2Node_MatineeController.cpp
@@ -2,10 +2,23 @@
 
 #include "BlueprintGraphPrivatePCH.h"
 #include "KismetCompiler.h"
+#include "Runtime/Engine/Public/MatineeDelegates.h"
 
 UK2Node_MatineeController::UK2Node_MatineeController(const class FPostConstructInitializeProperties& PCIP)
 	: Super(PCIP)
 {
+	FMatineeDelegates::Get().OnEventKeyframeAdded.AddUObject(this, &UK2Node_MatineeController::OnEventKeyframeAdded);
+	FMatineeDelegates::Get().OnEventKeyframeRenamed.AddUObject(this, &UK2Node_MatineeController::OnEventKeyframeRenamed);
+	FMatineeDelegates::Get().OnEventKeyframeRemoved.AddUObject(this, &UK2Node_MatineeController::OnEventKeyframeRemoved);
+}
+
+void UK2Node_MatineeController::BeginDestroy()
+{
+	Super::BeginDestroy();
+
+	FMatineeDelegates::Get().OnEventKeyframeAdded.RemoveAll(this);
+	FMatineeDelegates::Get().OnEventKeyframeRenamed.RemoveAll(this);
+	FMatineeDelegates::Get().OnEventKeyframeRemoved.RemoveAll(this);
 }
 
 void UK2Node_MatineeController::AllocateDefaultPins()
@@ -107,5 +120,63 @@ void UK2Node_MatineeController::ExpandNode(FKismetCompilerContext& CompilerConte
 			}
 		}
 	}
+}
+
+void UK2Node_MatineeController::OnEventKeyframeAdded(const AMatineeActor* InMatineeActor, const FName& InPinName, int32 InIndex)
+{
+	if ( MatineeActor == InMatineeActor )
+	{
+		const UEdGraphSchema_K2* K2Schema = GetDefault<UEdGraphSchema_K2>();
+
+		// Add one to the index as we insert "finished" at index 0
+		CreatePin(EGPD_Output, K2Schema->PC_Exec, TEXT(""), NULL, false, false, InPinName.ToString(), false, InIndex + 1);
+
+		// Update and refresh the blueprint
+		FBlueprintEditorUtils::MarkBlueprintAsStructurallyModified(GetBlueprint());
+	}
+}
+
+void UK2Node_MatineeController::OnEventKeyframeRenamed(const AMatineeActor* InMatineeActor, const FName& InOldPinName, const FName& InNewPinName)
+{
+	if ( MatineeActor == InMatineeActor )
+	{
+		// Rather than rename the pin directly, modify 'this' array as it's cheaper and quicker for the undo/redo system to not track UEdGraphPin's directly
+		if( UEdGraphPin* OldPin = FindPin(InOldPinName.ToString()) )
+		{
+			// Cache the pin index, then remove it from the array
+			const int32 PinIndex = GetPinIndex(OldPin);
+			check( PinIndex != INDEX_NONE );
+			DiscardPin(OldPin);
+
+			// Create a new pin at the old index and with the new name and copy it's data
+			const UEdGraphSchema_K2* K2Schema = GetDefault<UEdGraphSchema_K2>();
+			UEdGraphPin* NewPin = CreatePin(EGPD_Output, K2Schema->PC_Exec, TEXT(""), NULL, false, false, InNewPinName.ToString(), false, PinIndex);	
+			NewPin->CopyPersistentDataFromOldPin(*OldPin);
+			OldPin = NULL;
+
+			GetGraph()->NotifyGraphChanged();
+		}
+	}
+}
 
-}
\ No newline at end of file
+void UK2Node_MatineeController::OnEventKeyframeRemoved(const AMatineeActor* InMatineeActor, const TArray<FName>& InPinNames)
+{
+	if ( MatineeActor == InMatineeActor )
+	{
+		bool bNeedsRefresh = false;
+		for(int32 PinIdx=0; PinIdx<InPinNames.Num(); PinIdx++)
+		{
+			if(UEdGraphPin* Pin = FindPin(InPinNames[PinIdx].ToString()))
+			{
+				DiscardPin(Pin);
+				bNeedsRefresh = true;
+			}
+		}
+
+		if ( bNeedsRefresh )
+		{
+			// Update and refresh the blueprint
+			FBlueprintEditorUtils::MarkBlueprintAsStructurallyModified(GetBlueprint());
+		}
+	}
+}
diff --git a/Engine/Source/Editor/BlueprintGraph/Private/K2Node_SpawnActor.cpp b/Engine/Source/Editor/BlueprintGraph/Private/K2Node_SpawnActor.cpp
--- a/Engine/Source/Editor/BlueprintGraph/Private/K2Node_SpawnActor.cpp
+++ b/Engine/Source/Editor/BlueprintGraph/Private/K2Node_SpawnActor.cpp
@@ -482,7 +482,7 @@ bool UK2Node_SpawnActor::HasExternalBlueprintDependencies(TArray<class UStruct*>
 	{
 		OptionalOutput->Add(SourceClass);
 	}
-	return bResult;
+	return bResult || Super::HasExternalBlueprintDependencies(OptionalOutput);
 }
 
 bool UK2Node_SpawnActor::IsDeprecated() const
diff --git a/Engine/Source/Editor/BlueprintGraph/Private/K2Node_SpawnActorFromClass.cpp b/Engine/Source/Editor/BlueprintGraph/Private/K2Node_SpawnActorFromClass.cpp
--- a/Engine/Source/Editor/BlueprintGraph/Private/K2Node_SpawnActorFromClass.cpp
+++ b/Engine/Source/Editor/BlueprintGraph/Private/K2Node_SpawnActorFromClass.cpp
@@ -449,10 +449,7 @@ void UK2Node_SpawnActorFromClass::ExpandNode(class FKismetCompilerContext& Compi
 					// Move connection from the variable pin on the spawn node to the 'value' pin
 					UEdGraphPin* ValuePin = SetVarNode->FindPinChecked(ValueParamName);
 					CompilerContext.MovePinLinksToIntermediate(*SpawnVarPin, *ValuePin);
-					if(SpawnVarPin->PinType.bIsArray)
-					{
-						SetVarNode->PinConnectionListChanged(ValuePin);
-					}
+					SetVarNode->PinConnectionListChanged(ValuePin);
 
 					// Update 'last node in sequence' var
 					LastNode = SetVarNode;
@@ -478,7 +475,7 @@ bool UK2Node_SpawnActorFromClass::HasExternalBlueprintDependencies(TArray<class
 	{
 		OptionalOutput->Add(SourceClass);
 	}
-	return bResult;
+	return bResult || Super::HasExternalBlueprintDependencies(OptionalOutput);
 }
 
 #undef LOCTEXT_NAMESPACE
\ No newline at end of file
diff --git a/Engine/Source/Editor/BlueprintGraph/Private/K2Node_StructOperation.cpp b/Engine/Source/Editor/BlueprintGraph/Private/K2Node_StructOperation.cpp
--- a/Engine/Source/Editor/BlueprintGraph/Private/K2Node_StructOperation.cpp
+++ b/Engine/Source/Editor/BlueprintGraph/Private/K2Node_StructOperation.cpp
@@ -1,7 +1,8 @@
 // Copyright 1998-2014 Epic Games, Inc. All Rights Reserved.
 
 #include "BlueprintGraphPrivatePCH.h"
-
+#include "Runtime/Engine/Classes/Engine/UserDefinedStruct.h"
+#include "Editor/UnrealEd/Public/Kismet2/StructureEditorUtils.h"
 
 //////////////////////////////////////////////////////////////////////////
 // UK2Node_StructOperation
@@ -11,15 +12,58 @@ UK2Node_StructOperation::UK2Node_StructOperation(const class FPostConstructIniti
 {
 }
 
-/*
-FLinearColor UK2Node_StructOperation::GetNodeTitleColor() const
+bool UK2Node_StructOperation::HasExternalUserDefinedStructDependencies(TArray<class UStruct*>* OptionalOutput) const
 {
-	const UEdGraphSchema_K2* Schema = GetDefault<UEdGraphSchema_K2>();
+	const bool bResult = StructType && StructType->IsA<UUserDefinedStruct>();
+	if (bResult && OptionalOutput)
+	{
+		OptionalOutput->Add(StructType);
+	}
+	return bResult || Super::HasExternalUserDefinedStructDependencies(OptionalOutput);
+}
 
-	FEdGraphPinType StructPinType;
-	StructPinType.PinCategory = Schema->PC_Struct;
-	StructPinType.PinSubCategoryObject = StructType;
+void UK2Node_StructOperation::FStructOperationOptionalPinManager::CustomizePinData(UEdGraphPin* Pin, FName SourcePropertyName, int32 ArrayIndex, UProperty* Property) const
+{
+	FOptionalPinManager::CustomizePinData(Pin, SourcePropertyName, ArrayIndex, Property);
 
-	return Schema->GetPinTypeColor(StructPinType);
+	if (Pin && Property)
+	{
+		const auto UDStructure = Cast<const UUserDefinedStruct>(Property->GetOwnerStruct());
+		if (UDStructure)
+		{
+			const auto VarDesc = FStructureEditorUtils::GetVarDesc(UDStructure).FindByPredicate(
+				FStructureEditorUtils::FFindByNameHelper<FStructVariableDescription>(Property->GetFName()));
+			if (VarDesc)
+			{
+				Pin->PersistentGuid = VarDesc->VarGuid;
+			}
+		}
+	}
 }
-*/
+
+bool UK2Node_StructOperation::DoRenamedPinsMatch(const UEdGraphPin* NewPin, const UEdGraphPin* OldPin, bool bStructInVaraiablesOut)
+{
+	bool bResult = false;
+	if (NewPin && OldPin && (OldPin->Direction == NewPin->Direction))
+	{
+		const EEdGraphPinDirection StructDirection = bStructInVaraiablesOut ? EGPD_Input : EGPD_Output;
+		const EEdGraphPinDirection VariablesDirection = bStructInVaraiablesOut ? EGPD_Output : EGPD_Input;
+		const UEdGraphSchema_K2* K2Schema = GetDefault<UEdGraphSchema_K2>(); 
+		const bool bCompatible = K2Schema && K2Schema->ArePinTypesCompatible(NewPin->PinType, OldPin->PinType);
+
+		if (bCompatible && (StructDirection == OldPin->Direction))
+		{
+			// Struct name was changed
+			bResult = true;
+		}
+		else if (bCompatible && (VariablesDirection == OldPin->Direction))
+		{
+			// name of a member variable was changed
+			if ((NewPin->PersistentGuid == OldPin->PersistentGuid) && OldPin->PersistentGuid.IsValid())
+			{
+				bResult = true;
+			}
+		}
+	}
+	return bResult;
+}
\ No newline at end of file
diff --git a/Engine/Source/Editor/BlueprintGraph/Private/K2Node_Variable.cpp b/Engine/Source/Editor/BlueprintGraph/Private/K2Node_Variable.cpp
--- a/Engine/Source/Editor/BlueprintGraph/Private/K2Node_Variable.cpp
+++ b/Engine/Source/Editor/BlueprintGraph/Private/K2Node_Variable.cpp
@@ -410,7 +410,7 @@ bool UK2Node_Variable::HasExternalBlueprintDependencies(TArray<class UStruct*>*
 	{
 		OptionalOutput->Add(SourceClass);
 	}
-	return bResult;
+	return bResult || Super::HasExternalBlueprintDependencies(OptionalOutput);
 }
 
 FString UK2Node_Variable::GetDocumentationLink() const
diff --git a/Engine/Source/Editor/Blutility/Private/BlutilityModule.cpp b/Engine/Source/Editor/Blutility/Private/BlutilityModule.cpp
--- a/Engine/Source/Editor/Blutility/Private/BlutilityModule.cpp
+++ b/Engine/Source/Editor/Blutility/Private/BlutilityModule.cpp
@@ -2,7 +2,7 @@
 
 #include "BlutilityPrivatePCH.h"
 #include "BlutilityClasses.h"
-#include "Blutility.generated.inl"
+
 
 #include "AssetTypeActions_EditorUtilityBlueprint.h"
 #include "AssetToolsModule.h"
diff --git a/Engine/Source/Editor/Cascade/Private/Cascade.cpp b/Engine/Source/Editor/Cascade/Private/Cascade.cpp
--- a/Engine/Source/Editor/Cascade/Private/Cascade.cpp
+++ b/Engine/Source/Editor/Cascade/Private/Cascade.cpp
@@ -13,7 +13,7 @@
 #include "SCascadeEmitterCanvas.h"
 #include "CascadeActions.h"
 #include "Cascade.h"
-#include "Cascade.generated.inl"
+
 
 #include "Editor/WorkspaceMenuStructure/Public/WorkspaceMenuStructureModule.h"
 #include "Editor/PropertyEditor/Public/PropertyEditorModule.h"
diff --git a/Engine/Source/Editor/ContentBrowser/Private/SAssetView.cpp b/Engine/Source/Editor/ContentBrowser/Private/SAssetView.cpp
--- a/Engine/Source/Editor/ContentBrowser/Private/SAssetView.cpp
+++ b/Engine/Source/Editor/ContentBrowser/Private/SAssetView.cpp
@@ -3552,7 +3552,7 @@ FText SAssetView::GetAssetShowWarningText() const
 	{
 		DropText = LOCTEXT( "DragAssetsHere", "Drag and drop assets here to add them to the collection." );
 	}
-	else
+	else if ( OnGetAssetContextMenu.IsBound() )
 	{
 		DropText = LOCTEXT( "DropFilesOrRightClick", "Drop files here or right click to create content." );
 	}
diff --git a/Engine/Source/Editor/DestructibleMeshEditor/Private/DestructibleMeshEditor.cpp b/Engine/Source/Editor/DestructibleMeshEditor/Private/DestructibleMeshEditor.cpp
--- a/Engine/Source/Editor/DestructibleMeshEditor/Private/DestructibleMeshEditor.cpp
+++ b/Engine/Source/Editor/DestructibleMeshEditor/Private/DestructibleMeshEditor.cpp
@@ -8,7 +8,7 @@
 #include "Editor/PropertyEditor/Public/PropertyEditorModule.h"
 #include "Editor/PropertyEditor/Public/IDetailsView.h"
 
-#include "DestructibleMeshEditor.generated.inl"
+
 
 #define LOCTEXT_NAMESPACE "DestructibleMeshEditor"
 
diff --git a/Engine/Source/Editor/EditorStyle/Private/EditorStyleModule.cpp b/Engine/Source/Editor/EditorStyle/Private/EditorStyleModule.cpp
--- a/Engine/Source/Editor/EditorStyle/Private/EditorStyleModule.cpp
+++ b/Engine/Source/Editor/EditorStyle/Private/EditorStyleModule.cpp
@@ -5,7 +5,7 @@
 =============================================================================*/
 
 #include "EditorStylePrivatePCH.h"
-#include "EditorStyle.generated.inl"
+
 
 
 /**
diff --git a/Engine/Source/Editor/EditorStyle/Private/SlateEditorStyle.cpp b/Engine/Source/Editor/EditorStyle/Private/SlateEditorStyle.cpp
--- a/Engine/Source/Editor/EditorStyle/Private/SlateEditorStyle.cpp
+++ b/Engine/Source/Editor/EditorStyle/Private/SlateEditorStyle.cpp
@@ -3064,6 +3064,7 @@ void FSlateEditorStyle::FStyle::SetupGraphEditorStyles()
 		Set( "GraphEditor.EventGlyph", new IMAGE_BRUSH( "Graph/Icons/Event", Icon22x22, FLinearColor::White) );
 		Set( "GraphEditor.EventCustomGlyph", new IMAGE_BRUSH( "Graph/Icons/Event_Custom", Icon22x22, FLinearColor::White) );
 		Set( "GraphEditor.SCSGlyph", new IMAGE_BRUSH( "Graph/Icons/Hammer", Icon22x22, FLinearColor::White) );
+		Set( "GraphEditor.StructGlyph", new IMAGE_BRUSH( "Graph/Icons/Struct", Icon22x22, FLinearColor::White) );
 		// Find In Blueprints
 		Set( "GraphEditor.FIB_CallFunction", new IMAGE_BRUSH( "Graph/Icons/FIB_CallFunction", Icon22x22, FLinearColor::White) );
 		Set( "GraphEditor.FIB_MacroInstance", new IMAGE_BRUSH( "Graph/Icons/FIB_MacroInstance", Icon22x22, FLinearColor::White) );
@@ -4066,6 +4067,7 @@ void FSlateEditorStyle::FStyle::SetupClassIconsAndThumbnails()
 		Set( "ClassIcon.TextureRenderTarget2D", new IMAGE_BRUSH( "Icons/AssetIcons/TextureRenderTarget2D_16x", Icon16x16 ) );
 		Set( "ClassIcon.TextureRenderTargetCube", new IMAGE_BRUSH( "Icons/AssetIcons/TextureRenderTargetCube_16x", Icon16x16 ) );
 		Set( "ClassIcon.UserDefinedEnum", new IMAGE_BRUSH( "Icons/AssetIcons/Enumeration_16x", Icon16x16 ) );
+		Set( "ClassIcon.UserDefinedStruct", new IMAGE_BRUSH( "Icons/AssetIcons/Structure_16x", Icon16x16 ) );
 		Set( "ClassIcon.SlateWidgetStyleAsset", new IMAGE_BRUSH("Icons/AssetIcons/SlateStyle_16x", Icon16x16 ) );
 
 		// Builder brushes
@@ -4161,6 +4163,7 @@ void FSlateEditorStyle::FStyle::SetupClassIconsAndThumbnails()
 		Set( "ClassThumbnail.TextureRenderTarget2D", new IMAGE_BRUSH( "Icons/icon_class_TextureRenderTarget2D_64x", Icon32x32 ) );
 		Set( "ClassThumbnail.TextureRenderTargetCube", new IMAGE_BRUSH( "Icons/icon_class_TextureRenderTargetCube_64x", Icon32x32 ) );
 		Set( "ClassThumbnail.UserDefinedEnum", new IMAGE_BRUSH( "Icons/icon_class_Enumeration_64x", Icon32x32 ) );
+		Set( "ClassThumbnail.UserDefinedStruct", new IMAGE_BRUSH( "Icons/icon_class_Structure_64x", Icon32x32 ) );
 		Set( "ClassThumbnail.SlateWidgetStyleAsset", new IMAGE_BRUSH("Icons/icon_class_SlateStyle_64x", Icon32x32 ) );
 	}
 #endif
diff --git a/Engine/Source/Editor/EnvironmentQueryEditor/EnvironmentQueryEditor.Build.cs b/Engine/Source/Editor/EnvironmentQueryEditor/EnvironmentQueryEditor.Build.cs
--- a/Engine/Source/Editor/EnvironmentQueryEditor/EnvironmentQueryEditor.Build.cs
+++ b/Engine/Source/Editor/EnvironmentQueryEditor/EnvironmentQueryEditor.Build.cs
@@ -8,7 +8,8 @@ public class EnvironmentQueryEditor : ModuleRules
 	{
         PrivateIncludePaths.AddRange(
             new string[] {
-				"Editor/GraphEditor/Private"
+				"Editor/GraphEditor/Private",
+				"Editor/EnvironmentQueryEditor/Private",
 			}
             );
 
diff --git a/Engine/Source/Editor/EnvironmentQueryEditor/Private/DetailCustomizations/EnvDirectionCustomization.cpp b/Engine/Source/Editor/EnvironmentQueryEditor/Private/DetailCustomizations/EnvDirectionCustomization.cpp
--- a/Engine/Source/Editor/EnvironmentQueryEditor/Private/DetailCustomizations/EnvDirectionCustomization.cpp
+++ b/Engine/Source/Editor/EnvironmentQueryEditor/Private/DetailCustomizations/EnvDirectionCustomization.cpp
@@ -1,6 +1,6 @@
 // Copyright 1998-2013 Epic Games, Inc. All Rights Reserved.
 
-#include "../EnvironmentQueryEditorPrivatePCH.h"
+#include "EnvironmentQueryEditorPrivatePCH.h"
 #include "EnvDirectionCustomization.h"
 
 #define LOCTEXT_NAMESPACE "FEnvQueryCustomization"
diff --git a/Engine/Source/Editor/EnvironmentQueryEditor/Private/DetailCustomizations/EnvQueryParamSetupCustomization.cpp b/Engine/Source/Editor/EnvironmentQueryEditor/Private/DetailCustomizations/EnvQueryParamSetupCustomization.cpp
--- a/Engine/Source/Editor/EnvironmentQueryEditor/Private/DetailCustomizations/EnvQueryParamSetupCustomization.cpp
+++ b/Engine/Source/Editor/EnvironmentQueryEditor/Private/DetailCustomizations/EnvQueryParamSetupCustomization.cpp
@@ -1,6 +1,6 @@
 // Copyright 1998-2014 Epic Games, Inc. All Rights Reserved.
 
-#include "../EnvironmentQueryEditorPrivatePCH.h"
+#include "EnvironmentQueryEditorPrivatePCH.h"
 #include "EnvQueryParamSetupCustomization.h"
 
 #define LOCTEXT_NAMESPACE "FEnvQueryCustomization"
diff --git a/Engine/Source/Editor/EnvironmentQueryEditor/Private/EnvironmentQueryEditorModule.cpp b/Engine/Source/Editor/EnvironmentQueryEditor/Private/EnvironmentQueryEditorModule.cpp
--- a/Engine/Source/Editor/EnvironmentQueryEditor/Private/EnvironmentQueryEditorModule.cpp
+++ b/Engine/Source/Editor/EnvironmentQueryEditor/Private/EnvironmentQueryEditorModule.cpp
@@ -11,7 +11,7 @@
 #include "Toolkits/ToolkitManager.h"
 #include "SGraphNode_EnvironmentQuery.h"
 #include "EdGraphUtilities.h"
-#include "EnvironmentQueryEditor.generated.inl"
+
 
 IMPLEMENT_MODULE( FEnvironmentQueryEditorModule, EnvironmentQueryEditor );
 DEFINE_LOG_CATEGORY(LogEnvironmentQueryEditor);
diff --git a/Engine/Source/Editor/GameProjectGeneration/Private/GameProjectGenerationModule.cpp b/Engine/Source/Editor/GameProjectGeneration/Private/GameProjectGenerationModule.cpp
--- a/Engine/Source/Editor/GameProjectGeneration/Private/GameProjectGenerationModule.cpp
+++ b/Engine/Source/Editor/GameProjectGeneration/Private/GameProjectGenerationModule.cpp
@@ -3,7 +3,7 @@
 #include "GameProjectGenerationPrivatePCH.h"
 #include "ModuleManager.h"
 #include "GameProjectGenerationModule.h"
-#include "GameProjectGeneration.generated.inl"
+
 
 
 IMPLEMENT_MODULE( FGameProjectGenerationModule, GameProjectGeneration );
diff --git a/Engine/Source/Editor/GameplayTagsEditor/Private/GameplayTagsEditorModule.cpp b/Engine/Source/Editor/GameplayTagsEditor/Private/GameplayTagsEditorModule.cpp
--- a/Engine/Source/Editor/GameplayTagsEditor/Private/GameplayTagsEditorModule.cpp
+++ b/Engine/Source/Editor/GameplayTagsEditor/Private/GameplayTagsEditorModule.cpp
@@ -3,7 +3,7 @@
 #include "GameplayTagsEditorModulePrivatePCH.h"
 #include "GameplayTagsGraphPanelPinFactory.h"
 #include "GameplayTagContainerCustomization.h"
-#include "GameplayTagsEditor.generated.inl"
+
 
 
 class FGameplayTagsEditorModule : public IGameplayTagsEditorModule
diff --git a/Engine/Source/Editor/GeometryMode/Private/GeometryEdMode.cpp b/Engine/Source/Editor/GeometryMode/Private/GeometryEdMode.cpp
--- a/Engine/Source/Editor/GeometryMode/Private/GeometryEdMode.cpp
+++ b/Engine/Source/Editor/GeometryMode/Private/GeometryEdMode.cpp
@@ -1,7 +1,7 @@
 // Copyright 1998-2014 Epic Games, Inc. All Rights Reserved.
 
 #include "GeometryModePrivatePCH.h"
-#include "GeometryMode.generated.inl"
+
 
 IMPLEMENT_MODULE( FGeometryModeModule, GeometryMode );
 
diff --git a/Engine/Source/Editor/GraphEditor/Private/KismetNodes/SGraphNodeK2CreateDelegate.cpp b/Engine/Source/Editor/GraphEditor/Private/KismetNodes/SGraphNodeK2CreateDelegate.cpp
--- a/Engine/Source/Editor/GraphEditor/Private/KismetNodes/SGraphNodeK2CreateDelegate.cpp
+++ b/Engine/Source/Editor/GraphEditor/Private/KismetNodes/SGraphNodeK2CreateDelegate.cpp
@@ -59,17 +59,14 @@ FString SGraphNodeK2CreateDelegate::GetCurrentFunctionDescription() const
 		return TEXT("None");
 	}
 
-	for(TFieldIterator<UFunction> It(ScopeClass); It; ++It)
+	if (const auto Func = FindField<UFunction>(ScopeClass, Node->GetFunctionName()))
 	{
-		if(*It && (It->GetFName() == Node->SelectedFunctionName))
-		{
-			return FunctionDescription(*It);
-		}
+		return FunctionDescription(Func);
 	}
 
-	if(Node->SelectedFunctionName != NAME_None)
+	if (Node->GetFunctionName() != NAME_None)
 	{
-		return FString::Printf(TEXT("Error? %"), *Node->SelectedFunctionName.ToString());
+		return FString::Printf(TEXT("Error? %s"), *Node->GetFunctionName().ToString());
 	}
 
 	return TEXT("Select Function");
@@ -90,7 +87,7 @@ void SGraphNodeK2CreateDelegate::OnFunctionSelected(TSharedPtr<FFunctionItemData
 	{
 		if(UK2Node_CreateDelegate* Node = Cast<UK2Node_CreateDelegate>(GraphNode))
 		{
-			Node->SelectedFunctionName = FunctionItemData->Name;
+			Node->SetFunction(FunctionItemData->Name);
 			Node->HandleAnyChange(true);
 
 			auto SelectFunctionWidgetPtr = SelectFunctionWidget.Pin();
diff --git a/Engine/Source/Editor/GraphEditor/Private/SGraphActionMenu.cpp b/Engine/Source/Editor/GraphEditor/Private/SGraphActionMenu.cpp
--- a/Engine/Source/Editor/GraphEditor/Private/SGraphActionMenu.cpp
+++ b/Engine/Source/Editor/GraphEditor/Private/SGraphActionMenu.cpp
@@ -22,6 +22,8 @@ namespace GraphActionMenuHelpers
 			((FEdGraphSchemaAction_K2Graph*)InGraphAction)->EdGraph->GetFName() == ItemName);
 		bCheck |= (InGraphAction->GetTypeId() == FEdGraphSchemaAction_K2Enum::StaticGetTypeId() &&
 			((FEdGraphSchemaAction_K2Enum*)InGraphAction)->GetPathName() == ItemName);
+		bCheck |= (InGraphAction->GetTypeId() == FEdGraphSchemaAction_K2Struct::StaticGetTypeId() &&
+			((FEdGraphSchemaAction_K2Struct*)InGraphAction)->GetPathName() == ItemName);
 		bCheck |= (InGraphAction->GetTypeId() == FEdGraphSchemaAction_K2Delegate::StaticGetTypeId() &&
 			((FEdGraphSchemaAction_K2Delegate*)InGraphAction)->GetDelegateName() == ItemName);
 		bCheck |= (InGraphAction->GetTypeId() == FEdGraphSchemaAction_K2TargetNode::StaticGetTypeId() &&
diff --git a/Engine/Source/Editor/Kismet/Private/BlueprintDetailsCustomization.cpp b/Engine/Source/Editor/Kismet/Private/BlueprintDetailsCustomization.cpp
--- a/Engine/Source/Editor/Kismet/Private/BlueprintDetailsCustomization.cpp
+++ b/Engine/Source/Editor/Kismet/Private/BlueprintDetailsCustomization.cpp
@@ -450,8 +450,8 @@ void FBlueprintVarActionDetails::CustomizeDetails( IDetailLayoutBuilder& DetailL
 			else  if (auto StructProperty = Cast<const UStructProperty>(OryginalProperty))
 			{
 				// Prevent editing the default value of a stale struct
-				auto BGStruct = Cast<const UBlueprintGeneratedStruct>(StructProperty->Struct);
-				if (BGStruct && (EBlueprintStructureStatus::BSS_UpToDate != BGStruct->Status))
+				auto BGStruct = Cast<const UUserDefinedStruct>(StructProperty->Struct);
+				if (BGStruct && (EUserDefinedStructureStatus::UDSS_UpToDate != BGStruct->Status))
 				{
 					VariableProperty = NULL;
 				}
@@ -809,13 +809,16 @@ FEdGraphPinType FBlueprintVarActionDetails::OnGetVarType() const
 
 void FBlueprintVarActionDetails::OnVarTypeChanged(const FEdGraphPinType& NewPinType)
 {
-	FName VarName = GetVariableName();
-	if (VarName != NAME_None)
+	if (FBlueprintEditorUtils::IsPinTypeValid(NewPinType))
 	{
-		FBlueprintEditorUtils::ChangeMemberVariableType(GetBlueprintObj(), VarName, NewPinType);
+		FName VarName = GetVariableName();
+		if (VarName != NAME_None)
+		{
+			FBlueprintEditorUtils::ChangeMemberVariableType(GetBlueprintObj(), VarName, NewPinType);
 
-		// Set the MyBP tab's last pin type used as this, for adding lots of variables of the same type
-		MyBlueprint.Pin()->GetLastPinTypeUsed() = NewPinType;
+			// Set the MyBP tab's last pin type used as this, for adding lots of variables of the same type
+			MyBlueprint.Pin()->GetLastPinTypeUsed() = NewPinType;
+		}
 	}
 }
 
@@ -1661,13 +1664,16 @@ FEdGraphPinType FBlueprintLocalVarActionDetails::OnGetVarType() const
 
 void FBlueprintLocalVarActionDetails::OnVarTypeChanged(const FEdGraphPinType& NewPinType)
 {
-	UK2Node_LocalVariable* LocalVar = GetDetailsSelection();
-	if(LocalVar)
+	if (FBlueprintEditorUtils::IsPinTypeValid(NewPinType))
 	{
-		LocalVar->ChangeVariableType(NewPinType);
+		UK2Node_LocalVariable* LocalVar = GetDetailsSelection();
+		if (LocalVar)
+		{
+			LocalVar->ChangeVariableType(NewPinType);
 
-		// Set the MyBP tab's last pin type used as this, for adding lots of variables of the same type
-		MyBlueprint.Pin()->GetLastPinTypeUsed() = NewPinType;
+			// Set the MyBP tab's last pin type used as this, for adding lots of variables of the same type
+			MyBlueprint.Pin()->GetLastPinTypeUsed() = NewPinType;
+		}
 	}
 }
 
@@ -1983,7 +1989,7 @@ void FBlueprintGraphArgumentLayout::OnRefCheckStateChanged(ESlateCheckBoxState::
 
 void FBlueprintGraphArgumentLayout::PinInfoChanged(const FEdGraphPinType& PinType)
 {
-	if (ParamItemPtr.IsValid())
+	if (ParamItemPtr.IsValid() && FBlueprintEditorUtils::IsPinTypeValid(PinType))
 	{
 		ParamItemPtr.Pin()->PinType = PinType;
 		if (GraphActionDetailsPtr.IsValid())
@@ -3605,6 +3611,7 @@ FReply FBlueprintGraphActionDetails::OnAddNewOutputClicked()
 	{
 		FunctionResultNode->Modify();
 		FEdGraphPinType PinType = MyBlueprint.Pin()->GetLastFunctionPinTypeUsed();
+		PinType.bIsReference = false;
 
 		// Make sure that if this is an exec node we are allowed one.
 		const UEdGraphSchema_K2* Schema = GetDefault<UEdGraphSchema_K2>();
diff --git a/Engine/Source/Editor/Kismet/Private/BlueprintEditor.cpp b/Engine/Source/Editor/Kismet/Private/BlueprintEditor.cpp
--- a/Engine/Source/Editor/Kismet/Private/BlueprintEditor.cpp
+++ b/Engine/Source/Editor/Kismet/Private/BlueprintEditor.cpp
@@ -45,7 +45,6 @@
 #include "SBlueprintActionMenu.h"
 #include "SMyBlueprint.h"
 #include "FindInBlueprints.h"
-#include "SUserDefinedStructureEditor.h"
 // End of core kismet tabs
 
 // Debugging
@@ -68,7 +67,7 @@
 #include "ISCSEditorCustomization.h"
 #include "Editor/UnrealEd/Public/SourceCodeNavigation.h"
 
-#include "Kismet.generated.inl"
+
 
 #define LOCTEXT_NAMESPACE "BlueprintEditor"
 
@@ -863,6 +862,37 @@ void FBlueprintEditor::EnsureBlueprintIsUpToDate(UBlueprint* BlueprintObj)
 	FBlueprintEditorUtils::UpdateTransactionalFlags(BlueprintObj);
 }
 
+struct FLoadObjectsFromAssetRegistryHelper
+{
+	template<class TObjectType>
+	static void Load(TSet<TWeakObjectPtr<TObjectType>>& Collection)
+	{
+		FAssetRegistryModule& AssetRegistryModule = FModuleManager::LoadModuleChecked<FAssetRegistryModule>(TEXT("AssetRegistry"));
+
+		const double CompileStartTime = FPlatformTime::Seconds();
+
+		TArray<FAssetData> AssetData;
+		AssetRegistryModule.Get().GetAssetsByClass(TObjectType::StaticClass()->GetFName(), AssetData);
+
+		for (int32 AssetIndex = 0; AssetIndex < AssetData.Num(); ++AssetIndex)
+		{
+			if(AssetData[AssetIndex].IsValid())
+			{
+				FString AssetPath = AssetData[AssetIndex].ObjectPath.ToString();
+				TObjectType* Object = LoadObject<TObjectType>(NULL, *AssetPath, NULL, 0, NULL);
+				if (Object)
+				{
+					Collection.Add( TWeakObjectPtr<TObjectType>(Object) );
+				}
+			}
+		}
+
+		const double FinishTime = FPlatformTime::Seconds();
+
+		UE_LOG(LogBlueprint, Log, TEXT("Loading all assets of type: %s took %.2f seconds"), *TObjectType::StaticClass()->GetName(), static_cast<float>(FinishTime - CompileStartTime));
+	}
+};
+
 void FBlueprintEditor::CommonInitialization(const TArray<UBlueprint*>& InitBlueprints)
 {
 	TSharedPtr<FBlueprintEditor> ThisPtr(SharedThis(this));
@@ -891,7 +921,7 @@ void FBlueprintEditor::CommonInitialization(const TArray<UBlueprint*>& InitBluep
 		// Load blueprint libraries
 		LoadLibrariesFromAssetRegistry();
 
-		LoadUserDefinedEnumsFromAssetRegistry();
+		FLoadObjectsFromAssetRegistryHelper::Load<UUserDefinedEnum>(UserDefinedEnumerators);
 
 		UBlueprint* InitBlueprint = InitBlueprints[0];
 
@@ -904,40 +934,11 @@ void FBlueprintEditor::CommonInitialization(const TArray<UBlueprint*>& InitBluep
 		InitBlueprint->bIsNewlyCreated = false;
 
 		// When the blueprint that we are observing changes, it will notify this wrapper widget.
-		InitBlueprint->OnChanged().AddSP( this, &FBlueprintEditor::OnBlueprintChanged );
+		InitBlueprint->OnChanged().AddSP(this, &FBlueprintEditor::OnBlueprintChanged);
 	}
 
 	CreateDefaultCommands();
 	CreateDefaultTabContents(InitBlueprints);
-	
-	if(FStructureEditorUtils::StructureEditingEnabled())
-	{
-		UserDefinedStructureEditor = SNew(SUserDefinedStructureEditor, SharedThis(this));
-	}
-}
-
-void FBlueprintEditor::LoadUserDefinedEnumsFromAssetRegistry()
-{
-	if( GetBlueprintObj() )
-	{
-		FAssetRegistryModule& AssetRegistryModule = FModuleManager::LoadModuleChecked<FAssetRegistryModule>(TEXT("AssetRegistry"));
-
-		TArray<FAssetData> AssetData;
-		AssetRegistryModule.Get().GetAssetsByClass(UUserDefinedEnum::StaticClass()->GetFName(), AssetData);
-
-		for (int32 AssetIndex = 0; AssetIndex < AssetData.Num(); ++AssetIndex)
-		{
-			if(AssetData[AssetIndex].IsValid())
-			{
-				FString AssetPath = AssetData[AssetIndex].ObjectPath.ToString();
-				UUserDefinedEnum* EnumObject = LoadObject<UUserDefinedEnum>(NULL, *AssetPath, NULL, 0, NULL);
-				if (EnumObject)
-				{
-					UserDefinedEnumerators.AddUnique( TWeakObjectPtr<UUserDefinedEnum>(EnumObject) );
-				}
-			}
-		}
-	}
 }
 
 void FBlueprintEditor::LoadLibrariesFromAssetRegistry()
@@ -2431,16 +2432,22 @@ void FBlueprintEditor::AddReferencedObjects( FReferenceCollector& Collector )
 	{
 		Collector.AddReferencedObject(StandardLibraries[Index]);
 	}
-	for (int32 Index = 0; Index < UserDefinedEnumerators.Num();)
+
+	UserDefinedEnumerators.Remove(TWeakObjectPtr<UUserDefinedEnum>()); // Remove NULLs
+	for (auto ObjectPtr : UserDefinedEnumerators)
 	{
-		if(UObject* EnumObj = UserDefinedEnumerators[Index].Get(false))
+		if(UObject* Obj = ObjectPtr.Get())
 		{
-			Collector.AddReferencedObject(EnumObj);
-			Index++;
+			Collector.AddReferencedObject(Obj);
 		}
-		else
+	}
+
+	UserDefinedStructures.Remove(TWeakObjectPtr<UUserDefinedStruct>()); // Remove NULLs
+	for (auto ObjectPtr : UserDefinedStructures)
+	{
+		if(UObject* Obj = ObjectPtr.Get())
 		{
-			UserDefinedEnumerators.RemoveAtSwap(Index);
+			Collector.AddReferencedObject(Obj);
 		}
 	}
 }
@@ -2925,12 +2932,15 @@ FEdGraphPinType FBlueprintEditor::OnGetPinType(UEdGraphPin* SelectedPin) const
 
 void FBlueprintEditor::OnChangePinTypeFinished(const FEdGraphPinType& PinType, UEdGraphPin* InSelectedPin)
 {
-	InSelectedPin->PinType = PinType;
-	if (UEdGraphPin* SelectedPin = GetCurrentlySelectedPin())
+	if (FBlueprintEditorUtils::IsPinTypeValid(PinType))
 	{
-		if (UK2Node_Select* SelectNode = Cast<UK2Node_Select>(SelectedPin->GetOwningNode()))
+		InSelectedPin->PinType = PinType;
+		if (UEdGraphPin* SelectedPin = GetCurrentlySelectedPin())
 		{
-			SelectNode->ChangePinType(SelectedPin);
+			if (UK2Node_Select* SelectNode = Cast<UK2Node_Select>(SelectedPin->GetOwningNode()))
+			{
+				SelectNode->ChangePinType(SelectedPin);
+			}
 		}
 	}
 
diff --git a/Engine/Source/Editor/Kismet/Private/BlueprintEditorModes.cpp b/Engine/Source/Editor/Kismet/Private/BlueprintEditorModes.cpp
--- a/Engine/Source/Editor/Kismet/Private/BlueprintEditorModes.cpp
+++ b/Engine/Source/Editor/Kismet/Private/BlueprintEditorModes.cpp
@@ -60,11 +60,6 @@ FBlueprintEditorApplicationMode::FBlueprintEditorApplicationMode(TSharedPtr<clas
 	{
 		BlueprintEditorTabFactories.RegisterFactory(MakeShareable(new FDefaultsEditorSummoner(InBlueprintEditor)));
 	}
-	if(FStructureEditorUtils::StructureEditingEnabled())
-	{
-		BlueprintEditorTabFactories.RegisterFactory(MakeShareable(new FUserDefinedStructureEditorSummoner(InBlueprintEditor)));
-	}
-
 	CoreTabFactories.RegisterFactory(MakeShareable(new FSelectionDetailsSummoner(InBlueprintEditor)));
 
 	TabLayout = FTabManager::NewLayout( "Standalone_BlueprintEditor_Layout_v5" )
diff --git a/Engine/Source/Editor/Kismet/Private/BlueprintEditorModule.cpp b/Engine/Source/Editor/Kismet/Private/BlueprintEditorModule.cpp
--- a/Engine/Source/Editor/Kismet/Private/BlueprintEditorModule.cpp
+++ b/Engine/Source/Editor/Kismet/Private/BlueprintEditorModule.cpp
@@ -14,6 +14,7 @@
 #include "Developer/MessageLog/Public/MessageLogModule.h"
 #include "UObjectToken.h"
 #include "InstancedStaticMeshSCSEditorCustomization.h"
+#include "UserDefinedStructureEditor.h"
 
 #define LOCTEXT_NAMESPACE "BlueprintEditor"
 
@@ -172,6 +173,13 @@ TSharedRef<IUserDefinedEnumEditor> FBlueprintEditorModule::CreateUserDefinedEnum
 	return UserDefinedEnumEditor;
 }
 
+TSharedRef<IUserDefinedStructureEditor> FBlueprintEditorModule::CreateUserDefinedStructEditor(const EToolkitMode::Type Mode, const TSharedPtr<IToolkitHost>& InitToolkitHost, UUserDefinedStruct* UDStruct)
+{
+	TSharedRef<FUserDefinedStructureEditor> UserDefinedStructureEditor(new FUserDefinedStructureEditor());
+	UserDefinedStructureEditor->InitEditor(Mode, InitToolkitHost, UDStruct);
+	return UserDefinedStructureEditor;
+}
+
 void FBlueprintEditorModule::RegisterSCSEditorCustomization(const FName& InComponentName, FSCSEditorCustomizationBuilder InCustomizationBuilder)
 {
 	SCSEditorCustomizations.Add(InComponentName, InCustomizationBuilder);
diff --git a/Engine/Source/Editor/Kismet/Private/BlueprintEditorTabFactories.h b/Engine/Source/Editor/Kismet/Private/BlueprintEditorTabFactories.h
--- a/Engine/Source/Editor/Kismet/Private/BlueprintEditorTabFactories.h
+++ b/Engine/Source/Editor/Kismet/Private/BlueprintEditorTabFactories.h
@@ -4,8 +4,6 @@
 
 #define LOCTEXT_NAMESPACE "BlueprintEditor"
 
-#include "SUserDefinedStructureEditor.h"
-
 /////////////////////////////////////////////////////
 // FLocalKismetCallbacks
 
@@ -111,32 +109,6 @@ public:
 	}
 };
 
-/////////////////////////////////////////////////////
-// FUserDefinedStructureEditorSummoner
-
-struct FUserDefinedStructureEditorSummoner : public FWorkflowTabFactory
-{
-public:
-	FUserDefinedStructureEditorSummoner(TSharedPtr<class FAssetEditorToolkit> InHostingApp)
-		: FWorkflowTabFactory(FBlueprintEditorTabs::UserDefinedStructureID, InHostingApp)
-	{
-		TabLabel = LOCTEXT("BlueprintUserDefinedStructureTabTitle", "User Defined Structure");
-		TabIcon = FEditorStyle::GetBrush("LevelEditor.Tabs.Details");
-
-		bIsSingleton = true;
-
-		ViewMenuDescription = LOCTEXT("UserDefinedStructureEditorView", "Structures");
-		ViewMenuTooltip = LOCTEXT("UserDefinedStructureEditorView_ToolTip", "Shows the user defined structures editor view");
-	}
-
-	virtual TSharedRef<SWidget> CreateTabBody(const FWorkflowTabSpawnInfo& Info) const OVERRIDE
-	{
-		TSharedPtr<FBlueprintEditor> BlueprintEditorPtr = StaticCastSharedPtr<FBlueprintEditor>(HostingApp.Pin());
-
-		return BlueprintEditorPtr->GetUserDefinedStructureEditor();
-	}
-};
-
 /////////////////////////////////////////////////////
 // FGraphTabHistory
 
diff --git a/Engine/Source/Editor/Kismet/Private/SBlueprintPalette.cpp b/Engine/Source/Editor/Kismet/Private/SBlueprintPalette.cpp
--- a/Engine/Source/Editor/Kismet/Private/SBlueprintPalette.cpp
+++ b/Engine/Source/Editor/Kismet/Private/SBlueprintPalette.cpp
@@ -44,9 +44,10 @@ namespace BlueprintPalette
  * 
  * @param  VarClass	The class that owns the variable in question.
  * @param  VarName	The name of the variable you want the type of.
+ * @param  Detailed	If true the returned string includes SubCategoryObject
  * @return A string representing the variable's type (empty if the variable couldn't be found).
  */
-static FString GetVarType(UClass* VarClass, FName VarName, bool bUseObjToolTip)
+static FString GetVarType(UClass* VarClass, FName VarName, bool bUseObjToolTip, bool bDetailed = false)
 {
 	FString VarDesc;
 
@@ -70,7 +71,14 @@ static FString GetVarType(UClass* VarClass, FName VarName, bool bUseObjToolTip)
 				FEdGraphPinType PinType;
 				if (K2Schema->ConvertPropertyToPinType(Property, PinType)) // use schema to get the color
 				{
-					VarDesc = UEdGraphSchema_K2::TypeToString(PinType);
+					if (bDetailed && PinType.PinSubCategoryObject.IsValid())
+					{
+						VarDesc = FString::Printf(TEXT("%s '%s'"), *PinType.PinCategory, *PinType.PinSubCategoryObject->GetName());
+					}
+					else
+					{
+						VarDesc = UEdGraphSchema_K2::TypeToString(PinType);
+					}
 				}
 			}
 		}
@@ -316,16 +324,21 @@ static void GetPaletteItemIcon(TSharedPtr<FEdGraphSchemaAction> ActionIn, UBluep
 
 		UClass* VarClass = VarAction->GetVariableClass();
 		BrushOut = FBlueprintEditor::GetVarIconAndColor(VarClass, VarAction->GetVariableName(), ColorOut);
-		ToolTipOut = GetVarType(VarClass, VarAction->GetVariableName(), true);
+		ToolTipOut = GetVarType(VarClass, VarAction->GetVariableName(), true, true);
 
 		DocLinkOut = TEXT("Shared/Editor/Blueprint/VariableTypes");
-		DocExcerptOut = GetVarType(VarClass, VarAction->GetVariableName(), false);
+		DocExcerptOut = GetVarType(VarClass, VarAction->GetVariableName(), false, false);
 	}
 	else if(ActionIn->GetTypeId() == FEdGraphSchemaAction_K2Enum::StaticGetTypeId())
 	{
 		BrushOut = FEditorStyle::GetBrush(TEXT("GraphEditor.EnumGlyph"));
 		ToolTipOut = FString::Printf( *LOCTEXT("Enum_Tooltip", "Enum Asset").ToString() );
 	}
+	else if(ActionIn->GetTypeId() == FEdGraphSchemaAction_K2Struct::StaticGetTypeId())
+	{
+		BrushOut = FEditorStyle::GetBrush(TEXT("GraphEditor.StructGlyph"));
+		ToolTipOut =  FString::Printf( *LOCTEXT("Struct_Tooltip", "Struct Asset").ToString() );
+	}
 	else if (ActionIn->GetTypeId() == FEdGraphSchemaAction_K2AddComment::StaticGetTypeId() ||
 		ActionIn->GetTypeId() == FEdGraphSchemaAction_NewStateComment::StaticGetTypeId())
 	{
@@ -382,30 +395,22 @@ static TSharedRef<SToolTip> ConstructToolTipWithActionPath(TSharedPtr<FEdGraphSc
 /** A set of utilities to aid SBlueprintPaletteItem when the user attempts to rename one. */
 struct FBlueprintPaletteItemRenameUtils
 {
-	/**
-	 * Determines whether the enum node, associated with the selected action, 
-	 * can be renamed with the specified text.
-	 * 
-	 * @param  InNewText		The text you want to verify.
-	 * @param  OutErrorMessage	Text explaining why the associated node couldn't be renamed (if the return value is false).
-	 * @param  ActionPtr		The selected action that the calling palette item represents.
-	 * @return True if it is ok to rename the associated node with the given string (false if not).
-	 */
-	static bool VerifyNewEnumName(const FText& InNewText, FText& OutErrorMessage, TWeakPtr<FEdGraphSchemaAction> ActionPtr)
+private:
+	static bool VerifyNewAssetName(UObject* Object, const FText& InNewText, FText& OutErrorMessage)
 	{
-		// Should never make it here with anything but an enum action
-		check(ActionPtr.Pin()->GetTypeId() == FEdGraphSchemaAction_K2Enum::StaticGetTypeId());
-
-		FEdGraphSchemaAction_K2Enum* EnumAction = (FEdGraphSchemaAction_K2Enum*)ActionPtr.Pin().Get();
+		if (!Object)
+		{
+			return false;
+		}
 
-		if(EnumAction->Enum->GetName() == InNewText.ToString())
+		if (Object->GetName() == InNewText.ToString())
 		{
 			return true;
 		}
 
 		TArray<FAssetData> AssetData;
 		FAssetRegistryModule& AssetToolsModule = FModuleManager::LoadModuleChecked<FAssetRegistryModule>("AssetRegistry");
-		AssetToolsModule.Get().GetAssetsByPath(FName(*FPaths::GetPath(EnumAction->Enum->GetOutermost()->GetPathName())), AssetData);
+		AssetToolsModule.Get().GetAssetsByPath(FName(*FPaths::GetPath(Object->GetOutermost()->GetPathName())), AssetData);
 
 		if(!FEditorFileUtils::IsFilenameValidForSaving(InNewText.ToString(), OutErrorMessage) || !FName(*InNewText.ToString()).IsValidObjectName( OutErrorMessage ))
 		{
@@ -431,6 +436,48 @@ struct FBlueprintPaletteItemRenameUtils
 		return true;
 	}
 
+	static void CommitNewAssetName(UObject* Object, FBlueprintEditor* BlueprintEditor, const FText& NewText)
+	{
+		if (Object && BlueprintEditor)
+		{
+			if(Object->GetName() != NewText.ToString())
+			{
+				TArray<FAssetRenameData> AssetsAndNames;
+				const FString PackagePath = FPackageName::GetLongPackagePath(Object->GetOutermost()->GetName());
+				new(AssetsAndNames) FAssetRenameData(Object, PackagePath, NewText.ToString());
+
+				FAssetToolsModule& AssetToolsModule = FModuleManager::LoadModuleChecked<FAssetToolsModule>("AssetTools");
+				AssetToolsModule.Get().RenameAssets(AssetsAndNames);
+			}
+
+			auto MyBlueprint = BlueprintEditor->GetMyBlueprintWidget();
+			if (MyBlueprint.IsValid())
+			{
+				MyBlueprint->SelectItemByName(FName(*Object->GetPathName()));
+			}	
+		}
+	}
+
+public:
+	/**
+	 * Determines whether the enum node, associated with the selected action, 
+	 * can be renamed with the specified text.
+	 * 
+	 * @param  InNewText		The text you want to verify.
+	 * @param  OutErrorMessage	Text explaining why the associated node couldn't be renamed (if the return value is false).
+	 * @param  ActionPtr		The selected action that the calling palette item represents.
+	 * @return True if it is ok to rename the associated node with the given string (false if not).
+	 */
+	static bool VerifyNewEnumName(const FText& InNewText, FText& OutErrorMessage, TWeakPtr<FEdGraphSchemaAction> ActionPtr)
+	{
+		// Should never make it here with anything but an enum action
+		check(ActionPtr.Pin()->GetTypeId() == FEdGraphSchemaAction_K2Enum::StaticGetTypeId());
+
+		auto EnumAction = (FEdGraphSchemaAction_K2Enum*)ActionPtr.Pin().Get();
+		
+		return VerifyNewAssetName(EnumAction ? EnumAction->Enum : NULL, InNewText, OutErrorMessage);
+	}
+
 	/**
 	 * Take the verified text and renames the enum node associated with the 
 	 * selected action.
@@ -462,6 +509,25 @@ struct FBlueprintPaletteItemRenameUtils
 		BlueprintEditorPtr.Pin()->GetMyBlueprintWidget()->SelectItemByName(FName(*EnumAction->Enum->GetPathName()));
 	}
 
+	/**
+	* Determines whether the struct node, associated with the selected action,
+	* can be renamed with the specified text.
+	*
+	* @param  InNewText		The text you want to verify.
+	* @param  OutErrorMessage	Text explaining why the associated node couldn't be renamed (if the return value is false).
+	* @param  ActionPtr		The selected action that the calling palette item represents.
+	* @return True if it is ok to rename the associated node with the given string (false if not).
+	*/
+	static bool VerifyNewStructName(const FText& InNewText, FText& OutErrorMessage, TWeakPtr<FEdGraphSchemaAction> ActionPtr)
+	{
+		// Should never make it here with anything but a struct action
+		check(ActionPtr.Pin()->GetTypeId() == FEdGraphSchemaAction_K2Struct::StaticGetTypeId());
+
+		auto Action = (FEdGraphSchemaAction_K2Struct*)ActionPtr.Pin().Get();
+
+		return VerifyNewAssetName(Action ? Action->Struct : NULL, InNewText, OutErrorMessage);
+	}
+
 	/**
 	 * Determines whether the event node, associated with the selected action, 
 	 * can be renamed with the specified text.
@@ -488,6 +554,25 @@ struct FBlueprintPaletteItemRenameUtils
 		return bIsNameValid;
 	}
 
+	/**
+	 * Take the verified text and renames the struct node associated with the 
+	 * selected action.
+	 * 
+	 * @param  NewText				The new (verified) text to rename the node with.
+	 * @param  InTextCommit			A value denoting how the text was entered.
+	 * @param  ActionPtr			The selected action that the calling palette item represents.
+	 * @param  BlueprintEditorPtr	A pointer to the blueprint editor that the palette belongs to.
+	 */
+	static void CommitNewStructName(const FText& NewText, ETextCommit::Type InTextCommit, TWeakPtr<FEdGraphSchemaAction> ActionPtr, TWeakPtr<FBlueprintEditor> BlueprintEditorPtr)
+	{
+		// Should never make it here with anything but a struct action
+		check(ActionPtr.Pin()->GetTypeId() == FEdGraphSchemaAction_K2Struct::StaticGetTypeId());
+
+		FEdGraphSchemaAction_K2Struct* Action = (FEdGraphSchemaAction_K2Struct*)ActionPtr.Pin().Get();
+
+		CommitNewAssetName(Action ? Action->Struct : NULL, BlueprintEditorPtr.Pin().Get(), NewText);
+	}
+
 	/**
 	 * Take the verified text and renames the event node associated with the 
 	 * selected action.
@@ -821,6 +906,11 @@ TSharedRef<SWidget> SBlueprintPaletteItem::CreateTextSlotWidget(const FSlateFont
 		OnVerifyTextChanged.BindStatic(&FBlueprintPaletteItemRenameUtils::VerifyNewEnumName, ActionPtr);
 		OnTextCommitted.BindStatic(&FBlueprintPaletteItemRenameUtils::CommitNewEnumName, ActionPtr, BlueprintEditorPtr);
 	}
+	else if (ActionTypeId == FEdGraphSchemaAction_K2Struct::StaticGetTypeId())
+	{
+		OnVerifyTextChanged.BindStatic(&FBlueprintPaletteItemRenameUtils::VerifyNewStructName, ActionPtr );
+		OnTextCommitted.BindStatic(&FBlueprintPaletteItemRenameUtils::CommitNewStructName, ActionPtr, BlueprintEditorPtr);
+	}
 	else if (ActionTypeId == FEdGraphSchemaAction_K2Event::StaticGetTypeId())
 	{
 		OnVerifyTextChanged.BindStatic(&FBlueprintPaletteItemRenameUtils::VerifyNewEventName, ActionPtr);
@@ -897,6 +987,11 @@ FText SBlueprintPaletteItem::GetDisplayText() const
 		FEdGraphSchemaAction_K2Enum* EnumAction = (FEdGraphSchemaAction_K2Enum*)GraphAction.Get();
 		DisplayText = FText::FromString(EnumAction->Enum->GetName());
 	}
+	else if (GraphAction->GetTypeId() == FEdGraphSchemaAction_K2Struct::StaticGetTypeId())
+	{
+		FEdGraphSchemaAction_K2Struct* StructAction = (FEdGraphSchemaAction_K2Struct*)GraphAction.Get();
+		DisplayText = StructAction->Struct ? FText::FromString(StructAction->Struct->GetName()) : FText::FromString(TEXT("None"));
+	}
 	else
 	{
 		DisplayText = ActionPtr.Pin()->MenuDescription;
@@ -1172,7 +1267,7 @@ FText SBlueprintPaletteItem::GetToolTipText() const
 			{
 				FString Result = GetVarTooltip(BlueprintEditorPtr.Pin()->GetBlueprintObj(), VarClass, VarAction->GetVariableName());
 				// Only use the variable tooltip if it has been filled out.
-				ToolTipText = !Result.IsEmpty() ? Result : GetVarType(VarClass, VarAction->GetVariableName(), true);
+				ToolTipText = !Result.IsEmpty() ? Result : GetVarType(VarClass, VarAction->GetVariableName(), true, true);
 			}
 		}
 		else if (PaletteAction->GetTypeId() == FEdGraphSchemaAction_K2Delegate::StaticGetTypeId())
diff --git a/Engine/Source/Editor/Kismet/Private/SKismetInspector.cpp b/Engine/Source/Editor/Kismet/Private/SKismetInspector.cpp
--- a/Engine/Source/Editor/Kismet/Private/SKismetInspector.cpp
+++ b/Engine/Source/Editor/Kismet/Private/SKismetInspector.cpp
@@ -16,6 +16,7 @@
 #include "SMyBlueprint.h"
 #include "BlueprintDetailsCustomization.h"
 #include "UserDefinedEnumEditor.h"
+#include "UserDefinedStructureEditor.h"
 #include "FormatTextDetails.h"
 
 #define LOCTEXT_NAMESPACE "KismetInspector"
@@ -235,6 +236,9 @@ void SKismetInspector::Construct(const FArguments& InArgs)
 		FOnGetDetailCustomizationInstance LayoutEnumDetails = FOnGetDetailCustomizationInstance::CreateStatic(&FEnumDetails::MakeInstance);
 		PropertyView->RegisterInstancedCustomPropertyLayout(UUserDefinedEnum::StaticClass(), LayoutEnumDetails);
 
+		FOnGetDetailCustomizationInstance LayoutStructDetails = FOnGetDetailCustomizationInstance::CreateStatic(&FUserDefinedStructureDetails::MakeInstance);
+		PropertyView->RegisterInstancedCustomPropertyLayout(UUserDefinedStruct::StaticClass(), LayoutStructDetails);
+
 		FOnGetDetailCustomizationInstance LayoutFormatTextDetails = FOnGetDetailCustomizationInstance::CreateStatic(&FFormatTextDetails::MakeInstance);
 		PropertyView->RegisterInstancedCustomPropertyLayout(UK2Node_FormatText::StaticClass(), LayoutFormatTextDetails);
 
diff --git a/Engine/Source/Editor/Kismet/Private/SMyBlueprint.cpp b/Engine/Source/Editor/Kismet/Private/SMyBlueprint.cpp
--- a/Engine/Source/Editor/Kismet/Private/SMyBlueprint.cpp
+++ b/Engine/Source/Editor/Kismet/Private/SMyBlueprint.cpp
@@ -30,6 +30,7 @@
 #include "SBlueprintEditorToolbar.h"
 
 #include "Editor/UnrealEd/Public/Kismet2/BlueprintEditorUtils.h"
+#include "Editor/UnrealEd/Public/Kismet2/StructureEditorUtils.h"
 #include "ObjectEditorUtils.h"
 #include "AssetToolsModule.h"
 #include "Editor/GraphEditor/Private/GraphActionNode.h"
@@ -52,6 +53,9 @@ void FMyBlueprintCommands::RegisterCommands()
 	UI_COMMAND( FindUserDefinedEnumInContentBrowser, "Find in Content Browser...", "Find user defined enum in content browser...", EUserInterfaceActionType::Button, FInputGesture() );
 	UI_COMMAND( AddNewUserDefinedEnum, "Create Enum Asset", "Create new user defined enum asset", EUserInterfaceActionType::Button, FInputGesture() );
 	UI_COMMAND( GotoNativeVarDefinition, "Goto Code Definition", "Goto the native code definition of this variable", EUserInterfaceActionType::Button, FInputGesture() );
+	UI_COMMAND( FindUserDefinedStructInContentBrowser, "Find in Content Browser...", "Find user defined struct in content browser...", EUserInterfaceActionType::Button, FInputGesture() );
+	UI_COMMAND( AddNewUserDefinedStruct, "Create Struct Asset", "Create new user defined struct asset", EUserInterfaceActionType::Button, FInputGesture() );
+
 }
 
 //////////////////////////////////////////////////////////////////////////
@@ -168,8 +172,14 @@ void SMyBlueprint::Construct(const FArguments& InArgs, TWeakPtr<FBlueprintEditor
 	InBlueprintEditor.Pin()->GetToolkitCommands()->MapAction( FMyBlueprintCommands::Get().FindUserDefinedEnumInContentBrowser,
 		FExecuteAction::CreateSP(this, &SMyBlueprint::OnFindUserDefinedEnumInContentBrowser) );
 
+	InBlueprintEditor.Pin()->GetToolkitCommands()->MapAction( FMyBlueprintCommands::Get().FindUserDefinedStructInContentBrowser,
+		FExecuteAction::CreateSP(this, &SMyBlueprint::OnFindUserDefinedStructInContentBrowser) );
+
 	InBlueprintEditor.Pin()->GetToolkitCommands()->MapAction( FMyBlueprintCommands::Get().AddNewUserDefinedEnum,
 		FExecuteAction::CreateSP(this, &SMyBlueprint::AddNewUserDefinedEnum) );
+
+	InBlueprintEditor.Pin()->GetToolkitCommands()->MapAction( FMyBlueprintCommands::Get().AddNewUserDefinedStruct,
+		FExecuteAction::CreateSP(this, &SMyBlueprint::AddNewUserDefinedStruct) );
 	
 	ToolKitCommandList->MapAction( FMyBlueprintCommands::Get().DeleteEntry,
 		FExecuteAction::CreateSP(this, &SMyBlueprint::OnDeleteEntry),
@@ -342,6 +352,9 @@ FText SMyBlueprint::OnGetSectionTitle( int32 InSectionID )
 	case NodeSectionID::LOCAL_VARIABLE:
 		SeperatorTitle = NSLOCTEXT("GraphActionNode", "LocalVariables", "Local Variables");
 		break;	
+	case NodeSectionID::USER_STRUCT:
+		SeperatorTitle = NSLOCTEXT("GraphActionNode", "Userstructs", "User Structs");
+		break;	
 	default:
 	case NodeSectionID::NONE:
 		SeperatorTitle = FText::GetEmpty();
@@ -585,10 +598,111 @@ EVisibility SMyBlueprint::GetLocalActionsListVisibility() const
 	return EVisibility::Collapsed;
 }
 
-static FString UserDefinedEnumCategoryNameStr()
+struct FActionCollectingHelper
 {
-	return LOCTEXT("UserDefinedEnumCategory", "User defined enums").ToString();
-}
+public:
+	static FString UserDefinedEnumCategoryNameStr()
+	{
+		return LOCTEXT("UserDefinedEnumCategory", "User defined enums").ToString();
+	}
+
+	static FString UserDefinedStructCategoryNameStr()
+	{
+		return LOCTEXT("UserDefinedStructCategory", "User defined structs").ToString();
+	}
+
+	static void AddUserEnumsActions(FGraphActionListBuilderBase& OutAllActions, TArray<UUserDefinedEnum*>& ElementsUsedInBlueprint)
+	{
+		for(auto Element : ElementsUsedInBlueprint)
+		{
+			if(Element)
+			{
+				const FString EnumNameStr = Element->GetName();
+				TSharedPtr<FEdGraphSchemaAction_K2Enum> NewEnumAction = MakeShareable(new FEdGraphSchemaAction_K2Enum( 
+					UserDefinedEnumCategoryNameStr(), FText::FromString(EnumNameStr), EnumNameStr, -1));
+				NewEnumAction->Enum = Element;
+				NewEnumAction->SectionID = NodeSectionID::USER_ENUM;
+				OutAllActions.AddAction(NewEnumAction);
+			}
+		}
+	}
+
+	static void AddUserStructActions(FGraphActionListBuilderBase& OutAllActions, TArray<UUserDefinedStruct*>& ElementsUsedInBlueprint)
+	{
+		for(auto Element : ElementsUsedInBlueprint)
+		{
+			if(Element)
+			{
+				const FString NameStr = Element->GetName();
+				TSharedPtr<FEdGraphSchemaAction_K2Struct> NewEnumAction = MakeShareable(new FEdGraphSchemaAction_K2Struct( 
+					UserDefinedStructCategoryNameStr(), FText::FromString(NameStr), NameStr, -1));
+				NewEnumAction->Struct = Element;
+				NewEnumAction->SectionID = NodeSectionID::USER_STRUCT;
+				OutAllActions.AddAction(NewEnumAction);
+			}
+		}
+	}
+
+	static UUserDefinedEnum* FromProperty(UByteProperty* Property)
+	{
+		return Property ? Cast<UUserDefinedEnum>(Property->GetIntPropertyEnum()) : NULL;
+	}
+
+	static UUserDefinedStruct* FromProperty(UStructProperty* Property)
+	{
+		return Property ? Cast<UUserDefinedStruct>(Property->Struct) : NULL;
+	}
+
+	template<class TClassOfType, class TPropertyType>
+	static void CollectAllUsedTypes(FGraphActionListBuilderBase& OutAllActions, TArray<TClassOfType*>& ElementsUsedInBlueprint)
+	{
+		for ( int32 GraphActionIndex = 0; GraphActionIndex < OutAllActions.GetNumActions(); GraphActionIndex++ )
+		{
+			FGraphActionListBuilderBase::ActionGroup& CurrentGraphAction = OutAllActions.GetAction(GraphActionIndex);
+			for(auto ActionIter = CurrentGraphAction.Actions.CreateIterator(); ActionIter; ActionIter++)
+			{
+				TSharedPtr<FEdGraphSchemaAction> Action = *ActionIter;
+				if(Action.IsValid())
+				{
+					if(FEdGraphSchemaAction_K2Var::StaticGetTypeId() == Action->GetTypeId())
+					{
+						FEdGraphSchemaAction_K2Var* VarAction = (FEdGraphSchemaAction_K2Var*)Action.Get();
+						if(auto Property = Cast<TPropertyType>(VarAction->GetProperty()))
+						{
+							if(auto Element = FromProperty(Property))
+							{
+								ElementsUsedInBlueprint.AddUnique(Element);
+							}
+						}
+					}
+					else if(FEdGraphSchemaAction_K2Graph::StaticGetTypeId() == Action->GetTypeId())
+					{
+						FEdGraphSchemaAction_K2Graph* GraphAction = (FEdGraphSchemaAction_K2Graph*)Action.Get();
+						if(GraphAction->EdGraph)
+						{
+							for(auto NodeIter = GraphAction->EdGraph->Nodes.CreateIterator(); NodeIter; NodeIter++)
+							{
+								if(UEdGraphNode* Node = *NodeIter)
+								{
+									for(auto PinIter = Node->Pins.CreateIterator(); PinIter; PinIter++)
+									{
+										if(UEdGraphPin* Pin = *PinIter)
+										{
+											if(auto Element =  Cast<TClassOfType>(Pin->PinType.PinSubCategoryObject.Get()))
+											{
+												ElementsUsedInBlueprint.AddUnique(Element);
+											}
+										}
+									}
+								}
+							}
+						}
+					}
+				}
+			}
+		}
+	}
+};
 
 void SMyBlueprint::CollectAllActions(FGraphActionListBuilderBase& OutAllActions)
 {
@@ -838,71 +952,39 @@ void SMyBlueprint::CollectAllActions(FGraphActionListBuilderBase& OutAllActions)
 		GetChildEvents(Graph, NewFuncAction->SectionID, OutAllActions);
 	}
 
-	TArray<UUserDefinedEnum*> UserDefinedEnumUserInBlueprint;
-	for ( int32 GraphActionIndex = 0; GraphActionIndex < OutAllActions.GetNumActions(); GraphActionIndex++ )
 	{
-		FGraphActionListBuilderBase::ActionGroup& CurrentGraphAction = OutAllActions.GetAction(GraphActionIndex);
-		for(auto ActionIter = CurrentGraphAction.Actions.CreateIterator(); ActionIter; ActionIter++)
+		TArray<UUserDefinedEnum*> UserDefinedEnumUserInBlueprint;
+		FActionCollectingHelper::CollectAllUsedTypes<UUserDefinedEnum, UByteProperty>(OutAllActions, UserDefinedEnumUserInBlueprint);
+		for(auto EnumWeakIter = EnumsAddedToBlueprint.CreateIterator(); EnumWeakIter; EnumWeakIter++)
 		{
-			TSharedPtr<FEdGraphSchemaAction> Action = *ActionIter;
-			if(Action.IsValid())
+			if(EnumWeakIter->IsValid())
 			{
-				if(FEdGraphSchemaAction_K2Var::StaticGetTypeId() == Action->GetTypeId())
-				{
-					FEdGraphSchemaAction_K2Var* VarAction = (FEdGraphSchemaAction_K2Var*)Action.Get();
-					if(UByteProperty* ByteProperty = Cast<UByteProperty>(VarAction->GetProperty()))
-					{
-						if(UUserDefinedEnum* UserDefinedEnum =  Cast<UUserDefinedEnum>(ByteProperty->GetIntPropertyEnum()))
-						{
-							UserDefinedEnumUserInBlueprint.AddUnique(UserDefinedEnum);
-						}
-					}
-				}
-				else if(FEdGraphSchemaAction_K2Graph::StaticGetTypeId() == Action->GetTypeId())
-				{
-					FEdGraphSchemaAction_K2Graph* GraphAction = (FEdGraphSchemaAction_K2Graph*)Action.Get();
-					if(GraphAction->EdGraph)
-					{
-						for(auto NodeIter = GraphAction->EdGraph->Nodes.CreateIterator(); NodeIter; NodeIter++)
-						{
-							if(UEdGraphNode* Node = *NodeIter)
-							{
-								for(auto PinIter = Node->Pins.CreateIterator(); PinIter; PinIter++)
-								{
-									if(UEdGraphPin* Pin = *PinIter)
-									{
-										if(UUserDefinedEnum* UserDefinedEnum =  Cast<UUserDefinedEnum>(Pin->PinType.PinSubCategoryObject.Get()))
-										{
-											UserDefinedEnumUserInBlueprint.AddUnique(UserDefinedEnum);
-										}
-									}
-								}
-							}
-						}
-					}
-				}
+				UserDefinedEnumUserInBlueprint.AddUnique(EnumWeakIter->Get());
 			}
 		}
+		FActionCollectingHelper::AddUserEnumsActions(OutAllActions, UserDefinedEnumUserInBlueprint);
 	}
 
-	for(auto EnumWeakIter = EnumsAddedToBlueprint.CreateIterator(); EnumWeakIter; EnumWeakIter++)
+	if (FStructureEditorUtils::UserDefinedStructEnabled())
 	{
-		if(EnumWeakIter->IsValid())
+		TArray<UUserDefinedStruct*> UserDefinedStructUserInBlueprint;
+		FActionCollectingHelper::CollectAllUsedTypes<UUserDefinedStruct, UStructProperty>(OutAllActions, UserDefinedStructUserInBlueprint);
+		// Make sure, that UDStructs from invalid variables are also listed
+		for (auto& VarDesc : Blueprint->NewVariables)
 		{
-			UserDefinedEnumUserInBlueprint.AddUnique(EnumWeakIter->Get());
+			if (auto UDStruct = Cast<UUserDefinedStruct>(VarDesc.VarType.PinSubCategoryObject.Get()))
+			{
+				UserDefinedStructUserInBlueprint.AddUnique(UDStruct);
+			}
 		}
-	}
-	for(auto EnumIter = UserDefinedEnumUserInBlueprint.CreateIterator(); EnumIter; EnumIter++)
-	{
-		if(UEnum* Enum = *EnumIter)
+		for(auto StructWeakPtr : StructsAddedToBlueprint)
 		{
-			const FText EnumNameStr = FText::FromString(Enum->GetName());
-			TSharedPtr<FEdGraphSchemaAction_K2Enum> NewEnumAction = MakeShareable(new FEdGraphSchemaAction_K2Enum( 
-				UserDefinedEnumCategoryNameStr(), EnumNameStr, EnumNameStr.ToString(), -1));
-			NewEnumAction->Enum = Enum;
-			NewEnumAction->SectionID = NodeSectionID::USER_ENUM;
-			OutAllActions.AddAction(NewEnumAction);
+			if(StructWeakPtr.IsValid())
+			{
+				UserDefinedStructUserInBlueprint.AddUnique(StructWeakPtr.Get());
+			}
 		}
+		FActionCollectingHelper::AddUserStructActions(OutAllActions, UserDefinedStructUserInBlueprint);
 	}
 }
 
@@ -1082,6 +1164,15 @@ void SMyBlueprint::OnActionSelected( const TArray< TSharedPtr<FEdGraphSchemaActi
 
 			Inspector->ShowDetailsForSingleObject(EnumAction->Enum, Options);
 		}
+		else if (InAction->GetTypeId() == FEdGraphSchemaAction_K2Struct::StaticGetTypeId())
+		{
+			FEdGraphSchemaAction_K2Struct* StructAction = (FEdGraphSchemaAction_K2Struct*)InAction.Get();
+
+			SKismetInspector::FShowDetailsOptions Options(StructAction->GetPathName().ToString());
+			Options.bForceRefresh = true;
+
+			Inspector->ShowDetailsForSingleObject(StructAction->Struct, Options);
+		}
 		else if (InAction->GetTypeId() == FEdGraphSchemaAction_K2TargetNode::StaticGetTypeId())
 		{
 			FEdGraphSchemaAction_K2TargetNode* TargetNodeAction = (FEdGraphSchemaAction_K2TargetNode*)InAction.Get();
@@ -1179,6 +1270,12 @@ FEdGraphSchemaAction_K2Enum* SMyBlueprint::SelectionAsEnum() const
 	return SelectionAsType<FEdGraphSchemaAction_K2Enum>( GraphActionMenu );
 }
 
+
+FEdGraphSchemaAction_K2Struct* SMyBlueprint::SelectionAsStruct() const
+{
+	return SelectionAsType<FEdGraphSchemaAction_K2Struct>( GraphActionMenu );
+}
+
 FEdGraphSchemaAction_K2Graph* SMyBlueprint::SelectionAsGraph() const
 {
 	return SelectionAsType<FEdGraphSchemaAction_K2Graph>( GraphActionMenu );
@@ -1222,7 +1319,7 @@ bool SMyBlueprint::SelectionIsCategory() const
 
 bool SMyBlueprint::SelectionHasContextMenu() const
 {
-	return SelectionAsGraph() || SelectionAsVar() || SelectionIsCategory() || SelectionAsDelegate() || SelectionAsEnum() || SelectionAsEvent() || SelectionAsLocalVar();
+	return SelectionAsGraph() || SelectionAsVar() || SelectionIsCategory() || SelectionAsDelegate() || SelectionAsEnum() || SelectionAsEvent() || SelectionAsLocalVar() || SelectionAsStruct();
 }
 
 void SMyBlueprint::GetSelectedItemsForContextMenu(TArray<FComponentEventConstructionData>& OutSelectedItems) const
@@ -1257,6 +1354,11 @@ TSharedPtr<SWidget> SMyBlueprint::OnContextMenuOpening()
 		{
 			MenuBuilder.AddMenuEntry(FMyBlueprintCommands::Get().FindUserDefinedEnumInContentBrowser);
 		}
+		
+		if(SelectionAsStruct())
+		{
+			MenuBuilder.AddMenuEntry(FMyBlueprintCommands::Get().FindUserDefinedStructInContentBrowser);
+		}
 
 		MenuBuilder.BeginSection("BasicOperations");
 		{
@@ -1316,6 +1418,10 @@ TSharedPtr<SWidget> SMyBlueprint::OnContextMenuOpening()
 			{
 				MenuBuilder.AddMenuSeparator();
 				MenuBuilder.AddMenuEntry(FMyBlueprintCommands::Get().AddNewUserDefinedEnum);
+				if (FStructureEditorUtils::UserDefinedStructEnabled())
+				{
+					MenuBuilder.AddMenuEntry(FMyBlueprintCommands::Get().AddNewUserDefinedStruct);
+				}
 			}
 		}
 	}
@@ -1337,6 +1443,20 @@ void SMyBlueprint::OnFindUserDefinedEnumInContentBrowser() const
 	}
 }
 
+void SMyBlueprint::OnFindUserDefinedStructInContentBrowser() const
+{
+	if (FEdGraphSchemaAction_K2Struct* GraphAction = SelectionAsStruct())
+	{
+		UUserDefinedStruct* Struct = Cast<UUserDefinedStruct>(GraphAction->Struct);
+		if(Struct && Struct->IsAsset())
+		{
+			TArray<UObject*> ObjectsToSync;
+			ObjectsToSync.Add(Struct);
+			GEditor->SyncBrowserToObjects( ObjectsToSync );
+		}
+	}
+}
+
 void SMyBlueprint::AddNewUserDefinedEnum()
 {
 	UBlueprint* Blueprint = GetBlueprintObj();
@@ -1367,7 +1487,7 @@ void SMyBlueprint::AddNewUserDefinedEnum()
 			Refresh();
 			if (GraphActionMenu.IsValid())
 			{
-				GraphActionMenu->ExpandCategory(UserDefinedEnumCategoryNameStr());
+				GraphActionMenu->ExpandCategory(FActionCollectingHelper::UserDefinedEnumCategoryNameStr());
 				GraphActionMenu->SelectItemByName(FName(*Enum->GetPathName()),ESelectInfo::OnMouseClick);
 				GraphActionMenu->OnRequestRenameOnActionNode();
 			}
@@ -1375,6 +1495,44 @@ void SMyBlueprint::AddNewUserDefinedEnum()
 	}
 }
 
+void SMyBlueprint::AddNewUserDefinedStruct()
+{
+	UBlueprint* Blueprint = GetBlueprintObj();
+	if (Blueprint)
+	{
+		static FName AssetToolsModuleName = FName("AssetTools");
+		FAssetToolsModule& AssetToolsModule = FModuleManager::GetModuleChecked<FAssetToolsModule>(AssetToolsModuleName);
+
+		FString AssetName;
+		FString PackagePath;
+		AssetToolsModule.Get().CreateUniqueAssetName(Blueprint->GetOutermost()->GetName(), TEXT("_Struct"), PackagePath, AssetName);
+
+		UUserDefinedStruct* Struct = NULL;
+		if (!AssetName.IsEmpty())
+		{
+			auto* Factory = ConstructObject<UStructureFactory>(UStructureFactory::StaticClass());
+			if (Factory && Factory->ConfigureProperties())
+			{
+				PackagePath = FPackageName::GetLongPackagePath(Blueprint->GetPathName());
+				UObject* NewAsset = AssetToolsModule.Get().CreateAsset(AssetName, PackagePath, UUserDefinedStruct::StaticClass(), Factory );
+				Struct = Cast<UUserDefinedStruct>(NewAsset);
+			}
+		}
+
+		if (Struct)
+		{
+			StructsAddedToBlueprint.Add(Struct);
+			Refresh();
+			if (GraphActionMenu.IsValid())
+			{
+				GraphActionMenu->ExpandCategory(FActionCollectingHelper::UserDefinedStructCategoryNameStr());
+				GraphActionMenu->SelectItemByName(FName(*Struct->GetPathName()),ESelectInfo::OnMouseClick);
+				GraphActionMenu->OnRequestRenameOnActionNode();
+			}
+		}
+	}
+}
+
 bool SMyBlueprint::CanOpenGraph() const
 {
 	const FEdGraphSchemaAction_K2Graph* GraphAction = SelectionAsGraph();
@@ -1499,6 +1657,10 @@ void SMyBlueprint::OnFindEntry()
 	{
 		SearchTerm = EnumAction->Enum->GetName();
 	}
+	else if (FEdGraphSchemaAction_K2Struct* StructAction = SelectionAsStruct())
+	{
+		BlueprintEditorPtr.Pin()->SummonSearchUI(true, StructAction->Struct->GetName());
+	}
 	else if (FEdGraphSchemaAction_K2Event* EventAction = SelectionAsEvent())
 	{
 		SearchTerm = EventAction->MenuDescription.ToString();
diff --git a/Engine/Source/Editor/Kismet/Private/SMyBlueprint.h b/Engine/Source/Editor/Kismet/Private/SMyBlueprint.h
--- a/Engine/Source/Editor/Kismet/Private/SMyBlueprint.h
+++ b/Engine/Source/Editor/Kismet/Private/SMyBlueprint.h
@@ -16,7 +16,8 @@ namespace NodeSectionID
 		VARIABLE,				// Variables
 		DELEGATE,				// Delegate/Event
 		USER_ENUM,				// User defined enums
-		LOCAL_VARIABLE			// Local variables
+		LOCAL_VARIABLE,			// Local variables
+		USER_STRUCT,			// User defined structs
 	};
 };
 
@@ -40,7 +41,8 @@ public:
 	TSharedPtr<FUICommandInfo> FindUserDefinedEnumInContentBrowser;
 	TSharedPtr<FUICommandInfo> AddNewUserDefinedEnum;
 	TSharedPtr<FUICommandInfo> GotoNativeVarDefinition;
-
+	TSharedPtr<FUICommandInfo> FindUserDefinedStructInContentBrowser;
+	TSharedPtr<FUICommandInfo> AddNewUserDefinedStruct;
 	// Add New Item
 	/** Initialize commands */
 	virtual void RegisterCommands() OVERRIDE;
@@ -68,6 +70,9 @@ public:
 	/** Accessor for getting the current selection as a K2 enum */
 	FEdGraphSchemaAction_K2Enum* SelectionAsEnum() const;
 
+	/** Accessor for getting the current selection as a K2 enum */
+	FEdGraphSchemaAction_K2Struct* SelectionAsStruct() const;
+
 	/** Accessor for getting the current selection as a K2 var */
 	FEdGraphSchemaAction_K2Var* SelectionAsVar() const;
 	
@@ -177,6 +182,8 @@ private:
 	void OnDuplicateAction();
 	void GotoNativeCodeVarDefinition();
 	bool IsNativeVariable() const;
+	void OnFindUserDefinedStructInContentBrowser() const;
+	void AddNewUserDefinedStruct();
 
 	/** Callback when the filter is changed, forces the action tree(s) to filter */
 	void OnFilterTextChanged( const FText& InFilterText );
@@ -215,4 +222,7 @@ private:
 
 	/** Contains both the GraphActionMenu and LocalGraphActionMenu */
 	TSharedPtr< SSplitter > ActionMenuContainer;
+
+	/** Enums created from 'blueprint' level */
+	TArray<TWeakObjectPtr<UUserDefinedStruct>> StructsAddedToBlueprint;
 };
diff --git a/Engine/Source/Editor/Kismet/Private/UserDefinedEnumEditor.cpp b/Engine/Source/Editor/Kismet/Private/UserDefinedEnumEditor.cpp
--- a/Engine/Source/Editor/Kismet/Private/UserDefinedEnumEditor.cpp
+++ b/Engine/Source/Editor/Kismet/Private/UserDefinedEnumEditor.cpp
@@ -164,12 +164,12 @@ void FEnumDetails::CustomizeDetails( IDetailLayoutBuilder& DetailLayout )
 		InputsCategory.AddCustomBuilder( Layout.ToSharedRef() );
 	}
 
-	FEnumEditorUtils::FEnumEditorManager::Get().AddEnumListener(this);
+	FEnumEditorUtils::FEnumEditorManager::Get().AddListener(this);
 }
 
 FEnumDetails::~FEnumDetails()
 {
-	FEnumEditorUtils::FEnumEditorManager::Get().RemoveEnumListener(this);
+	FEnumEditorUtils::FEnumEditorManager::Get().RemoveListener(this);
 }
 
 void FEnumDetails::OnForceRefresh()
@@ -180,7 +180,7 @@ void FEnumDetails::OnForceRefresh()
 	}
 }
 
-void FEnumDetails::OnEnumChanged(const class UUserDefinedEnum* Enum)
+void FEnumDetails::OnChanged(const class UUserDefinedEnum* Enum)
 {
 	if (Enum && (TargetEnum.Get() == Enum))
 	{
diff --git a/Engine/Source/Editor/Kismet/Private/UserDefinedEnumEditor.h b/Engine/Source/Editor/Kismet/Private/UserDefinedEnumEditor.h
--- a/Engine/Source/Editor/Kismet/Private/UserDefinedEnumEditor.h
+++ b/Engine/Source/Editor/Kismet/Private/UserDefinedEnumEditor.h
@@ -69,7 +69,7 @@ public:
 	void OnForceRefresh();
 
 	/** FEnumEditorUtils::INotifyOnEnumChanged */
-	virtual void OnEnumChanged(const class UUserDefinedEnum* Enum) OVERRIDE;
+	virtual void OnChanged(const class UUserDefinedEnum* Enum) OVERRIDE;
 
 private:
 	/** Handles new enum element request */
diff --git a/Engine/Source/Editor/Kismet/Private/UserDefinedStructureEditor.cpp b/Engine/Source/Editor/Kismet/Private/UserDefinedStructureEditor.cpp
new file mode 100644
--- /dev/null
+++ b/Engine/Source/Editor/Kismet/Private/UserDefinedStructureEditor.cpp
@@ -0,0 +1,579 @@
+// Copyright 1998-2014 Epic Games, Inc. All Rights Reserved.
+
+#include "BlueprintEditorPrivatePCH.h"
+#include "UserDefinedStructureEditor.h"
+
+#include "PropertyCustomizationHelpers.h"
+#include "Editor/KismetWidgets//Public/SPinTypeSelector.h"
+#include "Editor/WorkspaceMenuStructure/Public/WorkspaceMenuStructureModule.h"
+
+#define LOCTEXT_NAMESPACE "StructureEditor"
+
+
+///////////////////////////////////////////////////////////////////////////////////////
+// FUserDefinedStructureEditor
+
+const FName FUserDefinedStructureEditor::MemberVariablesTabId( TEXT( "UserDefinedStruct_MemberVariablesEditor" ) );
+const FName FUserDefinedStructureEditor::UserDefinedStructureEditorAppIdentifier( TEXT( "UserDefinedStructEditorApp" ) );
+
+void FUserDefinedStructureEditor::RegisterTabSpawners(const TSharedRef<class FTabManager>& TabManager)
+{
+	FAssetEditorToolkit::RegisterTabSpawners(TabManager);
+
+	const IWorkspaceMenuStructure& MenuStructure = WorkspaceMenu::GetMenuStructure();
+
+	TabManager->RegisterTabSpawner( MemberVariablesTabId, FOnSpawnTab::CreateSP(this, &FUserDefinedStructureEditor::SpawnStructureTab) )
+		.SetDisplayName( LOCTEXT("MemberVariablesEditor", "Member Variables") )
+		.SetGroup( MenuStructure.GetAssetEditorCategory() );
+}
+
+void FUserDefinedStructureEditor::UnregisterTabSpawners(const TSharedRef<class FTabManager>& TabManager)
+{
+	FAssetEditorToolkit::UnregisterTabSpawners(TabManager);
+
+	TabManager->UnregisterTabSpawner( MemberVariablesTabId );
+}
+
+void FUserDefinedStructureEditor::InitEditor(const EToolkitMode::Type Mode, const TSharedPtr< class IToolkitHost >& InitToolkitHost, UUserDefinedStruct* Struct)
+{
+	const TSharedRef<FTabManager::FLayout> StandaloneDefaultLayout = FTabManager::NewLayout( "Standalone_UserDefinedStructureEditor_Layout_v1" )
+	->AddArea
+	(
+		FTabManager::NewPrimaryArea() ->SetOrientation(Orient_Vertical)
+		->Split
+		(
+			FTabManager::NewStack()
+			->SetSizeCoefficient(0.1f)
+			->SetHideTabWell( true )
+			->AddTab(GetToolbarTabId(), ETabState::OpenedTab)
+		)
+		->Split
+		(
+			FTabManager::NewSplitter()
+			->Split
+			(
+				FTabManager::NewStack()
+				->AddTab( MemberVariablesTabId, ETabState::OpenedTab )
+			)
+		)
+	);
+
+	const bool bCreateDefaultStandaloneMenu = true;
+	const bool bCreateDefaultToolbar = true;
+	FAssetEditorToolkit::InitAssetEditor( Mode, InitToolkitHost, UserDefinedStructureEditorAppIdentifier, StandaloneDefaultLayout, bCreateDefaultStandaloneMenu, bCreateDefaultToolbar, Struct );
+}
+
+FUserDefinedStructureEditor::~FUserDefinedStructureEditor()
+{
+}
+
+FName FUserDefinedStructureEditor::GetToolkitFName() const
+{
+	return FName("UserDefinedStructureEditor");
+}
+
+FText FUserDefinedStructureEditor::GetBaseToolkitName() const
+{
+	return LOCTEXT( "AppLabel", "Struct Editor" );
+}
+
+FText FUserDefinedStructureEditor::GetToolkitName() const
+{
+	if (1 == GetEditingObjects().Num())
+	{
+		return FAssetEditorToolkit::GetToolkitName();
+	}
+	return GetBaseToolkitName();
+}
+
+FString FUserDefinedStructureEditor::GetWorldCentricTabPrefix() const
+{
+	return LOCTEXT("UDStructWorldCentricTabPrefix", "Struct ").ToString();
+}
+
+FLinearColor FUserDefinedStructureEditor::GetWorldCentricTabColorScale() const
+{
+	return FLinearColor( 0.0f, 0.0f, 1.0f, 0.5f );
+}
+
+TSharedRef<SDockTab> FUserDefinedStructureEditor::SpawnStructureTab(const FSpawnTabArgs& Args)
+{
+	check( Args.GetTabId() == MemberVariablesTabId );
+
+	UUserDefinedStruct* EditedStruct = NULL;
+	const auto EditingObjects = GetEditingObjects();
+	if (EditingObjects.Num())
+	{
+		EditedStruct = Cast<UUserDefinedStruct>(EditingObjects[ 0 ]);
+	}
+
+	// Create a property view
+	FPropertyEditorModule& EditModule = FModuleManager::Get().GetModuleChecked<FPropertyEditorModule>("PropertyEditor");
+
+	FDetailsViewArgs DetailsViewArgs( /*bUpdateFromSelection=*/ false, /*bLockable=*/ false, /*bAllowSearch=*/ false, /*bObjectsUseNameArea=*/ true, /*bHideSelectionTip=*/ true );
+	DetailsViewArgs.bHideActorNameArea = true;
+	DetailsViewArgs.bShowOptions = false;
+
+	PropertyView = EditModule.CreateDetailView( DetailsViewArgs );
+
+	FOnGetDetailCustomizationInstance LayoutStructDetails = FOnGetDetailCustomizationInstance::CreateStatic(&FUserDefinedStructureDetails::MakeInstance);
+	PropertyView->RegisterInstancedCustomPropertyLayout(UUserDefinedStruct::StaticClass(), LayoutStructDetails);
+
+	PropertyView->SetObject(EditedStruct);
+
+	return SNew(SDockTab)
+		.Icon( FEditorStyle::GetBrush("GenericEditor.Tabs.Properties") )
+		.Label( LOCTEXT("UserDefinedStructureEditor", "Structure") )
+		.TabColorScale( GetTabColorScale() )
+		[
+			PropertyView.ToSharedRef()
+		];
+}
+
+///////////////////////////////////////////////////////////////////////////////////////
+// FUserDefinedStructureLayout
+
+//Represents single structure (List of fields)
+class FUserDefinedStructureLayout : public IDetailCustomNodeBuilder, public TSharedFromThis<FUserDefinedStructureLayout>
+{
+public:
+	FUserDefinedStructureLayout(TWeakPtr<class FUserDefinedStructureDetails> InStructureDetails) : StructureDetails(InStructureDetails) {}
+
+	void OnChanged()
+	{
+		OnRegenerateChildren.ExecuteIfBound();
+	}
+
+	FReply OnAddNewField()
+	{
+		auto StructureDetailsSP = StructureDetails.Pin();
+		if(StructureDetailsSP.IsValid())
+		{
+			const UEdGraphSchema_K2* K2Schema = GetDefault<UEdGraphSchema_K2>();
+			const  FEdGraphPinType InitialType(K2Schema->PC_Boolean, TEXT(""), NULL, false, false);
+			FStructureEditorUtils::AddVariable(StructureDetailsSP->GetUserDefinedStruct(), InitialType);
+			OnChanged();
+		}
+
+		return FReply::Handled();
+	}
+
+	const FSlateBrush* OnGetStructureStatus() const
+	{
+		auto StructureDetailsSP = StructureDetails.Pin();
+		if(StructureDetailsSP.IsValid())
+		{
+			if(auto Struct = StructureDetailsSP->GetUserDefinedStruct())
+			{
+				switch(Struct->Status.GetValue())
+				{
+				case EUserDefinedStructureStatus::UDSS_Error:
+					return FEditorStyle::GetBrush("Kismet.Status.Error.Small");
+				case EUserDefinedStructureStatus::UDSS_UpToDate:
+					return FEditorStyle::GetBrush("Kismet.Status.Good.Small");
+				default:
+					return FEditorStyle::GetBrush("Kismet.Status.Unknown.Small");
+				}
+				
+			}
+		}
+		return NULL;
+	}
+
+	/** IDetailCustomNodeBuilder Interface*/
+	virtual void SetOnRebuildChildren( FSimpleDelegate InOnRegenerateChildren ) OVERRIDE 
+	{
+		OnRegenerateChildren = InOnRegenerateChildren;
+	}
+	virtual void GenerateChildContent( IDetailChildrenBuilder& ChildrenBuilder ) OVERRIDE;
+
+	virtual void GenerateHeaderRowContent( FDetailWidgetRow& NodeRow ) OVERRIDE {}
+	virtual void Tick( float DeltaTime ) OVERRIDE {}
+	virtual bool RequiresTick() const OVERRIDE { return false; }
+	virtual FName GetName() const OVERRIDE 
+	{ 
+		auto StructureDetailsSP = StructureDetails.Pin();
+		if(StructureDetailsSP.IsValid())
+		{
+			if(auto Struct = StructureDetailsSP->GetUserDefinedStruct())
+			{
+				return Struct->GetFName();
+			}
+		}
+		return NAME_None; 
+	}
+	virtual bool InitiallyCollapsed() const OVERRIDE { return false; }
+
+private:
+	TWeakPtr<class FUserDefinedStructureDetails> StructureDetails;
+
+	FSimpleDelegate OnRegenerateChildren;
+};
+
+///////////////////////////////////////////////////////////////////////////////////////
+// FUserDefinedStructureFieldLayout
+
+//Represents single field
+class FUserDefinedStructureFieldLayout : public IDetailCustomNodeBuilder, public TSharedFromThis<FUserDefinedStructureFieldLayout>
+{
+public:
+	FUserDefinedStructureFieldLayout(TWeakPtr<class FUserDefinedStructureDetails> InStructureDetails, TWeakPtr<class FUserDefinedStructureLayout> InStructureLayout, FGuid InFieldGuid)
+		: FieldGuid(InFieldGuid)
+		, StructureLayout(InStructureLayout)
+		, StructureDetails(InStructureDetails) {}
+
+	void OnChanged()
+	{
+		OnRegenerateChildren.ExecuteIfBound();
+	}
+
+	FText OnGetNameText() const
+	{
+		auto StructureDetailsSP = StructureDetails.Pin();
+		if(StructureDetailsSP.IsValid())
+		{
+			return FText::FromString(FStructureEditorUtils::GetVariableDisplayName(StructureDetailsSP->GetUserDefinedStruct(), FieldGuid));
+		}
+		return FText::GetEmpty();
+	}
+
+	void OnNameTextCommitted(const FText& NewText, ETextCommit::Type InTextCommit)
+	{
+		auto StructureDetailsSP = StructureDetails.Pin();
+		if(StructureDetailsSP.IsValid())
+		{
+			const FString NewNameStr = NewText.ToString();
+			FStructureEditorUtils::RenameVariable(StructureDetailsSP->GetUserDefinedStruct(), FieldGuid, NewNameStr);
+			OnChanged();
+		}
+	}
+
+	FEdGraphPinType OnGetPinInfo() const
+	{
+		auto StructureDetailsSP = StructureDetails.Pin();
+		if(StructureDetailsSP.IsValid())
+		{
+			if(const FStructVariableDescription* FieldDesc = StructureDetailsSP->FindStructureFieldByGuid(FieldGuid))
+			{
+				return FieldDesc->ToPinType();
+			}
+		}
+		return FEdGraphPinType();
+	}
+
+	void PinInfoChanged(const FEdGraphPinType& PinType)
+	{
+		auto StructureDetailsSP = StructureDetails.Pin();
+		if(StructureDetailsSP.IsValid())
+		{
+			FStructureEditorUtils::ChangeVariableType(StructureDetailsSP->GetUserDefinedStruct(), FieldGuid, PinType);
+			OnChanged();
+		}
+	}
+
+	void OnPrePinInfoChange(const FEdGraphPinType& PinType)
+	{
+
+	}
+
+	void OnRemovField()
+	{
+		auto StructureDetailsSP = StructureDetails.Pin();
+		if(StructureDetailsSP.IsValid())
+		{
+			FStructureEditorUtils::RemoveVariable(StructureDetailsSP->GetUserDefinedStruct(), FieldGuid);
+			auto StructureLayoutSP = StructureLayout.Pin();
+			if(StructureLayoutSP.IsValid())
+			{
+				StructureLayoutSP->OnChanged();
+			}
+		}
+	}
+
+	bool IsRemoveButtonEnabled()
+	{
+		auto StructureDetailsSP = StructureDetails.Pin();
+		if (StructureDetailsSP.IsValid())
+		{
+			if (auto UDStruct = StructureDetailsSP->GetUserDefinedStruct())
+			{
+				return (FStructureEditorUtils::GetVarDesc(UDStruct).Num() > 1);
+			}
+		}
+		return false;
+	}
+
+	FText OnGetArgDefaultValueText() const
+	{
+		auto StructureDetailsSP = StructureDetails.Pin();
+		if(StructureDetailsSP.IsValid())
+		{
+			if(const FStructVariableDescription* FieldDesc = StructureDetailsSP->FindStructureFieldByGuid(FieldGuid))
+			{
+				return FText::FromString(FieldDesc->DefaultValue);
+			}
+		}
+		return FText();
+	}
+
+	void OnArgDefaultValueCommitted(const FText& NewText, ETextCommit::Type InTextCommit)
+	{
+		auto StructureDetailsSP = StructureDetails.Pin();
+		if(StructureDetailsSP.IsValid())
+		{
+			FStructureEditorUtils::ChangeVariableDefaultValue(StructureDetailsSP->GetUserDefinedStruct(), FieldGuid, NewText.ToString());
+			OnChanged();
+		}
+	}
+
+	/** IDetailCustomNodeBuilder Interface*/
+	virtual void SetOnRebuildChildren( FSimpleDelegate InOnRegenerateChildren ) OVERRIDE 
+	{
+		OnRegenerateChildren = InOnRegenerateChildren;
+	}
+
+	EVisibility GetErrorIconVisibility()
+	{
+		auto StructureDetailsSP = StructureDetails.Pin();
+		if(StructureDetailsSP.IsValid())
+		{
+			auto FieldDesc = StructureDetailsSP->FindStructureFieldByGuid(FieldGuid);
+			if (FieldDesc && FieldDesc->bInvalidMember)
+			{
+				return EVisibility::Visible;
+			}
+		}
+
+		return EVisibility::Collapsed;
+	}
+
+	void RemoveInvalidSubTypes(TSharedPtr<UEdGraphSchema_K2::FPinTypeTreeInfo> PinTypeNode, const UStruct* Parent) const
+	{
+		for (int32 ChildIndex = 0; ChildIndex < PinTypeNode->Children.Num();)
+		{
+			auto Child = PinTypeNode->Children[ChildIndex];
+			if(Child.IsValid())
+			{
+				auto ScriptStruct = Cast<const UScriptStruct>(Child->GetPinType(false).PinSubCategoryObject.Get());
+				if (ScriptStruct)
+				{
+					if (FStructureEditorUtils::EStructureError::Ok != FStructureEditorUtils::IsStructureValid(ScriptStruct, Parent))
+					{
+						PinTypeNode->Children.RemoveAt(ChildIndex);
+						continue;
+					}
+				}
+			}
+			++ChildIndex;
+		}
+	}
+
+	void GetFilteredVariableTypeTree( TArray< TSharedPtr<UEdGraphSchema_K2::FPinTypeTreeInfo> >& TypeTree, bool bAllowExec, bool bAllowWildcard ) const
+	{
+		auto K2Schema = GetDefault<UEdGraphSchema_K2>();
+		auto StructureDetailsSP = StructureDetails.Pin();
+		if(StructureDetailsSP.IsValid() && K2Schema)
+		{
+			K2Schema->GetVariableTypeTree(TypeTree, bAllowExec, bAllowWildcard);
+			const auto Parent = StructureDetailsSP->GetUserDefinedStruct();
+			// THE TREE HAS ONLY 2 LEVELS
+			for (auto PinTypePtr : TypeTree)
+			{
+				if (PinTypePtr.IsValid() && (K2Schema->PC_Struct == PinTypePtr->GetPinType(false).PinCategory))
+				{
+					RemoveInvalidSubTypes(PinTypePtr, Parent);
+				}
+			}
+		}
+	}
+
+	virtual void GenerateHeaderRowContent( FDetailWidgetRow& NodeRow ) OVERRIDE 
+	{
+		auto K2Schema = GetDefault<UEdGraphSchema_K2>();
+
+		TSharedPtr<SImage> ErrorIcon;
+
+		NodeRow
+		.NameContent()
+		[
+			SNew(SHorizontalBox)
+
+			+SHorizontalBox::Slot()
+			.AutoWidth()
+			.HAlign(HAlign_Left)
+			.VAlign(VAlign_Center)
+			[
+				SAssignNew(ErrorIcon, SImage)
+				.Image( FEditorStyle::GetBrush("Icons.Error") )
+			]
+
+			+SHorizontalBox::Slot()
+			.FillWidth(1)
+			.VAlign(VAlign_Center)
+			[
+				SNew(SEditableTextBox)
+				.Text( this, &FUserDefinedStructureFieldLayout::OnGetNameText )
+				.OnTextCommitted( this, &FUserDefinedStructureFieldLayout::OnNameTextCommitted )
+				.Font( IDetailLayoutBuilder::GetDetailFont() )
+			]
+		]
+		.ValueContent()
+		[
+			SNew(SHorizontalBox)
+			+SHorizontalBox::Slot()
+			.FillWidth(1)
+			.VAlign(VAlign_Center)
+			.Padding(0.0f, 0.0f, 4.0f, 0.0f)
+			[
+				SNew(SPinTypeSelector, FGetPinTypeTree::CreateSP(this, &FUserDefinedStructureFieldLayout::GetFilteredVariableTypeTree))
+				.TargetPinType(this, &FUserDefinedStructureFieldLayout::OnGetPinInfo)
+				.OnPinTypePreChanged(this, &FUserDefinedStructureFieldLayout::OnPrePinInfoChange)
+				.OnPinTypeChanged(this, &FUserDefinedStructureFieldLayout::PinInfoChanged)
+				.Schema(K2Schema)
+				.bAllowExec(false)
+				.bAllowWildcard(false)
+				.Font( IDetailLayoutBuilder::GetDetailFont() )
+			]
+			+SHorizontalBox::Slot()
+			.AutoWidth()
+			.HAlign(HAlign_Right)
+			.VAlign(VAlign_Center)
+			[
+				PropertyCustomizationHelpers::MakeClearButton(
+					FSimpleDelegate::CreateSP(this, &FUserDefinedStructureFieldLayout::OnRemovField),
+					LOCTEXT("RemoveVariable", "Remove member variable"),
+					TAttribute<bool>::Create(TAttribute<bool>::FGetter::CreateSP(this, &FUserDefinedStructureFieldLayout::IsRemoveButtonEnabled)))
+			]
+		];
+
+		if (ErrorIcon.IsValid())
+		{
+			ErrorIcon->SetVisibility(
+				TAttribute<EVisibility>::Create(
+					TAttribute<EVisibility>::FGetter::CreateSP(
+						this, &FUserDefinedStructureFieldLayout::GetErrorIconVisibility)));
+		}
+	}
+
+	virtual void GenerateChildContent( IDetailChildrenBuilder& ChildrenBuilder ) OVERRIDE 
+	{
+		ChildrenBuilder.AddChildContent( *LOCTEXT( "FunctionArgDetailsDefaultValue", "Default Value" ).ToString() )
+		.NameContent()
+		[
+			SNew(STextBlock)
+			.Text( LOCTEXT( "DefaultValue", "Default Value" ) )
+			.Font( IDetailLayoutBuilder::GetDetailFont() )
+		]
+		.ValueContent()
+		[
+			SNew(SEditableTextBox)
+			.Text( this, &FUserDefinedStructureFieldLayout::OnGetArgDefaultValueText )
+			.OnTextCommitted( this, &FUserDefinedStructureFieldLayout::OnArgDefaultValueCommitted )
+			.Font( IDetailLayoutBuilder::GetDetailFont() )
+		];
+	}
+
+	virtual void Tick( float DeltaTime ) OVERRIDE {}
+	virtual bool RequiresTick() const OVERRIDE { return false; }
+	virtual FName GetName() const OVERRIDE { return FName(*FieldGuid.ToString()); }
+	virtual bool InitiallyCollapsed() const OVERRIDE { return true; }
+
+private:
+	TWeakPtr<class FUserDefinedStructureDetails> StructureDetails;
+
+	TWeakPtr<class FUserDefinedStructureLayout> StructureLayout;
+
+	FGuid FieldGuid;
+
+	FSimpleDelegate OnRegenerateChildren;
+};
+
+///////////////////////////////////////////////////////////////////////////////////////
+// FUserDefinedStructureLayout
+void FUserDefinedStructureLayout::GenerateChildContent( IDetailChildrenBuilder& ChildrenBuilder ) 
+{
+	ChildrenBuilder.AddChildContent(FString())
+		[
+			SNew(SHorizontalBox)
+			+SHorizontalBox::Slot()
+			.AutoWidth()
+			.HAlign(HAlign_Left)
+			.VAlign(VAlign_Center)
+			[
+				SNew(SImage)
+				.Image( this, &FUserDefinedStructureLayout::OnGetStructureStatus )
+			]
+			+SHorizontalBox::Slot()
+				.HAlign(HAlign_Right)
+				[
+					SNew(SButton)
+					.Text(LOCTEXT("NewStructureField", "New Variable").ToString())
+					.OnClicked(this, &FUserDefinedStructureLayout::OnAddNewField)
+				]
+		];
+
+	auto StructureDetailsSP = StructureDetails.Pin();
+	if(StructureDetailsSP.IsValid())
+	{
+		if(auto Struct = StructureDetailsSP->GetUserDefinedStruct())
+		{
+			for (auto& VarDesc : FStructureEditorUtils::GetVarDesc(Struct))
+			{
+				TSharedRef<class FUserDefinedStructureFieldLayout> VarLayout = MakeShareable(new FUserDefinedStructureFieldLayout(StructureDetails,  SharedThis(this), VarDesc.VarGuid));
+				ChildrenBuilder.AddChildCustomBuilder(VarLayout);
+			}
+		}
+	}
+}
+
+///////////////////////////////////////////////////////////////////////////////////////
+// FUserDefinedStructureDetails
+
+UUserDefinedStruct* FUserDefinedStructureDetails::GetUserDefinedStruct()
+{
+	return UserDefinedStruct.Get();
+}
+
+FStructVariableDescription* FUserDefinedStructureDetails::FindStructureFieldByGuid(FGuid Guid)
+{
+	if (auto Struct = GetUserDefinedStruct())
+	{
+		return FStructureEditorUtils::GetVarDesc(Struct).FindByPredicate(FStructureEditorUtils::FFindByGuidHelper<FStructVariableDescription>(Guid));
+	}
+	return NULL;
+}
+
+void FUserDefinedStructureDetails::CustomizeDetails( class IDetailLayoutBuilder& DetailLayout )
+{
+	const TArray<TWeakObjectPtr<UObject>> Objects = DetailLayout.GetDetailsView().GetSelectedObjects();
+	check(Objects.Num() > 0);
+
+	if (Objects.Num() == 1)
+	{
+		UserDefinedStruct = CastChecked<UUserDefinedStruct>(Objects[0].Get());
+		IDetailCategoryBuilder& StructureCategory = DetailLayout.EditCategory("Structure", LOCTEXT("Structure", "Structure").ToString());
+
+		Layout = MakeShareable(new FUserDefinedStructureLayout(SharedThis(this)));
+		StructureCategory.AddCustomBuilder(Layout.ToSharedRef());
+	}
+
+	FStructureEditorUtils::FStructEditorManager::Get().AddListener(this);
+}
+
+void FUserDefinedStructureDetails::OnChanged(const class UUserDefinedStruct* Struct)
+{
+	if (Struct && (GetUserDefinedStruct() == Struct))
+	{
+		if (Layout.IsValid())
+		{
+			Layout->OnChanged();
+		}
+	}
+}
+
+FUserDefinedStructureDetails::~FUserDefinedStructureDetails()
+{
+	FStructureEditorUtils::FStructEditorManager::Get().RemoveListener(this);
+}
+
+#undef LOCTEXT_NAMESPACE
\ No newline at end of file
diff --git a/Engine/Source/Editor/Kismet/Private/UserDefinedStructureEditor.h b/Engine/Source/Editor/Kismet/Private/UserDefinedStructureEditor.h
new file mode 100644
--- /dev/null
+++ b/Engine/Source/Editor/Kismet/Private/UserDefinedStructureEditor.h
@@ -0,0 +1,71 @@
+// Copyright 1998-2014 Epic Games, Inc. All Rights Reserved.
+
+#pragma once
+
+#include "BlueprintEditorModule.h"
+#include "Editor/PropertyEditor/Public/PropertyEditing.h"
+#include "Editor/UnrealEd/Public/Kismet2/StructureEditorUtils.h"
+
+class FUserDefinedStructureEditor : public IUserDefinedStructureEditor
+{
+	/** App Identifier.*/
+	static const FName UserDefinedStructureEditorAppIdentifier;
+
+	/**	The tab ids for all the tabs used */
+	static const FName MemberVariablesTabId;
+	
+	/** Property viewing widget */
+	TSharedPtr<class IDetailsView> PropertyView;
+
+public:
+	/**
+	 * Edits the specified enum
+	 *
+	 * @param	Mode					Asset editing mode for this editor (standalone or world-centric)
+	 * @param	InitToolkitHost			When Mode is WorldCentric, this is the level editor instance to spawn this editor within
+	 * @param	EnumToEdit				The user defined enum to edit
+	 */
+	void InitEditor(const EToolkitMode::Type Mode, const TSharedPtr< class IToolkitHost >& InitToolkitHost, UUserDefinedStruct* EnumToEdit);
+
+	/** Destructor */
+	virtual ~FUserDefinedStructureEditor();
+
+	/** IToolkit interface */
+	virtual FName GetToolkitFName() const OVERRIDE;
+	virtual FText GetBaseToolkitName() const OVERRIDE;
+	virtual FText GetToolkitName() const OVERRIDE;
+	virtual FString GetWorldCentricTabPrefix() const OVERRIDE;
+	virtual FLinearColor GetWorldCentricTabColorScale() const OVERRIDE;
+
+	virtual void RegisterTabSpawners(const TSharedRef<class FTabManager>& TabManager) OVERRIDE;
+	virtual void UnregisterTabSpawners(const TSharedRef<class FTabManager>& TabManager) OVERRIDE;
+
+protected:
+	TSharedRef<SDockTab> SpawnStructureTab(const FSpawnTabArgs& Args);
+};
+
+class FUserDefinedStructureDetails : public IDetailCustomization, FStructureEditorUtils::INotifyOnStructChanged
+{
+public:
+	/** Makes a new instance of this detail layout class for a specific detail view requesting it */
+	static TSharedRef<IDetailCustomization> MakeInstance()
+	{
+		return MakeShareable(new FUserDefinedStructureDetails);
+	}
+
+	~FUserDefinedStructureDetails();
+
+	/** IDetailCustomization interface */
+	virtual void CustomizeDetails( class IDetailLayoutBuilder& DetailLayout ) OVERRIDE;
+
+	UUserDefinedStruct* GetUserDefinedStruct();
+
+	struct FStructVariableDescription* FindStructureFieldByGuid(FGuid Guid);
+
+	/** FStructureEditorUtils::INotifyOnStructChanged */
+	virtual void OnChanged(const class UUserDefinedStruct* Struct) OVERRIDE;
+
+private:
+	TWeakObjectPtr<UUserDefinedStruct> UserDefinedStruct;
+	TSharedPtr<class FUserDefinedStructureLayout> Layout;
+};
diff --git a/Engine/Source/Editor/Kismet/Public/BlueprintEditor.h b/Engine/Source/Editor/Kismet/Public/BlueprintEditor.h
--- a/Engine/Source/Editor/Kismet/Public/BlueprintEditor.h
+++ b/Engine/Source/Editor/Kismet/Public/BlueprintEditor.h
@@ -192,7 +192,6 @@ public:
 	TSharedRef<class SSCSEditor> GetSCSEditor() const {return SCSEditor.ToSharedRef();}
 	TSharedPtr<class SSCSEditorViewport> GetSCSViewport() const {return SCSViewport;}
 	TSharedPtr<class SMyBlueprint> GetMyBlueprintWidget() const {return MyBlueprintWidget;}
-	TSharedRef<class SUserDefinedStructureEditor> GetUserDefinedStructureEditor() const {return UserDefinedStructureEditor.ToSharedRef();}
 	AActor* GetPreviewActor() const;
 
 	TSharedPtr<class FBlueprintEditorToolbar> GetToolbarBuilder() {return Toolbar;}
@@ -769,9 +768,6 @@ protected:
 	/**Load macro & function blueprint libraries from asset registry*/
 	void LoadLibrariesFromAssetRegistry();
 
-	/**Load user-defined enums from asset registry */
-	void LoadUserDefinedEnumsFromAssetRegistry();
-
 	// Begin FEditorUndoClient Interface
 	virtual void	PostUndo(bool bSuccess) OVERRIDE;
 	virtual void	PostRedo(bool bSuccess) OVERRIDE;
@@ -841,7 +837,10 @@ protected:
 	TWeakPtr<FDocumentTabFactory> GraphEditorTabFactoryPtr;
 
 	/** User-defined enumerators to keep loaded */
-	TArray<TWeakObjectPtr<UUserDefinedEnum>> UserDefinedEnumerators;
+	TSet<TWeakObjectPtr<UUserDefinedEnum>> UserDefinedEnumerators;
+
+	/** User-defined structures to keep loaded */
+	TSet<TWeakObjectPtr<UUserDefinedStruct>> UserDefinedStructures;
 	
 	/** Macro/function libraries to keep loaded */
 	TArray<UBlueprint*> StandardLibraries;
@@ -849,9 +848,6 @@ protected:
 	/** SCS editor */
 	TSharedPtr<class SSCSEditor> SCSEditor;
 
-	/** SCS editor */
-	TSharedPtr<class SUserDefinedStructureEditor> UserDefinedStructureEditor;
-
 	/** Viewport widget */
 	TSharedPtr<class SSCSEditorViewport> SCSViewport;
 
diff --git a/Engine/Source/Editor/Kismet/Public/BlueprintEditorModule.h b/Engine/Source/Editor/Kismet/Public/BlueprintEditorModule.h
--- a/Engine/Source/Editor/Kismet/Public/BlueprintEditorModule.h
+++ b/Engine/Source/Editor/Kismet/Public/BlueprintEditorModule.h
@@ -17,6 +17,13 @@ class KISMET_API IUserDefinedEnumEditor : public FAssetEditorToolkit
 {
 };
 
+/**
+ * Enum editor public interface
+ */
+class KISMET_API IUserDefinedStructureEditor : public FAssetEditorToolkit
+{
+};
+
 /**
  * Blueprint editor public interface
  */
@@ -92,6 +99,17 @@ public:
 	 */
 	virtual TSharedRef<IUserDefinedEnumEditor> CreateUserDefinedEnumEditor(const EToolkitMode::Type Mode, const TSharedPtr<IToolkitHost>& InitToolkitHost, UUserDefinedEnum* UDEnum);
 
+	/**
+	 * Creates an instance of a Structure editor object.
+	 *
+	 * @param	Mode					Mode that this editor should operate in
+	 * @param	InitToolkitHost			When Mode is WorldCentric, this is the level editor instance to spawn this editor within
+	 * @param	UDEnum					The user-defined structure to start editing
+	 *
+	 * @return	Interface to the new Struct editor
+	 */
+	virtual TSharedRef<IUserDefinedStructureEditor> CreateUserDefinedStructEditor(const EToolkitMode::Type Mode, const TSharedPtr<IToolkitHost>& InitToolkitHost, UUserDefinedStruct* UDStruct);
+
 	/** Gets the extensibility managers for outside entities to extend blueprint editor's menus and toolbars */
 	virtual TSharedPtr<FExtensibilityManager> GetMenuExtensibilityManager() {return MenuExtensibilityManager;}
 
diff --git a/Engine/Source/Editor/KismetCompiler/Private/KismetCompiler.cpp b/Engine/Source/Editor/KismetCompiler/Private/KismetCompiler.cpp
--- a/Engine/Source/Editor/KismetCompiler/Private/KismetCompiler.cpp
+++ b/Engine/Source/Editor/KismetCompiler/Private/KismetCompiler.cpp
@@ -1035,11 +1035,7 @@ void FKismetCompilerContext::PrecompileFunction(FKismetFunctionContext& Context)
 
 		// Determine if this is a new function or if it overrides a parent function
 		//@TODO: Does not support multiple overloads for a parent virtual function
-		const UObject* ParentBlueprint = Blueprint->ParentClass->ClassGeneratedBy;
-
-		//@TODO: The parent function could be not loaded completely when the parent's blueprint is being regenerated. The condition should be removed once the GenerateBlueprintSkeleton is no longer called from Blueprint.Serialize
-		const bool bDontUseParentFunction = ParentBlueprint && ParentBlueprint->HasAnyFlags(RF_BeingRegenerated) && (EKismetCompileType::SkeletonOnly == CompileOptions.CompileType);
-		UFunction* ParentFunction = (!bDontUseParentFunction) ? Blueprint->ParentClass->FindFunctionByName(NewFunctionName) : NULL;
+		UFunction* ParentFunction = Blueprint->ParentClass->FindFunctionByName(NewFunctionName);
 
 		const FString NewFunctionNameString = NewFunctionName.ToString();
 		if (CreatedFunctionNames.Contains(NewFunctionNameString))
@@ -3149,7 +3145,7 @@ void FKismetCompilerContext::Compile()
 						}
 					}
 
-					UEditorEngine::CopyPropertiesForUnrelatedObjects(OldCDO, NewCDO);			
+					UEditorEngine::CopyPropertiesForUnrelatedObjects(OldCDO, NewCDO);
 				}
 
 				// >>> Backwards Compatibility: Propagate data from the skel CDO to the gen CDO if we haven't already done so for this blueprint
diff --git a/Engine/Source/Editor/KismetCompiler/Private/KismetCompilerMisc.cpp b/Engine/Source/Editor/KismetCompiler/Private/KismetCompilerMisc.cpp
--- a/Engine/Source/Editor/KismetCompiler/Private/KismetCompilerMisc.cpp
+++ b/Engine/Source/Editor/KismetCompiler/Private/KismetCompilerMisc.cpp
@@ -153,7 +153,7 @@ bool FKismetCompilerUtilities::IsTypeCompatibleWithProperty(UEdGraphPin* SourceP
 	}
 	else if (PinCategory == Schema->PC_Delegate)
 	{
-		const UFunction* SignatureFunction = Cast<const UFunction>(PinSubCategoryObject);
+		const UFunction* SignatureFunction = FMemberReference::ResolveSimpleMemberReference<UFunction>(Type.PinSubCategoryMemberReference);
 		const UDelegateProperty* PropertyDelegate = Cast<const UDelegateProperty>(TestProperty);
 		bTypeMismatch = !(SignatureFunction 
 			&& PropertyDelegate 
@@ -223,12 +223,25 @@ bool FKismetCompilerUtilities::IsTypeCompatibleWithProperty(UEdGraphPin* SourceP
 			if (StructProperty != NULL)
 			{
 				DesiredSubType = StructProperty->Struct->GetName();
-				bSubtypeMismatch = bTypeMismatch = (StructType != StructProperty->Struct);
+				bool bMatchingStructs = (StructType == StructProperty->Struct);
+				if (auto UserDefinedStructFromProperty = Cast<const UUserDefinedStruct>(StructProperty->Struct))
+				{
+					bMatchingStructs |= (UserDefinedStructFromProperty->PrimaryStruct.Get() == StructType);
+				}
+				bSubtypeMismatch = bTypeMismatch = !bMatchingStructs;
 			}
 			else
 			{
 				bTypeMismatch = true;
 			}
+
+			if (OwningFunction && bTypeMismatch)
+			{
+				if (UK2Node_CallFunction::IsStructureWildcardProperty(OwningFunction, SourcePin->PinName))
+				{
+					bSubtypeMismatch = bTypeMismatch = false;
+				}
+			}
 		}
 	}
 	else
@@ -608,7 +621,7 @@ UProperty* FKismetCompilerUtilities::CreatePropertyOnScope(UStruct* Scope, const
 	}
 	else if (Type.PinCategory == Schema->PC_Delegate)
 	{
-		if (UFunction* SignatureFunction = Cast<UFunction>(Type.PinSubCategoryObject.Get()))
+		if (UFunction* SignatureFunction = FMemberReference::ResolveSimpleMemberReference<UFunction>(Type.PinSubCategoryMemberReference))
 		{
 			UDelegateProperty* NewPropertyDelegate = NewNamedObject<UDelegateProperty>(PropertyScope, ValidatedPropertyName, ObjectFlags);
 			NewPropertyDelegate->SignatureFunction = SignatureFunction;
@@ -617,7 +630,7 @@ UProperty* FKismetCompilerUtilities::CreatePropertyOnScope(UStruct* Scope, const
 	}
 	else if (Type.PinCategory == Schema->PC_MCDelegate)
 	{
-		UFunction* const SignatureFunction = Cast<UFunction>(Type.PinSubCategoryObject.Get());
+		UFunction* const SignatureFunction = FMemberReference::ResolveSimpleMemberReference<UFunction>(Type.PinSubCategoryMemberReference);
 		UMulticastDelegateProperty* NewPropertyDelegate = NewNamedObject<UMulticastDelegateProperty>(PropertyScope, ValidatedPropertyName, ObjectFlags);
 		NewPropertyDelegate->SignatureFunction = SignatureFunction;
 		NewProperty = NewPropertyDelegate;
