#ue4 BP: UDS: UDS cannot use blueprint references. Minor UI code fixes. (last cl in pretest missed the merging)

diff --git a/Engine/Source/Editor/BlueprintGraph/Private/EdGraphSchema_K2.cpp b/Engine/Source/Editor/BlueprintGraph/Private/EdGraphSchema_K2.cpp
--- a/Engine/Source/Editor/BlueprintGraph/Private/EdGraphSchema_K2.cpp
+++ b/Engine/Source/Editor/BlueprintGraph/Private/EdGraphSchema_K2.cpp
@@ -194,6 +194,7 @@ UEdGraphSchema_K2::FPinTypeTreeInfo::FPinTypeTreeInfo(const FString& CategoryNam
 	PinType.PinCategory = CategoryName;
 
 	SubCategoryObjectAssetReference = SubCategoryObject;
+	PinType.PinSubCategoryObject = SubCategoryObjectAssetReference.ResolveObject();
 
 	bReadOnly = bInReadOnly;
 }
@@ -2054,7 +2055,7 @@ void UEdGraphSchema_K2::ConstructBasicPinTooltip(const UEdGraphPin& Pin, const F
 
 	if (!PinDescription.IsEmpty())
 	{
-		ConstructedTooltip += FString(TEXT("\n")) + PinDescription;
+	ConstructedTooltip += FString(TEXT("\n")) + PinDescription;
 	}
 	TooltipOut = ConstructedTooltip; // using a local FString, so PinDescription and TooltipOut can be the same variable
 }
diff --git a/Engine/Source/Editor/Kismet/Private/UserDefinedStructureEditor.cpp b/Engine/Source/Editor/Kismet/Private/UserDefinedStructureEditor.cpp
--- a/Engine/Source/Editor/Kismet/Private/UserDefinedStructureEditor.cpp
+++ b/Engine/Source/Editor/Kismet/Private/UserDefinedStructureEditor.cpp
@@ -364,21 +364,23 @@ public:
 		return EVisibility::Collapsed;
 	}
 
-	void RemoveInvalidSubTypes(TSharedPtr<UEdGraphSchema_K2::FPinTypeTreeInfo> PinTypeNode, const UStruct* Parent) const
+	void RemoveInvalidSubTypes(TSharedPtr<UEdGraphSchema_K2::FPinTypeTreeInfo> PinTypeNode, const UUserDefinedStruct* Parent) const
 	{
+		if (!PinTypeNode.IsValid() || !Parent)
+		{
+			return;
+		}
+
 		for (int32 ChildIndex = 0; ChildIndex < PinTypeNode->Children.Num();)
 		{
-			auto Child = PinTypeNode->Children[ChildIndex];
+			const auto Child = PinTypeNode->Children[ChildIndex];
 			if(Child.IsValid())
 			{
-				auto ScriptStruct = Cast<const UScriptStruct>(Child->GetPinType(false).PinSubCategoryObject.Get());
-				if (ScriptStruct)
+				const bool bCanCheckSubObjectWithoutLoading = Child->GetPinType(false).PinSubCategoryObject.IsValid();
+				if (bCanCheckSubObjectWithoutLoading && !FStructureEditorUtils::CanHaveAMemberVariableOfType(Parent, Child->GetPinType(false)))
 				{
-					if (FStructureEditorUtils::EStructureError::Ok != FStructureEditorUtils::IsStructureValid(ScriptStruct, Parent))
-					{
-						PinTypeNode->Children.RemoveAt(ChildIndex);
-						continue;
-					}
+					PinTypeNode->Children.RemoveAt(ChildIndex);
+					continue;
 				}
 			}
 			++ChildIndex;
@@ -396,10 +398,7 @@ public:
 			// THE TREE HAS ONLY 2 LEVELS
 			for (auto PinTypePtr : TypeTree)
 			{
-				if (PinTypePtr.IsValid() && (K2Schema->PC_Struct == PinTypePtr->GetPinType(false).PinCategory))
-				{
-					RemoveInvalidSubTypes(PinTypePtr, Parent);
-				}
+				RemoveInvalidSubTypes(PinTypePtr, Parent);
 			}
 		}
 	}
diff --git a/Engine/Source/Editor/UnrealEd/Private/Kismet2/StructureEditorUtils.cpp b/Engine/Source/Editor/UnrealEd/Private/Kismet2/StructureEditorUtils.cpp
--- a/Engine/Source/Editor/UnrealEd/Private/Kismet2/StructureEditorUtils.cpp
+++ b/Engine/Source/Editor/UnrealEd/Private/Kismet2/StructureEditorUtils.cpp
@@ -170,6 +170,18 @@ bool FStructureEditorUtils::CanHaveAMemberVariableOfType(const UUserDefinedStruc
 		}
 		return false;
 	}
+	else
+	{
+		const auto PinSubCategoryClass = Cast<const UClass>(VarType.PinSubCategoryObject.Get());
+		if (PinSubCategoryClass && PinSubCategoryClass->IsChildOf(UBlueprint::StaticClass()))
+		{
+			if (OutMsg)
+			{
+				*OutMsg = LOCTEXT("StructureUseBlueprintReferences", "Struct cannot use any blueprint references").ToString();
+			}
+			return false;
+		}
+	}
 	return true;
 }
 
