Commits the text in textboxes in the details tab when focus is lost as well as when Enter is pressed.
#ttp 333379 - BLUEPRINTS: Make the textboxes in the details page to store the data if the focus was lost
#branch UE4
#proj Editor.Kismet
#change Changed various On..Committed() delegates in BlueprintDetailsCustomization.cpp to assign the property on focus lost as well as on Enter pressed.

diff --git a/Engine/Source/Editor/Kismet/Private/BlueprintDetailsCustomization.cpp b/Engine/Source/Editor/Kismet/Private/BlueprintDetailsCustomization.cpp
--- a/Engine/Source/Editor/Kismet/Private/BlueprintDetailsCustomization.cpp
+++ b/Engine/Source/Editor/Kismet/Private/BlueprintDetailsCustomization.cpp
@@ -717,7 +717,7 @@ void FBlueprintVarActionDetails::OnVarNameChanged(const FText& InNewText)
 
 void FBlueprintVarActionDetails::OnVarNameCommitted(const FText& InNewText, ETextCommit::Type InTextCommit)
 {
-	if(InTextCommit == ETextCommit::OnEnter && !bIsVarNameInvalid)
+	if((InTextCommit == ETextCommit::OnEnter || InTextCommit == ETextCommit::OnUserMovedFocus) && !bIsVarNameInvalid)
 	{
 		const FScopedTransaction Transaction( LOCTEXT( "RenameVariable", "Rename Variable" ) );
 
@@ -1618,7 +1618,7 @@ void FBlueprintLocalVarActionDetails::OnLocalVarNameChanged(const FText& InNewTe
 
 void FBlueprintLocalVarActionDetails::OnLocalVarNameCommitted(const FText& InNewText, ETextCommit::Type InTextCommit)
 {
-	if(InTextCommit == ETextCommit::OnEnter && !bIsVarNameInvalid)
+	if((InTextCommit == ETextCommit::OnEnter || InTextCommit == ETextCommit::OnUserMovedFocus) && !bIsVarNameInvalid)
 	{
 		const FScopedTransaction Transaction( LOCTEXT( "RenameLocalVariable", "Rename Local Variable" ) );
 
@@ -2022,7 +2022,7 @@ FText FBlueprintGraphArgumentLayout::OnGetArgDefaultValueText() const
 
 void FBlueprintGraphArgumentLayout::OnArgDefaultValueCommitted(const FText& NewText, ETextCommit::Type InTextCommit)
 {
-	if (!NewText.IsEmpty() && !ShouldPinBeReadOnly() && (InTextCommit == ETextCommit::OnEnter) && ParamItemPtr.IsValid())
+	if (!NewText.IsEmpty() && !ShouldPinBeReadOnly() && (InTextCommit == ETextCommit::OnEnter || InTextCommit == ETextCommit::OnUserMovedFocus) && ParamItemPtr.IsValid())
 	{
 		bool bSuccess = TargetNode->ModifyUserDefinedPinDefaultValue(ParamItemPtr.Pin(), NewText.ToString());
 		
