Increased the size of the blueprint variable name text entry box to allow for displaying larger variable names

diff --git a/Engine/Source/Editor/Kismet/Private/BlueprintDetailsCustomization.cpp b/Engine/Source/Editor/Kismet/Private/BlueprintDetailsCustomization.cpp
--- a/Engine/Source/Editor/Kismet/Private/BlueprintDetailsCustomization.cpp
+++ b/Engine/Source/Editor/Kismet/Private/BlueprintDetailsCustomization.cpp
@@ -66,6 +66,7 @@ void FBlueprintVarActionDetails::CustomizeDetails( IDetailLayoutBuilder& DetailL
 		.Font(IDetailLayoutBuilder::GetDetailFont())
 	]
 	.ValueContent()
+	.MaxDesiredWidth(250.0f)
 	[
 		SAssignNew(VarNameEditableTextBox, SEditableTextBox)
 		.Text(this, &FBlueprintVarActionDetails::OnGetVarName)
