BP, UI: UserDefinedStruct and UserDefinedEnum removed from MyBlueprint

diff --git a/Engine/Source/Editor/Kismet/Private/SMyBlueprint.cpp b/Engine/Source/Editor/Kismet/Private/SMyBlueprint.cpp
--- a/Engine/Source/Editor/Kismet/Private/SMyBlueprint.cpp
+++ b/Engine/Source/Editor/Kismet/Private/SMyBlueprint.cpp
@@ -50,12 +50,7 @@ void FMyBlueprintCommands::RegisterCommands()
 	UI_COMMAND( ImplementFunction, "Implement Function", "Implements this overridable function as a new function.", EUserInterfaceActionType::Button, FInputGesture() );
 	UI_COMMAND( FindEntry, "Find References", "Searches for all references of this function or variable.", EUserInterfaceActionType::Button, FInputGesture() );
 	UI_COMMAND(DeleteEntry, "Delete", "Deletes this function or variable from this blueprint.", EUserInterfaceActionType::Button, FInputGesture(EKeys::Platform_Delete));
-	UI_COMMAND( FindUserDefinedEnumInContentBrowser, "Find in Content Browser...", "Find user defined enum in content browser...", EUserInterfaceActionType::Button, FInputGesture() );
-	UI_COMMAND( AddNewUserDefinedEnum, "Create Enum Asset", "Create new user defined enum asset", EUserInterfaceActionType::Button, FInputGesture() );
 	UI_COMMAND( GotoNativeVarDefinition, "Goto Code Definition", "Goto the native code definition of this variable", EUserInterfaceActionType::Button, FInputGesture() );
-	UI_COMMAND( FindUserDefinedStructInContentBrowser, "Find in Content Browser...", "Find user defined struct in content browser...", EUserInterfaceActionType::Button, FInputGesture() );
-	UI_COMMAND( AddNewUserDefinedStruct, "Create Struct Asset", "Create new user defined struct asset", EUserInterfaceActionType::Button, FInputGesture() );
-
 }
 
 //////////////////////////////////////////////////////////////////////////
@@ -168,18 +163,6 @@ void SMyBlueprint::Construct(const FArguments& InArgs, TWeakPtr<FBlueprintEditor
 		FExecuteAction::CreateSP(this, &SMyBlueprint::OnFindEntry),
 		FCanExecuteAction(),
 		FIsActionButtonVisible::CreateSP(this, &SMyBlueprint::CanFindEntry) );
-
-	InBlueprintEditor.Pin()->GetToolkitCommands()->MapAction( FMyBlueprintCommands::Get().FindUserDefinedEnumInContentBrowser,
-		FExecuteAction::CreateSP(this, &SMyBlueprint::OnFindUserDefinedEnumInContentBrowser) );
-
-	InBlueprintEditor.Pin()->GetToolkitCommands()->MapAction( FMyBlueprintCommands::Get().FindUserDefinedStructInContentBrowser,
-		FExecuteAction::CreateSP(this, &SMyBlueprint::OnFindUserDefinedStructInContentBrowser) );
-
-	InBlueprintEditor.Pin()->GetToolkitCommands()->MapAction( FMyBlueprintCommands::Get().AddNewUserDefinedEnum,
-		FExecuteAction::CreateSP(this, &SMyBlueprint::AddNewUserDefinedEnum) );
-
-	InBlueprintEditor.Pin()->GetToolkitCommands()->MapAction( FMyBlueprintCommands::Get().AddNewUserDefinedStruct,
-		FExecuteAction::CreateSP(this, &SMyBlueprint::AddNewUserDefinedStruct) );
 	
 	ToolKitCommandList->MapAction( FMyBlueprintCommands::Get().DeleteEntry,
 		FExecuteAction::CreateSP(this, &SMyBlueprint::OnDeleteEntry),
@@ -598,112 +581,6 @@ EVisibility SMyBlueprint::GetLocalActionsListVisibility() const
 	return EVisibility::Collapsed;
 }
 
-struct FActionCollectingHelper
-{
-public:
-	static FString UserDefinedEnumCategoryNameStr()
-	{
-		return LOCTEXT("UserDefinedEnumCategory", "User defined enums").ToString();
-	}
-
-	static FString UserDefinedStructCategoryNameStr()
-	{
-		return LOCTEXT("UserDefinedStructCategory", "User defined structs").ToString();
-	}
-
-	static void AddUserEnumsActions(FGraphActionListBuilderBase& OutAllActions, TArray<UUserDefinedEnum*>& ElementsUsedInBlueprint)
-	{
-		for(auto Element : ElementsUsedInBlueprint)
-		{
-			if(Element)
-			{
-				const FString EnumNameStr = Element->GetName();
-				TSharedPtr<FEdGraphSchemaAction_K2Enum> NewEnumAction = MakeShareable(new FEdGraphSchemaAction_K2Enum( 
-					UserDefinedEnumCategoryNameStr(), FText::FromString(EnumNameStr), EnumNameStr, -1));
-				NewEnumAction->Enum = Element;
-				NewEnumAction->SectionID = NodeSectionID::USER_ENUM;
-				OutAllActions.AddAction(NewEnumAction);
-			}
-		}
-	}
-
-	static void AddUserStructActions(FGraphActionListBuilderBase& OutAllActions, TArray<UUserDefinedStruct*>& ElementsUsedInBlueprint)
-	{
-		for(auto Element : ElementsUsedInBlueprint)
-		{
-			if(Element)
-			{
-				const FString NameStr = Element->GetName();
-				TSharedPtr<FEdGraphSchemaAction_K2Struct> NewEnumAction = MakeShareable(new FEdGraphSchemaAction_K2Struct( 
-					UserDefinedStructCategoryNameStr(), FText::FromString(NameStr), NameStr, -1));
-				NewEnumAction->Struct = Element;
-				NewEnumAction->SectionID = NodeSectionID::USER_STRUCT;
-				OutAllActions.AddAction(NewEnumAction);
-			}
-		}
-	}
-
-	static UUserDefinedEnum* FromProperty(UByteProperty* Property)
-	{
-		return Property ? Cast<UUserDefinedEnum>(Property->GetIntPropertyEnum()) : NULL;
-	}
-
-	static UUserDefinedStruct* FromProperty(UStructProperty* Property)
-	{
-		return Property ? Cast<UUserDefinedStruct>(Property->Struct) : NULL;
-	}
-
-	template<class TClassOfType, class TPropertyType>
-	static void CollectAllUsedTypes(FGraphActionListBuilderBase& OutAllActions, TArray<TClassOfType*>& ElementsUsedInBlueprint)
-	{
-		for ( int32 GraphActionIndex = 0; GraphActionIndex < OutAllActions.GetNumActions(); GraphActionIndex++ )
-		{
-			FGraphActionListBuilderBase::ActionGroup& CurrentGraphAction = OutAllActions.GetAction(GraphActionIndex);
-			for(auto ActionIter = CurrentGraphAction.Actions.CreateIterator(); ActionIter; ActionIter++)
-			{
-				TSharedPtr<FEdGraphSchemaAction> Action = *ActionIter;
-				if(Action.IsValid())
-				{
-					if(FEdGraphSchemaAction_K2Var::StaticGetTypeId() == Action->GetTypeId())
-					{
-						FEdGraphSchemaAction_K2Var* VarAction = (FEdGraphSchemaAction_K2Var*)Action.Get();
-						if(auto Property = Cast<TPropertyType>(VarAction->GetProperty()))
-						{
-							if(auto Element = FromProperty(Property))
-							{
-								ElementsUsedInBlueprint.AddUnique(Element);
-							}
-						}
-					}
-					else if(FEdGraphSchemaAction_K2Graph::StaticGetTypeId() == Action->GetTypeId())
-					{
-						FEdGraphSchemaAction_K2Graph* GraphAction = (FEdGraphSchemaAction_K2Graph*)Action.Get();
-						if(GraphAction->EdGraph)
-						{
-							for(auto NodeIter = GraphAction->EdGraph->Nodes.CreateIterator(); NodeIter; NodeIter++)
-							{
-								if(UEdGraphNode* Node = *NodeIter)
-								{
-									for(auto PinIter = Node->Pins.CreateIterator(); PinIter; PinIter++)
-									{
-										if(UEdGraphPin* Pin = *PinIter)
-										{
-											if(auto Element =  Cast<TClassOfType>(Pin->PinType.PinSubCategoryObject.Get()))
-											{
-												ElementsUsedInBlueprint.AddUnique(Element);
-											}
-										}
-									}
-								}
-							}
-						}
-					}
-				}
-			}
-		}
-	}
-};
-
 void SMyBlueprint::CollectAllActions(FGraphActionListBuilderBase& OutAllActions)
 {
 	const UEdGraphSchema_K2* K2Schema = GetDefault<UEdGraphSchema_K2>();
@@ -951,41 +828,6 @@ void SMyBlueprint::CollectAllActions(FGraphActionListBuilderBase& OutAllActions)
 		GetChildGraphs(Graph, OutAllActions);
 		GetChildEvents(Graph, NewFuncAction->SectionID, OutAllActions);
 	}
-
-	{
-		TArray<UUserDefinedEnum*> UserDefinedEnumUserInBlueprint;
-		FActionCollectingHelper::CollectAllUsedTypes<UUserDefinedEnum, UByteProperty>(OutAllActions, UserDefinedEnumUserInBlueprint);
-		for(auto EnumWeakIter = EnumsAddedToBlueprint.CreateIterator(); EnumWeakIter; EnumWeakIter++)
-		{
-			if(EnumWeakIter->IsValid())
-			{
-				UserDefinedEnumUserInBlueprint.AddUnique(EnumWeakIter->Get());
-			}
-		}
-		FActionCollectingHelper::AddUserEnumsActions(OutAllActions, UserDefinedEnumUserInBlueprint);
-	}
-
-	if (FStructureEditorUtils::UserDefinedStructEnabled())
-	{
-		TArray<UUserDefinedStruct*> UserDefinedStructUserInBlueprint;
-		FActionCollectingHelper::CollectAllUsedTypes<UUserDefinedStruct, UStructProperty>(OutAllActions, UserDefinedStructUserInBlueprint);
-		// Make sure, that UDStructs from invalid variables are also listed
-		for (auto& VarDesc : Blueprint->NewVariables)
-		{
-			if (auto UDStruct = Cast<UUserDefinedStruct>(VarDesc.VarType.PinSubCategoryObject.Get()))
-			{
-				UserDefinedStructUserInBlueprint.AddUnique(UDStruct);
-			}
-		}
-		for(auto StructWeakPtr : StructsAddedToBlueprint)
-		{
-			if(StructWeakPtr.IsValid())
-			{
-				UserDefinedStructUserInBlueprint.AddUnique(StructWeakPtr.Get());
-			}
-		}
-		FActionCollectingHelper::AddUserStructActions(OutAllActions, UserDefinedStructUserInBlueprint);
-	}
 }
 
 ESlateCheckBoxState::Type SMyBlueprint::OnUserVarsCheckState() const
@@ -1350,16 +1192,6 @@ TSharedPtr<SWidget> SMyBlueprint::OnContextMenuOpening()
 	// Check if the selected action is valid for a context menu
 	if (SelectionHasContextMenu())
 	{
-		if(SelectionAsEnum())
-		{
-			MenuBuilder.AddMenuEntry(FMyBlueprintCommands::Get().FindUserDefinedEnumInContentBrowser);
-		}
-		
-		if(SelectionAsStruct())
-		{
-			MenuBuilder.AddMenuEntry(FMyBlueprintCommands::Get().FindUserDefinedStructInContentBrowser);
-		}
-
 		MenuBuilder.BeginSection("BasicOperations");
 		{
 			MenuBuilder.AddMenuEntry(FMyBlueprintCommands::Get().OpenGraph);
@@ -1410,129 +1242,11 @@ TSharedPtr<SWidget> SMyBlueprint::OnContextMenuOpening()
 			MenuBuilder.AddMenuEntry(FBlueprintEditorCommands::Get().AddNewDelegate);
 		}
 		MenuBuilder.EndSection();
-
-		if(BlueprintEditorPtr.IsValid())
-		{
-			UBlueprint* Blueprint = BlueprintEditorPtr.Pin()->GetBlueprintObj();
-			if(Blueprint && (EBlueprintType::BPTYPE_LevelScript != Blueprint->BlueprintType))
-			{
-				MenuBuilder.AddMenuSeparator();
-				MenuBuilder.AddMenuEntry(FMyBlueprintCommands::Get().AddNewUserDefinedEnum);
-				if (FStructureEditorUtils::UserDefinedStructEnabled())
-				{
-					MenuBuilder.AddMenuEntry(FMyBlueprintCommands::Get().AddNewUserDefinedStruct);
-				}
-			}
-		}
 	}
 
 	return MenuBuilder.MakeWidget();
 }
 
-void SMyBlueprint::OnFindUserDefinedEnumInContentBrowser() const
-{
-	if (FEdGraphSchemaAction_K2Enum* GraphAction = SelectionAsEnum())
-	{
-		UUserDefinedEnum* Enum = Cast<UUserDefinedEnum>(GraphAction->Enum);
-		if(Enum && Enum->IsAsset())
-		{
-			TArray<UObject*> ObjectsToSync;
-				ObjectsToSync.Add(Enum);
-				GEditor->SyncBrowserToObjects( ObjectsToSync );
-		}
-	}
-}
-
-void SMyBlueprint::OnFindUserDefinedStructInContentBrowser() const
-{
-	if (FEdGraphSchemaAction_K2Struct* GraphAction = SelectionAsStruct())
-	{
-		UUserDefinedStruct* Struct = Cast<UUserDefinedStruct>(GraphAction->Struct);
-		if(Struct && Struct->IsAsset())
-		{
-			TArray<UObject*> ObjectsToSync;
-			ObjectsToSync.Add(Struct);
-			GEditor->SyncBrowserToObjects( ObjectsToSync );
-		}
-	}
-}
-
-void SMyBlueprint::AddNewUserDefinedEnum()
-{
-	UBlueprint* Blueprint = GetBlueprintObj();
-	if (Blueprint)
-	{
-		static FName AssetToolsModuleName = FName("AssetTools");
-		FAssetToolsModule& AssetToolsModule = FModuleManager::GetModuleChecked<FAssetToolsModule>(AssetToolsModuleName);
-
-		FString AssetName;
-		FString PackagePath;
-		AssetToolsModule.Get().CreateUniqueAssetName(Blueprint->GetOutermost()->GetName(), TEXT("_Enum"), PackagePath, AssetName);
-
-		UUserDefinedEnum* Enum = NULL;
-		if (!AssetName.IsEmpty())
-		{
-			UEnumFactory* Factory = ConstructObject<UEnumFactory>(UEnumFactory::StaticClass());
-			if (Factory && Factory->ConfigureProperties())
-			{
-				PackagePath = FPackageName::GetLongPackagePath(Blueprint->GetPathName());
-				UObject* NewAsset = AssetToolsModule.Get().CreateAsset(AssetName, PackagePath, UUserDefinedEnum::StaticClass(), Factory );
-				Enum = Cast<UUserDefinedEnum>(NewAsset);
-			}
-		}
-
-		if (Enum)
-		{
-			EnumsAddedToBlueprint.Add(Enum);
-			Refresh();
-			if (GraphActionMenu.IsValid())
-			{
-				GraphActionMenu->ExpandCategory(FActionCollectingHelper::UserDefinedEnumCategoryNameStr());
-				GraphActionMenu->SelectItemByName(FName(*Enum->GetPathName()),ESelectInfo::OnMouseClick);
-				GraphActionMenu->OnRequestRenameOnActionNode();
-			}
-		}
-	}
-}
-
-void SMyBlueprint::AddNewUserDefinedStruct()
-{
-	UBlueprint* Blueprint = GetBlueprintObj();
-	if (Blueprint)
-	{
-		static FName AssetToolsModuleName = FName("AssetTools");
-		FAssetToolsModule& AssetToolsModule = FModuleManager::GetModuleChecked<FAssetToolsModule>(AssetToolsModuleName);
-
-		FString AssetName;
-		FString PackagePath;
-		AssetToolsModule.Get().CreateUniqueAssetName(Blueprint->GetOutermost()->GetName(), TEXT("_Struct"), PackagePath, AssetName);
-
-		UUserDefinedStruct* Struct = NULL;
-		if (!AssetName.IsEmpty())
-		{
-			auto* Factory = ConstructObject<UStructureFactory>(UStructureFactory::StaticClass());
-			if (Factory && Factory->ConfigureProperties())
-			{
-				PackagePath = FPackageName::GetLongPackagePath(Blueprint->GetPathName());
-				UObject* NewAsset = AssetToolsModule.Get().CreateAsset(AssetName, PackagePath, UUserDefinedStruct::StaticClass(), Factory );
-				Struct = Cast<UUserDefinedStruct>(NewAsset);
-			}
-		}
-
-		if (Struct)
-		{
-			StructsAddedToBlueprint.Add(Struct);
-			Refresh();
-			if (GraphActionMenu.IsValid())
-			{
-				GraphActionMenu->ExpandCategory(FActionCollectingHelper::UserDefinedStructCategoryNameStr());
-				GraphActionMenu->SelectItemByName(FName(*Struct->GetPathName()),ESelectInfo::OnMouseClick);
-				GraphActionMenu->OnRequestRenameOnActionNode();
-			}
-		}
-	}
-}
-
 bool SMyBlueprint::CanOpenGraph() const
 {
 	const FEdGraphSchemaAction_K2Graph* GraphAction = SelectionAsGraph();
diff --git a/Engine/Source/Editor/Kismet/Private/SMyBlueprint.h b/Engine/Source/Editor/Kismet/Private/SMyBlueprint.h
--- a/Engine/Source/Editor/Kismet/Private/SMyBlueprint.h
+++ b/Engine/Source/Editor/Kismet/Private/SMyBlueprint.h
@@ -38,11 +38,7 @@ public:
 	TSharedPtr<FUICommandInfo> ImplementFunction;
 	TSharedPtr<FUICommandInfo> FindEntry;
 	TSharedPtr<FUICommandInfo> DeleteEntry;
-	TSharedPtr<FUICommandInfo> FindUserDefinedEnumInContentBrowser;
-	TSharedPtr<FUICommandInfo> AddNewUserDefinedEnum;
 	TSharedPtr<FUICommandInfo> GotoNativeVarDefinition;
-	TSharedPtr<FUICommandInfo> FindUserDefinedStructInContentBrowser;
-	TSharedPtr<FUICommandInfo> AddNewUserDefinedStruct;
 	// Add New Item
 	/** Initialize commands */
 	virtual void RegisterCommands() OVERRIDE;
@@ -173,8 +169,6 @@ private:
 	bool CanFindEntry() const;
 	void OnDeleteEntry();
 	bool CanDeleteEntry() const;
-	void OnFindUserDefinedEnumInContentBrowser() const;
-	void AddNewUserDefinedEnum();
 	FReply OnAddNewLocalVariable();
 	bool CanRequestRenameOnActionNode() const;
 	bool IsDuplicateActionVisible() const;
@@ -182,8 +176,6 @@ private:
 	void OnDuplicateAction();
 	void GotoNativeCodeVarDefinition();
 	bool IsNativeVariable() const;
-	void OnFindUserDefinedStructInContentBrowser() const;
-	void AddNewUserDefinedStruct();
 
 	/** Callback when the filter is changed, forces the action tree(s) to filter */
 	void OnFilterTextChanged( const FText& InFilterText );
