Fix for the blueprint goto code definition context menu being incorrectly unavailable on some native function nodes.
#TTP 333880 - Goto Code Definition is unavailable on blueprint nodes it should be.
#Branch UE4
#Proj  Kismet

#Change Fixed the way that the owning class is determined from the UFunction as it previously was returning incorrect results and preventing the item from being available.

ReviewdBy Chris.Wood, Nick.Whiting

diff --git a/Engine/Source/Editor/Kismet/Private/BlueprintEditor.cpp b/Engine/Source/Editor/Kismet/Private/BlueprintEditor.cpp
--- a/Engine/Source/Editor/Kismet/Private/BlueprintEditor.cpp
+++ b/Engine/Source/Editor/Kismet/Private/BlueprintEditor.cpp
@@ -5956,7 +5956,8 @@ bool FBlueprintEditor::IsSelectionNativeFunction()
 
 		if( FunctionNode && SelectedNodes.Num() == 1 )
 		{
-			UClass* OwningClass = FunctionNode->FunctionReference.GetMemberParentClass( FunctionNode );
+			UFunction* FunctionPtr = FunctionNode->FunctionReference.ResolveMember<UFunction>( FunctionNode );
+			UClass* OwningClass = FunctionPtr ? FunctionPtr->GetOuterUClass() : NULL;
 
 			if( OwningClass && OwningClass->HasAllClassFlags( CLASS_Native ))
 			{
