Modified variable Get/Set nodes dragdrop operations to be transactional so the blueprint is marked as dirty and undo/redo works when dropping a different variable on the node.
#TTP 314504 - UE4: K2: REGRESSION: Drag-dropping a variable on to a variable get node does not stick (it reverts to the old variable name when the BP is reloaded)!
#Branch UE4
#Proj  Kismet

diff --git a/Engine/Source/Editor/Kismet/Private/BPVariableDragDropAction.cpp b/Engine/Source/Editor/Kismet/Private/BPVariableDragDropAction.cpp
--- a/Engine/Source/Editor/Kismet/Private/BPVariableDragDropAction.cpp
+++ b/Engine/Source/Editor/Kismet/Private/BPVariableDragDropAction.cpp
@@ -274,6 +274,8 @@ FReply FKismetVariableDragDropAction::DroppedOnNode(FVector2D ScreenPosition, FV
 
 	if (TargetNode && (VariableName != TargetNode->GetVarName()))
 	{
+		const FScopedTransaction Transaction( LOCTEXT("ReplacePinVariable", "Replace Pin Variable") );
+
 		UProperty* VariableProperty = GetVariableProperty();
 		const FString OldVarName = TargetNode->GetVarNameString();
 		const UEdGraphSchema_K2* Schema = Cast<const UEdGraphSchema_K2>(TargetNode->GetSchema());
@@ -283,9 +285,13 @@ FReply FKismetVariableDragDropAction::DroppedOnNode(FVector2D ScreenPosition, FV
 
 		// Change the variable name and context
 		UBlueprint* DropOnBlueprint = FBlueprintEditorUtils::FindBlueprintForGraph(TargetNode->GetGraph());
+		UEdGraphPin* Pin = TargetNode->FindPin(OldVarName);
+		DropOnBlueprint->Modify();
+		TargetNode->Modify();
+		Pin->Modify();
+
 		ConfigureVarNode(TargetNode, VariableName, VariableSourceClass.Get(), DropOnBlueprint);
 
-		UEdGraphPin* Pin = TargetNode->FindPin(OldVarName);
 
 		if ((Pin == NULL) || (Pin->LinkedTo.Num() == BadLinks.Num()) || (Schema == NULL))
 		{
