Fix for unity build caused by Changelist 2065387

diff --git a/Engine/Source/Editor/Kismet/Private/BPVariableDragDropAction.cpp b/Engine/Source/Editor/Kismet/Private/BPVariableDragDropAction.cpp
--- a/Engine/Source/Editor/Kismet/Private/BPVariableDragDropAction.cpp
+++ b/Engine/Source/Editor/Kismet/Private/BPVariableDragDropAction.cpp
@@ -7,6 +7,7 @@
 #include "GraphEditorDragDropAction.h"
 #include "BPVariableDragDropAction.h"
 #include "SBlueprintPalette.h"
+#include "ScopedTransaction.h"
 
 #define LOCTEXT_NAMESPACE "VariableDragDropAction"
 
