Not clearing the "Preview Instance" debug selection for anim blueprints (reselecting debug objects if they share a name with the previous selection).
(TTP #333991)

diff --git a/Engine/Source/Editor/Kismet/Private/SBlueprintEditorSelectedDebugObjectWidget.cpp b/Engine/Source/Editor/Kismet/Private/SBlueprintEditorSelectedDebugObjectWidget.cpp
--- a/Engine/Source/Editor/Kismet/Private/SBlueprintEditorSelectedDebugObjectWidget.cpp
+++ b/Engine/Source/Editor/Kismet/Private/SBlueprintEditorSelectedDebugObjectWidget.cpp
@@ -143,11 +143,16 @@ void SBlueprintEditorSelectedDebugObjectWidget::Tick(const FGeometry& AllottedGe
 		{
 			if (Object != LastObjectObserved.Get())
 			{
-				LastObjectObserved = Object;
-
-				GenerateDebugObjectNames(false);
-
-				TSharedPtr<FString> NewSelection;
+				// bRestoreSelection attempts to restore the selection by name, 
+				// this ensures that if the last object we had selected was 
+				// regenerated (spawning a new object), then we select that  
+				// again, even if it is technically a different object
+				GenerateDebugObjectNames(/*bRestoreSelection =*/true);
+
+				TSharedPtr<FString> NewSelection = DebugObjectsComboBox->GetSelectedItem();
+				// just in case that object we want to select is actually in 
+				// there (and wasn't caught by bRestoreSelection), then let's 
+				// favor that over whatever was picked
 				for (int32 Index = 0; Index < DebugObjects.Num(); ++Index)
 				{
 					if (DebugObjects[Index] == Object)
@@ -163,6 +168,7 @@ void SBlueprintEditorSelectedDebugObjectWidget::Tick(const FGeometry& AllottedGe
 				}
 
 				DebugObjectsComboBox->SetSelectedItem(NewSelection);
+				LastObjectObserved = Object;
 			}
 		}
 		else
