Default and details inspectors are now disabled when editing a readonly graph.
#TTP 331389 - PERSONA: ANIM: Blueprints can be edited and compiled while the editor is simulating.

#branch UE4

#proj Editor.Kismet

#add Added new method IsFocusedGraphEditable which checks whether the currently focused graph is editable, then used this method to set the enabled state on the default inspector and details inspector.

diff --git a/Engine/Source/Editor/Kismet/Private/BlueprintEditor.cpp b/Engine/Source/Editor/Kismet/Private/BlueprintEditor.cpp
--- a/Engine/Source/Editor/Kismet/Private/BlueprintEditor.cpp
+++ b/Engine/Source/Editor/Kismet/Private/BlueprintEditor.cpp
@@ -1493,7 +1493,7 @@ void FBlueprintEditor::CreateDefaultTabContents(const TArray<UBlueprint*>& InBlu
 
 		this->Palette = 
 			SNew(SBlueprintPalette, SharedThis(this))
-				.IsEnabled(this, &FBlueprintEditor::InEditingMode);
+				.IsEnabled(this, &FBlueprintEditor::IsFocusedGraphEditable);
 	}
 
 	FIsPropertyEditingEnabled IsPropertyEditingEnabledDelegate = FIsPropertyEditingEnabled::CreateSP(this, &FBlueprintEditor::IsPropertyEditingEnabled);
@@ -1523,7 +1523,8 @@ void FBlueprintEditor::CreateDefaultTabContents(const TArray<UBlueprint*>& InBlu
 		. ViewIdentifier(FName("BlueprintInspector"))
 		. Kismet2(SharedThis(this))
 		. IsPropertyEditingEnabledDelegate( IsPropertyEditingEnabledDelegate )
-		. OnFinishedChangingProperties( FOnFinishedChangingProperties::FDelegate::CreateSP(this, &FBlueprintEditor::OnFinishedChangingProperties) );
+		. OnFinishedChangingProperties( FOnFinishedChangingProperties::FDelegate::CreateSP(this, &FBlueprintEditor::OnFinishedChangingProperties) )
+		. IsEnabled(this, &FBlueprintEditor::IsFocusedGraphEditable);
 
 	if (InBlueprints.Num() > 0)
 	{
@@ -1533,11 +1534,13 @@ void FBlueprintEditor::CreateDefaultTabContents(const TArray<UBlueprint*>& InBlu
 		const bool bShowPublicView = true;
 		const bool bHideNameArea = false;
 
+		TAttribute<bool> IsEnabled = (bIsInterface) ? false : TAttribute<bool>::Create(TAttribute<bool>::FGetter::CreateSP(this, &FBlueprintEditor::IsFocusedGraphEditable));
+
 		this->DefaultEditor = 
 			SNew(SKismetInspector)
 			. Kismet2(SharedThis(this))
 			. ViewIdentifier(FName("BlueprintDefaults"))
-			. IsEnabled(!bIsInterface)
+			. IsEnabled(IsEnabled)
 			. ShowPublicViewControl(bShowPublicView)
 			. ShowTitleArea(bShowTitle)
 			. HideNameArea(bHideNameArea)
@@ -6209,6 +6212,11 @@ bool FBlueprintEditor::IsGraphInCurrentBlueprint(UEdGraph* InGraph) const
 	return bEditable;
 }
 
+bool FBlueprintEditor::IsFocusedGraphEditable() const
+{
+	return IsEditable(GetFocusedGraph());
+}
+
 /////////////////////////////////////////////////////
 /////////////////////////////////////////////////////
 /////////////////////////////////////////////////////
diff --git a/Engine/Source/Editor/Kismet/Public/BlueprintEditor.h b/Engine/Source/Editor/Kismet/Public/BlueprintEditor.h
--- a/Engine/Source/Editor/Kismet/Public/BlueprintEditor.h
+++ b/Engine/Source/Editor/Kismet/Public/BlueprintEditor.h
@@ -826,6 +826,9 @@ private:
 	 */
 	bool IsEditingAnimGraph() const;
 
+	/** Returns whether the currently focused graph is editable or not */
+	bool IsFocusedGraphEditable() const;
+
 public://@TODO
 	TSharedPtr<FDocumentTracker> DocumentManager;
 
