Fix for issue introduced by CL 2066955

Details panel would be disabled by default when creating a new blueprint, this was due to checking whether the graph was editable however at that point no graph had been created.

Possible more correct fix would be to editor IsEditable to allow for a null graph, however this would result in changes in several areas of the blueprint editor. In the interest in fixing this issue blocker asap I opted for the more scoped fix to IsFocusedGraphEditable.

The Inspectors now have their scroll bars still enabled when they are disabled too, though this makes them inconsistant with other windows suchs as the pallet in the blueprint editor. These are seperate widgets though and would need addressing independantly.

diff --git a/Engine/Source/Editor/Kismet/Private/BlueprintEditor.cpp b/Engine/Source/Editor/Kismet/Private/BlueprintEditor.cpp
--- a/Engine/Source/Editor/Kismet/Private/BlueprintEditor.cpp
+++ b/Engine/Source/Editor/Kismet/Private/BlueprintEditor.cpp
@@ -1496,7 +1496,7 @@ void FBlueprintEditor::CreateDefaultTabContents(const TArray<UBlueprint*>& InBlu
 				.IsEnabled(this, &FBlueprintEditor::IsFocusedGraphEditable);
 	}
 
-	FIsPropertyEditingEnabled IsPropertyEditingEnabledDelegate = FIsPropertyEditingEnabled::CreateSP(this, &FBlueprintEditor::IsPropertyEditingEnabled);
+	FIsPropertyEditingEnabled IsPropertyEditingEnabledDelegate = FIsPropertyEditingEnabled::CreateSP(this, &FBlueprintEditor::IsFocusedGraphEditable);
 
 	if (IsEditingSingleBlueprint())
 	{
@@ -1523,8 +1523,7 @@ void FBlueprintEditor::CreateDefaultTabContents(const TArray<UBlueprint*>& InBlu
 		. ViewIdentifier(FName("BlueprintInspector"))
 		. Kismet2(SharedThis(this))
 		. IsPropertyEditingEnabledDelegate( IsPropertyEditingEnabledDelegate )
-		. OnFinishedChangingProperties( FOnFinishedChangingProperties::FDelegate::CreateSP(this, &FBlueprintEditor::OnFinishedChangingProperties) )
-		. IsEnabled(this, &FBlueprintEditor::IsFocusedGraphEditable);
+		. OnFinishedChangingProperties( FOnFinishedChangingProperties::FDelegate::CreateSP(this, &FBlueprintEditor::OnFinishedChangingProperties) );
 
 	if (InBlueprints.Num() > 0)
 	{
@@ -1534,13 +1533,11 @@ void FBlueprintEditor::CreateDefaultTabContents(const TArray<UBlueprint*>& InBlu
 		const bool bShowPublicView = true;
 		const bool bHideNameArea = false;
 
-		TAttribute<bool> IsEnabled = (bIsInterface) ? false : TAttribute<bool>::Create(TAttribute<bool>::FGetter::CreateSP(this, &FBlueprintEditor::IsFocusedGraphEditable));
-
 		this->DefaultEditor = 
 			SNew(SKismetInspector)
 			. Kismet2(SharedThis(this))
 			. ViewIdentifier(FName("BlueprintDefaults"))
-			. IsEnabled(IsEnabled)
+			. IsEnabled(!bIsInterface)
 			. ShowPublicViewControl(bShowPublicView)
 			. ShowTitleArea(bShowTitle)
 			. HideNameArea(bHideNameArea)
@@ -6214,7 +6211,12 @@ bool FBlueprintEditor::IsGraphInCurrentBlueprint(UEdGraph* InGraph) const
 
 bool FBlueprintEditor::IsFocusedGraphEditable() const
 {
-	return IsEditable(GetFocusedGraph());
+	UEdGraph* FocusedGraph = GetFocusedGraph();
+	if (FocusedGraph != nullptr)
+	{
+		return IsEditable(FocusedGraph);
+	}
+	return true;
 }
 
 /////////////////////////////////////////////////////
