BP: Expanding Macros/Functions will cause all contained nodes to be deleted in next editor session ttp334938

diff --git a/Engine/Source/Editor/Kismet/Private/BlueprintEditor.cpp b/Engine/Source/Editor/Kismet/Private/BlueprintEditor.cpp
--- a/Engine/Source/Editor/Kismet/Private/BlueprintEditor.cpp
+++ b/Engine/Source/Editor/Kismet/Private/BlueprintEditor.cpp
@@ -5106,6 +5106,16 @@ void FBlueprintEditor::ExpandNode(UEdGraphNode* InNodeToExpand, UEdGraph* InSour
 		Node->Modify();
 
 		Node->Rename(/*NewName=*/ NULL, /*NewOuter=*/ DestinationGraph);
+		if (Node->HasAllFlags(RF_Transient) && !DestinationGraph->HasAllFlags(RF_Transient))
+		{
+			Node->ClearFlags(RF_Transient);
+			TArray<UObject*> Subobjects;
+			GetObjectsWithOuter(Node, Subobjects);
+			for (auto Subobject : Subobjects)
+			{
+				Subobject->ClearFlags(RF_Transient);
+			}
+		}
 		DestinationGraph->Nodes.Add(Node);
 
 		// Want to test exactly against tunnel, we shouldn't collapse embedded collapsed
