Local function variables in Blueprints now utilize the Get/Set nodes

#change Deprecated UK2Node_LocalVariable
#change UK2Node_LocalVariable will replace itself with a Get and any Assign nodes it is hooked up to with a Set node, will also delete itself from Animation rule graphs that they should not have been allowed to be placed in

#ttp 332158 - BP:  Make local variables use standard get / set nodes, instead of the assigned named local nodes

diff --git a/Engine/Source/Editor/BlueprintGraph/Classes/EdGraphSchema_K2_Actions.h b/Engine/Source/Editor/BlueprintGraph/Classes/EdGraphSchema_K2_Actions.h
--- a/Engine/Source/Editor/BlueprintGraph/Classes/EdGraphSchema_K2_Actions.h
+++ b/Engine/Source/Editor/BlueprintGraph/Classes/EdGraphSchema_K2_Actions.h
@@ -533,6 +533,64 @@ public:
 	// End of FEdGraphSchemaAction interface
 };
 
+/*******************************************************************************
+* FEdGraphSchemaAction_K2LocalVar
+*******************************************************************************/
+
+/** Reference to a local variable (only used in 'docked' palette) */
+USTRUCT()
+struct BLUEPRINTGRAPH_API FEdGraphSchemaAction_K2LocalVar : public FEdGraphSchemaAction
+{
+	GENERATED_USTRUCT_BODY()
+
+private:
+	/** Name of function or class */
+	FName		VarName;
+
+	/** The struct that owns this item */
+	TWeakObjectPtr<UStruct>		VariableScope;
+
+public:
+	void SetVariableInfo(const FName& InVarName, const UStruct* InVariableScope)
+	{
+		VarName = InVarName;
+
+		check(InVariableScope);
+		VariableScope = InVariableScope;
+	}
+
+	// Simple type info
+	static FString StaticGetTypeId() {static FString Type = TEXT("FEdGraphSchemaAction_K2LocalVar"); return Type;}
+	virtual FString GetTypeId() const OVERRIDE { return StaticGetTypeId(); } 
+
+	FEdGraphSchemaAction_K2LocalVar() 
+		: FEdGraphSchemaAction()
+	{}
+
+	FEdGraphSchemaAction_K2LocalVar (const FString& InNodeCategory, const FText& InMenuDesc, const FString& InToolTip, const int32 InGrouping)
+		: FEdGraphSchemaAction(InNodeCategory, InMenuDesc, InToolTip, InGrouping)
+	{}
+
+	FName GetVariableName() const
+	{
+		return VarName;
+	}
+
+	UStruct* GetVariableScope() const
+	{
+		return VariableScope.Get();
+	}
+
+	UProperty* GetProperty() const
+	{
+		UStruct* SearchScope = GetVariableScope();
+		return FindField<UProperty>(SearchScope, VarName);
+	}
+
+	// FEdGraphSchemaAction interface
+	// End of FEdGraphSchemaAction interface
+};
+
 /*******************************************************************************
 * FEdGraphSchemaAction_K2Graph
 *******************************************************************************/
diff --git a/Engine/Source/Editor/BlueprintGraph/Classes/K2Node.h b/Engine/Source/Editor/BlueprintGraph/Classes/K2Node.h
--- a/Engine/Source/Editor/BlueprintGraph/Classes/K2Node.h
+++ b/Engine/Source/Editor/BlueprintGraph/Classes/K2Node.h
@@ -361,6 +361,10 @@ protected:
 	UPROPERTY()
 	mutable TSubclassOf<class UObject> MemberParentClass;
 
+	/** Class that this member is defined in. Should be NULL if bSelfContext is true.  */
+	UPROPERTY()
+	mutable FString MemberScope;
+
 	/** Name of variable */
 	UPROPERTY()
 	mutable FName MemberName;
@@ -432,6 +436,12 @@ public:
 	/** Set up this reference to a 'self' member name */
 	BLUEPRINTGRAPH_API void SetSelfMember(FName InMemberName);
 
+	/** Set up this reference to a 'self' member name, scoped to a struct */
+	BLUEPRINTGRAPH_API void SetLocalMember(FName InMemberName, UStruct* InScope, const FGuid InMemberGuid);
+
+	/** Set up this reference to a 'self' member name, scoped to a struct name */
+	BLUEPRINTGRAPH_API void SetLocalMember(FName InMemberName, FString InScopeName, const FGuid InMemberGuid);
+
 	/** Only intended for backwards compat! */
 	BLUEPRINTGRAPH_API void SetDirect(const FName InMemberName, const FGuid InMemberGuid, TSubclassOf<class UObject> InMemberParentClass, bool bIsConsideredSelfContext);
 
@@ -455,10 +465,22 @@ public:
 		return bSelfContext;
 	}
 
+	/** Returns if this is a local scope. */
+	bool IsLocalScope() const
+	{
+		return !MemberScope.IsEmpty();
+	}
 private:
 	/** Util to get the generated class from a node. */
 	BLUEPRINTGRAPH_API static UClass* GetBlueprintClassFromNode(const UK2Node* Node);
 
+	/**
+	 * Refreshes a local variable reference name if it has changed
+	 *
+	 * @param InSelfScope		Scope to lookup the variable in, to see if it has changed
+	 */
+	BLUEPRINTGRAPH_API FName RefreshLocalVariableName(UClass* InSelfScope) const;
+
 protected:
 	/** Only intended for backwards compat! */
 	BLUEPRINTGRAPH_API void SetGivenSelfScope(const FName InMemberName, const FGuid InMemberGuid, TSubclassOf<class UObject> InMemberParentClass, TSubclassOf<class UObject> SelfScope) const;
@@ -476,6 +498,18 @@ public:
 		return bSelfContext ? SelfScope : (UClass*)MemberParentClass;
 	}
 
+	/** Get the scope of this member, using a node to derive the class */
+	UStruct* GetMemberScope(const UK2Node* SelfNode) const
+	{
+		return GetMemberScope( GetBlueprintClassFromNode(SelfNode) );
+	}
+
+	/** Get the scope of this member */
+	UStruct* GetMemberScope(UClass* InMemberParentClass) const
+	{
+		return FindField<UStruct>(InMemberParentClass, *MemberScope);
+	}
+
 	/** 
 	 *	Returns the member UProperty/UFunction this reference is pointing to, or NULL if it no longer exists 
 	 *	Derives 'self' scope from supplied Blueprint node if required
@@ -491,42 +525,63 @@ public:
 			UE_LOG(LogBlueprint, Warning, TEXT("FMemberReference::ResolveMember (%s) bSelfContext == true, but no scope supplied!"), *MemberName.ToString() );
 		}
 
-		// Look for remapped member
-		UClass* TargetScope = bSelfContext ? SelfScope : (UClass*)MemberParentClass;
-		if( TargetScope != NULL &&  !GIsSavingPackage )
-		{
-			ReturnField = Cast<TFieldType>(UK2Node::FindRemappedField(TargetScope, MemberName, true));
-		}
-
-		if(ReturnField != NULL)
+		// Check if the member reference is function scoped
+		if(IsLocalScope())
 		{
-			// Fix up this struct, we found a redirect
-			MemberName = ReturnField->GetFName();
-			MemberParentClass = Cast<UClass>(ReturnField->GetOuter());
+			UStruct* MemberScopeStruct = FindField<UStruct>(SelfScope, *MemberScope);
 
-			MemberGuid.Invalidate();
-			UBlueprint::GetGuidFromClassByFieldName<TFieldType>(TargetScope, MemberName, MemberGuid);
+			// Find in target scope
+			ReturnField = FindField<TFieldType>(MemberScopeStruct, MemberName);
 
-			// Re-evaluate self-ness against the redirect if we were given a valid SelfScope
-			if(MemberParentClass != NULL && SelfScope != NULL)
+			if(ReturnField == NULL)
 			{
-				SetGivenSelfScope(MemberName, MemberGuid, MemberParentClass, SelfScope);
+				// If the property was not found, refresh the local variable name and try again
+				const FName RenamedMemberName = RefreshLocalVariableName(SelfScope);
+				if (RenamedMemberName != NAME_None)
+				{
+					ReturnField = FindField<TFieldType>(MemberScopeStruct, MemberName);
+				}
 			}
 		}
-		else if(TargetScope != NULL)
+		else
 		{
-			// Find in target scope
-			ReturnField = FindField<TFieldType>(TargetScope, MemberName);
+			// Look for remapped member
+			UClass* TargetScope = bSelfContext ? SelfScope : (UClass*)MemberParentClass;
+			if( TargetScope != NULL &&  !GIsSavingPackage )
+			{
+				ReturnField = Cast<TFieldType>(UK2Node::FindRemappedField(TargetScope, MemberName, true));
+			}
 
-			// If we have a GUID find the reference variable and make sure the name is up to date and find the field again
-			// For now only variable references will have valid GUIDs.  Will have to deal with finding other names subsequently
-			if (ReturnField == NULL && MemberGuid.IsValid())
+			if(ReturnField != NULL)
 			{
-				const FName RenamedMemberName = UBlueprint::GetFieldNameFromClassByGuid<TFieldType>(TargetScope, MemberGuid);
-				if (RenamedMemberName != NAME_None)
+				// Fix up this struct, we found a redirect
+				MemberName = ReturnField->GetFName();
+				MemberParentClass = Cast<UClass>(ReturnField->GetOuter());
+
+				MemberGuid.Invalidate();
+				UBlueprint::GetGuidFromClassByFieldName<TFieldType>(TargetScope, MemberName, MemberGuid);
+
+				// Re-evaluate self-ness against the redirect if we were given a valid SelfScope
+				if(MemberParentClass != NULL && SelfScope != NULL)
+				{
+					SetGivenSelfScope(MemberName, MemberGuid, MemberParentClass, SelfScope);
+				}
+			}
+			else if(TargetScope != NULL)
+			{
+				// Find in target scope
+				ReturnField = FindField<TFieldType>(TargetScope, MemberName);
+
+				// If we have a GUID find the reference variable and make sure the name is up to date and find the field again
+				// For now only variable references will have valid GUIDs.  Will have to deal with finding other names subsequently
+				if (ReturnField == NULL && MemberGuid.IsValid())
 				{
-					MemberName = RenamedMemberName;
-					ReturnField = FindField<TFieldType>(TargetScope, MemberName);
+					const FName RenamedMemberName = UBlueprint::GetFieldNameFromClassByGuid<TFieldType>(TargetScope, MemberGuid);
+					if (RenamedMemberName != NAME_None)
+					{
+						MemberName = RenamedMemberName;
+						ReturnField = FindField<TFieldType>(TargetScope, MemberName);
+					}
 				}
 			}
 		}
diff --git a/Engine/Source/Editor/BlueprintGraph/Classes/K2Node_FunctionEntry.h b/Engine/Source/Editor/BlueprintGraph/Classes/K2Node_FunctionEntry.h
--- a/Engine/Source/Editor/BlueprintGraph/Classes/K2Node_FunctionEntry.h
+++ b/Engine/Source/Editor/BlueprintGraph/Classes/K2Node_FunctionEntry.h
@@ -21,6 +21,10 @@ class UK2Node_FunctionEntry : public UK2Node_FunctionTerminator
 	UPROPERTY()
 	struct FKismetUserDeclaredFunctionMetadata MetaData;
 
+	/** Array of local variables to be added to generated function */
+	UPROPERTY()
+	TArray<struct FBPVariableDescription> LocalVariables;
+
 	// Begin UEdGraphNode interface
 	virtual void AllocateDefaultPins() OVERRIDE;
 	virtual FText GetNodeTitle(ENodeTitleType::Type TitleType) const OVERRIDE;
diff --git a/Engine/Source/Editor/BlueprintGraph/Classes/K2Node_LocalVariable.h b/Engine/Source/Editor/BlueprintGraph/Classes/K2Node_LocalVariable.h
--- a/Engine/Source/Editor/BlueprintGraph/Classes/K2Node_LocalVariable.h
+++ b/Engine/Source/Editor/BlueprintGraph/Classes/K2Node_LocalVariable.h
@@ -4,8 +4,8 @@
 
 #include "K2Node_LocalVariable.generated.h"
 
-UCLASS(MinimalAPI)
-class UK2Node_LocalVariable : public UK2Node_TemporaryVariable
+UCLASS(MinimalAPI, deprecated)
+class UDEPRECATED_K2Node_LocalVariable : public UK2Node_TemporaryVariable
 {
 	GENERATED_UCLASS_BODY()
 
@@ -26,6 +26,7 @@ class UK2Node_LocalVariable : public UK2Node_TemporaryVariable
 	virtual void PostPlacedNewNode() OVERRIDE;
 	virtual void PostPasteNode() OVERRIDE;
 	virtual bool CanPasteHere(const UEdGraph* TargetGraph, const UEdGraphSchema* Schema) const OVERRIDE;
+	virtual void ReconstructNode() OVERRIDE;
 	// End UEdGraphNode interface.
 
 	// Begin UK2Node interface.
diff --git a/Engine/Source/Editor/BlueprintGraph/Classes/K2Node_Variable.h b/Engine/Source/Editor/BlueprintGraph/Classes/K2Node_Variable.h
--- a/Engine/Source/Editor/BlueprintGraph/Classes/K2Node_Variable.h
+++ b/Engine/Source/Editor/BlueprintGraph/Classes/K2Node_Variable.h
@@ -121,12 +121,12 @@ public:
 	 * Utility method intended to serve as a choke point for various slate 
 	 * widgets to grab an icon from (for a specified variable).
 	 * 
-	 * @param  VarClass		The class that owns the variable in question.
+	 * @param  VarScope		The scope that owns the variable in question.
 	 * @param  VarName		The name of the variable you're querying for.
 	 * @param  IconColorOut	A color out, further discerning the variable's type.
 	 * @return A icon representing the specified variable's type.
 	 */
-	BLUEPRINTGRAPH_API static FName GetVariableIconAndColor(UClass* VarClass, FName VarName, FLinearColor& IconColorOut);
+	BLUEPRINTGRAPH_API static FName GetVariableIconAndColor(UStruct* VarScope, FName VarName, FLinearColor& IconColorOut);
 
 	/**
 	 * Utility method intended to serve as a choke point for various slate 
diff --git a/Engine/Source/Editor/BlueprintGraph/Private/K2ActionMenuBuilder.cpp b/Engine/Source/Editor/BlueprintGraph/Private/K2ActionMenuBuilder.cpp
--- a/Engine/Source/Editor/BlueprintGraph/Private/K2ActionMenuBuilder.cpp
+++ b/Engine/Source/Editor/BlueprintGraph/Private/K2ActionMenuBuilder.cpp
@@ -957,24 +957,14 @@ void FK2ActionMenuBuilder::GetContextAllowedNodeTypes(FBlueprintGraphActionListB
 
 			for (TArray<FEdGraphPinType>::TIterator TypeIt(TemporaryTypes); TypeIt; ++TypeIt)
 			{
-				// For function graphs, supply the menu options for named local variables, otherwise supply the un-named local variables
-				UK2Node* NodeTemplate;
-				if(GraphType == GT_Function)
+				// Add a standard version of the local
 				{
-					UK2Node_LocalVariable* LocalVariableNodeTemplate = ContextMenuBuilder.CreateTemplateNode<UK2Node_LocalVariable>();
-					LocalVariableNodeTemplate->VariableType = *TypeIt;
-					NodeTemplate = LocalVariableNodeTemplate;
-				}
-				else
-				{
-					UK2Node_TemporaryVariable* TemporaryVariableNodeTemplate = ContextMenuBuilder.CreateTemplateNode<UK2Node_TemporaryVariable>();
-					TemporaryVariableNodeTemplate->VariableType = *TypeIt;
-					NodeTemplate = TemporaryVariableNodeTemplate;
+					UK2Node_TemporaryVariable* NodeTemplate = ContextMenuBuilder.CreateTemplateNode<UK2Node_TemporaryVariable>();
+					NodeTemplate->VariableType = *TypeIt;
+					TSharedPtr<FEdGraphSchemaAction_K2NewNode> Action = FK2ActionMenuBuilder::AddNewNodeAction(ContextMenuBuilder, K2ActionCategories::MacroToolsCategory, NodeTemplate->GetNodeTitle(ENodeTitleType::ListView), NodeTemplate->GetTooltip(), 0, NodeTemplate->GetKeywords());
+					Action->NodeTemplate = NodeTemplate;
+					Action->SearchTitle = NodeTemplate->GetNodeSearchTitle();
 				}
-
-				TSharedPtr<FEdGraphSchemaAction_K2NewNode> Action = FK2ActionMenuBuilder::AddNewNodeAction(ContextMenuBuilder, K2ActionCategories::MacroToolsCategory, NodeTemplate->GetNodeTitle(ENodeTitleType::ListView), NodeTemplate->GetTooltip(), 0, NodeTemplate->GetKeywords());
-				Action->NodeTemplate = NodeTemplate;
-				Action->SearchTitle = NodeTemplate->GetNodeSearchTitle();
 			}
 
 			// Add in bool and int persistent pin types
@@ -1371,15 +1361,6 @@ void FK2ActionMenuBuilder::GetPinAllowedNodeTypes(FBlueprintGraphActionListBuild
 		{
 			GetEnumUtilitiesNodes(ContextMenuBuilder, false, false, false, true);
 		}
-
-		// Assignment statement for local variables
-		if(FromPin.GetOwningNode()->IsA(UK2Node_TemporaryVariable::StaticClass()))
-		{
-			UK2Node* NodeTemplate = ContextMenuBuilder.CreateTemplateNode<UK2Node_AssignmentStatement>();
-			TSharedPtr<FEdGraphSchemaAction_K2NewNode> Action = FK2ActionMenuBuilder::AddNewNodeAction(ContextMenuBuilder, TEXT(""), NodeTemplate->GetNodeTitle(ENodeTitleType::ListView), NodeTemplate->GetTooltip(), 0, NodeTemplate->GetKeywords());
-			Action->NodeTemplate = NodeTemplate;
-			Action->SearchTitle = NodeTemplate->GetNodeSearchTitle();
-		}
 	}
 }
 
@@ -1772,6 +1753,68 @@ void FK2ActionMenuBuilder::GetVariableGettersSettersForClass(FBlueprintGraphActi
 		}
 	}
 
+	// Add any local variables to the list as well
+	UEdGraph* FunctionGraph = FBlueprintEditorUtils::GetTopLevelGraph(ContextMenuBuilder.CurrentGraph);
+	TArray<UK2Node_FunctionEntry*> FunctionEntryNodes;
+	FunctionGraph->GetNodesOfClass<UK2Node_FunctionEntry>(FunctionEntryNodes);
+	if(FunctionEntryNodes.Num() == 1)
+	{
+		// Setting them without the property, cannot guarantee that properties in the UFunction are valid as local variables
+		for(auto& LocalVariable : FunctionEntryNodes[0]->LocalVariables)
+		{
+			bool bTypesMatchForRead = !bRequireDesiredPinType;
+			bool bTypesMatchForWrite = !bRequireDesiredPinType;
+			if (bRequireDesiredPinType)
+			{
+				// reading (property out -> pin in)
+				bTypesMatchForRead = K2Schema->ArePinTypesCompatible(LocalVariable.VarType, DesiredPinType, Class);
+
+				// writing (pin out -> property in)
+				// Always allow exec pins to write
+				bTypesMatchForWrite = DesiredPinType.PinCategory == K2Schema->PC_Exec || K2Schema->ArePinTypesCompatible(DesiredPinType, LocalVariable.VarType, Class);
+			}
+
+			FString Category = LocalVariable.Category.ToString();
+			FString PropertyTooltip;
+			if(LocalVariable.HasMetaData(TEXT("tooltip")))
+			{
+				PropertyTooltip = LocalVariable.GetMetaData(TEXT("tooltip"));
+			}
+			FString FullCategoryName = K2ActionCategories::VariablesCategory;
+			if(Category.IsEmpty() == false)
+			{
+				FullCategoryName += K2ActionCategories::SubCategoryDelim;
+				FullCategoryName += Category;
+			}
+
+			const FString PropertyDescription = GetDefault<UEditorStyleSettings>()->bShowFriendlyNames ? LocalVariable.FriendlyName : LocalVariable.VarName.ToString();
+			const bool bCanRead = bWantGetters && bTypesMatchForRead;
+			const bool bCanWrite = bWantSetters && bTypesMatchForWrite;
+			{
+				// Variable getters and setters
+				if (bCanRead)
+				{
+					TSharedPtr<FEdGraphSchemaAction_K2NewNode> ReadVar = FK2ActionMenuBuilder::AddNewNodeAction(ContextMenuBuilder, FullCategoryName, FText::FromString(GetString + PropertyDescription), PropertyTooltip);
+
+					UK2Node_VariableGet* VarGetNode = ContextMenuBuilder.CreateTemplateNode<UK2Node_VariableGet>();
+					VarGetNode->VariableReference.SetLocalMember(LocalVariable.VarName, FunctionGraph->GetName(), LocalVariable.VarGuid);
+					ReadVar->NodeTemplate = VarGetNode;
+					ReadVar->SearchTitle = ReadVar->NodeTemplate->GetNodeSearchTitle();
+				}
+
+				if (bCanWrite)
+				{
+					TSharedPtr<FEdGraphSchemaAction_K2NewNode> WriteVar = FK2ActionMenuBuilder::AddNewNodeAction(ContextMenuBuilder, FullCategoryName, FText::FromString(SetString + PropertyDescription), PropertyTooltip);
+
+					UK2Node_VariableSet* VarSetNode = ContextMenuBuilder.CreateTemplateNode<UK2Node_VariableSet>();
+					VarSetNode->VariableReference.SetLocalMember(LocalVariable.VarName, FunctionGraph->GetName(), LocalVariable.VarGuid);
+					WriteVar->NodeTemplate = VarSetNode;
+					WriteVar->SearchTitle = WriteVar->NodeTemplate->GetNodeSearchTitle();
+				}
+			}
+		}
+	}
+
 	if (bWantDelegates)
 	{
 		GetEventDispatcherNodesForClass(ContextMenuBuilder,Class, bSelfContext);
diff --git a/Engine/Source/Editor/BlueprintGraph/Private/K2Node.cpp b/Engine/Source/Editor/BlueprintGraph/Private/K2Node.cpp
--- a/Engine/Source/Editor/BlueprintGraph/Private/K2Node.cpp
+++ b/Engine/Source/Editor/BlueprintGraph/Private/K2Node.cpp
@@ -1087,6 +1087,7 @@ void FMemberReference::SetExternalMember(FName InMemberName, TSubclassOf<class U
 {
 	MemberName = InMemberName;
 	MemberParentClass = InMemberParentClass;
+	MemberScope.Empty();
 	bSelfContext = false;
 }
 
@@ -1094,6 +1095,7 @@ void FMemberReference::SetSelfMember(FName InMemberName)
 {
 	MemberName = InMemberName;
 	MemberParentClass = NULL;
+	MemberScope.Empty();
 	bSelfContext = true;
 }
 
@@ -1103,6 +1105,7 @@ void FMemberReference::SetDirect(const FName InMemberName, const FGuid InMemberG
 	MemberGuid = InMemberGuid;
 	bSelfContext = bIsConsideredSelfContext;
 	MemberParentClass = InMemberParentClass;
+	MemberScope.Empty();
 }
 
 void FMemberReference::SetGivenSelfScope(const FName InMemberName, const FGuid InMemberGuid, TSubclassOf<class UObject> InMemberParentClass, TSubclassOf<class UObject> SelfScope) const
@@ -1110,6 +1113,7 @@ void FMemberReference::SetGivenSelfScope(const FName InMemberName, const FGuid I
 	MemberName = InMemberName;
 	MemberGuid = InMemberGuid;
 	MemberParentClass = InMemberParentClass;
+	MemberScope.Empty();
 	bSelfContext = (SelfScope->IsChildOf(InMemberParentClass)) || (SelfScope->ClassGeneratedBy == InMemberParentClass->ClassGeneratedBy);
 
 	if (bSelfContext)
@@ -1118,6 +1122,18 @@ void FMemberReference::SetGivenSelfScope(const FName InMemberName, const FGuid I
 	}
 }
 
+void FMemberReference::SetLocalMember(FName InMemberName, UStruct* InScope, const FGuid InMemberGuid)
+{
+	SetLocalMember(InMemberName, InScope->GetName(), InMemberGuid);
+}
+
+void FMemberReference::SetLocalMember(FName InMemberName, FString InScopeName, const FGuid InMemberGuid)
+{
+	MemberName = InMemberName;
+	MemberScope = InScopeName;
+	MemberGuid = InMemberGuid;
+}
+
 void FMemberReference::InvalidateSelfScope()
 {
 	if( IsSelfContext() )
@@ -1136,4 +1152,22 @@ UClass* FMemberReference::GetBlueprintClassFromNode(const UK2Node* Node)
 	return BPClass;
 }
 
+FName FMemberReference::RefreshLocalVariableName(UClass* InSelfScope) const
+{
+	TArray<UBlueprint*> Blueprints;
+	UBlueprint::GetBlueprintHierarchyFromClass(InSelfScope, Blueprints);
+
+	FName RenamedMemberName = NAME_None;
+	for (int32 BPIndex = 0; BPIndex < Blueprints.Num(); ++BPIndex)
+	{
+		RenamedMemberName = FBlueprintEditorUtils::FindLocalVariableNameByGuid(Blueprints[BPIndex], MemberGuid);
+		if (RenamedMemberName != NAME_None)
+		{
+			MemberName = RenamedMemberName;
+			break;
+		}
+	}
+	return RenamedMemberName;
+}
+
 #undef LOCTEXT_NAMESPACE
diff --git a/Engine/Source/Editor/BlueprintGraph/Private/K2Node_LocalVariable.cpp b/Engine/Source/Editor/BlueprintGraph/Private/K2Node_LocalVariable.cpp
--- a/Engine/Source/Editor/BlueprintGraph/Private/K2Node_LocalVariable.cpp
+++ b/Engine/Source/Editor/BlueprintGraph/Private/K2Node_LocalVariable.cpp
@@ -5,14 +5,14 @@
 #include "Kismet2NameValidators.h"
 #include "ScopedTransaction.h"
 
-UK2Node_LocalVariable::UK2Node_LocalVariable(const class FPostConstructInitializeProperties& PCIP)
+UDEPRECATED_K2Node_LocalVariable::UDEPRECATED_K2Node_LocalVariable(const class FPostConstructInitializeProperties& PCIP)
 	: Super(PCIP)
 {
 	bCanRenameNode = true;
 	CustomVariableName = TEXT("NewLocalVar");
 }
 
-FString UK2Node_LocalVariable::GetTooltip() const
+FString UDEPRECATED_K2Node_LocalVariable::GetTooltip() const
 {
 	if(VariableTooltip.IsEmpty())
 	{
@@ -22,7 +22,7 @@ FString UK2Node_LocalVariable::GetTooltip() const
 	return VariableTooltip.ToString();
 }
 
-FText UK2Node_LocalVariable::GetNodeTitle(ENodeTitleType::Type TitleType) const
+FText UDEPRECATED_K2Node_LocalVariable::GetNodeTitle(ENodeTitleType::Type TitleType) const
 {
 	if (TitleType == ENodeTitleType::EditableTitle)
 	{
@@ -40,7 +40,7 @@ FText UK2Node_LocalVariable::GetNodeTitle(ENodeTitleType::Type TitleType) const
 	return FText::Format(NSLOCTEXT("K2Node", "LocalVariable_Name", "{Title}\nLocal Variable"), Args);
 }
 
-FString UK2Node_LocalVariable::GetNodeNativeTitle(ENodeTitleType::Type TitleType) const
+FString UDEPRECATED_K2Node_LocalVariable::GetNodeNativeTitle(ENodeTitleType::Type TitleType) const
 {
 	// Do not setup this function for localization, intentionally left unlocalized!
 	if (TitleType == ENodeTitleType::EditableTitle)
@@ -55,7 +55,7 @@ FString UK2Node_LocalVariable::GetNodeNativeTitle(ENodeTitleType::Type TitleType
 	return FString::Printf(TEXT("%s\nLocal Variable"), *CustomVariableName.ToString());
 }
 
-void UK2Node_LocalVariable::OnRenameNode(const FString& NewName)
+void UDEPRECATED_K2Node_LocalVariable::OnRenameNode(const FString& NewName)
 {
 	if(CustomVariableName != *NewName)
 	{
@@ -67,12 +67,12 @@ void UK2Node_LocalVariable::OnRenameNode(const FString& NewName)
 	}
 }
 
-TSharedPtr<class INameValidatorInterface> UK2Node_LocalVariable::MakeNameValidator() const
+TSharedPtr<class INameValidatorInterface> UDEPRECATED_K2Node_LocalVariable::MakeNameValidator() const
 {
 	return MakeShareable(new FKismetNameValidator(GetBlueprint(), *GetNodeTitle(ENodeTitleType::EditableTitle).ToString()));
 }
 
-void UK2Node_LocalVariable::ChangeVariableType(const FEdGraphPinType& InVariableType)
+void UDEPRECATED_K2Node_LocalVariable::ChangeVariableType(const FEdGraphPinType& InVariableType)
 {
 	// Local variables can never change type when the variable pin is hooked up
 	check(GetVariablePin()->LinkedTo.Num() == 0);
@@ -84,14 +84,14 @@ void UK2Node_LocalVariable::ChangeVariableType(const FEdGraphPinType& InVariable
 	FBlueprintEditorUtils::MarkBlueprintAsStructurallyModified(GetBlueprint());
 }
 
-void UK2Node_LocalVariable::PostPlacedNewNode()
+void UDEPRECATED_K2Node_LocalVariable::PostPlacedNewNode()
 {
 	Super::PostPlacedNewNode();
 
 	CustomVariableName = FBlueprintEditorUtils::FindUniqueKismetName(GetBlueprint(), CustomVariableName.ToString());
 }
 
-void UK2Node_LocalVariable::PostPasteNode()
+void UDEPRECATED_K2Node_LocalVariable::PostPasteNode()
 {
 	Super::PostPasteNode();
 
@@ -99,7 +99,7 @@ void UK2Node_LocalVariable::PostPasteNode()
 	CustomVariableName = FBlueprintEditorUtils::FindUniqueKismetName(GetBlueprint(), CustomVariableName.GetPlainNameString());
 }
 
-bool UK2Node_LocalVariable::CanPasteHere(const UEdGraph* TargetGraph, const UEdGraphSchema* Schema) const
+bool UDEPRECATED_K2Node_LocalVariable::CanPasteHere(const UEdGraph* TargetGraph, const UEdGraphSchema* Schema) const
 {
 	// Local variables can only be pasted into function graphs
 	if(Super::CanPasteHere(TargetGraph, Schema))
@@ -115,3 +115,96 @@ bool UK2Node_LocalVariable::CanPasteHere(const UEdGraph* TargetGraph, const UEdG
 
 	return false;
 }
+
+void UDEPRECATED_K2Node_LocalVariable::ReconstructNode()
+{
+	UEdGraph* Graph = GetGraph();
+
+	UEdGraph* TopLevelGraph = FBlueprintEditorUtils::GetTopLevelGraph(Graph);
+	if(TopLevelGraph->GetSchema()->GetGraphType(TopLevelGraph) == GT_Function)
+	{
+		const UEdGraphSchema_K2* K2Schema = GetDefault<UEdGraphSchema_K2>();
+
+		// First, add the variable to the entry node
+		FBPVariableDescription NewVar;
+		NewVar.VarName = *GetNodeTitle(ENodeTitleType::EditableTitle).ToString();
+		NewVar.VarGuid = FGuid::NewGuid();
+		NewVar.VarType = GetVariablePin()->PinType;
+		NewVar.FriendlyName = FName::NameToDisplayString( NewVar.VarName.ToString(), (NewVar.VarType.PinCategory == K2Schema->PC_Boolean) ? true : false );
+		NewVar.Category = K2Schema->VR_DefaultCategory;
+
+		// Find the function entry node in the graph
+		TArray<UK2Node_FunctionEntry*> FunctionEntryNodes;
+		TopLevelGraph->GetNodesOfClass<UK2Node_FunctionEntry>(FunctionEntryNodes);
+		check(FunctionEntryNodes.Num() == 1);
+
+		UK2Node_FunctionEntry* FunctionEntry = FunctionEntryNodes[0];
+		FunctionEntry->LocalVariables.Add(NewVar);
+
+		// Duplicate the array, it is semi-destructive as we iterate through
+		TArray< UEdGraphPin* > VariableLinkedPins = GetVariablePin()->LinkedTo;
+
+		// Search through all linked pins for AssignmentStatement nodes, those need to be replaced with a VariableSet node and hooked up to the new variable
+		for( auto LinkedPin : VariableLinkedPins )
+		{
+			if(LinkedPin->GetOwningNode()->IsA(UK2Node_AssignmentStatement::StaticClass()))
+			{
+				UK2Node_AssignmentStatement* AssignementNode = Cast<UK2Node_AssignmentStatement>(LinkedPin->GetOwningNode());
+
+				// Only replace the node if hooked up to the variable pin
+				if(AssignementNode->GetVariablePin() == LinkedPin)
+				{
+					UK2Node_VariableSet* SetNode = NewObject<UK2Node_VariableSet>(Graph);
+					SetNode->VariableReference.SetLocalMember(NewVar.VarName, TopLevelGraph->GetName(), NewVar.VarGuid);
+					Graph->AddNode(SetNode, false, false);
+					SetNode->CreateNewGuid();
+					SetNode->PostPlacedNewNode();
+
+					// Locally re-construct the pins, can't allow the node to do it because the property does not yet exist
+					UEdGraphPin* ExecPin = SetNode->CreatePin(EGPD_Input, K2Schema->PC_Exec, TEXT(""), NULL, false, false, K2Schema->PN_Execute);
+					UEdGraphPin* ThenPin = SetNode->CreatePin(EGPD_Output, K2Schema->PC_Exec, TEXT(""), NULL, false, false, K2Schema->PN_Then);
+					UEdGraphPin* ValuePin = SetNode->CreatePin(EGPD_Input, TEXT(""), TEXT(""), NULL, false, false, NewVar.VarName.ToString());
+					ValuePin->PinType = NewVar.VarType;
+
+					// Move the pin links over to the new node
+					K2Schema->MovePinLinks(*AssignementNode->FindPin(K2Schema->PN_Execute), *ExecPin);
+					K2Schema->MovePinLinks(*AssignementNode->FindPin(K2Schema->PN_Then), *ThenPin);
+					K2Schema->MovePinLinks(*AssignementNode->GetValuePin(), *ValuePin);
+
+					// Position the new node
+					SetNode->NodePosX = AssignementNode->NodePosX;
+					SetNode->NodePosY = AssignementNode->NodePosY;
+
+					// Destroy the assignment node
+					AssignementNode->DestroyNode();
+				}
+			}
+		}
+
+		// Only replace with a get node if the node still has connections
+		if(GetVariablePin()->LinkedTo.Num())
+		{
+			// Only the local variable node needs to be re-constructed as a VariableGet node, there are no other nodes representing this local variable
+			UK2Node_VariableGet* GetNode = NewObject<UK2Node_VariableGet>(Graph);
+			GetNode->VariableReference.SetLocalMember(NewVar.VarName, TopLevelGraph->GetName(), NewVar.VarGuid);
+			Graph->AddNode(GetNode, false, false);
+			GetNode->CreateNewGuid();
+			GetNode->PostPlacedNewNode();
+
+			// Locally re-construct the pins, can't allow the node to do it because the property does not yet exist
+			UEdGraphPin* VariablePin = GetNode->CreatePin(EGPD_Output, TEXT(""), TEXT(""), NULL, false, false, NewVar.VarName.ToString());
+			VariablePin->PinType = NewVar.VarType;
+			K2Schema->SetPinDefaultValueBasedOnType(VariablePin);
+
+			// Position the new node
+			GetNode->NodePosX = NodePosX;
+			GetNode->NodePosY = NodePosY;
+
+			// Move the pin links over to the new node
+			K2Schema->MovePinLinks(*GetVariablePin(), *GetNode->GetValuePin());
+		}
+	}
+
+	// This node should not persist anymore, it is deprecated and the above case catches any still valid uses and deletes the rest (invalid ones are from AnimationTransitionGraph)
+	DestroyNode();
+}
\ No newline at end of file
diff --git a/Engine/Source/Editor/BlueprintGraph/Private/K2Node_Variable.cpp b/Engine/Source/Editor/BlueprintGraph/Private/K2Node_Variable.cpp
--- a/Engine/Source/Editor/BlueprintGraph/Private/K2Node_Variable.cpp
+++ b/Engine/Source/Editor/BlueprintGraph/Private/K2Node_Variable.cpp
@@ -83,17 +83,21 @@ void UK2Node_Variable::CreatePinForSelf()
 	// Create the self pin
 	if (!K2Schema->FindSelfPin(*this, EGPD_Input))
 	{
-		if (VariableReference.IsSelfContext() && (ESelfContextInfo::NotSelfContext != SelfContextInfo))
+		// Do not create a self pin for locally scoped variables
+		if( !VariableReference.IsLocalScope() )
 		{
-			UEdGraphPin* SelfPin = CreatePin(EGPD_Input, K2Schema->PC_Object, K2Schema->PSC_Self, NULL, false, false, K2Schema->PN_Self);
-			SelfPin->bHidden = true; // don't show in 'self' context
-		}
-		else
-		{
-			UClass* MemberParentClass = VariableReference.GetMemberParentClass(this);
-			UClass* AuthoritativeClass = MemberParentClass ? MemberParentClass->GetAuthoritativeClass() : NULL;
-			UEdGraphPin* TargetPin = CreatePin(EGPD_Input, K2Schema->PC_Object, TEXT(""), AuthoritativeClass, false, false, K2Schema->PN_Self);
-			TargetPin->PinFriendlyName =  LOCTEXT("Target", "Target");
+			if (VariableReference.IsSelfContext() && (ESelfContextInfo::NotSelfContext != SelfContextInfo))
+			{
+				UEdGraphPin* SelfPin = CreatePin(EGPD_Input, K2Schema->PC_Object, K2Schema->PSC_Self, NULL, false, false, K2Schema->PN_Self);
+				SelfPin->bHidden = true; // don't show in 'self' context
+			}
+			else
+			{
+				UClass* MemberParentClass = VariableReference.GetMemberParentClass(this);
+				UClass* AuthoritativeClass = MemberParentClass ? MemberParentClass->GetAuthoritativeClass() : NULL;
+				UEdGraphPin* TargetPin = CreatePin(EGPD_Input, K2Schema->PC_Object, TEXT(""), AuthoritativeClass, false, false, K2Schema->PN_Self);
+				TargetPin->PinFriendlyName =  LOCTEXT("Target", "Target");
+			}
 		}
 	}
 	else
@@ -246,7 +250,9 @@ void UK2Node_Variable::ValidateNodeDuringCompilation(class FCompilerResultsLog&
 	Super::ValidateNodeDuringCompilation(MessageLog);
 
 	UProperty* VariableProperty = GetPropertyForVariable();
-	if (VariableProperty == NULL)
+
+	// Local variables do not exist until much later in the compilation than this function can provide
+	if (VariableProperty == NULL && !VariableReference.IsLocalScope())
 	{
 		MessageLog.Warning(*FString::Printf(*LOCTEXT("VariableNotFound", "Unable to find variable with name '%s' for @@").ToString(), *VariableReference.GetMemberName().ToString()), this);
 	}
@@ -254,7 +260,18 @@ void UK2Node_Variable::ValidateNodeDuringCompilation(class FCompilerResultsLog&
 
 FName UK2Node_Variable::GetPaletteIcon(FLinearColor& ColorOut) const
 {
-	return GetVariableIconAndColor(GetVariableSourceClass(), GetVarName(), ColorOut);
+	FName ReturnIconName;
+
+	if(VariableReference.IsLocalScope())
+	{
+		ReturnIconName = GetVariableIconAndColor(VariableReference.GetMemberScope(this), GetVarName(), ColorOut);
+	}
+	else
+	{
+		ReturnIconName = GetVariableIconAndColor(GetVariableSourceClass(), GetVarName(), ColorOut);
+	}
+
+	return ReturnIconName;
 }
 
 FName UK2Node_Variable::GetVarIconFromPinType(FEdGraphPinType& InPinType, FLinearColor& IconColorOut)
@@ -301,13 +318,13 @@ FText UK2Node_Variable::GetToolTipHeading() const
 	return Heading;
 }
 
-FName UK2Node_Variable::GetVariableIconAndColor(UClass* VarClass, FName VarName, FLinearColor& IconColorOut)
+FName UK2Node_Variable::GetVariableIconAndColor(UStruct* VarScope, FName VarName, FLinearColor& IconColorOut)
 {
 	FName IconBrush = TEXT("Kismet.AllClasses.VariableIcon");
 
-	if(VarClass != NULL)
+	if(VarScope != NULL)
 	{
-		UProperty* Property = FindField<UProperty>(VarClass, VarName);
+		UProperty* Property = FindField<UProperty>(VarScope, VarName);
 		if(Property != NULL)
 		{
 			const UEdGraphSchema_K2* K2Schema = GetDefault<UEdGraphSchema_K2>();
diff --git a/Engine/Source/Editor/BlueprintGraph/Private/K2Node_VariableGet.cpp b/Engine/Source/Editor/BlueprintGraph/Private/K2Node_VariableGet.cpp
--- a/Engine/Source/Editor/BlueprintGraph/Private/K2Node_VariableGet.cpp
+++ b/Engine/Source/Editor/BlueprintGraph/Private/K2Node_VariableGet.cpp
@@ -31,6 +31,15 @@ public:
 			VarNode->CheckForErrors(CompilerContext.GetSchema(), Context.MessageLog);
 		}
 
+		// Report an error that the local variable could not be found
+		if(VarNode->VariableReference.IsLocalScope() && VarNode->GetPropertyForVariable() == NULL)
+		{
+			FFormatNamedArguments Args;
+			Args.Add(TEXT("VariableName"), FText::FromName(VarNode->VariableReference.GetMemberName()));
+
+			CompilerContext.MessageLog.Warning(*FText::Format(LOCTEXT("LocalVariableNotFound_Error", "Unable to find local variable with name '{VariableName}' for @@"), Args).ToString(), Node);
+		}
+
 		FNodeHandlingFunctor::RegisterNets(Context, Node);
 	}
 };
diff --git a/Engine/Source/Editor/BlueprintGraph/Private/VariableSetHandler.cpp b/Engine/Source/Editor/BlueprintGraph/Private/VariableSetHandler.cpp
--- a/Engine/Source/Editor/BlueprintGraph/Private/VariableSetHandler.cpp
+++ b/Engine/Source/Editor/BlueprintGraph/Private/VariableSetHandler.cpp
@@ -26,6 +26,15 @@ void FKCHandler_VariableSet::RegisterNets(FKismetFunctionContext& Context, UEdGr
 		{
 			CompilerContext.MessageLog.Warning(*LOCTEXT("BlueprintReadOnlyOrPrivate_Error", "The property is mark as BlueprintReadOnly or Private. It cannot be modifed in the blueprint. @@").ToString(), Node);
 		}
+
+		// Report an error that the local variable could not be found
+		if(SetterNode->VariableReference.IsLocalScope() && SetterNode->GetPropertyForVariable() == NULL)
+		{
+			FFormatNamedArguments Args;
+			Args.Add(TEXT("VariableName"), FText::FromName(SetterNode->VariableReference.GetMemberName()));
+
+			CompilerContext.MessageLog.Warning(*FText::Format(LOCTEXT("LocalVariableNotFound_Error", "Unable to find local variable with name '{VariableName}' for @@"), Args).ToString(), Node);
+		}
 	}
 
 	for (int32 PinIndex = 0; PinIndex < Node->Pins.Num(); ++PinIndex)
diff --git a/Engine/Source/Editor/GraphEditor/Private/SGraphActionMenu.cpp b/Engine/Source/Editor/GraphEditor/Private/SGraphActionMenu.cpp
--- a/Engine/Source/Editor/GraphEditor/Private/SGraphActionMenu.cpp
+++ b/Engine/Source/Editor/GraphEditor/Private/SGraphActionMenu.cpp
@@ -17,6 +17,8 @@ namespace GraphActionMenuHelpers
 
 		bCheck |= (InGraphAction->GetTypeId() == FEdGraphSchemaAction_K2Var::StaticGetTypeId() &&
 			((FEdGraphSchemaAction_K2Var*)InGraphAction)->GetVariableName() == ItemName);
+		bCheck |= (InGraphAction->GetTypeId() == FEdGraphSchemaAction_K2LocalVar::StaticGetTypeId() &&
+			((FEdGraphSchemaAction_K2LocalVar*)InGraphAction)->GetVariableName() == ItemName);
 		bCheck |= (InGraphAction->GetTypeId() == FEdGraphSchemaAction_K2Graph::StaticGetTypeId() &&
 			((FEdGraphSchemaAction_K2Graph*)InGraphAction)->EdGraph &&
 			((FEdGraphSchemaAction_K2Graph*)InGraphAction)->EdGraph->GetFName() == ItemName);
diff --git a/Engine/Source/Editor/Kismet/Private/BPDelegateDragDropAction.cpp b/Engine/Source/Editor/Kismet/Private/BPDelegateDragDropAction.cpp
--- a/Engine/Source/Editor/Kismet/Private/BPDelegateDragDropAction.cpp
+++ b/Engine/Source/Editor/Kismet/Private/BPDelegateDragDropAction.cpp
@@ -138,9 +138,9 @@ FReply FKismetDelegateDragDropAction::DroppedOnPanel(const TSharedRef< SWidget >
 
 bool FKismetDelegateDragDropAction::IsValid() const
 {
-	return VariableSourceClass.IsValid() &&
+	return VariableSource.IsValid() &&
 		(VariableName != NAME_None) &&
-		(NULL != FindField<UMulticastDelegateProperty>(VariableSourceClass.Get(), VariableName));
+		(NULL != FindField<UMulticastDelegateProperty>(VariableSource.Get(), VariableName));
 }
 
 #undef LOCTEXT_NAMESPACE
diff --git a/Engine/Source/Editor/Kismet/Private/BPDelegateDragDropAction.h b/Engine/Source/Editor/Kismet/Private/BPDelegateDragDropAction.h
--- a/Engine/Source/Editor/Kismet/Private/BPDelegateDragDropAction.h
+++ b/Engine/Source/Editor/Kismet/Private/BPDelegateDragDropAction.h
@@ -21,11 +21,11 @@ public:
 
 	bool IsValid() const;
 	
-	static TSharedRef<FKismetDelegateDragDropAction> New( TSharedRef< SWidget > Owner, FName InVariableName, UClass* InClass, FNodeCreationAnalytic AnalyticCallback)
+	static TSharedRef<FKismetDelegateDragDropAction> New( TSharedRef< SWidget > Owner, FName InVariableName, UStruct* InSource, FNodeCreationAnalytic AnalyticCallback)
 	{
 		TSharedRef<FKismetDelegateDragDropAction> Operation = MakeShareable(new FKismetDelegateDragDropAction( Owner ) );
 		Operation->VariableName = InVariableName;
-		Operation->VariableSourceClass = InClass;
+		Operation->VariableSource = InSource;
 		Operation->AnalyticCallback = AnalyticCallback;
 		Operation->Construct();
 		return Operation;
diff --git a/Engine/Source/Editor/Kismet/Private/BPVariableDragDropAction.cpp b/Engine/Source/Editor/Kismet/Private/BPVariableDragDropAction.cpp
--- a/Engine/Source/Editor/Kismet/Private/BPVariableDragDropAction.cpp
+++ b/Engine/Source/Editor/Kismet/Private/BPVariableDragDropAction.cpp
@@ -55,12 +55,18 @@ void FKismetVariableDragDropAction::HoverTargetChanged()
 	UEdGraphPin* PinUnderCursor = GetHoveredPin();
 
 	bool bBadSchema = false;
-	if (UEdGraph* Graph = GetHoveredGraph())
+	bool bBadGraph = false;
+	UEdGraph* HoveredGraph = GetHoveredGraph();
+	if (HoveredGraph)
 	{
-		if (Cast<const UEdGraphSchema_K2>(Graph->GetSchema()) == NULL)
+		if (Cast<const UEdGraphSchema_K2>(HoveredGraph->GetSchema()) == NULL)
 		{
 			bBadSchema = true;
 		}
+		else if(!CanVariableBeDropped(VariableProperty, *HoveredGraph))
+		{
+			bBadGraph = true;
+		}
 	}
 
 	UEdGraphNode* VarNodeUnderCursor = Cast<UK2Node_Variable>(GetHoveredNode());
@@ -70,84 +76,126 @@ void FKismetVariableDragDropAction::HoverTargetChanged()
 		StatusSymbol = FEditorStyle::GetBrush(TEXT("Graph.ConnectorFeedback.Error"));
 		Message = LOCTEXT("CannotCreateInThisSchema", "Cannot access variables in this type of graph");
 	}
-	else if (PinUnderCursor != NULL)
+	else if(bBadGraph)
 	{
-		const UEdGraphSchema_K2* Schema = CastChecked<const UEdGraphSchema_K2>(PinUnderCursor->GetSchema());
-
-		const bool bIsRead = PinUnderCursor->Direction == EGPD_Input;
-		const UBlueprint* Blueprint = FBlueprintEditorUtils::FindBlueprintForNode(PinUnderCursor->GetOwningNode());
-		const bool bReadOnlyProperty = FBlueprintEditorUtils::IsPropertyReadOnlyInCurrentBlueprint(Blueprint, VariableProperty);
-		const bool bCanWriteIfNeeded = bIsRead || !bReadOnlyProperty;
-
-		FEdGraphPinType VariablePinType;
-		Schema->ConvertPropertyToPinType(VariableProperty, VariablePinType);
-		const bool bTypeMatch = Schema->ArePinTypesCompatible(VariablePinType, PinUnderCursor->PinType);
+		FFormatNamedArguments Args;
+		Args.Add(TEXT("VariableName"), FText::FromString(VariableString));
+		Args.Add(TEXT("Scope"), FText::FromString(HoveredGraph->GetName()));
 
+		StatusSymbol = FEditorStyle::GetBrush(TEXT("Graph.ConnectorFeedback.Error"));
+		Message = FText::Format( LOCTEXT("IncorrectGraph_Error", "Cannot place local variable '{VariableName}' in external scope '{Scope}'"), Args);
+	}
+	else if (PinUnderCursor != NULL)
+	{
 		FFormatNamedArguments Args;
 		Args.Add(TEXT("PinUnderCursor"), FText::FromString(PinUnderCursor->PinName));
 		Args.Add(TEXT("VariableName"), FText::FromString(VariableString));
 
-		if (bTypeMatch && bCanWriteIfNeeded)
+		if(CanVariableBeDropped(VariableProperty, *PinUnderCursor->GetOwningNode()->GetGraph()))
 		{
-			StatusSymbol = FEditorStyle::GetBrush(TEXT("Graph.ConnectorFeedback.OK"));
+			const UEdGraphSchema_K2* Schema = CastChecked<const UEdGraphSchema_K2>(PinUnderCursor->GetSchema());
+
+			const bool bIsRead = PinUnderCursor->Direction == EGPD_Input;
+			const UBlueprint* Blueprint = FBlueprintEditorUtils::FindBlueprintForNode(PinUnderCursor->GetOwningNode());
+			const bool bReadOnlyProperty = FBlueprintEditorUtils::IsPropertyReadOnlyInCurrentBlueprint(Blueprint, VariableProperty);
+			const bool bCanWriteIfNeeded = bIsRead || !bReadOnlyProperty;
+
+			FEdGraphPinType VariablePinType;
+			Schema->ConvertPropertyToPinType(VariableProperty, VariablePinType);
+			const bool bTypeMatch = Schema->ArePinTypesCompatible(VariablePinType, PinUnderCursor->PinType);
 
-			if (bIsRead)
+			Args.Add(TEXT("PinUnderCursor"), FText::FromString(PinUnderCursor->PinName));
+
+			if (bTypeMatch && bCanWriteIfNeeded)
 			{
-				Message = FText::Format(LOCTEXT("MakeThisEqualThat", "Make {PinUnderCursor} = {VariableName}"), Args);
+				StatusSymbol = FEditorStyle::GetBrush(TEXT("Graph.ConnectorFeedback.OK"));
+
+				if (bIsRead)
+				{
+					Message = FText::Format(LOCTEXT("MakeThisEqualThat_PinEqualVariableName", "Make {PinUnderCursor} = {VariableName}"), Args);
+				}
+				else
+				{
+					Message = FText::Format(LOCTEXT("MakeThisEqualThat_VariableNameEqualPin", "Make {VariableName} = {PinUnderCursor}"), Args);
+				}
 			}
 			else
 			{
-				Message = FText::Format(LOCTEXT("MakeThisEqualThat", "Make {VariableName} = {PinUnderCursor}"), Args);
+				StatusSymbol = FEditorStyle::GetBrush(TEXT("Graph.ConnectorFeedback.Error"));
+				if (!bCanWriteIfNeeded)
+				{
+					Message = FText::Format(LOCTEXT("ReadOnlyVar_Error", "Cannot write to read-only variable '{VariableName}'"), Args);
+				}
+				else
+				{
+					Message = FText::Format(LOCTEXT("NotCompatible_Error", "The type of '{VariableName}' is not compatible with {PinUnderCursor}"), Args);
+				}
 			}
 		}
 		else
 		{
+			Args.Add(TEXT("Scope"), FText::FromString(PinUnderCursor->GetOwningNode()->GetGraph()->GetName()));
+
 			StatusSymbol = FEditorStyle::GetBrush(TEXT("Graph.ConnectorFeedback.Error"));
-			if (!bCanWriteIfNeeded)
-			{
-				Message = FText::Format(LOCTEXT("ReadOnlyVar_Error", "Cannot write to read-only variable '{VariableName}'"), Args);
-			}
-			else
-			{
-				Message = FText::Format(LOCTEXT("NotCompatible_Error", "The type of '{VariableName}' is not compatible with {PinUnderCursor}"), Args);
-			}
+			Message = FText::Format( LOCTEXT("IncorrectGraphForPin_Error", "Cannot place local variable '{VariableName}' in external scope '{Scope}'"), Args);
 		}
 	}
 	else if (VarNodeUnderCursor != NULL)
 	{
-		const bool bIsRead = VarNodeUnderCursor->IsA(UK2Node_VariableGet::StaticClass());
-		const UBlueprint* Blueprint = FBlueprintEditorUtils::FindBlueprintForNode(VarNodeUnderCursor);
-		const bool bReadOnlyProperty = FBlueprintEditorUtils::IsPropertyReadOnlyInCurrentBlueprint(Blueprint, VariableProperty);
-		const bool bCanWriteIfNeeded = bIsRead || !bReadOnlyProperty;
-		
 		FFormatNamedArguments Args;
 		Args.Add(TEXT("VariableName"), FText::FromString(VariableString));
 
-		if (bCanWriteIfNeeded)
+		if(CanVariableBeDropped(VariableProperty, *VarNodeUnderCursor->GetGraph()))
 		{
-			Args.Add(TEXT("ReadOrWrite"), bIsRead ? LOCTEXT("Read", "read") : LOCTEXT("Write", "write"));
-			if(WillBreakLinks(VarNodeUnderCursor, VariableProperty))
+			const bool bIsRead = VarNodeUnderCursor->IsA(UK2Node_VariableGet::StaticClass());
+			const UBlueprint* Blueprint = FBlueprintEditorUtils::FindBlueprintForNode(VarNodeUnderCursor);
+			const bool bReadOnlyProperty = FBlueprintEditorUtils::IsPropertyReadOnlyInCurrentBlueprint(Blueprint, VariableProperty);
+			const bool bCanWriteIfNeeded = bIsRead || !bReadOnlyProperty;
+
+			if (bCanWriteIfNeeded)
 			{
-				StatusSymbol = FEditorStyle::GetBrush(TEXT("Graph.ConnectorFeedback.OKWarn"));
-				Message = FText::Format( LOCTEXT("ChangeNodeToWarnBreakLinks", "Change node to {ReadOrWrite} '{VariableName}', WARNING this will break links!"), Args);
+				Args.Add(TEXT("ReadOrWrite"), bIsRead ? LOCTEXT("Read", "read") : LOCTEXT("Write", "write"));
+				if(WillBreakLinks(VarNodeUnderCursor, VariableProperty))
+				{
+					StatusSymbol = FEditorStyle::GetBrush(TEXT("Graph.ConnectorFeedback.OKWarn"));
+					Message = FText::Format( LOCTEXT("ChangeNodeToWarnBreakLinks", "Change node to {ReadOrWrite} '{VariableName}', WARNING this will break links!"), Args);
+				}
+				else
+				{
+					StatusSymbol = FEditorStyle::GetBrush(TEXT("Graph.ConnectorFeedback.OK"));
+					Message = FText::Format( LOCTEXT("ChangeNodeTo", "Change node to {ReadOrWrite} '{VariableName}'"), Args);
+				}
 			}
 			else
 			{
-				StatusSymbol = FEditorStyle::GetBrush(TEXT("Graph.ConnectorFeedback.OK"));
-				Message = FText::Format( LOCTEXT("ChangeNodeTo", "Change node to {ReadOrWrite} '{VariableName}'"), Args);
+				StatusSymbol = FEditorStyle::GetBrush(TEXT("Graph.ConnectorFeedback.Error"));
+				Message = FText::Format( LOCTEXT("ReadOnlyVar_Error", "Cannot write to read-only variable '{VariableName}'"), Args);
 			}
 		}
 		else
 		{
+			Args.Add(TEXT("Scope"), FText::FromString(VarNodeUnderCursor->GetGraph()->GetName()));
+
 			StatusSymbol = FEditorStyle::GetBrush(TEXT("Graph.ConnectorFeedback.Error"));
-			Message = FText::Format( LOCTEXT("ReadOnlyVar_Error", "Cannot write to read-only variable '{VariableName}'"), Args);
+			Message = FText::Format( LOCTEXT("IncorrectGraphForNodeReplace_Error", "Cannot replace node with local variable '{VariableName}' in external scope '{Scope}'"), Args);
 		}
 	}
 	else if (HoveredCategoryName.Len() > 0)
 	{
 		// Find Blueprint that made this class and get category of variable
 		FName Category = NAME_None;
-		UBlueprint* Blueprint = UBlueprint::GetBlueprintFromClass(VariableSourceClass.Get());
+		UBlueprint* Blueprint;
+		
+		// Find the Blueprint for this property
+		if(Cast<UFunction>(VariableSource.Get()))
+		{
+			Blueprint = UBlueprint::GetBlueprintFromClass(Cast<UClass>(VariableSource->GetOuter()));
+		}
+		else
+		{
+			Blueprint = UBlueprint::GetBlueprintFromClass(Cast<UClass>(VariableSource.Get()));
+		}
+
 		if(Blueprint != NULL)
 		{
 			Category = FBlueprintEditorUtils::GetBlueprintVariableCategory(Blueprint, VariableProperty->GetFName() );
@@ -156,27 +204,30 @@ void FKismetVariableDragDropAction::HoverTargetChanged()
 		FName NewCategory = FName(*HoveredCategoryName,  FNAME_Find);
 
 		// See if class is native
-		UClass* OuterClass = CastChecked<UClass>(VariableProperty->GetOuter());
-		const bool bIsNativeVar = (OuterClass->ClassGeneratedBy == NULL);
+		UClass* OuterClass = Cast<UClass>(VariableProperty->GetOuter());
+		if(OuterClass || Cast<UFunction>(VariableProperty->GetOuter()))
+		{
+			const bool bIsNativeVar = (OuterClass && OuterClass->ClassGeneratedBy == NULL);
 
-		FFormatNamedArguments Args;
-		Args.Add(TEXT("VariableName"), FText::FromString(VariableString));
-		Args.Add(TEXT("HoveredCategoryName"), FText::FromString(HoveredCategoryName));
+			FFormatNamedArguments Args;
+			Args.Add(TEXT("VariableName"), FText::FromString(VariableString));
+			Args.Add(TEXT("HoveredCategoryName"), FText::FromString(HoveredCategoryName));
 
-		if (bIsNativeVar)
-		{
-			StatusSymbol = FEditorStyle::GetBrush(TEXT("Graph.ConnectorFeedback.Error"));
-			Message = FText::Format( LOCTEXT("ChangingCatagoryNotThisVar", "Cannot change category for variable '{VariableName}'"), Args );
-		}
-		else if (Category == NewCategory)
-		{
-			StatusSymbol = FEditorStyle::GetBrush(TEXT("Graph.ConnectorFeedback.Error"));
-			Message = FText::Format( LOCTEXT("ChangingCatagoryAlreadyIn", "Variable '{VariableName}' is already in category '{HoveredCategoryName}'"), Args );
-		}
-		else
-		{
-			StatusSymbol = FEditorStyle::GetBrush(TEXT("Graph.ConnectorFeedback.OK"));
-			Message = FText::Format( LOCTEXT("ChangingCatagoryOk", "Move variable '{VariableName}' to category '{HoveredCategoryName}'"), Args );
+			if (bIsNativeVar)
+			{
+				StatusSymbol = FEditorStyle::GetBrush(TEXT("Graph.ConnectorFeedback.Error"));
+				Message = FText::Format( LOCTEXT("ChangingCatagoryNotThisVar", "Cannot change category for variable '{VariableName}'"), Args );
+			}
+			else if (Category == NewCategory)
+			{
+				StatusSymbol = FEditorStyle::GetBrush(TEXT("Graph.ConnectorFeedback.Error"));
+				Message = FText::Format( LOCTEXT("ChangingCatagoryAlreadyIn", "Variable '{VariableName}' is already in category '{HoveredCategoryName}'"), Args );
+			}
+			else
+			{
+				StatusSymbol = FEditorStyle::GetBrush(TEXT("Graph.ConnectorFeedback.OK"));
+				Message = FText::Format( LOCTEXT("ChangingCatagoryOk", "Move variable '{VariableName}' to category '{HoveredCategoryName}'"), Args );
+			}
 		}
 	}
 	else if (HoveredAction.IsValid())
@@ -189,7 +240,7 @@ void FKismetVariableDragDropAction::HoverTargetChanged()
 			// Needs to have a valid index to move it (this excludes variables added through other means, like timelines/components
 			int32 MoveVarIndex = INDEX_NONE;
 			int32 TargetVarIndex = INDEX_NONE;
-			UBlueprint* Blueprint = UBlueprint::GetBlueprintFromClass(VariableSourceClass.Get());
+			UBlueprint* Blueprint = UBlueprint::GetBlueprintFromClass(Cast<UClass>(VariableSource.Get()));
 			if(Blueprint != NULL)
 			{
 				MoveVarIndex = FBlueprintEditorUtils::FindNewVariableIndex(Blueprint, VariableName);
@@ -225,7 +276,7 @@ void FKismetVariableDragDropAction::HoverTargetChanged()
 	// Draw variable icon
 	else
 	{
-		StatusSymbol = FBlueprintEditor::GetVarIconAndColor(VariableSourceClass.Get(), VariableName, IconColor);
+		StatusSymbol = FBlueprintEditor::GetVarIconAndColor(VariableSource.Get(), VariableName, IconColor);
 		Message = FText::FromString(VariableString);
 	}
 
@@ -242,26 +293,29 @@ FReply FKismetVariableDragDropAction::DroppedOnPin(FVector2D ScreenPosition, FVe
 		{
 			UProperty* VariableProperty = GetVariableProperty();
 
-			const bool bIsRead = TargetPin->Direction == EGPD_Input;
-			const UBlueprint* Blueprint = FBlueprintEditorUtils::FindBlueprintForNode(TargetPin->GetOwningNode());
-			const bool bReadOnlyProperty = FBlueprintEditorUtils::IsPropertyReadOnlyInCurrentBlueprint(Blueprint, VariableProperty);
-			const bool bCanWriteIfNeeded = bIsRead || !bReadOnlyProperty;
+			if(CanVariableBeDropped(VariableProperty, *TargetPin->GetOwningNode()->GetGraph()))
+			{
+				const bool bIsRead = TargetPin->Direction == EGPD_Input;
+				const UBlueprint* Blueprint = FBlueprintEditorUtils::FindBlueprintForNode(TargetPin->GetOwningNode());
+				const bool bReadOnlyProperty = FBlueprintEditorUtils::IsPropertyReadOnlyInCurrentBlueprint(Blueprint, VariableProperty);
+				const bool bCanWriteIfNeeded = bIsRead || !bReadOnlyProperty;
 
-			FEdGraphPinType VariablePinType;
-			Schema->ConvertPropertyToPinType(VariableProperty, VariablePinType);
-			const bool bTypeMatch = Schema->ArePinTypesCompatible(VariablePinType, TargetPin->PinType);
+				FEdGraphPinType VariablePinType;
+				Schema->ConvertPropertyToPinType(VariableProperty, VariablePinType);
+				const bool bTypeMatch = Schema->ArePinTypesCompatible(VariablePinType, TargetPin->PinType);
 
-			if (bTypeMatch && bCanWriteIfNeeded)
-			{
-				FEdGraphSchemaAction_K2NewNode Action;
+				if (bTypeMatch && bCanWriteIfNeeded)
+				{
+					FEdGraphSchemaAction_K2NewNode Action;
 
-				UK2Node_Variable* VarNode = bIsRead ? (UK2Node_Variable*)NewObject<UK2Node_VariableGet>() : (UK2Node_Variable*)NewObject<UK2Node_VariableSet>();
-				Action.NodeTemplate = VarNode;
+					UK2Node_Variable* VarNode = bIsRead ? (UK2Node_Variable*)NewObject<UK2Node_VariableGet>() : (UK2Node_Variable*)NewObject<UK2Node_VariableSet>();
+					Action.NodeTemplate = VarNode;
 
-				UBlueprint* DropOnBlueprint = FBlueprintEditorUtils::FindBlueprintForGraph(TargetPin->GetOwningNode()->GetGraph());
-				ConfigureVarNode(VarNode, VariableName, VariableSourceClass.Get(), DropOnBlueprint);
+					UBlueprint* DropOnBlueprint = FBlueprintEditorUtils::FindBlueprintForGraph(TargetPin->GetOwningNode()->GetGraph());
+					ConfigureVarNode(VarNode, VariableName, VariableSource.Get(), DropOnBlueprint);
 
-				Action.PerformAction(TargetPin->GetOwningNode()->GetGraph(), TargetPin, GraphPosition);
+					Action.PerformAction(TargetPin->GetOwningNode()->GetGraph(), TargetPin, GraphPosition);
+				}
 			}
 		}
 	}
@@ -278,91 +332,106 @@ FReply FKismetVariableDragDropAction::DroppedOnNode(FVector2D ScreenPosition, FV
 		const FScopedTransaction Transaction( LOCTEXT("ReplacePinVariable", "Replace Pin Variable") );
 
 		UProperty* VariableProperty = GetVariableProperty();
-		const FString OldVarName = TargetNode->GetVarNameString();
-		const UEdGraphSchema_K2* Schema = Cast<const UEdGraphSchema_K2>(TargetNode->GetSchema());
 
-		TArray<class UEdGraphPin*> BadLinks;
-		GetLinksThatWillBreak(TargetNode,VariableProperty,BadLinks);
-
-		// Change the variable name and context
-		UBlueprint* DropOnBlueprint = FBlueprintEditorUtils::FindBlueprintForGraph(TargetNode->GetGraph());
-		UEdGraphPin* Pin = TargetNode->FindPin(OldVarName);
-		DropOnBlueprint->Modify();
-		TargetNode->Modify();
-		Pin->Modify();
+		if(CanVariableBeDropped(VariableProperty, *TargetNode->GetGraph()))
+		{
+			const FString OldVarName = TargetNode->GetVarNameString();
+			const UEdGraphSchema_K2* Schema = Cast<const UEdGraphSchema_K2>(TargetNode->GetSchema());
 
-		ConfigureVarNode(TargetNode, VariableName, VariableSourceClass.Get(), DropOnBlueprint);
+			TArray<class UEdGraphPin*> BadLinks;
+			GetLinksThatWillBreak(TargetNode,VariableProperty,BadLinks);
 
+			// Change the variable name and context
+			UBlueprint* DropOnBlueprint = FBlueprintEditorUtils::FindBlueprintForGraph(TargetNode->GetGraph());
+			UEdGraphPin* Pin = TargetNode->FindPin(OldVarName);
+			DropOnBlueprint->Modify();
+			TargetNode->Modify();
+			Pin->Modify();
 
-		if ((Pin == NULL) || (Pin->LinkedTo.Num() == BadLinks.Num()) || (Schema == NULL))
-		{
-			TargetNode->GetSchema()->ReconstructNode(*TargetNode);
-		}
-		else 
-		{
-			FEdGraphPinType NewPinType;
-			Schema->ConvertPropertyToPinType(VariableProperty,NewPinType);
+			ConfigureVarNode(TargetNode, VariableName, VariableSource.Get(), DropOnBlueprint);
 
-			Pin->PinName = VariableName.ToString();
-			Pin->PinType = NewPinType;
 
-			//break bad links
-			for(TArray<class UEdGraphPin*>::TIterator OtherPinIt(BadLinks);OtherPinIt;)
+			if ((Pin == NULL) || (Pin->LinkedTo.Num() == BadLinks.Num()) || (Schema == NULL))
 			{
-				Pin->BreakLinkTo(*OtherPinIt);
+				TargetNode->GetSchema()->ReconstructNode(*TargetNode);
 			}
-		}
+			else 
+			{
+				FEdGraphPinType NewPinType;
+				Schema->ConvertPropertyToPinType(VariableProperty,NewPinType);
 
-		return FReply::Handled();
+				Pin->PinName = VariableName.ToString();
+				Pin->PinType = NewPinType;
+
+				//break bad links
+				for(TArray<class UEdGraphPin*>::TIterator OtherPinIt(BadLinks);OtherPinIt;)
+				{
+					Pin->BreakLinkTo(*OtherPinIt);
+				}
+			}
+
+			return FReply::Handled();
+		}
 	}
 	return FReply::Unhandled();
 }
 
-void FKismetVariableDragDropAction::ConfigureVarNode(UK2Node_Variable* VarNode, FName VariableName, UClass* VariableSourceClass, UBlueprint* TargetBlueprint)
+void FKismetVariableDragDropAction::ConfigureVarNode(UK2Node_Variable* InVarNode, FName InVariableName, UStruct* InVariableSource, UBlueprint* InTargetBlueprint)
 {
 	// See if this is a 'self context' (ie. blueprint class is owner (or child of owner) of dropped var class)
-	if ((VariableSourceClass == NULL) || TargetBlueprint->SkeletonGeneratedClass->IsChildOf(VariableSourceClass))
+	if ((InVariableSource == NULL) || InTargetBlueprint->SkeletonGeneratedClass->IsChildOf(InVariableSource))
+	{
+		InVarNode->VariableReference.SetSelfMember(InVariableName);
+	}
+	else if(InVariableSource->IsA(UStruct::StaticClass()))
 	{
-		VarNode->VariableReference.SetSelfMember(VariableName);
+		InVarNode->VariableReference.SetLocalMember(InVariableName, InVariableSource, FBlueprintEditorUtils::FindLocalVariableGuidByName(InTargetBlueprint, InVariableName));
 	}
 	else
 	{
-		VarNode->VariableReference.SetExternalMember(VariableName, VariableSourceClass);
+		InVarNode->VariableReference.SetExternalMember(InVariableName, CastChecked<UClass>(InVariableSource));
 	}
 }
 
 
-void FKismetVariableDragDropAction::MakeGetter(FVector2D GraphPosition, UEdGraph* Graph, FName VariableName, UClass* VariableSourceClass)
+void FKismetVariableDragDropAction::MakeGetter(FNodeConstructionParams InParams)
 {
-	check(Graph);
+	check(InParams.Graph);
 
 	UK2Node_VariableGet* GetVarNodeTemplate = NewObject<UK2Node_VariableGet>();
-	UBlueprint* DropOnBlueprint = FBlueprintEditorUtils::FindBlueprintForGraph(Graph);
-	ConfigureVarNode(GetVarNodeTemplate, VariableName, VariableSourceClass, DropOnBlueprint);
+	UBlueprint* DropOnBlueprint = FBlueprintEditorUtils::FindBlueprintForGraph(InParams.Graph);
 
-	FEdGraphSchemaAction_K2NewNode::SpawnNodeFromTemplate<UK2Node_VariableGet>(Graph, GetVarNodeTemplate, GraphPosition);
+	ConfigureVarNode(GetVarNodeTemplate, InParams.VariableName, InParams.VariableSource.Get(), DropOnBlueprint);
+
+	FEdGraphSchemaAction_K2NewNode::SpawnNodeFromTemplate<UK2Node_VariableGet>(InParams.Graph, GetVarNodeTemplate, InParams.GraphPosition);
 }
 
-void FKismetVariableDragDropAction::MakeSetter(FVector2D GraphPosition, UEdGraph* Graph, FName VariableName, UClass* VariableSourceClass)
+void FKismetVariableDragDropAction::MakeSetter(FNodeConstructionParams InParams)
 {
-	check(Graph);
+	check(InParams.Graph);
 
 	UK2Node_VariableSet* SetVarNodeTemplate = NewObject<UK2Node_VariableSet>();
-	UBlueprint* DropOnBlueprint = FBlueprintEditorUtils::FindBlueprintForGraph(Graph);
-	ConfigureVarNode(SetVarNodeTemplate, VariableName, VariableSourceClass, DropOnBlueprint);
+	UBlueprint* DropOnBlueprint = FBlueprintEditorUtils::FindBlueprintForGraph(InParams.Graph);
+
+	ConfigureVarNode(SetVarNodeTemplate, InParams.VariableName, InParams.VariableSource.Get(), DropOnBlueprint);
 
-	FEdGraphSchemaAction_K2NewNode::SpawnNodeFromTemplate<UK2Node_VariableSet>(Graph, SetVarNodeTemplate, GraphPosition);
+	FEdGraphSchemaAction_K2NewNode::SpawnNodeFromTemplate<UK2Node_VariableSet>(InParams.Graph, SetVarNodeTemplate, InParams.GraphPosition);
 }
 
-bool FKismetVariableDragDropAction::CanExecuteMakeSetter(UEdGraph* Graph, UProperty* VariableProperty, UClass* VariableSourceClass)
+bool FKismetVariableDragDropAction::CanExecuteMakeSetter(FNodeConstructionParams InParams, UProperty* InVariableProperty)
 {
-	check(VariableProperty);
-	check(VariableSourceClass);
+	check(InVariableProperty);
+	check(InParams.VariableSource.Get());
 
-	const UBlueprint* Blueprint = FBlueprintEditorUtils::FindBlueprintForGraph(Graph);
-	const bool bReadOnlyProperty = FBlueprintEditorUtils::IsPropertyReadOnlyInCurrentBlueprint(Blueprint, VariableProperty);
+	if(UClass* VariableSourceClass = Cast<UClass>(InParams.VariableSource.Get()))
+	{
+		const UBlueprint* Blueprint = FBlueprintEditorUtils::FindBlueprintForGraph(InParams.Graph);
+		const bool bReadOnlyProperty = FBlueprintEditorUtils::IsPropertyReadOnlyInCurrentBlueprint(Blueprint, InVariableProperty);
+
+		return (!bReadOnlyProperty) && (!VariableSourceClass->HasAnyClassFlags(CLASS_Const));
+	}
 
-	return (!bReadOnlyProperty) && (!VariableSourceClass->HasAnyClassFlags(CLASS_Const));
+	return true;
 }
 
 FReply FKismetVariableDragDropAction::DroppedOnPanel( const TSharedRef< SWidget >& Panel, FVector2D ScreenPosition, FVector2D GraphPosition, UEdGraph& Graph)
@@ -370,56 +439,66 @@ FReply FKismetVariableDragDropAction::DroppedOnPanel( const TSharedRef< SWidget
 	if (Cast<const UEdGraphSchema_K2>(Graph.GetSchema()) != NULL)
 	{
 		UProperty* VariableProperty = GetVariableProperty();
-		UClass* Outer = CastChecked<UClass>(VariableProperty->GetOuter());
-
-		// call analytics
-		AnalyticCallback.ExecuteIfBound();
+		UStruct* Outer = CastChecked<UStruct>(VariableProperty->GetOuter());
 
-		// Asking for a getter
-		if (bControlDrag)
+		if(CanVariableBeDropped(VariableProperty, Graph))
 		{
-			MakeGetter(GraphPosition, &Graph, VariableName, Outer);
-		}
-		// Asking for a setter
-		else if (bAltDrag && CanExecuteMakeSetter(&Graph, VariableProperty, Outer))
-		{
-			MakeSetter(GraphPosition, &Graph, VariableName, Outer);
-		}
-		// Show selection menu
-		else
-		{
-			FMenuBuilder MenuBuilder(true, NULL);
-			const FText VariableNameText = FText::FromName( VariableName );
-
-			MenuBuilder.BeginSection("BPVariableDroppedOn", VariableNameText );
-
-			MenuBuilder.AddMenuEntry(
-				LOCTEXT("CreateGetVariable", "Get"),
-				FText::Format( LOCTEXT("CreateVariableGetterToolTip", "Create Getter for variable '{0}'\n(Ctrl-drag to automatically create a getter)"), VariableNameText ),
-				FSlateIcon(),
-				FUIAction(
-				FExecuteAction::CreateStatic(&FKismetVariableDragDropAction::MakeGetter,GraphPosition,&Graph,VariableName,Outer), FCanExecuteAction())
-				);
-
-			MenuBuilder.AddMenuEntry(
-				LOCTEXT("CreateSetVariable", "Set"),
-				FText::Format( LOCTEXT("CreateVariableSetterToolTip", "Create Setter for variable '{0}'\n(Alt-drag to automatically create a setter)"), VariableNameText ),
-				FSlateIcon(),
-				FUIAction(
-				FExecuteAction::CreateStatic(&FKismetVariableDragDropAction::MakeSetter,GraphPosition,&Graph,VariableName,Outer),
-				FCanExecuteAction::CreateStatic(&FKismetVariableDragDropAction::CanExecuteMakeSetter, &Graph, VariableProperty, Outer ))
-				);
-
-			TSharedRef< SWidget > PanelWidget = Panel;
-			// Show dialog to choose getter vs setter
-			FSlateApplication::Get().PushMenu(
-				PanelWidget,
-				MenuBuilder.MakeWidget(),
-				ScreenPosition,
-				FPopupTransitionEffect( FPopupTransitionEffect::ContextMenu)
-				);
-
-			MenuBuilder.EndSection();
+			FNodeConstructionParams NewNodeParams;
+			NewNodeParams.VariableName = VariableName;
+			const UBlueprint* DropOnBlueprint = FBlueprintEditorUtils::FindBlueprintForGraph(&Graph);
+			NewNodeParams.Graph = &Graph;
+			NewNodeParams.GraphPosition = GraphPosition;
+			NewNodeParams.VariableSource= Outer;
+
+			// call analytics
+			AnalyticCallback.ExecuteIfBound();
+
+			// Asking for a getter
+			if (bControlDrag)
+			{
+				MakeGetter(NewNodeParams);
+			}
+			// Asking for a setter
+			else if (bAltDrag && CanExecuteMakeSetter(NewNodeParams, VariableProperty))
+			{
+				MakeSetter(NewNodeParams);
+			}
+			// Show selection menu
+			else
+			{
+				FMenuBuilder MenuBuilder(true, NULL);
+				const FText VariableNameText = FText::FromName( VariableName );
+
+				MenuBuilder.BeginSection("BPVariableDroppedOn", VariableNameText );
+
+				MenuBuilder.AddMenuEntry(
+					LOCTEXT("CreateGetVariable", "Get"),
+					FText::Format( LOCTEXT("CreateVariableGetterToolTip", "Create Getter for variable '{0}'\n(Ctrl-drag to automatically create a getter)"), VariableNameText ),
+					FSlateIcon(),
+					FUIAction(
+					FExecuteAction::CreateStatic(&FKismetVariableDragDropAction::MakeGetter, NewNodeParams), FCanExecuteAction())
+					);
+
+				MenuBuilder.AddMenuEntry(
+					LOCTEXT("CreateSetVariable", "Set"),
+					FText::Format( LOCTEXT("CreateVariableSetterToolTip", "Create Setter for variable '{0}'\n(Alt-drag to automatically create a setter)"), VariableNameText ),
+					FSlateIcon(),
+					FUIAction(
+					FExecuteAction::CreateStatic(&FKismetVariableDragDropAction::MakeSetter, NewNodeParams),
+					FCanExecuteAction::CreateStatic(&FKismetVariableDragDropAction::CanExecuteMakeSetter, NewNodeParams, VariableProperty ))
+					);
+
+				TSharedRef< SWidget > PanelWidget = Panel;
+				// Show dialog to choose getter vs setter
+				FSlateApplication::Get().PushMenu(
+					PanelWidget,
+					MenuBuilder.MakeWidget(),
+					ScreenPosition,
+					FPopupTransitionEffect( FPopupTransitionEffect::ContextMenu)
+					);
+
+				MenuBuilder.EndSection();
+			}
 		}
 	}
 
@@ -435,11 +514,11 @@ FReply FKismetVariableDragDropAction::DroppedOnAction(TSharedRef<FEdGraphSchemaA
 		FEdGraphSchemaAction_K2Var* VarAction = (FEdGraphSchemaAction_K2Var*)&Action.Get();
 
 		// Only let you drag and drop if variables are from same BP class, and not onto itself
-		UBlueprint* BP = UBlueprint::GetBlueprintFromClass(VariableSourceClass.Get());
+		UBlueprint* BP = UBlueprint::GetBlueprintFromClass(Cast<UClass>(VariableSource.Get()));
 		FName TargetVarName = VarAction->GetVariableName();
 		if( (BP != NULL) && 
 			(VariableName != TargetVarName) && 
-			(VariableSourceClass == VarAction->GetVariableClass()) )
+			(VariableSource == VarAction->GetVariableClass()) )
 		{
 			bool bMoved = FBlueprintEditorUtils::MoveVariableBeforeVariable(BP, VariableName, TargetVarName, true);
 			// If we moved successfully
@@ -467,7 +546,16 @@ FReply FKismetVariableDragDropAction::DroppedOnCategory(FString Category)
 {
 	UE_LOG(LogTemp, Log, TEXT("Dropped %s on Category %s"), *VariableName.ToString(), *Category);
 
-	UBlueprint* BP = UBlueprint::GetBlueprintFromClass(VariableSourceClass.Get());
+	UBlueprint* BP = NULL;
+	if(VariableSource.Get()->IsA(UFunction::StaticClass()))
+	{
+		BP = UBlueprint::GetBlueprintFromClass(Cast<UClass>(VariableSource.Get()->GetOuter()));
+	}
+	else
+	{
+		BP = UBlueprint::GetBlueprintFromClass(Cast<UClass>(VariableSource.Get()));
+	}
+
 	if(BP != NULL)
 	{
 		// Check this is actually a different category
@@ -481,6 +569,21 @@ FReply FKismetVariableDragDropAction::DroppedOnCategory(FString Category)
 	return FReply::Handled();
 }
 
+bool FKismetVariableDragDropAction::CanVariableBeDropped(const UProperty* InVariableProperty, const UEdGraph& InGraph) const
+{
+	UObject* Outer = InVariableProperty->GetOuter();
 
+	// Local variables have some special conditions for being allowed to be placed
+	bool bCanVariableBeDropped = true;
+	if(Outer->IsA(UFunction::StaticClass()))
+	{
+		// Check if the top level graph has the same name as the function, if they do not then the variable cannot be placed in the graph
+		if(FBlueprintEditorUtils::GetTopLevelGraph(&InGraph)->GetFName() != Outer->GetFName() || !IsFromBlueprint(FBlueprintEditorUtils::FindBlueprintForGraph(&InGraph)) )
+		{
+			bCanVariableBeDropped = false;
+		}
+	}
+	return bCanVariableBeDropped;
+}
 
 #undef LOCTEXT_NAMESPACE
diff --git a/Engine/Source/Editor/Kismet/Private/BPVariableDragDropAction.h b/Engine/Source/Editor/Kismet/Private/BPVariableDragDropAction.h
--- a/Engine/Source/Editor/Kismet/Private/BPVariableDragDropAction.h
+++ b/Engine/Source/Editor/Kismet/Private/BPVariableDragDropAction.h
@@ -20,11 +20,11 @@ public:
 	virtual FReply DroppedOnCategory(FString Category) OVERRIDE;
 	// End of FGraphEditorDragDropAction
 
-	static TSharedRef<FKismetVariableDragDropAction> New(FName InVariableName, UClass* InClass, FNodeCreationAnalytic AnalyticCallback)
+	static TSharedRef<FKismetVariableDragDropAction> New(FName InVariableName, UStruct* InVariableSource, FNodeCreationAnalytic AnalyticCallback)
 	{
 		TSharedRef<FKismetVariableDragDropAction> Operation = MakeShareable(new FKismetVariableDragDropAction);
 		Operation->VariableName = InVariableName;
-		Operation->VariableSourceClass = InClass;
+		Operation->VariableSource = InVariableSource;
 		Operation->AnalyticCallback = AnalyticCallback;
 		Operation->Construct();
 		return Operation;
@@ -33,22 +33,33 @@ public:
 	/** Helper method to see if we're dragging in the same blueprint */
 	bool IsFromBlueprint(class UBlueprint* InBlueprint) const
 	{ 
-		check(VariableSourceClass.IsValid());
-		UBlueprint* Blueprint = UBlueprint::GetBlueprintFromClass(VariableSourceClass.Get());
+		check(VariableSource.IsValid());
+
+		UClass* VariableSourceClass = NULL;
+		if(VariableSource.Get()->IsA(UClass::StaticClass()))
+		{
+			VariableSourceClass = CastChecked<UClass>(VariableSource.Get());
+		}
+		else
+		{
+			check(VariableSource.Get()->GetOuter());
+			VariableSourceClass = CastChecked<UClass>(VariableSource.Get()->GetOuter());
+		}
+		UBlueprint* Blueprint = UBlueprint::GetBlueprintFromClass(VariableSourceClass);
 		return (InBlueprint == Blueprint);
 	}
 
 	UProperty* GetVariableProperty()
 	{
-		check(VariableSourceClass.IsValid());
+		check(VariableSource.IsValid());
 		check(VariableName != NAME_None);
-		UProperty* VariableProperty = FindField<UProperty>(VariableSourceClass.Get(), VariableName);
+		UProperty* VariableProperty = FindField<UProperty>(VariableSource.Get(), VariableName);
 		check(VariableProperty != NULL);
 		return VariableProperty;
 	}
 
 	/** Configure the supplied variable node based on this action's info */
-	static void ConfigureVarNode(UK2Node_Variable* VarNode, FName VariableName, UClass* VariableSourceClass, UBlueprint* TargetBlueprint);
+	static void ConfigureVarNode(UK2Node_Variable* InVarNode, FName InVariableName, UStruct* InVariableSource, UBlueprint* InTargetBlueprint);
 
 	/** Set if operation is modified by alt */
 	void SetAltDrag(bool InIsAltDrag) {	bAltDrag = InIsAltDrag;}
@@ -60,12 +71,21 @@ protected:
 	 /** Construct a FKismetVariableDragDropAction */
 	FKismetVariableDragDropAction();
 
+	/** Structure for required node construction parameters */
+	struct FNodeConstructionParams
+	{
+		FVector2D GraphPosition;
+		UEdGraph* Graph;
+		FName VariableName;
+		TWeakObjectPtr<UStruct> VariableSource;
+	};
+
 	/** Called when user selects to create a Getter for the variable */
-	static void MakeGetter(FVector2D GraphPosition, UEdGraph* Graph, FName VariableName, UClass* VariableSourceClass);
+	static void MakeGetter(FNodeConstructionParams InParams);
 	/** Called when user selects to create a Setter for the variable */
-	static void MakeSetter(FVector2D GraphPosition, UEdGraph* Graph, FName VariableName, UClass* VariableSourceClass);
+	static void MakeSetter(FNodeConstructionParams InParams);
 	/** Called too check if we can execute a setter on a given property */
-	static bool CanExecuteMakeSetter(UEdGraph* Graph, UProperty* VariableProperty, UClass* VariableSourceClass);
+	static bool CanExecuteMakeSetter(FNodeConstructionParams InParams, UProperty* InVariableProperty);
 
 	/**
 	 * Test new variable type against existing links for node and get any links that will break
@@ -84,15 +104,23 @@ protected:
 		return BadLinks.Num() > 0;
 	}
 
+	/**
+	 * Checks if the property can be dropped in a graph
+	 *
+	 * @param InVariableProperty		The variable property to check with
+	 * @param InGraph					The graph to check against placing the variable
+	 */
+	bool CanVariableBeDropped(const UProperty* InVariableProperty, const UEdGraph& InGraph) const;
+
 protected:
 	/** Name of variable being dragged */
 	FName VariableName;
-	/** Class this variable belongs to */
-	TWeakObjectPtr<UClass> VariableSourceClass;
+	/** Scope this variable belongs to */
+	TWeakObjectPtr<UStruct> VariableSource;
 	/** Was ctrl held down at start of drag */
 	bool bControlDrag;
 	/** Was alt held down at the start of drag */
 	bool bAltDrag;
 	/** Analytic delegate to track node creation */
 	FNodeCreationAnalytic AnalyticCallback;
-};
\ No newline at end of file
+};
diff --git a/Engine/Source/Editor/Kismet/Private/BlueprintDetailsCustomization.cpp b/Engine/Source/Editor/Kismet/Private/BlueprintDetailsCustomization.cpp
--- a/Engine/Source/Editor/Kismet/Private/BlueprintDetailsCustomization.cpp
+++ b/Engine/Source/Editor/Kismet/Private/BlueprintDetailsCustomization.cpp
@@ -436,19 +436,19 @@ void FBlueprintVarActionDetails::CustomizeDetails( IDetailLayoutBuilder& DetailL
 		}
 	}
 
-	// Add in default value editing for properties that can be edited
+	// Add in default value editing for properties that can be edited, local properties cannot be edited
 	UBlueprint* Blueprint = GetBlueprintObj();
-	if ((Blueprint != NULL) && (Blueprint->GeneratedClass != NULL))
+	if ((Blueprint != NULL) && (Blueprint->GeneratedClass != NULL) && !IsALocalVariable(VariableProperty))
 	{
 		if (VariableProperty != NULL)
 		{
-			const UProperty* OryginalProperty = FindField<UProperty>(Blueprint->GeneratedClass, VariableProperty->GetFName());
-			if (OryginalProperty == NULL)
+			const UProperty* OriginalProperty = FindField<UProperty>(Blueprint->GeneratedClass, VariableProperty->GetFName());
+			if (OriginalProperty == NULL)
 			{
 				// Prevent editing the default value of a skeleton property
 				VariableProperty = NULL;
 			}
-			else  if (auto StructProperty = Cast<const UStructProperty>(OryginalProperty))
+			else  if (auto StructProperty = Cast<const UStructProperty>(OriginalProperty))
 			{
 				// Prevent editing the default value of a stale struct
 				auto BGStruct = Cast<const UUserDefinedStruct>(StructProperty->Struct);
@@ -627,8 +627,17 @@ bool FBlueprintVarActionDetails::IsAComponentVariable(UProperty* VariablePropert
 
 bool FBlueprintVarActionDetails::IsABlueprintVariable(UProperty* VariableProperty) const
 {
-	UClass* VarSourceClass = VariableProperty ? CastChecked<UClass>(VariableProperty->GetOuter()) : NULL;
-	return (VarSourceClass->ClassGeneratedBy != NULL);
+	UClass* VarSourceClass = VariableProperty ? Cast<UClass>(VariableProperty->GetOuter()) : NULL;
+	if(VarSourceClass)
+	{
+		return (VarSourceClass->ClassGeneratedBy != NULL);
+	}
+	return false;
+}
+
+bool FBlueprintVarActionDetails::IsALocalVariable(UProperty* VariableProperty) const
+{
+	return Cast<UFunction>(VariableProperty->GetOuter()) != NULL;
 }
 
 bool FBlueprintVarActionDetails::GetVariableNameChangeEnabled() const
@@ -662,6 +671,10 @@ bool FBlueprintVarActionDetails::GetVariableNameChangeEnabled() const
 				}
 			}
 		}
+		else if(IsALocalVariable(VariableProperty))
+		{
+			bIsReadOnly = false;
+		}
 	}
 
 	return bIsReadOnly;
@@ -742,6 +755,10 @@ void FBlueprintVarActionDetails::OnVarNameCommitted(const FText& InNewText, ETex
 		{
 			FBlueprintEditorUtils::RenameTimeline(GetBlueprintObj(), GetVariableName(), NewVarName);
 		}
+		else if(IsALocalVariable(VariableProperty))
+		{
+			FBlueprintEditorUtils::RenameLocalVariable(GetBlueprintObj(), GetVariableName(), NewVarName);
+		}
 		else
 		{
 			FBlueprintEditorUtils::RenameMemberVariable(GetBlueprintObj(), GetVariableName(), NewVarName);
@@ -773,22 +790,53 @@ bool FBlueprintVarActionDetails::GetVariableTypeChangeEnabled() const
 				break;
 			}
 		}
+		
 		bool bIsAVarInThisBlueprint = FBlueprintEditorUtils::FindNewVariableIndex(GetBlueprintObj(), VarAction->GetVariableName()) != INDEX_NONE;
 		return !bNodesPendingDeletion && bIsAVarInThisBlueprint;
 	}
+
+	FEdGraphSchemaAction_K2LocalVar* LocalVarAction = MyBlueprintSelectionAsLocalVar();
+	if (LocalVarAction)
+	{
+		TArray<UK2Node_Variable*> VariableNodes;
+		FBlueprintEditorUtils::GetAllNodesOfClass<UK2Node_Variable>(GetBlueprintObj(), VariableNodes);
+
+		bool bNodesPendingDeletion = false;
+		for( TArray<UK2Node_Variable*>::TConstIterator NodeIt(VariableNodes); NodeIt; ++NodeIt )
+		{
+			UK2Node_Variable* CurrentNode = *NodeIt;
+			if (LocalVarAction->GetVariableName() == CurrentNode->GetVarName())
+			{
+				bNodesPendingDeletion = true;
+				break;
+			}
+		}
+
+		bool bIsAVarInThisBlueprint = FBlueprintEditorUtils::FindLocalVariable(GetBlueprintObj(), LocalVarAction->GetVariableName()) != NULL;
+		return !bNodesPendingDeletion && bIsAVarInThisBlueprint;
+	}
+	FEdGraphSchemaAction_K2LocalVar* VarLocalAction = MyBlueprintSelectionAsLocalVar();
+	if(VarLocalAction)
+	{
+		return true;
+	}
 	return false;
 }
 
 bool FBlueprintVarActionDetails::GetVariableCategoryChangeEnabled() const
 {
 	UProperty* VariableProp = SelectionAsProperty();
-
 	if(VariableProp)
 	{
-		UClass* VarSourceClass = CastChecked<UClass>(VariableProp->GetOuter());
-
-		// If the variable's source class is the same as the current blueprint's class then it was created in this blueprint and it's category can be changed.
-		return VarSourceClass == GetBlueprintObj()->SkeletonGeneratedClass;
+		if(UClass* VarSourceClass = Cast<UClass>(VariableProp->GetOuter()))
+		{
+			// If the variable's source class is the same as the current blueprint's class then it was created in this blueprint and it's category can be changed.
+			return VarSourceClass == GetBlueprintObj()->SkeletonGeneratedClass;
+		}
+		else if(IsALocalVariable(VariableProp))
+		{
+			return true;
+		}
 	}
 
 	return false;
@@ -813,12 +861,24 @@ void FBlueprintVarActionDetails::OnVarTypeChanged(const FEdGraphPinType& NewPinT
 	if (FBlueprintEditorUtils::IsPinTypeValid(NewPinType))
 	{
 		FName VarName = GetVariableName();
+
 		if (VarName != NAME_None)
 		{
 			FBlueprintEditorUtils::ChangeMemberVariableType(GetBlueprintObj(), VarName, NewPinType);
-
 			// Set the MyBP tab's last pin type used as this, for adding lots of variables of the same type
 			MyBlueprint.Pin()->GetLastPinTypeUsed() = NewPinType;
+
+			FEdGraphSchemaAction_K2Var* VarAction = MyBlueprintSelectionAsVar();
+			if(VarAction)
+			{
+				FBlueprintEditorUtils::ChangeMemberVariableType(GetBlueprintObj(), VarName, NewPinType);
+			}
+
+			FEdGraphSchemaAction_K2LocalVar* VarLocalAction = MyBlueprintSelectionAsLocalVar();
+			if(VarLocalAction)
+			{
+				FBlueprintEditorUtils::ChangeLocalVariableType(GetBlueprintObj(), VarName, NewPinType);
+			}
 		}
 	}
 }
@@ -880,6 +940,26 @@ void FBlueprintVarActionDetails::PopulateCategories(SMyBlueprint* MyBlueprint, T
 			}
 		}
 	}
+
+	// Search through all function entry nodes for local variables to pull their categories
+	TArray<UK2Node_FunctionEntry*> FunctionEntryNodes;
+	FBlueprintEditorUtils::GetAllNodesOfClass<UK2Node_FunctionEntry>(Blueprint, FunctionEntryNodes);
+
+	for (UK2Node_FunctionEntry* const FunctionEntry : FunctionEntryNodes)
+	{
+		for( FBPVariableDescription& Variable : FunctionEntry->LocalVariables )
+		{
+			bool bNewCategory = true;
+			for (int32 j = 0; j < CategorySource.Num() && bNewCategory; ++j)
+			{
+				bNewCategory &= *CategorySource[j].Get() != Variable.Category.ToString();
+			}
+			if (bNewCategory)
+			{
+				CategorySource.Add(MakeShareable(new FString(Variable.Category.ToString())));
+			}
+		}
+	}
 }
 
 UK2Node_Variable* FBlueprintVarActionDetails::EdGraphSelectionAsVar() const
@@ -912,25 +992,15 @@ UProperty* FBlueprintVarActionDetails::SelectionAsProperty() const
 	{
 		return VarAction->GetProperty();
 	}
-	UK2Node_Variable* GraphVar = EdGraphSelectionAsVar();
-	if(GraphVar)
+	FEdGraphSchemaAction_K2LocalVar* LocalVarAction = MyBlueprintSelectionAsLocalVar();
+	if(LocalVarAction)
 	{
-		return GraphVar->GetPropertyForVariable();
-	}
-	return NULL;
-}
-
-UClass* FBlueprintVarActionDetails::SelectionAsClass() const
-{
-	FEdGraphSchemaAction_K2Var* VarAction = MyBlueprintSelectionAsVar();
-	if(VarAction)
-	{
-		return VarAction->GetVariableClass();
+		return LocalVarAction->GetProperty();
 	}
 	UK2Node_Variable* GraphVar = EdGraphSelectionAsVar();
 	if(GraphVar)
 	{
-		return GraphVar->GetClass();
+		return GraphVar->GetPropertyForVariable();
 	}
 	return NULL;
 }
@@ -942,6 +1012,11 @@ FName FBlueprintVarActionDetails::GetVariableName() const
 	{
 		return VarAction->GetVariableName();
 	}
+	FEdGraphSchemaAction_K2LocalVar* LocalVarAction = MyBlueprintSelectionAsLocalVar();
+	if(LocalVarAction)
+	{
+		return LocalVarAction->GetVariableName();
+	}
 	UK2Node_Variable* GraphVar = EdGraphSelectionAsVar();
 	if(GraphVar)
 	{
@@ -1205,7 +1280,7 @@ void FBlueprintVarActionDetails::OnExposedToMatineeChanged(ESlateCheckBoxState::
 EVisibility FBlueprintVarActionDetails::ExposeToMatineeVisibility() const
 {
 	UProperty* VariableProperty = SelectionAsProperty();
-	if (VariableProperty)
+	if (VariableProperty && !IsALocalVariable(VariableProperty))
 	{
 		const bool bIsInteger = VariableProperty->IsA(UIntProperty::StaticClass());
 		const bool bIsNonEnumByte = (VariableProperty->IsA(UByteProperty::StaticClass()) && Cast<const UByteProperty>(VariableProperty)->Enum == NULL);
@@ -1478,7 +1553,7 @@ EVisibility FBlueprintVarActionDetails::IsTooltipEditVisible() const
 	UProperty* VariableProperty = SelectionAsProperty();
 	if (VariableProperty)
 	{
-		if (IsABlueprintVariable(VariableProperty) && !IsAComponentVariable(VariableProperty))
+		if (IsABlueprintVariable(VariableProperty) && !IsAComponentVariable(VariableProperty) || IsALocalVariable(VariableProperty))
 		{
 			return EVisibility::Visible;
 		}
@@ -1486,230 +1561,6 @@ EVisibility FBlueprintVarActionDetails::IsTooltipEditVisible() const
 	return EVisibility::Collapsed;
 }
 
-// Local Variable Detail Customization
-BEGIN_SLATE_FUNCTION_BUILD_OPTIMIZATION
-	void FBlueprintLocalVarActionDetails::CustomizeDetails( IDetailLayoutBuilder& DetailLayout )
-{
-	FName VarName =	GetVariableName();
-	const UEdGraphSchema_K2* Schema = GetDefault<UEdGraphSchema_K2>();
-
-	IDetailCategoryBuilder& Category = DetailLayout.EditCategory("Local Variable", LOCTEXT("LocalVariableDetailsCategory", "Local Variable").ToString());
-	const FSlateFontInfo DetailFontInfo = IDetailLayoutBuilder::GetDetailFont();
-
-	Category.AddCustomRow( TEXT("Variable Name") )
-		.NameContent()
-		[
-			SNew(STextBlock)
-			.Text(LOCTEXT("BlueprintVarActionDetails_VariableNameLabel", "Variable Name").ToString())
-			.Font(IDetailLayoutBuilder::GetDetailFont())
-		]
-	.ValueContent()
-		[
-			SAssignNew(VarNameEditableTextBox, SEditableTextBox)
-			.Text(this, &FBlueprintLocalVarActionDetails::OnGetVarName)
-			.OnTextChanged(this, &FBlueprintLocalVarActionDetails::OnLocalVarNameChanged)
-			.OnTextCommitted(this, &FBlueprintLocalVarActionDetails::OnLocalVarNameCommitted)
-			.Font(IDetailLayoutBuilder::GetDetailFont())
-		];
-	Category.AddCustomRow( TEXT("Variable Type") )
-		.NameContent()
-		[
-			SNew(STextBlock)
-			.Text( LOCTEXT("VariableTypeLabel", "Variable Type") )
-			.Font( DetailFontInfo )
-		]
-	.ValueContent()
-		[
-			SNew(SPinTypeSelector, FGetPinTypeTree::CreateUObject(Schema, &UEdGraphSchema_K2::GetVariableTypeTree))
-			.TargetPinType(this, &FBlueprintLocalVarActionDetails::OnGetVarType)
-			.OnPinTypeChanged(this, &FBlueprintLocalVarActionDetails::OnVarTypeChanged)
-			.IsEnabled(this, &FBlueprintLocalVarActionDetails::GetVariableTypeChangeEnabled)
-			.Schema(Schema)
-			.bAllowExec(false)
-			.Font( DetailFontInfo )
-		];
-	Category.AddCustomRow( TEXT("Tooltip") )
-		.NameContent()
-		[
-			SNew(STextBlock)
-			.Text( LOCTEXT("IsVariableToolTipLabel", "Tooltip") )
-			.Font( DetailFontInfo )
-		]
-	.ValueContent()
-		[
-			SNew(SEditableTextBox)
-			.Text( this, &FBlueprintLocalVarActionDetails::OnGetTooltipText )
-			.OnTextCommitted( this, &FBlueprintLocalVarActionDetails::OnTooltipTextCommitted )
-			.Font( DetailFontInfo )
-		];
-}
-END_SLATE_FUNCTION_BUILD_OPTIMIZATION
-
-	UK2Node_LocalVariable* FBlueprintLocalVarActionDetails::EdGraphSelectionAsLocalVar() const
-{
-	TWeakPtr<FBlueprintEditor> BlueprintEditor = MyBlueprint.Pin()->GetBlueprintEditor();
-
-	if( BlueprintEditor.IsValid() )
-	{
-		/** Get the currently selected set of nodes */
-		TSet<UObject*> Objects = BlueprintEditor.Pin()->GetSelectedNodes();
-
-		if (Objects.Num() == 1)
-		{
-			TSet<UObject*>::TIterator Iter(Objects);
-			UObject* Object = *Iter;
-
-			if (Object && Object->IsA<UK2Node_LocalVariable>())
-			{
-				return Cast<UK2Node_LocalVariable>(Object);
-			}
-		}
-	}
-	return NULL;
-}
-
-UK2Node_LocalVariable* FBlueprintLocalVarActionDetails::GetDetailsSelection() const
-{
-	UK2Node_LocalVariable* LocalVar = nullptr;
-
-	// Check if the selection is from the MyBlueprints window
-	LocalVar = MyBlueprintSelectionAsLocalVar();
-	if(LocalVar)
-	{
-		return LocalVar;
-	}
-
-	// Check if the selection is in the graph
-	LocalVar = EdGraphSelectionAsLocalVar();
-	if(LocalVar)
-	{
-		return LocalVar;
-	}
-
-	return nullptr;
-}
-
-void FBlueprintLocalVarActionDetails::OnLocalVarNameChanged(const FText& InNewText)
-{
-	bIsVarNameInvalid = true;
-
-	UK2Node_LocalVariable* LocalVar = GetDetailsSelection();
-
-	TSharedPtr<INameValidatorInterface> NameValidator = LocalVar->MakeNameValidator();
-
-	EValidatorResult ValidatorResult = NameValidator->IsValid(InNewText.ToString());
-	if(ValidatorResult == EValidatorResult::AlreadyInUse)
-	{
-		VarNameEditableTextBox->SetError(FText::Format(LOCTEXT("RenameFailed_InUse", "{0} is in use by another variable or function!"), InNewText));
-	}
-	else if(ValidatorResult == EValidatorResult::EmptyName)
-	{
-		VarNameEditableTextBox->SetError(LOCTEXT("RenameFailed_LeftBlank", "Names cannot be left blank!"));
-	}
-	else if(ValidatorResult == EValidatorResult::TooLong)
-	{
-		VarNameEditableTextBox->SetError(LOCTEXT("RenameFailed_NameTooLong", "Names must have fewer than 100 characters!"));
-	}
-	else
-	{
-		bIsVarNameInvalid = false;
-		VarNameEditableTextBox->SetError(FText::GetEmpty());
-	}
-}
-
-void FBlueprintLocalVarActionDetails::OnLocalVarNameCommitted(const FText& InNewText, ETextCommit::Type InTextCommit)
-{
-	if((InTextCommit == ETextCommit::OnEnter || InTextCommit == ETextCommit::OnUserMovedFocus) && !bIsVarNameInvalid)
-	{
-		const FScopedTransaction Transaction( LOCTEXT( "RenameLocalVariable", "Rename Local Variable" ) );
-
-		FName NewVarName = FName(*InNewText.ToString());
-
-		auto LocalVar = GetDetailsSelection();
-		LocalVar->Modify();
-		LocalVar->CustomVariableName = NewVarName;
-
-		check(MyBlueprint.IsValid());
-		MyBlueprint.Pin()->Refresh();
-		MyBlueprint.Pin()->SelectItemByName(NewVarName, ESelectInfo::OnMouseClick);
-	}
-
-	bIsVarNameInvalid = false;
-	VarNameEditableTextBox->SetError(FText::GetEmpty());
-}
-
-FName FBlueprintLocalVarActionDetails::GetVariableName() const
-{
-	UK2Node_LocalVariable* LocalVar = GetDetailsSelection();
-	if(LocalVar)
-	{
-		return LocalVar->CustomVariableName;
-	}
-	return NAME_None;
-}
-
-FText FBlueprintLocalVarActionDetails::OnGetVarName() const
-{
-	return FText::FromName(GetVariableName());
-}
-
-FEdGraphPinType FBlueprintLocalVarActionDetails::OnGetVarType() const
-{
-	UK2Node_LocalVariable* LocalVar = GetDetailsSelection();
-	if(LocalVar)
-	{
-		return LocalVar->VariableType;
-	}
-	return FEdGraphPinType();
-}
-
-void FBlueprintLocalVarActionDetails::OnVarTypeChanged(const FEdGraphPinType& NewPinType)
-{
-	if (FBlueprintEditorUtils::IsPinTypeValid(NewPinType))
-	{
-		UK2Node_LocalVariable* LocalVar = GetDetailsSelection();
-		if (LocalVar)
-		{
-			LocalVar->ChangeVariableType(NewPinType);
-
-			// Set the MyBP tab's last pin type used as this, for adding lots of variables of the same type
-			MyBlueprint.Pin()->GetLastPinTypeUsed() = NewPinType;
-		}
-	}
-}
-
-bool FBlueprintLocalVarActionDetails::GetVariableTypeChangeEnabled() const
-{
-	UK2Node_LocalVariable* VarAction = GetDetailsSelection();
-	if (VarAction)
-	{
-		if(VarAction->GetVariablePin()->LinkedTo.Num() != 0)
-		{
-			return false;
-		}
-	}
-	return true;
-}
-
-FText FBlueprintLocalVarActionDetails::OnGetTooltipText() const
-{
-	UK2Node_LocalVariable* VarAction = GetDetailsSelection();
-	if (VarAction && !VarAction->VariableTooltip.IsEmpty())
-	{
-		return FText::FromString(VarAction->GetTooltip());
-	}
-	return FText::GetEmpty();
-}
-
-void FBlueprintLocalVarActionDetails::OnTooltipTextCommitted(const FText& NewText, ETextCommit::Type InTextCommit)
-{
-	UK2Node_LocalVariable* VarAction = GetDetailsSelection();
-	if (VarAction)
-	{
-		VarAction->VariableTooltip = NewText;
-	}
-}
-
 static FDetailWidgetRow& AddRow( TArray<TSharedRef<FDetailWidgetRow> >& OutChildRows )
 {
 	TSharedRef<FDetailWidgetRow> NewRow( new FDetailWidgetRow );
diff --git a/Engine/Source/Editor/Kismet/Private/BlueprintDetailsCustomization.h b/Engine/Source/Editor/Kismet/Private/BlueprintDetailsCustomization.h
--- a/Engine/Source/Editor/Kismet/Private/BlueprintDetailsCustomization.h
+++ b/Engine/Source/Editor/Kismet/Private/BlueprintDetailsCustomization.h
@@ -39,14 +39,15 @@ private:
 	/** Accessors passed to parent */
 	UBlueprint* GetBlueprintObj() const {return MyBlueprint.Pin()->GetBlueprintObj();}
 	FEdGraphSchemaAction_K2Var* MyBlueprintSelectionAsVar() const {return MyBlueprint.Pin()->SelectionAsVar();}
+	FEdGraphSchemaAction_K2LocalVar* MyBlueprintSelectionAsLocalVar() const {return MyBlueprint.Pin()->SelectionAsLocalVar();}
 	UK2Node_Variable* EdGraphSelectionAsVar() const;
 	UProperty* SelectionAsProperty() const;
-	UClass* SelectionAsClass() const;
 	FName GetVariableName() const;
 
 	/** Commonly queried attributes about the schema action */
 	bool IsAComponentVariable(UProperty* VariableProperty) const;
 	bool IsABlueprintVariable(UProperty* VariableProperty) const;
+	bool IsALocalVariable(UProperty* VariableProperty) const;
 
 	// Callbacks for uproperty details customization
 	bool GetVariableNameChangeEnabled() const;
@@ -142,55 +143,6 @@ private:
 	TWeakPtr< SListView< TSharedPtr< FString > > > PropertyFlagWidget;
 };
 
-/** Details customization for local variables selected in the MyBlueprint panel */
-class FBlueprintLocalVarActionDetails : public IDetailCustomization
-{
-public:
-	/** Makes a new instance of this detail layout class for a specific detail view requesting it */
-	static TSharedRef<class IDetailCustomization> MakeInstance(TWeakPtr<SMyBlueprint> InMyBlueprint)
-	{
-		return MakeShareable(new FBlueprintLocalVarActionDetails(InMyBlueprint));
-	}
-
-	FBlueprintLocalVarActionDetails(TWeakPtr<SMyBlueprint> InMyBlueprint)
-		: MyBlueprint(InMyBlueprint)
-		, bIsVarNameInvalid(false)
-	{
-	}
-
-	/** IDetailCustomization interface */
-	virtual void CustomizeDetails( IDetailLayoutBuilder& DetailLayout ) OVERRIDE;
-
-private:
-	/** Accessors passed to parent */
-	UBlueprint* GetBlueprintObj() const {return MyBlueprint.Pin()->GetBlueprintObj();}
-	UK2Node_LocalVariable* GetDetailsSelection() const;
-	UK2Node_LocalVariable* MyBlueprintSelectionAsLocalVar() const {return MyBlueprint.Pin()->SelectionAsLocalVar();}
-	UK2Node_LocalVariable* EdGraphSelectionAsLocalVar() const;
-	FName GetVariableName() const;
-
-	// Callbacks for uproperty details customization
-	FText OnGetVarName() const;
-	void OnLocalVarNameChanged(const FText& InNewText);
-	void OnLocalVarNameCommitted(const FText& InNewName, ETextCommit::Type InTextCommit);
-	bool GetVariableTypeChangeEnabled() const;
-	FEdGraphPinType OnGetVarType() const;
-	void OnVarTypeChanged(const FEdGraphPinType& NewPinType);
-
-	FText OnGetTooltipText() const;
-	void OnTooltipTextCommitted(const FText& NewText, ETextCommit::Type InTextCommit);
-
-private:
-	/** Pointer back to my parent tab */
-	TWeakPtr<SMyBlueprint> MyBlueprint;
-
-	/** The widget used when in variable name editing mode */ 
-	TSharedPtr<SEditableTextBox> VarNameEditableTextBox;
-
-	/** Flag to indicate whether or not the variable name is invalid */
-	bool bIsVarNameInvalid;
-};
-
 class FBaseBlueprintGraphActionDetails : public IDetailCustomization
 {
 public:
diff --git a/Engine/Source/Editor/Kismet/Private/BlueprintEditor.cpp b/Engine/Source/Editor/Kismet/Private/BlueprintEditor.cpp
--- a/Engine/Source/Editor/Kismet/Private/BlueprintEditor.cpp
+++ b/Engine/Source/Editor/Kismet/Private/BlueprintEditor.cpp
@@ -249,10 +249,10 @@ const FSlateBrush* FBlueprintEditor::GetGlyphForGraph(const UEdGraph* Graph, boo
 	return ReturnValue;
 }
 
-FSlateBrush const* FBlueprintEditor::GetVarIconAndColor(UClass* VarClass, FName VarName, FSlateColor& IconColorOut)
+FSlateBrush const* FBlueprintEditor::GetVarIconAndColor(UStruct* VarScope, FName VarName, FSlateColor& IconColorOut)
 {
 	FLinearColor ColorOut;
-	const FSlateBrush* IconBrush = FEditorStyle::GetBrush(UK2Node_Variable::GetVariableIconAndColor(VarClass, VarName, ColorOut));
+	const FSlateBrush* IconBrush = FEditorStyle::GetBrush(UK2Node_Variable::GetVariableIconAndColor(VarScope, VarName, ColorOut));
 	IconColorOut = ColorOut;
 	return IconBrush;
 }
@@ -5323,7 +5323,10 @@ void FBlueprintEditor::OnAddNewVariable()
 
 void FBlueprintEditor::OnAddNewLocalVariable()
 {
-	UEdGraph* TargetGraph = FocusedGraphEdPtr.Pin()->GetCurrentGraph();
+	// Find the top level graph to place the local variables into
+	UEdGraph* TargetGraph = FBlueprintEditorUtils::GetTopLevelGraph(FocusedGraphEdPtr.Pin()->GetCurrentGraph());
+	check(TargetGraph->GetSchema()->GetGraphType(TargetGraph) == GT_Function);
+
 	if(TargetGraph != NULL)
 	{
 		const FScopedTransaction Transaction( LOCTEXT("AddLocalVariable", "Add Local Variable") );
@@ -5332,19 +5335,26 @@ void FBlueprintEditor::OnAddNewLocalVariable()
 		// Figure out a decent place to stick the node
 		const FVector2D NewNodePos = TargetGraph->GetGoodPlaceForNewNode();
 
-		// Create a new event node
-		UK2Node_LocalVariable* LocalVariableTemplate = NewObject<UK2Node_LocalVariable>();
-		LocalVariableTemplate->VariableType = MyBlueprintWidget->GetLastPinTypeUsed();
+		TArray<UK2Node_FunctionEntry*> FunctionEntryNodes;
+		TargetGraph->GetNodesOfClass(FunctionEntryNodes);
+		check(FunctionEntryNodes.Num());
 
-		UK2Node_LocalVariable* LocalVariableNode = FEdGraphSchemaAction_K2NewNode::SpawnNodeFromTemplate<UK2Node_LocalVariable>(TargetGraph, LocalVariableTemplate, NewNodePos);
-		FBlueprintEditorUtils::MarkBlueprintAsStructurallyModified(GetBlueprintObj());
+		const UEdGraphSchema_K2* K2Schema = GetDefault<UEdGraphSchema_K2>();
 
-		// Finally, bring up kismet and jump to the new node
-		if (LocalVariableNode != NULL)
-		{
-			FKismetEditorUtilities::BringKismetToFocusAttentionOnObject(LocalVariableNode);
-			RenameNewlyAddedAction(LocalVariableNode->CustomVariableName);
-		}
+		// Now create new variable
+		FBPVariableDescription NewVar;
+
+		NewVar.VarName = FBlueprintEditorUtils::FindUniqueKismetName(GetBlueprintObj(), TEXT("NewLocalVar"));
+		NewVar.VarGuid = FGuid::NewGuid();
+		NewVar.VarType = GetMyBlueprintWidget()->GetLastPinTypeUsed();
+		NewVar.FriendlyName = FName::NameToDisplayString( NewVar.VarName.ToString(), (NewVar.VarType.PinCategory == K2Schema->PC_Boolean) ? true : false );
+		NewVar.Category = K2Schema->VR_DefaultCategory;
+
+		FunctionEntryNodes[0]->Modify();
+		FunctionEntryNodes[0]->LocalVariables.Add(NewVar);
+
+		FBlueprintEditorUtils::MarkBlueprintAsStructurallyModified(GetBlueprintObj());
+		RenameNewlyAddedAction(NewVar.VarName);
 	}
 }
 
@@ -5509,7 +5519,7 @@ bool FBlueprintEditor::NewDocument_IsVisibleForType(ECreatedDocumentType GraphTy
 	case CGT_NewEventGraph:
 		return FBlueprintEditorUtils::DoesSupportEventGraphs(GetBlueprintObj());
 	case CGT_NewLocalVariable:
-		return (GetBlueprintObj()->BlueprintType != BPTYPE_Interface) && (GetFocusedGraph() && GetFocusedGraph()->GetSchema()->GetGraphType(GetFocusedGraph()) == EGraphType::GT_Function);
+		return (GetBlueprintObj()->BlueprintType != BPTYPE_Interface) && (GetFocusedGraph() && GetFocusedGraph()->GetSchema()->GetGraphType(GetFocusedGraph()) == EGraphType::GT_Function) && !GetFocusedGraph()->IsA(UAnimationTransitionGraph::StaticClass());
 	}
 
 	return false;
diff --git a/Engine/Source/Editor/Kismet/Private/SBlueprintPalette.cpp b/Engine/Source/Editor/Kismet/Private/SBlueprintPalette.cpp
--- a/Engine/Source/Editor/Kismet/Private/SBlueprintPalette.cpp
+++ b/Engine/Source/Editor/Kismet/Private/SBlueprintPalette.cpp
@@ -47,18 +47,18 @@ namespace BlueprintPalette
  * associated with specific blueprint variables (gets a string representing the 
  * specified variable's type)
  * 
- * @param  VarClass	The class that owns the variable in question.
+ * @param  VarScope	The struct that owns the variable in question.
  * @param  VarName	The name of the variable you want the type of.
  * @param  Detailed	If true the returned string includes SubCategoryObject
  * @return A string representing the variable's type (empty if the variable couldn't be found).
  */
-static FString GetVarType(UClass* VarClass, FName VarName, bool bUseObjToolTip, bool bDetailed = false)
+static FString GetVarType(UStruct* VarScope, FName VarName, bool bUseObjToolTip, bool bDetailed = false)
 {
 	FString VarDesc;
 
-	if(VarClass != NULL)
+	if(VarScope != NULL)
 	{
-		UProperty* Property = FindField<UProperty>(VarClass, VarName);
+		UProperty* Property = FindField<UProperty>(VarScope, VarName);
 		if(Property != NULL)
 		{
 			// If it is an object property, see if we can get a nice class description instead of just the name
@@ -279,12 +279,6 @@ static void GetPaletteItemIcon(TSharedPtr<FEdGraphSchemaAction> ActionIn, UBluep
 		{
 			BrushOut = FEditorStyle::GetBrush(TEXT("GraphEditor.Event_16x"));
 		}
-		else if (UK2Node_LocalVariable* LocalVariableNode = Cast<UK2Node_LocalVariable>(TargetNodeAction->NodeTemplate))
-		{
-			FLinearColor LinearColor;
-			BrushOut = FEditorStyle::GetBrush(UK2Node_Variable::GetVarIconFromPinType(LocalVariableNode->VariableType, LinearColor));
-			ColorOut = LinearColor;
-		}
 	}
 	else if (ActionIn->GetTypeId() == FEdGraphSchemaAction_K2AddComponent::StaticGetTypeId())
 	{
@@ -334,6 +328,17 @@ static void GetPaletteItemIcon(TSharedPtr<FEdGraphSchemaAction> ActionIn, UBluep
 		DocLinkOut = TEXT("Shared/Editor/Blueprint/VariableTypes");
 		DocExcerptOut = GetVarType(VarClass, VarAction->GetVariableName(), false, false);
 	}
+	else if(ActionIn->GetTypeId() == FEdGraphSchemaAction_K2LocalVar::StaticGetTypeId())
+	{
+		FEdGraphSchemaAction_K2LocalVar* LocalVarAction = (FEdGraphSchemaAction_K2LocalVar*)ActionIn.Get();
+
+		UStruct* VarScope = LocalVarAction->GetVariableScope();
+		BrushOut = FBlueprintEditor::GetVarIconAndColor(VarScope, LocalVarAction->GetVariableName(), ColorOut);
+		ToolTipOut = GetVarType(VarScope, LocalVarAction->GetVariableName(), true);
+
+		DocLinkOut = TEXT("Shared/Editor/Blueprint/VariableTypes");
+		DocExcerptOut = GetVarType(VarScope, LocalVarAction->GetVariableName(), false);
+	}
 	else if(ActionIn->GetTypeId() == FEdGraphSchemaAction_K2Enum::StaticGetTypeId())
 	{
 		BrushOut = FEditorStyle::GetBrush(TEXT("GraphEditor.EnumGlyph"));
@@ -674,8 +679,8 @@ public:
 			UProperty* VariableProp = StaticCastSharedPtr<FEdGraphSchemaAction_K2Var>(PaletteAction)->GetProperty();
 			UObjectProperty* VariableObjProp = Cast<UObjectProperty>(VariableProp);
 
-			UClass* VarSourceClass = (VariableProp != NULL) ? CastChecked<UClass>(VariableProp->GetOuter()) : NULL;
-			const bool bIsBlueprintVariable = (VarSourceClass == BlueprintEditorPtr.Pin()->GetBlueprintObj()->SkeletonGeneratedClass);
+			UStruct* VarSourceScope = (VariableProp != NULL) ? CastChecked<UStruct>(VariableProp->GetOuter()) : NULL;
+			const bool bIsBlueprintVariable = (VarSourceScope == BlueprintEditorPtr.Pin()->GetBlueprintObj()->SkeletonGeneratedClass);
 			const bool bIsComponentVar = (VariableObjProp != NULL) && (VariableObjProp->PropertyClass != NULL) && (VariableObjProp->PropertyClass->IsChildOf(UActorComponent::StaticClass()));
 			bShouldHaveAVisibilityToggle = bIsBlueprintVariable && !bIsComponentVar;
 		}
@@ -1018,6 +1023,11 @@ bool SBlueprintPaletteItem::OnNameTextVerifyChanged(const FText& InNewText, FTex
 		FEdGraphSchemaAction_K2Var* VarAction = (FEdGraphSchemaAction_K2Var*)ActionPtr.Pin().Get();
 		OriginalName = (VarAction->GetVariableName());
 	}
+	else if (ActionPtr.Pin()->GetTypeId() == FEdGraphSchemaAction_K2LocalVar::StaticGetTypeId())
+	{
+		FEdGraphSchemaAction_K2LocalVar* LocalVarAction = (FEdGraphSchemaAction_K2LocalVar*)ActionPtr.Pin().Get();
+		OriginalName = (LocalVarAction->GetVariableName());
+	}
 	else
 	{
 		UEdGraph* Graph = NULL;
@@ -1109,6 +1119,19 @@ void SBlueprintPaletteItem::OnNameTextCommitted(const FText& NewText, ETextCommi
 				const FScopedTransaction Transaction( LOCTEXT( "Rename Function", "Rename Function" ) );
 				FString NewNameString = NewText.ToString();
 				FBlueprintEditorUtils::RenameGraph(Graph, *NewNameString );
+
+				// Search through all function entry nodes for local variables to update their scope name
+				TArray<UK2Node_Variable*> VariableNodes;
+				FBlueprintEditorUtils::GetAllNodesOfClass<UK2Node_Variable>(FBlueprintEditorUtils::FindBlueprintForGraph(Graph), VariableNodes);
+
+				for (UK2Node_Variable* const VariableNode : VariableNodes)
+				{
+					if(VariableNode->VariableReference.IsLocalScope())
+					{
+						// Update the variable's scope to be the graph's name (which mirrors the UFunction)
+						VariableNode->VariableReference.SetLocalMember(VariableNode->VariableReference.GetMemberName(), Graph->GetName(), VariableNode->VariableReference.GetMemberGuid());
+					}
+				}
 			}
 		}
 	}
@@ -1198,6 +1221,23 @@ void SBlueprintPaletteItem::OnNameTextCommitted(const FText& NewText, ETextCommi
 			FBlueprintEditorUtils::RenameMemberVariable(BlueprintEditorPtr.Pin()->GetBlueprintObj(), VarAction->GetVariableName(), NewVarName);
 		}
 	}
+	else if (ActionPtr.Pin()->GetTypeId() == FEdGraphSchemaAction_K2LocalVar::StaticGetTypeId())
+	{
+		FEdGraphSchemaAction_K2LocalVar* LocalVarAction = (FEdGraphSchemaAction_K2LocalVar*)ActionPtr.Pin().Get();
+
+		// Check if the name is unchanged
+		if(NewText.ToString() == LocalVarAction->GetVariableName().ToString())
+		{
+			return;
+		}
+
+		const FScopedTransaction Transaction( LOCTEXT( "RenameVariable", "Rename Variable" ) );
+
+		BlueprintEditorPtr.Pin()->GetBlueprintObj()->Modify();
+
+		FName NewVarName = FName(*NewText.ToString());
+		FBlueprintEditorUtils::RenameLocalVariable(BlueprintEditorPtr.Pin()->GetBlueprintObj(), LocalVarAction->GetVariableName(), NewVarName);
+	}
 	BlueprintEditorPtr.Pin()->GetMyBlueprintWidget()->SelectItemByName(FName(*NewText.ToString()), ESelectInfo::OnMouseClick);
 }
 
@@ -1275,6 +1315,23 @@ FText SBlueprintPaletteItem::GetToolTipText() const
 				ToolTipText = !Result.IsEmpty() ? Result : GetVarType(VarClass, VarAction->GetVariableName(), true, true);
 			}
 		}
+		else if (PaletteAction->GetTypeId() == FEdGraphSchemaAction_K2LocalVar::StaticGetTypeId())
+		{
+			FEdGraphSchemaAction_K2LocalVar* LocalVarAction = (FEdGraphSchemaAction_K2LocalVar*)PaletteAction.Get();
+			UClass* VarClass = CastChecked<UClass>(LocalVarAction->GetVariableScope()->GetOuter());
+			if (bShowClassInTooltip && (VarClass != NULL))
+			{
+				UBlueprint* Blueprint = UBlueprint::GetBlueprintFromClass(VarClass);
+				ClassDisplayName = (Blueprint != NULL) ? Blueprint->GetName() : VarClass->GetName();
+			}
+			else
+			{
+				FString Result;
+				FBlueprintEditorUtils::GetBlueprintVariableMetaData(BlueprintEditorPtr.Pin()->GetBlueprintObj(), LocalVarAction->GetVariableName(), TEXT("tooltip"), Result);
+				// Only use the variable tooltip if it has been filled out.
+				ToolTipText = !Result.IsEmpty() ? Result : GetVarType(LocalVarAction->GetVariableScope(), LocalVarAction->GetVariableName(), true, true);
+			}
+		}
 		else if (PaletteAction->GetTypeId() == FEdGraphSchemaAction_K2Delegate::StaticGetTypeId())
 		{
 			FEdGraphSchemaAction_K2Delegate* DelegateAction = (FEdGraphSchemaAction_K2Delegate*)PaletteAction.Get();
@@ -1359,6 +1416,12 @@ TSharedPtr<SToolTip> SBlueprintPaletteItem::ConstructToolTipWidget() const
 			DocLink = DefaultComment->GetDocumentationLink();
 			DocExcerptName = DefaultComment->GetDocumentationExcerptName();
 		}
+		else if (PaletteAction->GetTypeId() == FEdGraphSchemaAction_K2LocalVar::StaticGetTypeId())
+		{
+			// Don't show big tooltip if we are showing class as well (means we are not in MyBlueprint)
+			DocLink = TEXT("Shared/Editors/BlueprintEditor/GraphTypes");
+			DocExcerptName = TEXT("LocalVariable");
+		}
 	}
 
 	// Setup the attribute for dynamically pulling the tooltip
diff --git a/Engine/Source/Editor/Kismet/Private/SKismetInspector.cpp b/Engine/Source/Editor/Kismet/Private/SKismetInspector.cpp
--- a/Engine/Source/Editor/Kismet/Private/SKismetInspector.cpp
+++ b/Engine/Source/Editor/Kismet/Private/SKismetInspector.cpp
@@ -226,9 +226,6 @@ void SKismetInspector::Construct(const FArguments& InArgs)
 		PropertyView->RegisterInstancedCustomPropertyLayout(UK2Node_VariableGet::StaticClass(), LayoutVariableDetails);
 		PropertyView->RegisterInstancedCustomPropertyLayout(UK2Node_VariableSet::StaticClass(), LayoutVariableDetails);
 
-		FOnGetDetailCustomizationInstance LayoutLocalVariableDetails = FOnGetDetailCustomizationInstance::CreateStatic(&FBlueprintLocalVarActionDetails::MakeInstance, TWeakPtr<SMyBlueprint>(Kismet2->GetMyBlueprintWidget()));
-		PropertyView->RegisterInstancedCustomPropertyLayout(UK2Node_LocalVariable::StaticClass(), LayoutLocalVariableDetails);
-
 		FOnGetDetailCustomizationInstance LayoutOptionDetails = FOnGetDetailCustomizationInstance::CreateStatic(&FBlueprintGlobalOptionsDetails::MakeInstance, Kismet2Ptr);
 		PropertyView->RegisterInstancedCustomPropertyLayout(UBlueprint::StaticClass(), LayoutOptionDetails);
 
diff --git a/Engine/Source/Editor/Kismet/Private/SMyBlueprint.cpp b/Engine/Source/Editor/Kismet/Private/SMyBlueprint.cpp
--- a/Engine/Source/Editor/Kismet/Private/SMyBlueprint.cpp
+++ b/Engine/Source/Editor/Kismet/Private/SMyBlueprint.cpp
@@ -547,7 +547,9 @@ void SMyBlueprint::GetChildEvents(UEdGraph const* EdGraph, int32 const SectionId
 
 void SMyBlueprint::GetLocalVariables(FGraphActionListBuilderBase& OutAllActions) const
 {
-	if(UEdGraph* EdGraph = BlueprintEditorPtr.Pin()->GetFocusedGraph())
+	// We want to pull local variables from the top level function graphs
+	UEdGraph* EdGraph = FBlueprintEditorUtils::GetTopLevelGraph(BlueprintEditorPtr.Pin()->GetFocusedGraph());
+	if( EdGraph )
 	{
 		// grab the parent graph's name
 		UEdGraphSchema const* Schema = EdGraph->GetSchema();
@@ -555,19 +557,26 @@ void SMyBlueprint::GetLocalVariables(FGraphActionListBuilderBase& OutAllActions)
 		Schema->GetGraphDisplayInformation(*EdGraph, EdGraphDisplayInfo);
 		FText const EdGraphDisplayName = EdGraphDisplayInfo.DisplayName;
 
-		TArray<UK2Node_LocalVariable*> LocalVariableNodes;
-		EdGraph->GetNodesOfClass<UK2Node_LocalVariable>(LocalVariableNodes);
+		TArray<UK2Node_FunctionEntry*> FunctionEntryNodes;
+		EdGraph->GetNodesOfClass<UK2Node_FunctionEntry>(FunctionEntryNodes);
 
+		// Search in all FunctionEntry nodes for their local variables
 		FString ActionCategory;
-		for (UK2Node_LocalVariable* const LocalVariable : LocalVariableNodes)
+		for (UK2Node_FunctionEntry* const FunctionEntry : FunctionEntryNodes)
 		{
-			FString const Tooltip = LocalVariable->GetTooltip();
-			FText const Description = LocalVariable->GetNodeTitle(ENodeTitleType::EditableTitle);
+			for( const FBPVariableDescription& Variable : FunctionEntry->LocalVariables )
+			{
+				FString Category = Variable.Category.ToString();
+				if (Variable.Category == GetDefault<UEdGraphSchema_K2>()->VR_DefaultCategory)
+				{
+					Category = FString();
+				}
 
-			TSharedPtr<FEdGraphSchemaAction_K2TargetNode> TargetNodeAction = MakeShareable(new FEdGraphSchemaAction_K2TargetNode(ActionCategory, Description, Tooltip, 0));
-			TargetNodeAction->NodeTemplate = LocalVariable;
-			TargetNodeAction->SearchTitle = TargetNodeAction->NodeTemplate->GetNodeSearchTitle();
-			OutAllActions.AddAction(TargetNodeAction);
+				TSharedPtr<FEdGraphSchemaAction_K2LocalVar> NewVarAction = MakeShareable(new FEdGraphSchemaAction_K2LocalVar(Category, FText::FromName(Variable.VarName), TEXT(""), 0));
+				UFunction* Func = FindField<UFunction>(BlueprintEditorPtr.Pin()->GetBlueprintObj()->SkeletonGeneratedClass, *EdGraph->GetName());
+				NewVarAction->SetVariableInfo(Variable.VarName, Func);
+				OutAllActions.AddAction(NewVarAction);
+			}
 		}
 	}
 }
@@ -899,6 +908,22 @@ FReply SMyBlueprint::OnActionDragged( const TArray< TSharedPtr<FEdGraphSchemaAct
 				return FReply::Handled().BeginDragDrop(DragOperation);
 			}
 		}
+		else if( InAction->GetTypeId() == FEdGraphSchemaAction_K2LocalVar::StaticGetTypeId())
+		{
+			FEdGraphSchemaAction_K2LocalVar* VarAction = (FEdGraphSchemaAction_K2LocalVar*)InAction.Get();
+
+			UStruct* VariableScope = VarAction->GetVariableScope();
+			if(VariableScope != NULL)
+			{
+				const bool bIsAltDown = MouseEvent.IsAltDown();
+				const bool bIsCtrlDown = MouseEvent.IsLeftControlDown() || MouseEvent.IsRightControlDown();
+
+				TSharedRef<FKismetVariableDragDropAction> DragOperation = FKismetVariableDragDropAction::New(VarAction->GetVariableName(), VariableScope, AnalyticsDelegate);
+				DragOperation->SetAltDrag(bIsAltDown);
+				DragOperation->SetCtrlDrag(bIsCtrlDown);
+				return FReply::Handled().BeginDragDrop(DragOperation);
+			}
+		}
 		else if(InAction->GetTypeId() == FEdGraphSchemaAction_K2Var::StaticGetTypeId())
 		{
 			FEdGraphSchemaAction_K2Var* VarAction = (FEdGraphSchemaAction_K2Var*)InAction.Get();
@@ -997,6 +1022,15 @@ void SMyBlueprint::OnActionSelected( const TArray< TSharedPtr<FEdGraphSchemaActi
 
 			Inspector->ShowDetailsForSingleObject(VarAction->GetProperty(), Options);
 		}
+		else if (InAction->GetTypeId() == FEdGraphSchemaAction_K2LocalVar::StaticGetTypeId())
+		{
+			FEdGraphSchemaAction_K2LocalVar* VarAction = (FEdGraphSchemaAction_K2LocalVar*)InAction.Get();
+
+			SKismetInspector::FShowDetailsOptions Options(VarAction->GetVariableName().ToString());
+			Options.bHideFilterArea = true;
+
+			Inspector->ShowDetailsForSingleObject(VarAction->GetProperty(), Options);
+		}
 		else if (InAction->GetTypeId() == FEdGraphSchemaAction_K2Enum::StaticGetTypeId())
 		{
 			FEdGraphSchemaAction_K2Enum* EnumAction = (FEdGraphSchemaAction_K2Enum*)InAction.Get();
@@ -1128,19 +1162,12 @@ FEdGraphSchemaAction_K2Var* SMyBlueprint::SelectionAsVar() const
 	return SelectionAsType<FEdGraphSchemaAction_K2Var>( GraphActionMenu );
 }
 
-UK2Node_LocalVariable* SMyBlueprint::SelectionAsLocalVar() const
+FEdGraphSchemaAction_K2LocalVar* SMyBlueprint::SelectionAsLocalVar() const
 {
-	FEdGraphSchemaAction_K2TargetNode* TargetNode = NULL;
-	if (LocalGraphActionMenu.IsValid())
+	if(LocalGraphActionMenu.IsValid())
 	{
-		TargetNode = SelectionAsType<FEdGraphSchemaAction_K2TargetNode>( LocalGraphActionMenu );
+		return SelectionAsType<FEdGraphSchemaAction_K2LocalVar>( LocalGraphActionMenu );
 	}
-
-	if(TargetNode && TargetNode->NodeTemplate)
-	{
-		return Cast<UK2Node_LocalVariable>(TargetNode->NodeTemplate);
-	}
-
 	return NULL;
 }
 
@@ -1363,6 +1390,10 @@ void SMyBlueprint::OnFindEntry()
 	{
 		SearchTerm = VarAction->GetVariableName().ToString();
 	}
+	else if (FEdGraphSchemaAction_K2LocalVar* LocalVarAction = SelectionAsLocalVar())
+	{
+		SearchTerm = LocalVarAction->GetVariableName().ToString();
+	}
 	else if (FEdGraphSchemaAction_K2Delegate* DelegateAction = SelectionAsDelegate())
 	{
 		SearchTerm = DelegateAction->GetDelegateName().ToString();
@@ -1483,20 +1514,44 @@ void SMyBlueprint::OnDeleteEntry()
 		GetBlueprintObj()->Modify();
 		FBlueprintEditorUtils::RemoveMemberVariable(GetBlueprintObj(), VarAction->GetVariableName());
 	}
-	else if (FEdGraphSchemaAction_K2Event* EventAction = SelectionAsEvent())
+	else if ( FEdGraphSchemaAction_K2LocalVar* LocalVarAction = SelectionAsLocalVar() )
 	{
-		const FScopedTransaction Transaction(LOCTEXT( "RemoveEventNode", "Remove EventNode"));
+		if(FBlueprintEditorUtils::IsVariableUsed(GetBlueprintObj(), LocalVarAction->GetVariableName()))
+		{
+			FText ConfirmDelete = FText::Format(LOCTEXT( "ConfirmDeleteLocalVariableInUse",
+				"Local Variable {0} is in use! Do you really want to delete it?"),
+				FText::FromName( LocalVarAction->GetVariableName() ) );
+
+			// Warn the user that this may result in data loss
+			FSuppressableWarningDialog::FSetupInfo Info( ConfirmDelete, LOCTEXT("DeleteVar", "Delete Variable"), "DeleteVariableInUse_Warning" );
+			Info.ConfirmText = LOCTEXT( "DeleteVariable_Yes", "Yes");
+			Info.CancelText = LOCTEXT( "DeleteVariable_No", "No");	
+
+			FSuppressableWarningDialog DeleteVariableInUse( Info );
+			if ( DeleteVariableInUse.ShowModal() == FSuppressableWarningDialog::Cancel )
+			{
+				return;
+			}
+		}
+
+		const FScopedTransaction Transaction( LOCTEXT( "RemoveLocalVariable", "Remove Local Variable" ) );
 
 		GetBlueprintObj()->Modify();
-		FBlueprintEditorUtils::RemoveNode(GetBlueprintObj(), EventAction->NodeTemplate);
+
+		UEdGraph* FunctionGraph = FBlueprintEditorUtils::GetTopLevelGraph(BlueprintEditorPtr.Pin()->GetFocusedGraph());
+		TArray<UK2Node_FunctionEntry*> FunctionEntryNodes;
+		FunctionGraph->GetNodesOfClass<UK2Node_FunctionEntry>(FunctionEntryNodes);
+		check(FunctionEntryNodes.Num() == 1);
+		FunctionEntryNodes[0]->Modify();
+
+		FBlueprintEditorUtils::RemoveLocalVariable(GetBlueprintObj(), LocalVarAction->GetVariableName());
 	}
-	else if (UK2Node_LocalVariable* LocalVariable = SelectionAsLocalVar())
+	else if (FEdGraphSchemaAction_K2Event* EventAction = SelectionAsEvent())
 	{
-		const FScopedTransaction Transaction( LOCTEXT( "RemoveLocalVariable", "Remove Local Variable" ) );
-
-		LocalVariable->GetGraph()->Modify();
+		const FScopedTransaction Transaction(LOCTEXT( "RemoveEventNode", "Remove EventNode"));
 
-		LocalVariable->DestroyNode();
+		GetBlueprintObj()->Modify();
+		FBlueprintEditorUtils::RemoveNode(GetBlueprintObj(), EventAction->NodeTemplate);
 	}
 	else if ( SelectionIsCategory() )
 	{
@@ -1564,9 +1619,9 @@ bool SMyBlueprint::CanDeleteEntry() const
 	{
 		return EventAction->NodeTemplate != NULL;
 	}
-	else if (UK2Node_LocalVariable* LocalVariable = SelectionAsLocalVar())
+	else if (FEdGraphSchemaAction_K2LocalVar* LocalVariable = SelectionAsLocalVar())
 	{
-		return LocalVariable != NULL;
+		return true;
 	}
 	else if (SelectionIsCategory())
 	{
diff --git a/Engine/Source/Editor/Kismet/Private/SMyBlueprint.h b/Engine/Source/Editor/Kismet/Private/SMyBlueprint.h
--- a/Engine/Source/Editor/Kismet/Private/SMyBlueprint.h
+++ b/Engine/Source/Editor/Kismet/Private/SMyBlueprint.h
@@ -79,7 +79,7 @@ public:
 	FEdGraphSchemaAction_K2Event* SelectionAsEvent() const;
 
 	/** Accessor for getting the current selection as a K2 local var */
-	UK2Node_LocalVariable* SelectionAsLocalVar() const;
+	FEdGraphSchemaAction_K2LocalVar* SelectionAsLocalVar() const;
 
 	/** Accessor for determining if the current selection is a category*/
 	bool SelectionIsCategory() const;
diff --git a/Engine/Source/Editor/Kismet/Public/BlueprintEditor.h b/Engine/Source/Editor/Kismet/Public/BlueprintEditor.h
--- a/Engine/Source/Editor/Kismet/Public/BlueprintEditor.h
+++ b/Engine/Source/Editor/Kismet/Public/BlueprintEditor.h
@@ -286,8 +286,15 @@ public:
 	 */
 	static const FSlateBrush* GetGlyphForGraph(const UEdGraph* Graph, bool bInLargeIcon = false);
 
-	
-	static FSlateBrush const* GetVarIconAndColor(UClass* VarClass, FName VarName, FSlateColor& IconColorOut);
+	/**
+	 * Util for finding a glyph and color for a variable.
+	 *
+	 * @param VarScope			Scope to find the variable in
+	 * @param VarName			Name of variable
+	 * @param IconColorOut		The resulting color for the glyph
+	 * @return					The resulting glyph brush
+	 */
+	static FSlateBrush const* GetVarIconAndColor(UStruct* VarScope, FName VarName, FSlateColor& IconColorOut);
 
 	/** Overridable function for determining if the current mode can script */
 	virtual bool IsInAScriptingMode() const;
diff --git a/Engine/Source/Editor/KismetCompiler/Private/KismetCompiler.cpp b/Engine/Source/Editor/KismetCompiler/Private/KismetCompiler.cpp
--- a/Engine/Source/Editor/KismetCompiler/Private/KismetCompiler.cpp
+++ b/Engine/Source/Editor/KismetCompiler/Private/KismetCompiler.cpp
@@ -659,20 +659,26 @@ void FKismetCompilerContext::CreatePropertiesFromList(UStruct* Scope, UField**&
 	}
 }
 
-/** Creates the properties on a function that store the function parameters, results, and local variables */
 void FKismetCompilerContext::CreateLocalVariablesForFunction(FKismetFunctionContext& Context)
 {
 	// Local stack frame (or maybe class for the ubergraph)
 	{
 		const bool bArePropertiesLocal = true;
 
-		check(Context.Function->Children == NULL);
-		UField** PropertyStorageLocation = &(Context.Function->Children);
+		// Pull the local properties generated out of the function, they will be put at the end of the list
+		UField* LocalProperties = Context.Function->Children;
 
+		UField** PropertyStorageLocation = &Context.Function->Children;
 		CreatePropertiesFromList(Context.Function, PropertyStorageLocation, Context.Parameters, CPF_Parm, bArePropertiesLocal, /*bPropertiesAreParameters=*/ true);
 		CreatePropertiesFromList(Context.Function, PropertyStorageLocation, Context.Results, CPF_Parm | CPF_OutParm, bArePropertiesLocal, /*bPropertiesAreParameters=*/ true);
 		CreatePropertiesFromList(Context.Function, PropertyStorageLocation, Context.Locals, 0, bArePropertiesLocal, /*bPropertiesAreParameters=*/ true);
 
+		// If there were local properties, place them at the end of the property storage location
+		if(LocalProperties)
+		{
+			*PropertyStorageLocation = LocalProperties;
+		}
+
 		// Create debug data for variable reads/writes
 		if (Context.bCreateDebugData)
 		{
@@ -726,6 +732,39 @@ void FKismetCompilerContext::CreateLocalVariablesForFunction(FKismetFunctionCont
 	}
 }
 
+void FKismetCompilerContext::CreateUserDefinedLocalVariablesForFunction(FKismetFunctionContext& Context, UField**& PropertyStorageLocation)
+{
+	check(Context.Function->Children == NULL);
+
+	// Create local variables from the Context entry point
+	for (int32 i = 0; i < Context.EntryPoint->LocalVariables.Num(); ++i)
+	{
+		FBPVariableDescription& Variable = Context.EntryPoint->LocalVariables[Context.EntryPoint->LocalVariables.Num() - (i + 1)];
+
+		// Create the property based on the variable description, scoped to the function
+		UProperty* NewProperty = FKismetCompilerUtilities::CreatePropertyOnScope(Context.Function, Variable.VarName, Variable.VarType, NewClass, 0, Schema, MessageLog);
+		if (NewProperty != NULL)
+		{
+			// Link this object to the tail of the list (so properties remain in the desired order)
+			*PropertyStorageLocation = NewProperty;
+			PropertyStorageLocation = &(NewProperty->Next);
+		}
+
+		if (NewProperty != NULL)
+		{
+			NewProperty->SetPropertyFlags(Variable.PropertyFlags);
+			NewProperty->SetMetaData(TEXT("FriendlyName"), *Variable.FriendlyName);
+			NewProperty->SetMetaData(TEXT("Category"), *Variable.Category.ToString());
+			NewProperty->RepNotifyFunc = Variable.RepNotifyFunc;
+
+			if(!Variable.DefaultValue.IsEmpty())
+			{
+				SetPropertyDefaultValue(NewProperty, Variable.DefaultValue);
+			}
+		}
+	}
+}
+
 void FKismetCompilerContext::SetPropertyDefaultValue(const UProperty* PropertyToSet, FString& Value)
 {
 	DefaultPropertyValueMap.Add(PropertyToSet->GetFName(), Value);
@@ -1091,6 +1130,10 @@ void FKismetCompilerContext::PrecompileFunction(FKismetFunctionContext& Context)
 		// Find the execution path (and make sure it has no cycles)
 		CreateExecutionSchedule(Context.SourceGraph->Nodes, Context.LinearExecutionList);
 
+		// Create any user defined variables, this must occur before registering nets so that the properties are in place
+		UField** PropertyStorageLocation = &(Context.Function->Children);
+		CreateUserDefinedLocalVariablesForFunction(Context, PropertyStorageLocation);
+
 		bool bIsPureFunction = true;
 		for (int32 NodeIndex = 0; NodeIndex < Context.LinearExecutionList.Num(); ++NodeIndex)
 		{
diff --git a/Engine/Source/Editor/KismetCompiler/Private/KismetCompilerMisc.cpp b/Engine/Source/Editor/KismetCompiler/Private/KismetCompilerMisc.cpp
--- a/Engine/Source/Editor/KismetCompiler/Private/KismetCompilerMisc.cpp
+++ b/Engine/Source/Editor/KismetCompiler/Private/KismetCompilerMisc.cpp
@@ -715,6 +715,12 @@ void FNodeHandlingFunctor::ResolveAndRegisterScopedTerm(FKismetFunctionContext&
 			Term->bIsConst = true;
 		}
 
+		// Check if the property is a local variable and mark it so
+		if( SearchScope == Context.Function && BoundProperty->GetOuter() == Context.Function)
+		{
+			Term->bIsLocal = true;
+		}
+
 		// Resolve the context term
 		if (SelfPin != NULL)
 		{
diff --git a/Engine/Source/Editor/KismetCompiler/Public/KismetCompiler.h b/Engine/Source/Editor/KismetCompiler/Public/KismetCompiler.h
--- a/Engine/Source/Editor/KismetCompiler/Public/KismetCompiler.h
+++ b/Engine/Source/Editor/KismetCompiler/Public/KismetCompiler.h
@@ -203,6 +203,9 @@ protected:
 	/** Creates the properties on a function that store the function parameters, results, and local variables */
 	void CreateLocalVariablesForFunction(FKismetFunctionContext& Context);
 
+	/** Creates user defined local variables for function */
+	void CreateUserDefinedLocalVariablesForFunction(FKismetFunctionContext& Context, UField**& PropertyStorageLocation);
+
 	/** Adds a default value entry into the DefaultPropertyValueMap for the property specified */
 	void SetPropertyDefaultValue(const UProperty* PropertyToSet, FString& Value);
 
diff --git a/Engine/Source/Editor/UnrealEd/Private/Kismet2/BlueprintEditorUtils.cpp b/Engine/Source/Editor/UnrealEd/Private/Kismet2/BlueprintEditorUtils.cpp
--- a/Engine/Source/Editor/UnrealEd/Private/Kismet2/BlueprintEditorUtils.cpp
+++ b/Engine/Source/Editor/UnrealEd/Private/Kismet2/BlueprintEditorUtils.cpp
@@ -212,7 +212,12 @@ static void RenameVariableReferences(UBlueprint* Blueprint, UClass* VariableClas
 				if(NodeRefClass && NodeRefClass->IsChildOf(VariableClass) && OldVarName == VariableNode->GetVarName())
 				{
 					VariableNode->Modify();
-					if(VariableNode->VariableReference.IsSelfContext())
+
+					if(VariableNode->VariableReference.IsLocalScope())
+					{
+						VariableNode->VariableReference.SetLocalMember(NewVarName, VariableNode->VariableReference.GetMemberScope(NodeRefClass), VariableNode->VariableReference.GetMemberGuid());
+					}
+					else if(VariableNode->VariableReference.IsSelfContext())
 					{
 						VariableNode->VariableReference.SetSelfMember(NewVarName);
 					}
@@ -1885,7 +1890,7 @@ void FBlueprintEditorUtils::RemoveNode(UBlueprint* Blueprint, UEdGraphNode* Node
 	}
 }
 
-UEdGraph* FBlueprintEditorUtils::GetTopLevelGraph(UEdGraph* InGraph)
+UEdGraph* FBlueprintEditorUtils::GetTopLevelGraph(const UEdGraph* InGraph)
 {
 	UEdGraph* GraphToTest = const_cast<UEdGraph*>(InGraph);
 
@@ -2205,8 +2210,12 @@ bool FBlueprintEditorUtils::DoesBlueprintContainField(const UBlueprint* Blueprin
 	// Get the class of the field
 	if(TestField)
 	{
-		UClass* TestClass = CastChecked<UClass>(TestField->GetOuter());
-		return FBlueprintEditorUtils::DoesBlueprintDeriveFrom(Blueprint, TestClass);
+		// Local properties do not have a UClass outer but are also not a part of the Blueprint
+		UClass* TestClass = Cast<UClass>(TestField->GetOuter());
+		if(TestClass)
+		{
+			return FBlueprintEditorUtils::DoesBlueprintDeriveFrom(Blueprint, TestClass);
+		}
 	}
 	return false;
 }
@@ -2542,7 +2551,15 @@ void FBlueprintEditorUtils::SetBlueprintVariableMetaData(UBlueprint* Blueprint,
 			//Not a Timeline is this a SCS Node?
 			const int32 SCS_NodeIndex = FBlueprintEditorUtils::FindSCS_Node(Blueprint,VarName);
 
-			if (SCS_NodeIndex != INDEX_NONE)
+			if (SCS_NodeIndex == INDEX_NONE)
+			{
+				// Check to see if it is a local variable
+				if(FBPVariableDescription* LocalVariable = FindLocalVariable(Blueprint, VarName))
+				{
+					LocalVariable->SetMetaData(MetaDataKey, *MetaDataValue);
+				}
+			}
+			else
 			{
 				Blueprint->SimpleConstructionScript->GetAllNodes()[SCS_NodeIndex]->SetMetaData(MetaDataKey, MetaDataValue);
 			}
@@ -2581,7 +2598,20 @@ bool FBlueprintEditorUtils::GetBlueprintVariableMetaData(UBlueprint* Blueprint,
 			//Not a Timeline is this a SCS Node?
 			const int32 SCS_NodeIndex = FBlueprintEditorUtils::FindSCS_Node(Blueprint,VarName);
 
-			if (SCS_NodeIndex != INDEX_NONE)
+			if (SCS_NodeIndex == INDEX_NONE)
+			{
+				// Check to see if it is a local variable
+				if(FBPVariableDescription* LocalVariable = FindLocalVariable(Blueprint, VarName))
+				{
+					int32 EntryIndex = LocalVariable->FindMetaDataEntryIndexForKey(MetaDataKey);
+					if (EntryIndex != INDEX_NONE)
+					{
+						OutMetaDataValue = LocalVariable->GetMetaData(MetaDataKey);
+						return true;
+					}
+				}
+			}
+			else
 			{
 				USCS_Node& Desc = *Blueprint->SimpleConstructionScript->GetAllNodes()[SCS_NodeIndex];
 
@@ -2686,6 +2716,20 @@ void FBlueprintEditorUtils::SetBlueprintVariableCategory(UBlueprint* Blueprint,
 				}
 			}
 
+			if (bDontRecompile == false)
+			{
+				FBlueprintEditorUtils::MarkBlueprintAsStructurallyModified(Blueprint);
+			}
+		}
+	}
+	else
+	{
+		// Check to see if it is a local variable
+		if(FBPVariableDescription* LocalVariable = FindLocalVariable(Blueprint, VarName))
+		{
+			LocalVariable->SetMetaData(TEXT("Category"), *SetCategory.ToString());
+			LocalVariable->Category = SetCategory;
+
 			if (bDontRecompile == false)
 			{
 				FBlueprintEditorUtils::MarkBlueprintAsStructurallyModified(Blueprint);
@@ -2703,6 +2747,14 @@ FName FBlueprintEditorUtils::GetBlueprintVariableCategory(UBlueprint* Blueprint,
 	{
 		CategoryName = FObjectEditorUtils::GetCategoryFName(TargetProperty);
 	}
+	else
+	{
+		// Check to see if it is a local variable
+		if(FBPVariableDescription* LocalVariable = FindLocalVariable(Blueprint, VarName))
+		{
+			return LocalVariable->Category;
+		}
+	}
 
 	if(CategoryName == NAME_None && Blueprint->SimpleConstructionScript != NULL)
 	{
@@ -3198,6 +3250,162 @@ void FBlueprintEditorUtils::ChangeMemberVariableType(UBlueprint* Blueprint, cons
 	}
 }
 
+void FBlueprintEditorUtils::RemoveLocalVariable(UBlueprint* InBlueprint, const FName& InVarName)
+{
+	TArray<UK2Node_FunctionEntry*> FunctionEntryNodes;
+	GetAllNodesOfClass<UK2Node_FunctionEntry>(InBlueprint, FunctionEntryNodes);
+
+	for (UK2Node_FunctionEntry* const FunctionEntry : FunctionEntryNodes)
+	{
+		for( int32 VarIdx = 0; VarIdx < FunctionEntry->LocalVariables.Num(); ++VarIdx )
+		{
+			FunctionEntry->LocalVariables.RemoveAt(VarIdx);
+			FBlueprintEditorUtils::RemoveVariableNodes(InBlueprint, InVarName);
+			FBlueprintEditorUtils::MarkBlueprintAsStructurallyModified(InBlueprint);
+		}
+	}
+}
+
+void FBlueprintEditorUtils::RenameLocalVariable(UBlueprint* InBlueprint, const FName& InOldName, const FName& InNewName)
+{
+	if ((InOldName != InNewName) && (InNewName != NAME_None))
+	{
+		FBPVariableDescription* LocalVariable = FindLocalVariable(InBlueprint, InOldName);
+		
+		if(LocalVariable)
+		{
+			const FScopedTransaction Transaction( LOCTEXT("RenameVariable", "Rename Local Variable") );
+			InBlueprint->Modify();
+
+			// Update the name
+			const UEdGraphSchema_K2* K2Schema = GetDefault<UEdGraphSchema_K2>();
+			LocalVariable->VarName = InNewName;
+			LocalVariable->FriendlyName = FName::NameToDisplayString( InNewName.ToString(), (LocalVariable->VarType.PinCategory == K2Schema->PC_Boolean) ? true : false );
+
+			// Update any existing references to the old name
+			FBlueprintEditorUtils::ReplaceVariableReferences(InBlueprint, InOldName, InNewName);
+
+			// Validate child blueprints and adjust variable names to avoid a potential name collision
+			FBlueprintEditorUtils::ValidateBlueprintChildVariables(InBlueprint, InNewName);
+			
+			// And recompile
+			FBlueprintEditorUtils::MarkBlueprintAsStructurallyModified(InBlueprint);
+		}
+	}
+}
+
+FBPVariableDescription* FBlueprintEditorUtils::FindLocalVariable(UBlueprint* InBlueprint, const FName& InVariableName)
+{
+	// Search through all function entry nodes for a local variable with the passed name
+	TArray<UK2Node_FunctionEntry*> FunctionEntryNodes;
+	GetAllNodesOfClass<UK2Node_FunctionEntry>(InBlueprint, FunctionEntryNodes);
+
+	FBPVariableDescription* ReturnVariable = NULL;
+
+	for (UK2Node_FunctionEntry* const FunctionEntry : FunctionEntryNodes)
+	{
+		for( FBPVariableDescription& Variable : FunctionEntry->LocalVariables )
+		{
+			if(InVariableName == Variable.VarName)
+			{
+				ReturnVariable = &Variable;
+				break;
+			}
+		}
+
+		if(ReturnVariable)
+		{
+			break;
+		}
+	}
+
+	return ReturnVariable;
+}
+
+FName FBlueprintEditorUtils::FindLocalVariableNameByGuid(UBlueprint* InBlueprint, const FGuid& InVariableGuid)
+{
+	// Search through all function entry nodes for a local variable with the passed Guid
+	TArray<UK2Node_FunctionEntry*> FunctionEntryNodes;
+	GetAllNodesOfClass<UK2Node_FunctionEntry>(InBlueprint, FunctionEntryNodes);
+
+	FName ReturnVariableName = NAME_None;
+	for (UK2Node_FunctionEntry* const FunctionEntry : FunctionEntryNodes)
+	{
+		for( FBPVariableDescription& Variable : FunctionEntry->LocalVariables )
+		{
+			if(Variable.VarGuid == InVariableGuid)
+			{
+				ReturnVariableName = Variable.VarName;
+				break;
+			}
+		}
+
+		if(ReturnVariableName != NAME_None)
+		{
+			break;
+		}
+	}
+
+	return ReturnVariableName;
+}
+
+FGuid FBlueprintEditorUtils::FindLocalVariableGuidByName(UBlueprint* InBlueprint, const FName InVariableName)
+{
+	// Search through all function entry nodes for a local variable with the passed name
+	TArray<UK2Node_FunctionEntry*> FunctionEntryNodes;
+	GetAllNodesOfClass<UK2Node_FunctionEntry>(InBlueprint, FunctionEntryNodes);
+
+	FGuid ReturnGuid;
+	for (UK2Node_FunctionEntry* const FunctionEntry : FunctionEntryNodes)
+	{
+		for( FBPVariableDescription& Variable : FunctionEntry->LocalVariables )
+		{
+			if(Variable.VarName == InVariableName)
+			{
+				ReturnGuid = Variable.VarGuid;
+				break;
+			}
+		}
+
+		if(ReturnGuid.IsValid())
+		{
+			break;
+		}
+	}
+
+	return ReturnGuid;
+}
+
+void FBlueprintEditorUtils::ChangeLocalVariableType(UBlueprint* InBlueprint, const FName& InVariableName, const FEdGraphPinType& NewPinType)
+{
+	if (InVariableName != NAME_None)
+	{
+		FString ActionCategory;
+		const UEdGraphSchema_K2* K2Schema = GetDefault<UEdGraphSchema_K2>();
+		FBPVariableDescription* VariablePtr = FindLocalVariable(InBlueprint, InVariableName);
+
+		if(VariablePtr)
+		{
+			FBPVariableDescription& Variable = *VariablePtr;
+
+			// Update the variable type only if it is different
+			if (Variable.VarName == InVariableName && Variable.VarType != NewPinType)
+			{
+				const FScopedTransaction Transaction( LOCTEXT("ChangeLocalVariableType", "Change Local Variable Type") );
+				InBlueprint->Modify();
+
+				Variable.VarType = NewPinType;
+
+				// Destroy all our nodes, because the pin types could be incorrect now (as will the links)
+				RemoveVariableNodes(InBlueprint, InVariableName);
+
+				// And recompile
+				FBlueprintEditorUtils::MarkBlueprintAsStructurallyModified(InBlueprint);
+			}
+		}
+	}
+}
+
 void FBlueprintEditorUtils::ReplaceVariableReferences(UBlueprint* Blueprint, const FName& OldName, const FName& NewName)
 {
 	check((OldName != NAME_None) && (NewName != NAME_None));
diff --git a/Engine/Source/Editor/UnrealEd/Private/Kismet2/Kismet2NameValidators.cpp b/Engine/Source/Editor/UnrealEd/Private/Kismet2/Kismet2NameValidators.cpp
--- a/Engine/Source/Editor/UnrealEd/Private/Kismet2/Kismet2NameValidators.cpp
+++ b/Engine/Source/Editor/UnrealEd/Private/Kismet2/Kismet2NameValidators.cpp
@@ -141,13 +141,22 @@ EValidatorResult FKismetNameValidator::IsValid(const FName& Name, bool /* bOrigi
 
 		if(ValidatorResult == EValidatorResult::Ok)
 		{
-			TArray<UK2Node_LocalVariable*> EventNodes;
-			FBlueprintEditorUtils::GetAllNodesOfClass(BlueprintObject, EventNodes);
-			for (UK2Node_LocalVariable* LocalVariable : EventNodes)
+			// Search through all functions for their local variables and prevent duplicate names
+			TArray<UK2Node_FunctionEntry*> FunctionEntryNodes;
+			FBlueprintEditorUtils::GetAllNodesOfClass(BlueprintObject, FunctionEntryNodes);
+			for (UK2Node_FunctionEntry* const FunctionEntry : FunctionEntryNodes)
 			{
-				if(LocalVariable->CustomVariableName == Name)
+				for( const FBPVariableDescription& Variable : FunctionEntry->LocalVariables )
+				{
+					if(Variable.VarName == Name)
+					{
+						ValidatorResult = EValidatorResult::AlreadyInUse;
+						break;
+					}
+				}
+
+				if(ValidatorResult != EValidatorResult::Ok)
 				{
-					ValidatorResult = EValidatorResult::AlreadyInUse;
 					break;
 				}
 			}
diff --git a/Engine/Source/Editor/UnrealEd/Public/Kismet2/BlueprintEditorUtils.h b/Engine/Source/Editor/UnrealEd/Public/Kismet2/BlueprintEditorUtils.h
--- a/Engine/Source/Editor/UnrealEd/Public/Kismet2/BlueprintEditorUtils.h
+++ b/Engine/Source/Editor/UnrealEd/Public/Kismet2/BlueprintEditorUtils.h
@@ -249,7 +249,7 @@ public:
 	 *
 	 * @return				The top level graph
 	 */
-	static UEdGraph* GetTopLevelGraph(UEdGraph* InGraph);
+	static UEdGraph* GetTopLevelGraph(const UEdGraph* InGraph);
 
 	/** Look to see if an event already exists to override a particular function */
 	static class UK2Node_Event* FindOverrideForFunction(const UBlueprint* Blueprint, const UClass* SignatureClass, FName SignatureName);
@@ -473,6 +473,59 @@ public:
 	/** Changes the type of a member variable */
 	static void ChangeMemberVariableType(UBlueprint* Blueprint, const FName& VariableName, const FEdGraphPinType& NewPinType);
 
+	/**
+	 * Removes a member variable if it was declared in this blueprint and not in a base class.
+	 *
+	 * @param	InBlueprint			The Blueprint the local variable can be found in
+	 * @param	InVarName			Name of the variable to be removed.
+	 */
+	static void RemoveLocalVariable(UBlueprint* InBlueprint, const FName& InVarName);
+
+	/**
+	 * Returns a local variable
+	 *
+	 * @param InBlueprint		Blueprint to search for the local variable
+	 * @param InVariableName	Name of the variable to search for
+	 * @return					The local variable description
+	 */
+	static FBPVariableDescription* FindLocalVariable(UBlueprint* InBlueprint, const FName& InVariableName);
+
+	/**
+	 * Finds a local variable name using the variable's Guid
+	 *
+	 * @param InBlueprint		Blueprint to search for the local variable
+	 * @param InVariableGuid	Guid to identify the local variable with
+	 * @return					Local variable's name
+	 */
+	static FName FindLocalVariableNameByGuid(UBlueprint* InBlueprint, const FGuid& InVariableGuid);
+
+	/**
+	 * Finds a local variable Guid using the variable's name
+	 *
+	 * @param InBlueprint		Blueprint to search for the local variable
+	 * @param InVariableGuid	Local variable's name to search for
+	 * @return					The Guid associated with the local variable
+	 */
+	static FGuid FindLocalVariableGuidByName(UBlueprint* InBlueprint, const FName InVariableName);
+
+	/**
+	 * Rename a local variable
+	 *
+	 * @param InBlueprint		Blueprint to search for the local variable
+	 * @param InOldName The name of the local variable to change
+	 * @param InNewName	The new name of the local variable
+	 */
+	static void RenameLocalVariable(UBlueprint* InBlueprint, const FName& InOldName, const FName& InNewName);
+
+	/**
+	 * Changes the type of a local variable
+	 *
+	 * @param InBlueprint		Blueprint to search for the local variable
+	 * @param InVariableName	Name of the local variable to change the type of
+	 * @param InNewPinType		The pin type to change the local variable type to
+	 */
+	static void ChangeLocalVariableType(UBlueprint* InBlueprint, const FName& InVariableName, const FEdGraphPinType& InNewPinType);
+
 	/** Replaces all variable references in the specified blueprint */
 	static void ReplaceVariableReferences(UBlueprint* Blueprint, const FName& OldName, const FName& NewName);
 
diff --git a/Engine/Source/Runtime/Engine/Classes/Engine/Blueprint.h b/Engine/Source/Runtime/Engine/Classes/Engine/Blueprint.h
--- a/Engine/Source/Runtime/Engine/Classes/Engine/Blueprint.h
+++ b/Engine/Source/Runtime/Engine/Classes/Engine/Blueprint.h
@@ -167,7 +167,9 @@ struct FBPVariableDescription
 	/** Clear metadata value on the variable */
 	ENGINE_API void RemoveMetaData(const FName& Key);
 	/** Find the index in the array of a metadata entry */
-	ENGINE_API int32 FindMetaDataEntryIndexForKey(const FName& Key);
+	ENGINE_API int32 FindMetaDataEntryIndexForKey(const FName& Key) const;
+	/** Checks if there is metadata for a key */
+	ENGINE_API bool HasMetaData(const FName& Key) const;
 	
 };
 
diff --git a/Engine/Source/Runtime/Engine/Private/Blueprint.cpp b/Engine/Source/Runtime/Engine/Private/Blueprint.cpp
--- a/Engine/Source/Runtime/Engine/Private/Blueprint.cpp
+++ b/Engine/Source/Runtime/Engine/Private/Blueprint.cpp
@@ -16,7 +16,7 @@ DEFINE_LOG_CATEGORY(LogBlueprint);
 //////////////////////////////////////////////////////////////////////////
 // FBPVariableDescription
 
-int32 FBPVariableDescription::FindMetaDataEntryIndexForKey(const FName& Key)
+int32 FBPVariableDescription::FindMetaDataEntryIndexForKey(const FName& Key) const
 {
 	for(int32 i=0; i<MetaDataArray.Num(); i++)
 	{
@@ -28,6 +28,11 @@ int32 FBPVariableDescription::FindMetaDataEntryIndexForKey(const FName& Key)
 	return INDEX_NONE;
 }
 
+bool FBPVariableDescription::HasMetaData(const FName& Key) const
+{
+	return FindMetaDataEntryIndexForKey(Key) != INDEX_NONE;
+}
+
 /** Gets a metadata value on the variable; asserts if the value isn't present.  Check for validiy using FindMetaDataEntryIndexForKey. */
 FString FBPVariableDescription::GetMetaData(const FName& Key)
 {
