UDS: CPF_DisableEditOnInstance flag for variables in UDS. WIP

diff --git a/Engine/Source/Editor/Kismet/Private/UserDefinedStructureEditor.cpp b/Engine/Source/Editor/Kismet/Private/UserDefinedStructureEditor.cpp
--- a/Engine/Source/Editor/Kismet/Private/UserDefinedStructureEditor.cpp
+++ b/Engine/Source/Editor/Kismet/Private/UserDefinedStructureEditor.cpp
@@ -197,6 +197,31 @@ public:
 		return FString();
 	}
 
+	FText OnGetTooltipText() const
+	{
+		auto StructureDetailsSP = StructureDetails.Pin();
+		if (StructureDetailsSP.IsValid())
+		{
+			if (auto Struct = StructureDetailsSP->GetUserDefinedStruct())
+			{
+				return FText::FromString(FStructureEditorUtils::GetTooltip(Struct));
+			}
+		}
+		return FText();
+	}
+
+	void OnTooltipCommitted(const FText& NewText, ETextCommit::Type InTextCommit)
+	{
+		auto StructureDetailsSP = StructureDetails.Pin();
+		if (StructureDetailsSP.IsValid())
+		{
+			if (auto Struct = StructureDetailsSP->GetUserDefinedStruct())
+			{
+				FStructureEditorUtils::ChangeTooltip(Struct, NewText.ToString());
+			}
+		}
+	}
+
 	/** IDetailCustomNodeBuilder Interface*/
 	virtual void SetOnRebuildChildren( FSimpleDelegate InOnRegenerateChildren ) OVERRIDE 
 	{
@@ -205,6 +230,7 @@ public:
 	virtual void GenerateChildContent( IDetailChildrenBuilder& ChildrenBuilder ) OVERRIDE;
 
 	virtual void GenerateHeaderRowContent( FDetailWidgetRow& NodeRow ) OVERRIDE {}
+
 	virtual void Tick( float DeltaTime ) OVERRIDE {}
 	virtual bool RequiresTick() const OVERRIDE { return false; }
 	virtual FName GetName() const OVERRIDE 
@@ -333,6 +359,52 @@ public:
 		return FText();
 	}
 
+	FText OnGetTooltipText() const
+	{
+		auto StructureDetailsSP = StructureDetails.Pin();
+		if (StructureDetailsSP.IsValid())
+		{
+			if (const FStructVariableDescription* FieldDesc = StructureDetailsSP->FindStructureFieldByGuid(FieldGuid))
+			{
+				return FText::FromString(FieldDesc->ToolTip);
+			}
+		}
+		return FText();
+	}
+
+	void OnTooltipCommitted(const FText& NewText, ETextCommit::Type InTextCommit)
+	{
+		auto StructureDetailsSP = StructureDetails.Pin();
+		if (StructureDetailsSP.IsValid())
+		{
+			FStructureEditorUtils::ChangeVariableTooltip(StructureDetailsSP->GetUserDefinedStruct(), FieldGuid, NewText.ToString());
+			OnChanged();
+		}
+	}
+
+	ESlateCheckBoxState::Type OnGetEditableOnBPInstanceState() const
+	{
+		auto StructureDetailsSP = StructureDetails.Pin();
+		if (StructureDetailsSP.IsValid())
+		{
+			if (const FStructVariableDescription* FieldDesc = StructureDetailsSP->FindStructureFieldByGuid(FieldGuid))
+			{
+				return !FieldDesc->bDontEditoOnInstance ? ESlateCheckBoxState::Checked : ESlateCheckBoxState::Unchecked;
+			}
+		}
+		return ESlateCheckBoxState::Undetermined;
+	}
+
+	void OnEditableOnBPInstanceCommitted(ESlateCheckBoxState::Type InNewState)
+	{
+		auto StructureDetailsSP = StructureDetails.Pin();
+		if (StructureDetailsSP.IsValid())
+		{
+			FStructureEditorUtils::ChangeEditableOnBPInstance(StructureDetailsSP->GetUserDefinedStruct(), FieldGuid, ESlateCheckBoxState::Unchecked != InNewState);
+			OnChanged();
+		}
+	}
+
 	void OnArgDefaultValueCommitted(const FText& NewText, ETextCommit::Type InTextCommit)
 	{
 		auto StructureDetailsSP = StructureDetails.Pin();
@@ -473,7 +545,7 @@ public:
 
 	virtual void GenerateChildContent( IDetailChildrenBuilder& ChildrenBuilder ) OVERRIDE 
 	{
-		ChildrenBuilder.AddChildContent( *LOCTEXT( "FunctionArgDetailsDefaultValue", "Default Value" ).ToString() )
+		ChildrenBuilder.AddChildContent( *LOCTEXT( "UDSVarDefaultValue", "Default Value" ).ToString() )
 		.NameContent()
 		[
 			SNew(STextBlock)
@@ -487,6 +559,38 @@ public:
 			.OnTextCommitted( this, &FUserDefinedStructureFieldLayout::OnArgDefaultValueCommitted )
 			.Font( IDetailLayoutBuilder::GetDetailFont() )
 		];
+
+		ChildrenBuilder.AddChildContent(*LOCTEXT("Tooltip", "Tooltip").ToString())
+		.NameContent()
+		[
+			SNew(STextBlock)
+			.Text(LOCTEXT("Tooltip", "Tooltip"))
+			.Font(IDetailLayoutBuilder::GetDetailFont())
+		]
+		.ValueContent()
+		[
+			SNew(SEditableTextBox)
+			.Text(this, &FUserDefinedStructureFieldLayout::OnGetTooltipText)
+			.OnTextCommitted(this, &FUserDefinedStructureFieldLayout::OnTooltipCommitted)
+			.Font(IDetailLayoutBuilder::GetDetailFont())
+		];
+
+#if 0	// ttp335214 must be fixed first
+		ChildrenBuilder.AddChildContent(*LOCTEXT("EditableOnInstance", "EditableOnInstance").ToString())
+		.NameContent()
+		[
+			SNew(STextBlock)
+			.Text(LOCTEXT("Editable", "Editable"))
+			.Font(IDetailLayoutBuilder::GetDetailFont())
+		]
+		.ValueContent()
+		[
+			SNew(SCheckBox)
+			.ToolTipText(LOCTEXT("EditableOnBPInstance", "Variable can be edited on an instance of a Blueprint."))
+			.OnCheckStateChanged(this, &FUserDefinedStructureFieldLayout::OnEditableOnBPInstanceCommitted)
+			.IsChecked(this, &FUserDefinedStructureFieldLayout::OnGetEditableOnBPInstanceState)
+		];
+#endif 
 	}
 
 	virtual void Tick( float DeltaTime ) OVERRIDE {}
@@ -508,26 +612,61 @@ private:
 // FUserDefinedStructureLayout
 void FUserDefinedStructureLayout::GenerateChildContent( IDetailChildrenBuilder& ChildrenBuilder ) 
 {
+	const float NameWidth = 80.0f;
+	const float ContentWidth = 130.0f;
+
 	ChildrenBuilder.AddChildContent(FString())
+	[
+		SNew(SHorizontalBox)
+		+ SHorizontalBox::Slot()
+		.MaxWidth(NameWidth)
+		.HAlign(HAlign_Left)
+		.VAlign(VAlign_Center)
 		[
-			SNew(SHorizontalBox)
-			+SHorizontalBox::Slot()
-			.AutoWidth()
-			.HAlign(HAlign_Left)
-			.VAlign(VAlign_Center)
-			[
-				SNew(SImage)
-				.Image( this, &FUserDefinedStructureLayout::OnGetStructureStatus )
-				.ToolTipText(this, &FUserDefinedStructureLayout::GetStatusTooltip )
-			]
-			+SHorizontalBox::Slot()
-			.HAlign(HAlign_Right)
+			SNew(SImage)
+			.Image(this, &FUserDefinedStructureLayout::OnGetStructureStatus)
+			.ToolTipText(this, &FUserDefinedStructureLayout::GetStatusTooltip)
+		]
+		+ SHorizontalBox::Slot()
+		.AutoWidth()
+		.HAlign(HAlign_Left)
+		[
+			SNew(SBox)
+			.WidthOverride(ContentWidth)
 			[
 				SNew(SButton)
+				.HAlign(HAlign_Center)
 				.Text(LOCTEXT("NewStructureField", "New Variable").ToString())
 				.OnClicked(this, &FUserDefinedStructureLayout::OnAddNewField)
 			]
-		];
+		]
+	];
+
+	ChildrenBuilder.AddChildContent(FString())
+	[
+		SNew(SHorizontalBox)
+		+ SHorizontalBox::Slot()
+		.MaxWidth(NameWidth)
+		.HAlign(HAlign_Left)
+		[
+			SNew(STextBlock)
+			.Text(LOCTEXT("Tooltip", "Tooltip"))
+			.Font(IDetailLayoutBuilder::GetDetailFont())
+		]
+		+ SHorizontalBox::Slot()
+		.AutoWidth()
+		.HAlign(HAlign_Left)
+		[
+			SNew(SBox)
+			.WidthOverride(ContentWidth)
+			[
+				SNew(SEditableTextBox)
+				.Text(this, &FUserDefinedStructureLayout::OnGetTooltipText)
+				.OnTextCommitted(this, &FUserDefinedStructureLayout::OnTooltipCommitted)
+				.Font(IDetailLayoutBuilder::GetDetailFont())
+			]
+		]
+	];
 
 	auto StructureDetailsSP = StructureDetails.Pin();
 	if(StructureDetailsSP.IsValid())
diff --git a/Engine/Source/Editor/KismetCompiler/Private/UserDefinedStructureCompilerUtils.cpp b/Engine/Source/Editor/KismetCompiler/Private/UserDefinedStructureCompilerUtils.cpp
--- a/Engine/Source/Editor/KismetCompiler/Private/UserDefinedStructureCompilerUtils.cpp
+++ b/Engine/Source/Editor/KismetCompiler/Private/UserDefinedStructureCompilerUtils.cpp
@@ -146,6 +146,10 @@ struct FUserDefinedStructureCompilerInner
 			}
 
 			NewProperty->SetPropertyFlags(CPF_Edit);
+			if (VarDesc.bDontEditoOnInstance)
+			{
+				NewProperty->SetPropertyFlags(CPF_DisableEditOnInstance);
+			}
 			NewProperty->SetMetaData(TEXT("DisplayName"), *VarDesc.FriendlyName);
 			NewProperty->SetMetaData(TEXT("Category"), *VarDesc.Category);
 			NewProperty->SetMetaData(FBlueprintMetadata::MD_Tooltip, *VarDesc.ToolTip);
diff --git a/Engine/Source/Editor/UnrealEd/Classes/UserDefinedStructure/UserDefinedStructEditorData.h b/Engine/Source/Editor/UnrealEd/Classes/UserDefinedStructure/UserDefinedStructEditorData.h
--- a/Engine/Source/Editor/UnrealEd/Classes/UserDefinedStructure/UserDefinedStructEditorData.h
+++ b/Engine/Source/Editor/UnrealEd/Classes/UserDefinedStructure/UserDefinedStructEditorData.h
@@ -44,9 +44,18 @@ struct FStructVariableDescription
 	UPROPERTY()
 	FString ToolTip;
 
+	UPROPERTY()
+	bool bDontEditoOnInstance;
+
 	UNREALED_API bool SetPinType(const struct FEdGraphPinType& VarType);
 
 	UNREALED_API FEdGraphPinType ToPinType() const;
+
+	FStructVariableDescription()
+		: bIsArray(false)
+		, bInvalidMember(false)
+		, bDontEditoOnInstance(false)
+	{ }
 };
 
 UCLASS()
diff --git a/Engine/Source/Editor/UnrealEd/Private/Kismet2/StructureEditorUtils.cpp b/Engine/Source/Editor/UnrealEd/Private/Kismet2/StructureEditorUtils.cpp
--- a/Engine/Source/Editor/UnrealEd/Private/Kismet2/StructureEditorUtils.cpp
+++ b/Engine/Source/Editor/UnrealEd/Private/Kismet2/StructureEditorUtils.cpp
@@ -239,6 +239,8 @@ bool FStructureEditorUtils::AddVariable(UUserDefinedStruct* Struct, const FEdGra
 		NewVar.FriendlyName = DisplayName;
 		NewVar.SetPinType(VarType);
 		NewVar.VarGuid = FGuid::NewGuid();
+		NewVar.bDontEditoOnInstance = false;
+		NewVar.bInvalidMember = false;
 		GetVarDesc(Struct).Add(NewVar);
 
 		OnStructureChanged(Struct);
@@ -593,6 +595,26 @@ bool FStructureEditorUtils::ChangeVariableTooltip(UUserDefinedStruct* Struct, FG
 	return false;
 }
 
+bool FStructureEditorUtils::ChangeEditableOnBPInstance(UUserDefinedStruct* Struct, FGuid VarGuid, bool bInIsEditable)
+{
+	if (Struct)
+	{
+		const FScopedTransaction Transaction(LOCTEXT("ChangeVariableDefaultValue", "Change Variable Default Value"));
+		ModifyStructData(Struct);
+
+		const UEdGraphSchema_K2* K2Schema = GetDefault<UEdGraphSchema_K2>();
+		FStructVariableDescription* VarDesc = GetVarDesc(Struct).FindByPredicate(FFindByGuidHelper<FStructVariableDescription>(VarGuid));
+		const bool bNewDontEditoOnInstance = !bInIsEditable;
+		if (VarDesc && (bNewDontEditoOnInstance != VarDesc->bDontEditoOnInstance))
+		{
+			VarDesc->bDontEditoOnInstance = bNewDontEditoOnInstance;
+			OnStructureChanged(Struct);
+			return true;
+		}
+	}
+	return false;
+}
+
 void FStructureEditorUtils::ModifyStructData(UUserDefinedStruct* Struct)
 {
 	UUserDefinedStructEditorData* EditorData = Struct ? Cast<UUserDefinedStructEditorData>(Struct->EditorData) : NULL;
diff --git a/Engine/Source/Editor/UnrealEd/Public/Kismet2/StructureEditorUtils.h b/Engine/Source/Editor/UnrealEd/Public/Kismet2/StructureEditorUtils.h
--- a/Engine/Source/Editor/UnrealEd/Public/Kismet2/StructureEditorUtils.h
+++ b/Engine/Source/Editor/UnrealEd/Public/Kismet2/StructureEditorUtils.h
@@ -107,6 +107,9 @@ public:
 	
 	static bool ChangeVariableTooltip(UUserDefinedStruct* Struct, FGuid VarGuid, const FString& InTooltip);
 
+	static bool ChangeEditableOnBPInstance(UUserDefinedStruct* Struct, FGuid VarGuid, bool bInIsEditable);
+
+
 	//MISC
 	static TArray<FStructVariableDescription>& GetVarDesc(UUserDefinedStruct* Struct);
 
