Modified the goto code definition to just open the source file header for release 4.2 to side step issues with absolute paths in the symbol map failing to locate the correct source path ( if it was built on a different machine )
#TTP 335153 - PUNTABLE: Editor: Right-clicking a blueprint node and selecting 'Goto code definition' doesn't work
#Branch UE4
#Proj  Kismet

#Change replaced the jump to function definition using the symbol map with a simpler open source file for the header.

#Info This was neccesary for release 4.2 because of percieved issues with symbol maps that were built on different machine's with different paths failing to open the source code. With more time I'll investigate handling this case better and provide better failsafes and user notification and reinstate the original code ( which jumped to actual definition in code rather than the header of the class ).

diff --git a/Engine/Source/Editor/Kismet/Private/BlueprintEditor.cpp b/Engine/Source/Editor/Kismet/Private/BlueprintEditor.cpp
--- a/Engine/Source/Editor/Kismet/Private/BlueprintEditor.cpp
+++ b/Engine/Source/Editor/Kismet/Private/BlueprintEditor.cpp
@@ -5959,7 +5959,14 @@ void FBlueprintEditor::GotoNativeFunctionDefinition()
 
 			if( TargetFunction )
 			{
-				FSourceCodeNavigation::NavigateToFunctionAsync( TargetFunction );
+				FString NativeParentClassHeaderPath;
+				const bool bFileFound = FSourceCodeNavigation::FindClassHeaderPath(TargetFunction, NativeParentClassHeaderPath) 
+										&& (IFileManager::Get().FileSize(*NativeParentClassHeaderPath) != INDEX_NONE);
+				if( bFileFound )
+				{
+					const FString AbsNativeParentClassHeaderPath = FPaths::ConvertRelativePathToFull( NativeParentClassHeaderPath );
+					FSourceCodeNavigation::OpenSourceFile( AbsNativeParentClassHeaderPath );
+				}
 			}
 		}
 	}
@@ -6003,7 +6010,14 @@ void FBlueprintEditor::GotoNativeVariableDefinition()
 
 			if( VariableProperty )
 			{
-				FSourceCodeNavigation::NavigateToProperty( VariableProperty );
+				FString NativeParentClassHeaderPath;
+				const bool bFileFound = FSourceCodeNavigation::FindClassHeaderPath(VariableProperty, NativeParentClassHeaderPath) 
+										&& (IFileManager::Get().FileSize(*NativeParentClassHeaderPath) != INDEX_NONE);
+				if( bFileFound )
+				{
+					const FString AbsNativeParentClassHeaderPath = FPaths::ConvertRelativePathToFull( NativeParentClassHeaderPath );
+					FSourceCodeNavigation::OpenSourceFile( AbsNativeParentClassHeaderPath );
+				}
 			}
 		}
 	}
