UDS: CPF_DisableEditOnInstance flag for variables in UDS.

diff --git a/Engine/Source/Editor/Kismet/Private/UserDefinedStructureEditor.cpp b/Engine/Source/Editor/Kismet/Private/UserDefinedStructureEditor.cpp
--- a/Engine/Source/Editor/Kismet/Private/UserDefinedStructureEditor.cpp
+++ b/Engine/Source/Editor/Kismet/Private/UserDefinedStructureEditor.cpp
@@ -575,7 +575,6 @@ public:
 			.Font(IDetailLayoutBuilder::GetDetailFont())
 		];
 
-#if 0	// ttp335214 must be fixed first
 		ChildrenBuilder.AddChildContent(*LOCTEXT("EditableOnInstance", "EditableOnInstance").ToString())
 		.NameContent()
 		[
@@ -590,7 +589,6 @@ public:
 			.OnCheckStateChanged(this, &FUserDefinedStructureFieldLayout::OnEditableOnBPInstanceCommitted)
 			.IsChecked(this, &FUserDefinedStructureFieldLayout::OnGetEditableOnBPInstanceState)
 		];
-#endif 
 	}
 
 	virtual void Tick( float DeltaTime ) OVERRIDE {}
