Editoronly BP: TimelinesCurves outer is BPGC instead of BLueprint

diff --git a/Engine/Source/Editor/Kismet/Private/STimelineEditor.cpp b/Engine/Source/Editor/Kismet/Private/STimelineEditor.cpp
--- a/Engine/Source/Editor/Kismet/Private/STimelineEditor.cpp
+++ b/Engine/Source/Editor/Kismet/Private/STimelineEditor.cpp
@@ -915,6 +915,8 @@ FReply STimelineEditor::CreateNewTrack(FTimelineEdTrack::ETrackType Type)
 	TSharedPtr<FBlueprintEditor> Kismet2 = Kismet2Ptr.Pin();
 	UBlueprint* Blueprint = Kismet2->GetBlueprintObj();
 	UK2Node_Timeline* TimelineNode = FBlueprintEditorUtils::FindNodeForTimeline(Blueprint, TimelineObj);
+	UClass* OwnerClass = Blueprint->GeneratedClass;
+	check(OwnerClass);
 
 	FText ErrorMessage;
 
@@ -932,7 +934,7 @@ FReply STimelineEditor::CreateNewTrack(FTimelineEdTrack::ETrackType Type)
 			{
 				FTTEventTrack NewTrack;
 				NewTrack.TrackName = TrackName;
-				NewTrack.CurveKeys = NewNamedObject<UCurveFloat>(Blueprint, NAME_None, RF_Public); // Needs to be marked public so that it can be referenced from timeline instances in the level
+				NewTrack.CurveKeys = NewNamedObject<UCurveFloat>(OwnerClass, NAME_None, RF_Public); // Needs to be marked public so that it can be referenced from timeline instances in the level
 				NewTrack.CurveKeys->bIsEventCurve = true;
 				TimelineObj->EventTracks.Add(NewTrack);
 			}
@@ -944,7 +946,7 @@ FReply STimelineEditor::CreateNewTrack(FTimelineEdTrack::ETrackType Type)
 				NewTrack.CurveFloat = FindObject<UCurveFloat>(ANY_PACKAGE, *TrackName.ToString() );
 				if (NewTrack.CurveFloat == NULL)
 				{
-					NewTrack.CurveFloat = NewNamedObject<UCurveFloat>(Blueprint, NAME_None, RF_Public); 
+					NewTrack.CurveFloat = NewNamedObject<UCurveFloat>(OwnerClass, NAME_None, RF_Public);
 				}
 				TimelineObj->FloatTracks.Add(NewTrack);
 			}
@@ -952,14 +954,14 @@ FReply STimelineEditor::CreateNewTrack(FTimelineEdTrack::ETrackType Type)
 			{
 				FTTVectorTrack NewTrack;
 				NewTrack.TrackName = TrackName;
-				NewTrack.CurveVector = NewNamedObject<UCurveVector>(Blueprint, NAME_None, RF_Public); 
+				NewTrack.CurveVector = NewNamedObject<UCurveVector>(OwnerClass, NAME_None, RF_Public);
 				TimelineObj->VectorTracks.Add(NewTrack);
 			}
 			else if(Type == FTimelineEdTrack::TT_LinearColorInterp)
 			{
 				FTTLinearColorTrack NewTrack;
 				NewTrack.TrackName = TrackName;
-				NewTrack.CurveLinearColor = NewNamedObject<UCurveLinearColor>(Blueprint, NAME_None, RF_Public); 
+				NewTrack.CurveLinearColor = NewNamedObject<UCurveLinearColor>(OwnerClass, NAME_None, RF_Public);
 				TimelineObj->LinearColorTracks.Add(NewTrack);
 			}
 
@@ -1000,23 +1002,24 @@ UCurveBase* STimelineEditor::CreateNewCurve( FTimelineEdTrack::ETrackType Type )
 {
 	TSharedPtr<FBlueprintEditor> Kismet2 = Kismet2Ptr.Pin();
 	UBlueprint* Blueprint = Kismet2->GetBlueprintObj();
-
+	UClass* OwnerClass = Blueprint->GeneratedClass;
+	check(OwnerClass);
 	UCurveBase* NewCurve = NULL;
 	if(Type == FTimelineEdTrack::TT_Event)
 	{
-		NewCurve = NewNamedObject<UCurveFloat>(Blueprint, NAME_None, RF_Public); 
+		NewCurve = NewNamedObject<UCurveFloat>(OwnerClass, NAME_None, RF_Public);
 	}
 	else if(Type == FTimelineEdTrack::TT_FloatInterp)
 	{
-		NewCurve = NewNamedObject<UCurveFloat>(Blueprint, NAME_None, RF_Public); 
+		NewCurve = NewNamedObject<UCurveFloat>(OwnerClass, NAME_None, RF_Public);
 	}
 	else if(Type == FTimelineEdTrack::TT_VectorInterp)
 	{
-		NewCurve = NewNamedObject<UCurveVector>(Blueprint, NAME_None, RF_Public);
+		NewCurve = NewNamedObject<UCurveVector>(OwnerClass, NAME_None, RF_Public);
 	}
 	else if(Type == FTimelineEdTrack::TT_LinearColorInterp)
 	{
-		NewCurve = NewNamedObject<UCurveLinearColor>(Blueprint, NAME_None, RF_Public);
+		NewCurve = NewNamedObject<UCurveLinearColor>(OwnerClass, NAME_None, RF_Public);
 	}
 
 
diff --git a/Engine/Source/Editor/KismetCompiler/Private/KismetCompiler.cpp b/Engine/Source/Editor/KismetCompiler/Private/KismetCompiler.cpp
--- a/Engine/Source/Editor/KismetCompiler/Private/KismetCompiler.cpp
+++ b/Engine/Source/Editor/KismetCompiler/Private/KismetCompiler.cpp
@@ -197,6 +197,25 @@ void FKismetCompilerContext::CleanAndSanitizeClass(UBlueprintGeneratedClass* Cla
 				SubObjectsToSave.AddObject(SCSNode->ComponentTemplate);
 			}
 		}
+		{
+			TSet<class UCurveBase*> Curves;
+			for (auto Timeline : Blueprint->Timelines)
+			{
+				Timeline->GetAllCurves(Curves);
+			}
+			for (auto Component : Blueprint->ComponentTemplates)
+			{
+				if (auto TimelineComponent = Cast<UTimelineComponent>(Component))
+				{
+					TimelineComponent->GetAllCurves(Curves);
+				}
+			}
+			for (auto Curve : Curves)
+			{
+				SubObjectsToSave.AddObject(Curve);
+			}
+		}
+
 		ClassSubObjects.RemoveAllSwap(SubObjectsToSave);
 	}
 
diff --git a/Engine/Source/Runtime/Engine/Classes/Components/TimelineComponent.h b/Engine/Source/Runtime/Engine/Classes/Components/TimelineComponent.h
--- a/Engine/Source/Runtime/Engine/Classes/Components/TimelineComponent.h
+++ b/Engine/Source/Runtime/Engine/Classes/Components/TimelineComponent.h
@@ -345,6 +345,9 @@ public:
 
 	/** Set the delegate to call when timeline is finished */
 	ENGINE_API void SetDirectionPropertyName(FName InDirectionPropertyName);
+
+	/** Get all curves used by the Timeline */
+	void GetAllCurves(TSet<class UCurveBase*>& InOutCurves) const;
 private:
 	/** Returns the time value of the last keyframe in any of the timeline's curves */
 	float GetLastKeyframeTime() const;
@@ -478,6 +481,9 @@ public:
 
 	/** Set the delegate to call when timeline is finished */
 	ENGINE_API void SetDirectionPropertyName(FName DirectionPropertyName);
+
+	/** Get all curves used by the Timeline */
+	ENGINE_API void GetAllCurves(TSet<class UCurveBase*>& InOutCurves) const;
 };
 
 
diff --git a/Engine/Source/Runtime/Engine/Classes/Engine/TimelineTemplate.h b/Engine/Source/Runtime/Engine/Classes/Engine/TimelineTemplate.h
--- a/Engine/Source/Runtime/Engine/Classes/Engine/TimelineTemplate.h
+++ b/Engine/Source/Runtime/Engine/Classes/Engine/TimelineTemplate.h
@@ -203,6 +203,8 @@ class UTimelineTemplate : public UObject
 	ENGINE_API static FString TimelineTemplateNameToVariableName(FName Name);
 	ENGINE_API static FString TimelineVariableNameToTemplateName(FName Name);
 
+	ENGINE_API void GetAllCurves(TSet<class UCurveBase*>& InOutCurves) const;
+
 	/* Perform deep copy of curves */
 	virtual void PostDuplicate(bool bDuplicateForPIE) OVERRIDE;
 };
diff --git a/Engine/Source/Runtime/Engine/Private/Blueprint.cpp b/Engine/Source/Runtime/Engine/Private/Blueprint.cpp
--- a/Engine/Source/Runtime/Engine/Private/Blueprint.cpp
+++ b/Engine/Source/Runtime/Engine/Private/Blueprint.cpp
@@ -896,6 +896,7 @@ bool UBlueprint::ChangeOwnerOfTemplates()
 
 		// >>> Backwards Compatibility:  VER_UE4_EDITORONLY_BLUEPRINTS
 		bool bMigratedOwner = false;
+		TSet<class UCurveBase*> Curves;
 		for( auto CompIt = ComponentTemplates.CreateIterator(); CompIt; ++CompIt )
 		{
 			UActorComponent* Component = (*CompIt);
@@ -907,6 +908,10 @@ bool UBlueprint::ChangeOwnerOfTemplates()
 				bIsStillStale |= !bRenamed;
 				bMigratedOwner = true;
 			}
+			if (auto TimelineComponent = Cast<UTimelineComponent>(Component))
+			{
+				TimelineComponent->GetAllCurves(Curves);
+			}
 		}
 
 		for( auto CompIt = Timelines.CreateIterator(); CompIt; ++CompIt )
@@ -923,6 +928,16 @@ bool UBlueprint::ChangeOwnerOfTemplates()
 				ensure(OldTemplateName == UTimelineTemplate::TimelineTemplateNameToVariableName(Template->GetFName()));
 				bMigratedOwner = true;
 			}
+			Template->GetAllCurves(Curves);
+		}
+		for (auto Curve : Curves)
+		{
+			if (Curve && (Curve->GetOuter() == this))
+			{
+				const bool bRenamed = Curve->Rename(NULL, BPGClass, REN_ForceNoResetLoaders|REN_DoNotDirty);
+				ensure(bRenamed);
+				bIsStillStale |= !bRenamed;
+			}
 		}
 
 		if(USimpleConstructionScript* SCS = SimpleConstructionScript)
diff --git a/Engine/Source/Runtime/Engine/Private/Timeline.cpp b/Engine/Source/Runtime/Engine/Private/Timeline.cpp
--- a/Engine/Source/Runtime/Engine/Private/Timeline.cpp
+++ b/Engine/Source/Runtime/Engine/Private/Timeline.cpp
@@ -536,6 +536,22 @@ float FTimeline::GetLastKeyframeTime() const
 	return MaxTime;
 }
 
+void FTimeline::GetAllCurves(TSet<class UCurveBase*>& InOutCurves) const
+{
+	for (auto& Track : InterpVectors)
+	{
+		InOutCurves.Add(Track.VectorCurve);
+	}
+	for (auto& Track : InterpFloats)
+	{
+		InOutCurves.Add(Track.FloatCurve);
+	}
+	for (auto& Track : InterpLinearColors)
+	{
+		InOutCurves.Add(Track.LinearColorCurve);
+	}
+}
+
 //////////////////////////////////////////////////////////////////////////
 // UTimelineComponent
 
@@ -765,3 +781,8 @@ void UTimelineComponent::GetLifetimeReplicatedProps( TArray< FLifetimeProperty >
 
 	DOREPLIFETIME( UTimelineComponent, TheTimeline );
 }
+
+void UTimelineComponent::GetAllCurves(TSet<class UCurveBase*>& InOutCurves) const
+{
+	TheTimeline.GetAllCurves(InOutCurves);
+}
\ No newline at end of file
diff --git a/Engine/Source/Runtime/Engine/Private/TimelineTemplate.cpp b/Engine/Source/Runtime/Engine/Private/TimelineTemplate.cpp
--- a/Engine/Source/Runtime/Engine/Private/TimelineTemplate.cpp
+++ b/Engine/Source/Runtime/Engine/Private/TimelineTemplate.cpp
@@ -298,6 +298,26 @@ void UTimelineTemplate::PostDuplicate(bool bDuplicateForPIE)
 	TimelineGuid = FGuid::NewGuid();
 }
 
+void UTimelineTemplate::GetAllCurves(TSet<class UCurveBase*>& InOutCurves) const
+{
+	for (auto& Track : EventTracks)
+	{
+		InOutCurves.Add(Track.CurveKeys);
+	}
+	for (auto& Track : FloatTracks)
+	{
+		InOutCurves.Add(Track.CurveFloat);
+	}
+	for (auto& Track : VectorTracks)
+	{
+		InOutCurves.Add(Track.CurveVector);
+	}
+	for (auto& Track : LinearColorTracks)
+	{
+		InOutCurves.Add(Track.CurveLinearColor);
+	}
+}
+
 bool FTTTrackBase::operator==( const FTTTrackBase& T2 ) const
 {
 	return (TrackName == T2.TrackName) &&
