[336023] Fix jittery transforms when both parent and child nodes are selected in Components mode within the Blueprint editor.

- Modified FSCSEditorViewportClient::InputWidgetDelta() to only apply transform to a component if it is not attached to another [movable] component farther up the scene hierarchy.
- Modified SSCSEditor::GetSelectedNodes() to sort the selected node array from parent->child in order to ensure that selected nodes are always arranged in parent->child order.

diff --git a/Engine/Source/Editor/Kismet/Private/SCSEditorViewportClient.cpp b/Engine/Source/Editor/Kismet/Private/SCSEditorViewportClient.cpp
--- a/Engine/Source/Editor/Kismet/Private/SCSEditorViewportClient.cpp
+++ b/Engine/Source/Editor/Kismet/Private/SCSEditorViewportClient.cpp
@@ -100,6 +100,31 @@ namespace
 
 		return false;
 	}
+
+	// Determine whether or not the given node has a parent node that is not the root node, is movable and is selected
+	bool IsMovableParentNodeSelected(const FSCSEditorTreeNodePtrType& NodePtr, const TArray<FSCSEditorTreeNodePtrType>& SelectedNodes)
+	{
+		if(NodePtr.IsValid())
+		{
+			// Check for a valid parent node
+			FSCSEditorTreeNodePtrType ParentNodePtr = NodePtr->GetParent();
+			if(ParentNodePtr.IsValid() && !ParentNodePtr->IsRoot())
+			{
+				if(SelectedNodes.Contains(ParentNodePtr))
+				{
+					// The parent node is not the root node and is also selected; success
+					return true;
+				}
+				else
+				{
+					// Recursively search for any other parent nodes farther up the tree that might be selected
+					return IsParentNodeSelected(ParentNodePtr, SelectedNodes);
+				}
+			}
+		}
+
+		return false;
+	}
 }
 
 /////////////////////////////////////////////////////////////////////////
@@ -357,8 +382,9 @@ bool FSCSEditorViewportClient::InputWidgetDelta( FViewport* Viewport, EAxisList:
 				for(auto It(SelectedNodes.CreateIterator());It;++It)
 				{
 					FSCSEditorTreeNodePtrType SelectedNodePtr = *It;
-					// Don't allow editing of a root node or inherited SCS node
-					const bool bCanEdit =  !SelectedNodePtr->IsRoot() && !SelectedNodePtr->IsInherited();
+					// Don't allow editing of a root node, inherited SCS node or child node that also has a movable (non-root) parent node selected
+					const bool bCanEdit =  !SelectedNodePtr->IsRoot() && !SelectedNodePtr->IsInherited()
+						&& !IsMovableParentNodeSelected(SelectedNodePtr, SelectedNodes);
 
 					if(bCanEdit)
 					{
diff --git a/Engine/Source/Editor/Kismet/Private/SSCSEditor.cpp b/Engine/Source/Editor/Kismet/Private/SSCSEditor.cpp
--- a/Engine/Source/Editor/Kismet/Private/SSCSEditor.cpp
+++ b/Engine/Source/Editor/Kismet/Private/SSCSEditor.cpp
@@ -2341,14 +2341,22 @@ void SSCSEditor::PerformComboAddClass(TSubclassOf<UActorComponent> ComponentClas
 	}
 }
 
+TArray<FSCSEditorTreeNodePtrType>  SSCSEditor::GetSelectedNodes() const
+{
+	TArray<FSCSEditorTreeNodePtrType> SelectedTreeNodes = SCSTreeWidget->GetSelectedItems();
 
+	struct FCompareSelectedSCSEditorTreeNodes
+	{
+		FORCEINLINE bool operator()(const FSCSEditorTreeNodePtrType& A, const FSCSEditorTreeNodePtrType& B) const
+		{
+			return B.IsValid() && B->IsAttachedTo(A);
+		}
+	};
 
+	// Ensure that nodes are ordered from parent to child (otherwise they are sorted in the order that they were selected)
+	SelectedTreeNodes.Sort(FCompareSelectedSCSEditorTreeNodes());
 
-
-
-TArray<FSCSEditorTreeNodePtrType>  SSCSEditor::GetSelectedNodes() const
-{
-	return SCSTreeWidget->GetSelectedItems();
+	return SelectedTreeNodes;
 }
 
 FSCSEditorTreeNodePtrType SSCSEditor::GetNodeFromActorComponent(const UActorComponent* ActorComponent, bool bIncludeAttachedComponents) const
