CIS fix

diff --git a/Engine/Source/Editor/Kismet/Private/SCSEditorViewportClient.cpp b/Engine/Source/Editor/Kismet/Private/SCSEditorViewportClient.cpp
--- a/Engine/Source/Editor/Kismet/Private/SCSEditorViewportClient.cpp
+++ b/Engine/Source/Editor/Kismet/Private/SCSEditorViewportClient.cpp
@@ -118,7 +118,7 @@ namespace
 				else
 				{
 					// Recursively search for any other parent nodes farther up the tree that might be selected
-					return IsParentNodeSelected(ParentNodePtr, SelectedNodes);
+					return IsMovableParentNodeSelected(ParentNodePtr, SelectedNodes);
 				}
 			}
 		}
