Allow you to select a debug blueprint object from the editor world (before PIE'ing).
(TTP #332485)

diff --git a/Engine/Source/Editor/Kismet/Private/SBlueprintEditorSelectedDebugObjectWidget.cpp b/Engine/Source/Editor/Kismet/Private/SBlueprintEditorSelectedDebugObjectWidget.cpp
--- a/Engine/Source/Editor/Kismet/Private/SBlueprintEditorSelectedDebugObjectWidget.cpp
+++ b/Engine/Source/Editor/Kismet/Private/SBlueprintEditorSelectedDebugObjectWidget.cpp
@@ -413,7 +413,7 @@ void SBlueprintEditorSelectedDebugObjectWidget::GenerateDebugObjectNames(bool bR
 			}
 
 			// We don't have a specific debug world, but the object isnt in a PIE world
-			if (ObjWorld->WorldType != EWorldType::PIE)
+			if ((ObjWorld->WorldType != EWorldType::PIE) && (ObjWorld->WorldType != EWorldType::Editor))
 			{
 				continue;
 			}
