Local variables will no longer stop working when renaming a function.

#ttp 337236 - BP: Local variables cannot be found after adding and renaming a new function

diff --git a/Engine/Source/Editor/BlueprintGraph/Private/K2Node_VariableGet.cpp b/Engine/Source/Editor/BlueprintGraph/Private/K2Node_VariableGet.cpp
--- a/Engine/Source/Editor/BlueprintGraph/Private/K2Node_VariableGet.cpp
+++ b/Engine/Source/Editor/BlueprintGraph/Private/K2Node_VariableGet.cpp
@@ -37,7 +37,15 @@ public:
 			FFormatNamedArguments Args;
 			Args.Add(TEXT("VariableName"), FText::FromName(VarNode->VariableReference.GetMemberName()));
 
-			CompilerContext.MessageLog.Warning(*FText::Format(LOCTEXT("LocalVariableNotFound_Error", "Unable to find local variable with name '{VariableName}' for @@"), Args).ToString(), Node);
+			if(VarNode->VariableReference.GetMemberScopeName() != Context.Function->GetName())
+			{
+				Args.Add(TEXT("ScopeName"), FText::FromString(VarNode->VariableReference.GetMemberScopeName()));
+				CompilerContext.MessageLog.Warning(*FText::Format(LOCTEXT("LocalVariableNotFoundInScope_Error", "Unable to find local variable with name '{VariableName}' for @@, scope expected: @@, scope found: {ScopeName}"), Args).ToString(), Node, Node->GetGraph());
+			}
+			else
+			{
+				CompilerContext.MessageLog.Warning(*FText::Format(LOCTEXT("LocalVariableNotFound_Error", "Unable to find local variable with name '{VariableName}' for @@"), Args).ToString(), Node);
+			}
 		}
 
 		FNodeHandlingFunctor::RegisterNets(Context, Node);
diff --git a/Engine/Source/Editor/BlueprintGraph/Private/VariableSetHandler.cpp b/Engine/Source/Editor/BlueprintGraph/Private/VariableSetHandler.cpp
--- a/Engine/Source/Editor/BlueprintGraph/Private/VariableSetHandler.cpp
+++ b/Engine/Source/Editor/BlueprintGraph/Private/VariableSetHandler.cpp
@@ -33,7 +33,15 @@ void FKCHandler_VariableSet::RegisterNets(FKismetFunctionContext& Context, UEdGr
 			FFormatNamedArguments Args;
 			Args.Add(TEXT("VariableName"), FText::FromName(SetterNode->VariableReference.GetMemberName()));
 
-			CompilerContext.MessageLog.Warning(*FText::Format(LOCTEXT("LocalVariableNotFound_Error", "Unable to find local variable with name '{VariableName}' for @@"), Args).ToString(), Node);
+			if(SetterNode->VariableReference.GetMemberScopeName() != Context.Function->GetName())
+			{
+				Args.Add(TEXT("ScopeName"), FText::FromString(SetterNode->VariableReference.GetMemberScopeName()));
+				CompilerContext.MessageLog.Warning(*FText::Format(LOCTEXT("LocalVariableNotFoundInScope_Error", "Unable to find local variable with name '{VariableName}' for @@, scope expected: @@, scope found: {ScopeName}"), Args).ToString(), Node, Node->GetGraph());
+			}
+			else
+			{
+				CompilerContext.MessageLog.Warning(*FText::Format(LOCTEXT("LocalVariableNotFound_Error", "Unable to find local variable with name '{VariableName}' for @@"), Args).ToString(), Node);
+			}
 		}
 	}
 
diff --git a/Engine/Source/Editor/Kismet/Private/SBlueprintPalette.cpp b/Engine/Source/Editor/Kismet/Private/SBlueprintPalette.cpp
--- a/Engine/Source/Editor/Kismet/Private/SBlueprintPalette.cpp
+++ b/Engine/Source/Editor/Kismet/Private/SBlueprintPalette.cpp
@@ -1118,20 +1118,21 @@ void SBlueprintPaletteItem::OnNameTextCommitted(const FText& NewText, ETextCommi
 			{
 				const FScopedTransaction Transaction( LOCTEXT( "Rename Function", "Rename Function" ) );
 				FString NewNameString = NewText.ToString();
-				FBlueprintEditorUtils::RenameGraph(Graph, *NewNameString );
 
 				// Search through all function entry nodes for local variables to update their scope name
 				TArray<UK2Node_Variable*> VariableNodes;
-				FBlueprintEditorUtils::GetAllNodesOfClass<UK2Node_Variable>(FBlueprintEditorUtils::FindBlueprintForGraph(Graph), VariableNodes);
+				Graph->GetNodesOfClass<UK2Node_Variable>(VariableNodes);
 
 				for (UK2Node_Variable* const VariableNode : VariableNodes)
 				{
 					if(VariableNode->VariableReference.IsLocalScope())
 					{
 						// Update the variable's scope to be the graph's name (which mirrors the UFunction)
-						VariableNode->VariableReference.SetLocalMember(VariableNode->VariableReference.GetMemberName(), Graph->GetName(), VariableNode->VariableReference.GetMemberGuid());
+						VariableNode->VariableReference.SetLocalMember(VariableNode->VariableReference.GetMemberName(), NewNameString, VariableNode->VariableReference.GetMemberGuid());
 					}
 				}
+
+				FBlueprintEditorUtils::RenameGraph(Graph, *NewNameString );
 			}
 		}
 	}
