SKismetInspector now also displays struct properties marked 'show only inner'

diff --git a/Engine/Source/Editor/Kismet/Private/SKismetInspector.cpp b/Engine/Source/Editor/Kismet/Private/SKismetInspector.cpp
--- a/Engine/Source/Editor/Kismet/Private/SKismetInspector.cpp
+++ b/Engine/Source/Editor/Kismet/Private/SKismetInspector.cpp
@@ -408,6 +408,21 @@ void SKismetInspector::UpdateFromObjects(const TArray<UObject*>& PropertyObjects
 				check(Property != NULL);
 
 				SelectedObjectProperties.AddUnique(Property);
+
+				// If 'showing inners' on a struct, add them to the list as well
+				UStructProperty* StructProperty = Cast<UStructProperty>(Property);
+				static FName ShowOnlyInners("ShowOnlyInnerProperties");
+				if(	StructProperty != NULL && 
+					StructProperty->Struct != NULL && 
+					StructProperty->HasMetaData(ShowOnlyInners) )
+				{
+					for (TFieldIterator<UProperty> StructPropIt(StructProperty->Struct); StructPropIt; ++StructPropIt)
+					{
+						UProperty* InsideStructProperty = *StructPropIt;
+						check(InsideStructProperty != NULL);
+						SelectedObjectProperties.AddUnique(InsideStructProperty);
+					}
+				}
 			}
 
 			// Attempt to locate a matching property for the current component template
