#ue4 Blueprint node analytics update:  added in system to better track interesting analytics info on a per-node class basis (for Ben Cosh)

diff --git a/Engine/Plugins/Experimental/LiveEditor/Source/LiveEditor/Classes/K2Node_LiveEditObject.h b/Engine/Plugins/Experimental/LiveEditor/Source/LiveEditor/Classes/K2Node_LiveEditObject.h
--- a/Engine/Plugins/Experimental/LiveEditor/Source/LiveEditor/Classes/K2Node_LiveEditObject.h
+++ b/Engine/Plugins/Experimental/LiveEditor/Source/LiveEditor/Classes/K2Node_LiveEditObject.h
@@ -21,6 +21,7 @@ class UK2Node_LiveEditObject : public UK2Node
 	virtual void ReallocatePinsDuringReconstruction(TArray<UEdGraphPin*>& OldPins) OVERRIDE;
 	virtual class FNodeHandlingFunctor* CreateNodeHandler(class FKismetCompilerContext& CompilerContext) const OVERRIDE;
 	virtual void ExpandNode(class FKismetCompilerContext& CompilerContext, UEdGraph* SourceGraph) OVERRIDE;
+	virtual void GetNodeAttributes( TArray<TKeyValuePair<FString, FString>>& OutNodeAttributes ) const OVERRIDE;
 	// End UK2Node interface
 
 	/** Create new pins to show LiveEditable properties on archetype */
diff --git a/Engine/Plugins/Experimental/LiveEditor/Source/LiveEditor/Private/K2Node_LiveEditObject.cpp b/Engine/Plugins/Experimental/LiveEditor/Source/LiveEditor/Private/K2Node_LiveEditObject.cpp
--- a/Engine/Plugins/Experimental/LiveEditor/Source/LiveEditor/Private/K2Node_LiveEditObject.cpp
+++ b/Engine/Plugins/Experimental/LiveEditor/Source/LiveEditor/Private/K2Node_LiveEditObject.cpp
@@ -587,6 +587,17 @@ void UK2Node_LiveEditObject::ExpandNode(class FKismetCompilerContext& CompilerCo
 	}
 }
 
+void UK2Node_LiveEditObject::GetNodeAttributes( TArray<TKeyValuePair<FString, FString>>& OutNodeAttributes ) const
+{
+	UClass* ClassToSpawn = GetClassToSpawn();
+	const FString ClassToSpawnStr = ClassToSpawn ? ClassToSpawn->GetName() : TEXT( "UnknownClass" );
+	OutNodeAttributes.Add( TKeyValuePair<FString, FString>( TEXT( "Type" ), TEXT( "LiveEditObject" ) ));
+	OutNodeAttributes.Add( TKeyValuePair<FString, FString>( TEXT( "Class" ), GetClass()->GetName() ));
+	OutNodeAttributes.Add( TKeyValuePair<FString, FString>( TEXT( "Name" ), GetName() ));
+	OutNodeAttributes.Add( TKeyValuePair<FString, FString>( TEXT( "ClassToSpawn" ), ClassToSpawnStr ));
+	OutNodeAttributes.Add( TKeyValuePair<FString, FString>( TEXT( "EventName" ), GetEventName() ));
+}
+
 void UK2Node_LiveEditObject::CreatePinsForClass(UClass* InClass)
 {
 	check(InClass != NULL);
diff --git a/Engine/Source/Editor/AnimGraph/Classes/AnimGraphNode_Base.h b/Engine/Source/Editor/AnimGraph/Classes/AnimGraphNode_Base.h
--- a/Engine/Source/Editor/AnimGraph/Classes/AnimGraphNode_Base.h
+++ b/Engine/Source/Editor/AnimGraph/Classes/AnimGraphNode_Base.h
@@ -112,6 +112,7 @@ class UAnimGraphNode_Base : public UK2Node
 	ANIMGRAPH_API virtual void ReallocatePinsDuringReconstruction(TArray<UEdGraphPin*>& OldPins) OVERRIDE;
 	ANIMGRAPH_API virtual bool CanCreateUnderSpecifiedSchema(const UEdGraphSchema* DesiredSchema) const OVERRIDE;
 	ANIMGRAPH_API virtual void GetMenuEntries(FGraphContextMenuBuilder& ContextMenuBuilder) const OVERRIDE;
+	virtual void GetNodeAttributes( TArray<TKeyValuePair<FString, FString>>& OutNodeAttributes ) const OVERRIDE;
 	// End of UK2Node interface
 
 	// UAnimGraphNode_Base interface
diff --git a/Engine/Source/Editor/AnimGraph/Private/AnimGraphNode_Base.cpp b/Engine/Source/Editor/AnimGraph/Private/AnimGraphNode_Base.cpp
--- a/Engine/Source/Editor/AnimGraph/Private/AnimGraphNode_Base.cpp
+++ b/Engine/Source/Editor/AnimGraph/Private/AnimGraphNode_Base.cpp
@@ -219,6 +219,13 @@ void UAnimGraphNode_Base::GetMenuEntries(FGraphContextMenuBuilder& ContextMenuBu
 	CreateDefaultMenuEntry(ContextMenuBuilder);
 }
 
+void UAnimGraphNode_Base::GetNodeAttributes( TArray<TKeyValuePair<FString, FString>>& OutNodeAttributes ) const
+{
+	OutNodeAttributes.Add( TKeyValuePair<FString, FString>( TEXT( "Type" ), TEXT( "AnimGraphNode" ) ));
+	OutNodeAttributes.Add( TKeyValuePair<FString, FString>( TEXT( "Class" ), GetClass()->GetName() ));
+	OutNodeAttributes.Add( TKeyValuePair<FString, FString>( TEXT( "Name" ), GetName() ));
+}
+
 TSharedPtr<FEdGraphSchemaAction_K2NewNode> UAnimGraphNode_Base::CreateDefaultMenuEntry(FGraphContextMenuBuilder& ContextMenuBuilder) const
 {
 	UAnimGraphNode_Base* TemplateNode = NewObject<UAnimGraphNode_Base>(GetTransientPackage(), GetClass());
diff --git a/Engine/Source/Editor/BlueprintGraph/Classes/K2Node.h b/Engine/Source/Editor/BlueprintGraph/Classes/K2Node.h
--- a/Engine/Source/Editor/BlueprintGraph/Classes/K2Node.h
+++ b/Engine/Source/Editor/BlueprintGraph/Classes/K2Node.h
@@ -278,6 +278,10 @@ class UK2Node : public UEdGraphNode
 
 	/** This function if used for nodes that needs CDO for validation (Called before expansion)*/
 	BLUEPRINTGRAPH_API virtual void EarlyValidation(class FCompilerResultsLog& MessageLog) const {}
+
+	/** This function returns an arbitrary number of attributes that describe this node for analytics events */
+	BLUEPRINTGRAPH_API virtual void GetNodeAttributes( TArray<TKeyValuePair<FString, FString>>& OutNodeAttributes ) const;
+
 protected:
 	/** 
 	 * A mapping from old property and function names to new ones.  Get primed from INI files, and should contain entries for properties, functions, and delegates that get moved, so they can be fixed up
diff --git a/Engine/Source/Editor/BlueprintGraph/Classes/K2Node_AddDelegate.h b/Engine/Source/Editor/BlueprintGraph/Classes/K2Node_AddDelegate.h
--- a/Engine/Source/Editor/BlueprintGraph/Classes/K2Node_AddDelegate.h
+++ b/Engine/Source/Editor/BlueprintGraph/Classes/K2Node_AddDelegate.h
@@ -17,5 +17,6 @@ public:
 
 	// Begin of UK2Node interface
 	virtual class FNodeHandlingFunctor* CreateNodeHandler(class FKismetCompilerContext& CompilerContext) const OVERRIDE;
+	virtual void GetNodeAttributes( TArray<TKeyValuePair<FString, FString>>& OutNodeAttributes ) const OVERRIDE;
 	// End of UK2Node interface
 };
diff --git a/Engine/Source/Editor/BlueprintGraph/Classes/K2Node_BaseMCDelegate.h b/Engine/Source/Editor/BlueprintGraph/Classes/K2Node_BaseMCDelegate.h
--- a/Engine/Source/Editor/BlueprintGraph/Classes/K2Node_BaseMCDelegate.h
+++ b/Engine/Source/Editor/BlueprintGraph/Classes/K2Node_BaseMCDelegate.h
@@ -23,6 +23,7 @@ public:
 	virtual bool AllowMultipleSelfs(bool bInputAsArray) const OVERRIDE { return true; }
 	virtual void ExpandNode(class FKismetCompilerContext& CompilerContext, UEdGraph* SourceGraph) OVERRIDE;
 	virtual bool HasExternalBlueprintDependencies(TArray<class UStruct*>* OptionalOutput) const OVERRIDE;
+	virtual void GetNodeAttributes( TArray<TKeyValuePair<FString, FString>>& OutNodeAttributes ) const OVERRIDE;
 	// End of UK2Node interface
 
 	// UEdGraphNode interface
diff --git a/Engine/Source/Editor/BlueprintGraph/Classes/K2Node_CallFunction.h b/Engine/Source/Editor/BlueprintGraph/Classes/K2Node_CallFunction.h
--- a/Engine/Source/Editor/BlueprintGraph/Classes/K2Node_CallFunction.h
+++ b/Engine/Source/Editor/BlueprintGraph/Classes/K2Node_CallFunction.h
@@ -88,6 +88,7 @@ public:
 	virtual void ExpandNode(class FKismetCompilerContext& CompilerContext, UEdGraph* SourceGraph) OVERRIDE;
 	virtual FName GetCornerIcon() const OVERRIDE;
 	virtual FText GetToolTipHeading() const OVERRIDE;
+	virtual void GetNodeAttributes( TArray<TKeyValuePair<FString, FString>>& OutNodeAttributes ) const OVERRIDE;
 	// End of UK2Node interface
 
 	/** Returns the UFunction that this class is pointing to */
diff --git a/Engine/Source/Editor/BlueprintGraph/Classes/K2Node_Event.h b/Engine/Source/Editor/BlueprintGraph/Classes/K2Node_Event.h
--- a/Engine/Source/Editor/BlueprintGraph/Classes/K2Node_Event.h
+++ b/Engine/Source/Editor/BlueprintGraph/Classes/K2Node_Event.h
@@ -62,6 +62,7 @@ class UK2Node_Event : public UK2Node_EditablePinBase
 	BLUEPRINTGRAPH_API virtual FNodeHandlingFunctor* CreateNodeHandler(FKismetCompilerContext& CompilerContext) const OVERRIDE;
 	BLUEPRINTGRAPH_API virtual void ExpandNode(class FKismetCompilerContext& CompilerContext, UEdGraph* SourceGraph) OVERRIDE;
 	BLUEPRINTGRAPH_API virtual FText GetToolTipHeading() const;
+	virtual void GetNodeAttributes( TArray<TKeyValuePair<FString, FString>>& OutNodeAttributes ) const OVERRIDE;
 	// End UK2Node interface
 
 	/** Checks whether the parameters for this event node are compatible with the specified function entry node */
diff --git a/Engine/Source/Editor/BlueprintGraph/Classes/K2Node_MacroInstance.h b/Engine/Source/Editor/BlueprintGraph/Classes/K2Node_MacroInstance.h
--- a/Engine/Source/Editor/BlueprintGraph/Classes/K2Node_MacroInstance.h
+++ b/Engine/Source/Editor/BlueprintGraph/Classes/K2Node_MacroInstance.h
@@ -55,6 +55,7 @@ public:
 	virtual void PostReconstructNode() OVERRIDE;
 	virtual FText GetActiveBreakpointToolTipText() const OVERRIDE;
 	virtual bool HasExternalBlueprintDependencies(TArray<class UStruct*>* OptionalOutput) const OVERRIDE;
+	virtual void GetNodeAttributes( TArray<TKeyValuePair<FString, FString>>& OutNodeAttributes ) const OVERRIDE;
 	// End UK2Node interface
 
 	void SetMacroGraph(UEdGraph* Graph) { MacroGraphReference.SetGraph(Graph); }
diff --git a/Engine/Source/Editor/BlueprintGraph/Classes/K2Node_SpawnActor.h b/Engine/Source/Editor/BlueprintGraph/Classes/K2Node_SpawnActor.h
--- a/Engine/Source/Editor/BlueprintGraph/Classes/K2Node_SpawnActor.h
+++ b/Engine/Source/Editor/BlueprintGraph/Classes/K2Node_SpawnActor.h
@@ -30,6 +30,7 @@ class UK2Node_SpawnActor : public UK2Node
 	virtual class FNodeHandlingFunctor* CreateNodeHandler(class FKismetCompilerContext& CompilerContext) const OVERRIDE;
 	virtual void ExpandNode(class FKismetCompilerContext& CompilerContext, UEdGraph* SourceGraph) OVERRIDE;
 	virtual bool HasExternalBlueprintDependencies(TArray<class UStruct*>* OptionalOutput) const OVERRIDE;
+	virtual void GetNodeAttributes( TArray<TKeyValuePair<FString, FString>>& OutNodeAttributes ) const OVERRIDE;
 	// End UK2Node interface
 
 
diff --git a/Engine/Source/Editor/BlueprintGraph/Classes/K2Node_SpawnActorFromClass.h b/Engine/Source/Editor/BlueprintGraph/Classes/K2Node_SpawnActorFromClass.h
--- a/Engine/Source/Editor/BlueprintGraph/Classes/K2Node_SpawnActorFromClass.h
+++ b/Engine/Source/Editor/BlueprintGraph/Classes/K2Node_SpawnActorFromClass.h
@@ -26,6 +26,7 @@ class UK2Node_SpawnActorFromClass : public UK2Node
 	virtual bool IsNodeSafeToIgnore() const OVERRIDE { return true; }
 	virtual void ReallocatePinsDuringReconstruction(TArray<UEdGraphPin*>& OldPins) OVERRIDE;
 	virtual bool CanPasteHere(const UEdGraph* TargetGraph, const UEdGraphSchema* Schema) const OVERRIDE;
+	virtual void GetNodeAttributes( TArray<TKeyValuePair<FString, FString>>& OutNodeAttributes ) const OVERRIDE;
 	// End UK2Node interface
 
 
diff --git a/Engine/Source/Editor/BlueprintGraph/Classes/K2Node_Timeline.h b/Engine/Source/Editor/BlueprintGraph/Classes/K2Node_Timeline.h
--- a/Engine/Source/Editor/BlueprintGraph/Classes/K2Node_Timeline.h
+++ b/Engine/Source/Editor/BlueprintGraph/Classes/K2Node_Timeline.h
@@ -51,6 +51,7 @@ class UK2Node_Timeline : public UK2Node
 	virtual bool NodeCausesStructuralBlueprintChange() const OVERRIDE { return true; }
 	virtual class FNodeHandlingFunctor* CreateNodeHandler(class FKismetCompilerContext& CompilerContext) const OVERRIDE;
 	virtual void ExpandNode(class FKismetCompilerContext& CompilerContext, UEdGraph* SourceGraph) OVERRIDE;
+	virtual void GetNodeAttributes( TArray<TKeyValuePair<FString, FString>>& OutNodeAttributes ) const OVERRIDE;
 	// End UK2Node interface.
 
 	/** Get the 'play' input pin */
diff --git a/Engine/Source/Editor/BlueprintGraph/Classes/K2Node_Variable.h b/Engine/Source/Editor/BlueprintGraph/Classes/K2Node_Variable.h
--- a/Engine/Source/Editor/BlueprintGraph/Classes/K2Node_Variable.h
+++ b/Engine/Source/Editor/BlueprintGraph/Classes/K2Node_Variable.h
@@ -60,6 +60,7 @@ public:
 	virtual ERedirectType DoPinsMatchForReconstruction(const UEdGraphPin* NewPin, int32 NewPinIndex, const UEdGraphPin* OldPin, int32 OldPinIndex)  const OVERRIDE;
 	virtual void ValidateNodeDuringCompilation(class FCompilerResultsLog& MessageLog) const OVERRIDE;
 	virtual FText GetToolTipHeading() const OVERRIDE;
+	virtual void GetNodeAttributes( TArray<TKeyValuePair<FString, FString>>& OutNodeAttributes ) const OVERRIDE;
 	// End K2Node interface
 
 	/** Set up this variable node from the supplied UProperty */
diff --git a/Engine/Source/Editor/BlueprintGraph/Private/EdGraphSchema_K2_Actions.cpp b/Engine/Source/Editor/BlueprintGraph/Private/EdGraphSchema_K2_Actions.cpp
--- a/Engine/Source/Editor/BlueprintGraph/Private/EdGraphSchema_K2_Actions.cpp
+++ b/Engine/Source/Editor/BlueprintGraph/Private/EdGraphSchema_K2_Actions.cpp
@@ -103,25 +103,8 @@ UEdGraphNode* FEdGraphSchemaAction_K2NewNode::CreateNode(class UEdGraph* ParentG
 	// auto-wire creates a conversion node to put between them)
 	ResultNode->AutowireNewNode(FromPinPtr);
 
-	// Update Analytics for these nodes
-	if( UK2Node_CallFunction* CallFunctionNode = Cast<UK2Node_CallFunction>( ResultNode ))
-	{
-		if( UFunction* TargetFunction = CallFunctionNode->GetTargetFunction() )
-		{
-			FBlueprintEditorUtils::AnalyticsTrackNewNode( ResultNode, FName( *UK2Node_CallFunction::StaticClass()->GetName()), *TargetFunction->GetName() );
-		}
-	}
-	else if( UK2Node_MacroInstance* MacroNode = Cast<UK2Node_MacroInstance>( ResultNode ))
-	{
-		if( UEdGraph* MacroGraph = MacroNode->GetMacroGraph() )
-		{
-			FBlueprintEditorUtils::AnalyticsTrackNewNode( ResultNode, FName( *UK2Node_MacroInstance::StaticClass()->GetName()), FName( *MacroGraph->GetName() ));
-		}
-	}
-	else if( UK2Node_Event* EventNode = Cast<UK2Node_Event>( ResultNode ))
-	{
-		FBlueprintEditorUtils::AnalyticsTrackNewNode( ResultNode, FName( *UK2Node_Event::StaticClass()->GetName() ), EventNode->GetFunctionName() );
-	}
+	// Update Analytics for the new nodes
+	FBlueprintEditorUtils::AnalyticsTrackNewNode( ResultNode );
 	// NOTE: At this point the node may have been reconstructed, depending on node type!
 
 	return ResultNode;
@@ -615,6 +598,9 @@ UEdGraphNode* FEdGraphSchemaAction_K2AddComment::PerformAction(class UEdGraph* P
 
 	UEdGraphNode* NewNode = FEdGraphSchemaAction_NewNode::SpawnNodeFromTemplate<UEdGraphNode_Comment>(ParentGraph, CommentTemplate, SpawnLocation, bSelectNewNode);
 
+	// Update Analytics for these nodes
+	FBlueprintEditorUtils::AnalyticsTrackNewNode( NewNode );
+
 	// Mark Blueprint as structurally modified since
 	// UK2Node_Comment::NodeCausesStructuralBlueprintChange used to return true
 	FBlueprintEditorUtils::MarkBlueprintAsStructurallyModified(Blueprint);
diff --git a/Engine/Source/Editor/BlueprintGraph/Private/K2Node.cpp b/Engine/Source/Editor/BlueprintGraph/Private/K2Node.cpp
--- a/Engine/Source/Editor/BlueprintGraph/Private/K2Node.cpp
+++ b/Engine/Source/Editor/BlueprintGraph/Private/K2Node.cpp
@@ -191,6 +191,13 @@ void UK2Node::InsertNewNode(UEdGraphPin* FromPin, UEdGraphPin* NewLinkPin, TSet<
 	}
 }
 
+void UK2Node::GetNodeAttributes( TArray<TKeyValuePair<FString, FString>>& OutNodeAttributes ) const
+{
+	OutNodeAttributes.Add( TKeyValuePair<FString, FString>( TEXT( "Type" ), TEXT( "GraphNode" ) ));
+	OutNodeAttributes.Add( TKeyValuePair<FString, FString>( TEXT( "Class" ), GetClass()->GetName() ));
+	OutNodeAttributes.Add( TKeyValuePair<FString, FString>( TEXT( "Name" ), GetName() ));
+}
+
 void UK2Node::PinConnectionListChanged(UEdGraphPin* Pin) 
 {
 	// If the pin has been connected, clear the default values so that we don't hold on to references
diff --git a/Engine/Source/Editor/BlueprintGraph/Private/K2Node_CallFunction.cpp b/Engine/Source/Editor/BlueprintGraph/Private/K2Node_CallFunction.cpp
--- a/Engine/Source/Editor/BlueprintGraph/Private/K2Node_CallFunction.cpp
+++ b/Engine/Source/Editor/BlueprintGraph/Private/K2Node_CallFunction.cpp
@@ -1477,6 +1477,15 @@ FText UK2Node_CallFunction::GetToolTipHeading() const
 	return HeadingBuilder.ConstructedHeading;
 }
 
+void UK2Node_CallFunction::GetNodeAttributes( TArray<TKeyValuePair<FString, FString>>& OutNodeAttributes ) const
+{
+	UFunction* TargetFunction = GetTargetFunction();
+	const FString TargetFunctionName = TargetFunction ? TargetFunction->GetName() : TEXT( "InvalidFunction" );
+	OutNodeAttributes.Add( TKeyValuePair<FString, FString>( TEXT( "Type" ), TEXT( "Function" ) ));
+	OutNodeAttributes.Add( TKeyValuePair<FString, FString>( TEXT( "Class" ), GetClass()->GetName() ));
+	OutNodeAttributes.Add( TKeyValuePair<FString, FString>( TEXT( "Name" ), TargetFunctionName ));
+}
+
 bool UK2Node_CallFunction::HasExternalBlueprintDependencies(TArray<class UStruct*>* OptionalOutput) const
 {
 	const UClass* SourceClass = FunctionReference.GetMemberParentClass(this);
diff --git a/Engine/Source/Editor/BlueprintGraph/Private/K2Node_Event.cpp b/Engine/Source/Editor/BlueprintGraph/Private/K2Node_Event.cpp
--- a/Engine/Source/Editor/BlueprintGraph/Private/K2Node_Event.cpp
+++ b/Engine/Source/Editor/BlueprintGraph/Private/K2Node_Event.cpp
@@ -657,6 +657,13 @@ FText UK2Node_Event::GetToolTipHeading() const
 	return CompleteHeading;
 }
 
+void UK2Node_Event::GetNodeAttributes( TArray<TKeyValuePair<FString, FString>>& OutNodeAttributes ) const
+{
+	OutNodeAttributes.Add( TKeyValuePair<FString, FString>( TEXT( "Type" ), TEXT( "Event" ) ));
+	OutNodeAttributes.Add( TKeyValuePair<FString, FString>( TEXT( "Class" ), GetClass()->GetName() ));
+	OutNodeAttributes.Add( TKeyValuePair<FString, FString>( TEXT( "Name" ), GetFunctionName().ToString() ));
+}
+
 bool UK2Node_Event::IsDeprecated() const
 {
 	if (UFunction* Function = FindField<UFunction>(EventSignatureClass, EventSignatureName))
diff --git a/Engine/Source/Editor/BlueprintGraph/Private/K2Node_MCDelegate.cpp b/Engine/Source/Editor/BlueprintGraph/Private/K2Node_MCDelegate.cpp
--- a/Engine/Source/Editor/BlueprintGraph/Private/K2Node_MCDelegate.cpp
+++ b/Engine/Source/Editor/BlueprintGraph/Private/K2Node_MCDelegate.cpp
@@ -185,6 +185,13 @@ bool UK2Node_BaseMCDelegate::HasExternalBlueprintDependencies(TArray<class UStru
 	return bResult || Super::HasExternalBlueprintDependencies(OptionalOutput);
 }
 
+void UK2Node_BaseMCDelegate::GetNodeAttributes( TArray<TKeyValuePair<FString, FString>>& OutNodeAttributes ) const
+{
+	OutNodeAttributes.Add( TKeyValuePair<FString, FString>( TEXT( "Type" ), TEXT( "EventDelegate" ) ));
+	OutNodeAttributes.Add( TKeyValuePair<FString, FString>( TEXT( "Class" ), GetClass()->GetName() ));
+	OutNodeAttributes.Add( TKeyValuePair<FString, FString>( TEXT( "Name" ), GetPropertyName().ToString() ));
+}
+
 /////// UK2Node_AddDelegate ///////////
 
 UK2Node_AddDelegate::UK2Node_AddDelegate(const class FPostConstructInitializeProperties& PCIP)
@@ -208,8 +215,8 @@ void UK2Node_AddDelegate::AllocateDefaultPins()
 FText UK2Node_AddDelegate::GetNodeTitle(ENodeTitleType::Type TitleType) const
 {
 	FFormatNamedArguments Args;
-	Args.Add(TEXT("ProeprtyName"), FText::FromName(GetPropertyName()));
-	return FText::Format(NSLOCTEXT("K2Node", "AddDelegate", "Bind Event to {ProeprtyName}"), Args);
+	Args.Add(TEXT("PropertyName"), FText::FromName(GetPropertyName()));
+	return FText::Format(NSLOCTEXT("K2Node", "AddDelegate", "Bind Event to {PropertyName}"), Args);
 }
 
 FString UK2Node_AddDelegate::GetNodeNativeTitle(ENodeTitleType::Type TitleType) const
@@ -224,6 +231,13 @@ FNodeHandlingFunctor* UK2Node_AddDelegate::CreateNodeHandler(FKismetCompilerCont
 	return new FKCHandler_AddRemoveDelegate(CompilerContext, KCST_AddMulticastDelegate);
 }
 
+void UK2Node_AddDelegate::GetNodeAttributes( TArray<TKeyValuePair<FString, FString>>& OutNodeAttributes ) const
+{
+	OutNodeAttributes.Add( TKeyValuePair<FString, FString>( TEXT( "Type" ), TEXT( "AddDelegate" ) ));
+	OutNodeAttributes.Add( TKeyValuePair<FString, FString>( TEXT( "Class" ), GetClass()->GetName() ));
+	OutNodeAttributes.Add( TKeyValuePair<FString, FString>( TEXT( "Name" ), GetPropertyName().ToString() ));
+}
+
 /////// UK2Node_ClearDelegate ///////////
 
 UK2Node_ClearDelegate::UK2Node_ClearDelegate(const class FPostConstructInitializeProperties& PCIP)
@@ -359,4 +373,4 @@ FNodeHandlingFunctor* UK2Node_CallDelegate::CreateNodeHandler(FKismetCompilerCon
 FName UK2Node_CallDelegate::GetCornerIcon() const
 {
 	return TEXT("Graph.Message.MessageIcon");
-}
\ No newline at end of file
+}
diff --git a/Engine/Source/Editor/BlueprintGraph/Private/K2Node_MacroInstance.cpp b/Engine/Source/Editor/BlueprintGraph/Private/K2Node_MacroInstance.cpp
--- a/Engine/Source/Editor/BlueprintGraph/Private/K2Node_MacroInstance.cpp
+++ b/Engine/Source/Editor/BlueprintGraph/Private/K2Node_MacroInstance.cpp
@@ -423,4 +423,19 @@ bool UK2Node_MacroInstance::HasExternalBlueprintDependencies(TArray<class UStruc
 	return bResult;
 }
 
+void UK2Node_MacroInstance::GetNodeAttributes( TArray<TKeyValuePair<FString, FString>>& OutNodeAttributes ) const
+{
+	FString MacroName( TEXT( "InvalidMacro" ));
+
+	if( UEdGraph* MacroGraph = GetMacroGraph() )
+	{
+		MacroName = MacroGraph->GetName();
+	}
+
+	OutNodeAttributes.Add( TKeyValuePair<FString, FString>( TEXT( "Type" ), TEXT( "Macro" ) ));
+	OutNodeAttributes.Add( TKeyValuePair<FString, FString>( TEXT( "Class" ), GetClass()->GetName() ));
+	OutNodeAttributes.Add( TKeyValuePair<FString, FString>( TEXT( "Name" ), MacroName ));
+}
+
+
 #undef LOCTEXT_NAMESPACE
diff --git a/Engine/Source/Editor/BlueprintGraph/Private/K2Node_SpawnActor.cpp b/Engine/Source/Editor/BlueprintGraph/Private/K2Node_SpawnActor.cpp
--- a/Engine/Source/Editor/BlueprintGraph/Private/K2Node_SpawnActor.cpp
+++ b/Engine/Source/Editor/BlueprintGraph/Private/K2Node_SpawnActor.cpp
@@ -483,6 +483,16 @@ bool UK2Node_SpawnActor::HasExternalBlueprintDependencies(TArray<class UStruct*>
 	return bResult || Super::HasExternalBlueprintDependencies(OptionalOutput);
 }
 
+void UK2Node_SpawnActor::GetNodeAttributes( TArray<TKeyValuePair<FString, FString>>& OutNodeAttributes ) const
+{
+	UClass* ClassToSpawn = GetClassToSpawn();
+	const FString ClassToSpawnStr = ClassToSpawn ? ClassToSpawn->GetName() : TEXT( "InvalidClass" );
+	OutNodeAttributes.Add( TKeyValuePair<FString, FString>( TEXT( "Type" ), TEXT( "SpawnActor" ) ));
+	OutNodeAttributes.Add( TKeyValuePair<FString, FString>( TEXT( "Class" ), GetClass()->GetName() ));
+	OutNodeAttributes.Add( TKeyValuePair<FString, FString>( TEXT( "Name" ), GetName() ));
+	OutNodeAttributes.Add( TKeyValuePair<FString, FString>( TEXT( "ActorClass" ), ClassToSpawnStr ));
+}
+
 bool UK2Node_SpawnActor::IsDeprecated() const
 {
 	return false;
diff --git a/Engine/Source/Editor/BlueprintGraph/Private/K2Node_SpawnActorFromClass.cpp b/Engine/Source/Editor/BlueprintGraph/Private/K2Node_SpawnActorFromClass.cpp
--- a/Engine/Source/Editor/BlueprintGraph/Private/K2Node_SpawnActorFromClass.cpp
+++ b/Engine/Source/Editor/BlueprintGraph/Private/K2Node_SpawnActorFromClass.cpp
@@ -302,6 +302,16 @@ bool UK2Node_SpawnActorFromClass::CanPasteHere( const UEdGraph* TargetGraph, con
 	return CanCreateUnderSpecifiedSchema(Schema) && (!Blueprint || FBlueprintEditorUtils::FindUserConstructionScript(Blueprint) != TargetGraph);
 }
 
+void UK2Node_SpawnActorFromClass::GetNodeAttributes( TArray<TKeyValuePair<FString, FString>>& OutNodeAttributes ) const
+{
+	UClass* ClassToSpawn = GetClassToSpawn();
+	const FString ClassToSpawnStr = ClassToSpawn ? ClassToSpawn->GetName() : TEXT( "InvalidClass" );
+	OutNodeAttributes.Add( TKeyValuePair<FString, FString>( TEXT( "Type" ), TEXT( "SpawnActorFromClass" ) ));
+	OutNodeAttributes.Add( TKeyValuePair<FString, FString>( TEXT( "Class" ), GetClass()->GetName() ));
+	OutNodeAttributes.Add( TKeyValuePair<FString, FString>( TEXT( "Name" ), GetName() ));
+	OutNodeAttributes.Add( TKeyValuePair<FString, FString>( TEXT( "ActorClass" ), ClassToSpawnStr ));
+}
+
 void UK2Node_SpawnActorFromClass::ExpandNode(class FKismetCompilerContext& CompilerContext, UEdGraph* SourceGraph)
 {
 	Super::ExpandNode(CompilerContext, SourceGraph);
diff --git a/Engine/Source/Editor/BlueprintGraph/Private/K2Node_Timeline.cpp b/Engine/Source/Editor/BlueprintGraph/Private/K2Node_Timeline.cpp
--- a/Engine/Source/Editor/BlueprintGraph/Private/K2Node_Timeline.cpp
+++ b/Engine/Source/Editor/BlueprintGraph/Private/K2Node_Timeline.cpp
@@ -546,4 +546,11 @@ FString UK2Node_Timeline::GetDocumentationExcerptName() const
 	return TEXT("UK2Node_Timeline");
 }
 
+void UK2Node_Timeline::GetNodeAttributes( TArray<TKeyValuePair<FString, FString>>& OutNodeAttributes ) const
+{
+	OutNodeAttributes.Add( TKeyValuePair<FString, FString>( TEXT( "Type" ), TEXT( "TimeLine" ) ));
+	OutNodeAttributes.Add( TKeyValuePair<FString, FString>( TEXT( "Class" ), GetClass()->GetName() ));
+	OutNodeAttributes.Add( TKeyValuePair<FString, FString>( TEXT( "Name" ), GetName() ));
+}
+
 #undef LOCTEXT_NAMESPACE
diff --git a/Engine/Source/Editor/BlueprintGraph/Private/K2Node_Variable.cpp b/Engine/Source/Editor/BlueprintGraph/Private/K2Node_Variable.cpp
--- a/Engine/Source/Editor/BlueprintGraph/Private/K2Node_Variable.cpp
+++ b/Engine/Source/Editor/BlueprintGraph/Private/K2Node_Variable.cpp
@@ -318,6 +318,15 @@ FText UK2Node_Variable::GetToolTipHeading() const
 	return Heading;
 }
 
+void UK2Node_Variable::GetNodeAttributes( TArray<TKeyValuePair<FString, FString>>& OutNodeAttributes ) const
+{
+	UProperty* VariableProperty = GetPropertyForVariable();
+	const FString VariableName = VariableProperty ? VariableProperty->GetName() : TEXT( "InvalidVariable" );
+	OutNodeAttributes.Add( TKeyValuePair<FString, FString>( TEXT( "Type" ), TEXT( "Variable" ) ));
+	OutNodeAttributes.Add( TKeyValuePair<FString, FString>( TEXT( "Class" ), GetClass()->GetName() ));
+	OutNodeAttributes.Add( TKeyValuePair<FString, FString>( TEXT( "Name" ), VariableName ));
+}
+
 FName UK2Node_Variable::GetVariableIconAndColor(UStruct* VarScope, FName VarName, FLinearColor& IconColorOut)
 {
 	FName IconBrush = TEXT("Kismet.AllClasses.VariableIcon");
diff --git a/Engine/Source/Editor/Kismet/Private/BPDelegateDragDropAction.cpp b/Engine/Source/Editor/Kismet/Private/BPDelegateDragDropAction.cpp
--- a/Engine/Source/Editor/Kismet/Private/BPDelegateDragDropAction.cpp
+++ b/Engine/Source/Editor/Kismet/Private/BPDelegateDragDropAction.cpp
@@ -27,7 +27,12 @@ void FKismetDelegateDragDropAction::MakeEvent(FNodeConstructionParams Params)
 		Params.Graph->Modify();
 
 		const FString FunctionName = FString::Printf(TEXT("%s_Event"), *Params.Property->GetName());
-		UK2Node_CustomEvent::CreateFromFunction(Params.GraphPosition, Params.Graph, FunctionName, SignatureFunction);
+		UK2Node_CustomEvent* NewNode = UK2Node_CustomEvent::CreateFromFunction(Params.GraphPosition, Params.Graph, FunctionName, SignatureFunction);
+
+		if( NewNode )
+		{
+			FBlueprintEditorUtils::AnalyticsTrackNewNode( NewNode );
+		}
 
 		UBlueprint* Blueprint = FBlueprintEditorUtils::FindBlueprintForGraphChecked(Params.Graph);
 		FBlueprintEditorUtils::MarkBlueprintAsStructurallyModified(Blueprint);
diff --git a/Engine/Source/Editor/Kismet/Private/BlueprintEditor.cpp b/Engine/Source/Editor/Kismet/Private/BlueprintEditor.cpp
--- a/Engine/Source/Editor/Kismet/Private/BlueprintEditor.cpp
+++ b/Engine/Source/Editor/Kismet/Private/BlueprintEditor.cpp
@@ -312,6 +312,43 @@ FGraphPanelSelectionSet FBlueprintEditor::GetSelectedNodes() const
 	return CurrentSelection;
 }
 
+void FBlueprintEditor::AnalyticsTrackNodeEvent( UBlueprint* Blueprint, UEdGraphNode *GraphNode, bool bNodeDelete ) const
+{
+	if( Blueprint && GraphNode && FEngineAnalytics::IsAvailable() )
+	{
+		// Build Node Details
+		const UGeneralProjectSettings& ProjectSettings = *GetDefault<UGeneralProjectSettings>();
+		FString ProjectID = ProjectSettings.ProjectID.ToString();
+		TArray< FAnalyticsEventAttribute > NodeAttributes;
+		NodeAttributes.Add( FAnalyticsEventAttribute( TEXT( "ProjectId" ), ProjectID ) );
+		NodeAttributes.Add( FAnalyticsEventAttribute( TEXT( "BlueprintId" ), Blueprint->GetBlueprintGuid().ToString() ) );
+		TArray< TKeyValuePair<FString, FString> > Attributes;
+
+		if( UK2Node* K2Node = Cast<UK2Node>( GraphNode ))
+		{
+			K2Node->GetNodeAttributes( Attributes );
+		}
+		else if( UEdGraphNode_Comment* CommentNode = Cast<UEdGraphNode_Comment>( GraphNode ))
+		{
+			Attributes.Add( TKeyValuePair<FString, FString>( TEXT( "Type" ), TEXT( "Comment" ) ));
+			Attributes.Add( TKeyValuePair<FString, FString>( TEXT( "Class" ), CommentNode->GetClass()->GetName() ));
+			Attributes.Add( TKeyValuePair<FString, FString>( TEXT( "Name" ), CommentNode->GetName() ));
+		}
+		if( Attributes.Num() > 0 )
+		{
+			// Build Node Attributes
+			for( auto Iter = Attributes.CreateConstIterator(); Iter; ++Iter )
+			{
+				NodeAttributes.Add( FAnalyticsEventAttribute( Iter->Key, Iter->Value ));
+			}
+			// Send Analytics event
+			FString EventType = bNodeDelete ?	TEXT( "Editor.Usage.BlueprintEditor.NodeDeleted" ) :
+												TEXT( "Editor.Usage.BlueprintEditor.NodeCreated" );
+			FEngineAnalytics::GetProvider().RecordEvent( EventType, NodeAttributes );
+		}
+	}
+}
+
 void FBlueprintEditor::RefreshEditors()
 {
 	DocumentManager->CleanInvalidTabs();
@@ -1437,17 +1474,7 @@ FBlueprintEditor::~FBlueprintEditor()
 		BPEditorAttribs.Add( FAnalyticsEventAttribute( FString( "PastedNodesCreated" ), AnalyticsStats.NodePasteCreateCount ));
 
 		BPEditorAttribs.Add( FAnalyticsEventAttribute( FString( "ProjectId" ), ProjectID ) );
-
-		TArray< FAnalyticsEventAttribute > BPEditorNodeStats;
-		for( auto Iter = AnalyticsStats.CreatedNodeTypes.CreateIterator(); Iter; ++Iter )
-		{
-			const FString NodeType = FString::Printf( TEXT( "%s.%s" ), *Iter->Value.NodeClass.ToString(), *Iter->Key.ToString() ); 
-			BPEditorNodeStats.Add( FAnalyticsEventAttribute( NodeType, Iter->Value.Instances ));
-		}
-		AnalyticsStats.CreatedNodeTypes.Reset();
-
 		FEngineAnalytics::GetProvider().RecordEvent( FString( "Editor.Usage.BlueprintEditor" ), BPEditorAttribs );
-		FEngineAnalytics::GetProvider().RecordEvent( FString( "Editor.Usage.BlueprintEditor.NodeCreation" ), BPEditorNodeStats );
 
 		for (auto Iter = AnalyticsStats.GraphDisallowedPinConnections.CreateConstIterator(); Iter; ++Iter)
 		{
@@ -2034,21 +2061,6 @@ void FBlueprintEditor::PostRedo(bool bSuccess)
 	}
 }
 
-void FBlueprintEditor::AnalyticsTrackNewNode( FName NodeClass, FName NodeType )
-{
-	if( FAnalyticsStatistics::FNodeDetails* Result = AnalyticsStats.CreatedNodeTypes.Find( NodeType ))
-	{
-		Result->Instances++;
-	}
-	else
-	{
-		FAnalyticsStatistics::FNodeDetails NewEntry;
-		NewEntry.NodeClass = NodeClass;
-		NewEntry.Instances = 1;
-		AnalyticsStats.CreatedNodeTypes.Add( NodeType, NewEntry );
-	}
-}
-
 void FBlueprintEditor::UndoGraphAction()
 {
 	GEditor->UndoTransaction();
@@ -3868,7 +3880,7 @@ void FBlueprintEditor::DeleteSelectedNodes()
 				{
 					DocumentManager->CleanInvalidTabs();
 				}
-
+				AnalyticsTrackNodeEvent( GetBlueprintObj(), Node, true );
 				FBlueprintEditorUtils::RemoveNode(GetBlueprintObj(), Node, true);
 			}
 		}
@@ -4257,6 +4269,8 @@ void FBlueprintEditor::PasteNodesHere(class UEdGraph* DestinationGraph, const FV
 		{
 			bNeedToModifyStructurally = true;
 		}
+		// Log new node created to analytics
+		AnalyticsTrackNodeEvent( GetBlueprintObj(), Node, false );
 	}
 
 	if (bNeedToModifyStructurally)
diff --git a/Engine/Source/Editor/Kismet/Public/BlueprintEditor.h b/Engine/Source/Editor/Kismet/Public/BlueprintEditor.h
--- a/Engine/Source/Editor/Kismet/Public/BlueprintEditor.h
+++ b/Engine/Source/Editor/Kismet/Public/BlueprintEditor.h
@@ -154,6 +154,7 @@ public:
 	virtual TArray<TSharedPtr<class FSCSEditorTreeNode> > GetSelectedSCSEditorTreeNodes() const OVERRIDE;
 	virtual TSharedPtr<class FSCSEditorTreeNode> FindAndSelectSCSEditorTreeNode(const UActorComponent* InComponent, bool IsCntrlDown) OVERRIDE;
 	virtual int32 GetNumberOfSelectedNodes() const OVERRIDE;
+	virtual void AnalyticsTrackNodeEvent( UBlueprint* Blueprint, UEdGraphNode *GraphNode, bool bNodeDelete = false ) const OVERRIDE;
 	// End of IBlueprintEditor interface
 
 	// FTickableEditorObject interface
@@ -789,8 +790,6 @@ protected:
 	virtual void	PostRedo(bool bSuccess) OVERRIDE;
 	// End of FEditorUndoClient
 
-	virtual void AnalyticsTrackNewNode( FName NodeClass, FName NodeType ) OVERRIDE;
-
 	/** Get graph appearance */
 	virtual FGraphAppearanceInfo GetGraphAppearance() const;
 
@@ -975,9 +974,6 @@ private:
 			int32 Instances;
 		};
 
-		/** Stats collected about class/type/instances of user created nodes */
-		TMap<FName,FNodeDetails> CreatedNodeTypes;
-
 		/** Stats collected about warning tooltips */
 		TArray<FDisallowedPinConnection> GraphDisallowedPinConnections;
 	};
diff --git a/Engine/Source/Editor/Kismet/Public/BlueprintEditorModule.h b/Engine/Source/Editor/Kismet/Public/BlueprintEditorModule.h
--- a/Engine/Source/Editor/Kismet/Public/BlueprintEditorModule.h
+++ b/Engine/Source/Editor/Kismet/Public/BlueprintEditorModule.h
@@ -52,8 +52,8 @@ public:
 	/** Find and select a specific SCS editor tree node associated with the given component */
 	virtual TSharedPtr<class FSCSEditorTreeNode> FindAndSelectSCSEditorTreeNode(const class UActorComponent* InComponent, bool IsCntrlDown) = 0;
 
-	/** Used to track node class creation ( event, call function, macro ) and the type */
-	virtual void AnalyticsTrackNewNode( FName NodeClass, FName NodeType ) = 0;
+	/** Used to track node create/delete events for Analytics */
+	virtual void AnalyticsTrackNodeEvent( UBlueprint* Blueprint, UEdGraphNode *GraphNode, bool bNodeDelete = false ) const = 0;
 
 };
 
diff --git a/Engine/Source/Editor/UnrealEd/Private/Kismet2/BlueprintEditorUtils.cpp b/Engine/Source/Editor/UnrealEd/Private/Kismet2/BlueprintEditorUtils.cpp
--- a/Engine/Source/Editor/UnrealEd/Private/Kismet2/BlueprintEditorUtils.cpp
+++ b/Engine/Source/Editor/UnrealEd/Private/Kismet2/BlueprintEditorUtils.cpp
@@ -1380,6 +1380,9 @@ void FBlueprintEditorUtils::PostDuplicateBlueprint(UBlueprint* Blueprint)
 			UEditorEngine::CopyPropertiesForUnrelatedObjects(OldCDO, NewCDO);
 		}
 
+		// Create a new blueprint guid
+		Blueprint->GenerateNewGuid();
+
 		// And compile again to make sure they go into the generated class, get cleaned up, etc...
 		FKismetEditorUtilities::CompileBlueprint(Blueprint, false, true);
 
@@ -5493,14 +5496,14 @@ void FBlueprintEditorUtils::FindAndSetDebuggableBlueprintInstances()
 	}
 }
 
-void FBlueprintEditorUtils::AnalyticsTrackNewNode( UEdGraphNode* NewNode, FName NodeClass, FName NodeType )
+void FBlueprintEditorUtils::AnalyticsTrackNewNode( UEdGraphNode *NewNode )
 {
 	UBlueprint* Blueprint = FindBlueprintForNodeChecked(NewNode);
-	TSharedPtr<IToolkit> FoundAssetEditor = FToolkitManager::Get().FindEditorForAsset(Blueprint); 
+	TSharedPtr<IToolkit> FoundAssetEditor = FToolkitManager::Get().FindEditorForAsset(Blueprint);
 	if (FoundAssetEditor.IsValid() && FoundAssetEditor->IsBlueprintEditor()) 
 	{ 
 		TSharedPtr<IBlueprintEditor> BlueprintEditor = StaticCastSharedPtr<IBlueprintEditor>(FoundAssetEditor);
-		BlueprintEditor->AnalyticsTrackNewNode(NodeClass, NodeType);
+		BlueprintEditor->AnalyticsTrackNodeEvent(Blueprint, NewNode, false);
 	}
 }
 
diff --git a/Engine/Source/Editor/UnrealEd/Private/Kismet2/Kismet2.cpp b/Engine/Source/Editor/UnrealEd/Private/Kismet2/Kismet2.cpp
--- a/Engine/Source/Editor/UnrealEd/Private/Kismet2/Kismet2.cpp
+++ b/Engine/Source/Editor/UnrealEd/Private/Kismet2/Kismet2.cpp
@@ -194,6 +194,7 @@ UBlueprint* FKismetEditorUtilities::CreateBlueprint(UClass* ParentClass, UObject
 	NewBP->BlueprintSystemVersion = UBlueprint::GetCurrentBlueprintSystemVersion();
 	NewBP->bIsNewlyCreated = true;
 	NewBP->bLegacyNeedToPurgeSkelRefs = false;
+	NewBP->GenerateNewGuid();
 
 	const UEdGraphSchema_K2* K2Schema = GetDefault<UEdGraphSchema_K2>();
 
@@ -291,6 +292,7 @@ UBlueprint* FKismetEditorUtilities::CreateBlueprint(UClass* ParentClass, UObject
 
 		const UGeneralProjectSettings& ProjectSettings = *GetDefault<UGeneralProjectSettings>();
 		Attribs.Add(FAnalyticsEventAttribute(FString("ProjectId"), ProjectSettings.ProjectID.ToString()));
+		Attribs.Add(FAnalyticsEventAttribute(FString("BlueprintId"), NewBP->GetBlueprintGuid().ToString()));
 
 		FEngineAnalytics::GetProvider().RecordEvent(FString("Editor.Usage.BlueprintCreated"), Attribs);
 	}
diff --git a/Engine/Source/Editor/UnrealEd/Public/Kismet2/BlueprintEditorUtils.h b/Engine/Source/Editor/UnrealEd/Public/Kismet2/BlueprintEditorUtils.h
--- a/Engine/Source/Editor/UnrealEd/Public/Kismet2/BlueprintEditorUtils.h
+++ b/Engine/Source/Editor/UnrealEd/Public/Kismet2/BlueprintEditorUtils.h
@@ -816,9 +816,9 @@ public:
 	static void FindAndSetDebuggableBlueprintInstances();
 
 	/**
-	 * Records new node class and type for analytics submission
+	 * Records node create events for analytics
 	 */
-	static void AnalyticsTrackNewNode( UEdGraphNode* NewNode, FName NodeClass, FName NodeType );
+	static void AnalyticsTrackNewNode( UEdGraphNode* NewNode );
 
 	/**
 	 * Generates a unique graph name for the supplied blueprint (guaranteed to not 
diff --git a/Engine/Source/Runtime/Engine/Classes/Engine/BlueprintCore.h b/Engine/Source/Runtime/Engine/Classes/Engine/BlueprintCore.h
--- a/Engine/Source/Runtime/Engine/Classes/Engine/BlueprintCore.h
+++ b/Engine/Source/Runtime/Engine/Classes/Engine/BlueprintCore.h
@@ -27,6 +27,10 @@ private:
 	UPROPERTY()
 	bool bLegacyGeneratedClassIsAuthoritative;
 
+	/** Blueprint Guid */
+	UPROPERTY()
+	FGuid BlueprintGuid;
+
 public:
 	void SetLegacyGeneratedClassIsAuthoritative()
 	{
@@ -42,4 +46,19 @@ public:
 
 	/** Gets asset registry tags */
 	virtual void GetAssetRegistryTags(TArray<FAssetRegistryTag>& OutTags) const OVERRIDE;
+
+	/** Generates a new blueprint Guid, used when creating new blueprints */
+	void GenerateNewGuid()
+	{
+		BlueprintGuid = FGuid::NewGuid();
+	}
+
+	/** Gets the Blueprint Guid */
+	const FGuid& GetBlueprintGuid() const { return BlueprintGuid; }
+
+private:
+	
+	/** Generates a new deterministic guid based on blueprint properties */
+	void GenerateDeterministicGuid();
+
 };
\ No newline at end of file
diff --git a/Engine/Source/Runtime/Engine/Private/Blueprint.cpp b/Engine/Source/Runtime/Engine/Private/Blueprint.cpp
--- a/Engine/Source/Runtime/Engine/Private/Blueprint.cpp
+++ b/Engine/Source/Runtime/Engine/Private/Blueprint.cpp
@@ -120,6 +120,10 @@ void UBlueprintCore::Serialize(FArchive& Ar)
 			}
 		}
 	}
+	if( Ar.ArIsLoading && !BlueprintGuid.IsValid() )
+	{
+		GenerateDeterministicGuid();
+	}
 }
 
 void UBlueprintCore::GetAssetRegistryTags(TArray<FAssetRegistryTag>& OutTags) const
@@ -139,6 +143,20 @@ void UBlueprintCore::GetAssetRegistryTags(TArray<FAssetRegistryTag>& OutTags) co
 	OutTags.Add( FAssetRegistryTag("GeneratedClass", GeneratedClassVal, FAssetRegistryTag::TT_Hidden) );
 }
 
+void UBlueprintCore::GenerateDeterministicGuid()
+{
+	FString HashString = GetPathName();
+	HashString.Shrink();
+	ensure( HashString.Len() );
+
+	uint32 HashBuffer[ 5 ];
+	uint32 BufferLength = HashString.Len() * sizeof( HashString[0] );
+	FSHA1::HashBuffer(*HashString, BufferLength, reinterpret_cast< uint8 * >( HashBuffer ));
+	BlueprintGuid.A = HashBuffer[ 1 ];
+	BlueprintGuid.B = HashBuffer[ 2 ];
+	BlueprintGuid.C = HashBuffer[ 3 ];
+	BlueprintGuid.D = HashBuffer[ 4 ];
+}
 
 UBlueprint::UBlueprint(const class FPostConstructInitializeProperties& PCIP)
 	: Super(PCIP)
