[331442] BP compile time optimization

- modified SBlueprintSubPalette to consolidate and defer tree refresh operations until the next slate Tick()
- note: more robust improvements to the palette system are in progress and will replace this [temp] fix in the near future

diff --git a/Engine/Source/Editor/Kismet/Private/SBlueprintSubPalette.cpp b/Engine/Source/Editor/Kismet/Private/SBlueprintSubPalette.cpp
--- a/Engine/Source/Editor/Kismet/Private/SBlueprintSubPalette.cpp
+++ b/Engine/Source/Editor/Kismet/Private/SBlueprintSubPalette.cpp
@@ -229,10 +229,22 @@ void SBlueprintSubPalette::Construct(FArguments const& InArgs, TWeakPtr<FBluepri
 	// has to come after GraphActionMenu has been set
 	BindCommands(CommandList);
 
-	GEditor->OnBlueprintCompiled().AddSP(this, &SBlueprintSubPalette::RefreshActionsList, true);
+	GEditor->OnBlueprintCompiled().AddSP(this, &SBlueprintSubPalette::RequestRefreshActionsList);
 	// Register with the Asset Registry to be informed when it is done loading up files.
 	FAssetRegistryModule& AssetRegistryModule = FModuleManager::GetModuleChecked<FAssetRegistryModule>(TEXT("AssetRegistry"));
-	AssetRegistryModule.Get().OnFilesLoaded().AddSP(this, &SBlueprintSubPalette::RefreshActionsList, true);
+	AssetRegistryModule.Get().OnFilesLoaded().AddSP(this, &SBlueprintSubPalette::RequestRefreshActionsList);
+}
+
+//------------------------------------------------------------------------------
+void SBlueprintSubPalette::Tick(const FGeometry& AllottedGeometry, const double InCurrentTime, const float InDeltaTime)
+{
+	SGraphPalette::Tick(AllottedGeometry, InCurrentTime, InDeltaTime);
+
+	if(bNeedsRefresh)
+	{
+		bNeedsRefresh = false;
+		RefreshActionsList(true);
+	}
 }
 
 //------------------------------------------------------------------------------
@@ -346,6 +358,12 @@ void SBlueprintSubPalette::GenerateContextMenuEntries(FMenuBuilder& MenuBuilder)
 	MenuBuilder.AddMenuEntry(PaletteCommands.RefreshPalette);
 }
 
+//------------------------------------------------------------------------------
+void SBlueprintSubPalette::RequestRefreshActionsList()
+{
+	bNeedsRefresh = true;
+}
+
 /*******************************************************************************
 * Private SBlueprintSubPalette Methods
 *******************************************************************************/
diff --git a/Engine/Source/Editor/Kismet/Private/SBlueprintSubPalette.h b/Engine/Source/Editor/Kismet/Private/SBlueprintSubPalette.h
--- a/Engine/Source/Editor/Kismet/Private/SBlueprintSubPalette.h
+++ b/Engine/Source/Editor/Kismet/Private/SBlueprintSubPalette.h
@@ -34,6 +34,15 @@ public:
 	 */
 	void Construct(const FArguments& InArgs, TWeakPtr<FBlueprintEditor> InBlueprintEditor);
 
+	/**
+	 * Ticks this widget.
+	 *
+	 * @param  AllottedGeometry The space allotted for this widget
+	 * @param  InCurrentTime  Current absolute real time
+	 * @param  InDeltaTime  Real time passed since last tick
+	 */
+	virtual void Tick( const FGeometry& AllottedGeometry, const double InCurrentTime, const float InDeltaTime ) override;
+
 	/**
 	 * Retrieves, from the owning blueprint-editor, the blueprint currently  
 	 * being worked on.
@@ -84,6 +93,9 @@ protected:
 	 */
 	virtual void GenerateContextMenuEntries(FMenuBuilder& MenuBuilder) const;
 
+	/** Delegate to call to request a refresh */
+	void RequestRefreshActionsList();
+
 private:
 	/**
 	 * Constructs a slate widget for the right-click context menu in this 
@@ -99,5 +111,8 @@ private:
 
 	/** Pointer to the command list created for this (so multiple sub-palettes can have their own bindings)*/
 	TSharedPtr<FUICommandList> CommandList;
+
+	/** Flag to indicate whether or not the action list needs to be refreshed */
+	bool bNeedsRefresh;
 };
 
