UE4 Refactoring. Changed OVERRIDE and FINAL macros to keywords override and final.

diff --git a/Engine/Source/Editor/GraphEditor/Private/KismetNodes/SGraphNodeCallParameterCollectionFunction.h b/Engine/Source/Editor/GraphEditor/Private/KismetNodes/SGraphNodeCallParameterCollectionFunction.h
--- a/Engine/Source/Editor/GraphEditor/Private/KismetNodes/SGraphNodeCallParameterCollectionFunction.h
+++ b/Engine/Source/Editor/GraphEditor/Private/KismetNodes/SGraphNodeCallParameterCollectionFunction.h
@@ -7,6 +7,6 @@ class SGraphNodeCallParameterCollectionFunction : public SGraphNodeK2Default
 protected:
 
 	// SGraphNode interface
-	virtual TSharedPtr<SGraphPin> CreatePinWidget(UEdGraphPin* Pin) const OVERRIDE;
+	virtual TSharedPtr<SGraphPin> CreatePinWidget(UEdGraphPin* Pin) const override;
 	// End of SGraphNode interface
 };
diff --git a/Engine/Source/Editor/GraphEditor/Private/KismetNodes/SGraphNodeFormatText.h b/Engine/Source/Editor/GraphEditor/Private/KismetNodes/SGraphNodeFormatText.h
--- a/Engine/Source/Editor/GraphEditor/Private/KismetNodes/SGraphNodeFormatText.h
+++ b/Engine/Source/Editor/GraphEditor/Private/KismetNodes/SGraphNodeFormatText.h
@@ -11,11 +11,11 @@ public:
 	void Construct(const FArguments& InArgs, UK2Node_FormatText* InNode);
 
 	// SGraphNode interface
-	virtual void CreatePinWidgets() OVERRIDE;
+	virtual void CreatePinWidgets() override;
 
 protected:
-	virtual void CreateInputSideAddButton(TSharedPtr<SVerticalBox> InputBox) OVERRIDE;
-	virtual EVisibility IsAddPinButtonVisible() const OVERRIDE;
-	virtual FReply OnAddPin() OVERRIDE;
+	virtual void CreateInputSideAddButton(TSharedPtr<SVerticalBox> InputBox) override;
+	virtual EVisibility IsAddPinButtonVisible() const override;
+	virtual FReply OnAddPin() override;
 	// End of SGraphNode interface
 };
diff --git a/Engine/Source/Editor/GraphEditor/Private/KismetNodes/SGraphNodeK2Base.h b/Engine/Source/Editor/GraphEditor/Private/KismetNodes/SGraphNodeK2Base.h
--- a/Engine/Source/Editor/GraphEditor/Private/KismetNodes/SGraphNodeK2Base.h
+++ b/Engine/Source/Editor/GraphEditor/Private/KismetNodes/SGraphNodeK2Base.h
@@ -7,18 +7,18 @@ class GRAPHEDITOR_API SGraphNodeK2Base : public SGraphNode
 public:
 
 	// SGraphNode interface
-	virtual void UpdateGraphNode() OVERRIDE;
+	virtual void UpdateGraphNode() override;
 
 	// SNodePanel::SNode interface
-	virtual bool RequiresSecondPassLayout() const OVERRIDE;
-	virtual void GetOverlayBrushes(bool bSelected, const FVector2D WidgetSize, TArray<FOverlayBrushInfo>& Brushes) const OVERRIDE;
-	virtual void GetNodeInfoPopups(FNodeInfoContext* Context, TArray<FGraphInformationPopupInfo>& Popups) const OVERRIDE;
-	virtual const FSlateBrush* GetShadowBrush(bool bSelected) const OVERRIDE;
-	void PerformSecondPassLayout(const TMap< UObject*, TSharedRef<SNode> >& NodeToWidgetLookup) const OVERRIDE;
+	virtual bool RequiresSecondPassLayout() const override;
+	virtual void GetOverlayBrushes(bool bSelected, const FVector2D WidgetSize, TArray<FOverlayBrushInfo>& Brushes) const override;
+	virtual void GetNodeInfoPopups(FNodeInfoContext* Context, TArray<FGraphInformationPopupInfo>& Popups) const override;
+	virtual const FSlateBrush* GetShadowBrush(bool bSelected) const override;
+	void PerformSecondPassLayout(const TMap< UObject*, TSharedRef<SNode> >& NodeToWidgetLookup) const override;
 
 protected :
 	// Begin SGraphNode interface
-	virtual TSharedPtr<SToolTip> GetComplexTooltip() OVERRIDE;
+	virtual TSharedPtr<SToolTip> GetComplexTooltip() override;
 	// End SGraphNode interface
 
 	/** Set up node in 'standard' mode */
diff --git a/Engine/Source/Editor/GraphEditor/Private/KismetNodes/SGraphNodeK2Composite.h b/Engine/Source/Editor/GraphEditor/Private/KismetNodes/SGraphNodeK2Composite.h
--- a/Engine/Source/Editor/GraphEditor/Private/KismetNodes/SGraphNodeK2Composite.h
+++ b/Engine/Source/Editor/GraphEditor/Private/KismetNodes/SGraphNodeK2Composite.h
@@ -11,8 +11,8 @@ public:
 	void Construct(const FArguments& InArgs, UK2Node_Composite* InNode);
 
 	// SGraphNode interface
-	virtual void UpdateGraphNode() OVERRIDE;
-	virtual TSharedPtr<SToolTip> GetComplexTooltip() OVERRIDE;
+	virtual void UpdateGraphNode() override;
+	virtual TSharedPtr<SToolTip> GetComplexTooltip() override;
 	// End of SGraphNode interface
 protected:
 	virtual UEdGraph* GetInnerGraph() const;
diff --git a/Engine/Source/Editor/GraphEditor/Private/KismetNodes/SGraphNodeK2CreateDelegate.h b/Engine/Source/Editor/GraphEditor/Private/KismetNodes/SGraphNodeK2CreateDelegate.h
--- a/Engine/Source/Editor/GraphEditor/Private/KismetNodes/SGraphNodeK2CreateDelegate.h
+++ b/Engine/Source/Editor/GraphEditor/Private/KismetNodes/SGraphNodeK2CreateDelegate.h
@@ -20,7 +20,7 @@ public:
 
 	virtual ~SGraphNodeK2CreateDelegate();
 	void Construct(const FArguments& InArgs, UK2Node* InNode);
-	virtual void CreateBelowWidgetControls(TSharedPtr<SVerticalBox> MainBox) OVERRIDE;
+	virtual void CreateBelowWidgetControls(TSharedPtr<SVerticalBox> MainBox) override;
 
 protected:
 	FString GetCurrentFunctionDescription() const;
diff --git a/Engine/Source/Editor/GraphEditor/Private/KismetNodes/SGraphNodeK2Event.h b/Engine/Source/Editor/GraphEditor/Private/KismetNodes/SGraphNodeK2Event.h
--- a/Engine/Source/Editor/GraphEditor/Private/KismetNodes/SGraphNodeK2Event.h
+++ b/Engine/Source/Editor/GraphEditor/Private/KismetNodes/SGraphNodeK2Event.h
@@ -8,12 +8,12 @@ public:
 	SGraphNodeK2Event() : SGraphNodeK2Default(), bHasDelegateOutputPin(false) {}
 
 protected:
-	virtual TSharedRef<SWidget> CreateTitleWidget(TSharedPtr<SNodeTitle> NodeTitle) OVERRIDE;
-	virtual bool UseLowDetailNodeTitles() const OVERRIDE;
-	virtual void AddPin( const TSharedRef<SGraphPin>& PinToAdd ) OVERRIDE;
+	virtual TSharedRef<SWidget> CreateTitleWidget(TSharedPtr<SNodeTitle> NodeTitle) override;
+	virtual bool UseLowDetailNodeTitles() const override;
+	virtual void AddPin( const TSharedRef<SGraphPin>& PinToAdd ) override;
 
 
-	virtual void SetDefaultTitleAreaWidget(TSharedRef<SOverlay> DefaultTitleAreaWidget) OVERRIDE
+	virtual void SetDefaultTitleAreaWidget(TSharedRef<SOverlay> DefaultTitleAreaWidget) override
 	{
 		TitleAreaWidget = DefaultTitleAreaWidget;
 	}
diff --git a/Engine/Source/Editor/GraphEditor/Private/KismetNodes/SGraphNodeK2Sequence.h b/Engine/Source/Editor/GraphEditor/Private/KismetNodes/SGraphNodeK2Sequence.h
--- a/Engine/Source/Editor/GraphEditor/Private/KismetNodes/SGraphNodeK2Sequence.h
+++ b/Engine/Source/Editor/GraphEditor/Private/KismetNodes/SGraphNodeK2Sequence.h
@@ -14,8 +14,8 @@ public:
 
 protected:
 	// SGraphNode interface
-	virtual void CreateOutputSideAddButton(TSharedPtr<SVerticalBox> OutputBox) OVERRIDE;
-	virtual FReply OnAddPin() OVERRIDE;
+	virtual void CreateOutputSideAddButton(TSharedPtr<SVerticalBox> OutputBox) override;
+	virtual FReply OnAddPin() override;
 	// End of SGraphNode interface
 };
 
diff --git a/Engine/Source/Editor/GraphEditor/Private/KismetNodes/SGraphNodeK2Terminator.h b/Engine/Source/Editor/GraphEditor/Private/KismetNodes/SGraphNodeK2Terminator.h
--- a/Engine/Source/Editor/GraphEditor/Private/KismetNodes/SGraphNodeK2Terminator.h
+++ b/Engine/Source/Editor/GraphEditor/Private/KismetNodes/SGraphNodeK2Terminator.h
@@ -15,7 +15,7 @@ public:
 	virtual const FSlateBrush* GetShadowBrush(bool bSelected) const;
 
 protected:
-	virtual void UpdateGraphNode() OVERRIDE;
+	virtual void UpdateGraphNode() override;
 };
 
 #endif // __SGraphNodeK2Terminator_h__
\ No newline at end of file
diff --git a/Engine/Source/Editor/GraphEditor/Private/KismetNodes/SGraphNodeK2Timeline.h b/Engine/Source/Editor/GraphEditor/Private/KismetNodes/SGraphNodeK2Timeline.h
--- a/Engine/Source/Editor/GraphEditor/Private/KismetNodes/SGraphNodeK2Timeline.h
+++ b/Engine/Source/Editor/GraphEditor/Private/KismetNodes/SGraphNodeK2Timeline.h
@@ -11,6 +11,6 @@ public:
 	void Construct(const FArguments& InArgs, UK2Node_Timeline* InNode);
 
 	// SNodePanel::SNode interface
-	void GetNodeInfoPopups(FNodeInfoContext* Context, TArray<FGraphInformationPopupInfo>& Popups) const OVERRIDE;
+	void GetNodeInfoPopups(FNodeInfoContext* Context, TArray<FGraphInformationPopupInfo>& Popups) const override;
 	// End of SNodePanel::SNode interface
 };
diff --git a/Engine/Source/Editor/GraphEditor/Private/KismetNodes/SGraphNodeK2Var.h b/Engine/Source/Editor/GraphEditor/Private/KismetNodes/SGraphNodeK2Var.h
--- a/Engine/Source/Editor/GraphEditor/Private/KismetNodes/SGraphNodeK2Var.h
+++ b/Engine/Source/Editor/GraphEditor/Private/KismetNodes/SGraphNodeK2Var.h
@@ -15,7 +15,7 @@ public:
 	virtual const FSlateBrush* GetShadowBrush(bool bSelected) const;
 
 	// SGraphNode interface
-	virtual void UpdateGraphNode() OVERRIDE;
+	virtual void UpdateGraphNode() override;
 	// End of SGraphNode interface
 
 protected:
diff --git a/Engine/Source/Editor/GraphEditor/Private/KismetNodes/SGraphNodeSpawnActor.cpp b/Engine/Source/Editor/GraphEditor/Private/KismetNodes/SGraphNodeSpawnActor.cpp
--- a/Engine/Source/Editor/GraphEditor/Private/KismetNodes/SGraphNodeSpawnActor.cpp
+++ b/Engine/Source/Editor/GraphEditor/Private/KismetNodes/SGraphNodeSpawnActor.cpp
@@ -40,7 +40,7 @@ class SGraphPinActorBasedBlueprintClass : public SGraphPinObject
 	{
 	public:
 
-		virtual bool IsClassAllowed(const FClassViewerInitializationOptions& InInitOptions, const UClass* InClass, TSharedRef< FClassViewerFilterFuncs > InFilterFuncs ) OVERRIDE
+		virtual bool IsClassAllowed(const FClassViewerInitializationOptions& InInitOptions, const UClass* InClass, TSharedRef< FClassViewerFilterFuncs > InFilterFuncs ) override
 		{
 			if(NULL != InClass)
 			{
@@ -51,7 +51,7 @@ class SGraphPinActorBasedBlueprintClass : public SGraphPinObject
 			return false;
 		}
 
-		virtual bool IsUnloadedClassAllowed(const FClassViewerInitializationOptions& InInitOptions, const TSharedRef< const IUnloadedBlueprintData > InUnloadedClassData, TSharedRef< FClassViewerFilterFuncs > InFilterFuncs) OVERRIDE
+		virtual bool IsUnloadedClassAllowed(const FClassViewerInitializationOptions& InInitOptions, const TSharedRef< const IUnloadedBlueprintData > InUnloadedClassData, TSharedRef< FClassViewerFilterFuncs > InFilterFuncs) override
 		{
 			const bool bActorBased = InUnloadedClassData->IsChildOf(AActor::StaticClass());
 			return bActorBased;
@@ -60,9 +60,9 @@ class SGraphPinActorBasedBlueprintClass : public SGraphPinObject
 
 protected:
 
-	virtual FText GetDefaultComboText() const OVERRIDE { return LOCTEXT( "DefaultComboText", "Select Blueprint" ); }
+	virtual FText GetDefaultComboText() const override { return LOCTEXT( "DefaultComboText", "Select Blueprint" ); }
 
-	virtual TSharedRef<SWidget> GenerateAssetPicker() OVERRIDE
+	virtual TSharedRef<SWidget> GenerateAssetPicker() override
 	{
 		FClassViewerModule& ClassViewerModule = FModuleManager::LoadModuleChecked<FClassViewerModule>("ClassViewer");
 
diff --git a/Engine/Source/Editor/GraphEditor/Private/KismetNodes/SGraphNodeSpawnActor.h b/Engine/Source/Editor/GraphEditor/Private/KismetNodes/SGraphNodeSpawnActor.h
--- a/Engine/Source/Editor/GraphEditor/Private/KismetNodes/SGraphNodeSpawnActor.h
+++ b/Engine/Source/Editor/GraphEditor/Private/KismetNodes/SGraphNodeSpawnActor.h
@@ -7,6 +7,6 @@ class SGraphNodeSpawnActor : public SGraphNodeK2Default
 public:
 
 	// SGraphNode interface
-	virtual void CreatePinWidgets() OVERRIDE;
+	virtual void CreatePinWidgets() override;
 	// End of SGraphNode interface
 };
diff --git a/Engine/Source/Editor/GraphEditor/Private/KismetNodes/SGraphNodeSpawnActorFromClass.cpp b/Engine/Source/Editor/GraphEditor/Private/KismetNodes/SGraphNodeSpawnActorFromClass.cpp
--- a/Engine/Source/Editor/GraphEditor/Private/KismetNodes/SGraphNodeSpawnActorFromClass.cpp
+++ b/Engine/Source/Editor/GraphEditor/Private/KismetNodes/SGraphNodeSpawnActorFromClass.cpp
@@ -40,7 +40,7 @@ class SGraphPinActorBasedClass : public SGraphPinObject
 	{
 	public:
 
-		virtual bool IsClassAllowed(const FClassViewerInitializationOptions& InInitOptions, const UClass* InClass, TSharedRef< FClassViewerFilterFuncs > InFilterFuncs ) OVERRIDE
+		virtual bool IsClassAllowed(const FClassViewerInitializationOptions& InInitOptions, const UClass* InClass, TSharedRef< FClassViewerFilterFuncs > InFilterFuncs ) override
 		{
 			if(NULL != InClass)
 			{
@@ -52,7 +52,7 @@ class SGraphPinActorBasedClass : public SGraphPinObject
 			return false;
 		}
 
-		virtual bool IsUnloadedClassAllowed(const FClassViewerInitializationOptions& InInitOptions, const TSharedRef< const IUnloadedBlueprintData > InUnloadedClassData, TSharedRef< FClassViewerFilterFuncs > InFilterFuncs) OVERRIDE
+		virtual bool IsUnloadedClassAllowed(const FClassViewerInitializationOptions& InInitOptions, const TSharedRef< const IUnloadedBlueprintData > InUnloadedClassData, TSharedRef< FClassViewerFilterFuncs > InFilterFuncs) override
 		{
 			const bool bActorBased = InUnloadedClassData->IsChildOf(AActor::StaticClass());
 			const bool bNotAbstract = !InUnloadedClassData->HasAnyClassFlags(CLASS_Abstract);
@@ -62,7 +62,7 @@ class SGraphPinActorBasedClass : public SGraphPinObject
 
 protected:
 
-	virtual FReply OnClickUse() OVERRIDE
+	virtual FReply OnClickUse() override
 	{
 		if(GraphPinObj && GraphPinObj->GetSchema())
 		{
@@ -78,14 +78,14 @@ protected:
 		return FReply::Handled();
 	}
 
-	virtual FOnClicked GetOnUseButtonDelegate() OVERRIDE
+	virtual FOnClicked GetOnUseButtonDelegate() override
 	{
 		return FOnClicked::CreateSP( this, &SGraphPinActorBasedClass::OnClickUse );
 	}
 
-	virtual FText GetDefaultComboText() const OVERRIDE { return LOCTEXT( "DefaultComboText", "Select Class" ); }
+	virtual FText GetDefaultComboText() const override { return LOCTEXT( "DefaultComboText", "Select Class" ); }
 
-	virtual TSharedRef<SWidget> GenerateAssetPicker() OVERRIDE
+	virtual TSharedRef<SWidget> GenerateAssetPicker() override
 	{
 		FClassViewerModule& ClassViewerModule = FModuleManager::LoadModuleChecked<FClassViewerModule>("ClassViewer");
 
diff --git a/Engine/Source/Editor/GraphEditor/Private/KismetNodes/SGraphNodeSpawnActorFromClass.h b/Engine/Source/Editor/GraphEditor/Private/KismetNodes/SGraphNodeSpawnActorFromClass.h
--- a/Engine/Source/Editor/GraphEditor/Private/KismetNodes/SGraphNodeSpawnActorFromClass.h
+++ b/Engine/Source/Editor/GraphEditor/Private/KismetNodes/SGraphNodeSpawnActorFromClass.h
@@ -7,6 +7,6 @@ class SGraphNodeSpawnActorFromClass : public SGraphNodeK2Default
 public:
 
 	// SGraphNode interface
-	virtual void CreatePinWidgets() OVERRIDE;
+	virtual void CreatePinWidgets() override;
 	// End of SGraphNode interface
 };
diff --git a/Engine/Source/Editor/GraphEditor/Private/KismetNodes/SGraphNodeSwitchStatement.h b/Engine/Source/Editor/GraphEditor/Private/KismetNodes/SGraphNodeSwitchStatement.h
--- a/Engine/Source/Editor/GraphEditor/Private/KismetNodes/SGraphNodeSwitchStatement.h
+++ b/Engine/Source/Editor/GraphEditor/Private/KismetNodes/SGraphNodeSwitchStatement.h
@@ -11,11 +11,11 @@ public:
 	void Construct(const FArguments& InArgs, UK2Node_Switch* InNode);
 
 	// SGraphNode interface
-	virtual void CreatePinWidgets() OVERRIDE;
+	virtual void CreatePinWidgets() override;
 
 protected:
-	virtual void CreateOutputSideAddButton(TSharedPtr<SVerticalBox> OutputBox) OVERRIDE;
-	virtual EVisibility IsAddPinButtonVisible() const OVERRIDE;
-	virtual FReply OnAddPin() OVERRIDE;
+	virtual void CreateOutputSideAddButton(TSharedPtr<SVerticalBox> OutputBox) override;
+	virtual EVisibility IsAddPinButtonVisible() const override;
+	virtual FReply OnAddPin() override;
 	// End of SGraphNode interface
 };
diff --git a/Engine/Source/Editor/GraphEditor/Private/KismetPins/SGraphPinClass.h b/Engine/Source/Editor/GraphEditor/Private/KismetPins/SGraphPinClass.h
--- a/Engine/Source/Editor/GraphEditor/Private/KismetPins/SGraphPinClass.h
+++ b/Engine/Source/Editor/GraphEditor/Private/KismetPins/SGraphPinClass.h
@@ -20,11 +20,11 @@ protected:
 	void OnPickedNewClass(UClass* ChosenClass);
 
 	// Begin SGraphPinObject interface
-	virtual FReply OnClickUse() OVERRIDE;
-	virtual bool AllowSelfPinWidget() const OVERRIDE { return false; }
-	virtual TSharedRef<SWidget> GenerateAssetPicker() OVERRIDE;
-	virtual FText GetDefaultComboText() const OVERRIDE;
-	virtual FOnClicked GetOnUseButtonDelegate() OVERRIDE;
+	virtual FReply OnClickUse() override;
+	virtual bool AllowSelfPinWidget() const override { return false; }
+	virtual TSharedRef<SWidget> GenerateAssetPicker() override;
+	virtual FText GetDefaultComboText() const override;
+	virtual FOnClicked GetOnUseButtonDelegate() override;
 	// End SGraphPinObject interface
 
 };
diff --git a/Engine/Source/Editor/GraphEditor/Private/KismetPins/SGraphPinExec.h b/Engine/Source/Editor/GraphEditor/Private/KismetPins/SGraphPinExec.h
--- a/Engine/Source/Editor/GraphEditor/Private/KismetPins/SGraphPinExec.h
+++ b/Engine/Source/Editor/GraphEditor/Private/KismetPins/SGraphPinExec.h
@@ -12,8 +12,8 @@ public:
 
 protected:
 	// Begin SGraphPin interface
-	virtual TSharedRef<SWidget>	GetDefaultValueWidget() OVERRIDE;
-	virtual const FSlateBrush* GetPinIcon() const OVERRIDE;
+	virtual TSharedRef<SWidget>	GetDefaultValueWidget() override;
+	virtual const FSlateBrush* GetPinIcon() const override;
 	// End SGraphPin interface
 
 	void CachePinBrushes(bool bForceCache = false) const;
diff --git a/Engine/Source/Editor/GraphEditor/Private/KismetPins/SGraphPinNum.h b/Engine/Source/Editor/GraphEditor/Private/KismetPins/SGraphPinNum.h
--- a/Engine/Source/Editor/GraphEditor/Private/KismetPins/SGraphPinNum.h
+++ b/Engine/Source/Editor/GraphEditor/Private/KismetPins/SGraphPinNum.h
@@ -12,6 +12,6 @@ public:
 
 protected:
 	// Begin SGraphPinText interface
-	virtual void SetTypeInValue(const FText& NewTypeInValue, ETextCommit::Type CommitInfo) OVERRIDE;
+	virtual void SetTypeInValue(const FText& NewTypeInValue, ETextCommit::Type CommitInfo) override;
 	// End SGraphPinText interface
 };
diff --git a/Engine/Source/Editor/GraphEditor/Private/MaterialNodes/SGraphNodeMaterialBase.h b/Engine/Source/Editor/GraphEditor/Private/MaterialNodes/SGraphNodeMaterialBase.h
--- a/Engine/Source/Editor/GraphEditor/Private/MaterialNodes/SGraphNodeMaterialBase.h
+++ b/Engine/Source/Editor/GraphEditor/Private/MaterialNodes/SGraphNodeMaterialBase.h
@@ -11,10 +11,10 @@ public:
 	~FPreviewViewport();
 
 	// ISlateViewport interface
-	virtual void OnDrawViewport( const FGeometry& AllottedGeometry, const FSlateRect& MyClippingRect, class FSlateWindowElementList& OutDrawElements, int32 LayerId, const FWidgetStyle& InWidgetStyle, bool bParentEnabled ) OVERRIDE;
-	virtual FIntPoint GetSize() const OVERRIDE;
-	virtual class FSlateShaderResource* GetViewportRenderTargetTexture() const OVERRIDE {return NULL;}
-	virtual bool RequiresVsync() const OVERRIDE {return false;}
+	virtual void OnDrawViewport( const FGeometry& AllottedGeometry, const FSlateRect& MyClippingRect, class FSlateWindowElementList& OutDrawElements, int32 LayerId, const FWidgetStyle& InWidgetStyle, bool bParentEnabled ) override;
+	virtual FIntPoint GetSize() const override;
+	virtual class FSlateShaderResource* GetViewportRenderTargetTexture() const override {return NULL;}
+	virtual bool RequiresVsync() const override {return false;}
 
 	/** Material node to get expression preview from */
 	UMaterialGraphNode* MaterialNode;
@@ -44,7 +44,7 @@ private:
 	/**
 	 * ICustomSlateElement interface 
 	 */
-	virtual void DrawRenderThread( const void* InWindowBackBuffer ) OVERRIDE;
+	virtual void DrawRenderThread( const void* InWindowBackBuffer ) override;
 
 private:
 	/** Render target that the canvas renders to */
@@ -64,21 +64,21 @@ public:
 	void Construct(const FArguments& InArgs, class UMaterialGraphNode* InNode);
 
 	// SGraphNode interface
-	virtual void CreatePinWidgets() OVERRIDE;
+	virtual void CreatePinWidgets() override;
 	// End of SGraphNode interface
 
 	// SNodePanel::SNode interface
-	virtual void MoveTo(const FVector2D& NewPosition, FNodeSet& NodeFilter) OVERRIDE;
+	virtual void MoveTo(const FVector2D& NewPosition, FNodeSet& NodeFilter) override;
 	// End of SNodePanel::SNode interface
 
 	UMaterialGraphNode* GetMaterialGraphNode() const {return MaterialNode;}
 
 protected:
 	// SGraphNode interface
-	virtual void AddPin( const TSharedRef<SGraphPin>& PinToAdd ) OVERRIDE;
-	virtual void CreateBelowPinControls(TSharedPtr<SVerticalBox> MainBox) OVERRIDE;
-	virtual void SetDefaultTitleAreaWidget(TSharedRef<SOverlay> DefaultTitleAreaWidget) OVERRIDE;
-	virtual TSharedRef<SWidget> CreateNodeContentArea() OVERRIDE;
+	virtual void AddPin( const TSharedRef<SGraphPin>& PinToAdd ) override;
+	virtual void CreateBelowPinControls(TSharedPtr<SVerticalBox> MainBox) override;
+	virtual void SetDefaultTitleAreaWidget(TSharedRef<SOverlay> DefaultTitleAreaWidget) override;
+	virtual TSharedRef<SWidget> CreateNodeContentArea() override;
 	// End of SGraphNode interface
 
 	/** Creates a preview viewport if necessary */
diff --git a/Engine/Source/Editor/GraphEditor/Private/MaterialNodes/SGraphNodeMaterialComment.h b/Engine/Source/Editor/GraphEditor/Private/MaterialNodes/SGraphNodeMaterialComment.h
--- a/Engine/Source/Editor/GraphEditor/Private/MaterialNodes/SGraphNodeMaterialComment.h
+++ b/Engine/Source/Editor/GraphEditor/Private/MaterialNodes/SGraphNodeMaterialComment.h
@@ -11,7 +11,7 @@ public:
 	void Construct(const FArguments& InArgs, class UMaterialGraphNode_Comment* InNode);
 
 	// SNodePanel::SNode interface
-	virtual void MoveTo(const FVector2D& NewPosition, FNodeSet& NodeFilter) OVERRIDE;
+	virtual void MoveTo(const FVector2D& NewPosition, FNodeSet& NodeFilter) override;
 	// End of SNodePanel::SNode interface
 
 private:
diff --git a/Engine/Source/Editor/GraphEditor/Private/MaterialNodes/SGraphNodeMaterialResult.h b/Engine/Source/Editor/GraphEditor/Private/MaterialNodes/SGraphNodeMaterialResult.h
--- a/Engine/Source/Editor/GraphEditor/Private/MaterialNodes/SGraphNodeMaterialResult.h
+++ b/Engine/Source/Editor/GraphEditor/Private/MaterialNodes/SGraphNodeMaterialResult.h
@@ -11,11 +11,11 @@ public:
 	void Construct(const FArguments& InArgs, class UMaterialGraphNode_Root* InNode);
 
 	// SGraphNode interface
-	virtual void CreatePinWidgets() OVERRIDE;
+	virtual void CreatePinWidgets() override;
 	// End of SGraphNode interface
 
 	// SNodePanel::SNode interface
-	virtual void MoveTo(const FVector2D& NewPosition, FNodeSet& NodeFilter) OVERRIDE;
+	virtual void MoveTo(const FVector2D& NewPosition, FNodeSet& NodeFilter) override;
 	// End of SNodePanel::SNode interface
 
 private:
diff --git a/Engine/Source/Editor/GraphEditor/Private/MaterialPins/SGraphPinMaterialInput.h b/Engine/Source/Editor/GraphEditor/Private/MaterialPins/SGraphPinMaterialInput.h
--- a/Engine/Source/Editor/GraphEditor/Private/MaterialPins/SGraphPinMaterialInput.h
+++ b/Engine/Source/Editor/GraphEditor/Private/MaterialPins/SGraphPinMaterialInput.h
@@ -12,7 +12,7 @@ public:
 
 protected:
 	// Begin SGraphPin interface
-	virtual FSlateColor GetPinColor() const OVERRIDE;
+	virtual FSlateColor GetPinColor() const override;
 	// End SGraphPin interface
 
 private:
diff --git a/Engine/Source/Editor/GraphEditor/Private/SGraphActionMenu.cpp b/Engine/Source/Editor/GraphEditor/Private/SGraphActionMenu.cpp
--- a/Engine/Source/Editor/GraphEditor/Private/SGraphActionMenu.cpp
+++ b/Engine/Source/Editor/GraphEditor/Private/SGraphActionMenu.cpp
@@ -110,7 +110,7 @@ public:
 	}
 
 	// SWidget interface
-	virtual FReply OnDrop( const FGeometry& MyGeometry, const FDragDropEvent& DragDropEvent ) OVERRIDE
+	virtual FReply OnDrop( const FGeometry& MyGeometry, const FDragDropEvent& DragDropEvent ) override
 	{
 		TSharedPtr<FGraphEditorDragDropAction> GraphDropOp = DragDropEvent.GetOperationAs<FGraphEditorDragDropAction>();
 		if (GraphDropOp.IsValid())
@@ -121,7 +121,7 @@ public:
 		return FReply::Unhandled();
 	}
 
-	virtual void OnDragEnter( const FGeometry& MyGeometry, const FDragDropEvent& DragDropEvent ) OVERRIDE
+	virtual void OnDragEnter( const FGeometry& MyGeometry, const FDragDropEvent& DragDropEvent ) override
 	{
 		TSharedPtr<FGraphEditorDragDropAction> GraphDropOp = DragDropEvent.GetOperationAs<FGraphEditorDragDropAction>();
 		if (GraphDropOp.IsValid())
@@ -130,7 +130,7 @@ public:
 		}
 	}
 
-	virtual void OnDragLeave( const FDragDropEvent& DragDropEvent ) OVERRIDE
+	virtual void OnDragLeave( const FDragDropEvent& DragDropEvent ) override
 	{
 		TSharedPtr<FGraphEditorDragDropAction> GraphDropOp = DragDropEvent.GetOperationAs<FGraphEditorDragDropAction>();
 		if (GraphDropOp.IsValid())
diff --git a/Engine/Source/Editor/GraphEditor/Private/SGraphEditorImpl.h b/Engine/Source/Editor/GraphEditor/Private/SGraphEditorImpl.h
--- a/Engine/Source/Editor/GraphEditor/Private/SGraphEditorImpl.h
+++ b/Engine/Source/Editor/GraphEditor/Private/SGraphEditorImpl.h
@@ -101,36 +101,36 @@ public:
 	void GraphEd_OnPanelUpdated();
 
 	// SWidget interface
-	virtual void Tick( const FGeometry& AllottedGeometry, const double InCurrentTime, const float InDeltaTime ) OVERRIDE;
+	virtual void Tick( const FGeometry& AllottedGeometry, const double InCurrentTime, const float InDeltaTime ) override;
 
 	void FocusLockedEditorHere();
 
-	virtual FReply OnKeyboardFocusReceived( const FGeometry& MyGeometry, const FKeyboardFocusEvent& InKeyboardFocusEvent ) OVERRIDE;
-	virtual FReply OnMouseButtonDown( const FGeometry& MyGeometry, const FPointerEvent& MouseEvent ) OVERRIDE;
-	virtual FReply OnKeyDown( const FGeometry& MyGeometry, const FKeyboardEvent& InKeyboardEvent ) OVERRIDE;
-	virtual bool SupportsKeyboardFocus() const OVERRIDE;
+	virtual FReply OnKeyboardFocusReceived( const FGeometry& MyGeometry, const FKeyboardFocusEvent& InKeyboardFocusEvent ) override;
+	virtual FReply OnMouseButtonDown( const FGeometry& MyGeometry, const FPointerEvent& MouseEvent ) override;
+	virtual FReply OnKeyDown( const FGeometry& MyGeometry, const FKeyboardEvent& InKeyboardEvent ) override;
+	virtual bool SupportsKeyboardFocus() const override;
 	// End of SWidget interface
 
 	// SGraphEditor interface
-	virtual const TSet<class UObject*>& GetSelectedNodes() const OVERRIDE;
-	virtual void ClearSelectionSet() OVERRIDE;
-	virtual void SetNodeSelection(UEdGraphNode* Node, bool bSelect) OVERRIDE;
-	virtual void SelectAllNodes() OVERRIDE;
-	virtual FVector2D GetPasteLocation() const OVERRIDE;
-	virtual bool IsNodeTitleVisible( const UEdGraphNode* Node, bool bRequestRename ) OVERRIDE;
-	virtual void JumpToNode( const UEdGraphNode* JumpToMe, bool bRequestRename = false ) OVERRIDE;
-	virtual void JumpToPin( const UEdGraphPin* JumpToMe ) OVERRIDE;
-	virtual UEdGraphPin* GetGraphPinForMenu() OVERRIDE;
-	virtual void ZoomToFit(bool bOnlySelection) OVERRIDE;
-	virtual bool GetBoundsForSelectedNodes( class FSlateRect& Rect, float Padding) OVERRIDE;
+	virtual const TSet<class UObject*>& GetSelectedNodes() const override;
+	virtual void ClearSelectionSet() override;
+	virtual void SetNodeSelection(UEdGraphNode* Node, bool bSelect) override;
+	virtual void SelectAllNodes() override;
+	virtual FVector2D GetPasteLocation() const override;
+	virtual bool IsNodeTitleVisible( const UEdGraphNode* Node, bool bRequestRename ) override;
+	virtual void JumpToNode( const UEdGraphNode* JumpToMe, bool bRequestRename = false ) override;
+	virtual void JumpToPin( const UEdGraphPin* JumpToMe ) override;
+	virtual UEdGraphPin* GetGraphPinForMenu() override;
+	virtual void ZoomToFit(bool bOnlySelection) override;
+	virtual bool GetBoundsForSelectedNodes( class FSlateRect& Rect, float Padding) override;
 	virtual void NotifyGraphChanged();
-	virtual TSharedPtr<SWidget> GetTitleBar() const OVERRIDE;
-	virtual void SetViewLocation(const FVector2D& Location, float ZoomAmount) OVERRIDE;
-	virtual void GetViewLocation(FVector2D& Location, float& ZoomAmount) OVERRIDE;
-	virtual void LockToGraphEditor(TWeakPtr<SGraphEditor> Other) OVERRIDE;
-	virtual void UnlockFromGraphEditor(TWeakPtr<SGraphEditor> Other) OVERRIDE;
-	virtual void AddNotification ( FNotificationInfo& Info, bool bSuccess ) OVERRIDE;
-	virtual void SetPinVisibility(SGraphEditor::EPinVisibility Visibility) OVERRIDE;
+	virtual TSharedPtr<SWidget> GetTitleBar() const override;
+	virtual void SetViewLocation(const FVector2D& Location, float ZoomAmount) override;
+	virtual void GetViewLocation(FVector2D& Location, float& ZoomAmount) override;
+	virtual void LockToGraphEditor(TWeakPtr<SGraphEditor> Other) override;
+	virtual void UnlockFromGraphEditor(TWeakPtr<SGraphEditor> Other) override;
+	virtual void AddNotification ( FNotificationInfo& Info, bool bSuccess ) override;
+	virtual void SetPinVisibility(SGraphEditor::EPinVisibility Visibility) override;
 	// End of SGraphEditor interface
 protected:
 	//
@@ -145,7 +145,7 @@ protected:
 	void BreakPinLinks(bool bSendNodeNotification);
 
 	// SGraphEditor interface
-	virtual void OnGraphChanged( const FEdGraphEditAction& InAction ) OVERRIDE;
+	virtual void OnGraphChanged( const FEdGraphEditAction& InAction ) override;
 	// End of SGraphEditorInterface
 private:
 	FText GetZoomText() const;
diff --git a/Engine/Source/Editor/GraphEditor/Private/SGraphNodeKnot.cpp b/Engine/Source/Editor/GraphEditor/Private/SGraphNodeKnot.cpp
--- a/Engine/Source/Editor/GraphEditor/Private/SGraphNodeKnot.cpp
+++ b/Engine/Source/Editor/GraphEditor/Private/SGraphNodeKnot.cpp
@@ -20,13 +20,13 @@ public:
 	}
 		
 	// FDragDropOperation interface
-	virtual void OnDragged(const class FDragDropEvent& DragDropEvent) OVERRIDE;
+	virtual void OnDragged(const class FDragDropEvent& DragDropEvent) override;
 	// End of FDragDropOperation interface
 
 	UEdGraphPin* GetBestPin() const;
 
 	// FDragConnection interface
-	virtual void ValidateGraphPinList(TArray<UEdGraphPin*>& OutValidPins) OVERRIDE;
+	virtual void ValidateGraphPinList(TArray<UEdGraphPin*>& OutValidPins) override;
 	// End of FDragConnection interface
 
 protected:
@@ -141,15 +141,15 @@ public:
 	void Construct(const FArguments& InArgs, UEdGraphPin* InPin);
 
 	// SWidget interface
-	virtual void OnDragEnter(const FGeometry& MyGeometry, const FDragDropEvent& DragDropEvent) OVERRIDE;
+	virtual void OnDragEnter(const FGeometry& MyGeometry, const FDragDropEvent& DragDropEvent) override;
 	// End of SWidget interface
 
 protected:
 	// Begin SGraphPin interface
-	virtual TSharedRef<SWidget>	GetDefaultValueWidget() OVERRIDE;
-	virtual TSharedRef<FDragDropOperation> SpawnPinDragEvent(const TSharedRef<SGraphPanel>& InGraphPanel, const TArray< TSharedRef<SGraphPin> >& InStartingPins, bool bInShiftOperation) OVERRIDE;
-	virtual FReply OnPinMouseDown(const FGeometry& SenderGeometry, const FPointerEvent& MouseEvent) OVERRIDE;
-	virtual FSlateColor GetPinColor() const OVERRIDE;
+	virtual TSharedRef<SWidget>	GetDefaultValueWidget() override;
+	virtual TSharedRef<FDragDropOperation> SpawnPinDragEvent(const TSharedRef<SGraphPanel>& InGraphPanel, const TArray< TSharedRef<SGraphPin> >& InStartingPins, bool bInShiftOperation) override;
+	virtual FReply OnPinMouseDown(const FGeometry& SenderGeometry, const FPointerEvent& MouseEvent) override;
+	virtual FSlateColor GetPinColor() const override;
 	// End SGraphPin interface
 };
 
diff --git a/Engine/Source/Editor/GraphEditor/Private/SGraphNodeKnot.h b/Engine/Source/Editor/GraphEditor/Private/SGraphNodeKnot.h
--- a/Engine/Source/Editor/GraphEditor/Private/SGraphNodeKnot.h
+++ b/Engine/Source/Editor/GraphEditor/Private/SGraphNodeKnot.h
@@ -14,10 +14,10 @@ public:
 	void Construct(const FArguments& InArgs, class UK2Node_Knot* InKnot);
 
 	// SGraphNode interface
-	virtual void UpdateGraphNode() OVERRIDE;
-	virtual const FSlateBrush* GetShadowBrush(bool bSelected) const OVERRIDE;
-	virtual TSharedPtr<SGraphPin> CreatePinWidget(UEdGraphPin* Pin) const OVERRIDE;
-	virtual void AddPin(const TSharedRef<SGraphPin>& PinToAdd) OVERRIDE;
-	virtual void RequestRenameOnSpawn() OVERRIDE { }
+	virtual void UpdateGraphNode() override;
+	virtual const FSlateBrush* GetShadowBrush(bool bSelected) const override;
+	virtual TSharedPtr<SGraphPin> CreatePinWidget(UEdGraphPin* Pin) const override;
+	virtual void AddPin(const TSharedRef<SGraphPin>& PinToAdd) override;
+	virtual void RequestRenameOnSpawn() override { }
 	// End of SGraphNode interface
 };
diff --git a/Engine/Source/Editor/GraphEditor/Private/SGraphPanel.h b/Engine/Source/Editor/GraphEditor/Private/SGraphPanel.h
--- a/Engine/Source/Editor/GraphEditor/Private/SGraphPanel.h
+++ b/Engine/Source/Editor/GraphEditor/Private/SGraphPanel.h
@@ -93,24 +93,24 @@ public:
 	~SGraphPanel();
 public:
 	// SWidget interface
-	virtual void OnDragEnter( const FGeometry& MyGeometry, const FDragDropEvent& DragDropEvent ) OVERRIDE;
-	virtual void OnDragLeave( const FDragDropEvent& DragDropEvent ) OVERRIDE;
-	virtual FReply OnDragOver( const FGeometry& MyGeometry, const FDragDropEvent& DragDropEvent ) OVERRIDE;
-	virtual FReply OnDrop( const FGeometry& MyGeometry, const FDragDropEvent& DragDropEvent ) OVERRIDE;
-	virtual int32 OnPaint( const FGeometry& AllottedGeometry, const FSlateRect& MyClippingRect, FSlateWindowElementList& OutDrawElements, int32 LayerId, const FWidgetStyle& InWidgetStyle, bool bParentEnabled ) const OVERRIDE;
-	virtual FReply OnKeyDown( const FGeometry& MyGeometry, const FKeyboardEvent& InKeyboardEvent ) OVERRIDE;
-	virtual bool SupportsKeyboardFocus() const OVERRIDE;
+	virtual void OnDragEnter( const FGeometry& MyGeometry, const FDragDropEvent& DragDropEvent ) override;
+	virtual void OnDragLeave( const FDragDropEvent& DragDropEvent ) override;
+	virtual FReply OnDragOver( const FGeometry& MyGeometry, const FDragDropEvent& DragDropEvent ) override;
+	virtual FReply OnDrop( const FGeometry& MyGeometry, const FDragDropEvent& DragDropEvent ) override;
+	virtual int32 OnPaint( const FGeometry& AllottedGeometry, const FSlateRect& MyClippingRect, FSlateWindowElementList& OutDrawElements, int32 LayerId, const FWidgetStyle& InWidgetStyle, bool bParentEnabled ) const override;
+	virtual FReply OnKeyDown( const FGeometry& MyGeometry, const FKeyboardEvent& InKeyboardEvent ) override;
+	virtual bool SupportsKeyboardFocus() const override;
 	// End of SWidget interface
 
 	// SNodePanel interface
-	virtual TSharedPtr<SWidget> OnSummonContextMenu(const FGeometry& MyGeometry, const FPointerEvent& MouseEvent) OVERRIDE;
-	virtual bool OnHandleLeftMouseRelease(const FGeometry& MyGeometry, const FPointerEvent& MouseEvent) OVERRIDE;
-	virtual void AddGraphNode(const TSharedRef<SNode>& NodeToAdd) OVERRIDE;
-	virtual void RemoveAllNodes() OVERRIDE;
+	virtual TSharedPtr<SWidget> OnSummonContextMenu(const FGeometry& MyGeometry, const FPointerEvent& MouseEvent) override;
+	virtual bool OnHandleLeftMouseRelease(const FGeometry& MyGeometry, const FPointerEvent& MouseEvent) override;
+	virtual void AddGraphNode(const TSharedRef<SNode>& NodeToAdd) override;
+	virtual void RemoveAllNodes() override;
 	// End of SNodePanel interface
 
 	// FGCObject interface.
-	virtual void AddReferencedObjects( FReferenceCollector& Collector ) OVERRIDE;
+	virtual void AddReferencedObjects( FReferenceCollector& Collector ) override;
 	// End of FGCObject interface.
 
 	void ArrangeChildrenForContextMenuSummon(const FGeometry& AllottedGeometry, FArrangedChildren& ArrangedChildren) const;
diff --git a/Engine/Source/Editor/GraphEditor/Private/SNodePanel.cpp b/Engine/Source/Editor/GraphEditor/Private/SNodePanel.cpp
--- a/Engine/Source/Editor/GraphEditor/Private/SNodePanel.cpp
+++ b/Engine/Source/Editor/GraphEditor/Private/SNodePanel.cpp
@@ -51,13 +51,13 @@ struct FFixedZoomLevelsContainer : public FZoomLevelsContainer
 		}
 	}
 
-	float GetZoomAmount(int32 InZoomLevel) const OVERRIDE
+	float GetZoomAmount(int32 InZoomLevel) const override
 	{
 		checkSlow(ZoomLevels.IsValidIndex(InZoomLevel));
 		return ZoomLevels[InZoomLevel].ZoomAmount;
 	}
 
-	int32 GetNearestZoomLevel(float InZoomAmount) const OVERRIDE
+	int32 GetNearestZoomLevel(float InZoomAmount) const override
 	{
 		for (int32 ZoomLevelIndex=0; ZoomLevelIndex < GetNumZoomLevels(); ++ZoomLevelIndex)
 		{
@@ -70,23 +70,23 @@ struct FFixedZoomLevelsContainer : public FZoomLevelsContainer
 		return GetDefaultZoomLevel();
 	}
 	
-	FText GetZoomText(int32 InZoomLevel) const OVERRIDE
+	FText GetZoomText(int32 InZoomLevel) const override
 	{
 		checkSlow(ZoomLevels.IsValidIndex(InZoomLevel));
 		return ZoomLevels[InZoomLevel].DisplayText;
 	}
 	
-	int32 GetNumZoomLevels() const OVERRIDE
+	int32 GetNumZoomLevels() const override
 	{
 		return ZoomLevels.Num();
 	}
 	
-	int32 GetDefaultZoomLevel() const OVERRIDE
+	int32 GetDefaultZoomLevel() const override
 	{
 		return 12;
 	}
 
-	EGraphRenderingLOD::Type GetLOD(int32 InZoomLevel) const OVERRIDE
+	EGraphRenderingLOD::Type GetLOD(int32 InZoomLevel) const override
 	{
 		checkSlow(ZoomLevels.IsValidIndex(InZoomLevel));
 		return ZoomLevels[InZoomLevel].LOD;
diff --git a/Engine/Source/Editor/GraphEditor/Private/SoundNodes/SGraphNodeSoundBase.h b/Engine/Source/Editor/GraphEditor/Private/SoundNodes/SGraphNodeSoundBase.h
--- a/Engine/Source/Editor/GraphEditor/Private/SoundNodes/SGraphNodeSoundBase.h
+++ b/Engine/Source/Editor/GraphEditor/Private/SoundNodes/SGraphNodeSoundBase.h
@@ -12,9 +12,9 @@ public:
 
 protected:
 	// SGraphNode Interface
-	virtual void CreateOutputSideAddButton(TSharedPtr<SVerticalBox> OutputBox) OVERRIDE;
-	virtual EVisibility IsAddPinButtonVisible() const OVERRIDE;
-	virtual FReply OnAddPin() OVERRIDE;
+	virtual void CreateOutputSideAddButton(TSharedPtr<SVerticalBox> OutputBox) override;
+	virtual EVisibility IsAddPinButtonVisible() const override;
+	virtual FReply OnAddPin() override;
 
 private:
 	USoundCueGraphNode* SoundNode;
diff --git a/Engine/Source/Editor/GraphEditor/Private/SoundNodes/SGraphNodeSoundResult.h b/Engine/Source/Editor/GraphEditor/Private/SoundNodes/SGraphNodeSoundResult.h
--- a/Engine/Source/Editor/GraphEditor/Private/SoundNodes/SGraphNodeSoundResult.h
+++ b/Engine/Source/Editor/GraphEditor/Private/SoundNodes/SGraphNodeSoundResult.h
@@ -12,6 +12,6 @@ public:
 
 protected:
 	// SGraphNode interface
-	virtual TSharedRef<SWidget> CreateNodeContentArea() OVERRIDE;
+	virtual TSharedRef<SWidget> CreateNodeContentArea() override;
 	// End of SGraphNode interface
 };
diff --git a/Engine/Source/Editor/GraphEditor/Public/GraphEditorActions.h b/Engine/Source/Editor/GraphEditor/Public/GraphEditorActions.h
--- a/Engine/Source/Editor/GraphEditor/Public/GraphEditorActions.h
+++ b/Engine/Source/Editor/GraphEditor/Public/GraphEditorActions.h
@@ -17,7 +17,7 @@ public:
 	{
 	}
 
-	GRAPHEDITOR_API virtual void RegisterCommands() OVERRIDE;
+	GRAPHEDITOR_API virtual void RegisterCommands() override;
 
 	TSharedPtr< FUICommandInfo > ReconstructNodes;
 	TSharedPtr< FUICommandInfo > BreakNodeLinks;
diff --git a/Engine/Source/Editor/GraphEditor/Public/GraphEditorDragDropAction.h b/Engine/Source/Editor/GraphEditor/Public/GraphEditorDragDropAction.h
--- a/Engine/Source/Editor/GraphEditor/Public/GraphEditorDragDropAction.h
+++ b/Engine/Source/Editor/GraphEditor/Public/GraphEditorDragDropAction.h
@@ -62,8 +62,8 @@ public:
 	DRAG_DROP_OPERATOR_TYPE(FGraphSchemaActionDragDropAction, FGraphEditorDragDropAction)
 
 	// FGraphEditorDragDropAction interface
-	virtual void HoverTargetChanged() OVERRIDE;
-	virtual FReply DroppedOnPanel( const TSharedRef< class SWidget >& Panel, FVector2D ScreenPosition, FVector2D GraphPosition, UEdGraph& Graph) OVERRIDE;
+	virtual void HoverTargetChanged() override;
+	virtual FReply DroppedOnPanel( const TSharedRef< class SWidget >& Panel, FVector2D ScreenPosition, FVector2D GraphPosition, UEdGraph& Graph) override;
 	// End of FGraphEditorDragDropAction
 
 	static TSharedRef<FGraphSchemaActionDragDropAction> New(TSharedPtr<FEdGraphSchemaAction> InActionNode )
diff --git a/Engine/Source/Editor/GraphEditor/Public/SGraphActionMenu.h b/Engine/Source/Editor/GraphEditor/Public/SGraphActionMenu.h
--- a/Engine/Source/Editor/GraphEditor/Public/SGraphActionMenu.h
+++ b/Engine/Source/Editor/GraphEditor/Public/SGraphActionMenu.h
@@ -18,7 +18,7 @@ class GRAPHEDITOR_API SDefaultGraphActionWidget : public SCompoundWidget
 	SLATE_END_ARGS()
 
 	void Construct(const FArguments& InArgs, const FCreateWidgetForActionData* InCreateData);
-	virtual FReply OnMouseButtonDown( const FGeometry& MyGeometry, const FPointerEvent& MouseEvent ) OVERRIDE;
+	virtual FReply OnMouseButtonDown( const FGeometry& MyGeometry, const FPointerEvent& MouseEvent ) override;
 
 	/** The item that we want to display with this widget */
 	TWeakPtr<FEdGraphSchemaAction> ActionPtr;
@@ -116,7 +116,7 @@ public:
 	void Construct( const FArguments& InArgs, bool bIsReadOnly = true );
 
 	// FGCObject override
-	virtual void AddReferencedObjects( FReferenceCollector& Collector ) OVERRIDE;
+	virtual void AddReferencedObjects( FReferenceCollector& Collector ) override;
 
 	/**
 	 * Refreshes the actions that this widget should display
@@ -175,7 +175,7 @@ protected:
 	FGetFilterText OnGetFilterText;
 public:
 	// SWidget interface
-	virtual FReply OnKeyDown(const FGeometry& MyGeometry, const FKeyboardEvent& KeyboardEvent) OVERRIDE;
+	virtual FReply OnKeyDown(const FGeometry& MyGeometry, const FKeyboardEvent& KeyboardEvent) override;
 	// End of SWidget interface
 
 	/** Get filter text box widget */
diff --git a/Engine/Source/Editor/GraphEditor/Public/SGraphNode.h b/Engine/Source/Editor/GraphEditor/Public/SGraphNode.h
--- a/Engine/Source/Editor/GraphEditor/Public/SGraphNode.h
+++ b/Engine/Source/Editor/GraphEditor/Public/SGraphNode.h
@@ -33,7 +33,7 @@ public:
 	void Construct(const FArguments& InArgs, UEdGraphNode* InNode);
 
 	// SWidget interface
-	virtual void Tick(const FGeometry& AllottedGeometry, const double InCurrentTime, const float InDeltaTime) OVERRIDE;
+	virtual void Tick(const FGeometry& AllottedGeometry, const double InCurrentTime, const float InDeltaTime) override;
 	// End of SWidget interface
 
 	/** Returns the main title for the node */
@@ -65,26 +65,26 @@ class GRAPHEDITOR_API SGraphNode : public SNodePanel::SNode
 {
 public:
 	// SWidget interface
-	virtual void OnDragEnter( const FGeometry& MyGeometry, const FDragDropEvent& DragDropEvent ) OVERRIDE;
-	virtual void OnDragLeave( const FDragDropEvent& DragDropEvent ) OVERRIDE;
-	virtual FReply OnDragOver( const FGeometry& MyGeometry, const FDragDropEvent& DragDropEvent ) OVERRIDE;
-	virtual FReply OnDrop( const FGeometry& MyGeometry, const FDragDropEvent& DragDropEvent ) OVERRIDE;
+	virtual void OnDragEnter( const FGeometry& MyGeometry, const FDragDropEvent& DragDropEvent ) override;
+	virtual void OnDragLeave( const FDragDropEvent& DragDropEvent ) override;
+	virtual FReply OnDragOver( const FGeometry& MyGeometry, const FDragDropEvent& DragDropEvent ) override;
+	virtual FReply OnDrop( const FGeometry& MyGeometry, const FDragDropEvent& DragDropEvent ) override;
 
-	virtual FReply OnMouseButtonDown( const FGeometry& MyGeometry, const FPointerEvent& MouseEvent ) OVERRIDE;
-	virtual FReply OnMouseButtonUp( const FGeometry& MyGeometry, const FPointerEvent& MouseEvent ) OVERRIDE;
-	virtual FReply OnMouseButtonDoubleClick( const FGeometry& InMyGeometry, const FPointerEvent& InMouseEvent ) OVERRIDE;
+	virtual FReply OnMouseButtonDown( const FGeometry& MyGeometry, const FPointerEvent& MouseEvent ) override;
+	virtual FReply OnMouseButtonUp( const FGeometry& MyGeometry, const FPointerEvent& MouseEvent ) override;
+	virtual FReply OnMouseButtonDoubleClick( const FGeometry& InMyGeometry, const FPointerEvent& InMouseEvent ) override;
 
-	virtual TSharedPtr<IToolTip> GetToolTip() OVERRIDE;
-	virtual void OnToolTipClosing() OVERRIDE;
+	virtual TSharedPtr<IToolTip> GetToolTip() override;
+	virtual void OnToolTipClosing() override;
 
-	virtual void Tick( const FGeometry& AllottedGeometry, const double InCurrentTime, const float InDeltaTime ) OVERRIDE;
+	virtual void Tick( const FGeometry& AllottedGeometry, const double InCurrentTime, const float InDeltaTime ) override;
 	// End of SWidget interface
 
 	// SNodePanel::SNode interface
-	virtual void MoveTo( const FVector2D& NewPosition, FNodeSet& NodeFilter ) OVERRIDE;
-	virtual FVector2D GetPosition() const OVERRIDE;
-	virtual FString GetNodeComment() const OVERRIDE;
-	virtual UObject* GetObjectBeingDisplayed() const OVERRIDE;
+	virtual void MoveTo( const FVector2D& NewPosition, FNodeSet& NodeFilter ) override;
+	virtual FVector2D GetPosition() const override;
+	virtual FString GetNodeComment() const override;
+	virtual UObject* GetObjectBeingDisplayed() const override;
 	// End of SNodePanel::SNode interface
 
 	/** Set attribute for determining if widget is editable */
diff --git a/Engine/Source/Editor/GraphEditor/Public/SGraphNodeComment.h b/Engine/Source/Editor/GraphEditor/Public/SGraphNodeComment.h
--- a/Engine/Source/Editor/GraphEditor/Public/SGraphNodeComment.h
+++ b/Engine/Source/Editor/GraphEditor/Public/SGraphNodeComment.h
@@ -27,48 +27,48 @@ public:
 	SLATE_END_ARGS()
 
 	// Begin SWidget Interface
-	virtual FReply OnMouseMove( const FGeometry& MyGeometry, const FPointerEvent& MouseEvent ) OVERRIDE;
-	virtual void OnMouseEnter( const FGeometry& MyGeometry, const FPointerEvent& MouseEvent ) OVERRIDE;
-	virtual void OnMouseLeave( const FPointerEvent& MouseEvent ) OVERRIDE;
-	virtual FCursorReply OnCursorQuery( const FGeometry& MyGeometry, const FPointerEvent& CursorEvent ) const OVERRIDE;
-	virtual FReply OnMouseButtonDoubleClick( const FGeometry& InMyGeometry, const FPointerEvent& InMouseEvent ) OVERRIDE;
-	virtual FReply OnMouseButtonDown( const FGeometry& MyGeometry, const FPointerEvent& MouseEvent ) OVERRIDE;
-	virtual FReply OnMouseButtonUp( const FGeometry& MyGeometry, const FPointerEvent& MouseEvent ) OVERRIDE;
-	virtual bool OnHitTest( const FGeometry& MyGeometry, FVector2D InAbsoluteCursorPosition ) OVERRIDE;
-	virtual void Tick(const FGeometry& AllottedGeometry, const double InCurrentTime, const float InDeltaTime) OVERRIDE;
-	virtual FReply OnDrop( const FGeometry& MyGeometry, const FDragDropEvent& DragDropEvent ) OVERRIDE;
-	virtual void OnDragEnter( const FGeometry& MyGeometry, const FDragDropEvent& DragDropEvent ) OVERRIDE;
+	virtual FReply OnMouseMove( const FGeometry& MyGeometry, const FPointerEvent& MouseEvent ) override;
+	virtual void OnMouseEnter( const FGeometry& MyGeometry, const FPointerEvent& MouseEvent ) override;
+	virtual void OnMouseLeave( const FPointerEvent& MouseEvent ) override;
+	virtual FCursorReply OnCursorQuery( const FGeometry& MyGeometry, const FPointerEvent& CursorEvent ) const override;
+	virtual FReply OnMouseButtonDoubleClick( const FGeometry& InMyGeometry, const FPointerEvent& InMouseEvent ) override;
+	virtual FReply OnMouseButtonDown( const FGeometry& MyGeometry, const FPointerEvent& MouseEvent ) override;
+	virtual FReply OnMouseButtonUp( const FGeometry& MyGeometry, const FPointerEvent& MouseEvent ) override;
+	virtual bool OnHitTest( const FGeometry& MyGeometry, FVector2D InAbsoluteCursorPosition ) override;
+	virtual void Tick(const FGeometry& AllottedGeometry, const double InCurrentTime, const float InDeltaTime) override;
+	virtual FReply OnDrop( const FGeometry& MyGeometry, const FDragDropEvent& DragDropEvent ) override;
+	virtual void OnDragEnter( const FGeometry& MyGeometry, const FDragDropEvent& DragDropEvent ) override;
 	// End SWidget Interface
 
 	// Begin SNodePanel::SNode interface
-	virtual const FSlateBrush* GetShadowBrush(bool bSelected) const OVERRIDE;
+	virtual const FSlateBrush* GetShadowBrush(bool bSelected) const override;
 	virtual void GetOverlayBrushes(bool bSelected, const FVector2D WidgetSize, TArray<FOverlayBrushInfo>& Brushes) const;
-	virtual bool ShouldAllowCulling() const OVERRIDE { return false; }
+	virtual bool ShouldAllowCulling() const override { return false; }
 	// End SNodePanel::SNode interface
 
 	// Begin SPanel Interface
-	virtual FVector2D ComputeDesiredSize() const OVERRIDE;
+	virtual FVector2D ComputeDesiredSize() const override;
 	// End SPanel interface
 
 	// Begin SGraphNode Interface
-	virtual bool IsNameReadOnly() const OVERRIDE;
+	virtual bool IsNameReadOnly() const override;
 	// ENd SGraphNode Interface
 
 	void Construct( const FArguments& InArgs, UEdGraphNode* InNode );
 
 	/** return if the node can be selected, by pointing given location */
-	virtual bool CanBeSelected( const FVector2D& MousePositionInNode ) const OVERRIDE;
+	virtual bool CanBeSelected( const FVector2D& MousePositionInNode ) const override;
 
 	/** return size of the title bar */
-	virtual FVector2D GetDesiredSizeForMarquee() const OVERRIDE;
+	virtual FVector2D GetDesiredSizeForMarquee() const override;
 
 	/** return rect of the title bar */
-	virtual FSlateRect GetTitleRect() const OVERRIDE;
+	virtual FSlateRect GetTitleRect() const override;
 
 protected:
 	// SGraphNode Interface
-	virtual void UpdateGraphNode() OVERRIDE;
-	virtual bool ShouldScaleNodeComment()const OVERRIDE;
+	virtual void UpdateGraphNode() override;
+	virtual bool ShouldScaleNodeComment()const override;
 
 	/**
 	 * Helper method to update selection state of comment and any nodes 'contained' within it
@@ -78,7 +78,7 @@ protected:
 	void HandleSelection(bool bIsSelected, bool bUpdateNodesUnderComment = false) const;
 
 	/** called when user is moving the comment node */
-	virtual void MoveTo(const FVector2D& NewPosition, FNodeSet& NodeFilter ) OVERRIDE;
+	virtual void MoveTo(const FVector2D& NewPosition, FNodeSet& NodeFilter ) override;
 
 private:
 
diff --git a/Engine/Source/Editor/GraphEditor/Public/SGraphPalette.h b/Engine/Source/Editor/GraphEditor/Public/SGraphPalette.h
--- a/Engine/Source/Editor/GraphEditor/Public/SGraphPalette.h
+++ b/Engine/Source/Editor/GraphEditor/Public/SGraphPalette.h
@@ -25,10 +25,10 @@ public:
 	void Construct(const FArguments& InArgs, FCreateWidgetForActionData* const InCreateData);
 
 	// SWidget interface
-	virtual FReply OnDrop( const FGeometry& MyGeometry, const FDragDropEvent& DragDropEvent ) OVERRIDE;
-	virtual void OnDragEnter( const FGeometry& MyGeometry, const FDragDropEvent& DragDropEvent ) OVERRIDE;
-	virtual void OnDragLeave( const FDragDropEvent& DragDropEvent ) OVERRIDE;
-	virtual FReply OnMouseButtonDown( const FGeometry& MyGeometry, const FPointerEvent& MouseEvent ) OVERRIDE;
+	virtual FReply OnDrop( const FGeometry& MyGeometry, const FDragDropEvent& DragDropEvent ) override;
+	virtual void OnDragEnter( const FGeometry& MyGeometry, const FDragDropEvent& DragDropEvent ) override;
+	virtual void OnDragLeave( const FDragDropEvent& DragDropEvent ) override;
+	virtual FReply OnMouseButtonDown( const FGeometry& MyGeometry, const FPointerEvent& MouseEvent ) override;
 	// End of SWidget interface
 
 protected:
diff --git a/Engine/Source/Editor/GraphEditor/Public/SGraphPin.h b/Engine/Source/Editor/GraphEditor/Public/SGraphPin.h
--- a/Engine/Source/Editor/GraphEditor/Public/SGraphPin.h
+++ b/Engine/Source/Editor/GraphEditor/Public/SGraphPin.h
@@ -44,15 +44,15 @@ public:
 	FReply OnPinNameMouseDown(const FGeometry& SenderGeometry, const FPointerEvent& MouseEvent);
 
 	// SWidget interface
-	virtual FReply OnMouseButtonUp(const FGeometry& MyGeometry, const FPointerEvent& MouseEvent) OVERRIDE;
-	virtual FReply OnMouseMove(const FGeometry& MyGeometry, const FPointerEvent& MouseEvent) OVERRIDE;
-	virtual void OnMouseEnter(const FGeometry& MyGeometry, const FPointerEvent& MouseEvent) OVERRIDE;
-	virtual void OnMouseLeave(const FPointerEvent& MouseEvent) OVERRIDE;
-	virtual void OnDragEnter(const FGeometry& MyGeometry, const FDragDropEvent& DragDropEvent) OVERRIDE;
-	virtual void OnDragLeave(const FDragDropEvent& DragDropEvent) OVERRIDE;
-	virtual FReply OnDragOver(const FGeometry& MyGeometry, const FDragDropEvent& DragDropEvent) OVERRIDE;
-	virtual FReply OnDrop(const FGeometry& MyGeometry, const FDragDropEvent& DragDropEvent) OVERRIDE;
-	virtual void Tick( const FGeometry& AllottedGeometry, const double InCurrentTime, const float InDeltaTime ) OVERRIDE;
+	virtual FReply OnMouseButtonUp(const FGeometry& MyGeometry, const FPointerEvent& MouseEvent) override;
+	virtual FReply OnMouseMove(const FGeometry& MyGeometry, const FPointerEvent& MouseEvent) override;
+	virtual void OnMouseEnter(const FGeometry& MyGeometry, const FPointerEvent& MouseEvent) override;
+	virtual void OnMouseLeave(const FPointerEvent& MouseEvent) override;
+	virtual void OnDragEnter(const FGeometry& MyGeometry, const FDragDropEvent& DragDropEvent) override;
+	virtual void OnDragLeave(const FDragDropEvent& DragDropEvent) override;
+	virtual FReply OnDragOver(const FGeometry& MyGeometry, const FDragDropEvent& DragDropEvent) override;
+	virtual FReply OnDrop(const FGeometry& MyGeometry, const FDragDropEvent& DragDropEvent) override;
+	virtual void Tick( const FGeometry& AllottedGeometry, const double InCurrentTime, const float InDeltaTime ) override;
 	// End of SWidget interface
 
 public:
diff --git a/Engine/Source/Editor/GraphEditor/Public/SNodePanel.h b/Engine/Source/Editor/GraphEditor/Public/SNodePanel.h
--- a/Engine/Source/Editor/GraphEditor/Public/SNodePanel.h
+++ b/Engine/Source/Editor/GraphEditor/Public/SNodePanel.h
@@ -339,22 +339,22 @@ public:
 	};
 
 	// SPanel interface
-	virtual void OnArrangeChildren(const FGeometry& AllottedGeometry, FArrangedChildren& ArrangedChildren) const OVERRIDE;
-	virtual FVector2D ComputeDesiredSize() const OVERRIDE;
-	virtual FChildren* GetChildren() OVERRIDE;
+	virtual void OnArrangeChildren(const FGeometry& AllottedGeometry, FArrangedChildren& ArrangedChildren) const override;
+	virtual FVector2D ComputeDesiredSize() const override;
+	virtual FChildren* GetChildren() override;
 	// End of SPanel interface
 
 	// SWidget interface
-	virtual void Tick(const FGeometry& AllottedGeometry, const double InCurrentTime, const float InDeltaTime) OVERRIDE;
-	virtual FReply OnMouseButtonDown(const FGeometry& MyGeometry, const FPointerEvent& MouseEvent) OVERRIDE;
-	virtual FReply OnMouseButtonUp(const FGeometry& MyGeometry, const FPointerEvent& MouseEvent) OVERRIDE;
-	virtual FReply OnMouseMove(const FGeometry& MyGeometry, const FPointerEvent& MouseEvent) OVERRIDE;
-	virtual FReply OnMouseWheel(const FGeometry& MyGeometry, const FPointerEvent& MouseEvent) OVERRIDE;
-	virtual FCursorReply OnCursorQuery( const FGeometry& MyGeometry, const FPointerEvent& CursorEvent ) const OVERRIDE;
-	virtual FReply OnKeyDown( const FGeometry& MyGeometry, const FKeyboardEvent& InKeyboardEvent ) OVERRIDE;
-	virtual FReply OnKeyUp( const FGeometry& MyGeometry, const FKeyboardEvent& InKeyboardEvent ) OVERRIDE;
-	virtual void OnKeyboardFocusLost( const FKeyboardFocusEvent& InKeyboardFocusEvent ) OVERRIDE;
-	virtual FReply OnTouchGesture( const FGeometry& MyGeometry, const FPointerEvent& GestureEvent ) OVERRIDE;
+	virtual void Tick(const FGeometry& AllottedGeometry, const double InCurrentTime, const float InDeltaTime) override;
+	virtual FReply OnMouseButtonDown(const FGeometry& MyGeometry, const FPointerEvent& MouseEvent) override;
+	virtual FReply OnMouseButtonUp(const FGeometry& MyGeometry, const FPointerEvent& MouseEvent) override;
+	virtual FReply OnMouseMove(const FGeometry& MyGeometry, const FPointerEvent& MouseEvent) override;
+	virtual FReply OnMouseWheel(const FGeometry& MyGeometry, const FPointerEvent& MouseEvent) override;
+	virtual FCursorReply OnCursorQuery( const FGeometry& MyGeometry, const FPointerEvent& CursorEvent ) const override;
+	virtual FReply OnKeyDown( const FGeometry& MyGeometry, const FKeyboardEvent& InKeyboardEvent ) override;
+	virtual FReply OnKeyUp( const FGeometry& MyGeometry, const FKeyboardEvent& InKeyboardEvent ) override;
+	virtual void OnKeyboardFocusLost( const FKeyboardFocusEvent& InKeyboardFocusEvent ) override;
+	virtual FReply OnTouchGesture( const FGeometry& MyGeometry, const FPointerEvent& GestureEvent ) override;
 
 	// End of SWidget interface
 public:
diff --git a/Engine/Source/Editor/InputBindingEditor/Private/InputBindingEditorModule.cpp b/Engine/Source/Editor/InputBindingEditor/Private/InputBindingEditorModule.cpp
--- a/Engine/Source/Editor/InputBindingEditor/Private/InputBindingEditorModule.cpp
+++ b/Engine/Source/Editor/InputBindingEditor/Private/InputBindingEditorModule.cpp
@@ -18,7 +18,7 @@ public:
 
 	// Begin IInputBindingEditorModule interface
 
-	virtual TWeakPtr<SWidget> CreateInputBindingEditorPanel( ) OVERRIDE
+	virtual TWeakPtr<SWidget> CreateInputBindingEditorPanel( ) override
 	{
 		TSharedPtr<SWidget> Panel = SNew(SInputBindingEditorPanel);
 		BindingEditorPanels.Add(Panel);
@@ -26,7 +26,7 @@ public:
 		return Panel;
 	}
 
-	virtual void DestroyInputBindingEditorPanel( const TWeakPtr<SWidget>& Panel ) OVERRIDE
+	virtual void DestroyInputBindingEditorPanel( const TWeakPtr<SWidget>& Panel ) override
 	{
 		BindingEditorPanels.Remove(Panel.Pin());
 	}
diff --git a/Engine/Source/Editor/InputBindingEditor/Private/Widgets/SGestureEditor.h b/Engine/Source/Editor/InputBindingEditor/Private/Widgets/SGestureEditor.h
--- a/Engine/Source/Editor/InputBindingEditor/Private/Widgets/SGestureEditor.h
+++ b/Engine/Source/Editor/InputBindingEditor/Private/Widgets/SGestureEditor.h
@@ -82,17 +82,17 @@ private:
 	/** Gesture that is currently being edited. */
 	static TWeakPtr<SGestureEditor> GestureBeingEdited;
 
-	virtual FReply OnKeyDown( const FGeometry& MyGeometry, const FKeyboardEvent& InKeyboardEvent ) OVERRIDE;
+	virtual FReply OnKeyDown( const FGeometry& MyGeometry, const FKeyboardEvent& InKeyboardEvent ) override;
 
-	virtual FReply OnKeyUp( const FGeometry& MyGeometry, const FKeyboardEvent& InKeyboardEvent ) OVERRIDE;
+	virtual FReply OnKeyUp( const FGeometry& MyGeometry, const FKeyboardEvent& InKeyboardEvent ) override;
 
-	virtual FReply OnMouseButtonDown( const FGeometry& MyGeometry, const FPointerEvent& InMouseEvent ) OVERRIDE;
+	virtual FReply OnMouseButtonDown( const FGeometry& MyGeometry, const FPointerEvent& InMouseEvent ) override;
 
-	virtual FReply OnMouseButtonDoubleClick( const FGeometry& InMyGeometry, const FPointerEvent& InMouseEvent ) OVERRIDE;
+	virtual FReply OnMouseButtonDoubleClick( const FGeometry& InMyGeometry, const FPointerEvent& InMouseEvent ) override;
 
-	virtual void OnKeyboardFocusLost( const FKeyboardFocusEvent& InKeyboardFocusEvent ) OVERRIDE;
+	virtual void OnKeyboardFocusLost( const FKeyboardFocusEvent& InKeyboardFocusEvent ) override;
 
-	virtual FReply OnKeyChar( const FGeometry& MyGeometry, const FCharacterEvent& InCharacterEvent ) OVERRIDE;
+	virtual FReply OnKeyChar( const FGeometry& MyGeometry, const FCharacterEvent& InCharacterEvent ) override;
 
 	/** 
 	 * Called when the gesture changes 
diff --git a/Engine/Source/Editor/InputBindingEditor/Private/Widgets/SGestureTreeItem.h b/Engine/Source/Editor/InputBindingEditor/Private/Widgets/SGestureTreeItem.h
--- a/Engine/Source/Editor/InputBindingEditor/Private/Widgets/SGestureTreeItem.h
+++ b/Engine/Source/Editor/InputBindingEditor/Private/Widgets/SGestureTreeItem.h
@@ -114,7 +114,7 @@ public:
 	}
 
 	/** SWidget Interface*/
-	virtual FReply OnMouseButtonDown( const FGeometry& MyGeometry, const FPointerEvent& MouseEvent ) OVERRIDE
+	virtual FReply OnMouseButtonDown( const FGeometry& MyGeometry, const FPointerEvent& MouseEvent ) override
 	{
 		if( MouseEvent.GetEffectingButton() == EKeys::LeftMouseButton && FMultiBoxSettings::IsInToolbarEditMode() )
 		{
@@ -124,7 +124,7 @@ public:
 		return FReply::Unhandled();
 	}
 
-	virtual FReply OnDragDetected( const FGeometry& MyGeometry, const FPointerEvent& MouseEvent ) OVERRIDE
+	virtual FReply OnDragDetected( const FGeometry& MyGeometry, const FPointerEvent& MouseEvent ) override
 	{
 		FToolBarBuilder TempBuilder( MakeShareable(new FUICommandList), FMultiBoxCustomization::None );
 		TempBuilder.AddToolBarButton( TreeItem->CommandInfo.ToSharedRef() );
diff --git a/Engine/Source/Editor/InputBindingEditor/Private/Widgets/SInputBindingEditorPanel.h b/Engine/Source/Editor/InputBindingEditor/Private/Widgets/SInputBindingEditorPanel.h
--- a/Engine/Source/Editor/InputBindingEditor/Private/Widgets/SInputBindingEditorPanel.h
+++ b/Engine/Source/Editor/InputBindingEditor/Private/Widgets/SInputBindingEditorPanel.h
@@ -83,8 +83,8 @@ public:
 	void Construct(const FArguments& InArgs);
 
 	// SWidget interface
-	virtual bool SupportsKeyboardFocus() const OVERRIDE;
-	virtual FReply OnKeyboardFocusReceived( const FGeometry& MyGeometry, const FKeyboardFocusEvent& InKeyboardFocusEvent ) OVERRIDE;
+	virtual bool SupportsKeyboardFocus() const override;
+	virtual FReply OnKeyboardFocusReceived( const FGeometry& MyGeometry, const FKeyboardFocusEvent& InKeyboardFocusEvent ) override;
 	// End of SWidget interface
 
 private:
diff --git a/Engine/Source/Editor/Kismet/Classes/BlueprintPaletteFavorites.h b/Engine/Source/Editor/Kismet/Classes/BlueprintPaletteFavorites.h
--- a/Engine/Source/Editor/Kismet/Classes/BlueprintPaletteFavorites.h
+++ b/Engine/Source/Editor/Kismet/Classes/BlueprintPaletteFavorites.h
@@ -108,8 +108,8 @@ class KISMET_API UBlueprintPaletteFavorites : public UObject
 
 public:
 	// Begin UObject interface
-	virtual void PostInitProperties() OVERRIDE;
-	virtual void PostEditChangeProperty(struct FPropertyChangedEvent& PropertyChangedEvent) OVERRIDE;
+	virtual void PostInitProperties() override;
+	virtual void PostEditChangeProperty(struct FPropertyChangedEvent& PropertyChangedEvent) override;
 	// End UObject interface
 
 	/**
diff --git a/Engine/Source/Editor/Kismet/Private/BPDelegateDragDropAction.h b/Engine/Source/Editor/Kismet/Private/BPDelegateDragDropAction.h
--- a/Engine/Source/Editor/Kismet/Private/BPDelegateDragDropAction.h
+++ b/Engine/Source/Editor/Kismet/Private/BPDelegateDragDropAction.h
@@ -12,11 +12,11 @@ public:
 	DRAG_DROP_OPERATOR_TYPE(FKismetDelegateDragDropAction, FKismetVariableDragDropAction)
 
 	// FGraphEditorDragDropAction interface
-	virtual void HoverTargetChanged() OVERRIDE { FGraphEditorDragDropAction::HoverTargetChanged(); }
-	virtual FReply DroppedOnPin(FVector2D ScreenPosition, FVector2D GraphPosition) OVERRIDE { return FGraphEditorDragDropAction::DroppedOnPin(ScreenPosition, GraphPosition); }
-	virtual FReply DroppedOnNode(FVector2D ScreenPosition, FVector2D GraphPosition) OVERRIDE { return FGraphEditorDragDropAction::DroppedOnNode(ScreenPosition, GraphPosition); }
+	virtual void HoverTargetChanged() override { FGraphEditorDragDropAction::HoverTargetChanged(); }
+	virtual FReply DroppedOnPin(FVector2D ScreenPosition, FVector2D GraphPosition) override { return FGraphEditorDragDropAction::DroppedOnPin(ScreenPosition, GraphPosition); }
+	virtual FReply DroppedOnNode(FVector2D ScreenPosition, FVector2D GraphPosition) override { return FGraphEditorDragDropAction::DroppedOnNode(ScreenPosition, GraphPosition); }
 
-	virtual FReply DroppedOnPanel(const TSharedRef< SWidget >& Panel, FVector2D ScreenPosition, FVector2D GraphPosition, UEdGraph& Graph) OVERRIDE;
+	virtual FReply DroppedOnPanel(const TSharedRef< SWidget >& Panel, FVector2D ScreenPosition, FVector2D GraphPosition, UEdGraph& Graph) override;
 	// End of FGraphEditorDragDropAction
 
 	bool IsValid() const;
diff --git a/Engine/Source/Editor/Kismet/Private/BPFunctionDragDropAction.h b/Engine/Source/Editor/Kismet/Private/BPFunctionDragDropAction.h
--- a/Engine/Source/Editor/Kismet/Private/BPFunctionDragDropAction.h
+++ b/Engine/Source/Editor/Kismet/Private/BPFunctionDragDropAction.h
@@ -15,8 +15,8 @@ public:
 	DRAG_DROP_OPERATOR_TYPE(FKismetDragDropAction, FGraphSchemaActionDragDropAction)
 		
 	// FGraphEditorDragDropAction interface
-	virtual void HoverTargetChanged() OVERRIDE;
-	virtual FReply DroppedOnPanel( const TSharedRef< class SWidget >& Panel, FVector2D ScreenPosition, FVector2D GraphPosition, UEdGraph& Graph) OVERRIDE;
+	virtual void HoverTargetChanged() override;
+	virtual FReply DroppedOnPanel( const TSharedRef< class SWidget >& Panel, FVector2D ScreenPosition, FVector2D GraphPosition, UEdGraph& Graph) override;
 	// End of FGraphSchemaActionDragDropAction
 
 	DECLARE_DELEGATE_RetVal_ThreeParams(bool, FCanBeDroppedDelegate, TSharedPtr<FEdGraphSchemaAction> /*DropAction*/, UEdGraph* /*HoveredGraphIn*/, FText& /*ImpededReasonOut*/);
@@ -53,8 +53,8 @@ public:
 	FKismetFunctionDragDropAction();
 
 	// FGraphEditorDragDropAction interface
-	virtual void HoverTargetChanged() OVERRIDE;
-	virtual FReply DroppedOnPanel( const TSharedRef< class SWidget >& Panel, FVector2D ScreenPosition, FVector2D GraphPosition, UEdGraph& Graph) OVERRIDE;
+	virtual void HoverTargetChanged() override;
+	virtual FReply DroppedOnPanel( const TSharedRef< class SWidget >& Panel, FVector2D ScreenPosition, FVector2D GraphPosition, UEdGraph& Graph) override;
 	// End of FGraphEditorDragDropAction
 
 	static TSharedRef<FKismetFunctionDragDropAction> New(FName InFunctionName, UClass* InOwningClass, const FMemberReference& InCallOnMember, FNodeCreationAnalytic AnalyticCallback, FCanBeDroppedDelegate CanBeDroppedDelegate = FCanBeDroppedDelegate());
@@ -86,8 +86,8 @@ public:
 	FKismetMacroDragDropAction();
 
 	// FGraphEditorDragDropAction interface
-	virtual void HoverTargetChanged() OVERRIDE;
-	virtual FReply DroppedOnPanel( const TSharedRef< class SWidget >& Panel, FVector2D ScreenPosition, FVector2D GraphPosition, UEdGraph& Graph) OVERRIDE;
+	virtual void HoverTargetChanged() override;
+	virtual FReply DroppedOnPanel( const TSharedRef< class SWidget >& Panel, FVector2D ScreenPosition, FVector2D GraphPosition, UEdGraph& Graph) override;
 	// End of FGraphEditorDragDropAction
 
 	static TSharedRef<FKismetMacroDragDropAction> New(FName InMacroName, UBlueprint* InBlueprint, UEdGraph* InMacro, FNodeCreationAnalytic AnalyticCallback);
diff --git a/Engine/Source/Editor/Kismet/Private/BPVariableDragDropAction.h b/Engine/Source/Editor/Kismet/Private/BPVariableDragDropAction.h
--- a/Engine/Source/Editor/Kismet/Private/BPVariableDragDropAction.h
+++ b/Engine/Source/Editor/Kismet/Private/BPVariableDragDropAction.h
@@ -13,11 +13,11 @@ public:
 
 	// FGraphEditorDragDropAction interface
 	virtual void HoverTargetChanged();
-	virtual FReply DroppedOnPin(FVector2D ScreenPosition, FVector2D GraphPosition) OVERRIDE;
-	virtual FReply DroppedOnNode(FVector2D ScreenPosition, FVector2D GraphPosition) OVERRIDE;
-	virtual FReply DroppedOnPanel( const TSharedRef< class SWidget >& Panel, FVector2D ScreenPosition, FVector2D GraphPosition, UEdGraph& Graph) OVERRIDE;
-	virtual FReply DroppedOnAction(TSharedRef<struct FEdGraphSchemaAction> Action) OVERRIDE;
-	virtual FReply DroppedOnCategory(FString Category) OVERRIDE;
+	virtual FReply DroppedOnPin(FVector2D ScreenPosition, FVector2D GraphPosition) override;
+	virtual FReply DroppedOnNode(FVector2D ScreenPosition, FVector2D GraphPosition) override;
+	virtual FReply DroppedOnPanel( const TSharedRef< class SWidget >& Panel, FVector2D ScreenPosition, FVector2D GraphPosition, UEdGraph& Graph) override;
+	virtual FReply DroppedOnAction(TSharedRef<struct FEdGraphSchemaAction> Action) override;
+	virtual FReply DroppedOnCategory(FString Category) override;
 	// End of FGraphEditorDragDropAction
 
 	static TSharedRef<FKismetVariableDragDropAction> New(FName InVariableName, UStruct* InVariableSource, FNodeCreationAnalytic AnalyticCallback)
diff --git a/Engine/Source/Editor/Kismet/Private/BlueprintDetailsCustomization.cpp b/Engine/Source/Editor/Kismet/Private/BlueprintDetailsCustomization.cpp
--- a/Engine/Source/Editor/Kismet/Private/BlueprintDetailsCustomization.cpp
+++ b/Engine/Source/Editor/Kismet/Private/BlueprintDetailsCustomization.cpp
@@ -2699,7 +2699,7 @@ struct FParamsChangedHelper : public FBasePinChangeHelper
 	TSet<UBlueprint*> ModifiedBlueprints;
 	TSet<UEdGraph*> ModifiedGraphs;
 
-	virtual void EditCompositeTunnelNode(UK2Node_Tunnel* TunnelNode) OVERRIDE
+	virtual void EditCompositeTunnelNode(UK2Node_Tunnel* TunnelNode) override
 	{
 		if (TunnelNode->InputSinkNode != NULL)
 		{
@@ -2712,7 +2712,7 @@ struct FParamsChangedHelper : public FBasePinChangeHelper
 		}
 	}
 
-	virtual void EditMacroInstance(UK2Node_MacroInstance* MacroInstance, UBlueprint* Blueprint) OVERRIDE
+	virtual void EditMacroInstance(UK2Node_MacroInstance* MacroInstance, UBlueprint* Blueprint) override
 	{
 		MacroInstance->ReconstructNode();
 		if (Blueprint)
@@ -2721,7 +2721,7 @@ struct FParamsChangedHelper : public FBasePinChangeHelper
 		}
 	}
 
-	virtual void EditCallSite(UK2Node_CallFunction* CallSite, UBlueprint* Blueprint) OVERRIDE
+	virtual void EditCallSite(UK2Node_CallFunction* CallSite, UBlueprint* Blueprint) override
 	{
 		CallSite->Modify();
 		CallSite->ReconstructNode();
@@ -2731,7 +2731,7 @@ struct FParamsChangedHelper : public FBasePinChangeHelper
 		}
 	}
 
-	virtual void EditDelegates(UK2Node_BaseMCDelegate* CallSite, UBlueprint* Blueprint) OVERRIDE
+	virtual void EditDelegates(UK2Node_BaseMCDelegate* CallSite, UBlueprint* Blueprint) override
 	{
 		CallSite->Modify();
 		CallSite->ReconstructNode();
@@ -2755,7 +2755,7 @@ struct FParamsChangedHelper : public FBasePinChangeHelper
 		}
 	}
 
-	virtual void EditCreateDelegates(UK2Node_CreateDelegate* CallSite) OVERRIDE
+	virtual void EditCreateDelegates(UK2Node_CreateDelegate* CallSite) override
 	{
 		UBlueprint* Blueprint = NULL;
 		UEdGraph* Graph = NULL;
@@ -2815,7 +2815,7 @@ struct FPinRenamedHelper : public FBasePinChangeHelper
 	TSet<UBlueprint*> ModifiedBlueprints;
 	TSet<UK2Node*> NodesToRename;
 
-	virtual void EditMacroInstance(UK2Node_MacroInstance* MacroInstance, UBlueprint* Blueprint) OVERRIDE
+	virtual void EditMacroInstance(UK2Node_MacroInstance* MacroInstance, UBlueprint* Blueprint) override
 	{
 		NodesToRename.Add(MacroInstance);
 		if (Blueprint)
@@ -2824,7 +2824,7 @@ struct FPinRenamedHelper : public FBasePinChangeHelper
 		}
 	}
 
-	virtual void EditCallSite(UK2Node_CallFunction* CallSite, UBlueprint* Blueprint) OVERRIDE
+	virtual void EditCallSite(UK2Node_CallFunction* CallSite, UBlueprint* Blueprint) override
 	{
 		NodesToRename.Add(CallSite);
 		if (Blueprint)
diff --git a/Engine/Source/Editor/Kismet/Private/BlueprintDetailsCustomization.h b/Engine/Source/Editor/Kismet/Private/BlueprintDetailsCustomization.h
--- a/Engine/Source/Editor/Kismet/Private/BlueprintDetailsCustomization.h
+++ b/Engine/Source/Editor/Kismet/Private/BlueprintDetailsCustomization.h
@@ -31,7 +31,7 @@ public:
 	}
 	
 	/** IDetailCustomization interface */
-	virtual void CustomizeDetails( IDetailLayoutBuilder& DetailLayout ) OVERRIDE;
+	virtual void CustomizeDetails( IDetailLayoutBuilder& DetailLayout ) override;
 	
 	static void PopulateCategories(SMyBlueprint* MyBlueprint, TArray<TSharedPtr<FString>>& CategorySource);
 
@@ -154,7 +154,7 @@ public:
 	}
 	
 	/** IDetailCustomization interface */
-	virtual void CustomizeDetails( IDetailLayoutBuilder& DetailLayout ) OVERRIDE 
+	virtual void CustomizeDetails( IDetailLayoutBuilder& DetailLayout ) override 
 	{
 		check(false);
 	}
@@ -229,10 +229,10 @@ public:
 	}
 	
 	/** IDetailCustomization interface */
-	virtual void CustomizeDetails( IDetailLayoutBuilder& DetailLayout ) OVERRIDE;
+	virtual void CustomizeDetails( IDetailLayoutBuilder& DetailLayout ) override;
 
 	/** Gets the graph that we are currently editing */
-	virtual UEdGraph* GetGraph() const OVERRIDE;
+	virtual UEdGraph* GetGraph() const override;
 
 private:
 
@@ -273,13 +273,13 @@ public:
 
 private:
 	/** IDetailCustomNodeBuilder Interface*/
-	virtual void SetOnRebuildChildren( FSimpleDelegate InOnRegenerateChildren ) OVERRIDE;
-	virtual void GenerateHeaderRowContent( FDetailWidgetRow& NodeRow ) OVERRIDE {}
-	virtual void GenerateChildContent( IDetailChildrenBuilder& ChildrenBuilder ) OVERRIDE;
-	virtual void Tick( float DeltaTime ) OVERRIDE {}
-	virtual bool RequiresTick() const OVERRIDE { return false; }
-	virtual FName GetName() const OVERRIDE { return NAME_None; }
-	virtual bool InitiallyCollapsed() const OVERRIDE { return false; }
+	virtual void SetOnRebuildChildren( FSimpleDelegate InOnRegenerateChildren ) override;
+	virtual void GenerateHeaderRowContent( FDetailWidgetRow& NodeRow ) override {}
+	virtual void GenerateChildContent( IDetailChildrenBuilder& ChildrenBuilder ) override;
+	virtual void Tick( float DeltaTime ) override {}
+	virtual bool RequiresTick() const override { return false; }
+	virtual FName GetName() const override { return NAME_None; }
+	virtual bool InitiallyCollapsed() const override { return false; }
 	
 private:
 	/** The parent graph action details customization */
@@ -302,13 +302,13 @@ public:
 	
 private:
 	/** IDetailCustomNodeBuilder Interface*/
-	virtual void SetOnRebuildChildren( FSimpleDelegate InOnRegenerateChildren ) OVERRIDE {}
-	virtual void GenerateHeaderRowContent( FDetailWidgetRow& NodeRow ) OVERRIDE;
-	virtual void GenerateChildContent( IDetailChildrenBuilder& ChildrenBuilder ) OVERRIDE;
-	virtual void Tick( float DeltaTime ) OVERRIDE {}
-	virtual bool RequiresTick() const OVERRIDE { return false; }
-	virtual FName GetName() const OVERRIDE { return ArgumentName; }
-	virtual bool InitiallyCollapsed() const OVERRIDE { return true; }
+	virtual void SetOnRebuildChildren( FSimpleDelegate InOnRegenerateChildren ) override {}
+	virtual void GenerateHeaderRowContent( FDetailWidgetRow& NodeRow ) override;
+	virtual void GenerateChildContent( IDetailChildrenBuilder& ChildrenBuilder ) override;
+	virtual void Tick( float DeltaTime ) override {}
+	virtual bool RequiresTick() const override { return false; }
+	virtual FName GetName() const override { return ArgumentName; }
+	virtual bool InitiallyCollapsed() const override { return true; }
 
 private:
 	/** Determines if this pin should not be editable */
@@ -368,7 +368,7 @@ public:
 	}
 	
 	/** IDetailCustomization interface */
-	virtual void CustomizeDetails( IDetailLayoutBuilder& DetailLayout ) OVERRIDE;
+	virtual void CustomizeDetails( IDetailLayoutBuilder& DetailLayout ) override;
 
 private:
 	struct FAccessSpecifierLabel
@@ -462,13 +462,13 @@ public:
 
 private:
 	/** IDetailCustomNodeBuilder Interface*/
-	virtual void SetOnRebuildChildren( FSimpleDelegate InOnRegenerateChildren ) OVERRIDE {RegenerateChildrenDelegate = InOnRegenerateChildren;}
-	virtual void GenerateHeaderRowContent( FDetailWidgetRow& NodeRow ) OVERRIDE;
-	virtual void GenerateChildContent( IDetailChildrenBuilder& ChildrenBuilder ) OVERRIDE;
-	virtual void Tick( float DeltaTime ) OVERRIDE {}
-	virtual bool RequiresTick() const OVERRIDE { return false; }
-	virtual FName GetName() const OVERRIDE { return NAME_None; }
-	virtual bool InitiallyCollapsed() const OVERRIDE { return false; }
+	virtual void SetOnRebuildChildren( FSimpleDelegate InOnRegenerateChildren ) override {RegenerateChildrenDelegate = InOnRegenerateChildren;}
+	virtual void GenerateHeaderRowContent( FDetailWidgetRow& NodeRow ) override;
+	virtual void GenerateChildContent( IDetailChildrenBuilder& ChildrenBuilder ) override;
+	virtual void Tick( float DeltaTime ) override {}
+	virtual bool RequiresTick() const override { return false; }
+	virtual FName GetName() const override { return NAME_None; }
+	virtual bool InitiallyCollapsed() const override { return false; }
 	
 private:
 	/** Callbacks for details UI */
@@ -514,7 +514,7 @@ public:
 	}
 
 	/** IDetailCustomization interface */
-	virtual void CustomizeDetails( IDetailLayoutBuilder& DetailLayout ) OVERRIDE;
+	virtual void CustomizeDetails( IDetailLayoutBuilder& DetailLayout ) override;
 	
 	/** Gets the Blueprint being edited */
 	UBlueprint* GetBlueprintObj() const;
@@ -564,7 +564,7 @@ public:
 	}
 
 	/** IDetailCustomization interface */
-	virtual void CustomizeDetails( IDetailLayoutBuilder& DetailLayout ) OVERRIDE;
+	virtual void CustomizeDetails( IDetailLayoutBuilder& DetailLayout ) override;
 
 protected:
 	/** Accessors passed to parent */
@@ -639,7 +639,7 @@ public:
 	}
 
 	/** IDetailCustomization interface */
-	virtual void CustomizeDetails( IDetailLayoutBuilder& DetailLayout ) OVERRIDE;
+	virtual void CustomizeDetails( IDetailLayoutBuilder& DetailLayout ) override;
 	
 private:
 
@@ -670,7 +670,7 @@ public:
 	FChildActorComponentDetails(TWeakPtr<FBlueprintEditor> BlueprintEditorPtr);
 
 	/** IDetailCustomization interface */
-	virtual void CustomizeDetails(IDetailLayoutBuilder& DetailBuilder) OVERRIDE;
+	virtual void CustomizeDetails(IDetailLayoutBuilder& DetailBuilder) override;
 
 private:
 	/** Weak reference to the Blueprint editor */
diff --git a/Engine/Source/Editor/Kismet/Private/BlueprintEditorCommands.cpp b/Engine/Source/Editor/Kismet/Private/BlueprintEditorCommands.cpp
--- a/Engine/Source/Editor/Kismet/Private/BlueprintEditorCommands.cpp
+++ b/Engine/Source/Editor/Kismet/Private/BlueprintEditorCommands.cpp
@@ -162,7 +162,7 @@ public:
 	FEdGraphNodeSpawnInfo(UClass* InClass) : NodeClass(InClass), GraphNode(NULL) {}
 
 	// FNodeSpawnInfo interface
-	virtual TSharedPtr< FEdGraphSchemaAction > GetAction(FBlueprintPaletteListBuilder& InPaletteBuilder, UEdGraph* InDestGraph) OVERRIDE
+	virtual TSharedPtr< FEdGraphSchemaAction > GetAction(FBlueprintPaletteListBuilder& InPaletteBuilder, UEdGraph* InDestGraph) override
 	{
 		TSharedPtr<FEdGraphSchemaAction_NewNode> NewActionNode = TSharedPtr<FEdGraphSchemaAction_NewNode>(new FEdGraphSchemaAction_NewNode);
 
@@ -194,7 +194,7 @@ public:
 	FFunctionNodeSpawnInfo(UFunction* InFunctionPtr) : FunctionPtr(InFunctionPtr) {}
 
 	// FNodeSpawnInfo interface
-	virtual TSharedPtr< FEdGraphSchemaAction > GetAction(FBlueprintPaletteListBuilder& InPaletteBuilder, UEdGraph* InDestGraph) OVERRIDE
+	virtual TSharedPtr< FEdGraphSchemaAction > GetAction(FBlueprintPaletteListBuilder& InPaletteBuilder, UEdGraph* InDestGraph) override
 	{
 		const UBlueprint* Blueprint = InPaletteBuilder.Blueprint;
 		FBlueprintPaletteListBuilder TempListBuilder(Blueprint);
@@ -247,7 +247,7 @@ public:
 	FMacroNodeSpawnInfo(UEdGraph* InMacroGraph) : MacroGraph(InMacroGraph) {}
 
 	// FNodeSpawnInfo interface
-	virtual TSharedPtr< FEdGraphSchemaAction > GetAction(FBlueprintPaletteListBuilder& InPaletteBuilder, UEdGraph* InDestGraph) OVERRIDE
+	virtual TSharedPtr< FEdGraphSchemaAction > GetAction(FBlueprintPaletteListBuilder& InPaletteBuilder, UEdGraph* InDestGraph) override
 	{
 		const UBlueprint* Blueprint = InPaletteBuilder.Blueprint;
 		FBlueprintPaletteListBuilder TempListBuilder(Blueprint);
diff --git a/Engine/Source/Editor/Kismet/Private/BlueprintEditorCommands.h b/Engine/Source/Editor/Kismet/Private/BlueprintEditorCommands.h
--- a/Engine/Source/Editor/Kismet/Private/BlueprintEditorCommands.h
+++ b/Engine/Source/Editor/Kismet/Private/BlueprintEditorCommands.h
@@ -15,7 +15,7 @@ public:
 	{
 	}	
 
-	virtual void RegisterCommands() OVERRIDE;
+	virtual void RegisterCommands() override;
 
 
 	// File-ish commands
@@ -81,7 +81,7 @@ public:
 	{
 	}	
 
-	virtual void RegisterCommands() OVERRIDE;
+	virtual void RegisterCommands() override;
 
 	/**
 	 * Returns a graph action assigned to the passed in gesture
diff --git a/Engine/Source/Editor/Kismet/Private/BlueprintEditorTabFactories.h b/Engine/Source/Editor/Kismet/Private/BlueprintEditorTabFactories.h
--- a/Engine/Source/Editor/Kismet/Private/BlueprintEditorTabFactories.h
+++ b/Engine/Source/Editor/Kismet/Private/BlueprintEditorTabFactories.h
@@ -70,7 +70,7 @@ public:
 		ViewMenuTooltip = LOCTEXT("DebugView_ToolTip", "Shows the debugging view");
 	}
 
-	virtual TSharedRef<SWidget> CreateTabBody(const FWorkflowTabSpawnInfo& Info) const OVERRIDE
+	virtual TSharedRef<SWidget> CreateTabBody(const FWorkflowTabSpawnInfo& Info) const override
 	{
 		TSharedPtr<FBlueprintEditor> BlueprintEditorPtr = StaticCastSharedPtr<FBlueprintEditor>(HostingApp.Pin());
 
@@ -96,14 +96,14 @@ public:
 		ViewMenuTooltip = LOCTEXT("DefaultEditorView_ToolTip", "Shows the default editor view");
 	}
 
-	virtual TSharedRef<SWidget> CreateTabBody(const FWorkflowTabSpawnInfo& Info) const OVERRIDE
+	virtual TSharedRef<SWidget> CreateTabBody(const FWorkflowTabSpawnInfo& Info) const override
 	{
 		TSharedPtr<FBlueprintEditor> BlueprintEditorPtr = StaticCastSharedPtr<FBlueprintEditor>(HostingApp.Pin());
 
 		return BlueprintEditorPtr->GetDefaultEditor();
 	}
 
-	virtual FText GetTabToolTipText(const FWorkflowTabSpawnInfo& Info) const OVERRIDE
+	virtual FText GetTabToolTipText(const FWorkflowTabSpawnInfo& Info) const override
 	{
 		return LOCTEXT("DefaultsEditorTooltip", "The defaults editor lets you set the default value for all variables in your Blueprint.");
 	}
@@ -127,7 +127,7 @@ public:
 
 	}
 
-	virtual void EvokeHistory(TSharedPtr<FTabInfo> InTabInfo) OVERRIDE
+	virtual void EvokeHistory(TSharedPtr<FTabInfo> InTabInfo) override
 	{
 		FWorkflowTabSpawnInfo SpawnInfo;
 		SpawnInfo.Payload = Payload;
@@ -139,7 +139,7 @@ public:
 		FactoryPtr.Pin()->UpdateTab(InTabInfo->GetTab().Pin(), SpawnInfo, GraphEditorRef);
 	}
 
-	virtual void SaveHistory() OVERRIDE
+	virtual void SaveHistory() override
 	{
 		if(IsHistoryValid())
 		{
@@ -148,7 +148,7 @@ public:
 		}
 	}
 
-	virtual void RestoreHistory() OVERRIDE
+	virtual void RestoreHistory() override
 	{
 		if(IsHistoryValid())
 		{
@@ -181,19 +181,19 @@ public:
 	{
 	}
 
-	virtual void OnTabActivated(TSharedPtr<SDockTab> Tab) const OVERRIDE
+	virtual void OnTabActivated(TSharedPtr<SDockTab> Tab) const override
 	{
 		TSharedRef<SGraphEditor> GraphEditor = StaticCastSharedRef<SGraphEditor>(Tab->GetContent());
 		BlueprintEditorPtr.Pin()->OnGraphEditorFocused(GraphEditor);
 	}
 
-	virtual void OnTabRefreshed(TSharedPtr<SDockTab> Tab) const OVERRIDE
+	virtual void OnTabRefreshed(TSharedPtr<SDockTab> Tab) const override
 	{
 		TSharedRef<SGraphEditor> GraphEditor = StaticCastSharedRef<SGraphEditor>(Tab->GetContent());
 		GraphEditor->NotifyGraphChanged();
 	}
 
-	virtual void SaveState(TSharedPtr<SDockTab> Tab, TSharedPtr<FTabPayload> Payload) const OVERRIDE
+	virtual void SaveState(TSharedPtr<SDockTab> Tab, TSharedPtr<FTabPayload> Payload) const override
 	{
 		TSharedRef<SGraphEditor> GraphEditor = StaticCastSharedRef<SGraphEditor>(Tab->GetContent());
 
@@ -215,23 +215,23 @@ public:
 //		BlueprintEditorPtr.Pin()->RequestSaveEditedObjectState();
 //	}
 protected:
-	virtual TAttribute<FText> ConstructTabNameForObject(UEdGraph* DocumentID) const OVERRIDE
+	virtual TAttribute<FText> ConstructTabNameForObject(UEdGraph* DocumentID) const override
 	{
 		return TAttribute<FText>::Create(TAttribute<FText>::FGetter::CreateStatic<UEdGraph*>(&FLocalKismetCallbacks::GetGraphDisplayName, DocumentID));
 	}
 
-	virtual TSharedRef<SWidget> CreateTabBodyForObject(const FWorkflowTabSpawnInfo& Info, UEdGraph* DocumentID) const OVERRIDE
+	virtual TSharedRef<SWidget> CreateTabBodyForObject(const FWorkflowTabSpawnInfo& Info, UEdGraph* DocumentID) const override
 	{
 		check(Info.TabInfo.IsValid());
 		return OnCreateGraphEditorWidget.Execute(Info.TabInfo.ToSharedRef(), DocumentID);
 	}
 
-	virtual const FSlateBrush* GetTabIconForObject(const FWorkflowTabSpawnInfo& Info, UEdGraph* DocumentID) const OVERRIDE
+	virtual const FSlateBrush* GetTabIconForObject(const FWorkflowTabSpawnInfo& Info, UEdGraph* DocumentID) const override
 	{
 		return FBlueprintEditor::GetGlyphForGraph(DocumentID, false);
 	}
 
-	virtual TSharedRef<FGenericTabHistory> CreateTabHistoryNode(TSharedPtr<FTabPayload> Payload) OVERRIDE
+	virtual TSharedRef<FGenericTabHistory> CreateTabHistoryNode(TSharedPtr<FTabPayload> Payload) override
 	{
 		return MakeShareable(new FGraphTabHistory(SharedThis(this), Payload));
 	}
@@ -252,22 +252,22 @@ public:
 		, BlueprintEditorPtr(InBlueprintEditorPtr)
 	{}
 protected:
-	virtual TAttribute<FText> ConstructTabNameForObject(UTimelineTemplate* DocumentID) const OVERRIDE
+	virtual TAttribute<FText> ConstructTabNameForObject(UTimelineTemplate* DocumentID) const override
 	{
 		return TAttribute<FText>::Create(TAttribute<FText>::FGetter::CreateStatic<UObject*>(&FLocalKismetCallbacks::GetObjectName, DocumentID));
 	}
 
-	virtual TSharedRef<SWidget> CreateTabBodyForObject(const FWorkflowTabSpawnInfo& Info, UTimelineTemplate* DocumentID) const OVERRIDE
+	virtual TSharedRef<SWidget> CreateTabBodyForObject(const FWorkflowTabSpawnInfo& Info, UTimelineTemplate* DocumentID) const override
 	{
 		return SNew(STimelineEditor, BlueprintEditorPtr.Pin(), DocumentID);
 	}
 
-	virtual const FSlateBrush* GetTabIconForObject(const FWorkflowTabSpawnInfo& Info, UTimelineTemplate* DocumentID) const OVERRIDE
+	virtual const FSlateBrush* GetTabIconForObject(const FWorkflowTabSpawnInfo& Info, UTimelineTemplate* DocumentID) const override
 	{
 		return FEditorStyle::GetBrush("GraphEditor.Timeline_16x");
 	}
 
-	virtual void SaveState(TSharedPtr<SDockTab> Tab, TSharedPtr<FTabPayload> Payload) const OVERRIDE
+	virtual void SaveState(TSharedPtr<SDockTab> Tab, TSharedPtr<FTabPayload> Payload) const override
 	{
 		UTimelineTemplate* Timeline = FTabPayload_UObject::CastChecked<UTimelineTemplate>(Payload);
 		BlueprintEditorPtr.Pin()->GetBlueprintObj()->LastEditedDocuments.Add(FEditedDocumentInfo(Timeline));
@@ -278,13 +278,13 @@ protected:
 // 		BlueprintEditorPtr.Pin()->RequestSaveEditedObjectState();
 // 	}
 
-	virtual void OnTabRefreshed(TSharedPtr<SDockTab> Tab) const OVERRIDE
+	virtual void OnTabRefreshed(TSharedPtr<SDockTab> Tab) const override
 	{
 		TSharedRef<STimelineEditor> TimelineEditor = StaticCastSharedRef<STimelineEditor>(Tab->GetContent());
 		TimelineEditor->OnTimelineChanged();
 	}
 
-	virtual FText GetTabToolTipText(const FWorkflowTabSpawnInfo& Info) const OVERRIDE
+	virtual FText GetTabToolTipText(const FWorkflowTabSpawnInfo& Info) const override
 	{
 		return LOCTEXT("TimelineEditorTooltip", "The Timeline editor lets you add tracks and visually keyframe values over time, to be used inside a Blueprint graph.");
 	}
@@ -311,14 +311,14 @@ public:
 		ViewMenuTooltip = LOCTEXT("ComponentsView_ToolTip", "Show the components view");
 	}
 
-	virtual TSharedRef<SWidget> CreateTabBody(const FWorkflowTabSpawnInfo& Info) const OVERRIDE
+	virtual TSharedRef<SWidget> CreateTabBody(const FWorkflowTabSpawnInfo& Info) const override
 	{
 		TSharedPtr<FBlueprintEditor> BlueprintEditorPtr = StaticCastSharedPtr<FBlueprintEditor>(HostingApp.Pin());
 
 		return BlueprintEditorPtr->GetSCSEditor();
 	}
 
-	virtual FText GetTabToolTipText(const FWorkflowTabSpawnInfo& Info) const OVERRIDE
+	virtual FText GetTabToolTipText(const FWorkflowTabSpawnInfo& Info) const override
 	{
 		return LOCTEXT("ComponentsTabTooltip", "The Components tab is for easily adding, selecting and attaching Components within your Blueprint.");
 	}
@@ -341,7 +341,7 @@ public:
 		ViewMenuTooltip = LOCTEXT("SCSViewportView_ToolTip", "Show the viewport view");
 	}
 
-	virtual TSharedRef<SWidget> CreateTabBody(const FWorkflowTabSpawnInfo& Info) const OVERRIDE
+	virtual TSharedRef<SWidget> CreateTabBody(const FWorkflowTabSpawnInfo& Info) const override
 	{
 		TSharedPtr<FBlueprintEditor> BlueprintEditorPtr = StaticCastSharedPtr<FBlueprintEditor>(HostingApp.Pin());
 
@@ -382,14 +382,14 @@ public:
 		ViewMenuTooltip = LOCTEXT("PaletteView_ToolTip", "Show palette of all functions and variables");
 	}
 
-	virtual TSharedRef<SWidget> CreateTabBody(const FWorkflowTabSpawnInfo& Info) const OVERRIDE
+	virtual TSharedRef<SWidget> CreateTabBody(const FWorkflowTabSpawnInfo& Info) const override
 	{
 		TSharedPtr<FBlueprintEditor> BlueprintEditorPtr = StaticCastSharedPtr<FBlueprintEditor>(HostingApp.Pin());
 
 		return BlueprintEditorPtr->GetPalette();
 	}
 
-	virtual FText GetTabToolTipText(const FWorkflowTabSpawnInfo& Info) const OVERRIDE
+	virtual FText GetTabToolTipText(const FWorkflowTabSpawnInfo& Info) const override
 	{
 		return LOCTEXT("PaletteTooltip", "The Palette tab provides access to _all_ nodes that can be placed (functions, variables, events etc).");
 	}
@@ -413,14 +413,14 @@ public:
 		ViewMenuTooltip = LOCTEXT("MyBlueprintTabView_ToolTip", "Show the my blueprint view");
 	}
 
-	virtual TSharedRef<SWidget> CreateTabBody(const FWorkflowTabSpawnInfo& Info) const OVERRIDE
+	virtual TSharedRef<SWidget> CreateTabBody(const FWorkflowTabSpawnInfo& Info) const override
 	{
 		TSharedPtr<FBlueprintEditor> BlueprintEditorPtr = StaticCastSharedPtr<FBlueprintEditor>(HostingApp.Pin());
 
 		return BlueprintEditorPtr->GetMyBlueprintWidget().ToSharedRef();
 	}
 
-	virtual FText GetTabToolTipText(const FWorkflowTabSpawnInfo& Info) const OVERRIDE
+	virtual FText GetTabToolTipText(const FWorkflowTabSpawnInfo& Info) const override
 	{
 		return LOCTEXT("MyBlueprintTooltip", "The MyBlueprint tab shows you elements that belong to _this_ Blueprint (variables, graphs etc.)");
 	}
@@ -444,14 +444,14 @@ public:
 		ViewMenuTooltip = LOCTEXT("CompilerResultsView_ToolTip", "Show compiler results of all functions and variables");
 	}
 
-	virtual TSharedRef<SWidget> CreateTabBody(const FWorkflowTabSpawnInfo& Info) const OVERRIDE
+	virtual TSharedRef<SWidget> CreateTabBody(const FWorkflowTabSpawnInfo& Info) const override
 	{
 		TSharedPtr<FBlueprintEditor> BlueprintEditorPtr = StaticCastSharedPtr<FBlueprintEditor>(HostingApp.Pin());
 
 		return BlueprintEditorPtr->GetCompilerResults();
 	}
 
-	virtual FText GetTabToolTipText(const FWorkflowTabSpawnInfo& Info) const OVERRIDE
+	virtual FText GetTabToolTipText(const FWorkflowTabSpawnInfo& Info) const override
 	{
 		return LOCTEXT("CompilerResultsTooltip", "The compiler results tab shows any errors or warnings generated when compiling this Blueprint.");
 	}
@@ -475,14 +475,14 @@ public:
 		ViewMenuTooltip = LOCTEXT("FindResultsView_ToolTip", "Show find results for searching in this blueprint or all blueprints");
 	}
 
-	virtual TSharedRef<SWidget> CreateTabBody(const FWorkflowTabSpawnInfo& Info) const OVERRIDE
+	virtual TSharedRef<SWidget> CreateTabBody(const FWorkflowTabSpawnInfo& Info) const override
 	{
 		TSharedPtr<FBlueprintEditor> BlueprintEditorPtr = StaticCastSharedPtr<FBlueprintEditor>(HostingApp.Pin());
 
 		return BlueprintEditorPtr->GetFindResults();
 	}
 
-	virtual FText GetTabToolTipText(const FWorkflowTabSpawnInfo& Info) const OVERRIDE
+	virtual FText GetTabToolTipText(const FWorkflowTabSpawnInfo& Info) const override
 	{
 		return LOCTEXT("FindResultsTooltip", "The find results tab shows results of searching this Blueprint (or all Blueprints).");
 	}
diff --git a/Engine/Source/Editor/Kismet/Private/Debugging/SKismetDebuggingView.cpp b/Engine/Source/Editor/Kismet/Private/Debugging/SKismetDebuggingView.cpp
--- a/Engine/Source/Editor/Kismet/Private/Debugging/SKismetDebuggingView.cpp
+++ b/Engine/Source/Editor/Kismet/Private/Debugging/SKismetDebuggingView.cpp
@@ -124,18 +124,18 @@ public:
 	{
 	}
 protected:
-	virtual bool Compare(const FDebugLineItem* BaseOther) const OVERRIDE
+	virtual bool Compare(const FDebugLineItem* BaseOther) const override
 	{
 		FMessageLineItem* Other = (FMessageLineItem*)BaseOther;
 		return Message == Other->Message;
 	}
 
-	virtual FDebugLineItem* Duplicate() const OVERRIDE
+	virtual FDebugLineItem* Duplicate() const override
 	{
 		return new FMessageLineItem(Message);
 	}
 
-	virtual FString GetDescription() const OVERRIDE
+	virtual FString GetDescription() const override
 	{
 		return Message;
 	}
@@ -158,21 +158,21 @@ public:
 		ParentObjectRef = ParentObject;
 	}
 protected:
-	virtual bool Compare(const FDebugLineItem* BaseOther) const OVERRIDE
+	virtual bool Compare(const FDebugLineItem* BaseOther) const override
 	{
 		FLatentActionLineItem* Other = (FLatentActionLineItem*)BaseOther;
 		return (ParentObjectRef.Get() == Other->ParentObjectRef.Get()) &&
 			(UUID == Other->UUID);
 	}
 
-	virtual FDebugLineItem* Duplicate() const OVERRIDE
+	virtual FDebugLineItem* Duplicate() const override
 	{
 		return new FLatentActionLineItem(UUID, ParentObjectRef.Get());
 	}
 protected:
-	virtual TSharedRef<SWidget> GenerateNameWidget() OVERRIDE;
-	virtual FString GetDescription() const OVERRIDE;
-	virtual FText GetDisplayName() const OVERRIDE;
+	virtual TSharedRef<SWidget> GenerateNameWidget() override;
+	virtual FString GetDescription() const override;
+	virtual FText GetDisplayName() const override;
 	void OnNavigateToLatentNode();
 
 	class UEdGraphNode* FindAssociatedNode() const;
@@ -265,19 +265,19 @@ public:
 		ParentObjectRef = ParentObject;
 	}
 
-	virtual bool Compare(const FDebugLineItem* BaseOther) const OVERRIDE
+	virtual bool Compare(const FDebugLineItem* BaseOther) const override
 	{
 		FWatchLineItem* Other = (FWatchLineItem*)BaseOther;
 		return (ParentObjectRef.Get() == Other->ParentObjectRef.Get()) &&
 			(ObjectRef.Get() == Other->ObjectRef.Get());
 	}
 
-	virtual FDebugLineItem* Duplicate() const OVERRIDE
+	virtual FDebugLineItem* Duplicate() const override
 	{
 		return new FWatchLineItem(Cast<UEdGraphPin>(ObjectRef.Get()), ParentObjectRef.Get());
 	}	
 
-	virtual void MakeMenu(class FMenuBuilder& MenuBuilder) OVERRIDE
+	virtual void MakeMenu(class FMenuBuilder& MenuBuilder) override
 	{
 		if (UEdGraphPin* WatchedPin = Cast<UEdGraphPin>(ObjectRef.Get()))
 		{
@@ -293,9 +293,9 @@ public:
 		}
 	}
 protected:
-	virtual FString GetDescription() const OVERRIDE;
-	virtual FText GetDisplayName() const OVERRIDE;
-	virtual TSharedRef<SWidget> GenerateNameWidget() OVERRIDE;
+	virtual FString GetDescription() const override;
+	virtual FText GetDisplayName() const override;
+	virtual TSharedRef<SWidget> GenerateNameWidget() override;
 
 	void OnNavigateToWatchLocation();
 };
@@ -399,19 +399,19 @@ public:
 		ParentObjectRef = ParentObject;
 	}
 
-	virtual bool Compare(const FDebugLineItem* BaseOther) const OVERRIDE
+	virtual bool Compare(const FDebugLineItem* BaseOther) const override
 	{
 		FBreakpointLineItem* Other = (FBreakpointLineItem*)BaseOther;
 		return (ParentObjectRef.Get() == Other->ParentObjectRef.Get()) &&
 			(BreakpointRef.Get() == Other->BreakpointRef.Get());
 	}
 
-	virtual FDebugLineItem* Duplicate() const OVERRIDE
+	virtual FDebugLineItem* Duplicate() const override
 	{
 		return new FBreakpointLineItem(BreakpointRef.Get(), ParentObjectRef.Get());
 	}	
 
-	virtual void MakeMenu(class FMenuBuilder& MenuBuilder) OVERRIDE
+	virtual void MakeMenu(class FMenuBuilder& MenuBuilder) override
 	{
 		UBreakpoint* Breakpoint = BreakpointRef.Get();
 		UBlueprint* ParentBlueprint = GetBlueprintForObject(ParentObjectRef.Get());
@@ -464,7 +464,7 @@ public:
 		}
 	}
 protected:
-	virtual TSharedRef<SWidget> GenerateNameWidget() OVERRIDE
+	virtual TSharedRef<SWidget> GenerateNameWidget() override
 	{
 		return SNew(SHorizontalBox)
 
@@ -585,12 +585,12 @@ public:
 		ObjectRef = Object;
 	}
 
-	virtual UObject* GetParentObject() OVERRIDE
+	virtual UObject* GetParentObject() override
 	{
 		return ObjectRef.Get();
 	}
 
-	virtual void GatherChildren(TArray<FDebugTreeItemPtr>& OutChildren) OVERRIDE
+	virtual void GatherChildren(TArray<FDebugTreeItemPtr>& OutChildren) override
 	{
 		if (UObject* ParentObject = ObjectRef.Get())
 		{
@@ -668,18 +668,18 @@ public:
 	}
 
 protected:
-	virtual bool Compare(const FDebugLineItem* BaseOther) const OVERRIDE
+	virtual bool Compare(const FDebugLineItem* BaseOther) const override
 	{
 		FParentLineItem* Other = (FParentLineItem*)BaseOther;
 		return ObjectRef.Get() == Other->ObjectRef.Get();
 	}
 
-	virtual FDebugLineItem* Duplicate() const OVERRIDE
+	virtual FDebugLineItem* Duplicate() const override
 	{
 		return new FParentLineItem(ObjectRef.Get());
 	}
 
-	virtual FText GetDisplayName() const OVERRIDE
+	virtual FText GetDisplayName() const override
 	{
 		UObject* Object = ObjectRef.Get();
 		AActor* Actor = Cast<AActor>(Object);
@@ -694,7 +694,7 @@ protected:
 		}
 	}
 
-	virtual void MakeMenu(class FMenuBuilder& MenuBuilder) OVERRIDE
+	virtual void MakeMenu(class FMenuBuilder& MenuBuilder) override
 	{
 		if (UBlueprint* BP = Cast<UBlueprint>(ObjectRef.Get()))
 		{
@@ -741,13 +741,13 @@ public:
 		StackIndex = InStackIndex;
 	}
 protected:
-	virtual bool Compare(const FDebugLineItem* BaseOther) const OVERRIDE
+	virtual bool Compare(const FDebugLineItem* BaseOther) const override
 	{
 		check(false);
 		return false;
 	}
 
-	virtual FDebugLineItem* Duplicate() const OVERRIDE
+	virtual FDebugLineItem* Duplicate() const override
 	{
 		check(false);
 		return NULL;
@@ -775,7 +775,7 @@ protected:
 		return NULL;
 	}
 
-	virtual FText GetDisplayName() const OVERRIDE
+	virtual FText GetDisplayName() const override
 	{
 		UEdGraphNode* Node = GetNode();
 		if (Node != NULL)
@@ -789,7 +789,7 @@ protected:
 	}
 
 	// Index icon and node name
-	virtual TSharedRef<SWidget> GenerateNameWidget() OVERRIDE
+	virtual TSharedRef<SWidget> GenerateNameWidget() override
 	{
 		return SNew(SHorizontalBox)
 
@@ -813,7 +813,7 @@ protected:
 	}
 
 	// Visit time and actor name
-	virtual TSharedRef<SWidget> GenerateValueWidget() OVERRIDE
+	virtual TSharedRef<SWidget> GenerateValueWidget() override
 	{
 		return SNew(SHorizontalBox)
 			
@@ -896,7 +896,7 @@ public:
 	{
 	}
 
-	virtual void GatherChildren(TArray<FDebugTreeItemPtr>& OutChildren) OVERRIDE
+	virtual void GatherChildren(TArray<FDebugTreeItemPtr>& OutChildren) override
 	{
 		const TSimpleRingBuffer<FKismetTraceSample>& TraceStack = FKismetDebugUtilities::GetTraceStack();
 		const int32 NumVisible = TraceStack.Num();
@@ -919,13 +919,13 @@ protected:
 		return LOCTEXT("ExecutionTrace", "Execution Trace");
 	}
 
-	virtual bool Compare(const FDebugLineItem* BaseOther) const OVERRIDE
+	virtual bool Compare(const FDebugLineItem* BaseOther) const override
 	{
 		check(false);
 		return false;
 	}
 
-	virtual FDebugLineItem* Duplicate() const OVERRIDE
+	virtual FDebugLineItem* Duplicate() const override
 	{
 		check(false);
 		return NULL;
@@ -943,7 +943,7 @@ public:
 	SLATE_BEGIN_ARGS(SDebugLineItem){}
 	SLATE_END_ARGS()
 
-	virtual TSharedRef<SWidget> GenerateWidgetForColumn( const FName& ColumnName ) OVERRIDE
+	virtual TSharedRef<SWidget> GenerateWidgetForColumn( const FName& ColumnName ) override
 	{
 		if (ColumnName == KismetDebugViewConstants::ColumnId_Name)
 		{
