PropertyEditor for UDS default values

diff --git a/Engine/Source/Editor/Kismet/Private/UserDefinedStructureEditor.cpp b/Engine/Source/Editor/Kismet/Private/UserDefinedStructureEditor.cpp
--- a/Engine/Source/Editor/Kismet/Private/UserDefinedStructureEditor.cpp
+++ b/Engine/Source/Editor/Kismet/Private/UserDefinedStructureEditor.cpp
@@ -2,9 +2,12 @@
 
 #include "BlueprintEditorPrivatePCH.h"
 #include "UserDefinedStructureEditor.h"
+#include "BlueprintEditorUtils.h"
+#include "PropertyEditorModule.h"
+#include "IStructureDetailsView.h"
 
 #include "PropertyCustomizationHelpers.h"
-#include "Editor/KismetWidgets//Public/SPinTypeSelector.h"
+#include "Editor/KismetWidgets/Public/SPinTypeSelector.h"
 #include "Editor/WorkspaceMenuStructure/Public/WorkspaceMenuStructureModule.h"
 
 #define LOCTEXT_NAMESPACE "StructureEditor"
@@ -139,11 +142,54 @@ class FUserDefinedStructureLayout : public IDetailCustomNodeBuilder, public TSha
 public:
 	FUserDefinedStructureLayout(TWeakPtr<class FUserDefinedStructureDetails> InStructureDetails) : StructureDetails(InStructureDetails) {}
 
+	void OnFinishedChangingProperties(const FPropertyChangedEvent& PropertyChangedEvent)
+	{
+		auto StructureDetailsSP = StructureDetails.Pin();
+		if (StructureDetailsSP.IsValid())
+		{
+			check(PropertyChangedEvent.MemberProperty
+				&& PropertyChangedEvent.MemberProperty->GetOwnerStruct()
+				&& (PropertyChangedEvent.MemberProperty->GetOwnerStruct() == StructureDetailsSP->GetUserDefinedStruct())
+				&& PropertyChangedEvent.MemberProperty->GetOwnerStruct()->IsA<UUserDefinedStruct>());
+
+			const UProperty* DirectProperty = PropertyChangedEvent.MemberProperty;
+			while (!Cast<const UUserDefinedStruct>(DirectProperty->GetOuter()))
+			{
+				DirectProperty = CastChecked<const UProperty>(DirectProperty->GetOuter());
+			}
+
+			FString DefaultValueString;
+			bool bDefaultValueSet = false;
+			{
+				auto DefaultStructInstance = StructureDetailsSP->GetStructInstance();
+				if (DefaultStructInstance.IsValid() && DefaultStructInstance->IsValid())
+				{
+					bDefaultValueSet = FBlueprintEditorUtils::PropertyValueToString(DirectProperty, DefaultStructInstance->GetStructMemory(), DefaultValueString);
+				}
+			}
+
+			const FGuid VarGuid = FStructureEditorUtils::GetGuidForProperty(DirectProperty);
+			if (bDefaultValueSet && VarGuid.IsValid())
+			{
+				FStructureEditorUtils::ChangeVariableDefaultValue(StructureDetailsSP->GetUserDefinedStruct(), VarGuid, DefaultValueString);
+			}
+		}
+	}
+
 	void OnChanged()
 	{
 		OnRegenerateChildren.ExecuteIfBound();
 	}
 
+	void ResetStructureView()
+	{
+		auto StructureDetailsViewPinned = StructureDetailsView.Pin();
+		if (StructureDetailsView.IsValid())
+		{
+			StructureDetailsViewPinned->SetStructureData(NULL);
+		}
+	}
+
 	FReply OnAddNewField()
 	{
 		auto StructureDetailsSP = StructureDetails.Pin();
@@ -152,7 +198,6 @@ public:
 			const UEdGraphSchema_K2* K2Schema = GetDefault<UEdGraphSchema_K2>();
 			const  FEdGraphPinType InitialType(K2Schema->PC_Boolean, TEXT(""), NULL, false, false);
 			FStructureEditorUtils::AddVariable(StructureDetailsSP->GetUserDefinedStruct(), InitialType);
-			OnChanged();
 		}
 
 		return FReply::Handled();
@@ -249,7 +294,7 @@ public:
 
 private:
 	TWeakPtr<class FUserDefinedStructureDetails> StructureDetails;
-
+	TWeakPtr<class IStructureDetailsView> StructureDetailsView;
 	FSimpleDelegate OnRegenerateChildren;
 };
 
@@ -287,7 +332,6 @@ public:
 		{
 			const FString NewNameStr = NewText.ToString();
 			FStructureEditorUtils::RenameVariable(StructureDetailsSP->GetUserDefinedStruct(), FieldGuid, NewNameStr);
-			OnChanged();
 		}
 	}
 
@@ -310,7 +354,6 @@ public:
 		if(StructureDetailsSP.IsValid())
 		{
 			FStructureEditorUtils::ChangeVariableType(StructureDetailsSP->GetUserDefinedStruct(), FieldGuid, PinType);
-			OnChanged();
 		}
 	}
 
@@ -325,11 +368,6 @@ public:
 		if(StructureDetailsSP.IsValid())
 		{
 			FStructureEditorUtils::RemoveVariable(StructureDetailsSP->GetUserDefinedStruct(), FieldGuid);
-			auto StructureLayoutSP = StructureLayout.Pin();
-			if(StructureLayoutSP.IsValid())
-			{
-				StructureLayoutSP->OnChanged();
-			}
 		}
 	}
 
@@ -378,7 +416,6 @@ public:
 		if (StructureDetailsSP.IsValid())
 		{
 			FStructureEditorUtils::ChangeVariableTooltip(StructureDetailsSP->GetUserDefinedStruct(), FieldGuid, NewText.ToString());
-			OnChanged();
 		}
 	}
 
@@ -401,7 +438,6 @@ public:
 		if (StructureDetailsSP.IsValid())
 		{
 			FStructureEditorUtils::ChangeEditableOnBPInstance(StructureDetailsSP->GetUserDefinedStruct(), FieldGuid, ESlateCheckBoxState::Unchecked != InNewState);
-			OnChanged();
 		}
 	}
 
@@ -432,7 +468,6 @@ public:
 		if (StructureDetailsSP.IsValid() && (ESlateCheckBoxState::Undetermined != InNewState))
 		{
 			FStructureEditorUtils::Change3dWidgetEnabled(StructureDetailsSP->GetUserDefinedStruct(), FieldGuid, ESlateCheckBoxState::Checked == InNewState);
-			OnChanged();
 		}
 	}
 
@@ -442,7 +477,6 @@ public:
 		if(StructureDetailsSP.IsValid())
 		{
 			FStructureEditorUtils::ChangeVariableDefaultValue(StructureDetailsSP->GetUserDefinedStruct(), FieldGuid, NewText.ToString());
-			OnChanged();
 		}
 	}
 
@@ -725,6 +759,33 @@ void FUserDefinedStructureLayout::GenerateChildContent( IDetailChildrenBuilder&
 				TSharedRef<class FUserDefinedStructureFieldLayout> VarLayout = MakeShareable(new FUserDefinedStructureFieldLayout(StructureDetails,  SharedThis(this), VarDesc.VarGuid));
 				ChildrenBuilder.AddChildCustomBuilder(VarLayout);
 			}
+
+			struct FShowDefaultValuePropertyEditor
+			{
+				bool bShowDefaultValuePropertyEditor;
+				FShowDefaultValuePropertyEditor() : bShowDefaultValuePropertyEditor(false)
+				{
+					GConfig->GetBool(TEXT("UserDefinedStructure"), TEXT("bShowDefaultValuePropertyEditor"), bShowDefaultValuePropertyEditor, GEditorIni);
+				}
+			};
+			static FShowDefaultValuePropertyEditor Helper;
+			if (Helper.bShowDefaultValuePropertyEditor)
+			{
+				FPropertyEditorModule& PropertyModule = FModuleManager::LoadModuleChecked<FPropertyEditorModule>("PropertyEditor");
+				FDetailsViewArgs ViewArgs;
+				ViewArgs.bAllowSearch = false;
+				ViewArgs.bHideSelectionTip = false;
+				ViewArgs.bShowActorLabel = false;
+
+				auto StructureDetailViewRef = PropertyModule.CreateStructureDetailView(ViewArgs, StructureDetailsSP->GetStructInstance());
+				StructureDetailViewRef->GetOnFinishedChangingPropertiesDelegate().AddSP(this, &FUserDefinedStructureLayout::OnFinishedChangingProperties);
+				ChildrenBuilder.AddChildContent(FString())
+				[
+					StructureDetailViewRef->GetWidget().ToSharedRef()
+				];
+
+				StructureDetailsView = StructureDetailViewRef;
+			}
 		}
 	}
 }
@@ -732,11 +793,6 @@ void FUserDefinedStructureLayout::GenerateChildContent( IDetailChildrenBuilder&
 ///////////////////////////////////////////////////////////////////////////////////////
 // FUserDefinedStructureDetails
 
-UUserDefinedStruct* FUserDefinedStructureDetails::GetUserDefinedStruct()
-{
-	return UserDefinedStruct.Get();
-}
-
 FStructVariableDescription* FUserDefinedStructureDetails::FindStructureFieldByGuid(FGuid Guid)
 {
 	if (auto Struct = GetUserDefinedStruct())
@@ -754,8 +810,9 @@ void FUserDefinedStructureDetails::CustomizeDetails( class IDetailLayoutBuilder&
 	if (Objects.Num() == 1)
 	{
 		UserDefinedStruct = CastChecked<UUserDefinedStruct>(Objects[0].Get());
-		IDetailCategoryBuilder& StructureCategory = DetailLayout.EditCategory("Structure", LOCTEXT("Structure", "Structure").ToString());
+		StructData = MakeShareable(new FStructOnScope(UserDefinedStruct.Get()));
 
+		IDetailCategoryBuilder& StructureCategory = DetailLayout.EditCategory("Structure", LOCTEXT("Structure", "Structure").ToString());
 		Layout = MakeShareable(new FUserDefinedStructureLayout(SharedThis(this)));
 		StructureCategory.AddCustomBuilder(Layout.ToSharedRef());
 	}
@@ -763,10 +820,25 @@ void FUserDefinedStructureDetails::CustomizeDetails( class IDetailLayoutBuilder&
 	FStructureEditorUtils::FStructEditorManager::Get().AddListener(this);
 }
 
-void FUserDefinedStructureDetails::OnChanged(const class UUserDefinedStruct* Struct)
+void FUserDefinedStructureDetails::PreChange(const class UUserDefinedStruct* Struct)
+{
+	if (Struct && (GetUserDefinedStruct() == Struct) && StructData.IsValid())
+	{
+		if (Layout.IsValid())
+		{
+			Layout->ResetStructureView();
+		}
+		StructData->Destroy();
+		StructData.Reset();
+	}
+}
+
+void FUserDefinedStructureDetails::PostChange(const class UUserDefinedStruct* Struct)
 {
 	if (Struct && (GetUserDefinedStruct() == Struct))
 	{
+		StructData = MakeShareable(new FStructOnScope(Struct));
+		FStructureEditorUtils::Fill_MakeStructureDefaultValue(Struct, StructData->GetStructMemory());
 		if (Layout.IsValid())
 		{
 			Layout->OnChanged();
@@ -774,6 +846,8 @@ void FUserDefinedStructureDetails::OnChanged(const class UUserDefinedStruct* Str
 	}
 }
 
+
+
 FUserDefinedStructureDetails::~FUserDefinedStructureDetails()
 {
 	FStructureEditorUtils::FStructEditorManager::Get().RemoveListener(this);
diff --git a/Engine/Source/Editor/Kismet/Private/UserDefinedStructureEditor.h b/Engine/Source/Editor/Kismet/Private/UserDefinedStructureEditor.h
--- a/Engine/Source/Editor/Kismet/Private/UserDefinedStructureEditor.h
+++ b/Engine/Source/Editor/Kismet/Private/UserDefinedStructureEditor.h
@@ -58,14 +58,24 @@ public:
 	/** IDetailCustomization interface */
 	virtual void CustomizeDetails( class IDetailLayoutBuilder& DetailLayout ) override;
 
-	UUserDefinedStruct* GetUserDefinedStruct();
+	UUserDefinedStruct* GetUserDefinedStruct()
+	{
+		return UserDefinedStruct.Get();
+	}
+
+	TSharedPtr<FStructOnScope> GetStructInstance()
+	{
+		return StructData;
+	}
 
 	struct FStructVariableDescription* FindStructureFieldByGuid(FGuid Guid);
 
 	/** FStructureEditorUtils::INotifyOnStructChanged */
-	virtual void OnChanged(const class UUserDefinedStruct* Struct) override;
+	virtual void PreChange(const class UUserDefinedStruct* Struct) override;
+	virtual void PostChange(const class UUserDefinedStruct* Struct) override;
 
 private:
 	TWeakObjectPtr<UUserDefinedStruct> UserDefinedStruct;
 	TSharedPtr<class FUserDefinedStructureLayout> Layout;
+	TSharedPtr<FStructOnScope> StructData;
 };
diff --git a/Engine/Source/Editor/PropertyEditor/Private/PropertyEditorModule.cpp b/Engine/Source/Editor/PropertyEditor/Private/PropertyEditorModule.cpp
--- a/Engine/Source/Editor/PropertyEditor/Private/PropertyEditorModule.cpp
+++ b/Engine/Source/Editor/PropertyEditor/Private/PropertyEditorModule.cpp
@@ -21,6 +21,7 @@
 #include "TextPropertyTableCellPresenter.h"
 
 #include "PropertyTableConstants.h"
+#include "SStructureDetailsView.h"
 
 IMPLEMENT_MODULE( FPropertyEditorModule, PropertyEditor );
 
@@ -630,4 +631,25 @@ FPropertyTypeLayoutCallback FPropertyEditorModule::GetPropertyTypeCustomization(
 	}
 
 	return FPropertyTypeLayoutCallback();
+}
+
+TSharedRef<class IStructureDetailsView> FPropertyEditorModule::CreateStructureDetailView(const FDetailsViewArgs& DetailsViewArgs, TSharedPtr<FStructOnScope> StructData)
+{
+	TSharedRef<SStructureDetailsView> DetailView =
+		SNew(SStructureDetailsView)
+		.DetailsViewArgs(DetailsViewArgs);
+
+	struct FDontShowObjects
+	{
+		static bool CanShow(const UProperty* Property)
+		{
+			return Property && !Property->IsA<UObjectPropertyBase>();
+		}
+	};
+
+	DetailView->SetIsPropertyVisibleDelegate(FIsPropertyVisible::CreateStatic(&FDontShowObjects::CanShow));
+
+	DetailView->SetStructureData(StructData);
+
+	return DetailView;
 }
\ No newline at end of file
diff --git a/Engine/Source/Editor/PropertyEditor/Private/SStructureDetailsView.cpp b/Engine/Source/Editor/PropertyEditor/Private/SStructureDetailsView.cpp
new file mode 100644
--- /dev/null
+++ b/Engine/Source/Editor/PropertyEditor/Private/SStructureDetailsView.cpp
@@ -0,0 +1,236 @@
+// Copyright 1998-2014 Epic Games, Inc. All Rights Reserved.
+
+#include "PropertyEditorPrivatePCH.h"
+#include "PropertyEditorHelpers.h"
+#include "ObjectPropertyNode.h"
+#include "PropertyEditor.h"
+#include "PropertyDetailsUtilities.h"
+#include "SPropertyEditor.h"
+#include "ScopedTransaction.h"
+#include "SResetToDefaultPropertyEditor.h"
+#include "SStructureDetailsView.h"
+#include "StructurePropertyNode.h"
+
+#define LOCTEXT_NAMESPACE "SStructureDetailsView"
+
+SStructureDetailsView::~SStructureDetailsView()
+{
+	SaveExpandedItems();
+}
+
+UScriptStruct* SStructureDetailsView::GetBaseScriptStruct() const
+{
+	const UScriptStruct* Struct = StructData.IsValid() ? StructData->GetStruct() : NULL;
+	return const_cast<UScriptStruct*>(Struct);
+}
+
+void SStructureDetailsView::Construct(const FArguments& InArgs)
+{
+	DetailsViewArgs = InArgs._DetailsViewArgs;
+	bHasActiveFilter = false;
+	bIsLocked = false;
+	bHasOpenColorPicker = false;
+
+	// Create the root property now
+	RootNode = MakeShareable(new FStructurePropertyNode);
+		
+	PropertyUtilities = MakeShareable( new FPropertyDetailsUtilities( *this ) );
+	
+	ColumnWidth = 0.65f;
+	ColumnSizeData.LeftColumnWidth = TAttribute<float>(this, &SStructureDetailsView::OnGetLeftColumnWidth);
+	ColumnSizeData.RightColumnWidth = TAttribute<float>(this, &SStructureDetailsView::OnGetRightColumnWidth);
+	ColumnSizeData.OnWidthChanged = SSplitter::FOnSlotResized::CreateSP(this, &SStructureDetailsView::OnSetColumnWidth);
+
+	TSharedRef<SScrollBar> ExternalScrollbar = 
+		SNew(SScrollBar)
+		.AlwaysShowScrollbar( true );
+
+		FMenuBuilder DetailViewOptions( true, NULL );
+
+		FUIAction ShowOnlyModifiedAction( 
+			FExecuteAction::CreateSP(this, &SStructureDetailsView::OnShowOnlyModifiedClicked),
+			FCanExecuteAction(),
+			FIsActionChecked::CreateSP(this, &SStructureDetailsView::IsShowOnlyModifiedChecked)
+		);
+
+		if (DetailsViewArgs.bShowModifiedPropertiesOption)
+		{
+			DetailViewOptions.AddMenuEntry( 
+				LOCTEXT("ShowOnlyModified", "Show Only Modified Properties"),
+				LOCTEXT("ShowOnlyModified_ToolTip", "Displays only properties which have been changed from their default"),
+				FSlateIcon(),
+				ShowOnlyModifiedAction,
+				NAME_None,
+				EUserInterfaceActionType::ToggleButton 
+			);
+		}
+
+		FUIAction ShowAllAdvancedAction( 
+			FExecuteAction::CreateSP(this, &SStructureDetailsView::OnShowAllAdvancedClicked),
+			FCanExecuteAction(),
+			FIsActionChecked::CreateSP(this, &SStructureDetailsView::IsShowAllAdvancedChecked)
+		);
+
+		DetailViewOptions.AddMenuEntry(
+			LOCTEXT("ShowAllAdvanced", "Show All Advanced Details"),
+			LOCTEXT("ShowAllAdvanced_ToolTip", "Shows all advanced detail sections in each category"),
+			FSlateIcon(),
+			ShowAllAdvancedAction,
+			NAME_None,
+			EUserInterfaceActionType::ToggleButton 
+			);
+
+		DetailViewOptions.AddMenuEntry(
+			LOCTEXT("CollapseAll", "Collapse All Categories"),
+			LOCTEXT("CollapseAll_ToolTip", "Collapses all root level categories"),
+			FSlateIcon(),
+			FUIAction(FExecuteAction::CreateSP(this, &SStructureDetailsView::SetRootExpansionStates, /*bExpanded=*/false, /*bRecurse=*/false)));
+		DetailViewOptions.AddMenuEntry(
+			LOCTEXT("ExpandAll", "Expand All Categories"),
+			LOCTEXT("ExpandAll_ToolTip", "Expands all root level categories"),
+			FSlateIcon(),
+			FUIAction(FExecuteAction::CreateSP(this, &SStructureDetailsView::SetRootExpansionStates, /*bExpanded=*/true, /*bRecurse=*/false)));
+
+	TSharedRef<SHorizontalBox> FilterRow = SNew( SHorizontalBox )
+		.Visibility(this, &SStructureDetailsView::GetFilterBoxVisibility)
+		+SHorizontalBox::Slot()
+		.FillWidth( 1 )
+		.VAlign( VAlign_Center )
+		[
+			// Create the search box
+			SAssignNew( SearchBox, SSearchBox )
+			.OnTextChanged(this, &SStructureDetailsView::OnFilterTextChanged)
+		];
+
+	if (DetailsViewArgs.bShowOptions)
+	{
+		FilterRow->AddSlot()
+			.HAlign(HAlign_Right)
+			.AutoWidth()
+			[
+				SNew( SComboButton )
+				.ContentPadding(0)
+				.ForegroundColor( FSlateColor::UseForeground() )
+				.ButtonStyle( FEditorStyle::Get(), "ToggleButton" )
+				.MenuContent()
+				[
+					DetailViewOptions.MakeWidget()
+				]
+				.ButtonContent()
+				[
+					SNew(SImage)
+					.Image( FEditorStyle::GetBrush("GenericViewButton") )
+				]
+			];
+	}
+
+	SAssignNew(DetailTree, SDetailTree)
+		.Visibility(this, &SStructureDetailsView::GetTreeVisibility)
+		.TreeItemsSource(&RootTreeNodes)
+		.OnGetChildren(this, &SStructureDetailsView::OnGetChildrenForDetailTree)
+		.OnSetExpansionRecursive(this, &SStructureDetailsView::SetNodeExpansionStateRecursive)
+		.OnGenerateRow(this, &SStructureDetailsView::OnGenerateRowForDetailTree)
+		.OnExpansionChanged(this, &SStructureDetailsView::OnItemExpansionChanged)
+		.SelectionMode(ESelectionMode::None)
+		.ExternalScrollbar(ExternalScrollbar);
+
+	ChildSlot
+	[
+		SNew( SVerticalBox )
+		+ SVerticalBox::Slot()
+		.AutoHeight()
+		.Padding( 0.0f, 0.0f, 0.0f, 2.0f )
+		[
+			FilterRow
+		]
+		+ SVerticalBox::Slot()
+		.FillHeight(1)
+		.Padding(0)
+		[
+			SNew( SHorizontalBox )
+			+ SHorizontalBox::Slot()
+			[
+				DetailTree.ToSharedRef()
+			]
+			+ SHorizontalBox::Slot()
+			.AutoWidth()
+			[
+				SNew( SBox )
+				.WidthOverride( 16.0f )
+				[
+					ExternalScrollbar
+				]
+			]
+		]
+	];
+}
+
+void SStructureDetailsView::SetStructureData(TSharedPtr<FStructOnScope> InStructData)
+{
+	//PRE SET
+	SaveExpandedItems();
+	RootNodePendingKill = RootNode;
+	RootNode = MakeShareable(new FStructurePropertyNode);
+
+	//SET
+	StructData = InStructData;
+	RootNode->SetStructure(StructData);
+	if (!StructData.IsValid() || !StructData->IsValid())
+	{
+		bIsLocked = false;
+	}
+	
+	//POST SET
+	DestroyColorPicker();
+	ColorPropertyNode = NULL;
+
+	FPropertyNodeInitParams InitParams;
+	InitParams.ParentNode = NULL;
+	InitParams.Property = NULL;
+	InitParams.ArrayOffset = 0;
+	InitParams.ArrayIndex = INDEX_NONE;
+	InitParams.bAllowChildren = true;
+	InitParams.bForceHiddenPropertyVisibility = FPropertySettings::Get().ShowHiddenProperties();
+	InitParams.bCreateCategoryNodes = false;
+
+	RootNode->InitNode(InitParams);
+
+	RestoreExpandedItems();
+
+	UpdatePropertyMap();
+}
+
+void SStructureDetailsView::ForceRefresh()
+{
+	SetStructureData(StructData);
+}
+
+const TArray< TWeakObjectPtr<UObject> >& SStructureDetailsView::GetSelectedObjects() const
+{
+	static const TArray< TWeakObjectPtr<UObject> > DummyRef;
+	return DummyRef;
+}
+
+const TArray< TWeakObjectPtr<AActor> >& SStructureDetailsView::GetSelectedActors() const
+{
+	static const TArray< TWeakObjectPtr<AActor> > DummyRef;
+	return DummyRef;
+}
+
+const FSelectedActorInfo& SStructureDetailsView::GetSelectedActorInfo() const
+{
+	static const FSelectedActorInfo DummyRef;
+	return DummyRef;
+}
+
+bool SStructureDetailsView::IsConnected() const
+{
+	return StructData.IsValid() && StructData->IsValid() && RootNode.IsValid() && RootNode->IsValid();
+}
+
+TSharedPtr<class FComplexPropertyNode> SStructureDetailsView::GetRootNode()
+{
+	return RootNode;
+}
+
+#undef LOCTEXT_NAMESPACE
\ No newline at end of file
diff --git a/Engine/Source/Editor/PropertyEditor/Private/SStructureDetailsView.h b/Engine/Source/Editor/PropertyEditor/Private/SStructureDetailsView.h
new file mode 100644
--- /dev/null
+++ b/Engine/Source/Editor/PropertyEditor/Private/SStructureDetailsView.h
@@ -0,0 +1,84 @@
+// Copyright 1998-2014 Epic Games, Inc. All Rights Reserved.
+#pragma once
+
+#include "SDetailsViewBase.h"
+#include "IStructureDetailsView.h"
+
+class SStructureDetailsView : public SDetailsViewBase, public IStructureDetailsView
+{
+	TSharedPtr<struct FStructOnScope> StructData;
+	TSharedPtr<class FStructurePropertyNode> RootNode;
+
+public:
+
+	SLATE_BEGIN_ARGS(SStructureDetailsView)
+		: _DetailsViewArgs()
+		{}
+		/** The user defined args for the details view */
+		SLATE_ARGUMENT(FDetailsViewArgs, DetailsViewArgs)
+	SLATE_END_ARGS()
+
+	~SStructureDetailsView();
+
+	/**
+	 * Constructs the property view widgets                   
+	 */
+	void Construct(const FArguments& InArgs);
+
+	UScriptStruct* GetBaseScriptStruct() const;
+
+	virtual bool IsConnected() const override;
+
+	/** IStructureDetailsView interface */
+	virtual TSharedPtr<SWidget> GetWidget() override
+	{
+		return AsShared();
+	}
+
+	virtual void SetStructureData(TSharedPtr<FStructOnScope> InStructData) override;
+	virtual FOnFinishedChangingProperties& GetOnFinishedChangingPropertiesDelegate() override
+	{
+		return OnFinishedChangingProperties();
+	}
+	/** IDetailsViewPrivate interface */
+	virtual const UClass* GetBaseClass() const override
+	{
+		return NULL;
+	}
+
+	virtual UClass* GetBaseClass() override
+	{
+		return NULL;
+	}
+
+	virtual bool IsCategoryHiddenByClass(FName CategoryName) const override
+	{
+		return false;
+	}
+
+	virtual void ForceRefresh() override;
+	virtual void AddExternalRootPropertyNode(TSharedRef<FObjectPropertyNode> ExternalRootNode) override {}
+
+	/** IDetailsView interface */
+	virtual UStruct* GetBaseStruct() const override
+	{
+		return GetBaseScriptStruct();
+	}
+	virtual const TArray< TWeakObjectPtr<UObject> >& GetSelectedObjects() const override;
+	virtual const TArray< TWeakObjectPtr<AActor> >& GetSelectedActors() const override;
+	virtual const struct FSelectedActorInfo& GetSelectedActorInfo() const override;
+	virtual bool HasClassDefaultObject() const override
+	{
+		return false;
+	}
+
+	virtual void SetOnObjectArrayChanged(FOnObjectArrayChanged OnObjectArrayChangedDelegate) override{}
+	virtual void RegisterInstancedCustomPropertyLayout(UClass* Class, FOnGetDetailCustomizationInstance DetailLayoutDelegate) override {}
+	virtual void UnregisterInstancedCustomPropertyLayout(UClass* Class) override {}
+	virtual void SetObjects(const TArray<UObject*>& InObjects, bool bForceRefresh = false) override{}
+	virtual void SetObjects(const TArray< TWeakObjectPtr< UObject > >& InObjects, bool bForceRefresh = false) override{}
+	virtual void SetObject(UObject* InObject, bool bForceRefresh = false) override{}
+
+	virtual TSharedPtr<class FComplexPropertyNode> GetRootNode() override;
+};
+
diff --git a/Engine/Source/Editor/PropertyEditor/Private/StructurePropertyNode.cpp b/Engine/Source/Editor/PropertyEditor/Private/StructurePropertyNode.cpp
new file mode 100644
--- /dev/null
+++ b/Engine/Source/Editor/PropertyEditor/Private/StructurePropertyNode.cpp
@@ -0,0 +1,36 @@
+// Copyright 1998-2014 Epic Games, Inc. All Rights Reserved.
+
+
+#include "PropertyEditorPrivatePCH.h"
+#include "PropertyNode.h"
+#include "CategoryPropertyNode.h"
+#include "ItemPropertyNode.h"
+#include "ObjectEditorUtils.h"
+#include "StructurePropertyNode.h"
+
+void FStructurePropertyNode::InitChildNodes()
+{
+	const bool bShouldShowHiddenProperties = !!HasNodeFlags(EPropertyNodeFlags::ShouldShowHiddenProperties);
+	const UScriptStruct* Struct = StructData.IsValid() ? StructData->GetStruct() : NULL;
+
+	for (TFieldIterator<UProperty> It(Struct); It; ++It)
+	{
+		UProperty* StructMember = *It;
+		if (StructMember && (bShouldShowHiddenProperties || (StructMember->PropertyFlags & CPF_Edit)))
+		{
+			TSharedPtr<FItemPropertyNode> NewItemNode(new FItemPropertyNode);//;//CreatePropertyItem(StructMember,INDEX_NONE,this);
+
+			FPropertyNodeInitParams InitParams;
+			InitParams.ParentNode = SharedThis(this);
+			InitParams.Property = StructMember;
+			InitParams.ArrayOffset = 0;
+			InitParams.ArrayIndex = INDEX_NONE;
+			InitParams.bAllowChildren = true;
+			InitParams.bForceHiddenPropertyVisibility = !!HasNodeFlags(EPropertyNodeFlags::ShouldShowHiddenProperties);
+			InitParams.bCreateCategoryNodes = false;
+
+			NewItemNode->InitNode(InitParams);
+			AddChildNode(NewItemNode);
+		}
+	}
+}
\ No newline at end of file
diff --git a/Engine/Source/Editor/PropertyEditor/Private/StructurePropertyNode.h b/Engine/Source/Editor/PropertyEditor/Private/StructurePropertyNode.h
new file mode 100644
--- /dev/null
+++ b/Engine/Source/Editor/PropertyEditor/Private/StructurePropertyNode.h
@@ -0,0 +1,109 @@
+// Copyright 1998-2014 Epic Games, Inc. All Rights Reserved.
+
+#pragma once
+
+//-----------------------------------------------------------------------------
+//	FStructPropertyNode - Used for the root and various sub-nodes
+//-----------------------------------------------------------------------------
+
+class FStructurePropertyNode : public FComplexPropertyNode
+{
+public:
+	FStructurePropertyNode() : FComplexPropertyNode() {}
+	virtual ~FStructurePropertyNode() {}
+
+	/** FPropertyNode Interface */
+	virtual uint8* GetValueBaseAddress(uint8* Base) override
+	{
+		ensure(true);
+		return IsValid() ? StructData->GetStructMemory() : NULL;
+	}
+
+	void SetStructure(TSharedPtr<FStructOnScope> InStructData)
+	{
+		DestroyTree();
+		StructData = InStructData;
+	}
+
+	bool IsValid() const
+	{
+		return StructData.IsValid() && StructData->IsValid();
+	}
+
+	bool GetReadAddressUncached(FPropertyNode& InPropertyNode, FReadAddressListData& OutAddresses) const override
+	{
+		if (!IsValid())
+		{
+			return false;
+		}
+
+		UProperty* InItemProperty = InPropertyNode.GetProperty();
+		if (!InItemProperty)
+		{
+			return false;
+		}
+
+		uint8* ReadAddress = StructData->GetStructMemory();
+		check(ReadAddress);
+		OutAddresses.Add(NULL, InPropertyNode.GetValueBaseAddress(ReadAddress), true);
+		return true;
+	}
+
+	bool GetReadAddressUncached(FPropertyNode& InPropertyNode,
+		bool InRequiresSingleSelection,
+		FReadAddressListData& OutAddresses,
+		bool bComparePropertyContents,
+		bool bObjectForceCompare,
+		bool bArrayPropertiesCanDifferInSize) const override
+	{
+		return GetReadAddressUncached(InPropertyNode, OutAddresses);
+	}
+
+	/** FComplexPropertyNode Interface */
+	virtual const UStruct* GetBaseStructure() const override
+	{ 
+		return IsValid() ? StructData->GetStruct() : NULL; 
+	}
+	virtual UStruct* GetBaseStructure() override
+	{
+		const UStruct* Struct = IsValid() ? StructData->GetStruct() : NULL;
+		return const_cast<UStruct*>(Struct);
+	}
+	virtual int32 GetInstancesNum() const override
+	{ 
+		return IsValid() ? 1 : 0; 
+	}
+	virtual uint8* GetMemoryOfInstance(int32 Index) override
+	{ 
+		check(0 == Index);
+		return IsValid() ? StructData->GetStructMemory() : NULL;
+	}
+	virtual TWeakObjectPtr<UObject> GetInstanceAsUObject(int32 Index) override
+	{
+		check(0 == Index);
+		return NULL;
+	}
+	virtual EPropertyType GetPropertyType() const override
+	{
+		return EPT_StandaloneStructure;
+	}
+
+	virtual void Disconnect() override
+	{
+		SetStructure(NULL);
+	}
+
+protected:
+
+	/** FPropertyNode interface */
+	virtual void InitChildNodes() override;
+
+	virtual void GetQualifiedName(FString& PathPlusIndex, const bool bWithArrayIndex) const override
+	{
+		PathPlusIndex += TEXT("Struct");
+	}
+
+private:
+	TSharedPtr<FStructOnScope> StructData;
+};
+
diff --git a/Engine/Source/Editor/PropertyEditor/Public/IStructureDetailsView.h b/Engine/Source/Editor/PropertyEditor/Public/IStructureDetailsView.h
new file mode 100644
--- /dev/null
+++ b/Engine/Source/Editor/PropertyEditor/Public/IStructureDetailsView.h
@@ -0,0 +1,17 @@
+// Copyright 1998-2014 Epic Games, Inc. All Rights Reserved.
+
+#pragma once
+
+#include "PropertyEditorDelegates.h"
+
+/**
+ * Interface class for all detail views
+ */
+class IStructureDetailsView
+{
+public:
+	virtual TSharedPtr<class SWidget> GetWidget() = 0;
+
+	virtual void SetStructureData(TSharedPtr<struct FStructOnScope> StructData) = 0;
+	virtual FOnFinishedChangingProperties& GetOnFinishedChangingPropertiesDelegate() = 0;
+};
diff --git a/Engine/Source/Editor/PropertyEditor/Public/PropertyEditorModule.h b/Engine/Source/Editor/PropertyEditor/Public/PropertyEditorModule.h
--- a/Engine/Source/Editor/PropertyEditor/Public/PropertyEditorModule.h
+++ b/Engine/Source/Editor/PropertyEditor/Public/PropertyEditorModule.h
@@ -52,7 +52,7 @@ class SWindow;
 class IPropertyTableCellPresenter;
 class IPropertyTypeCustomization;
 
-/** 
+/**
  * Base class for adding an extra data to identify a custom property type
  */
 class IPropertyTypeIdentifier
@@ -60,9 +60,9 @@ class IPropertyTypeIdentifier
 public:
 	virtual ~IPropertyTypeIdentifier() {}
 
-	/**
+/**
 	 * Called to identify if a property type is customized
-	 * 
+	 *
 	 * @param IPropertyHandle	Handle to the property being tested
 	 */
 	virtual bool IsPropertyTypeCustomized( const IPropertyHandle& PropertyHandle ) const = 0;
@@ -118,6 +118,8 @@ class FPropertyEditorModule : public IModuleInterface
 {
 	friend class SPropertyTreeView;
 	friend class SDetailsView;
+	friend class SStructureDetailsView;
+	friend class SDetailsViewBase;
 public:
 	
 	/**
@@ -253,6 +255,8 @@ public:
 	 */
 	virtual TSharedPtr<class ISinglePropertyView> CreateSingleProperty( UObject* InObject, FName InPropertyName, const struct FSinglePropertyParams& InitParams );
 
+	virtual TSharedRef<class IStructureDetailsView> CreateStructureDetailView(const struct FDetailsViewArgs& DetailsViewArgs, TSharedPtr<struct FStructOnScope> StructData);
+
 	/**
 	 * Creates a property change listener that notifies users via a  delegate when a property on an object changes
 	 *
