Fixed the context menu incorrectly appearing when right clicking in the local variable list view when an item part of the general list is selected and visa versa.

#ttp 333366 - BP: The variable right click context menu appears instead of the "Add New" menu when right clicking the empty space under Local Variables

diff --git a/Engine/Source/Editor/Kismet/Private/SMyBlueprint.cpp b/Engine/Source/Editor/Kismet/Private/SMyBlueprint.cpp
--- a/Engine/Source/Editor/Kismet/Private/SMyBlueprint.cpp
+++ b/Engine/Source/Editor/Kismet/Private/SMyBlueprint.cpp
@@ -1188,7 +1188,19 @@ bool SMyBlueprint::SelectionIsCategory() const
 
 bool SMyBlueprint::SelectionHasContextMenu() const
 {
-	return SelectionAsGraph() || SelectionAsVar() || SelectionIsCategory() || SelectionAsDelegate() || SelectionAsEnum() || SelectionAsEvent() || SelectionAsLocalVar() || SelectionAsStruct();
+	bool bSelectionHasContextMenu = false;
+
+	// Do not consider the selection having a context menu if the menu is being brought up in the other action menu, that being if an item is selected in the GraphActionMenu and you right click in LocalGraphActionMenu
+	if(GraphActionMenu.IsValid() && GraphActionMenu->HasFocusedDescendants())
+	{
+		bSelectionHasContextMenu = SelectionAsGraph() || SelectionAsVar() || SelectionIsCategory() || SelectionAsDelegate() || SelectionAsEnum() || SelectionAsEvent() || SelectionAsStruct();
+	}
+	else if(LocalGraphActionMenu.IsValid() && LocalGraphActionMenu->HasFocusedDescendants())
+	{
+		bSelectionHasContextMenu = SelectionAsLocalVar() != NULL;
+	}
+
+	return bSelectionHasContextMenu;
 }
 
 void SMyBlueprint::GetSelectedItemsForContextMenu(TArray<FComponentEventConstructionData>& OutSelectedItems) const
