A check was added during blueprint editor startup to determine if the current project is blueprint or code based, the result is cached and used to enable/disable goto code definition options.
#TTP 336331 - EDITOR: BP: Functions have the right-click option to Goto Code Definition and will send user to VS with a project that has a lot of errors
#Branch UE4
#Proj Kismet

#Add Added FBlueprintEditor::bCodeBasedProject which is cached during blueprint editor startup based on the FProjectStatus::bCodeBasedBasedProject. This was the only place i found reference to a project being blueprint only or not.
#Add Added FBlueprintEditor::IsNativeCodeBrowsingAvailable which returns the current availability of GotoCode functionality based on project type.

diff --git a/Engine/Source/Editor/Kismet/Kismet.Build.cs b/Engine/Source/Editor/Kismet/Kismet.Build.cs
--- a/Engine/Source/Editor/Kismet/Kismet.Build.cs
+++ b/Engine/Source/Editor/Kismet/Kismet.Build.cs
@@ -35,7 +35,8 @@ public class Kismet : ModuleRules
 				"SourceControl",
 				"LevelEditor",
                 "InputCore",
-				"EngineSettings"
+				"EngineSettings",
+                "Projects"
 			}
 			);
 
diff --git a/Engine/Source/Editor/Kismet/Private/BlueprintEditor.cpp b/Engine/Source/Editor/Kismet/Private/BlueprintEditor.cpp
--- a/Engine/Source/Editor/Kismet/Private/BlueprintEditor.cpp
+++ b/Engine/Source/Editor/Kismet/Private/BlueprintEditor.cpp
@@ -682,12 +682,16 @@ TSharedRef<SGraphEditor> FBlueprintEditor::CreateGraphEditorWidget(TSharedRef<FT
 
 			GraphEditorCommands->MapAction( FGraphEditorCommands::Get().GotoNativeFunctionDefinition,
 				FExecuteAction::CreateSP( this, &FBlueprintEditor::GotoNativeFunctionDefinition ),
-				FCanExecuteAction::CreateSP(this, &FBlueprintEditor::IsSelectionNativeFunction)
+				FCanExecuteAction::CreateSP(this, &FBlueprintEditor::IsSelectionNativeFunction),
+				FIsActionChecked(),
+				FIsActionButtonVisible::CreateSP( this, &FBlueprintEditor::IsNativeCodeBrowsingAvailable )
 				);
 
 			GraphEditorCommands->MapAction( FGraphEditorCommands::Get().GotoNativeVariableDefinition,
 				FExecuteAction::CreateSP( this, &FBlueprintEditor::GotoNativeVariableDefinition ),
-				FCanExecuteAction::CreateSP(this, &FBlueprintEditor::IsSelectionNativeVariable)
+				FCanExecuteAction::CreateSP(this, &FBlueprintEditor::IsSelectionNativeVariable),
+				FIsActionChecked(),
+				FIsActionButtonVisible::CreateSP( this, &FBlueprintEditor::IsNativeCodeBrowsingAvailable )
 				);
 		}
 	}
@@ -799,6 +803,7 @@ FBlueprintEditor::FBlueprintEditor()
 	, bIsActionMenuContextSensitive(true)
 	, CurrentUISelection(FBlueprintEditor::NoSelection)
 	, bEditorMarkedAsClosed(false)
+	, bCodeBasedProject(false)
 {
 	AnalyticsStats.GraphActionMenusNonCtxtSensitiveExecCount = 0;
 	AnalyticsStats.GraphActionMenusCtxtSensitiveExecCount = 0;
@@ -1174,6 +1179,16 @@ void FBlueprintEditor::InitBlueprintEditor(const EToolkitMode::Type Mode, const
 		SetCurrentMode(FBlueprintEditorApplicationModes::BlueprintDefaultsMode);
 	}
 
+	// Cache the project type ( Blueprint or Code Based )
+	if( FPaths::IsProjectFilePathSet() )
+	{
+		FProjectStatus ProjectStatus;
+		if( IProjectManager::Get().QueryStatusForProject( FPaths::GetProjectFilePath(), ProjectStatus ) && ProjectStatus.bCodeBasedProject )
+		{
+			bCodeBasedProject = true;
+		}
+	}
+
 	// Post-layout initialization
 	PostLayoutBlueprintEditorInitialization();
 
@@ -6137,6 +6152,11 @@ bool FBlueprintEditor::IsSelectionNativeVariable()
 	return false;
 }
 
+bool FBlueprintEditor::IsNativeCodeBrowsingAvailable() const
+{
+	return bCodeBasedProject;
+}
+
 void FBlueprintEditor::OnFindInstancesCustomEvent()
 {
 	auto GraphEditor = FocusedGraphEdPtr.Pin();
diff --git a/Engine/Source/Editor/Kismet/Public/BlueprintEditor.h b/Engine/Source/Editor/Kismet/Public/BlueprintEditor.h
--- a/Engine/Source/Editor/Kismet/Public/BlueprintEditor.h
+++ b/Engine/Source/Editor/Kismet/Public/BlueprintEditor.h
@@ -841,6 +841,9 @@ private:
 	/** Returns whether the currently focused graph is editable or not */
 	bool IsFocusedGraphEditable() const;
 
+	/** Called to check if native code browsing is available */
+	bool IsNativeCodeBrowsingAvailable() const;
+
 public://@TODO
 	TSharedPtr<FDocumentTracker> DocumentManager;
 
@@ -928,6 +931,7 @@ protected:
 
 	/** Whether we are already in the process of closing this editor */
 	bool bEditorMarkedAsClosed;
+
 public:
 	// Tries to open the specified graph and bring it's document to the front (note: this can return NULL)
 	TSharedPtr<SGraphEditor> OpenGraphAndBringToFront(UEdGraph* Graph);
@@ -989,6 +993,10 @@ private:
 
 	/** Customizations for the SCS editor */
 	TMap< FName, TSharedPtr<ISCSEditorCustomization> > SCSEditorCustomizations;
+
+	/** Whether the current project is C++ or blueprint based */
+	bool bCodeBasedProject;
+
 };
 
 #undef LOCTEXT_NAMESPACE
\ No newline at end of file
