StructureDetailsView refactored, custom name added.

diff --git a/Engine/Source/Editor/Kismet/Private/SKismetInspector.cpp b/Engine/Source/Editor/Kismet/Private/SKismetInspector.cpp
--- a/Engine/Source/Editor/Kismet/Private/SKismetInspector.cpp
+++ b/Engine/Source/Editor/Kismet/Private/SKismetInspector.cpp
@@ -229,12 +229,6 @@ void SKismetInspector::Construct(const FArguments& InArgs)
 		FOnGetDetailCustomizationInstance LayoutOptionDetails = FOnGetDetailCustomizationInstance::CreateStatic(&FBlueprintGlobalOptionsDetails::MakeInstance, Kismet2Ptr);
 		PropertyView->RegisterInstancedCustomPropertyLayout(UBlueprint::StaticClass(), LayoutOptionDetails);
 
-		FOnGetDetailCustomizationInstance LayoutEnumDetails = FOnGetDetailCustomizationInstance::CreateStatic(&FEnumDetails::MakeInstance);
-		PropertyView->RegisterInstancedCustomPropertyLayout(UUserDefinedEnum::StaticClass(), LayoutEnumDetails);
-
-		FOnGetDetailCustomizationInstance LayoutStructDetails = FOnGetDetailCustomizationInstance::CreateStatic(&FUserDefinedStructureDetails::MakeInstance);
-		PropertyView->RegisterInstancedCustomPropertyLayout(UUserDefinedStruct::StaticClass(), LayoutStructDetails);
-
 		FOnGetDetailCustomizationInstance LayoutFormatTextDetails = FOnGetDetailCustomizationInstance::CreateStatic(&FFormatTextDetails::MakeInstance);
 		PropertyView->RegisterInstancedCustomPropertyLayout(UK2Node_FormatText::StaticClass(), LayoutFormatTextDetails);
 
diff --git a/Engine/Source/Editor/Kismet/Private/UserDefinedStructureEditor.cpp b/Engine/Source/Editor/Kismet/Private/UserDefinedStructureEditor.cpp
--- a/Engine/Source/Editor/Kismet/Private/UserDefinedStructureEditor.cpp
+++ b/Engine/Source/Editor/Kismet/Private/UserDefinedStructureEditor.cpp
@@ -9,9 +9,154 @@
 #include "PropertyCustomizationHelpers.h"
 #include "Editor/KismetWidgets/Public/SPinTypeSelector.h"
 #include "Editor/WorkspaceMenuStructure/Public/WorkspaceMenuStructureModule.h"
+#include "Editor/UnrealEd/Public/Kismet2/StructureEditorUtils.h"
 
 #define LOCTEXT_NAMESPACE "StructureEditor"
 
+class FStructureDefaultValueView : public FStructureEditorUtils::INotifyOnStructChanged, public TSharedFromThis<FStructureDefaultValueView>
+{
+public:
+	FStructureDefaultValueView::FStructureDefaultValueView(UUserDefinedStruct* EditedStruct) 
+		: UserDefinedStruct(EditedStruct)
+	{
+		FStructureEditorUtils::FStructEditorManager::Get().AddListener(this);
+	}
+
+	void Initialize()
+	{
+		StructData = MakeShareable(new FStructOnScope(GetUserDefinedStruct()));
+
+		FPropertyEditorModule& PropertyModule = FModuleManager::LoadModuleChecked<FPropertyEditorModule>("PropertyEditor");
+		FDetailsViewArgs ViewArgs;
+		ViewArgs.bAllowSearch = false;
+		ViewArgs.bHideSelectionTip = false;
+		ViewArgs.bShowActorLabel = false;
+
+		StructureDetailsView = PropertyModule.CreateStructureDetailView(ViewArgs, StructData, false, LOCTEXT("DefaultValues", "Default Values").ToString());
+		StructureDetailsView->GetOnFinishedChangingPropertiesDelegate().AddSP(this, &FStructureDefaultValueView::OnFinishedChangingProperties);
+	}
+
+	~FStructureDefaultValueView()
+	{
+		FStructureEditorUtils::FStructEditorManager::Get().RemoveListener(this);
+	}
+
+	void OnFinishedChangingProperties(const FPropertyChangedEvent& PropertyChangedEvent)
+	{
+		check(PropertyChangedEvent.MemberProperty
+			&& PropertyChangedEvent.MemberProperty->GetOwnerStruct()
+			&& (PropertyChangedEvent.MemberProperty->GetOwnerStruct() == GetUserDefinedStruct())
+			&& PropertyChangedEvent.MemberProperty->GetOwnerStruct()->IsA<UUserDefinedStruct>());
+
+		const UProperty* DirectProperty = PropertyChangedEvent.MemberProperty;
+		while (!Cast<const UUserDefinedStruct>(DirectProperty->GetOuter()))
+		{
+			DirectProperty = CastChecked<const UProperty>(DirectProperty->GetOuter());
+		}
+
+		FString DefaultValueString;
+		bool bDefaultValueSet = false;
+		{
+			if (StructData.IsValid() && StructData->IsValid())
+			{
+				bDefaultValueSet = FBlueprintEditorUtils::PropertyValueToString(DirectProperty, StructData->GetStructMemory(), DefaultValueString);
+			}
+		}
+
+		const FGuid VarGuid = FStructureEditorUtils::GetGuidForProperty(DirectProperty);
+		if (bDefaultValueSet && VarGuid.IsValid())
+		{
+			FStructureEditorUtils::ChangeVariableDefaultValue(GetUserDefinedStruct(), VarGuid, DefaultValueString);
+		}
+	}
+
+	UUserDefinedStruct* GetUserDefinedStruct()
+	{
+		return UserDefinedStruct.Get();
+	}
+
+	TSharedPtr<class SWidget> GetWidget()
+	{
+		return StructureDetailsView.IsValid() ? StructureDetailsView->GetWidget() : NULL;
+	}
+
+	virtual void PreChange(const class UUserDefinedStruct* Struct) override
+	{
+		if (Struct && (GetUserDefinedStruct() == Struct))
+		{
+			if (StructureDetailsView.IsValid())
+			{
+				StructureDetailsView->SetStructureData(NULL);
+			}
+			if (StructData.IsValid())
+			{
+				StructData->Destroy();
+				StructData.Reset();
+			}
+		}
+	}
+
+	virtual void PostChange(const class UUserDefinedStruct* Struct) override
+	{
+		if (Struct && (GetUserDefinedStruct() == Struct))
+		{
+			StructData = MakeShareable(new FStructOnScope(Struct));
+			FStructureEditorUtils::Fill_MakeStructureDefaultValue(Struct, StructData->GetStructMemory());
+
+			if (StructureDetailsView.IsValid())
+			{
+				StructureDetailsView->SetStructureData(StructData);
+			}
+		}
+	}
+private:
+	TSharedPtr<FStructOnScope> StructData;
+	TSharedPtr<class IStructureDetailsView> StructureDetailsView;
+	const TWeakObjectPtr<UUserDefinedStruct> UserDefinedStruct;
+};
+
+///////////////////////////////////////////////////////////////////////////////////////
+// FUserDefinedStructureDetails
+
+class FUserDefinedStructureDetails : public IDetailCustomization, FStructureEditorUtils::INotifyOnStructChanged
+{
+public:
+	/** Makes a new instance of this detail layout class for a specific detail view requesting it */
+	static TSharedRef<IDetailCustomization> MakeInstance()
+	{
+		return MakeShareable(new FUserDefinedStructureDetails);
+	}
+
+	~FUserDefinedStructureDetails()
+	{
+		FStructureEditorUtils::FStructEditorManager::Get().RemoveListener(this);
+	}
+
+	UUserDefinedStruct* GetUserDefinedStruct()
+	{
+		return UserDefinedStruct.Get();
+	}
+
+	struct FStructVariableDescription* FindStructureFieldByGuid(FGuid Guid)
+	{
+		if (auto Struct = GetUserDefinedStruct())
+		{
+			return FStructureEditorUtils::GetVarDesc(Struct).FindByPredicate(FStructureEditorUtils::FFindByGuidHelper<FStructVariableDescription>(Guid));
+		}
+		return NULL;
+	}
+
+	/** IDetailCustomization interface */
+	virtual void CustomizeDetails(class IDetailLayoutBuilder& DetailLayout) override;
+
+	/** FStructureEditorUtils::INotifyOnStructChanged */
+	virtual void PreChange(const class UUserDefinedStruct* Struct) override {}
+	virtual void PostChange(const class UUserDefinedStruct* Struct) override;
+
+private:
+	TWeakObjectPtr<UUserDefinedStruct> UserDefinedStruct;
+	TSharedPtr<class FUserDefinedStructureLayout> Layout;
+};
 
 ///////////////////////////////////////////////////////////////////////////////////////
 // FUserDefinedStructureEditor
@@ -110,26 +255,57 @@ TSharedRef<SDockTab> FUserDefinedStructureEditor::SpawnStructureTab(const FSpawn
 		EditedStruct = Cast<UUserDefinedStruct>(EditingObjects[ 0 ]);
 	}
 
-	// Create a property view
-	FPropertyEditorModule& EditModule = FModuleManager::Get().GetModuleChecked<FPropertyEditorModule>("PropertyEditor");
-
-	FDetailsViewArgs DetailsViewArgs( /*bUpdateFromSelection=*/ false, /*bLockable=*/ false, /*bAllowSearch=*/ false, /*bObjectsUseNameArea=*/ true, /*bHideSelectionTip=*/ true );
-	DetailsViewArgs.bHideActorNameArea = true;
-	DetailsViewArgs.bShowOptions = false;
+	auto Box = SNew(SVerticalBox);
 
-	PropertyView = EditModule.CreateDetailView( DetailsViewArgs );
-
-	FOnGetDetailCustomizationInstance LayoutStructDetails = FOnGetDetailCustomizationInstance::CreateStatic(&FUserDefinedStructureDetails::MakeInstance);
-	PropertyView->RegisterInstancedCustomPropertyLayout(UUserDefinedStruct::StaticClass(), LayoutStructDetails);
+	{
+		// Create a property view
+		FPropertyEditorModule& EditModule = FModuleManager::Get().GetModuleChecked<FPropertyEditorModule>("PropertyEditor");
+		FDetailsViewArgs DetailsViewArgs( /*bUpdateFromSelection=*/ false, /*bLockable=*/ false, /*bAllowSearch=*/ false, /*bObjectsUseNameArea=*/ true, /*bHideSelectionTip=*/ true);
+		DetailsViewArgs.bHideActorNameArea = true;
+		DetailsViewArgs.bShowOptions = false;
+		PropertyView = EditModule.CreateDetailView(DetailsViewArgs);
+		FOnGetDetailCustomizationInstance LayoutStructDetails = FOnGetDetailCustomizationInstance::CreateStatic(&FUserDefinedStructureDetails::MakeInstance);
+		PropertyView->RegisterInstancedCustomPropertyLayout(UUserDefinedStruct::StaticClass(), LayoutStructDetails);
+		PropertyView->SetObject(EditedStruct);
+		Box->AddSlot()
+		.AutoHeight()
+		[
+			PropertyView.ToSharedRef()
+		];
+	}
 
-	PropertyView->SetObject(EditedStruct);
+	DefaultValueView = NULL;
+	struct FShowDefaultValuePropertyEditor
+	{
+		bool bShowDefaultValuePropertyEditor;
+		FShowDefaultValuePropertyEditor() : bShowDefaultValuePropertyEditor(false)
+		{
+			GConfig->GetBool(TEXT("UserDefinedStructure"), TEXT("bShowDefaultValuePropertyEditor"), bShowDefaultValuePropertyEditor, GEditorIni);
+		}
+	};
+	static FShowDefaultValuePropertyEditor Helper;
+	if (Helper.bShowDefaultValuePropertyEditor)
+	{
+		DefaultValueView = MakeShareable(new FStructureDefaultValueView(EditedStruct));
+		DefaultValueView->Initialize();
+		auto DefaultValueWidget = DefaultValueView->GetWidget();
+		if (DefaultValueWidget.IsValid())
+		{
+			Box->AddSlot()
+			.AutoHeight()
+			.VAlign(EVerticalAlignment::VAlign_Top)
+			[
+				DefaultValueWidget.ToSharedRef()
+			];
+		}
+	}
 
 	return SNew(SDockTab)
 		.Icon( FEditorStyle::GetBrush("GenericEditor.Tabs.Properties") )
 		.Label( LOCTEXT("UserDefinedStructureEditor", "Structure") )
 		.TabColorScale( GetTabColorScale() )
 		[
-			PropertyView.ToSharedRef()
+			Box
 		];
 }
 
@@ -142,54 +318,11 @@ class FUserDefinedStructureLayout : public IDetailCustomNodeBuilder, public TSha
 public:
 	FUserDefinedStructureLayout(TWeakPtr<class FUserDefinedStructureDetails> InStructureDetails) : StructureDetails(InStructureDetails) {}
 
-	void OnFinishedChangingProperties(const FPropertyChangedEvent& PropertyChangedEvent)
-	{
-		auto StructureDetailsSP = StructureDetails.Pin();
-		if (StructureDetailsSP.IsValid())
-		{
-			check(PropertyChangedEvent.MemberProperty
-				&& PropertyChangedEvent.MemberProperty->GetOwnerStruct()
-				&& (PropertyChangedEvent.MemberProperty->GetOwnerStruct() == StructureDetailsSP->GetUserDefinedStruct())
-				&& PropertyChangedEvent.MemberProperty->GetOwnerStruct()->IsA<UUserDefinedStruct>());
-
-			const UProperty* DirectProperty = PropertyChangedEvent.MemberProperty;
-			while (!Cast<const UUserDefinedStruct>(DirectProperty->GetOuter()))
-			{
-				DirectProperty = CastChecked<const UProperty>(DirectProperty->GetOuter());
-			}
-
-			FString DefaultValueString;
-			bool bDefaultValueSet = false;
-			{
-				auto DefaultStructInstance = StructureDetailsSP->GetStructInstance();
-				if (DefaultStructInstance.IsValid() && DefaultStructInstance->IsValid())
-				{
-					bDefaultValueSet = FBlueprintEditorUtils::PropertyValueToString(DirectProperty, DefaultStructInstance->GetStructMemory(), DefaultValueString);
-				}
-			}
-
-			const FGuid VarGuid = FStructureEditorUtils::GetGuidForProperty(DirectProperty);
-			if (bDefaultValueSet && VarGuid.IsValid())
-			{
-				FStructureEditorUtils::ChangeVariableDefaultValue(StructureDetailsSP->GetUserDefinedStruct(), VarGuid, DefaultValueString);
-			}
-		}
-	}
-
 	void OnChanged()
 	{
 		OnRegenerateChildren.ExecuteIfBound();
 	}
 
-	void ResetStructureView()
-	{
-		auto StructureDetailsViewPinned = StructureDetailsView.Pin();
-		if (StructureDetailsView.IsValid())
-		{
-			StructureDetailsViewPinned->SetStructureData(NULL);
-		}
-	}
-
 	FReply OnAddNewField()
 	{
 		auto StructureDetailsSP = StructureDetails.Pin();
@@ -294,7 +427,6 @@ public:
 
 private:
 	TWeakPtr<class FUserDefinedStructureDetails> StructureDetails;
-	TWeakPtr<class IStructureDetailsView> StructureDetailsView;
 	FSimpleDelegate OnRegenerateChildren;
 };
 
@@ -759,50 +891,15 @@ void FUserDefinedStructureLayout::GenerateChildContent( IDetailChildrenBuilder&
 				TSharedRef<class FUserDefinedStructureFieldLayout> VarLayout = MakeShareable(new FUserDefinedStructureFieldLayout(StructureDetails,  SharedThis(this), VarDesc.VarGuid));
 				ChildrenBuilder.AddChildCustomBuilder(VarLayout);
 			}
-
-			struct FShowDefaultValuePropertyEditor
-			{
-				bool bShowDefaultValuePropertyEditor;
-				FShowDefaultValuePropertyEditor() : bShowDefaultValuePropertyEditor(false)
-				{
-					GConfig->GetBool(TEXT("UserDefinedStructure"), TEXT("bShowDefaultValuePropertyEditor"), bShowDefaultValuePropertyEditor, GEditorIni);
-				}
-			};
-			static FShowDefaultValuePropertyEditor Helper;
-			if (Helper.bShowDefaultValuePropertyEditor)
-			{
-				FPropertyEditorModule& PropertyModule = FModuleManager::LoadModuleChecked<FPropertyEditorModule>("PropertyEditor");
-				FDetailsViewArgs ViewArgs;
-				ViewArgs.bAllowSearch = false;
-				ViewArgs.bHideSelectionTip = false;
-				ViewArgs.bShowActorLabel = false;
-
-				auto StructureDetailViewRef = PropertyModule.CreateStructureDetailView(ViewArgs, StructureDetailsSP->GetStructInstance());
-				StructureDetailViewRef->GetOnFinishedChangingPropertiesDelegate().AddSP(this, &FUserDefinedStructureLayout::OnFinishedChangingProperties);
-				ChildrenBuilder.AddChildContent(FString())
-				[
-					StructureDetailViewRef->GetWidget().ToSharedRef()
-				];
-
-				StructureDetailsView = StructureDetailViewRef;
-			}
 		}
 	}
 }
 
 ///////////////////////////////////////////////////////////////////////////////////////
-// FUserDefinedStructureDetails
-
-FStructVariableDescription* FUserDefinedStructureDetails::FindStructureFieldByGuid(FGuid Guid)
-{
-	if (auto Struct = GetUserDefinedStruct())
-	{
-		return FStructureEditorUtils::GetVarDesc(Struct).FindByPredicate(FStructureEditorUtils::FFindByGuidHelper<FStructVariableDescription>(Guid));
-	}
-	return NULL;
-}
+// FUserDefinedStructureLayout
 
-void FUserDefinedStructureDetails::CustomizeDetails( class IDetailLayoutBuilder& DetailLayout )
+/** IDetailCustomization interface */
+void FUserDefinedStructureDetails::CustomizeDetails(class IDetailLayoutBuilder& DetailLayout) 
 {
 	const TArray<TWeakObjectPtr<UObject>> Objects = DetailLayout.GetDetailsView().GetSelectedObjects();
 	check(Objects.Num() > 0);
@@ -810,7 +907,6 @@ void FUserDefinedStructureDetails::CustomizeDetails( class IDetailLayoutBuilder&
 	if (Objects.Num() == 1)
 	{
 		UserDefinedStruct = CastChecked<UUserDefinedStruct>(Objects[0].Get());
-		StructData = MakeShareable(new FStructOnScope(UserDefinedStruct.Get()));
 
 		IDetailCategoryBuilder& StructureCategory = DetailLayout.EditCategory("Structure", LOCTEXT("Structure", "Structure").ToString());
 		Layout = MakeShareable(new FUserDefinedStructureLayout(SharedThis(this)));
@@ -820,25 +916,11 @@ void FUserDefinedStructureDetails::CustomizeDetails( class IDetailLayoutBuilder&
 	FStructureEditorUtils::FStructEditorManager::Get().AddListener(this);
 }
 
-void FUserDefinedStructureDetails::PreChange(const class UUserDefinedStruct* Struct)
-{
-	if (Struct && (GetUserDefinedStruct() == Struct) && StructData.IsValid())
-	{
-		if (Layout.IsValid())
-		{
-			Layout->ResetStructureView();
-		}
-		StructData->Destroy();
-		StructData.Reset();
-	}
-}
-
+/** FStructureEditorUtils::INotifyOnStructChanged */
 void FUserDefinedStructureDetails::PostChange(const class UUserDefinedStruct* Struct)
 {
 	if (Struct && (GetUserDefinedStruct() == Struct))
 	{
-		StructData = MakeShareable(new FStructOnScope(Struct));
-		FStructureEditorUtils::Fill_MakeStructureDefaultValue(Struct, StructData->GetStructMemory());
 		if (Layout.IsValid())
 		{
 			Layout->OnChanged();
@@ -846,11 +928,4 @@ void FUserDefinedStructureDetails::PostChange(const class UUserDefinedStruct* St
 	}
 }
 
-
-
-FUserDefinedStructureDetails::~FUserDefinedStructureDetails()
-{
-	FStructureEditorUtils::FStructEditorManager::Get().RemoveListener(this);
-}
-
 #undef LOCTEXT_NAMESPACE
\ No newline at end of file
diff --git a/Engine/Source/Editor/Kismet/Private/UserDefinedStructureEditor.h b/Engine/Source/Editor/Kismet/Private/UserDefinedStructureEditor.h
--- a/Engine/Source/Editor/Kismet/Private/UserDefinedStructureEditor.h
+++ b/Engine/Source/Editor/Kismet/Private/UserDefinedStructureEditor.h
@@ -4,7 +4,6 @@
 
 #include "BlueprintEditorModule.h"
 #include "Editor/PropertyEditor/Public/PropertyEditing.h"
-#include "Editor/UnrealEd/Public/Kismet2/StructureEditorUtils.h"
 
 class FUserDefinedStructureEditor : public IUserDefinedStructureEditor
 {
@@ -16,7 +15,7 @@ class FUserDefinedStructureEditor : public IUserDefinedStructureEditor
 	
 	/** Property viewing widget */
 	TSharedPtr<class IDetailsView> PropertyView;
-
+	TSharedPtr<class FStructureDefaultValueView> DefaultValueView;
 public:
 	/**
 	 * Edits the specified enum
@@ -25,7 +24,7 @@ public:
 	 * @param	InitToolkitHost			When Mode is WorldCentric, this is the level editor instance to spawn this editor within
 	 * @param	EnumToEdit				The user defined enum to edit
 	 */
-	void InitEditor(const EToolkitMode::Type Mode, const TSharedPtr< class IToolkitHost >& InitToolkitHost, UUserDefinedStruct* EnumToEdit);
+	void InitEditor(const EToolkitMode::Type Mode, const TSharedPtr< class IToolkitHost >& InitToolkitHost, class UUserDefinedStruct* EnumToEdit);
 
 	/** Destructor */
 	virtual ~FUserDefinedStructureEditor();
@@ -44,38 +43,3 @@ protected:
 	TSharedRef<SDockTab> SpawnStructureTab(const FSpawnTabArgs& Args);
 };
 
-class FUserDefinedStructureDetails : public IDetailCustomization, FStructureEditorUtils::INotifyOnStructChanged
-{
-public:
-	/** Makes a new instance of this detail layout class for a specific detail view requesting it */
-	static TSharedRef<IDetailCustomization> MakeInstance()
-	{
-		return MakeShareable(new FUserDefinedStructureDetails);
-	}
-
-	~FUserDefinedStructureDetails();
-
-	/** IDetailCustomization interface */
-	virtual void CustomizeDetails( class IDetailLayoutBuilder& DetailLayout ) override;
-
-	UUserDefinedStruct* GetUserDefinedStruct()
-	{
-		return UserDefinedStruct.Get();
-	}
-
-	TSharedPtr<FStructOnScope> GetStructInstance()
-	{
-		return StructData;
-	}
-
-	struct FStructVariableDescription* FindStructureFieldByGuid(FGuid Guid);
-
-	/** FStructureEditorUtils::INotifyOnStructChanged */
-	virtual void PreChange(const class UUserDefinedStruct* Struct) override;
-	virtual void PostChange(const class UUserDefinedStruct* Struct) override;
-
-private:
-	TWeakObjectPtr<UUserDefinedStruct> UserDefinedStruct;
-	TSharedPtr<class FUserDefinedStructureLayout> Layout;
-	TSharedPtr<FStructOnScope> StructData;
-};
diff --git a/Engine/Source/Editor/PropertyEditor/Private/PropertyEditorModule.cpp b/Engine/Source/Editor/PropertyEditor/Private/PropertyEditorModule.cpp
--- a/Engine/Source/Editor/PropertyEditor/Private/PropertyEditorModule.cpp
+++ b/Engine/Source/Editor/PropertyEditor/Private/PropertyEditorModule.cpp
@@ -633,22 +633,25 @@ FPropertyTypeLayoutCallback FPropertyEditorModule::GetPropertyTypeCustomization(
 	return FPropertyTypeLayoutCallback();
 }
 
-TSharedRef<class IStructureDetailsView> FPropertyEditorModule::CreateStructureDetailView(const FDetailsViewArgs& DetailsViewArgs, TSharedPtr<FStructOnScope> StructData)
+TSharedRef<class IStructureDetailsView> FPropertyEditorModule::CreateStructureDetailView(const FDetailsViewArgs& DetailsViewArgs, TSharedPtr<FStructOnScope> StructData, bool bShowObjects, const FString& CustomName)
 {
 	TSharedRef<SStructureDetailsView> DetailView =
 		SNew(SStructureDetailsView)
-		.DetailsViewArgs(DetailsViewArgs);
+		.DetailsViewArgs(DetailsViewArgs)
+		.CustomName(CustomName);
 
-	struct FDontShowObjects
+	if (!bShowObjects)
 	{
-		static bool CanShow(const UProperty* Property)
+		struct FDontShowObjects
 		{
-			return Property && !Property->IsA<UObjectPropertyBase>();
-		}
-	};
-
-	DetailView->SetIsPropertyVisibleDelegate(FIsPropertyVisible::CreateStatic(&FDontShowObjects::CanShow));
+			static bool CanShow(const UProperty* Property)
+			{
+				return Property && !Property->IsA<UObjectPropertyBase>() && !Property->IsA<UInterfaceProperty>();
+			}
+		};
 
+		DetailView->SetIsPropertyVisibleDelegate(FIsPropertyVisible::CreateStatic(&FDontShowObjects::CanShow));
+	}
 	DetailView->SetStructureData(StructData);
 
 	return DetailView;
diff --git a/Engine/Source/Editor/PropertyEditor/Private/SDetailsViewBase.cpp b/Engine/Source/Editor/PropertyEditor/Private/SDetailsViewBase.cpp
--- a/Engine/Source/Editor/PropertyEditor/Private/SDetailsViewBase.cpp
+++ b/Engine/Source/Editor/PropertyEditor/Private/SDetailsViewBase.cpp
@@ -987,6 +987,8 @@ void SDetailsViewBase::UpdatePropertyMap()
 	// When we encounter object property nodes, add their children instead of adding them to the tree.
 	UpdatePropertyMapRecursive(*RootPropertyNode, *DetailLayout, NAME_None, RootPropertyNode.Get());
 
+	CustomUpdatePropertyMap();
+
 	// Ask for custom detail layouts
 	QueryCustomDetailLayout(*DetailLayout);
 
diff --git a/Engine/Source/Editor/PropertyEditor/Private/SDetailsViewBase.h b/Engine/Source/Editor/PropertyEditor/Private/SDetailsViewBase.h
--- a/Engine/Source/Editor/PropertyEditor/Private/SDetailsViewBase.h
+++ b/Engine/Source/Editor/PropertyEditor/Private/SDetailsViewBase.h
@@ -265,7 +265,7 @@ protected:
 
 	/** Updates the property map for access when customizing the details view.  Generates default layout for properties */
 	void UpdatePropertyMap();
-
+	virtual void CustomUpdatePropertyMap() {}
 	/** 
 	 * Recursively updates children of property nodes. Generates default layout for properties 
 	 * 
diff --git a/Engine/Source/Editor/PropertyEditor/Private/SStructureDetailsView.cpp b/Engine/Source/Editor/PropertyEditor/Private/SStructureDetailsView.cpp
--- a/Engine/Source/Editor/PropertyEditor/Private/SStructureDetailsView.cpp
+++ b/Engine/Source/Editor/PropertyEditor/Private/SStructureDetailsView.cpp
@@ -31,6 +31,8 @@ void SStructureDetailsView::Construct(const FArguments& InArgs)
 	bIsLocked = false;
 	bHasOpenColorPicker = false;
 
+	CustomName = InArgs._CustomName;
+
 	// Create the root property now
 	RootNode = MakeShareable(new FStructurePropertyNode);
 		
@@ -194,6 +196,7 @@ void SStructureDetailsView::SetStructureData(TSharedPtr<FStructOnScope> InStruct
 	InitParams.bCreateCategoryNodes = false;
 
 	RootNode->InitNode(InitParams);
+	RootNode->SetDisplayNameOverride(CustomName);
 
 	RestoreExpandedItems();
 
@@ -233,4 +236,9 @@ TSharedPtr<class FComplexPropertyNode> SStructureDetailsView::GetRootNode()
 	return RootNode;
 }
 
+void SStructureDetailsView::CustomUpdatePropertyMap()
+{
+	DetailLayout->DefaultCategory(NAME_None).SetDisplayName(NAME_None, CustomName);
+}
+
 #undef LOCTEXT_NAMESPACE
\ No newline at end of file
diff --git a/Engine/Source/Editor/PropertyEditor/Private/SStructureDetailsView.h b/Engine/Source/Editor/PropertyEditor/Private/SStructureDetailsView.h
--- a/Engine/Source/Editor/PropertyEditor/Private/SStructureDetailsView.h
+++ b/Engine/Source/Editor/PropertyEditor/Private/SStructureDetailsView.h
@@ -8,7 +8,7 @@ class SStructureDetailsView : public SDetailsViewBase, public IStructureDetailsV
 {
 	TSharedPtr<struct FStructOnScope> StructData;
 	TSharedPtr<class FStructurePropertyNode> RootNode;
-
+	FString CustomName;
 public:
 
 	SLATE_BEGIN_ARGS(SStructureDetailsView)
@@ -16,6 +16,7 @@ public:
 		{}
 		/** The user defined args for the details view */
 		SLATE_ARGUMENT(FDetailsViewArgs, DetailsViewArgs)
+		SLATE_ARGUMENT(FString, CustomName)
 	SLATE_END_ARGS()
 
 	~SStructureDetailsView();
@@ -80,5 +81,7 @@ public:
 	virtual void SetObject(UObject* InObject, bool bForceRefresh = false) override{}
 
 	virtual TSharedPtr<class FComplexPropertyNode> GetRootNode() override;
+protected:
+	virtual void CustomUpdatePropertyMap() override;
 };
 
diff --git a/Engine/Source/Editor/PropertyEditor/Public/PropertyEditorModule.h b/Engine/Source/Editor/PropertyEditor/Public/PropertyEditorModule.h
--- a/Engine/Source/Editor/PropertyEditor/Public/PropertyEditorModule.h
+++ b/Engine/Source/Editor/PropertyEditor/Public/PropertyEditorModule.h
@@ -255,7 +255,7 @@ public:
 	 */
 	virtual TSharedPtr<class ISinglePropertyView> CreateSingleProperty( UObject* InObject, FName InPropertyName, const struct FSinglePropertyParams& InitParams );
 
-	virtual TSharedRef<class IStructureDetailsView> CreateStructureDetailView(const struct FDetailsViewArgs& DetailsViewArgs, TSharedPtr<struct FStructOnScope> StructData);
+	virtual TSharedRef<class IStructureDetailsView> CreateStructureDetailView(const struct FDetailsViewArgs& DetailsViewArgs, TSharedPtr<struct FStructOnScope> StructData, bool bShowObjects, const FString& CustomName = FString());
 
 	/**
 	 * Creates a property change listener that notifies users via a  delegate when a property on an object changes
