Mac compilation fix. extra qualification on member 'FStructureDefaultValueView'

diff --git a/Engine/Source/Editor/Kismet/Private/UserDefinedStructureEditor.cpp b/Engine/Source/Editor/Kismet/Private/UserDefinedStructureEditor.cpp
--- a/Engine/Source/Editor/Kismet/Private/UserDefinedStructureEditor.cpp
+++ b/Engine/Source/Editor/Kismet/Private/UserDefinedStructureEditor.cpp
@@ -16,7 +16,7 @@
 class FStructureDefaultValueView : public FStructureEditorUtils::INotifyOnStructChanged, public TSharedFromThis<FStructureDefaultValueView>
 {
 public:
-	FStructureDefaultValueView::FStructureDefaultValueView(UUserDefinedStruct* EditedStruct) 
+	FStructureDefaultValueView(UUserDefinedStruct* EditedStruct) 
 		: UserDefinedStruct(EditedStruct)
 	{
 		FStructureEditorUtils::FStructEditorManager::Get().AddListener(this);
