Find-in-Blueprints will now find a variety of new, useful bits of data in Blueprints (both for find-in-all and finding in current). Searches can now be done in the current culture and source culture.

Find-in-Blueprints now uses the same search data for both Find-in-all-Blueprints and Find-in-Blueprint. That searchdata is no longer stored in the asset registry, it is stored in the DDC and is managed by a new class. Searching is multithreaded to reduce overhead on the main thread during long searches in large projects.

virtual UEdGraphNode::AddSearchMetaDataInfo which sets base find-in-blueprints metadata and can be overriden with more detailed search data.
#change UEdGraphPin::GetDefaultAsString will return the text value as string if set

Added FBlueprintEditorUtils::GetEntryAndResultNode, moving most of the functionality from FBlueprintGraphActionDetails::SetEntryAndResultNodes to the new function
Added FBlueprintEditorUtils::GetGraphDescription to get graph descriptions from their metadata
Added FBlueprintEditorUtils::GetNodeByGUID to find nodes using their Guid

Blueprints now cache their search metadata on compile, this is a CPU intensive process and can no longer occur every time the Asset Registry queries for the Blueprint's tags

#change Fixed an issue with finding a pin type's icon (UK2Node_Variable::GetVarIconFromPinType) where we would search for a class we already had available.

Added an overridable function to UEdGraphNode to collect Blueprint search metadata.

#change FJsonObjectConverter::UPropertyToJsonValue is now available outside of the .cpp

#change Rewrote how Find-in-Blueprints works both for local and global searches, now uses Json to resolve a variety of format issues and also to make it much more extensible.

#ttp 290376 - K2: Ability to search (find in blueprints) important pin values
#ttp 298599 - ROCKET: TASK: K2: Search results should indicate where the instance is located in the results page (construction script, event graph, custom graph)
#ttp 304819 - Blueprints: Editor: L10N: "Find in Blueprint" needs to work with other languages
#ttp 308572 - UE4: ANIM: K2: Searching doesn't work for anim nodes right now
#ttp 316678 - UE4: BLUEPRINTS: SHADOW: Default values are not searchable
#ttp 318180 - UE4: K2: Cannot search for spawnactor using find in blueprints
#ttp 320141 - UE4: K2: Find in blueprints does not find functions declared in another blueprint (if the function isn't called in that BP)

diff --git a/Engine/Plugins/Editor/EpicSurvey/Source/EpicSurvey/Private/EpicSurvey.cpp b/Engine/Plugins/Editor/EpicSurvey/Source/EpicSurvey/Private/EpicSurvey.cpp
--- a/Engine/Plugins/Editor/EpicSurvey/Source/EpicSurvey/Private/EpicSurvey.cpp
+++ b/Engine/Plugins/Editor/EpicSurvey/Source/EpicSurvey/Private/EpicSurvey.cpp
@@ -152,7 +152,7 @@ void FEpicSurvey::Initialize()
 
 void FEpicSurvey::InitializeTitleCloud()
 {
-	TSharedRef<FCulture> Culture = FInternationalization::Get().GetCurrentCulture();
+	TSharedRef<FCulture, ESPMode::ThreadSafe> Culture = FInternationalization::Get().GetCurrentCulture();
 
 	FString CultureString;
 	switch (CurrentCulture)
diff --git a/Engine/Source/Developer/InternationalizationSettings/Private/InternationalizationSettingsModelDetails.cpp b/Engine/Source/Developer/InternationalizationSettings/Private/InternationalizationSettingsModelDetails.cpp
--- a/Engine/Source/Developer/InternationalizationSettings/Private/InternationalizationSettingsModelDetails.cpp
+++ b/Engine/Source/Developer/InternationalizationSettings/Private/InternationalizationSettingsModelDetails.cpp
@@ -7,14 +7,14 @@
 /** Functions for sorting the languages */
 struct FCompareCultureByNativeLanguage
 {
-	static FText GetCultureNativeLanguageText( const TSharedPtr<FCulture> Culture )
+	static FText GetCultureNativeLanguageText( const TSharedPtr<FCulture, ESPMode::ThreadSafe> Culture )
 	{
 		check( Culture.IsValid() );
 		const FString Language = Culture->GetNativeLanguage();
 		return FText::FromString(Language);
 	}
 
-	FORCEINLINE bool operator()( const TSharedPtr<FCulture> A, const TSharedPtr<FCulture> B ) const
+	FORCEINLINE bool operator()( const TSharedPtr<FCulture, ESPMode::ThreadSafe> A, const TSharedPtr<FCulture, ESPMode::ThreadSafe> B ) const
 	{
 		check( A.IsValid() );
 		check( B.IsValid() );
@@ -26,7 +26,7 @@ struct FCompareCultureByNativeLanguage
 /** Functions for sorting the regions */
 struct FCompareCultureByNativeRegion
 {
-	static FText GetCultureNativeRegionText( const TSharedPtr<FCulture> Culture )
+	static FText GetCultureNativeRegionText( const TSharedPtr<FCulture, ESPMode::ThreadSafe> Culture )
 	{
 		check( Culture.IsValid() );
 		FString Region = Culture->GetNativeRegion();
@@ -38,7 +38,7 @@ struct FCompareCultureByNativeRegion
 		return FText::FromString(Region);
 	}
 
-	FORCEINLINE bool operator()( const TSharedPtr<FCulture> A, const TSharedPtr<FCulture> B ) const
+	FORCEINLINE bool operator()( const TSharedPtr<FCulture, ESPMode::ThreadSafe> A, const TSharedPtr<FCulture, ESPMode::ThreadSafe> B ) const
 	{
 		check( A.IsValid() );
 		check( B.IsValid() );
@@ -127,6 +127,9 @@ void FInternationalizationSettingsModelDetails::CustomizeDetails( IDetailLayoutB
 
 	const FText LanguageToolTipText = LOCTEXT("EditorLanguageTooltip", "Change the Editor language (requires restart to take effect)");
 
+	// For use in the Slate macros below, the type must be typedef'd to compile.
+	typedef TSharedPtr<FCulture, ESPMode::ThreadSafe> ThreadSafeCulturePtr;
+
 	CategoryBuilder.AddCustomRow("Language")
 	.NameContent()
 	[
@@ -144,7 +147,7 @@ void FInternationalizationSettingsModelDetails::CustomizeDetails( IDetailLayoutB
 	.ValueContent()
 	.MaxDesiredWidth(300.0f)
 	[
-		SAssignNew(LanguageComboBox, SComboBox< TSharedPtr<FCulture> >)
+		SAssignNew(LanguageComboBox, SComboBox< ThreadSafeCulturePtr > )
 		.OptionsSource( &AvailableLanguages )
 		.InitiallySelectedItem(SelectedLanguage)
 		.OnGenerateWidget(this, &FInternationalizationSettingsModelDetails::OnLanguageGenerateWidget, &DetailBuilder)
@@ -177,7 +180,7 @@ void FInternationalizationSettingsModelDetails::CustomizeDetails( IDetailLayoutB
 	.ValueContent()
 	.MaxDesiredWidth(300.0f)
 	[
-		SAssignNew(RegionComboBox, SComboBox< TSharedPtr<FCulture> >)
+		SAssignNew(RegionComboBox, SComboBox< ThreadSafeCulturePtr > )
 		.OptionsSource( &AvailableRegions )
 		.InitiallySelectedItem(SelectedCulture)
 		.OnGenerateWidget(this, &FInternationalizationSettingsModelDetails::OnRegionGenerateWidget, &DetailBuilder)
@@ -292,7 +295,7 @@ void FInternationalizationSettingsModelDetails::RefreshAvailableRegions()
 {
 	AvailableRegions.Empty();
 
-	TSharedPtr<FCulture> DefaultCulture = NULL;
+	TSharedPtr<FCulture, ESPMode::ThreadSafe> DefaultCulture = NULL;
 	if ( SelectedLanguage.IsValid() )
 	{
 		const FString SelectedLanguageName = SelectedLanguage->GetNativeLanguage();
@@ -320,7 +323,7 @@ void FInternationalizationSettingsModelDetails::RefreshAvailableRegions()
 	// If we have a preferred default (or there's only one in the list), select that now
 	if ( DefaultCulture.IsValid() || AvailableCultures.Num() == 1 )
 	{
-		TSharedPtr<FCulture> Culture = DefaultCulture.IsValid() ? DefaultCulture : AvailableCultures.Last();
+		TSharedPtr<FCulture, ESPMode::ThreadSafe> Culture = DefaultCulture.IsValid() ? DefaultCulture : AvailableCultures.Last();
 		// Set it as our default region, if one hasn't already been chosen
 		if ( !SelectedCulture.IsValid() && RegionComboBox.IsValid() )
 		{
@@ -344,14 +347,14 @@ FText FInternationalizationSettingsModelDetails::GetCurrentLanguageText() const
 	return FText::GetEmpty();
 }
 
-TSharedRef<SWidget> FInternationalizationSettingsModelDetails::OnLanguageGenerateWidget( TSharedPtr<FCulture> Culture, IDetailLayoutBuilder* DetailBuilder ) const
+TSharedRef<SWidget> FInternationalizationSettingsModelDetails::OnLanguageGenerateWidget( TSharedPtr<FCulture, ESPMode::ThreadSafe> Culture, IDetailLayoutBuilder* DetailBuilder ) const
 {
 	return SNew(STextBlock)
 		.Text(FCompareCultureByNativeLanguage::GetCultureNativeLanguageText(Culture))
 		.Font(DetailBuilder->GetDetailFont());
 }
 
-void FInternationalizationSettingsModelDetails::OnLanguageSelectionChanged( TSharedPtr<FCulture> Culture, ESelectInfo::Type SelectionType )
+void FInternationalizationSettingsModelDetails::OnLanguageSelectionChanged( TSharedPtr<FCulture, ESPMode::ThreadSafe> Culture, ESelectInfo::Type SelectionType )
 {
 	SelectedLanguage = Culture;
 	SelectedCulture = NULL;
@@ -369,14 +372,14 @@ FText FInternationalizationSettingsModelDetails::GetCurrentRegionText() const
 	return FText::GetEmpty();
 }
 
-TSharedRef<SWidget> FInternationalizationSettingsModelDetails::OnRegionGenerateWidget( TSharedPtr<FCulture> Culture, IDetailLayoutBuilder* DetailBuilder ) const
+TSharedRef<SWidget> FInternationalizationSettingsModelDetails::OnRegionGenerateWidget( TSharedPtr<FCulture, ESPMode::ThreadSafe> Culture, IDetailLayoutBuilder* DetailBuilder ) const
 {
 	return SNew(STextBlock)
 		.Text(FCompareCultureByNativeRegion::GetCultureNativeRegionText(Culture))
 		.Font(DetailBuilder->GetDetailFont());
 }
 
-void FInternationalizationSettingsModelDetails::OnRegionSelectionChanged( TSharedPtr<FCulture> Culture, ESelectInfo::Type SelectionType )
+void FInternationalizationSettingsModelDetails::OnRegionSelectionChanged( TSharedPtr<FCulture, ESPMode::ThreadSafe> Culture, ESelectInfo::Type SelectionType )
 {
 	SelectedCulture = Culture;
 
@@ -417,4 +420,4 @@ void FInternationalizationSettingsModelDetails::HandleShutdownPostPackagesSaved(
 	}
 }
 
-#undef LOCTEXT_NAMESPACE
\ No newline at end of file
+#undef LOCTEXT_NAMESPACE
diff --git a/Engine/Source/Developer/InternationalizationSettings/Public/InternationalizationSettingsModelDetails.h b/Engine/Source/Developer/InternationalizationSettings/Public/InternationalizationSettingsModelDetails.h
--- a/Engine/Source/Developer/InternationalizationSettings/Public/InternationalizationSettingsModelDetails.h
+++ b/Engine/Source/Developer/InternationalizationSettings/Public/InternationalizationSettingsModelDetails.h
@@ -42,19 +42,19 @@ private:
 	FText GetCurrentLanguageText() const;
 
 	/** Generate a widget for the language combo */
-	TSharedRef<SWidget> OnLanguageGenerateWidget( TSharedPtr<FCulture> Culture, IDetailLayoutBuilder* DetailBuilder ) const;
+	TSharedRef<SWidget> OnLanguageGenerateWidget( TSharedPtr<FCulture, ESPMode::ThreadSafe> Culture, IDetailLayoutBuilder* DetailBuilder ) const;
 
 	/** Delegate called when the language selection has changed */
-	void OnLanguageSelectionChanged( TSharedPtr<FCulture> Culture, ESelectInfo::Type SelectionType );
+	void OnLanguageSelectionChanged( TSharedPtr<FCulture, ESPMode::ThreadSafe> Culture, ESelectInfo::Type SelectionType );
 
 	/** Generate a widget for the region combo */
-	TSharedRef<SWidget> OnRegionGenerateWidget( TSharedPtr<FCulture> Culture, IDetailLayoutBuilder* DetailBuilder ) const;
+	TSharedRef<SWidget> OnRegionGenerateWidget( TSharedPtr<FCulture, ESPMode::ThreadSafe> Culture, IDetailLayoutBuilder* DetailBuilder ) const;
 
 	/** Delegate called to display the current region */
 	FText GetCurrentRegionText() const;
 
 	/** Delegate called when the region selection has changed */
-	void OnRegionSelectionChanged( TSharedPtr<FCulture> Culture, ESelectInfo::Type SelectionType );
+	void OnRegionSelectionChanged( TSharedPtr<FCulture, ESPMode::ThreadSafe> Culture, ESelectInfo::Type SelectionType );
 
 	/** Delegate called to see if we're allowed to change the region */
 	bool IsRegionSelectionAllowed() const;
@@ -72,27 +72,27 @@ private:
 	TWeakObjectPtr<UInternationalizationSettingsModel> Model;
 
 	/** The culture selected at present */
-	TSharedPtr<FCulture> SelectedCulture;
+	TSharedPtr<FCulture, ESPMode::ThreadSafe> SelectedCulture;
 
 	/** The language selected at present */
-	TSharedPtr<FCulture> SelectedLanguage;
+	TSharedPtr<FCulture, ESPMode::ThreadSafe> SelectedLanguage;
 
 	/** The cultures we have available to us */
-	TArray< TSharedPtr<FCulture> > AvailableCultures;
+	TArray< TSharedPtr<FCulture, ESPMode::ThreadSafe> > AvailableCultures;
 
 	/** The languages we have available to us */
-	TArray< TSharedPtr<FCulture> > AvailableLanguages;
+	TArray< TSharedPtr<FCulture, ESPMode::ThreadSafe> > AvailableLanguages;
 
 	/** The dropdown for the languages */
-	TSharedPtr<SComboBox< TSharedPtr<FCulture> >> LanguageComboBox;
+	TSharedPtr<SComboBox< TSharedPtr<FCulture, ESPMode::ThreadSafe> > > LanguageComboBox;
 
 	/** The regions we have available to us */
-	TArray< TSharedPtr<FCulture> > AvailableRegions;
+	TArray< TSharedPtr<FCulture, ESPMode::ThreadSafe> > AvailableRegions;
 
 	/** The dropdown for the regions */
-	TSharedPtr<SComboBox< TSharedPtr<FCulture> >> RegionComboBox;
+	TSharedPtr<SComboBox< TSharedPtr<FCulture, ESPMode::ThreadSafe> > > RegionComboBox;
 
-	/** The check boxc for using localized field names */
+	/** The check box for using localized field names */
 	TSharedPtr<SCheckBox> FieldNamesCheckBox;
 
 	bool RequiresRestart;
diff --git a/Engine/Source/Editor/BlueprintGraph/Classes/K2Node_Variable.h b/Engine/Source/Editor/BlueprintGraph/Classes/K2Node_Variable.h
--- a/Engine/Source/Editor/BlueprintGraph/Classes/K2Node_Variable.h
+++ b/Engine/Source/Editor/BlueprintGraph/Classes/K2Node_Variable.h
@@ -137,6 +137,6 @@ public:
 	 * @param  IconColorOut	A color out, further discerning the variable's type.
 	 * @return A icon representing the specified variable's type.
 	 */
-	BLUEPRINTGRAPH_API static FName GetVarIconFromPinType(FEdGraphPinType& InPinType, FLinearColor& IconColorOut);
+	BLUEPRINTGRAPH_API static FName GetVarIconFromPinType(const FEdGraphPinType& InPinType, FLinearColor& IconColorOut);
 };
 
diff --git a/Engine/Source/Editor/BlueprintGraph/Private/K2Node_Variable.cpp b/Engine/Source/Editor/BlueprintGraph/Private/K2Node_Variable.cpp
--- a/Engine/Source/Editor/BlueprintGraph/Private/K2Node_Variable.cpp
+++ b/Engine/Source/Editor/BlueprintGraph/Private/K2Node_Variable.cpp
@@ -274,7 +274,7 @@ FName UK2Node_Variable::GetPaletteIcon(FLinearColor& ColorOut) const
 	return ReturnIconName;
 }
 
-FName UK2Node_Variable::GetVarIconFromPinType(FEdGraphPinType& InPinType, FLinearColor& IconColorOut)
+FName UK2Node_Variable::GetVarIconFromPinType(const FEdGraphPinType& InPinType, FLinearColor& IconColorOut)
 {
 	FName IconBrush = TEXT("Kismet.AllClasses.VariableIcon");
 
@@ -287,10 +287,9 @@ FName UK2Node_Variable::GetVarIconFromPinType(FEdGraphPinType& InPinType, FLinea
 	}
 	else if(InPinType.PinSubCategoryObject.IsValid())
 	{
-		UClass* VarClass = FindObject<UClass>(ANY_PACKAGE, *InPinType.PinSubCategoryObject->GetName());
-		if( VarClass )
+		if(UClass* Class = Cast<UClass>(InPinType.PinSubCategoryObject.Get()))
 		{
-			IconBrush = FClassIconFinder::FindIconNameForClass( VarClass );
+			IconBrush = FClassIconFinder::FindIconNameForClass( Class );
 		}
 	}
 
diff --git a/Engine/Source/Editor/Documentation/Private/Documentation.cpp b/Engine/Source/Editor/Documentation/Private/Documentation.cpp
--- a/Engine/Source/Editor/Documentation/Private/Documentation.cpp
+++ b/Engine/Source/Editor/Documentation/Private/Documentation.cpp
@@ -30,7 +30,7 @@ bool FDocumentation::OpenHome() const
 	return Open( TEXT("%ROOT%") );
 }
 
-bool FDocumentation::OpenHome(const TSharedRef<FCulture>& Culture) const
+bool FDocumentation::OpenHome(const TSharedRef<FCulture, ESPMode::ThreadSafe>& Culture) const
 {
 	return Open(TEXT("%ROOT%"), Culture);
 }
@@ -84,7 +84,7 @@ bool FDocumentation::Open( const FString& Link ) const
 	return !DocumentationUrl.IsEmpty();
 }
 
-bool FDocumentation::Open(const FString& Link, const TSharedRef<FCulture>& Culture) const
+bool FDocumentation::Open(const FString& Link, const TSharedRef<FCulture, ESPMode::ThreadSafe>& Culture) const
 {
 	FString DocumentationUrl;
 
@@ -157,7 +157,7 @@ bool FDocumentation::PageExists(const FString& Link) const
 	return FPaths::FileExists(SourcePath);
 }
 
-bool FDocumentation::PageExists(const FString& Link, const TSharedRef<FCulture>& Culture) const
+bool FDocumentation::PageExists(const FString& Link, const TSharedRef<FCulture, ESPMode::ThreadSafe>& Culture) const
 {
 	const TWeakPtr< IDocumentationPage >* ExistingPagePtr = LoadedPages.Find(Link);
 	if (ExistingPagePtr != NULL)
diff --git a/Engine/Source/Editor/Documentation/Private/Documentation.h b/Engine/Source/Editor/Documentation/Private/Documentation.h
--- a/Engine/Source/Editor/Documentation/Private/Documentation.h
+++ b/Engine/Source/Editor/Documentation/Private/Documentation.h
@@ -15,13 +15,13 @@ public:
 
 	virtual bool OpenHome() const override;
 
-	virtual bool OpenHome(const TSharedRef<FCulture>& Culture) const override;
+	virtual bool OpenHome(const TSharedRef<FCulture, ESPMode::ThreadSafe>& Culture) const override;
 
 	virtual bool OpenAPIHome() const override;
 
 	virtual bool Open( const FString& Link ) const override;
 
-	virtual bool Open(const FString& Link, const TSharedRef<FCulture>& Culture) const override;
+	virtual bool Open(const FString& Link, const TSharedRef<FCulture, ESPMode::ThreadSafe>& Culture) const override;
 
 	virtual TSharedRef< SWidget > CreateAnchor( const FString& Link, const FString& PreviewLink = FString(), const FString& PreviewExcerptName = FString() ) const override;
 
@@ -29,7 +29,7 @@ public:
 
 	virtual bool PageExists(const FString& Link) const override;
 
-	virtual bool PageExists(const FString& Link, const TSharedRef<FCulture>& Culture) const override;
+	virtual bool PageExists(const FString& Link, const TSharedRef<FCulture, ESPMode::ThreadSafe>& Culture) const override;
 
 	virtual TSharedRef< class SToolTip > CreateToolTip( const TAttribute<FText>& Text, const TSharedPtr<SWidget>& OverrideContent, const FString& Link, const FString& ExcerptName ) const override;
 
@@ -40,4 +40,4 @@ private:
 private:
 
 	TMap< FString, TWeakPtr< IDocumentationPage > > LoadedPages;
-};
\ No newline at end of file
+};
diff --git a/Engine/Source/Editor/Documentation/Private/DocumentationLink.h b/Engine/Source/Editor/Documentation/Private/DocumentationLink.h
--- a/Engine/Source/Editor/Documentation/Private/DocumentationLink.h
+++ b/Engine/Source/Editor/Documentation/Private/DocumentationLink.h
@@ -23,7 +23,7 @@ public:
 		return GetHomeUrl(FInternationalization::Get().GetCurrentCulture());
 	}
 
-	static FString GetHomeUrl(const TSharedRef<FCulture>& Culture)
+	static FString GetHomeUrl(const TSharedRef<FCulture, ESPMode::ThreadSafe>& Culture)
 	{
 		FString Url;
 		FUnrealEdMisc::Get().GetURL( TEXT("UDNURL"), Url, true );
@@ -37,7 +37,7 @@ public:
 		return ToUrl(Link, FInternationalization::Get().GetCurrentCulture());
 	}
 
-	static FString ToUrl(const FString& Link, const TSharedRef<FCulture>& Culture)
+	static FString ToUrl(const FString& Link, const TSharedRef<FCulture, ESPMode::ThreadSafe>& Culture)
 	{
 		FString Path;
 		FString Anchor;
@@ -56,7 +56,7 @@ public:
 
 		if (!FPaths::FileExists(FilePath))
 		{
-			const TSharedPtr<FCulture> FallbackCulture = I18N.GetCulture(TEXT("en"));
+			const TSharedPtr<FCulture, ESPMode::ThreadSafe> FallbackCulture = I18N.GetCulture(TEXT("en"));
 			if (FallbackCulture.IsValid())
 			{
 				const FString FallbackFilePath = ToFilePath(Link, FallbackCulture.ToSharedRef());
@@ -70,7 +70,7 @@ public:
 		return FilePath;
 	}
 
-	static FString ToFilePath(const FString& Link, const TSharedRef<FCulture>& Culture)
+	static FString ToFilePath(const FString& Link, const TSharedRef<FCulture, ESPMode::ThreadSafe>& Culture)
 	{
 		FString Path;
 		FString Anchor;
@@ -84,12 +84,12 @@ public:
 	{
 		FInternationalization& I18N = FInternationalization::Get();
 
-		TSharedRef<FCulture> Culture = I18N.GetCurrentCulture();
+		TSharedRef<FCulture, ESPMode::ThreadSafe> Culture = I18N.GetCurrentCulture();
 		FString FilePath = ToFilePath(Link, Culture);
 
 		if (!FPaths::FileExists(FilePath))
 		{
-			const TSharedPtr<FCulture> FallbackCulture = I18N.GetCulture(TEXT("en"));
+			const TSharedPtr<FCulture, ESPMode::ThreadSafe> FallbackCulture = I18N.GetCulture(TEXT("en"));
 			if (FallbackCulture.IsValid())
 			{
 				const FString FallbackFilePath = ToFilePath(Link, FallbackCulture.ToSharedRef());
@@ -103,7 +103,7 @@ public:
 		return ToFileUrl(Link, Culture);
 	}
 
-	static FString ToFileUrl(const FString& Link, const TSharedRef<FCulture>& Culture)
+	static FString ToFileUrl(const FString& Link, const TSharedRef<FCulture, ESPMode::ThreadSafe>& Culture)
 	{
 		FString Path;
 		FString Anchor;
@@ -120,7 +120,7 @@ public:
 
 		if (!FPaths::FileExists(SourcePath))
 		{
-			const TSharedPtr<FCulture> FallbackCulture = I18N.GetCulture(TEXT("en"));
+			const TSharedPtr<FCulture, ESPMode::ThreadSafe> FallbackCulture = I18N.GetCulture(TEXT("en"));
 			if (FallbackCulture.IsValid())
 			{
 				const FString FallbackSourcePath = ToSourcePath(Link, FallbackCulture.ToSharedRef());
@@ -134,7 +134,7 @@ public:
 		return SourcePath;
 	}
 
-	static FString ToSourcePath(const FString& Link, const TSharedRef<FCulture>& Culture)
+	static FString ToSourcePath(const FString& Link, const TSharedRef<FCulture, ESPMode::ThreadSafe>& Culture)
 	{
 		FString Path;
 		FString Anchor;
diff --git a/Engine/Source/Editor/Documentation/Public/IDocumentation.h b/Engine/Source/Editor/Documentation/Public/IDocumentation.h
--- a/Engine/Source/Editor/Documentation/Public/IDocumentation.h
+++ b/Engine/Source/Editor/Documentation/Public/IDocumentation.h
@@ -139,13 +139,13 @@ public:
 
 	virtual bool OpenHome() const = 0;
 
-	virtual bool OpenHome(const TSharedRef<FCulture>& Culture) const = 0;
+	virtual bool OpenHome(const TSharedRef<FCulture, ESPMode::ThreadSafe>& Culture) const = 0;
 
 	virtual bool OpenAPIHome() const = 0;
 
 	virtual bool Open( const FString& Link ) const = 0;
 
-	virtual bool Open(const FString& Link, const TSharedRef<FCulture>& Culture) const = 0;
+	virtual bool Open(const FString& Link, const TSharedRef<FCulture, ESPMode::ThreadSafe>& Culture) const = 0;
 
 	virtual TSharedRef< class SWidget > CreateAnchor( const FString& Link, const FString& PreviewLink = FString(), const FString& PreviewExcerptName = FString() ) const = 0;
 
@@ -153,7 +153,7 @@ public:
 
 	virtual bool PageExists(const FString& Link) const = 0;
 
-	virtual bool PageExists(const FString& Link, const TSharedRef<FCulture>& Culture) const = 0;
+	virtual bool PageExists(const FString& Link, const TSharedRef<FCulture, ESPMode::ThreadSafe>& Culture) const = 0;
 
 	virtual TSharedRef< class SToolTip > CreateToolTip( const TAttribute<FText>& Text, const TSharedPtr<SWidget>& OverrideContent, const FString& Link, const FString& ExcerptName ) const = 0;
 };
diff --git a/Engine/Source/Editor/Kismet/Kismet.Build.cs b/Engine/Source/Editor/Kismet/Kismet.Build.cs
--- a/Engine/Source/Editor/Kismet/Kismet.Build.cs
+++ b/Engine/Source/Editor/Kismet/Kismet.Build.cs
@@ -13,7 +13,8 @@ public class Kismet : ModuleRules
 				"AssetRegistry", 
 				"AssetTools",
 				"EditorWidgets",
-				"Analytics"
+				"Analytics",
+                "DerivedDataCache",
 			}
 			);
 
@@ -36,7 +37,9 @@ public class Kismet : ModuleRules
 				"LevelEditor",
                 "InputCore",
 				"EngineSettings",
-                "Projects"
+                "Projects",
+                "JsonUtilities",
+                "DerivedDataCache",
 			}
 			);
 
diff --git a/Engine/Source/Editor/Kismet/Private/BlueprintDetailsCustomization.cpp b/Engine/Source/Editor/Kismet/Private/BlueprintDetailsCustomization.cpp
--- a/Engine/Source/Editor/Kismet/Private/BlueprintDetailsCustomization.cpp
+++ b/Engine/Source/Editor/Kismet/Private/BlueprintDetailsCustomization.cpp
@@ -2896,62 +2896,7 @@ void FBlueprintGraphActionDetails::SetEntryAndResultNodes()
 
 	if (UEdGraph* NewTargetGraph = GetGraph())
 	{
-		// There are a few different potential configurations for editable graphs (FunctionEntry/Result, Tunnel Pairs, etc).
-		// Step through each case until we find one that matches what appears to be in the graph.  This could be improved if
-		// we want to add more robust typing to the graphs themselves
-
-		// Case 1:  Function Entry / Result Pair ------------------
-		TArray<UK2Node_FunctionEntry*> EntryNodes;
-		NewTargetGraph->GetNodesOfClass(EntryNodes);
-
-		if (EntryNodes.Num() > 0)
-		{
-			if (EntryNodes[0]->IsEditable())
-			{
-				FunctionEntryNodePtr = EntryNodes[0];
-
-				// Find a result node
-				TArray<UK2Node_FunctionResult*> ResultNodes;
-				NewTargetGraph->GetNodesOfClass(ResultNodes);
-
-				check(ResultNodes.Num() <= 1);
-				UK2Node_FunctionResult* ResultNode = ResultNodes.Num() ? ResultNodes[0] : NULL;
-				// Note:  we assume that if the entry is editable, the result is too (since the entry node is guaranteed to be there on graph creation, but the result isn't)
-				if( ResultNode )
-				{
-					FunctionResultNodePtr = ResultNode;
-				}
-			}
-		}
-		else
-		{
-			// Case 2:  Tunnel Pair -----------------------------------
-			TArray<UK2Node_Tunnel*> TunnelNodes;
-			NewTargetGraph->GetNodesOfClass(TunnelNodes);
-
-			if (TunnelNodes.Num() > 0)
-			{
-				// Iterate over the tunnel nodes, and try to find an entry and exit
-				for (int32 i = 0; i < TunnelNodes.Num(); i++)
-				{
-					UK2Node_Tunnel* Node = TunnelNodes[i];
-					// Composite nodes should never be considered for function entry / exit, since we're searching for a graph's terminals
-					if (Node->IsEditable() && !Node->IsA(UK2Node_Composite::StaticClass()))
-					{
-						if (Node->bCanHaveOutputs)
-						{
-							ensure(!FunctionEntryNodePtr.IsValid());
-							FunctionEntryNodePtr = Node;
-						}
-						else if (Node->bCanHaveInputs)
-						{
-							ensure(!FunctionResultNodePtr.IsValid());
-							FunctionResultNodePtr = Node;
-						}
-					}
-				}
-			}
-		}
+		FBlueprintEditorUtils::GetEntryAndResultNodes(NewTargetGraph, FunctionEntryNodePtr, FunctionResultNodePtr);
 	}
 	else if (UK2Node_EditablePinBase* Node = GetEditableNode())
 	{
diff --git a/Engine/Source/Editor/Kismet/Private/BlueprintEditor.cpp b/Engine/Source/Editor/Kismet/Private/BlueprintEditor.cpp
--- a/Engine/Source/Editor/Kismet/Private/BlueprintEditor.cpp
+++ b/Engine/Source/Editor/Kismet/Private/BlueprintEditor.cpp
@@ -18,6 +18,7 @@
 #include "MessageLog.h"
 #include "UObjectToken.h"
 #include "Kismet/GameplayStatics.h"
+#include "Editor/Kismet/Public/FindInBlueprintManager.h"
 #include "Editor/BlueprintGraph/Public/K2ActionMenuBuilder.h"
 #include "BlueprintEditorCommands.h"
 #include "BlueprintEditor.h"
@@ -6365,6 +6366,15 @@ bool FBlueprintEditor::IsFocusedGraphEditable() const
 	return true;
 }
 
+void FBlueprintEditor::SaveAsset_Execute()
+{
+	// Update the Blueprint's search data
+	GetBlueprintObj()->SearchGuid = FGuid::NewGuid();
+	FFindInBlueprintSearchManager::Get().AddOrUpdateBlueprintSearchMetadata(GetBlueprintObj());
+
+	IBlueprintEditor::SaveAsset_Execute();
+}
+
 /////////////////////////////////////////////////////
 /////////////////////////////////////////////////////
 /////////////////////////////////////////////////////
diff --git a/Engine/Source/Editor/Kismet/Private/FindInBlueprintManager.cpp b/Engine/Source/Editor/Kismet/Private/FindInBlueprintManager.cpp
new file mode 100644
--- /dev/null
+++ b/Engine/Source/Editor/Kismet/Private/FindInBlueprintManager.cpp
@@ -0,0 +1,1131 @@
+#include "BlueprintEditorPrivatePCH.h"
+
+#include "FindInBlueprintManager.h"
+#include "FindInBlueprints.h"
+#include "Editor/UnrealEd/Public/Kismet2/BlueprintEditorUtils.h"
+#include "AssetRegistryModule.h"
+#include "DerivedDataCacheInterface.h"
+#include "DerivedDataPluginInterface.h"
+
+#include "JsonUtilities.h"
+
+#define LOCTEXT_NAMESPACE "FindInBlueprintManager"
+
+FFindInBlueprintSearchManager* FFindInBlueprintSearchManager::Instance = NULL;
+
+const FText FFindInBlueprintSearchTags::FiB_Properties = LOCTEXT("Properties", "Properties");
+
+const FText FFindInBlueprintSearchTags::FiB_Nodes = LOCTEXT("Nodes", "Nodes");
+
+const FText FFindInBlueprintSearchTags::FiB_SchemaName = LOCTEXT("SchemaName", "SchemaName");
+
+const FText FFindInBlueprintSearchTags::FiB_UberGraphs = LOCTEXT("Uber", "Uber");
+const FText FFindInBlueprintSearchTags::FiB_Functions = LOCTEXT("Functions", "Functions");
+const FText FFindInBlueprintSearchTags::FiB_Macros = LOCTEXT("Macros", "Macros");
+const FText FFindInBlueprintSearchTags::FiB_SubGraphs = LOCTEXT("Sub", "Sub");
+
+const FText FFindInBlueprintSearchTags::FiB_Name = LOCTEXT("Name", "Name");
+const FText FFindInBlueprintSearchTags::FiB_ClassName = LOCTEXT("ClassName", "ClassName");
+const FText FFindInBlueprintSearchTags::FiB_NodeGuid = LOCTEXT("NodeGuid", "NodeGuid");
+const FText FFindInBlueprintSearchTags::FiB_Tooltip = LOCTEXT("Tooltip", "Tooltip");
+const FText FFindInBlueprintSearchTags::FiB_DefaultValue = LOCTEXT("DefaultValue", "DefaultValue");
+const FText FFindInBlueprintSearchTags::FiB_Description = LOCTEXT("Description", "Description");
+const FText FFindInBlueprintSearchTags::FiB_Comment = LOCTEXT("Comment", "Comment");
+
+const FText FFindInBlueprintSearchTags::FiB_Pins = LOCTEXT("Pins", "Pins");
+const FText FFindInBlueprintSearchTags::FiB_PinCategory = LOCTEXT("PinCategory", "PinCategory");
+const FText FFindInBlueprintSearchTags::FiB_PinSubCategory = LOCTEXT("SubCategory", "SubCategory");
+const FText FFindInBlueprintSearchTags::FiB_ObjectClass = LOCTEXT("ObjectClass", "ObjectClass");
+const FText FFindInBlueprintSearchTags::FiB_IsArray = LOCTEXT("IsArray", "IsArray");
+const FText FFindInBlueprintSearchTags::FiB_IsReference = LOCTEXT("IsReference", "IsReference");
+const FText FFindInBlueprintSearchTags::FiB_Glyph = LOCTEXT("Glyph", "Glyph");
+const FText FFindInBlueprintSearchTags::FiB_GlyphColor = LOCTEXT("GlyphColor", "GlyphColor");
+
+namespace BlueprintSearchMetaDataHelpers
+{
+	typedef TJsonWriter<TCHAR, TCondensedJsonPrintPolicy<TCHAR>> SearchMetaDataWriter;
+
+	/** Helper function to set a Json string field using FText's, the FText's will be serialized and converted into hex strings for storage */
+	void SetJsonTextField(TSharedRef< FJsonObject > InOutJsonObject, FText InFieldName, FText InFieldValue)
+	{
+		InOutJsonObject->SetStringField(FFindInBlueprintSearchManager::ConvertFTextToHexString(InFieldName), FFindInBlueprintSearchManager::ConvertFTextToHexString(InFieldValue));
+	}
+
+	/**
+	 * Checks if Json value is searchable, eliminating data that not considered useful to search for
+	 *
+	 * @param InJsonValue		The Json value object to examine for searchability
+	 * @return					TRUE if the value should be searchable
+	 */
+	bool CheckIfJsonValueIsSearchable( TSharedPtr< FJsonValue > InJsonValue )
+	{
+		/** Check for interesting values
+		 *  booleans are not interesting, there are a lot of them
+		 *  strings are not interesting if they are empty
+		 *  numbers are not interesting if they are 0
+		 *  arrays are not interesting if they are empty or if they are filled with un-interesting types
+		 *  objects may not have interesting values when dug into
+		 */
+		bool bValidPropetyValue = true;
+		if(InJsonValue->Type == EJson::Boolean || InJsonValue->Type == EJson::None || InJsonValue->Type == EJson::Null)
+		{
+			bValidPropetyValue = false;
+		}
+		else if(InJsonValue->Type == EJson::String)
+		{
+			FString temp = InJsonValue->AsString();
+			if(InJsonValue->AsString().IsEmpty())
+			{
+				bValidPropetyValue = false;
+			}
+		}
+		else if(InJsonValue->Type == EJson::Number)
+		{
+			if(InJsonValue->AsNumber() == 0.0)
+			{
+				bValidPropetyValue = false;
+			}
+		}
+		else if(InJsonValue->Type == EJson::Array)
+		{
+			auto JsonArray = InJsonValue->AsArray();
+			if(JsonArray.Num() > 0)
+			{
+				// Some types are never interesting and the contents of the array should be ignored. Other types can be interesting, the contents of the array should be stored (even if
+				// the values may not be interesting, so that index values can be obtained)
+				if(JsonArray[0]->Type != EJson::Array && JsonArray[0]->Type != EJson::String && JsonArray[0]->Type != EJson::Number && JsonArray[0]->Type != EJson::Object)
+				{
+					bValidPropetyValue = false;
+				}
+			}
+		}
+		else if(InJsonValue->Type == EJson::Object)
+		{
+			// Start it out as not being valid, if we find any sub-items that are searchable, it will be marked to TRUE
+			bValidPropetyValue = false;
+
+			// Go through all value/key pairs to see if any of them are searchable, remove the ones that are not
+			auto JsonObject = InJsonValue->AsObject();
+			for(auto Iter = JsonObject->Values.CreateIterator(); Iter; ++Iter)
+			{
+				if(!CheckIfJsonValueIsSearchable(Iter->Value))
+				{
+					Iter.RemoveCurrent();
+				}
+				else
+				{
+					bValidPropetyValue = true;
+				}
+			}
+			
+		}
+
+		return bValidPropetyValue;
+	}
+
+	/**
+	 * Saves a graph pin type to a Json object
+	 *
+	 * @param InPinType				The pin type to save
+	 * @param InOutJsonObject		The Json object to save the pintype to
+	 */
+	void SavePinTypeToJson(const FEdGraphPinType& InPinType, TSharedRef< FJsonObject > InOutJsonObject)
+	{
+		// Only save strings that are not empty
+
+		if(!InPinType.PinCategory.IsEmpty())
+		{
+			SetJsonTextField(InOutJsonObject, FFindInBlueprintSearchTags::FiB_PinCategory, FText::FromString(InPinType.PinCategory));
+		}
+		
+		if(!InPinType.PinSubCategory.IsEmpty())
+		{
+			SetJsonTextField(InOutJsonObject, FFindInBlueprintSearchTags::FiB_PinSubCategory, FText::FromString(InPinType.PinSubCategory));
+		}
+
+		if(InPinType.PinSubCategoryObject.IsValid())
+		{
+			SetJsonTextField(InOutJsonObject, FFindInBlueprintSearchTags::FiB_ObjectClass, FText::FromString(InPinType.PinSubCategoryObject->GetName()));
+		}
+		InOutJsonObject->SetBoolField(FFindInBlueprintSearchManager::ConvertFTextToHexString(FFindInBlueprintSearchTags::FiB_IsArray), InPinType.bIsArray);
+		InOutJsonObject->SetBoolField(FFindInBlueprintSearchManager::ConvertFTextToHexString(FFindInBlueprintSearchTags::FiB_IsReference), InPinType.bIsReference);
+	}
+
+	/** Recursively digs into a JsonValue to correct all Strings into FTexts, including all fields keys. */
+	void FixUpJsonValueToStrings(TSharedRef< FJsonValue >& InJsonValue)
+	{
+		if(InJsonValue->Type == EJson::String)
+		{
+			InJsonValue = MakeShareable( new FJsonValueString(FFindInBlueprintSearchManager::ConvertFTextToHexString(FText::FromString(InJsonValue->AsString()))));
+		}
+		else if(InJsonValue->Type == EJson::Object)
+		{
+			TSharedRef< FJsonObject > NewJsonObject = MakeShareable(new FJsonObject);
+
+			TSharedPtr< FJsonObject > ParentObject = InJsonValue->AsObject();
+
+			if(ParentObject.IsValid())
+			{
+				// Go through all object sub-values and attempt to potentially fix up any strings
+				for( auto MapValues : ParentObject->Values )
+				{
+					TSharedRef< FJsonValue > JsonValueAsSharedRef = MapValues.Value.ToSharedRef();
+					FixUpJsonValueToStrings(JsonValueAsSharedRef);
+
+					NewJsonObject->SetField(FFindInBlueprintSearchManager::ConvertFTextToHexString(FText::FromString(MapValues.Key)), JsonValueAsSharedRef);
+				}
+			}
+
+
+			InJsonValue = MakeShareable( new FJsonValueObject(NewJsonObject));
+		}
+		else if( InJsonValue->Type == EJson::Array)
+		{
+			// Build a new Json array and replace the old one
+			TArray< TSharedPtr< FJsonValue > > JsonArrayData = InJsonValue->AsArray();
+			TArray< TSharedPtr< FJsonValue > > NewJsonArrayData;
+			for( TSharedPtr< FJsonValue >& JsonArrayItem : JsonArrayData)
+			{
+				TSharedRef< FJsonValue > JsonValueAsSharedRef = JsonArrayItem.ToSharedRef();
+				FixUpJsonValueToStrings(JsonValueAsSharedRef);
+
+				NewJsonArrayData.Add(JsonValueAsSharedRef);
+			}
+
+			InJsonValue = MakeShareable(new FJsonValueArray(NewJsonArrayData));
+		}
+	}
+
+	/**
+	 * Helper function to save a variable description to Json
+	 *
+	 * @param InWriter					Json writer object
+	 * @param InBlueprint				Blueprint the property for the variable can be found in, if any
+	 * @param InVariableDescription		The variable description being serialized to Json
+	 */
+	void SaveVariableDescriptionToJson(TSharedRef< SearchMetaDataWriter>& InWriter, const UBlueprint* InBlueprint, const FBPVariableDescription& InVariableDescription)
+	{
+		FEdGraphPinType VariableType = InVariableDescription.VarType;
+
+		TSharedRef< FJsonObject > JsonObject = MakeShareable(new FJsonObject);
+		SetJsonTextField(JsonObject, FFindInBlueprintSearchTags::FiB_Name, FText::FromString(InVariableDescription.FriendlyName));
+
+		// Find the variable's tooltip
+		FString TooltipResult;
+		FBlueprintEditorUtils::GetBlueprintVariableMetaData(InBlueprint, InVariableDescription.VarName, FBlueprintMetadata::MD_Tooltip, TooltipResult);
+		SetJsonTextField(JsonObject, FFindInBlueprintSearchTags::FiB_Tooltip, FText::FromString(TooltipResult));
+
+		// Save the variable's pin type
+		BlueprintSearchMetaDataHelpers::SavePinTypeToJson(VariableType, JsonObject);
+
+		// Find the UProperty and convert it into a Json value.
+		UProperty* VariableProperty = FindField<UProperty>(InBlueprint->GeneratedClass, InVariableDescription.VarName);
+		if(VariableProperty)
+		{
+			const uint8* PropData = VariableProperty->ContainerPtrToValuePtr<uint8>(InBlueprint->GeneratedClass->GetDefaultObject());
+			auto JsonValue = FJsonObjectConverter::UPropertyToJsonValue(VariableProperty, PropData, 0, 0);
+
+			
+			// Only use the value if it is searchable
+			if(BlueprintSearchMetaDataHelpers::CheckIfJsonValueIsSearchable(JsonValue))
+			{
+				TSharedRef< FJsonValue > JsonValueAsSharedRef = JsonValue.ToSharedRef();
+				FixUpJsonValueToStrings(JsonValueAsSharedRef);
+				//FixUpJsonValueToStrings(JsonValue.ToSharedRef());
+				JsonObject->SetField(FFindInBlueprintSearchManager::ConvertFTextToHexString(FFindInBlueprintSearchTags::FiB_DefaultValue), JsonValueAsSharedRef);
+				//JsonObject->SetStringField(FFindInBlueprintSearchManager::ConvertFTextToHexString(FFindInBlueprintSearchTags::FiB_DefaultValue), FFindInBlueprintSearchManager::ConvertFTextToHexString(FText::FromString(JsonValue->AsString())));
+			}
+		}
+
+		// Serialize the object out using the writer, so that it turns into raw text in the metadata string being made
+		FJsonSerializer::Serialize(JsonObject, InWriter);
+	}
+
+	/**
+	 * Gathers all nodes from a specified graph and serializes their searchable data to Json
+	 *
+	 * @param InWriter		The Json writer to use for serialization
+	 * @param InGraph		The graph to search through
+	 */
+	void GatherNodesFromGraph(TSharedRef< SearchMetaDataWriter>& InWriter, const UEdGraph* InGraph)
+	{
+		// Collect all macro graphs
+		InWriter->WriteArrayStart(FFindInBlueprintSearchManager::ConvertFTextToHexString(FFindInBlueprintSearchTags::FiB_Nodes));
+		{
+			for(auto* Node : InGraph->Nodes)
+			{
+				if(Node)
+				{
+					// Make sure we don't collect search data for nodes that are going away soon
+					if( Node->GetOuter()->IsPendingKill() )
+					{
+						continue;
+					}
+
+					TSharedRef< FJsonObject > JsonNodeObject = MakeShareable(new FJsonObject);
+
+					// Retrieve the search metadata from the node, some node types may have extra metadata to be searchable.
+					TArray<struct FSearchTagDataPair> Tags;
+					Node->AddSearchMetaDataInfo(Tags);
+
+					// Go through the node metadata tags and put them into the Json object.
+					for(const FSearchTagDataPair& SearchData : Tags)
+					{
+						FString Key = FFindInBlueprintSearchManager::ConvertFTextToHexString(SearchData.Key);
+						FString Value = FFindInBlueprintSearchManager::ConvertFTextToHexString(SearchData.Value);
+						JsonNodeObject->SetStringField(Key, Value);
+					}
+
+					// Find all the pins and extract their metadata
+					TArray<TSharedPtr< FJsonValue > > JsonPinList;
+					for(UEdGraphPin* Pin : Node->Pins)
+					{
+						TSharedRef< FJsonObject > JsonPinObject = MakeShareable(new FJsonObject);
+						SetJsonTextField(JsonPinObject, FFindInBlueprintSearchTags::FiB_Name, FText::FromString(Pin->GetSchema()->GetPinDisplayName(Pin)));
+						SetJsonTextField(JsonPinObject, FFindInBlueprintSearchTags::FiB_DefaultValue, FText::FromString(Pin->GetDefaultAsString()));
+						BlueprintSearchMetaDataHelpers::SavePinTypeToJson(Pin->PinType, JsonPinObject);
+
+						// Add the Json pin data to the pin list
+						JsonPinList.Add(MakeShareable(new FJsonValueObject(JsonPinObject)));
+					}
+					JsonNodeObject->SetArrayField(FFindInBlueprintSearchManager::ConvertFTextToHexString(FFindInBlueprintSearchTags::FiB_Pins), JsonPinList);
+
+					// Serialize the object out using the writer, so that it turns into raw text in the metadata string being made
+					FJsonSerializer::Serialize(JsonNodeObject, InWriter);
+				}
+			}
+		}
+		InWriter->WriteArrayEnd();
+	}
+
+	/** 
+	 * Gathers all graph's search data (and subojects) and serializes them to Json
+	 *
+	 * @param InWriter			The Json writer to use for serialization
+	 * @param InGraphArray		All the graphs to process
+	 * @param InTitle			The array title to place these graphs into
+	 * @param InOutSubGraphs	All the subgraphs that need to be processed later
+	 */
+	void GatherGraphSearchData(TSharedRef< SearchMetaDataWriter>& InWriter, const UBlueprint* InBlueprint, const TArray< UEdGraph* >& InGraphArray, FText InTitle, TArray< UEdGraph* >* InOutSubGraphs)
+	{
+		if(InGraphArray.Num() > 0)
+		{
+			// Collect all graphs
+			InWriter->WriteArrayStart(FFindInBlueprintSearchManager::ConvertFTextToHexString(InTitle));
+			{
+				for(const UEdGraph* Graph : InGraphArray)
+				{
+					InWriter->WriteObjectStart();
+
+					FGraphDisplayInfo DisplayInfo;
+					Graph->GetSchema()->GetGraphDisplayInformation(*Graph, DisplayInfo);
+					InWriter->WriteValue(FFindInBlueprintSearchManager::ConvertFTextToHexString(FFindInBlueprintSearchTags::FiB_Name), FFindInBlueprintSearchManager::ConvertFTextToHexString(DisplayInfo.PlainName));
+
+					FString GraphDescription = FBlueprintEditorUtils::GetGraphDescription(Graph).ToString();
+					if(!GraphDescription.IsEmpty())
+					{
+						InWriter->WriteValue(FFindInBlueprintSearchManager::ConvertFTextToHexString(FFindInBlueprintSearchTags::FiB_Description), FFindInBlueprintSearchManager::ConvertFTextToHexString(FText::FromString(GraphDescription)));
+					}
+
+					// All nodes will appear as children to the graph in search results
+					GatherNodesFromGraph(InWriter, Graph);
+
+					// Collect local variables
+					TArray<UK2Node_FunctionEntry*> FunctionEntryNodes;
+					Graph->GetNodesOfClass<UK2Node_FunctionEntry>(FunctionEntryNodes);
+
+					InWriter->WriteArrayStart(FFindInBlueprintSearchManager::ConvertFTextToHexString(FFindInBlueprintSearchTags::FiB_Properties));
+					{
+						// Search in all FunctionEntry nodes for their local variables and add them to the list
+						FString ActionCategory;
+						for (UK2Node_FunctionEntry* const FunctionEntry : FunctionEntryNodes)
+						{
+							for( const FBPVariableDescription& Variable : FunctionEntry->LocalVariables )
+							{
+								SaveVariableDescriptionToJson(InWriter, InBlueprint, Variable);
+							}
+						}
+					}
+					InWriter->WriteArrayEnd(); // Properties
+
+					InWriter->WriteObjectEnd();
+
+					// Only if asked to do it
+					if(InOutSubGraphs)
+					{
+						Graph->GetAllChildrenGraphs(*InOutSubGraphs);
+					}
+				}
+			}
+			InWriter->WriteArrayEnd();
+		}
+	}
+}
+
+class FCacheAllBlueprintsTickableObject
+	: public FTickableEditorObject
+{
+
+public:
+
+	FCacheAllBlueprintsTickableObject(TArray<FName>& InUncachedBlueprints)
+		: TickCacheIndex(0)
+		, UncachedBlueprints(InUncachedBlueprints)
+	{
+		// Start the Blueprint indexing 'progress' notification
+		FNotificationInfo Info( LOCTEXT("BlueprintIndexMessage", "Indexing Blueprints...") );
+		Info.bFireAndForget = false;
+		Info.ButtonDetails.Add(FNotificationButtonInfo(
+			LOCTEXT("BlueprintIndexCancel","Cancel"),
+			LOCTEXT("BlueprintIndexCancelToolTip","Cancels indexing Blueprints."), FSimpleDelegate::CreateRaw(this, &FCacheAllBlueprintsTickableObject::OnCancelCaching)));
+
+		ProgressNotification = FSlateNotificationManager::Get().AddNotification(Info);
+		if (ProgressNotification.IsValid())
+		{
+			ProgressNotification.Pin()->SetCompletionState(SNotificationItem::CS_Pending);
+		}
+	}
+
+	~FCacheAllBlueprintsTickableObject()
+	{
+		
+	}
+
+	/** Returns the current cache index of the object */
+	int32 GetCurrentCacheIndex() const
+	{
+		return TickCacheIndex + 1;
+	}
+
+	/** Returns the name of the current Blueprint being cached */
+	FString GetCurrentCacheBlueprintName() const
+	{
+		if(UncachedBlueprints.Num() && TickCacheIndex >= 0)
+		{
+			return UncachedBlueprints[TickCacheIndex].ToString();
+		}
+		return FString();
+	}
+
+	/** Returns the progress as a percent */
+	float GetCacheProgress() const
+	{
+		return (float)TickCacheIndex / (float)UncachedBlueprints.Num();
+	}
+
+	/** Cancels caching and destroys this object */
+	void OnCancelCaching()
+	{
+		ProgressNotification.Pin()->SetText(LOCTEXT("BlueprintIndexCancelled", "Cancelled Indexing Blueprints!"));
+
+		ProgressNotification.Pin()->SetCompletionState(SNotificationItem::CS_Fail);
+		ProgressNotification.Pin()->ExpireAndFadeout();
+
+		FFindInBlueprintSearchManager::Get().FinishedCachingBlueprints(TickCacheIndex);
+	}
+
+protected:
+	/** FTickableEditorObject interface */
+	virtual void Tick(float DeltaTime) override
+	{
+		if(GWarn->ReceivedUserCancel())
+		{
+			FFindInBlueprintSearchManager::Get().FinishedCachingBlueprints(TickCacheIndex);
+		}
+		else
+		{
+			FAssetRegistryModule* AssetRegistryModule = &FModuleManager::LoadModuleChecked<FAssetRegistryModule>(TEXT("AssetRegistry"));
+
+			FAssetData AssetData = AssetRegistryModule->Get().GetAssetByObjectPath(UncachedBlueprints[TickCacheIndex]);
+
+			bool bIsLoaded = AssetData.IsAssetLoaded();
+
+			UObject* Asset = AssetData.GetAsset();
+
+			// If the Blueprint was not just loaded, cache it again, otherwise it was cached on load.
+			if(!bIsLoaded)
+			{
+				FFindInBlueprintSearchManager::Get().AddOrUpdateBlueprintSearchMetadata(Cast<UBlueprint>(Asset), true);
+			}
+
+			++TickCacheIndex;
+
+			// Check if done caching Blueprints
+			if(TickCacheIndex == UncachedBlueprints.Num())
+			{
+				ProgressNotification.Pin()->SetCompletionState(SNotificationItem::CS_Success);
+				ProgressNotification.Pin()->ExpireAndFadeout();
+
+				ProgressNotification.Pin()->SetText(LOCTEXT("BlueprintIndexComplete", "Finished indexing Blueprints!"));
+
+				FFindInBlueprintSearchManager::Get().FinishedCachingBlueprints(TickCacheIndex);
+			}
+			else
+			{
+				FFormatNamedArguments Args;
+				Args.Add(TEXT("Percent"), FText::AsPercent(GetCacheProgress()));
+				ProgressNotification.Pin()->SetText(FText::Format(LOCTEXT("BlueprintIndexProgress", "Indexing Blueprints... ({Percent})"), Args));
+			}
+		}
+	}
+
+	virtual bool IsTickable() const override
+	{
+		return true;
+	}
+
+	virtual TStatId GetStatId() const override
+	{
+		RETURN_QUICK_DECLARE_CYCLE_STAT(FCacheAllBlueprintsTickableObject, STATGROUP_Tickables);
+	}
+
+private:
+
+	/** The current index, increases at a rate of once per tick */
+	int32 TickCacheIndex;
+
+	/** The list of uncached Blueprints that are in the process of being cached */
+	TArray<FName> UncachedBlueprints;
+
+	/** Notification that appears and details progress */
+	TWeakPtr<SNotificationItem> ProgressNotification;
+};
+
+FFindInBlueprintSearchManager& FFindInBlueprintSearchManager::Get()
+{
+	if (Instance == NULL)
+	{
+		Instance = new FFindInBlueprintSearchManager();
+		Instance->Initialize();
+	}
+
+	return *Instance;
+}
+
+FFindInBlueprintSearchManager::FFindInBlueprintSearchManager()
+	: bIsPausing(false)
+	, CachingObject(NULL)
+{
+}
+
+FFindInBlueprintSearchManager::~FFindInBlueprintSearchManager()
+{
+	FinishedCachingBlueprints(0);
+
+	AssetRegistryModule->Get().OnAssetAdded().RemoveAll(this);
+	AssetRegistryModule->Get().OnAssetRemoved().RemoveAll(this);
+	AssetRegistryModule->Get().OnAssetRenamed().RemoveAll(this);
+	FCoreDelegates::PreGarbageCollect.RemoveAll(this);
+	FCoreDelegates::PostGarbageCollect.RemoveAll(this);
+	FCoreDelegates::OnAssetLoaded.RemoveAll(this);
+}
+
+void FFindInBlueprintSearchManager::Initialize()
+{
+	AssetRegistryModule = &FModuleManager::LoadModuleChecked<FAssetRegistryModule>(TEXT("AssetRegistry"));
+	AssetRegistryModule->Get().OnAssetAdded().AddRaw(this, &FFindInBlueprintSearchManager::OnAssetAdded);
+	AssetRegistryModule->Get().OnAssetRemoved().AddRaw(this, &FFindInBlueprintSearchManager::OnAssetRemoved);
+	AssetRegistryModule->Get().OnAssetRenamed().AddRaw(this, &FFindInBlueprintSearchManager::OnAssetRenamed);
+	FCoreDelegates::PreGarbageCollect.AddRaw(this, &FFindInBlueprintSearchManager::PauseFindInBlueprintSearch);
+	FCoreDelegates::PostGarbageCollect.AddRaw(this, &FFindInBlueprintSearchManager::UnpauseFindInBlueprintSearch);
+	FCoreDelegates::OnAssetLoaded.AddRaw(this, &FFindInBlueprintSearchManager::OnAssetLoaded);
+
+	// Do an immediate load of the cache to catch any Blueprints that were discovered by the asset registry before we initialized.
+	BuildCache();
+}
+
+void FFindInBlueprintSearchManager::OnAssetAdded(const FAssetData& InAssetData)
+{
+	if(InAssetData.GetClass()->IsChildOf(UBlueprint::StaticClass()))
+	{
+		FString BlueprintPackagePath = FPaths::GetPath(InAssetData.ObjectPath.ToString()) / FPaths::GetBaseFilename(InAssetData.ObjectPath.ToString());
+
+		// Confirm that the Blueprint has not been added already, this can occur during duplication of Blueprints.
+		int32* IndexPtr = SearchMap.Find(BlueprintPackagePath);
+		if(!IndexPtr)
+		{
+			FString SearchGuid;
+			const FString* SearchGuidPtr = InAssetData.TagsAndValues.Find("SearchGuid");
+			if(SearchGuidPtr)
+			{
+				SearchGuid = *SearchGuidPtr;
+			}
+			else
+			{
+				SearchGuid = FMD5::HashAnsiString(*InAssetData.ObjectPath.ToString());
+
+				if(InAssetData.IsAssetLoaded())
+				{
+					AddOrUpdateBlueprintSearchMetadata(Cast<UBlueprint>(InAssetData.GetAsset()));
+				}
+			}
+
+			FDerivedDataCacheInterface& DDC = GetDerivedDataCacheRef();
+			FSearchData NewSearchData;
+			if(InAssetData.IsAssetLoaded())
+			{
+				NewSearchData.Blueprint = Cast<UBlueprint>(InAssetData.GetAsset());
+			}
+			NewSearchData.BlueprintPath = InAssetData.ObjectPath.ToString();
+			NewSearchData.DDCRetrievalID = DDC.GetAsynchronous(*SearchGuid);
+			NewSearchData.bMarkedForDeletion = false;
+
+			int32 ArrayIndex = SearchArray.Add(NewSearchData);
+
+			// Add the asset file path to the map along with the index into the array
+			SearchMap.Add(BlueprintPackagePath, ArrayIndex);
+		}
+	}
+}
+
+void FFindInBlueprintSearchManager::OnAssetRemoved(const class FAssetData& InAssetData)
+{
+	if(InAssetData.IsAssetLoaded())
+	{
+		int32* SearchIdx = SearchMap.Find(InAssetData.GetAsset()->GetOutermost()->GetPathName());
+		if(SearchIdx)
+		{
+			SearchArray[*SearchIdx].bMarkedForDeletion = true;
+		}
+	}
+}
+
+void FFindInBlueprintSearchManager::OnAssetRenamed(const class FAssetData& InAssetData, const FString& InOldName)
+{
+	if(InAssetData.IsAssetLoaded())
+	{
+		int32* SearchIdx = SearchMap.Find(FPaths::GetPath(InOldName) / FPaths::GetBaseFilename(InOldName));
+		if(SearchIdx)
+		{
+			SearchArray[*SearchIdx].bMarkedForDeletion = true;
+		}
+	}
+}
+
+void FFindInBlueprintSearchManager::OnAssetLoaded(UObject* InAsset)
+{
+	if(UBlueprint* Blueprint = Cast<UBlueprint>(InAsset))
+	{
+		// Find and update the item in the search array. Searches may currently be active, this will do no harm to them
+
+		FString BlueprintPackagePath = InAsset->GetPathName();
+
+		// Confirm that the Blueprint has not been added already, this can occur during duplication of Blueprints.
+		int32* IndexPtr = SearchMap.Find(BlueprintPackagePath);
+
+		// The asset registry might not have informed us of this asset yet.
+		if(IndexPtr)
+		{
+			// That index should never have a Blueprint already
+			check(!SearchArray[*IndexPtr].Blueprint);
+
+			SearchArray[*IndexPtr].Blueprint = Blueprint;
+		}
+	}
+}
+
+FString FFindInBlueprintSearchManager::GatherBlueprintSearchMetadata(const UBlueprint* Blueprint)
+{
+	FString SearchMetaData;
+
+	// The search registry tags for a Blueprint are all in Json
+	TSharedRef< BlueprintSearchMetaDataHelpers::SearchMetaDataWriter > Writer = TJsonWriterFactory<TCHAR, TCondensedJsonPrintPolicy<TCHAR>>::Create( &SearchMetaData );
+
+	TMap<FString, TMap<FString,int>> AllPaths;
+	Writer->WriteObjectStart();
+
+	// Pull out all nodes and their pins to be searchable
+	TArray<UEdGraphNode*> Nodes;
+	FBlueprintEditorUtils::GetAllNodesOfClass(Blueprint, Nodes );
+
+	// Only pull properties if the Blueprint has been compiled
+	if(Blueprint->SkeletonGeneratedClass)
+	{
+		Writer->WriteArrayStart(FFindInBlueprintSearchManager::ConvertFTextToHexString(FFindInBlueprintSearchTags::FiB_Properties));
+		{
+			for( const FBPVariableDescription& Variable : Blueprint->NewVariables )
+			{
+				BlueprintSearchMetaDataHelpers::SaveVariableDescriptionToJson(Writer, Blueprint, Variable);
+			}
+		}
+		Writer->WriteArrayEnd(); // Properties
+	}
+
+	// Gather all graph searchable data
+	TArray< UEdGraph* > SubGraphs;
+	BlueprintSearchMetaDataHelpers::GatherGraphSearchData(Writer, Blueprint, Blueprint->UbergraphPages, FFindInBlueprintSearchTags::FiB_UberGraphs, &SubGraphs);
+	BlueprintSearchMetaDataHelpers::GatherGraphSearchData(Writer, Blueprint, Blueprint->FunctionGraphs, FFindInBlueprintSearchTags::FiB_Functions, &SubGraphs);
+	BlueprintSearchMetaDataHelpers::GatherGraphSearchData(Writer, Blueprint, Blueprint->MacroGraphs, FFindInBlueprintSearchTags::FiB_Macros, &SubGraphs);
+	// Sub graphs are processed separately so that they do not become children in the TreeView, cluttering things up if the tree is deep
+	BlueprintSearchMetaDataHelpers::GatherGraphSearchData(Writer, Blueprint, SubGraphs, FFindInBlueprintSearchTags::FiB_SubGraphs, NULL);
+
+	Writer->WriteObjectEnd();
+	Writer->Close();
+
+	return SearchMetaData;
+}
+
+void FFindInBlueprintSearchManager::AddOrUpdateBlueprintSearchMetadata(UBlueprint* InBlueprint, bool bInForceReCache/* = false*/)
+{
+	check(InBlueprint);
+
+	// Allow only one thread to query the DDC at a single time, this is important in the case that two threads are waiting for the same DDC information, once the information is pulled, it is deleted and is no longer accessible.
+	FScopeLock ScopeLock(&SafeModifyCacheCriticalSection);
+
+	FString BlueprintPackagePath = InBlueprint->GetOutermost()->GetPathName();
+
+	int32* IndexPtr = SearchMap.Find(BlueprintPackagePath);
+	int32 Index = 0;
+	if(!IndexPtr)
+	{
+		FSearchData SearchData;
+		SearchData.Blueprint = InBlueprint;
+		SearchData.BlueprintPath = InBlueprint->GetPathName();
+		Index = SearchArray.Add(SearchData);
+
+		SearchMap.Add(BlueprintPackagePath, Index);
+	}
+	else
+	{
+		Index = *IndexPtr;
+	}
+
+	FDerivedDataCacheInterface& DDC = GetDerivedDataCacheRef();
+
+	FAssetData AssetData = AssetRegistryModule->Get().GetAssetByObjectPath(*InBlueprint->GetPathName());
+
+	FString* SearchGuidPtr = AssetData.TagsAndValues.Find("SearchGuid");
+	FString SearchGuid;
+
+	// If there is no search Guid stored, generate one using the asset's path. This transitions Blueprints from the old search system to the new one without resaving
+	if(!SearchGuidPtr)
+	{
+		SearchGuid = FMD5::HashAnsiString(*AssetData.ObjectPath.ToString());
+	}
+	else
+	{
+		SearchGuid = *SearchGuidPtr;
+	}
+
+	// If the Blueprint has been cached, don't cache it again, it's wasted time
+	if(!DDC.CachedDataProbablyExists(*SearchGuid) || bInForceReCache)
+	{
+		// The data is valid, but we may still need to wait on the DDC.
+		if(SearchArray[Index].DDCRetrievalID != INDEX_NONE)
+		{
+			// Allow only one thread to query the DDC at a single time, this is important in the case that two threads are waiting for the same DDC information, once the information is pulled, it is deleted and is no longer accessible.
+			FScopeLock ScopeLock(&SafeModifyCacheCriticalSection);
+
+			// If there was a DDC retrieval request, we have to see it through
+			if(SearchArray[Index].DDCRetrievalID != INDEX_NONE)
+			{
+				// Wait for completion of the DDC retrieval, we are on a separate thread so this will not stall the editor
+				DDC.WaitAsynchronousCompletion(SearchArray[Index].DDCRetrievalID);
+
+				// Fetch the data, we don't need it but the DDC doesn't need to manage the query anymore
+				TArray<uint8> DerivedData;
+				DDC.GetAsynchronousResults(SearchArray[Index].DDCRetrievalID, DerivedData);
+
+				SearchArray[Index].DDCRetrievalID = INDEX_NONE;
+			}
+		}
+
+		// Build the search data
+		SearchArray[Index].BlueprintPath = InBlueprint->GetPathName();
+		SearchArray[Index].Value = GatherBlueprintSearchMetadata(InBlueprint);
+		SearchArray[Index].bMarkedForDeletion = false;
+
+		// Save the search data to the DDC
+		TArray<uint8> DerivedData;
+		FMemoryWriter Ar(DerivedData, /*bIsPersistent=*/ true);
+		Ar << SearchArray[Index].BlueprintPath;
+		Ar << SearchArray[Index].Value;
+
+		// Insert the search data into the DDC
+		check(DerivedData.Num());
+		DDC.Put(*SearchGuid, DerivedData);
+
+		Ar.Close();
+	}
+}
+
+void FFindInBlueprintSearchManager::BeginSearchQuery(const FStreamSearch* InSearchOriginator)
+{
+	// Cannot begin a search thread while saving
+	FScopeLock ScopeLock(&PauseThreadsCriticalSection);
+	FScopeLock ScopeLock2(&SafeQueryModifyCriticalSection);
+
+	ActiveSearchCounter.Increment();
+	ActiveSearchQueries.FindOrAdd(InSearchOriginator) = 0;
+}
+
+bool FFindInBlueprintSearchManager::ContinueSearchQuery(const FStreamSearch* InSearchOriginator, FSearchData& OutSearchData)
+{
+	// Check if the thread has been told to pause, this occurs for the Garbage Collector and for saving to disk
+	if(bIsPausing == true)
+	{
+		// Pause all searching, the GC is running and we will also be saving the database
+		ActiveSearchCounter.Decrement();
+		FScopeLock ScopeLock(&PauseThreadsCriticalSection);
+		ActiveSearchCounter.Increment();
+	}
+
+	int32* SearchIdxPtr = NULL;
+
+	{
+		// Must lock this behind a critical section to ensure that no other thread is accessing it at the same time
+		FScopeLock ScopeLock(&SafeQueryModifyCriticalSection);
+		SearchIdxPtr = ActiveSearchQueries.Find(InSearchOriginator);
+	}
+
+	if(SearchIdxPtr)
+	{
+		int32& SearchIdx = *SearchIdxPtr;
+		while(SearchIdx < SearchArray.Num())
+		{
+			// If the Blueprint is not marked for deletion, and the asset is valid, we will check to see if we want to refresh the searchable data.
+			if( SearchArray[SearchIdx].bMarkedForDeletion || (SearchArray[SearchIdx].Blueprint && SearchArray[SearchIdx].Blueprint->HasAllFlags(RF_PendingKill)) )
+			{
+				// Mark it for deletion, it will be removed on next save
+				SearchArray[SearchIdx].bMarkedForDeletion = true;
+			}
+			else if(SearchArray[SearchIdx].Blueprint)
+			{
+				// Don't rebuild the search data if the Blueprint isn't dirty (the package) or possibly dirty (compilation).
+				if(Cast<UPackage>(SearchArray[SearchIdx].Blueprint->GetOutermost())->IsDirty() || SearchArray[SearchIdx].Blueprint->IsPossiblyDirty())
+				{
+					// Update the Blueprint's search metadata so our search is up-to-date.
+					SearchArray[SearchIdx].Blueprint->SearchGuid = FGuid::NewGuid();
+					AddOrUpdateBlueprintSearchMetadata(SearchArray[SearchIdx].Blueprint);
+				}
+				else
+				{
+					RetrieveDDCData(SearchArray[SearchIdx], *SearchArray[SearchIdx].Blueprint->GetPathName());
+				}
+
+ 				OutSearchData = SearchArray[SearchIdx++];
+				return true;
+			}
+			else
+			{
+				RetrieveDDCData(SearchArray[SearchIdx], *SearchArray[SearchIdx].BlueprintPath);
+				OutSearchData = SearchArray[SearchIdx++];
+				return true;
+			}
+
+			++SearchIdx;
+		}
+	}
+
+	{
+		// Must lock this behind a critical section to ensure that no other thread is accessing it at the same time
+		FScopeLock ScopeLock(&SafeQueryModifyCriticalSection);
+		ActiveSearchQueries.Remove(InSearchOriginator);
+	}
+	ActiveSearchCounter.Decrement();
+
+	return false;
+}
+
+void FFindInBlueprintSearchManager::EnsureSearchQueryEnds(const class FStreamSearch* InSearchOriginator)
+{
+	// Must lock this behind a critical section to ensure that no other thread is accessing it at the same time
+	FScopeLock ScopeLock(&SafeQueryModifyCriticalSection);
+	int32* SearchIdxPtr = ActiveSearchQueries.Find(InSearchOriginator);
+
+	// If the search thread is still considered active, remove it
+	if(SearchIdxPtr)
+	{
+		ActiveSearchQueries.Remove(InSearchOriginator);
+		ActiveSearchCounter.Decrement();
+	}
+}
+
+float FFindInBlueprintSearchManager::GetPercentComplete(const FStreamSearch* InSearchOriginator) const
+{
+	FScopeLock ScopeLock(&SafeQueryModifyCriticalSection);
+	const int32* SearchIdxPtr = ActiveSearchQueries.Find(InSearchOriginator);
+
+	float ReturnPercent = 0.0f;
+
+	if(SearchIdxPtr)
+	{
+		ReturnPercent = (float)*SearchIdxPtr / (float)SearchArray.Num();
+	}
+
+	return ReturnPercent;
+}
+
+void FFindInBlueprintSearchManager::RetrieveDDCData(FSearchData& InOutSearchData, FName InBlueprintPath)
+{
+	FDerivedDataCacheInterface& DDC = GetDerivedDataCacheRef();
+
+	// The data is valid, but we may still need to wait on the DDC.
+	if(InOutSearchData.DDCRetrievalID != INDEX_NONE)
+	{
+		// Allow only one thread to query the DDC at a single time, this is important in the case that two threads are waiting for the same DDC information, once the information is pulled, it is deleted and is no longer accessible.
+		FScopeLock ScopeLock(&SafeModifyCacheCriticalSection);
+
+		// Check to see if the data was pulled by another thread while we waited
+		if(InOutSearchData.DDCRetrievalID != INDEX_NONE)
+		{
+			// Wait for completion of the DDC retrieval, we are on a separate thread so this will not stall the editor
+			DDC.WaitAsynchronousCompletion(InOutSearchData.DDCRetrievalID);
+
+			TArray<uint8> DerivedData;
+
+			if( !DDC.GetAsynchronousResults(InOutSearchData.DDCRetrievalID, DerivedData) ) 
+			{
+				// Since the asset cannot be loaded at this point, cache the Blueprint, we'll index it later.
+				UncachedBlueprints.AddUnique(InBlueprintPath);
+			}
+			else
+			{
+				check(DerivedData.Num());
+
+				FMemoryReader DDCAr(DerivedData, /*bIsPersistent=*/ true);
+				DDCAr << InOutSearchData.BlueprintPath;
+				DDCAr << InOutSearchData.Value;
+				DDCAr.Close();
+			}
+
+			InOutSearchData.DDCRetrievalID = INDEX_NONE;
+		}
+	}
+}
+
+FString FFindInBlueprintSearchManager::QuerySingleBlueprint(UBlueprint* InBlueprint)
+{
+	// Update the Blueprint, make sure it is fully up-to-date
+	AddOrUpdateBlueprintSearchMetadata(InBlueprint, true);
+
+	FString Key = InBlueprint->GetOutermost()->GetPathName();
+
+	int32* ArrayIdx = SearchMap.Find(Key);
+	// This should always be true since we make sure to refresh the search data for this Blueprint when doing the search
+	check(ArrayIdx && *ArrayIdx < SearchArray.Num());
+
+	RetrieveDDCData(SearchArray[*ArrayIdx], *InBlueprint->GetPathName());
+
+	return SearchArray[*ArrayIdx].Value;
+}
+
+void FFindInBlueprintSearchManager::PauseFindInBlueprintSearch()
+{
+	// Lock the critical section and flag that threads need to pause, they will pause when they can
+	PauseThreadsCriticalSection.Lock();
+	bIsPausing = true;
+
+	// It is UNSAFE to lock any other critical section here, threads need them to finish a cycle of searching. Next cycle they will pause
+
+	// Wait until all threads have come to a stop, it won't take long
+	while(ActiveSearchCounter.GetValue() > 0)
+	{
+		FPlatformProcess::Sleep(0.1f);
+	}
+}
+
+void FFindInBlueprintSearchManager::UnpauseFindInBlueprintSearch()
+{
+	// Before unpausing, we clean the cache of any excess data to keep it from bloating in size
+	CleanCache();
+	bIsPausing = false;
+
+	// Release the threads to continue searching.
+	PauseThreadsCriticalSection.Unlock();
+}
+
+void FFindInBlueprintSearchManager::CleanCache()
+{
+	// *NOTE* SaveCache is a thread safe operation by design, all searching threads are paused during the operation so there is no critical section locking
+
+	// We need to cache where the active queries are so that we can put them back in a safe and expected position
+	TMap< const FStreamSearch*, FString > CacheQueries;
+	for( auto It = ActiveSearchQueries.CreateIterator() ; It ; ++It )
+	{
+	 	const FStreamSearch* ActiveSearch = It.Key();
+	 	check(ActiveSearch);
+	 	{
+			FSearchData searchData;
+	 		ContinueSearchQuery(ActiveSearch, searchData);
+
+			// We will be looking the path up in the map, which uses the package path instead of the Blueprint, so rebuild the package path from the Blueprint path
+			FString CachePath = FPaths::GetPath(searchData.BlueprintPath) / FPaths::GetBaseFilename(searchData.BlueprintPath);
+	 		CacheQueries.Add(ActiveSearch, CachePath);
+	 	}
+	}
+
+	FDerivedDataCacheInterface& DDC = GetDerivedDataCacheRef();
+
+	TMap<FString, int> NewSearchMap;
+	TArray<FSearchData> NewSearchArray;
+
+	for(auto& SearchValuePair : SearchMap)
+	{
+		// Here it builds the new map/array, clean of deleted content.
+		if( SearchArray[SearchValuePair.Value].bMarkedForDeletion || (SearchArray[SearchValuePair.Value].Blueprint && SearchArray[SearchValuePair.Value].Blueprint->HasAllFlags(RF_PendingKill)) )
+		{
+			// Item should not be kept, is is being destroyed and removed
+			int x = 5;
+			x = x;
+		}
+		else
+		{
+			// Build the new map/array
+			NewSearchMap.Add(SearchValuePair.Key, NewSearchArray.Add(SearchArray[SearchValuePair.Value]) );
+		}
+	}
+
+	SearchMap = NewSearchMap;
+	SearchArray = NewSearchArray;
+
+	// After the search, we have to place the active search queries where they belong
+	for( auto& CacheQuery : CacheQueries )
+	{
+	 	int32 NewMappedIndex = 0;
+	 	// If the CachePath has a length, it means it is valid, otherwise we are at the end and there are no more search results, leave the query there so it can handle shutdown on it's own
+	 	if(CacheQuery.Value.Len() > 0)
+	 	{
+	 		int32* NewMappedIndexPtr = SearchMap.Find(CacheQuery.Value);
+	 		check(NewMappedIndexPtr);
+	 
+	 		NewMappedIndex = *NewMappedIndexPtr;
+	 	}
+	 	else
+	 	{
+	 		NewMappedIndex = SearchArray.Num();
+	 	}
+	 
+		// Update the active search to the new index of where it is at in the search
+	 	*(ActiveSearchQueries.Find(CacheQuery.Key)) = NewMappedIndex;
+	}
+}
+
+void FFindInBlueprintSearchManager::BuildCache()
+{
+	AssetRegistryModule = &FModuleManager::LoadModuleChecked<FAssetRegistryModule>(TEXT("AssetRegistry"));
+
+	TArray< FAssetData > BlueprintAssets;
+	AssetRegistryModule->Get().GetAssetsByClass(UBlueprint::StaticClass()->GetFName(), BlueprintAssets, true);
+	
+	FDerivedDataCacheInterface& DDC = GetDerivedDataCacheRef();
+	for( FAssetData& Asset : BlueprintAssets )
+	{
+		if(FString* SearchGuid = Asset.TagsAndValues.Find("SearchGuid"))
+		{
+			OnAssetAdded(Asset);
+		}
+	}
+}
+
+void FFindInBlueprintSearchManager::CacheAllUncachedBlueprints(TWeakPtr< SFindInBlueprints > InSourceWidget)
+{
+	// Do not start another caching process if one is in progress
+	if(!IsCacheInProgress())
+	{
+		// Multiple threads can be adding to this at the same time
+		FScopeLock ScopeLock(&SafeModifyCacheCriticalSection);
+
+		CachingObject = new FCacheAllBlueprintsTickableObject(UncachedBlueprints);
+	}
+
+	// Add the widget to the callback list, to be called when caching is complete
+	OnCachingCompleteCallbackWidgets.Add(InSourceWidget);
+}
+
+void FFindInBlueprintSearchManager::CancelCacheAll()
+{
+	if(IsCacheInProgress())
+	{
+		CachingObject->OnCancelCaching();
+	}
+}
+
+int32 FFindInBlueprintSearchManager::GetCurrentCacheIndex() const
+{
+	int32 CachingIndex = 0;
+	if(CachingObject)
+	{
+		CachingIndex = CachingObject->GetCurrentCacheIndex();
+	}
+
+	return CachingIndex;
+}
+
+FString FFindInBlueprintSearchManager::GetCurrentCacheBlueprintName() const
+{
+	FString CachingBPName;
+	if(CachingObject)
+	{
+		CachingBPName = CachingObject->GetCurrentCacheBlueprintName();
+	}
+
+	return CachingBPName;
+}
+
+float  FFindInBlueprintSearchManager::GetCacheProgress() const
+{
+	float ReturnCacheValue = 1.0f;
+
+	if(CachingObject)
+	{
+		ReturnCacheValue = CachingObject->GetCacheProgress();
+	}
+
+	return ReturnCacheValue;
+}
+
+void FFindInBlueprintSearchManager::FinishedCachingBlueprints(int32 InNumberCached)
+{
+	// Multiple threads could be adding to this at the same time
+	FScopeLock ScopeLock(&SafeModifyCacheCriticalSection);
+
+	// Clean out the list by the count that were successfully cached, this allows us to continue the caching later.
+	UncachedBlueprints.RemoveAt(0, InNumberCached);
+
+	// Signal the callback, so all FindInBlueprints can resubmit their search queries
+	for(TWeakPtr< SFindInBlueprints > FindInBlueprintsWidget : OnCachingCompleteCallbackWidgets)
+	{
+		if(FindInBlueprintsWidget.IsValid())
+		{
+			FindInBlueprintsWidget.Pin()->OnCacheComplete();
+		}
+	}
+	OnCachingCompleteCallbackWidgets.Empty();
+
+	// Delete the object and NULL it out so we can do it again in the future if needed (if it was canceled)
+	delete CachingObject;
+	CachingObject = NULL;
+}
+
+bool FFindInBlueprintSearchManager::IsCacheInProgress() const
+{
+	return CachingObject != NULL;
+}
+
+FText FFindInBlueprintSearchManager::ConvertHexStringToFText(FString InHexString)
+{
+	TArray<uint8> SerializedData;
+	SerializedData.AddZeroed(InHexString.Len());
+
+	HexToBytes(InHexString, SerializedData.GetData());
+
+	FText ResultText;
+
+	FMemoryReader Ar(SerializedData, /*bIsPersistent=*/ true);
+	Ar << ResultText;
+	Ar.Close();
+
+	return ResultText;
+}
+
+FString FFindInBlueprintSearchManager::ConvertFTextToHexString(FText InValue)
+{
+	TArray<uint8> SerializedData;
+	FMemoryWriter Ar(SerializedData, /*bIsPersistent=*/ true);
+
+	Ar << InValue;
+	Ar.Close();
+
+	return BytesToHex(SerializedData.GetData(), SerializedData.Num());
+}
\ No newline at end of file
diff --git a/Engine/Source/Editor/Kismet/Private/FindInBlueprints.cpp b/Engine/Source/Editor/Kismet/Private/FindInBlueprints.cpp
--- a/Engine/Source/Editor/Kismet/Private/FindInBlueprints.cpp
+++ b/Engine/Source/Editor/Kismet/Private/FindInBlueprints.cpp
@@ -3,97 +3,482 @@
 #include "BlueprintEditorCommands.h"
 #include "BlueprintEditor.h"
 #include "FindInBlueprints.h"
+#include "FindInBlueprintManager.h"
 #include "Editor/UnrealEd/Public/Kismet2/DebuggerCommands.h"
 #include "Editor/UnrealEd/Public/Kismet2/BlueprintEditorUtils.h"
 #include "AssetRegistryModule.h"
-#include "FindInBlueprintUtils.h"
+#include "Json.h"
+#include "IDocumentation.h"
 
 #define LOCTEXT_NAMESPACE "FindInBlueprints"
 
+/** Some utility functions to help with Find-in-Blueprint functionality */
+namespace FindInBlueprintsHelpers
+{
+	/**
+	 * Retrieves the pin type as a string value
+	 *
+	 * @param InPinType		The pin type to look at
+	 *
+	 * @return				The pin type as a string in format [category]'[sub-category object]'
+	 */
+	FString GetPinTypeAsString(const FEdGraphPinType& InPinType)
+	{
+		FString Result = InPinType.PinCategory;
+		if(UObject* SubCategoryObject = InPinType.PinSubCategoryObject.Get()) 
+		{
+			Result += FString(" '") + SubCategoryObject->GetName() + "'";
+		}
+
+		return Result;
+	}
+
+	/**
+	 * Determines if a string matches the search tokens
+	 *
+	 * @param InTokens					Search tokens to compare with
+	 * @param InComparisonString		The string to check for matching tokens in
+	 *
+	 * @return							TRUE if the string matches any of the search tokens
+	 */
+	bool StringMatchesSearchTokens(const TArray<FString>& InTokens, const FString& InComparisonString)
+	{
+		const FString SafeSearchString = InComparisonString.Replace( TEXT( " " ), TEXT( "" ) );
+
+		bool bFoundAllTokens = true;
+		//search the entry for each token, it must have all of them to pass
+		for(auto TokIT(InTokens.CreateConstIterator());TokIT;++TokIT)
+		{
+			const FString& Token = *TokIT;
+			if(!SafeSearchString.Contains(Token))
+			{
+				bFoundAllTokens = false;
+				break;
+			}
+		}
+		return bFoundAllTokens;
+	}
+
+	/**
+	 * Determines if an FText matches the search tokens, comparing both the (localized) display string and the (unlocalized) source string
+	 *
+	 * @param InTokens					Search tokens to compare with
+	 * @param InComparisonText			The FText to check for matching tokens in
+	 *
+	 * @return							TRUE if the FText matches any of the search tokens
+	 */
+	bool TextMatchesSearchTokens(const TArray<FString>& InTokens, const FText& InComparisonText)
+	{
+		return StringMatchesSearchTokens(InTokens, InComparisonText.ToString()) || StringMatchesSearchTokens(InTokens, InComparisonText.BuildSourceString());
+	}
+
+	/**
+	 * Helper function to extract all graph data from a Json object and into a search result. There is little type safety, it assumes all Json objects found in the InCategoryTitle are graphs
+	 *
+	 * @param InJsonObject					The JsonObject to extract the graphs from
+	 * @param InCategoryTitle				The category the graphs are titled under.
+	 * @param InTokens						Filtering tokens to use
+	 * @param InGraphType					There is nothing in the Json script to identify graphs, they are all grouped by category, so pass in what type they should be considered
+	 * @param OutBlueprintCategory			The search category these graphs will be children to
+	 */
+	void ExtractGraph( TSharedPtr< FJsonObject > InJsonObject, FString InCategoryTitle, const TArray<FString>& InTokens, EGraphType InGraphType, FSearchResult& OutBlueprintCategory )
+	{
+		if(InJsonObject->HasField(InCategoryTitle))
+		{
+			// Pulls out all functions for this Blueprint
+			TArray<TSharedPtr< FJsonValue > > FunctionList = InJsonObject->GetArrayField(InCategoryTitle);
+			for( TSharedPtr< FJsonValue > FunctionValue : FunctionList )
+			{
+				FSearchResult Result(new FFindInBlueprintsGraph(FText::GetEmpty(), OutBlueprintCategory, InGraphType) );
+
+				if( Result->ExtractContent(FunctionValue->AsObject(), InTokens) )
+				{
+					OutBlueprintCategory->Children.Add(Result);
+				}
+			}
+		}
+	}
+
+	/**
+	 * Utility function to begin a search
+	 *
+	 * @param InBuffer					The search buffer, a Json string to be parsed and converted into search result
+	 * @param InTokens					Search tokens to filter the results by
+	 * @param OutBlueprintCategory		The root category for the search results
+	 */
+	void Extract( const FString& InBuffer, const TArray<FString>& InTokens, FSearchResult& OutBlueprintCategory )
+	{
+		// Skip empty buffers
+		if(InBuffer.Len() == 0)
+		{
+			return;
+		}
+
+		TSharedPtr< FJsonObject > JsonObject = NULL;
+		TSharedRef< TJsonReader<> > Reader = TJsonReaderFactory<TCHAR>::Create( InBuffer );
+		FJsonSerializer::Deserialize( Reader, JsonObject );
+
+		if(JsonObject.IsValid())
+		{
+			if(JsonObject->HasField(FFindInBlueprintSearchManager::ConvertFTextToHexString(FFindInBlueprintSearchTags::FiB_Properties)))
+			{
+				// Pulls out all properties (variables) for this Blueprint
+				TArray<TSharedPtr< FJsonValue > > PropertyList = JsonObject->GetArrayField(FFindInBlueprintSearchManager::ConvertFTextToHexString(FFindInBlueprintSearchTags::FiB_Properties));
+				for( TSharedPtr< FJsonValue > PropertyValue : PropertyList )
+				{
+					FSearchResult Result(new FFindInBlueprintsProperty(FText::GetEmpty(), OutBlueprintCategory) );
+
+					if( Result->ExtractContent(PropertyValue->AsObject(), InTokens) )
+					{
+						OutBlueprintCategory->Children.Add(Result);
+					}
+				}
+			}
+
+			ExtractGraph(JsonObject, FFindInBlueprintSearchManager::ConvertFTextToHexString(FFindInBlueprintSearchTags::FiB_Functions), InTokens, EGraphType::GT_Function, OutBlueprintCategory);
+			ExtractGraph(JsonObject, FFindInBlueprintSearchManager::ConvertFTextToHexString(FFindInBlueprintSearchTags::FiB_Macros), InTokens, EGraphType::GT_Macro, OutBlueprintCategory);
+			ExtractGraph(JsonObject, FFindInBlueprintSearchManager::ConvertFTextToHexString(FFindInBlueprintSearchTags::FiB_UberGraphs), InTokens, EGraphType::GT_Ubergraph, OutBlueprintCategory);
+			ExtractGraph(JsonObject, FFindInBlueprintSearchManager::ConvertFTextToHexString(FFindInBlueprintSearchTags::FiB_SubGraphs), InTokens, EGraphType::GT_Ubergraph, OutBlueprintCategory);
+		}
+	}
+
+	/**
+	 * Parses a pin type from passed in key names and values
+	 *
+	 * @param InKey					The key name for what the data should be translated as
+	 * @param InValue				Value to be be translated
+	 * @param InOutPinType			Modifies the PinType based on the passed parameters, building it up over multiple calls
+	 * @param InOutPinSubCategory	When the pin's sub-category is found, this value will be changed. This value will never be wiped during other translations. It is not thread safe to determine the sub-object.
+	 * @return						TRUE when the parsing is successful
+	 */
+	bool ParsePinType(FString InKey, FText InValue, FEdGraphPinType& InOutPinType, FString& OutPinSubCategory)
+	{
+		bool bParsed = true;
+
+		if(InKey == FFindInBlueprintSearchTags::FiB_PinCategory.ToString())
+		{
+			InOutPinType.PinCategory = InValue.ToString();
+		}
+		else if(InKey == FFindInBlueprintSearchTags::FiB_PinSubCategory.ToString())
+		{
+			InOutPinType.PinSubCategory = InValue.ToString();
+		}
+		else if(InKey == FFindInBlueprintSearchTags::FiB_ObjectClass.ToString())
+		{
+			OutPinSubCategory = InValue.ToString();
+		}
+		else if(InKey == FFindInBlueprintSearchTags::FiB_IsArray.ToString())
+		{
+			InOutPinType.bIsArray = InValue.ToString().ToBool();
+		}
+		else if(InKey == FFindInBlueprintSearchTags::FiB_IsReference.ToString())
+		{
+			InOutPinType.bIsReference = InValue.ToString().ToBool();
+		}
+		else
+		{
+			bParsed = false;
+		}
+
+		return bParsed;
+	}
+}
+
 //////////////////////////////////////////////////////////////////////////
 // FBlueprintSearchResult
 
-FFindInBlueprintsResult::FFindInBlueprintsResult(const FString& InValue ) 
-	:Value(InValue), DuplicationIndex(0), Class(NULL), Pin(NULL), GraphNode(NULL)
+FFindInBlueprintsResult::FFindInBlueprintsResult(const FText& InDisplayText ) 
+	: DisplayText(InDisplayText)
 {
 }
 
-FFindInBlueprintsResult::FFindInBlueprintsResult( const FString& InValue, TSharedPtr<FFindInBlueprintsResult>& InParent, UClass* InClass, int InDuplicationIndex ) 
-	:Value(InValue), DuplicationIndex(InDuplicationIndex), Parent(InParent), Class(InClass), Pin(NULL), GraphNode(NULL)
-{	
-}
-
-FFindInBlueprintsResult::FFindInBlueprintsResult( const FString& InValue, TSharedPtr<FFindInBlueprintsResult>& InParent, UEdGraphPin* InPin) 
-	: Value(InValue), DuplicationIndex(0), Parent(InParent), Class(InPin->GetClass()), Pin(InPin), GraphNode(NULL)
+FFindInBlueprintsResult::FFindInBlueprintsResult( const FText& InDisplayText, TSharedPtr<FFindInBlueprintsResult> InParent) 
+	: DisplayText(InDisplayText), Parent(InParent)
 {
 }
 
-FFindInBlueprintsResult::FFindInBlueprintsResult( const FString& InValue, TSharedPtr<FFindInBlueprintsResult>& InParent, UEdGraphNode* InNode) 
-	: Value(InValue), DuplicationIndex(0), Parent(InParent), Class(InNode->GetClass()), Pin(NULL), GraphNode(InNode)
+FReply FFindInBlueprintsResult::OnClick()
 {
-	if(GraphNode.IsValid())
+	// If there is a parent, handle it using the parent's functionality
+	if(Parent.IsValid())
+	{
+		return Parent.Pin()->OnClick();
+	}
+	else
 	{
-		CommentText = GraphNode->NodeComment;
+		// As a last resort, find the parent Blueprint, and open that, it will get the user close to what they want
+		UBlueprint* Blueprint = GetParentBlueprint();
+		if(Blueprint)
+		{
+			FKismetEditorUtilities::BringKismetToFocusAttentionOnObject(Blueprint, false);
+		}
 	}
+	
+	return FReply::Handled();
+}
+
+bool FFindInBlueprintsResult::ParseSearchInfo(const TArray<FString> &InTokens, FString InKey, FText InValue, TSharedPtr< FFindInBlueprintsResult > InParent)
+{
+	// By default, nothing much gets processed, just check to see if it matches the search requirements
+	return FindInBlueprintsHelpers::TextMatchesSearchTokens(InTokens, InValue);
+}
+
+void FFindInBlueprintsResult::AddExtraSearchInfo(FString InKey, FText InValue, TSharedPtr< FFindInBlueprintsResult > InParent)
+{
+	FText DisplayText;
+
+	FFormatNamedArguments Args;
+	Args.Add(TEXT("Key"), FText::FromString(InKey));
+	Args.Add(TEXT("Value"), InValue);
+	DisplayText = FText::Format(LOCTEXT("ExtraSearchInfo", "{Key}: {Value}"), Args);
+
+	TSharedPtr< FFindInBlueprintsResult > SearchResult(new FFindInBlueprintsResult(DisplayText, InParent) );
+	InParent->Children.Add(SearchResult);
 }
 
-/* Get a node from the blueprint and the find result */
-template<class NodeType>
-inline bool GetNode(UBlueprint* Blueprint, const FFindInBlueprintsResult& Result, UEdGraphNode*& OutNode)
+bool FFindInBlueprintsResult::ExtractJsonValue(const TArray<FString>& InTokens, FText InKey, TSharedPtr< FJsonValue > InValue, TSharedPtr< FFindInBlueprintsResult > InParentOverride)
 {
-	if(Result.Class == NodeType::StaticClass()) 
+	bool bMatchesSearchResults = false;
+
+	TSharedPtr< FFindInBlueprintsResult > ParentResultPtr = InParentOverride.IsValid()? InParentOverride : SharedThis( this );
+
+	if( InValue->Type == EJson::Array)
 	{
-		TArray<NodeType*> Nodes;
-		FBlueprintEditorUtils::GetAllNodesOfClass(Blueprint, Nodes );
+		TSharedPtr< FFindInBlueprintsResult > ArrayHeader = MakeShareable(new FFindInBlueprintsResult(InKey, ParentResultPtr) );
+		bool bArrayHasMatches = false;
 
-		int FoundCount = 0;
-		for(auto It(Nodes.CreateConstIterator());It;++It)
+		TArray< TSharedPtr< FJsonValue > > Array = InValue->AsArray();
+		for( int ArrayIdx = 0; ArrayIdx < Array.Num(); ++ArrayIdx )
 		{
-			NodeType* Node = *It;
-			const FString NodePath = FindInBlueprintsUtil::GetNodeTypePath(Node) + "::" + Node->GetNodeTitle(ENodeTitleType::ListView).ToString() ;
-			if(NodePath == Result.Value)
+			if(Array[ArrayIdx]->Type == EJson::String)
 			{
-				if(FoundCount == Result.DuplicationIndex)
-				{
-					OutNode = Node;
-					break;
-				}
-				else
-				{
-					FoundCount++;
-				}
+				bArrayHasMatches |= ParseSearchInfo(InTokens, FString::FromInt(ArrayIdx), FFindInBlueprintSearchManager::ConvertHexStringToFText(Array[ArrayIdx]->AsString()), ArrayHeader);
 			}
+			else
+			{
+				bArrayHasMatches |= ExtractJsonValue(InTokens, FText::FromString(FString::FromInt(ArrayIdx)), Array[ArrayIdx], ArrayHeader);
+			}
+		}
+
+		bMatchesSearchResults |= bArrayHasMatches;
+
+		if(bArrayHasMatches)
+		{
+			ParentResultPtr->Children.Add(ArrayHeader);
+		}
+	}
+	else if(InValue->Type == EJson::Object)
+	{
+		TSharedPtr< FFindInBlueprintsResult > ObjectResultHeader = MakeShareable(new FFindInBlueprintsResult(InKey, ParentResultPtr) );
+
+		bMatchesSearchResults |= ExtractContent(InValue->AsObject(), InTokens, ObjectResultHeader);
+
+		if(bMatchesSearchResults)
+		{
+			ParentResultPtr->Children.Add(ObjectResultHeader);
 		}
 	}
-	return (OutNode != NULL);
+	else if(InValue->Type == EJson::String)
+	{
+		bMatchesSearchResults |= ParseSearchInfo(InTokens, InKey.ToString(), FFindInBlueprintSearchManager::ConvertHexStringToFText(InValue->AsString()), ParentResultPtr);
+	}
+	else
+	{
+		bMatchesSearchResults |= ParseSearchInfo(InTokens, InKey.ToString(), FText::FromString(InValue->AsString()), ParentResultPtr);
+	}
+
+	return bMatchesSearchResults;
 }
 
-FReply FFindInBlueprintsResult::OnClick()
+bool FFindInBlueprintsResult::ExtractContent(TSharedPtr< FJsonObject > InJsonNode, const TArray<FString>& InTokens, TSharedPtr< FFindInBlueprintsResult > InParentOverride)
+{
+	bool bMatchesSearchResults = false;
+
+	TSharedPtr< FFindInBlueprintsResult > ParentResultPtr = InParentOverride.IsValid()? InParentOverride : SharedThis( this );
+
+	for( auto MapValues : InJsonNode->Values )
+	{
+		bMatchesSearchResults |= ExtractJsonValue(InTokens, FFindInBlueprintSearchManager::ConvertHexStringToFText(MapValues.Key), MapValues.Value, ParentResultPtr);
+	}
+	return bMatchesSearchResults;
+}
+
+FString FFindInBlueprintsResult::GetCategory() const
+{
+	return FString();
+}
+
+TSharedRef<SWidget> FFindInBlueprintsResult::CreateIcon() const
+{
+	FLinearColor IconColor = FLinearColor::White;
+	const FSlateBrush* Brush = NULL;
+
+	return 	SNew(SImage)
+			.Image(Brush)
+			.ColorAndOpacity(IconColor)
+			.ToolTipText( GetCategory() );
+}
+
+FString FFindInBlueprintsResult::GetCommentText() const
+{
+	return CommentText;
+}
+
+UBlueprint* FFindInBlueprintsResult::GetParentBlueprint() const
+{
+	if (Parent.IsValid())
+	{
+		return Parent.Pin()->GetParentBlueprint();
+	}
+	else
+	{
+		return Cast<UBlueprint>(StaticLoadObject(UBlueprint::StaticClass(), NULL, *DisplayText.ToString()));
+	}
+}
+
+void FFindInBlueprintsResult::ExpandAllChildren( TSharedPtr< STreeView< TSharedPtr< FFindInBlueprintsResult > > > InTreeView )
+{
+	if( Children.Num() )
+	{
+		InTreeView->SetItemExpansion(this->AsShared(), true);
+		for( int32 i = 0; i < Children.Num(); i++ )
+		{
+			Children[i]->ExpandAllChildren(InTreeView);
+		}
+	}
+}
+
+FString FFindInBlueprintsResult::GetDisplayString() const
+{
+	return DisplayText.ToString();
+}
+
+//////////////////////////////////////////////////////////
+// FFindInBlueprintsGraphNode
+
+FFindInBlueprintsGraphNode::FFindInBlueprintsGraphNode(const FText& InValue, TSharedPtr<FFindInBlueprintsResult> InParent)
+	: FFindInBlueprintsResult(InValue, InParent)
+	, Class(nullptr)
+{
+}
+
+FReply FFindInBlueprintsGraphNode::OnClick()
 {
 	UBlueprint* Blueprint = GetParentBlueprint();
 	if(Blueprint)
 	{
-		if (GraphNode.IsValid() )
+		UEdGraphNode* OutNode = NULL;
+		if(	UEdGraphNode* GraphNode = FBlueprintEditorUtils::GetNodeByGUID(Blueprint, NodeGuid) )
 		{
-			FKismetEditorUtilities::BringKismetToFocusAttentionOnObject(GraphNode.Get(), /*bRequestRename=*/false);
+			FKismetEditorUtilities::BringKismetToFocusAttentionOnObject(GraphNode, /*bRequestRename=*/false);
+			return FReply::Handled();
 		}
-		else if (Pin.IsValid())
+	}
+
+	return FFindInBlueprintsResult::OnClick();
+}
+
+TSharedRef<SWidget> FFindInBlueprintsGraphNode::CreateIcon() const
+{
+	return 	SNew(SImage)
+		.Image(GlyphBrush)
+		.ColorAndOpacity(GlyphColor)
+		.ToolTipText( GetCategory() );
+}
+
+bool FFindInBlueprintsGraphNode::ParseSearchInfo(const TArray<FString> &InTokens, FString InKey, FText InValue, TSharedPtr< FFindInBlueprintsResult > InParent)
+{
+	if(InKey == FFindInBlueprintSearchTags::FiB_NodeGuid.ToString())
+	{
+		FString NodeGUIDAsString = InValue.ToString();
+		FGuid::Parse(NodeGUIDAsString, NodeGuid);
+
+		// NodeGuid's are not searchable
+		return false;
+	}
+
+	bool bMatchesTokens = FindInBlueprintsHelpers::TextMatchesSearchTokens(InTokens, InValue);
+
+	if(InKey == FFindInBlueprintSearchTags::FiB_ClassName.ToString())
+	{
+		ClassName = InValue.ToString();
+		bMatchesTokens = false;
+	}
+	else if(InKey == FFindInBlueprintSearchTags::FiB_Name.ToString())
+	{
+		DisplayText = InValue;
+	}
+	else if(InKey == FFindInBlueprintSearchTags::FiB_Comment.ToString())
+	{
+		CommentText = InValue.ToString();
+	}
+	else if(InKey == FFindInBlueprintSearchTags::FiB_Glyph.ToString())
+	{
+		GlyphBrush = FEditorStyle::GetBrush(*InValue.ToString());
+		bMatchesTokens = false;
+	}
+	else if(InKey == FFindInBlueprintSearchTags::FiB_GlyphColor.ToString())
+	{
+		GlyphColor.InitFromString(InValue.ToString());
+		bMatchesTokens = false;
+	}
+	else
+	{
+		// Any generic keys and value pairs get put in as "extra" search info
+		if(bMatchesTokens)
 		{
-			FKismetEditorUtilities::BringKismetToFocusAttentionOnObject(Pin.Get(), /*bRequestRename=*/false);
+			AddExtraSearchInfo(InKey, InValue, InParent);
+			return true;
 		}
-		else
+	}
+	return bMatchesTokens;
+}
+
+bool FFindInBlueprintsGraphNode::ExtractContent(TSharedPtr< FJsonObject > InJsonNode, const TArray<FString>& InTokens)
+{
+	// Very important to get the schema first, other bits of data depend on it
+	FString SchemaName;
+	TSharedPtr< FJsonValue > SchemaNameValue = InJsonNode->GetField< EJson::Type::String >(FFindInBlueprintSearchManager::ConvertFTextToHexString(FFindInBlueprintSearchTags::FiB_SchemaName));
+	if(SchemaNameValue.IsValid())
+	{
+		SchemaName = FFindInBlueprintSearchManager::ConvertHexStringToFText(SchemaNameValue->AsString()).ToString();
+	}
+
+	bool bMatchesSearchResults = false;
+
+	for( auto MapValues : InJsonNode->Values )
+	{
+		FText Key = FFindInBlueprintSearchManager::ConvertHexStringToFText(MapValues.Key);
+
+		if(Key.CompareTo(FFindInBlueprintSearchTags::FiB_Pins) == 0)
 		{
-			UEdGraphNode* OutNode = NULL;
-			if(	GetNode<UEdGraphNode>(Blueprint, *this, OutNode))
+			TArray< TSharedPtr< FJsonValue > > PinsList = MapValues.Value->AsArray();
+
+			for( TSharedPtr< FJsonValue > Pin : PinsList )
 			{
-				FKismetEditorUtilities::BringKismetToFocusAttentionOnObject(OutNode, /*bRequestRename=*/false);
+				// The pin
+				TSharedPtr< FFindInBlueprintsPin > PinResult(new FFindInBlueprintsPin(FText::GetEmpty(), SharedThis( this ), SchemaName) );
+				if(PinResult->ExtractContent(Pin->AsObject(), InTokens))
+				{
+					Children.Add(PinResult);
+					bMatchesSearchResults = true;
+				}
 			}
 		}
+		else if(Key.CompareTo(FFindInBlueprintSearchTags::FiB_SchemaName) == 0)
+		{
+			// Previously extracted
+			continue;
+		}
+		else
+		{
+			bMatchesSearchResults |= ExtractJsonValue(InTokens, FFindInBlueprintSearchManager::ConvertHexStringToFText(MapValues.Key), MapValues.Value, SharedThis( this ));
+		}
 	}
-	return FReply::Handled();
+	return bMatchesSearchResults;
 }
 
-FString FFindInBlueprintsResult::GetCategory() const
+FString FFindInBlueprintsGraphNode::GetCategory() const
 {
 	if(Class == UK2Node_CallFunction::StaticClass())
 	{
@@ -115,111 +500,432 @@ FString FFindInBlueprintsResult::GetCategory() const
 	{
 		return LOCTEXT("VariableSetCategory", "Variable Set").ToString();
 	}
-	else if(Class == UEdGraphPin::StaticClass())
+
+	return LOCTEXT("NodeCategory", "Node").ToString();
+}
+
+void FFindInBlueprintsGraphNode::FinalizeSearchData()
+{
+	if(!ClassName.IsEmpty())
+	{
+		Class = FindObject<UClass>(ANY_PACKAGE, *ClassName, true);
+		ClassName.Empty();
+	}
+}
+
+//////////////////////////////////////////////////////////
+// FFindInBlueprintsPin
+
+FFindInBlueprintsPin::FFindInBlueprintsPin(const FText& InValue, TSharedPtr<FFindInBlueprintsResult> InParent, FString InSchemaName)
+	: FFindInBlueprintsResult(InValue, InParent)
+	, SchemaName(InSchemaName)
+	, IconColor(FSlateColor::UseForeground())
+{
+}
+
+TSharedRef<SWidget> FFindInBlueprintsPin::CreateIcon() const
+{
+	const FSlateBrush* Brush = NULL;
+
+	if( PinType.bIsArray )
+	{
+		Brush = FEditorStyle::GetBrush( TEXT("GraphEditor.ArrayPinIcon") );
+	}
+	else if( PinType.bIsReference )
 	{
-		return LOCTEXT("PinCategory", "Pin").ToString();
+		Brush = FEditorStyle::GetBrush( TEXT("GraphEditor.RefPinIcon") );
 	}
 	else
 	{
-		return LOCTEXT("NodeCategory", "Node").ToString();
+		Brush = FEditorStyle::GetBrush( TEXT("GraphEditor.PinIcon") );
 	}
-	return FString();
+
+	return 	SNew(SImage)
+		.Image(Brush)
+		.ColorAndOpacity(IconColor)
+		.ToolTipText(FindInBlueprintsHelpers::GetPinTypeAsString(PinType) );
 }
 
-TSharedRef<SWidget> FFindInBlueprintsResult::CreateIcon() const
+bool FFindInBlueprintsPin::ParseSearchInfo(const TArray<FString> &InTokens, FString InKey, FText InValue, TSharedPtr< FFindInBlueprintsResult > InParent)
+{
+	bool bMatchesTokens = FindInBlueprintsHelpers::TextMatchesSearchTokens(InTokens, InValue);
+
+	if(InKey == FFindInBlueprintSearchTags::FiB_Name.ToString())
+	{
+		DisplayText = InValue;
+		bMatchesTokens = FindInBlueprintsHelpers::StringMatchesSearchTokens(InTokens, DisplayText.ToString());
+	}
+	else if(FindInBlueprintsHelpers::ParsePinType(InKey, InValue, PinType, PinSubCategory))
+	{
+		// PinType information is not searchable
+		bMatchesTokens = false;
+	}
+	else
+	{
+		// Any generic keys and value pairs get put in as "extra" search info
+		if(bMatchesTokens)
+		{
+			AddExtraSearchInfo(InKey, InValue, InParent);
+		}
+	}
+	return bMatchesTokens;
+}
+
+FString FFindInBlueprintsPin::GetCategory() const
+{
+	return LOCTEXT("PinCategory", "Pin").ToString();
+}
+
+void FFindInBlueprintsPin::FinalizeSearchData()
+{
+	if(!PinSubCategory.IsEmpty())
+	{
+		PinType.PinSubCategoryObject = FindObject<UClass>(ANY_PACKAGE, *PinSubCategory, true);
+		if(!PinType.PinSubCategoryObject.IsValid())
+		{
+			PinType.PinSubCategoryObject = FindObject<UScriptStruct>(UObject::StaticClass(), *PinSubCategory);
+		}
+		PinSubCategory.Empty();
+	}
+
+	if(!SchemaName.IsEmpty())
+	{
+		UEdGraphSchema* Schema = nullptr;
+		UClass* SchemaClass = FindObject<UClass>(ANY_PACKAGE, *SchemaName, true);
+		if(SchemaClass)
+		{
+			Schema = SchemaClass->GetDefaultObject<UEdGraphSchema>();
+		} 
+		IconColor = Schema->GetPinTypeColor(PinType);
+
+		SchemaName.Empty();
+	}
+}
+
+//////////////////////////////////////////////////////////
+// FFindInBlueprintsProperty
+
+FFindInBlueprintsProperty::FFindInBlueprintsProperty(const FText& InValue, TSharedPtr<FFindInBlueprintsResult> InParent)
+	: FFindInBlueprintsResult(InValue, InParent)
+{
+}
+
+TSharedRef<SWidget> FFindInBlueprintsProperty::CreateIcon() const
+{
+	FLinearColor IconColor = FLinearColor::White;
+	const FSlateBrush* Brush = FEditorStyle::GetBrush(UK2Node_Variable::GetVarIconFromPinType(PinType, IconColor));
+	IconColor = UEdGraphSchema_K2::StaticClass()->GetDefaultObject<UEdGraphSchema_K2>()->GetPinTypeColor(PinType);
+
+	return 	SNew(SImage)
+		.Image(Brush)
+		.ColorAndOpacity(IconColor)
+		.ToolTipText( FindInBlueprintsHelpers::GetPinTypeAsString(PinType) );
+}
+
+bool FFindInBlueprintsProperty::ParseSearchInfo(const TArray<FString> &InTokens, FString InKey, FText InValue, TSharedPtr< FFindInBlueprintsResult > InParent)
+{
+	bool bMatchesTokens = FindInBlueprintsHelpers::TextMatchesSearchTokens(InTokens, InValue);
+
+	if(InKey == FFindInBlueprintSearchTags::FiB_Name.ToString())
+	{
+		DisplayText = InValue;
+	}
+	else if(FindInBlueprintsHelpers::ParsePinType(InKey, InValue, PinType, PinSubCategory))
+	{
+		// Do not need to do anything
+	}
+	else
+	{
+		// Any generic keys and value pairs get put in as "extra" search info
+		if(bMatchesTokens)
+		{
+			AddExtraSearchInfo(InKey, InValue, InParent);
+		}
+	}
+	return bMatchesTokens;
+}
+
+FString FFindInBlueprintsProperty::GetCategory() const
+{
+	return LOCTEXT("Variable", "Variable").ToString();
+}
+
+void FFindInBlueprintsProperty::FinalizeSearchData()
+{
+	if(!PinSubCategory.IsEmpty())
+	{
+		PinType.PinSubCategoryObject = FindObject<UClass>(ANY_PACKAGE, *PinSubCategory, true);
+		if(!PinType.PinSubCategoryObject.IsValid())
+		{
+			PinType.PinSubCategoryObject = FindObject<UScriptStruct>(UObject::StaticClass(), *PinSubCategory);
+		}
+		PinSubCategory.Empty();
+	}
+}
+
+//////////////////////////////////////////////////////////
+// FFindInBlueprintsGraph
+
+FFindInBlueprintsGraph::FFindInBlueprintsGraph(const FText& InValue, TSharedPtr<FFindInBlueprintsResult> InParent, EGraphType InGraphType)
+	: FFindInBlueprintsResult(InValue, InParent)
+	, GraphType(InGraphType)
+{
+}
+
+FReply FFindInBlueprintsGraph::OnClick()
+{
+	UBlueprint* Blueprint = GetParentBlueprint();
+	if(Blueprint)
+	{
+		TArray<UEdGraph*> BlueprintGraphs;
+		Blueprint->GetAllGraphs(BlueprintGraphs);
+
+		for( auto Graph : BlueprintGraphs)
+		{
+			FGraphDisplayInfo DisplayInfo;
+			Graph->GetSchema()->GetGraphDisplayInformation(*Graph, DisplayInfo);
+
+			if(DisplayInfo.PlainName.EqualTo(DisplayText))
+			{
+				FKismetEditorUtilities::BringKismetToFocusAttentionOnObject(Graph);
+				break;
+			}
+		}
+	}
+	else
+	{
+		return FFindInBlueprintsResult::OnClick();
+	}
+	return FReply::Handled();
+}
+
+TSharedRef<SWidget> FFindInBlueprintsGraph::CreateIcon() const
 {
-	FSlateColor IconColor = FSlateColor::UseForeground();
 	const FSlateBrush* Brush = NULL;
-	if(Class == UK2Node_Event::StaticClass())
+	if(GraphType == GT_Function)
 	{
-		Brush = FEditorStyle::GetBrush( TEXT("GraphEditor.FIB_Event") );
+		Brush = FEditorStyle::GetBrush(TEXT("GraphEditor.Function_16x"));
 	}
-	else if(Class == UK2Node_MacroInstance::StaticClass())
+	else if(GraphType == GT_Macro)
+	{
+		Brush = FEditorStyle::GetBrush(TEXT("GraphEditor.Macro_16x"));
+	}
+
+	return 	SNew(SImage)
+		.Image(Brush)
+		.ToolTipText( GetCategory() );
+}
+
+bool FFindInBlueprintsGraph::ExtractContent(TSharedPtr< FJsonObject > InJsonNode, const TArray<FString>& InTokens)
+{
+	bool bMatchesSearchResults = false;
+
+	for( auto MapValues : InJsonNode->Values )
+	{
+		FText Key = FFindInBlueprintSearchManager::ConvertHexStringToFText(MapValues.Key);
+
+		if(Key.CompareTo(FFindInBlueprintSearchTags::FiB_Nodes) == 0)
+		{
+			TArray< TSharedPtr< FJsonValue > > NodeList = MapValues.Value->AsArray();
+
+			for( TSharedPtr< FJsonValue > NodeValue : NodeList )
+			{
+				TSharedPtr< FFindInBlueprintsGraphNode > NodeResult(new FFindInBlueprintsGraphNode(FText::GetEmpty(), SharedThis( this )) );
+
+				if( NodeResult->ExtractContent(NodeValue->AsObject(), InTokens) )
+				{
+					Children.Add(NodeResult);
+					bMatchesSearchResults = true;
+				}
+			}
+		}
+		else if(Key.CompareTo(FFindInBlueprintSearchTags::FiB_SchemaName) == 0)
+		{
+			// Previously extracted
+			continue;
+		}
+		else if(Key.CompareTo(FFindInBlueprintSearchTags::FiB_Properties) == 0)
+		{
+			// Pulls out all properties (variables) for this Blueprint
+			TArray<TSharedPtr< FJsonValue > > PropertyList = MapValues.Value->AsArray();
+			for( TSharedPtr< FJsonValue > PropertyValue : PropertyList )
+			{
+				FSearchResult Result(new FFindInBlueprintsProperty(FText::GetEmpty(), SharedThis( this )) );
+
+				if( Result->ExtractContent(PropertyValue->AsObject(), InTokens) )
+				{
+					Children.Add(Result);
+					bMatchesSearchResults = true;
+				}
+			}
+		}
+		else
+		{
+			bMatchesSearchResults |= ExtractJsonValue( InTokens, FFindInBlueprintSearchManager::ConvertHexStringToFText(MapValues.Key), MapValues.Value, SharedThis( this ) );
+		}
+	}
+	return bMatchesSearchResults;
+}
+
+bool FFindInBlueprintsGraph::ParseSearchInfo(const TArray<FString> &InTokens, FString InKey, FText InValue, TSharedPtr< FFindInBlueprintsResult > InParent)
+{
+	bool bMatchesTokens = FindInBlueprintsHelpers::TextMatchesSearchTokens(InTokens, InValue);
+
+	if(InKey == FFindInBlueprintSearchTags::FiB_Name.ToString())
 	{
-		Brush =  FEditorStyle::GetBrush( TEXT("GraphEditor.FIB_MacroInstance") );
+		DisplayText = InValue;
 	}
-	else if(Class == UK2Node_CallFunction::StaticClass())
+	else
 	{
-		Brush =  FEditorStyle::GetBrush( TEXT("GraphEditor.FIB_CallFunction") );
+		if(bMatchesTokens)
+		{
+			AddExtraSearchInfo(InKey, InValue, InParent);
+		}
 	}
-	else if(Class == UK2Node_VariableGet::StaticClass())
+
+	return bMatchesTokens;
+}
+
+FString FFindInBlueprintsGraph::GetCategory() const
+{
+	if(GraphType == GT_Function)
+	{
+		return LOCTEXT("FunctionGraphCategory", "Function").ToString();
+	}
+	else if(GraphType == GT_Macro)
+	{
+		return LOCTEXT("MacroGraphCategory", "Macro").ToString();
+	}
+	return LOCTEXT("GraphCategory", "Graph").ToString();
+}
+
+////////////////////////////////////
+// FStreamSearch
+
+/**
+ * Async task for searching Blueprints
+ */
+class FStreamSearch : public FRunnable
+{
+public:
+	/** Constructor */
+	FStreamSearch(const TArray<FString>& InSearchTokens )
+		: SearchTokens(InSearchTokens)
+		, bThreadCompleted(false)
+		, StopTaskCounter(0)
 	{
-		Brush =  FEditorStyle::GetBrush( TEXT("GraphEditor.FIB_VariableGet") );
+		Thread = FRunnableThread::Create( this, TEXT("FStreamSearch"), 0, TPri_BelowNormal );
 	}
-	else if(Class == UK2Node_VariableSet::StaticClass())
+
+	/** Begin FRunnable Interface */
+	virtual bool Init() override
 	{
-		Brush =  FEditorStyle::GetBrush( TEXT("GraphEditor.FIB_VariableSet") );
+		return true;
 	}
-	else if(Pin.IsValid())
+
+	virtual uint32 Run() override
 	{
-		if( Pin->PinType.bIsArray )
-		{
-			Brush = FEditorStyle::GetBrush( TEXT("GraphEditor.ArrayPinIcon") );
-		}
-		else if( Pin->PinType.bIsReference )
+		FFindInBlueprintSearchManager::Get().BeginSearchQuery(this);
+
+		FSearchData QueryResult;
+
+		// Searching comes to an end if it is requested using the StopTaskCounter or continuing the search query yields no results
+		while(StopTaskCounter.GetValue() == 0 && FFindInBlueprintSearchManager::Get().ContinueSearchQuery(this, QueryResult))
 		{
-			Brush = FEditorStyle::GetBrush( TEXT("GraphEditor.RefPinIcon") );
+			if(QueryResult.BlueprintPath.Len())
+			{
+				//Each blueprints acts as a category
+				FSearchResult BlueprintCategory = FSearchResult(new FFindInBlueprintsResult(FText::FromString(QueryResult.BlueprintPath)));
+				FindInBlueprintsHelpers::Extract(QueryResult.Value, SearchTokens, BlueprintCategory);
+
+				// If there are children, add the item to the search results
+				if(BlueprintCategory->Children.Num() > 0)
+				{
+					FScopeLock ScopeLock(&SearchCriticalSection);
+					ItemsFound.Add(BlueprintCategory);
+				}
+			}
 		}
-		else
+
+		if(StopTaskCounter.GetValue() > 0)
 		{
-			Brush = FEditorStyle::GetBrush( TEXT("GraphEditor.PinIcon") );
+			// Ensure that the FiB Manager knows that we are done searching
+			FFindInBlueprintSearchManager::Get().EnsureSearchQueryEnds(this);
 		}
-		const UEdGraphSchema* Schema = Pin->GetSchema();
-		IconColor = Schema->GetPinTypeColor(Pin->PinType);
+
+		bThreadCompleted = true;
+
+		return 0;
 	}
-	else if (GraphNode.IsValid())
+
+	virtual void Stop() override
 	{
-		Brush = FEditorStyle::GetBrush( TEXT("GraphEditor.NodeGlyph") );
+		StopTaskCounter.Increment();
 	}
 
-	return 	SNew(SImage)
-			.Image(Brush)
-			.ColorAndOpacity(IconColor)
-			.ToolTipText( GetCategory() );
-}
-
-FString FFindInBlueprintsResult::GetCommentText() const
-{
-	return CommentText;
-}
+	virtual void Exit() override
+	{
 
-FString FFindInBlueprintsResult::GetValueText() const
-{
-	FString DisplayVal;
+	}
+	/** End FRunnable Interface */
 
-	if( Pin.IsValid()
-		&& (!Pin->DefaultValue.IsEmpty() || !Pin->AutogeneratedDefaultValue.IsEmpty() || Pin->DefaultObject || !Pin->DefaultTextValue.IsEmpty())
-		&& (Pin->Direction == EGPD_Input)
-		&& (Pin->LinkedTo.Num() == 0) )
+	/** Brings the thread to a safe stop before continuing. */
+	void EnsureCompletion()
 	{
-		if( Pin->DefaultObject )
-		{
-			const AActor* DefaultActor = Cast<AActor>(Pin->DefaultObject);
-			DisplayVal = FString::Printf(TEXT("(%s)"), DefaultActor ? *DefaultActor->GetActorLabel() : *Pin->DefaultObject->GetName());
-		}
-		else if( !Pin->DefaultValue.IsEmpty() )
-		{
-			DisplayVal = FString::Printf(TEXT("(%s)"), *Pin->DefaultValue);
-		}
-		else if( !Pin->AutogeneratedDefaultValue.IsEmpty() )
-		{
-			DisplayVal = FString::Printf(TEXT("(%s)"), *Pin->AutogeneratedDefaultValue);
-		}
-		else if( !Pin->DefaultTextValue.IsEmpty() )
 		{
-			DisplayVal = FString::Printf(TEXT("(%s)"), *Pin->DefaultTextValue.ToString());
+			FScopeLock CritSectionLock(&SearchCriticalSection);
+			ItemsFound.Empty();
 		}
+
+		Stop();
+		Thread->WaitForCompletion();
+		delete Thread;
+		Thread = NULL;
 	}
 
-	return DisplayVal;
-}
+	/** Returns TRUE if the thread is done with it's work. */
+	bool IsComplete() const
+	{
+		return bThreadCompleted;
+	}
 
-UBlueprint* FFindInBlueprintsResult::GetParentBlueprint() const
-{
-	if (Parent.IsValid()) {return Parent.Pin()->GetParentBlueprint();}
-	else {return Cast<UBlueprint>(StaticLoadObject(UBlueprint::StaticClass(), NULL, *Value));}
-}
+	/**
+	 * Appends the items filtered through the search tokens to the passed array
+	 *
+	 * @param OutItemsFound		All the items found since last queried
+	 */
+	void GetFilteredItems(TArray<FSearchResult>& OutItemsFound)
+	{
+		FScopeLock ScopeLock(&SearchCriticalSection);
+		OutItemsFound.Append(ItemsFound);
+		ItemsFound.Empty();
+	}
+
+	/** Helper function to query the percent complete this search is */
+	float GetPercentComplete() const
+	{
+		return FFindInBlueprintSearchManager::Get().GetPercentComplete(this);
+	}
+
+public:
+	/** Thread to run the cleanup FRunnable on */
+	FRunnableThread* Thread;
+
+	/** A list of items found, cleared whenever the main thread pulls them to display to screen */
+	TArray<FSearchResult> ItemsFound;
+
+	/** The search tokens to filter results by */
+	TArray<FString> SearchTokens;
 
+	/** Prevents searching while other threads are pulling search results */
+	FCriticalSection SearchCriticalSection;
+
+	// Whether the thread has finished running
+	bool bThreadCompleted;
+
+	/** > 0 if we've been asked to abort work in progress at the next opportunity */
+	FThreadSafeCounter StopTaskCounter;
+};
 
 //////////////////////////////////////////////////////////////////////////
 // SBlueprintSearch
@@ -229,10 +935,10 @@ void SFindInBlueprints::Construct( const FArguments& InArgs, TSharedPtr<FBluepri
 	BlueprintEditorPtr = InBlueprintEditor;
 
 	bIsInFindWithinBlueprintMode = true;
-
+	
 	this->ChildSlot
 		[
-			SNew(SVerticalBox)
+			SAssignNew(MainVerticalBox, SVerticalBox)
 			+SVerticalBox::Slot()
 			.AutoHeight()
 			[
@@ -270,13 +976,208 @@ void SFindInBlueprints::Construct( const FArguments& InArgs, TSharedPtr<FBluepri
 					.TreeItemsSource( &ItemsFound )
 					.OnGenerateRow( this, &SFindInBlueprints::OnGenerateRow )
 					.OnGetChildren( this, &SFindInBlueprints::OnGetChildren )
-					.OnSelectionChanged(this,&SFindInBlueprints::OnTreeSelectionChanged)
+					.OnMouseButtonDoubleClick(this,&SFindInBlueprints::OnTreeSelectionDoubleClicked)
 					.SelectionMode( ESelectionMode::Multi )
 				]
 			]
+
+			+SVerticalBox::Slot()
+				.AutoHeight()
+			[
+				SNew(SHorizontalBox)
+
+				// Text
+				+SHorizontalBox::Slot()
+				.AutoWidth()
+				.Padding(0, 2)
+				[
+					SNew(STextBlock)
+					.Font( FEditorStyle::GetFontStyle("AssetDiscoveryIndicator.DiscovertingAssetsFont") )
+					.Text( LOCTEXT("SearchResults", "Searching...").ToString() )
+					.Visibility(this, &SFindInBlueprints::GetSearchbarVisiblity)
+				]
+
+				// Progress bar
+				+SHorizontalBox::Slot()
+					.FillWidth(1.0f)
+					.Padding(2.0f, 3.0f, 0.0f, 0.0f)
+				[
+					SNew(SProgressBar)
+						.Visibility(this, &SFindInBlueprints::GetSearchbarVisiblity).Percent( this, &SFindInBlueprints::GetPercentCompleteSearch )
+				]
+			]
 		];
 }
 
+void SFindInBlueprints::ConditionallyAddCacheBar()
+{
+	// Do not add a second cache bar and do not add it when there are no uncached Blueprints
+	if(FFindInBlueprintSearchManager::Get().GetNumberUncachedBlueprints() > 0)
+	{
+		if(MainVerticalBox.IsValid() && !CacheBarSlot.IsValid())
+		{
+			MainVerticalBox.Pin()->AddSlot()
+				.AutoHeight()
+				[
+					SAssignNew(CacheBarSlot, SBorder)
+					.Visibility( this, &SFindInBlueprints::GetCachingBarVisibility )
+					.BorderBackgroundColor( this, &SFindInBlueprints::GetCachingBarColor )
+					.BorderImage( FCoreStyle::Get().GetBrush("ErrorReporting.Box") )
+					.Padding( FMargin(3,1) )
+					[
+						SNew(SVerticalBox)
+
+						+SVerticalBox::Slot()
+							.AutoHeight()
+						[
+							SNew(SHorizontalBox)
+							+SHorizontalBox::Slot()
+							.VAlign(EVerticalAlignment::VAlign_Center)
+							.AutoWidth()
+							[
+								SNew(STextBlock)
+								.Text(this, &SFindInBlueprints::GetUncachedBlueprintWarningText)
+								.ColorAndOpacity( FCoreStyle::Get().GetColor("ErrorReporting.ForegroundColor") )
+							]
+
+							// Cache All button
+							+SHorizontalBox::Slot()
+								.AutoWidth()
+								.VAlign(EVerticalAlignment::VAlign_Center)
+								.Padding(6.0f, 2.0f, 4.0f, 2.0f)
+								[
+									SNew(SButton)
+									.Text(LOCTEXT("IndexAllBlueprints", "Index All"))
+									.OnClicked( this, &SFindInBlueprints::OnCacheAllBlueprints )
+									.Visibility( this, &SFindInBlueprints::GetCacheAllButtonVisibility )
+									.ToolTip(IDocumentation::Get()->CreateToolTip(
+										LOCTEXT("IndexAlLBlueprints_Tooltip", "Loads all non-indexed Blueprints and stores the data with the DDC. This can be a very slow process, feel free to cancel it whenever and benefit from the Blueprints it indexes."),
+										NULL,
+										TEXT("Shared/Editors/BlueprintEditor"),
+										TEXT("FindInBlueprint_IndexAll")))
+								]
+
+							// Cache progress bar
+							+SHorizontalBox::Slot()
+								.FillWidth(1.0f)
+								.Padding(4.0f, 2.0f, 4.0f, 2.0f)
+								[
+									SNew(SProgressBar)
+									.Percent( this, &SFindInBlueprints::GetPercentCompleteCache )
+									.Visibility( this, &SFindInBlueprints::GetCachingProgressBarVisiblity )
+								]
+
+							// Cancel button
+							+SHorizontalBox::Slot()
+								.AutoWidth()
+								.Padding(4.0f, 2.0f, 0.0f, 2.0f)
+								[
+									SNew(SButton)
+									.Text(LOCTEXT("CancelCacheAll", "Cancel"))
+									.OnClicked( this, &SFindInBlueprints::OnCancelCacheAll )
+									.Visibility( this, &SFindInBlueprints::GetCachingProgressBarVisiblity )
+									.ToolTipText( LOCTEXT("CancelCacheAll_Tooltip", "Stops the caching process from where ever it is, can be started back up where it left off when needed.") )
+								]
+
+							// "X" to remove the bar
+							+SHorizontalBox::Slot()
+								.HAlign(HAlign_Right)
+								[
+									SNew(SButton)
+									.ButtonStyle( FCoreStyle::Get(), "NoBorder" )
+									.ContentPadding(0)
+									.HAlign(HAlign_Center)
+									.VAlign(VAlign_Center)
+									.OnClicked( this, &SFindInBlueprints::OnRemoveCacheBar )
+									.ForegroundColor( FSlateColor::UseForeground() )
+									[
+										SNew(SImage)
+										.Image( FCoreStyle::Get().GetBrush("EditableComboBox.Delete") )
+										.ColorAndOpacity( FSlateColor::UseForeground() )
+									]
+								]
+						]
+
+						+SVerticalBox::Slot()
+							.AutoHeight()
+							.Padding(8.0f, 0.0f, 0.0f, 2.0f)
+						[
+							SNew(STextBlock)
+							.Text(this, &SFindInBlueprints::GetCurrentCacheBlueprintName)
+							.Visibility( this, &SFindInBlueprints::GetCachingBlueprintNameVisiblity )
+							.ColorAndOpacity( FCoreStyle::Get().GetColor("ErrorReporting.ForegroundColor") )
+						]
+					]
+				];
+		}
+	}
+	else
+	{
+		// Because there are no uncached Blueprints, remove the bar
+		OnRemoveCacheBar();
+	}
+}
+
+FReply SFindInBlueprints::OnRemoveCacheBar()
+{
+	if(MainVerticalBox.IsValid() && CacheBarSlot.IsValid())
+	{
+		MainVerticalBox.Pin()->RemoveSlot(CacheBarSlot.Pin().ToSharedRef());
+	}
+
+	return FReply::Handled();
+}
+
+SFindInBlueprints::~SFindInBlueprints()
+{
+	if(StreamSearch.IsValid())
+	{
+		StreamSearch->Stop();
+		StreamSearch->EnsureCompletion();
+	}
+}
+
+void SFindInBlueprints::Tick( const FGeometry& AllottedGeometry, const double InCurrentTime, const float InDeltaTime )
+{
+	SCompoundWidget::Tick(AllottedGeometry, InCurrentTime, InDeltaTime);
+
+	if(StreamSearch.IsValid())
+	{
+		bool bShouldShutdownThread = false;
+		bShouldShutdownThread = StreamSearch->IsComplete();
+
+		TArray<FSearchResult> BackgroundItemsFound;
+
+		StreamSearch->GetFilteredItems(BackgroundItemsFound);
+		if(BackgroundItemsFound.Num())
+		{
+			for(auto& Item : BackgroundItemsFound)
+			{
+				Item->ExpandAllChildren(TreeView);
+				ItemsFound.Add(Item);
+			}
+			TreeView->RequestTreeRefresh();
+		}
+
+		// If the thread is complete, shut it down properly
+		if(bShouldShutdownThread)
+		{
+			if(ItemsFound.Num() == 0)
+			{
+				// Insert a fake result to inform user if none found
+				ItemsFound.Add(FSearchResult(new FFindInBlueprintsResult(LOCTEXT("BlueprintSearchNoResults", "No Results found"))));
+				TreeView->RequestTreeRefresh();
+			}
+
+			// Add the cache bar if needed.
+			ConditionallyAddCacheBar();
+
+			StreamSearch->EnsureCompletion();
+			StreamSearch.Reset();
+		}
+	}
+}
+
 void SFindInBlueprints::FocusForUse(bool bSetFindWithinBlueprint, FString NewSearchTerms, bool bSelectFirstResult)
 {
 	// NOTE: Careful, GeneratePathToWidget can be reentrant in that it can call visibility delegates and such
@@ -355,13 +1256,20 @@ void SFindInBlueprints::InitiateSearch()
 		SearchValue.ParseIntoArray(&Tokens, TEXT(" "), true);
 	}
 
-	for (auto It(ItemsFound.CreateIterator());It;++It)
+	if(ItemsFound.Num())
 	{
-		TreeView->SetItemExpansion(*It, false);
+		// Reset the scroll to the top
+		TreeView->RequestScrollIntoView(ItemsFound[0]);
 	}
+
 	ItemsFound.Empty();
+
 	if (Tokens.Num() > 0)
 	{
+		OnRemoveCacheBar();
+
+		TreeView->RequestTreeRefresh();
+
 		HighlightText = FText::FromString( SearchValue );
 		if (bIsInFindWithinBlueprintMode)
 		{
@@ -369,238 +1277,65 @@ void SFindInBlueprints::InitiateSearch()
 		}
 		else
 		{
-			MatchTokens(Tokens);
+			LaunchStreamThread(Tokens);
 		}
 	}
-
-	// Insert a fake result to inform user if none found
-	if (ItemsFound.Num() == 0)
-	{
-		ItemsFound.Add(FSearchResult(new FFindInBlueprintsResult(LOCTEXT("BlueprintSearchNoResults", "No Results found").ToString())));
-	}
-
-	TreeView->RequestTreeRefresh();
-
-	for (auto It(ItemsFound.CreateIterator()); It; ++It)
-	{
-		TreeView->SetItemExpansion(*It, true);
-	}
 }
 
-void SFindInBlueprints::MatchTokens( const TArray<FString> &Tokens ) 
+void SFindInBlueprints::MatchTokensWithinBlueprint(const TArray<FString>& Tokens)
 {
-	RootSearchResult.Reset();
-
-	FAssetRegistryModule& AssetRegistryModule = FModuleManager::LoadModuleChecked<FAssetRegistryModule>(TEXT("AssetRegistry"));
-	
-	TArray<FAssetData> Assets;
-	AssetRegistryModule.Get().GetAssetsByClass(UBlueprint::StaticClass()->GetFName(), Assets, true);
-	for(auto It(Assets.CreateConstIterator());It;++It)
+	if(StreamSearch.IsValid() && !StreamSearch->IsComplete())
 	{
-		const FAssetData& AssetData = *It;
-		//Each blueprints acts as a category
-		FSearchResult BlueprintCategory = FSearchResult(new FFindInBlueprintsResult(AssetData.ObjectPath.ToString()));
-		
-		//find results within the blueprint 
-		Extract(AssetData.TagsAndValues.Find("CallsFunctions"), Tokens, BlueprintCategory, UK2Node_CallFunction::StaticClass());
-		Extract(AssetData.TagsAndValues.Find("MacroInstances"), Tokens, BlueprintCategory, UK2Node_MacroInstance::StaticClass());
-		Extract(AssetData.TagsAndValues.Find("ImplementsFunction"), Tokens, BlueprintCategory, UK2Node_Event::StaticClass());
-		Extract(AssetData.TagsAndValues.Find("VariableGet"), Tokens, BlueprintCategory, UK2Node_VariableGet::StaticClass());
-		Extract(AssetData.TagsAndValues.Find("VariableSet"), Tokens, BlueprintCategory, UK2Node_VariableSet::StaticClass());
-		Extract(AssetData.TagsAndValues.Find("MulticastDelegate"), Tokens, BlueprintCategory, UK2Node_BaseMCDelegate::StaticClass());
-		Extract(AssetData.TagsAndValues.Find("DelegateSet"), Tokens, BlueprintCategory, UK2Node_DelegateSet::StaticClass());
-		Extract(AssetData.TagsAndValues.Find("ImplementedInterfaces"), Tokens, BlueprintCategory, UInterface::StaticClass());
-		Extract(AssetData.TagsAndValues.Find("Comments"), Tokens, BlueprintCategory, UEdGraphNode::StaticClass());
-
-		//if we got any matches and in FIB mode, add the blueprint as a category to the tree view
-		if(BlueprintCategory->Children.Num() > 0 && !bIsInFindWithinBlueprintMode)
-		{
-			ItemsFound.Add(BlueprintCategory);
-		}
+		StreamSearch->Stop();
+		StreamSearch->EnsureCompletion();
+		StreamSearch->GetFilteredItems(ItemsFound);
+		ItemsFound.Empty();
+		StreamSearch.Reset();
 	}
-}
-
-void SFindInBlueprints::MatchTokensWithinBlueprint(const TArray<FString>& Tokens)
-{
-	RootSearchResult.Reset();
 
+	// No multi-threading for single Blueprint search, just query for the data block and process it.
 	UBlueprint* Blueprint = BlueprintEditorPtr.Pin()->GetBlueprintObj();
 
-	RootSearchResult = FSearchResult(new FFindInBlueprintsResult(Blueprint->GetPathName()));
-		
-	TArray<UK2Node*> Nodes;
-	FBlueprintEditorUtils::GetAllNodesOfClass(Blueprint, Nodes);
+	RootSearchResult = FSearchResult(new FFindInBlueprintsResult(FText::FromString(Blueprint->GetPathName())));
+	
+	FindInBlueprintsHelpers::Extract(FFindInBlueprintSearchManager::Get().QuerySingleBlueprint(Blueprint), Tokens, RootSearchResult);
 
-	for(auto It(Nodes.CreateConstIterator());It;++It)
+	// We do not use the root, the user already knows what Blueprint they are searching within
+	ItemsFound = RootSearchResult->Children;
+	
+	if(ItemsFound.Num() == 0)
 	{
-		UK2Node* Node = *It;
-			
-		const FString NodeName = Node->GetNodeTitle(ENodeTitleType::ListView).ToString();
-		FSearchResult NodeResult(new FFindInBlueprintsResult(NodeName, RootSearchResult, Node));
-		UK2Node* K2Node = Cast<UK2Node>(Node);
-		AActor* NodeActor = K2Node ? K2Node->GetReferencedLevelActor() : NULL;
-		FString NodeSearchString = NodeName + Node->GetClass()->GetName() + Node->NodeComment + (NodeActor ? NodeActor->GetName() + NodeActor->GetActorLabel() : FString());
-		NodeSearchString = NodeSearchString.Replace( TEXT( " " ), TEXT( "" ) );
-
-		bool bNodeMatchesSearch = StringMatchesSearchTokens(Tokens, NodeSearchString);
-
-		for( TArray<UEdGraphPin*>::TIterator PinIt(Node->Pins); PinIt; ++PinIt )
-		{
-			UEdGraphPin* Pin = *PinIt;
-			FString PinName = Pin->PinName;
-			FString PinSearchString = Pin->PinName + Pin->PinFriendlyName.ToString() + Pin->DefaultValue + Pin->PinType.PinCategory + Pin->PinType.PinSubCategory + (Pin->PinType.PinSubCategoryObject.IsValid() ? Pin->PinType.PinSubCategoryObject.Get()->GetFullName() : TEXT(""));
-			PinSearchString = PinSearchString.Replace( TEXT( " " ), TEXT( "" ) );
-			if ((Pin != NULL) && StringMatchesSearchTokens(Tokens, PinSearchString))
-			{
-				FSearchResult PinResult(new FFindInBlueprintsResult(PinName, NodeResult, Pin));
-				NodeResult->Children.Add(PinResult);
-			}
-		}
-
-		if ((NodeResult->Children.Num() > 0) || bNodeMatchesSearch)
-		{
-			ItemsFound.Add(NodeResult);
-		}
+		// Insert a fake result to inform user if none found
+		ItemsFound.Add(FSearchResult(new FFindInBlueprintsResult(LOCTEXT("BlueprintSearchNoResults", "No Results found"))));
+		TreeView->RequestTreeRefresh();
 	}
-
-	// Now search for comment nodes containing our search
-	TArray<UEdGraphNode_Comment*> CommentNodes;
-	FBlueprintEditorUtils::GetAllNodesOfClass(Blueprint, CommentNodes);	
-	if( CommentNodes.Num() != 0 )
+	else
 	{
-		// Create the root for the comment hit results
-		const FString NodeName = LOCTEXT("BlueprintResultCommentResults", "Comments containing search term(s)").ToString();
-		FSearchResult RootCommentResult(new FFindInBlueprintsResult(NodeName, RootSearchResult, CommentNodes[0]->GetClass(), 0));
-		
-		for(auto It(CommentNodes.CreateConstIterator());It;++It)
-		{
-			UEdGraphNode_Comment* Node = *It;
-			const FString CommentNodeName = Node->NodeComment;
-			FSearchResult NodeResult(new FFindInBlueprintsResult(CommentNodeName, RootCommentResult, Node));		
-			FString NodeSearchString = CommentNodeName + Node->GetClass()->GetName() + Node->NodeComment;
-			NodeSearchString = NodeSearchString.Replace( TEXT( " " ), TEXT( "" ) );
-
-			bool bNodeMatchesSearch = StringMatchesSearchTokens(Tokens, NodeSearchString);
-
-			if (bNodeMatchesSearch)
-			{
-				RootCommentResult->Children.Add( NodeResult );
-			}			
-		}	
-		
-		if (RootCommentResult->Children.Num() > 0)
+		for(auto Item : ItemsFound)
 		{
-			ItemsFound.Add(RootCommentResult);
+			Item->ExpandAllChildren(TreeView);
 		}
 	}
-	
-	// Search for interfaces/functions within interfaces matching the search
-	TArray<UK2Node*> AddedInterfaceFunctions;
-	for (int32 i=0; i < Blueprint->ImplementedInterfaces.Num(); i++)
-	{
-		// Create the root for the comment hit results
-		const FString NodeName = FString::Printf(*LOCTEXT("BlueprintResultInterfacesImplemented", "Interface implementations containing search term(s)").ToString() );
-		FSearchResult RootInterfaceResult(new FFindInBlueprintsResult(NodeName, RootSearchResult,  UFunction::StaticClass(), 0));
-
-		const FBPInterfaceDescription& InterfaceDesc = Blueprint->ImplementedInterfaces[i];
-		
-		// First check the interface name
- 		FString InterfaceNodeSearchString = InterfaceDesc.Interface->GetName();
-		const bool bInterfaceMatchesSearch = StringMatchesSearchTokens(Tokens, InterfaceNodeSearchString);
-
-		for (int32 FuncIdx = 0; FuncIdx < InterfaceDesc.Graphs.Num(); FuncIdx++)
-		{
-			UEdGraph* Graph = InterfaceDesc.Graphs[FuncIdx];
-			check(Graph);
-
-			FString NodeSearchString = Graph->GetFName().ToString();// + FunctionName.ToString();
-			const bool bNodeMatchesSearch = StringMatchesSearchTokens(Tokens, NodeSearchString);
 
-			if (bNodeMatchesSearch)
-			{
-				FString SigName;
-				for (int32 iNode = 0; iNode < Graph->Nodes.Num() ; iNode++)
-				{
-					// Check if this token is implemented as a function
-					UK2Node_FunctionEntry* FunctionEntryNode = Cast<UK2Node_FunctionEntry>(Graph->Nodes[ iNode ]);
-					if( ( FunctionEntryNode != NULL ) && ( AddedInterfaceFunctions.Find( FunctionEntryNode ) == INDEX_NONE ) )
-					{						
-						SigName = FunctionEntryNode->SignatureName.ToString();
-						if( StringMatchesSearchTokens(Tokens, SigName ) == true )
-						{
-							const FString InterfaceImpl = FString::Printf(*LOCTEXT("BlueprintResultInterfaceImplemented", "Interface Function implementation: %s").ToString(), *SigName );
-							FSearchResult NodeResult(new FFindInBlueprintsResult( InterfaceImpl, RootInterfaceResult, FunctionEntryNode ));	
-							RootInterfaceResult->Children.Add( NodeResult );							
-						}						
-					}				
-				}
-			}	
-		}
-		// Did we match interface name but no functions ?
-		if( ( bInterfaceMatchesSearch == true ) && ( RootInterfaceResult->Children.Num() == 0 ) )
-		{
-			const FString InterfaceNameResult = FString::Printf(*LOCTEXT("BlueprintResultInterfaceImplementation", "Interface implementation: %s").ToString(), *InterfaceNodeSearchString);
-			const FSearchResult NodeResult(new FFindInBlueprintsResult(InterfaceNameResult, RootInterfaceResult, InterfaceDesc.Interface->GetClass(), 0));
-			RootInterfaceResult->Children.Add(NodeResult);
-		}
-		
-		if (RootInterfaceResult->Children.Num() > 0)
-		{
-			ItemsFound.Add(RootInterfaceResult);
-		}
-	}
+	TreeView->RequestTreeRefresh();
 }
 
-void SFindInBlueprints::Extract( const FString* Buffer, const TArray<FString> &Tokens, FSearchResult& BlueprintCategory, UClass* Class )
+void SFindInBlueprints::LaunchStreamThread(const TArray<FString>& InTokens)
 {
-	if(!Buffer)
+	if(StreamSearch.IsValid() && !StreamSearch->IsComplete())
 	{
-		return;
+		StreamSearch->Stop();
+		StreamSearch->EnsureCompletion();
 	}
 
-	TArray<FString> SubStrings;
-	Buffer->ParseIntoArray(&SubStrings, TEXT(","), true);
-	//go through each comma separated entry
-	for(auto SubIt(SubStrings.CreateConstIterator());SubIt;++SubIt)
-	{
-		const FString SubString = *SubIt;
-
-		// Remove the spaces from the test string to increase the chance of finding the item.
-		const FString SafeSubString = SubString.Replace( TEXT( " " ), TEXT( "" ) );
-		if(StringMatchesSearchTokens(Tokens, SafeSubString))
-		{
-			//found a match. Read the count and add that many entries
-			FString SubStringWithoutCount;
-			const int DuplicateCount = ReadDuplicateCount(SubString, SubStringWithoutCount);
-
-			FString Comment;
-
-			// Extracting comments are signified by the UEdGraphNode
-			if(Class == UEdGraphNode::StaticClass())
-			{
-				// Comment meta-data is always arranged in form of "[comment]::[node name]". Parse the string and change the label the result will display as to the node name.
-				TArray<FString> CommentSubStrings;
-				SubStringWithoutCount.ParseIntoArray(&CommentSubStrings, TEXT("::"), true);
-
-				// Certain searches will drop the comment string and find only the node name, so it is not guranteed to properly have the two expected strings and causing a crash.
-				if(CommentSubStrings.Num() >= 2)
-				{
-					Comment = CommentSubStrings[0];
-				}
-			}
-
-			for(int i = 0;i<DuplicateCount;++i)
-			{
-				FSearchResult Result(new FFindInBlueprintsResult(SubStringWithoutCount, BlueprintCategory, Class, i) );
-				Result->CommentText = Comment;
-				BlueprintCategory->Children.Add(Result);
-			}
-		}
-	}
+	StreamSearch = MakeShareable(new FStreamSearch(InTokens));
 }
 
 TSharedRef<ITableRow> SFindInBlueprints::OnGenerateRow( FSearchResult InItem, const TSharedRef<STableViewBase>& OwnerTable )
 {
+	// Finalize the search data, this does some non-thread safe actions that could not be done on the separate thread.
+	InItem->FinalizeSearchData();
+
 	bool bIsACategoryWidget = !bIsInFindWithinBlueprintMode && !InItem->Parent.IsValid();
 
 	if (bIsACategoryWidget)
@@ -614,19 +1349,36 @@ TSharedRef<ITableRow> SFindInBlueprints::OnGenerateRow( FSearchResult InItem, co
 				.ForegroundColor(FEditorStyle::GetColor("PropertyWindow.CategoryForeground"))
 				[
 					SNew(STextBlock)
-					.Text(InItem->Value)
+					.Text(InItem.Get(), &FFindInBlueprintsResult::GetDisplayString)
 					.ToolTipText(LOCTEXT("BlueprintCatSearchToolTip", "Blueprint").ToString())
 				]
 			];
 	}
 	else // Functions/Event/Pin widget
 	{
-		FString CommentText = InItem->GetCommentText();
+		FString CommentText;
+
+		if(!InItem->GetCommentText().IsEmpty())
+		{
+			FFormatNamedArguments Args;
+			Args.Add(TEXT("Comment"), FText::FromString(InItem->GetCommentText()));
+
+			CommentText = FText::Format(LOCTEXT("NodeComment", "Node Comment:[{Comment}]"), Args).ToString();
+		}
+
+		FString Tooltip;
+
+		FFormatNamedArguments Args;
+		Args.Add(TEXT("Category"), FText::FromString(InItem->GetCategory()));
+		Args.Add(TEXT("DisplayTitle"), InItem->DisplayText);
+
+		Tooltip = FText::Format(LOCTEXT("BlueprintResultSearchToolTip", "{Category} : {DisplayTitle}"), Args).ToString();
 
 		return SNew( STableRow< TSharedPtr<FFindInBlueprintsResult> >, OwnerTable )
 			[
 				SNew(SHorizontalBox)
 				+SHorizontalBox::Slot()
+				.VAlign(VAlign_Top)
 				.AutoWidth()
 				[
 					InItem->CreateIcon()
@@ -637,19 +1389,9 @@ TSharedRef<ITableRow> SFindInBlueprints::OnGenerateRow( FSearchResult InItem, co
 				.Padding(2,0)
 				[
 					SNew(STextBlock)
-					.Text(InItem->Value)
-					.HighlightText(HighlightText)
-					.ToolTipText(FString::Printf(*LOCTEXT("BlueprintResultSearchToolTip", "%s : %s").ToString(), *InItem->GetCategory(),* InItem->Value))
-				]
-				+SHorizontalBox::Slot()
-				.FillWidth(1)
-				.HAlign(HAlign_Right)
-				.VAlign(VAlign_Center)
-				.Padding(2,0)
-				[
-					SNew(STextBlock)
-					.Text( InItem->GetValueText() )
-					.HighlightText(HighlightText)
+						.Text(InItem.Get(), &FFindInBlueprintsResult::GetDisplayString)
+						.HighlightText(HighlightText)
+						.ToolTipText(Tooltip)
 				]
 				+SHorizontalBox::Slot()
 				.FillWidth(1)
@@ -658,7 +1400,7 @@ TSharedRef<ITableRow> SFindInBlueprints::OnGenerateRow( FSearchResult InItem, co
 				.Padding(2,0)
 				[
 					SNew(STextBlock)
-					.Text( CommentText.IsEmpty() ? FString() : FString::Printf(TEXT("Node Comment:[%s]"), *CommentText) )
+					.Text( CommentText )
 					.ColorAndOpacity(FLinearColor::Yellow)
 					.HighlightText(HighlightText)
 				]
@@ -671,7 +1413,7 @@ void SFindInBlueprints::OnGetChildren( FSearchResult InItem, TArray< FSearchResu
 	OutChildren += InItem->Children;
 }
 
-void SFindInBlueprints::OnTreeSelectionChanged( FSearchResult Item , ESelectInfo::Type  )
+void SFindInBlueprints::OnTreeSelectionDoubleClicked( FSearchResult Item )
 {
 	if(Item.IsValid())
 	{
@@ -679,42 +1421,112 @@ void SFindInBlueprints::OnTreeSelectionChanged( FSearchResult Item , ESelectInfo
 	}
 }
 
-int SFindInBlueprints::ReadDuplicateCount( const FString& ValueWCount,  FString& OutValue )
+TOptional<float> SFindInBlueprints::GetPercentCompleteSearch() const
 {
-	TArray<FString> SubStrings;
-	OutValue = ValueWCount;
-	ValueWCount.ParseIntoArray(&SubStrings, TEXT("::"), true);
-	if(SubStrings.Num() > 2)
+	if(StreamSearch.IsValid())
 	{
-		const int Count =  FMath::Max(FCString::Atoi(*SubStrings.Last()), 1);
-		if(Count != 1)
-		{
-			OutValue = OutValue.LeftChop(SubStrings.Last().Len()+2);
-		}
-		return Count;
+		return StreamSearch->GetPercentComplete();
 	}
-	else
+	return 0.0f;
+}
+
+EVisibility SFindInBlueprints::GetSearchbarVisiblity() const
+{
+	return StreamSearch.IsValid()? EVisibility::Visible : EVisibility::Collapsed;
+}
+
+FReply SFindInBlueprints::OnCacheAllBlueprints()
+{
+	FFindInBlueprintSearchManager::Get().CacheAllUncachedBlueprints(SharedThis(this));
+	return FReply::Handled();
+}
+
+FReply SFindInBlueprints::OnCancelCacheAll()
+{
+	FFindInBlueprintSearchManager::Get().CancelCacheAll();
+
+	// Resubmit the last search
+	OnSearchTextCommitted(SearchTextField->GetText(), ETextCommit::OnEnter);
+
+	return FReply::Handled();
+}
+
+int32 SFindInBlueprints::GetCurrentCacheIndex() const
+{
+	return FFindInBlueprintSearchManager::Get().GetCurrentCacheIndex();
+}
+
+TOptional<float> SFindInBlueprints::GetPercentCompleteCache() const
+{
+	return FFindInBlueprintSearchManager::Get().GetCacheProgress();
+}
+
+EVisibility SFindInBlueprints::GetCachingProgressBarVisiblity() const
+{
+	return IsCacheInProgress()? EVisibility::Visible : EVisibility::Hidden;
+}
+
+EVisibility SFindInBlueprints::GetCacheAllButtonVisibility() const
+{
+	return IsCacheInProgress()? EVisibility::Collapsed : EVisibility::Visible;
+}
+
+EVisibility SFindInBlueprints::GetCachingBarVisibility() const
+{
+	return (FFindInBlueprintSearchManager::Get().GetNumberUncachedBlueprints() > 0)? EVisibility::Visible : EVisibility::Collapsed;
+}
+
+EVisibility SFindInBlueprints::GetCachingBlueprintNameVisiblity() const
+{
+	return IsCacheInProgress()? EVisibility::Visible : EVisibility::Collapsed;
+}
+
+bool SFindInBlueprints::IsCacheInProgress() const
+{
+	return FFindInBlueprintSearchManager::Get().IsCacheInProgress();
+}
+
+FSlateColor SFindInBlueprints::GetCachingBarColor() const
+{
+	// The caching bar's default color is a darkish red
+	FSlateColor ReturnColor = FSlateColor(FLinearColor(0.4f, 0.0f, 0.0f));
+	if(IsCacheInProgress())
 	{
-		return 1;
+		// It turns yellow when in progress
+		ReturnColor = FSlateColor(FLinearColor(0.4f, 0.4f, 0.0f));
 	}
+	return ReturnColor;
 }
 
-bool SFindInBlueprints::StringMatchesSearchTokens(const TArray<FString>& Tokens, const FString& ComparisonString)
+FText SFindInBlueprints::GetUncachedBlueprintWarningText() const
 {
-	bool bFoundAllTokens = true;
-	//search the entry for each token, it must have all of them to pass
-	for(auto TokIT(Tokens.CreateConstIterator());TokIT;++TokIT)
+	FFormatNamedArguments Args;
+	Args.Add(TEXT("Count"), FFindInBlueprintSearchManager::Get().GetNumberUncachedBlueprints());
+
+	FText ReturnDisplayText;
+	if(IsCacheInProgress())
 	{
-		const FString& Token = *TokIT;
-		if(!ComparisonString.Contains(Token))
-		{
-			bFoundAllTokens = false;
-			break;
-		}
+		Args.Add(TEXT("CurrentIndex"), FFindInBlueprintSearchManager::Get().GetCurrentCacheIndex());
+
+		ReturnDisplayText = FText::Format(LOCTEXT("CachingBlueprints", "Indexing Blueprints... {CurrentIndex}/{Count}"), Args);
+	}
+	else
+	{
+		ReturnDisplayText = FText::Format(LOCTEXT("UncachedBlueprints", "Search incomplete. {Count} Blueprints are not indexed!"), Args);
 	}
-	return bFoundAllTokens;
+
+	return ReturnDisplayText;
+}
+
+FText SFindInBlueprints::GetCurrentCacheBlueprintName() const
+{
+	return FText::FromString(FFindInBlueprintSearchManager::Get().GetCurrentCacheBlueprintName());
 }
 
-/////////////////////////////////////////////////////
+void SFindInBlueprints::OnCacheComplete()
+{
+	// Resubmit the last search, which will also remove the bar if needed
+	OnSearchTextCommitted(SearchTextField->GetText(), ETextCommit::OnEnter);
+}
 
 #undef LOCTEXT_NAMESPACE
diff --git a/Engine/Source/Editor/Kismet/Public/BlueprintEditor.h b/Engine/Source/Editor/Kismet/Public/BlueprintEditor.h
--- a/Engine/Source/Editor/Kismet/Public/BlueprintEditor.h
+++ b/Engine/Source/Editor/Kismet/Public/BlueprintEditor.h
@@ -132,6 +132,7 @@ public:
 	// FAssetEditorToolkit interface
 	virtual bool OnRequestClose() override;
 	virtual void ToolkitBroughtToFront() override;
+	virtual void SaveAsset_Execute() override;
 	// End of FAssetEditorToolkit 
 
 	// IToolkit interface
@@ -999,4 +1000,4 @@ private:
 
 };
 
-#undef LOCTEXT_NAMESPACE
\ No newline at end of file
+#undef LOCTEXT_NAMESPACE
diff --git a/Engine/Source/Editor/Kismet/Public/FindInBlueprintManager.h b/Engine/Source/Editor/Kismet/Public/FindInBlueprintManager.h
new file mode 100644
--- /dev/null
+++ b/Engine/Source/Editor/Kismet/Public/FindInBlueprintManager.h
@@ -0,0 +1,268 @@
+// Copyright 1998-2014 Epic Games, Inc. All Rights Reserved.
+
+#pragma once
+
+/**
+ *Const values for Find-in-Blueprints to tag searchable data
+ */
+struct KISMET_API FFindInBlueprintSearchTags
+{
+	/** Properties tag, for Blueprint variables */
+	static const FText FiB_Properties;
+
+	/** Nodes tag */
+	static const FText FiB_Nodes;
+
+	/** Schema Name tag, to identify the schema that a graph uses */
+	static const FText FiB_SchemaName;
+	/** Uber graphs tag */
+	static const FText FiB_UberGraphs;
+	/** Function graph tag */
+	static const FText FiB_Functions;
+	/** Macro graph tag */
+	static const FText FiB_Macros;
+	/** Sub graph tag, for any sub-graphs in a Blueprint */
+	static const FText FiB_SubGraphs;
+
+	/** Name tag */
+	static const FText FiB_Name;
+	/** Class Name tag */
+	static const FText FiB_ClassName;
+	/** NodeGuid tag */
+	static const FText FiB_NodeGuid;
+	/** Default value */
+	static const FText FiB_DefaultValue;
+	/** Tooltip tag */
+	static const FText FiB_Tooltip;
+	/** Description tag */
+	static const FText FiB_Description;
+	/** Comment tag */
+	static const FText FiB_Comment;
+
+	/** Pin type tags */
+
+	/** Pins tag */
+	static const FText FiB_Pins;
+
+	/** Pin Category tag */
+	static const FText FiB_PinCategory;
+	/** Pin Sub-Category tag */
+	static const FText FiB_PinSubCategory;
+	/** Pin object class tag */
+	static const FText FiB_ObjectClass;
+	/** Pin IsArray tag */
+	static const FText FiB_IsArray;
+	/** Pin IsReference tag */
+	static const FText FiB_IsReference;
+	/** Glyph icon tag */
+	static const FText FiB_Glyph;
+	/** Glyph icon color tag */
+	static const FText FiB_GlyphColor;
+	/** End const values for Find-in-Blueprint */
+};
+
+/** Tracks data relevant to a Blueprint for searches */
+struct FSearchData
+{
+	/** The Blueprint this search data points to, if available */
+	UBlueprint* Blueprint;
+
+	/** The full Blueprint path this search data is associated with */
+	FString BlueprintPath;
+
+	/** Search data block for the Blueprint */
+	FString Value;
+
+	/** Cached to determine if the Blueprint is seen as no longer valid, allows it to be cleared out next save to disk */
+	bool bMarkedForDeletion;
+
+	/** Retrieval ID from the DDC, set to INDEX_None when retrieval is not necessary */
+	int32 DDCRetrievalID;
+
+	FSearchData()
+		: bMarkedForDeletion(false)
+		, DDCRetrievalID(INDEX_NONE)
+		, Blueprint(nullptr)
+	{
+
+	}
+};
+
+/** Used for external gather functions to add Key/Value pairs to be placed into Json */
+struct FSearchTagDataPair
+{
+	FSearchTagDataPair(FText InKey, FText InValue)
+		: Key(InKey)
+		, Value(InValue)
+	{}
+
+	FText Key;
+	FText Value;
+};
+
+/** Singleton manager for handling all Blueprint searches, helps to manage the going progress of Blueprints, and is thread-safe. */
+class KISMET_API FFindInBlueprintSearchManager
+{
+public:
+	static FFindInBlueprintSearchManager* Instance;
+	static FFindInBlueprintSearchManager& Get();
+
+	FFindInBlueprintSearchManager();
+	~FFindInBlueprintSearchManager();
+
+	/**
+	 * Gathers the Blueprint's search metadata and adds or updates it in the cache
+	 *
+	 * @param InBlueprint		Blueprint to cache the searchable data for
+	 * @param bInForceReCache	Forces the Blueprint to be recache'd, regardless of what data it believes exists
+	 */
+	void AddOrUpdateBlueprintSearchMetadata(UBlueprint* InBlueprint, bool bInForceReCache = false);
+
+	/**
+	 * Starts a search query, the FiB manager handles where the thread is at in the search query at all times so that post-save of the cache to disk it can correct the index
+	 *
+	 * @param InSearchOriginator		Pointer to the thread object that the query originates from
+	 */
+	void BeginSearchQuery(const class FStreamSearch* InSearchOriginator);
+
+	/**
+	 * Continues a search query, returning a single piece of search data
+	 *
+	 * @param InSearchOriginator		Pointer to the thread object that the query originates from
+	 * @param OutSearchData				Result of the search, if there is any Blueprint search data still available to query
+	 * @return							TRUE if the search was successful, FALSE if the search is complete
+	 */
+	bool ContinueSearchQuery(const class FStreamSearch* InSearchOriginator, FSearchData& OutSearchData);
+
+	/**
+	 * This function ensures that the passed in search query ends in a safe manner. The search will no longer be valid to this manager, though it does not destroy any objects.
+	 * Use this whenever the search is finished or canceled.
+	 *
+	 * @param InSearchOriginator	Identifying search stream to be stopped
+	 */
+	void EnsureSearchQueryEnds(const class FStreamSearch* InSearchOriginator);
+
+	/**
+	 * Query how far along a search thread is
+	 *
+	 * @param OutSearchData				Result of the search, if there is any Blueprint search data still available to query
+	 * @return							Percent along the search thread is
+	 */
+	float GetPercentComplete(const class FStreamSearch* InSearchOriginator) const;
+
+	/**
+	 * Query for a single, specific Blueprint's search block
+	 *
+	 * @param InBlueprint		The Blueprint to search for
+	 * @return					The search block
+	 */
+	FString QuerySingleBlueprint(UBlueprint* InBlueprint);
+
+	/** Returns the number of uncached Blueprints */
+	int32 GetNumberUncachedBlueprints() const { return UncachedBlueprints.Num(); };
+
+	/**
+	 * Starts caching all uncached Blueprints at a rate of 1 per tick
+	 *
+	 * @param InSourceWidget		The source FindInBlueprints widget, this widget will be informed when caching is complete
+	 */
+	void CacheAllUncachedBlueprints(TWeakPtr< class SFindInBlueprints > InSourceWidget);
+
+	/** Stops the caching process where it currently is at, the rest can be continued later */
+	void CancelCacheAll();
+
+	/** Returns the current index in the caching */
+	int32 GetCurrentCacheIndex() const;
+
+	/** Returns the name of the current Blueprint being cached */
+	FString GetCurrentCacheBlueprintName() const;
+
+	/** Returns the progress complete on the caching */
+	float GetCacheProgress() const;
+
+	/**
+	 * Callback to note that Blueprint caching is complete
+	 *
+	 * @param InNumberCached		The number of Blueprints cached, to be chopped off the existing array so the rest (if any) can be finished later
+	 */
+	void FinishedCachingBlueprints(int32 InNumberCached);
+
+	/** Returns TRUE if Blueprints are being cached. */
+	bool IsCacheInProgress() const;
+
+	/** Converts a string of hex characters, previously converted by ConvertFTextToHexString, to an FText. */
+	static FText ConvertHexStringToFText(FString InHexString);
+
+	/** Serializes an FText to memory and converts the memory into a string of hex characters */
+	static FString ConvertFTextToHexString(FText InValue);
+private:
+	/** Initializes the FiB manager */
+	void Initialize();
+
+	/** Callback hook during pre-garbage collection, pauses all processing searches so that the GC can do it's job */ 
+	void PauseFindInBlueprintSearch();
+
+	/** Callback hook during post-garbage collection, saves the cache to disk and unpauses all processing searches */
+	void UnpauseFindInBlueprintSearch();
+
+	/** Callback hook from the Asset Registry when an asset is added */
+	void OnAssetAdded(const class FAssetData& InAssetData);
+
+	/** Callback hook from the Asset Registry, marks the asset for deletion from the cache */
+	void OnAssetRemoved(const class FAssetData& InAssetData);
+
+	/** Callback hook from the Asset Registry, marks the asset for deletion from the cache */
+	void OnAssetRenamed(const class FAssetData& InAssetData, const FString& InOldName);
+
+	/** Callback hook from the Asset Registry when an asset is loaded */
+	void OnAssetLoaded(class UObject* InAsset);
+
+	/** Helper to gathers the Blueprint's search metadata */
+	FString GatherBlueprintSearchMetadata(const UBlueprint* Blueprint);
+
+	/** Cleans the cache of any excess data from Blueprints that have been moved, renamed, or deleted. Occurs during post-garbage collection */
+	void CleanCache();
+
+	/** Builds the cache from all available Blueprint assets that the asset registry has discovered at the time of this function. Occurs on startup */
+	void BuildCache();
+
+	/** Retrieves Find-in-Blueprint data from the DDC */
+	void RetrieveDDCData(FSearchData& InOutSearchData, FName InBlueprintPath);
+
+protected:
+	/** Maps the Blueprint paths to their index in the SearchArray */
+	TMap<FString, int> SearchMap;
+
+	/** Stores the Blueprint search data and is used to iterate over in small chunks */
+	TArray<FSearchData> SearchArray;
+
+	/** Counter of active searches */
+	FThreadSafeCounter ActiveSearchCounter;
+
+	/** A mapping of active search queries and where they are currently at in the search data */
+	TMap< const class FStreamSearch*, int32 > ActiveSearchQueries;
+
+	/** Critical section to safely add, remove, and find data in ActiveSearchQueries */
+	mutable FCriticalSection SafeQueryModifyCriticalSection;
+
+	/** Critical section to lock threads during the pausing procedure */
+	FCriticalSection PauseThreadsCriticalSection;
+
+	/** Critical section to safely modify cached data */
+	FCriticalSection SafeModifyCacheCriticalSection;
+
+	/** TRUE when the the FiB manager wants to pause all searches, helps manage the pausing procedure */
+	volatile bool bIsPausing;
+
+	/** Because we are unable to query for the module on another thread, cache it for use later */
+	class FAssetRegistryModule* AssetRegistryModule;
+
+	/** Blueprint paths that have not been cached for searching due to lack of data, this means that they are either older Blueprints, or the DDC cannot find the data */
+	TArray<FName> UncachedBlueprints;
+
+	/** Tickable object that does the caching of uncached Blueprints at a rate of once per tick */
+	class FCacheAllBlueprintsTickableObject* CachingObject;
+
+	/** A list of all the FindInBlueprints widgets that need to be informed that caching is complete */
+	TArray< TWeakPtr<class SFindInBlueprints> > OnCachingCompleteCallbackWidgets;
+};
\ No newline at end of file
diff --git a/Engine/Source/Editor/Kismet/Public/FindInBlueprintUtils.h b/Engine/Source/Editor/Kismet/Public/FindInBlueprintUtils.h
deleted file mode 100644
diff --git a/Engine/Source/Editor/Kismet/Public/FindInBlueprints.h b/Engine/Source/Editor/Kismet/Public/FindInBlueprints.h
--- a/Engine/Source/Editor/Kismet/Public/FindInBlueprints.h
+++ b/Engine/Source/Editor/Kismet/Public/FindInBlueprints.h
@@ -2,63 +2,209 @@
 
 #pragma once
 
+typedef TSharedPtr<class FFindInBlueprintsResult> FSearchResult;
+typedef STreeView<FSearchResult>  STreeViewType;
+
 /* Item that matched the search results */
-class FFindInBlueprintsResult
+class FFindInBlueprintsResult : public TSharedFromThis< FFindInBlueprintsResult >
 {
 public: 
 	/* Create a root */
-	FFindInBlueprintsResult(const FString& InValue);
+	FFindInBlueprintsResult(const FText& InDisplayText);
 	
 	/* Create a listing for a search result*/
-	FFindInBlueprintsResult(const FString& InValue, TSharedPtr<FFindInBlueprintsResult>& InParent, UClass* InClass, int InDuplicationIndex );
-	
-	/* Create a listing for a pin result */
-	FFindInBlueprintsResult(const FString& InValue, TSharedPtr<FFindInBlueprintsResult>& InParent, UEdGraphPin* InPin);
-	
-	/* Create a listing for a pin result */
-	FFindInBlueprintsResult(const FString& InValue, TSharedPtr<FFindInBlueprintsResult>& InParent, UEdGraphNode* InNode);
+	FFindInBlueprintsResult(const FText& InDisplayText, TSharedPtr<FFindInBlueprintsResult> InParent);
 	
 	/* Called when user clicks on the search item */
-	FReply OnClick();
+	virtual FReply OnClick();
 
 	/* Get Category for this search result */
-	FString GetCategory() const;
+	virtual FString GetCategory() const;
 
 	/* Create an icon to represent the result */
-	TSharedRef<SWidget>	CreateIcon() const;
+	virtual TSharedRef<SWidget>	CreateIcon() const;
+
+	/** Finalizes any content for the search data that was unsafe to do on a separate thread */
+	virtual void FinalizeSearchData() {};
 
 	/* Gets the comment on this node if any */
 	FString GetCommentText() const;
-	
-	/* Gets the value of the pin if any */
-	FString GetValueText() const;
 
 	/** gets the blueprint housing all these search results */
 	UBlueprint* GetParentBlueprint() const;
 
+	/**
+	 * Extracts the content of a JsonNode, building the sub-hierarchy for display
+	 *
+	 * @param InJsonNode		Node to extract data from
+	 * @param InTokens			Tokens to match searchable content against
+	 * @return					Returns TRUE if the content matches search results and was extracted
+	 */
+	virtual bool ExtractContent(TSharedPtr< FJsonObject > InJsonNode, const TArray<FString>& InTokens) { return ExtractContent(InJsonNode, InTokens, NULL); };
+
+	/**
+	 * Extracts the content of a JsonNode, building the sub-hierarchy for display
+	 *
+	 * @param InJsonNode		Node to extract data from
+	 * @param InTokens			Tokens to match searchable content against
+	 * @param InParentOverride	Parent search result to attach sub-content to
+	 * @return					Returns TRUE if the content matches search results and was extracted
+	 */
+	virtual bool ExtractContent(TSharedPtr< FJsonObject > InJsonNode, const TArray<FString>& InTokens, TSharedPtr< FFindInBlueprintsResult > InParentOverride);
+
+	/**
+	 * Parses search info for specific data important for displaying the search result in an easy to understand format
+	 *
+	 * @param	InTokens		The search tokens to check results against
+	 * @param	InKey			This is the tag for the data, describing what it is so special handling can occur if needed
+	 * @param	InValue			Compared against search query to see if it passes the filter, sometimes data is rejected because it is deemed unsearchable
+	 * @param	InParent		The parent search result
+	 *
+	 * @return					TRUE if the item passes the search filter or is otherwise accepted
+	 */
+	virtual bool ParseSearchInfo(const TArray<FString> &InTokens, FString InKey, FText InValue, TSharedPtr< FFindInBlueprintsResult > InParent);
+
+	/**
+	 * Adds extra search info, anything that does not have a predestined place in the search result. Adds a sub-item to the searches and formats its description so the tag displays
+	 *
+	 * @param	InKey			This is the tag for the data, describing what it is so special handling can occur if needed
+	 * @param	InValue			Compared against search query to see if it passes the filter, sometimes data is rejected because it is deemed unsearchable
+	 * @param	InParent		The parent search result
+	 */
+	void AddExtraSearchInfo(FString InKey, FText InValue, TSharedPtr< FFindInBlueprintsResult > InParent);
+
+	/**
+	 * Iterates through all this node's children, and tells the tree view to expand them
+	 */
+	void ExpandAllChildren( TSharedPtr< STreeView< TSharedPtr< FFindInBlueprintsResult > > > InTreeView );
+
+	/** Returns the display string for the row */
+	FString GetDisplayString() const;
+
+protected:
+
+	/**
+	 * Extracts the details of a JsonValue in the search results, pushing arrays into sub-objects correctly, breaking up objects, and passing on the rest of the types
+	 *
+	 * @param	InTokens				The search tokens to check results against
+	 * @param	InKey					This is the tag for the data, describing what it is so special handling can occur if needed
+	 * @param	InValue					Has all the data for this object that needs to be deciphered or broken up into base types to be deciphered
+	 * @param	InParentOverride		The parent to attach all children data to, can be NULL.
+	 */
+	bool ExtractJsonValue(const TArray<FString>& InTokens, FText InKey, TSharedPtr< FJsonValue > InValue, TSharedPtr< FFindInBlueprintsResult > InParentOverride);
+
+public:
 	/*Any children listed under this category */
 	TArray< TSharedPtr<FFindInBlueprintsResult> > Children;
 
 	/*If it exists it is the blueprint*/
 	TWeakPtr<FFindInBlueprintsResult> Parent;
 
-	/*The meta string that was stored in the asset registry for this item */
-	FString Value;
+	/*The display text for this item */
+	FText DisplayText;
+
+	/** Display text for comment information */
+	FString CommentText;
+};
+
+/** Graph nodes use this class to store their data */
+class FFindInBlueprintsGraphNode : public FFindInBlueprintsResult
+{
+public:
+	FFindInBlueprintsGraphNode(const FText& InValue, TSharedPtr<FFindInBlueprintsResult> InParent);
+
+	/** FFindInBlueprintsResult Interface */
+	virtual FReply OnClick() override;
+	virtual TSharedRef<SWidget>	CreateIcon() const override;
+	virtual bool ParseSearchInfo(const TArray<FString> &InTokens, FString InKey, FText InValue, TSharedPtr< FFindInBlueprintsResult > InParentOverride = NULL) override;
+	virtual bool ExtractContent(TSharedPtr< FJsonObject > InJsonNode, const TArray<FString>& InTokens) override;
+	virtual FString GetCategory() const override;
+	virtual void FinalizeSearchData() override;
+	/** End FFindInBlueprintsResult Interface */
+
+private:
+	/** The Node Guid to find when jumping to the node */
+	FGuid NodeGuid;
+
+	/** The glyph brush for this node */
+	const struct FSlateBrush* GlyphBrush;
 
-	/*The blueprint may have multiple instances of whatever we are looking for, this tells us which instance # we refer to*/
-	int	DuplicationIndex;
+	/** The glyph color for this node */
+	FLinearColor GlyphColor;
 
 	/*The class this item refers to */
 	UClass* Class;
 
+	/*The class name this item refers to */
+	FString ClassName;
+};
+
+/** Pins use this class to store their data */
+class FFindInBlueprintsPin : public FFindInBlueprintsResult
+{
+public:
+	FFindInBlueprintsPin(const FText& InValue, TSharedPtr<FFindInBlueprintsResult> InParent, FString InSchemaName);
+
+	/** FFindInBlueprintsResult Interface */
+	virtual TSharedRef<SWidget>	CreateIcon() const override;
+	virtual bool ParseSearchInfo(const TArray<FString> &InTokens, FString InKey, FText InValue, TSharedPtr< FFindInBlueprintsResult > InParentOverride = NULL) override;
+	virtual FString GetCategory() const override;
+	virtual void FinalizeSearchData() override;
+	/** End FFindInBlueprintsResult Interface */
+protected:
+	/** The name of the schema this pin exists under */
+	FString SchemaName;
+
 	/** The pin that this search result refers to */
-	TWeakObjectPtr<UEdGraphPin> Pin;
+	FEdGraphPinType PinType;
 
-	/** The graph node that this search result refers to (if not by asset registry or UK2Node) */
-	TWeakObjectPtr<UEdGraphNode> GraphNode;
+	/** The Pin's sub-category name */
+	FString PinSubCategory;
 
-	/** Display text for comment information */
-	FString CommentText;
+	/** Pin's icon color */
+	FSlateColor IconColor;
+};
+
+/** Property data is stored here */
+class FFindInBlueprintsProperty : public FFindInBlueprintsResult
+{
+public:
+	FFindInBlueprintsProperty(const FText& InValue, TSharedPtr<FFindInBlueprintsResult> InParent);
+
+	/** FFindInBlueprintsResult Interface */
+	virtual TSharedRef<SWidget>	CreateIcon() const override;
+	virtual bool ParseSearchInfo(const TArray<FString> &InTokens, FString InKey, FText InValue, TSharedPtr< FFindInBlueprintsResult > InParentOverride = NULL) override;
+	virtual FString GetCategory() const override;
+	virtual void FinalizeSearchData() override;
+	/** End FFindInBlueprintsResult Interface */
+protected:
+	/** The pin that this search result refers to */
+	FEdGraphPinType PinType;
+
+	/** The default value of a property as a string */
+	FString DefaultValue;
+
+	/** The Property's sub-category name, for the pin display type */
+	FString PinSubCategory;
+};
+
+/** Graphs, such as functions and macros, are stored here */
+class FFindInBlueprintsGraph : public FFindInBlueprintsResult
+{
+public:
+	FFindInBlueprintsGraph(const FText& InValue, TSharedPtr<FFindInBlueprintsResult> InParent, EGraphType InGraphType);
+
+	/** FFindInBlueprintsResult Interface */
+	virtual FReply OnClick() override;
+	virtual TSharedRef<SWidget>	CreateIcon() const override;
+	virtual bool ParseSearchInfo(const TArray<FString> &InTokens, FString InKey, FText InValue, TSharedPtr< FFindInBlueprintsResult > InParentOverride = NULL) override;
+	virtual FString GetCategory() const override;
+	virtual bool ExtractContent(TSharedPtr< FJsonObject > InJsonNode, const TArray<FString>& InTokens) override;
+	/** End FFindInBlueprintsResult Interface */
+protected:
+
+	/** The type of graph this represents */
+	EGraphType GraphType;
 };
 
 
@@ -70,13 +216,18 @@ public:
 	SLATE_END_ARGS()
 
 	void Construct(const FArguments& InArgs, TSharedPtr<class FBlueprintEditor> InBlueprintEditor);
+	~SFindInBlueprints();
+
+	/** Begin SWidget Interface */
+	virtual void Tick( const FGeometry& AllottedGeometry, const double InCurrentTime, const float InDeltaTime ) override;
+	/** End SWidget Interface */
 
 	/** Focuses this widget's search box, and changes the mode as well, and optionally the search terms */
 	void FocusForUse(bool bSetFindWithinBlueprint, FString NewSearchTerms = FString(), bool bSelectFirstResult = false);
 
+	/** Called when caching Blueprints is complete, if this widget initiated the indexing */
+	void OnCacheComplete();
 private:
-	typedef TSharedPtr<FFindInBlueprintsResult> FSearchResult;
-	typedef STreeView<FSearchResult> STreeViewType;
 
 	/*Called when user changes the text they are searching for */
 	void OnSearchTextChanged(const FText& Text);
@@ -93,29 +244,68 @@ private:
 	/* Get the children of a row */
 	void OnGetChildren( FSearchResult InItem, TArray< FSearchResult >& OutChildren );
 
-	/* Called when user clicks on a new result */
-	void OnTreeSelectionChanged(FSearchResult Item, ESelectInfo::Type SelectInfo );
+	/* Called when user double clicks on a new result */
+	void OnTreeSelectionDoubleClicked( FSearchResult Item );
 
 	/* Called when a new row is being generated */
 	TSharedRef<ITableRow> OnGenerateRow(FSearchResult InItem, const TSharedRef<STableViewBase>& OwnerTable);
 
 	/** Begins the search based on the SearchValue */
 	void InitiateSearch();
-
-	/* Find any results that contain all of the tokens, only for find in all blueprints */
-	void MatchTokens(const TArray<FString>& Tokens);
 	
 	/* Find any results that contain all of the tokens, however, only in find within blueprints */
 	void MatchTokensWithinBlueprint(const TArray<FString>& Tokens);
 
-	/* Extract results from a single buffer */
-	void Extract(const FString* Buffer, const TArray<FString> &Tokens, FSearchResult& BlueprintCategory, UClass* Class);
+	/** Launches a thread for streaming more content into the results widget */
+	void LaunchStreamThread(const TArray<FString>& InTokens);
+
+	/** Returns the percent complete on the search for the progress bar */
+	TOptional<float> GetPercentCompleteSearch() const;
+
+	/** Returns the progress bar visiblity */
+	EVisibility GetSearchbarVisiblity() const;
+
+	/** Adds the "cache" bar at the bottom of the Find-in-Blueprints widget, to notify the user that the search is incomplete */
+	void ConditionallyAddCacheBar();
+
+	/** Callback to remove the "cache" bar when a button is pressed */
+	FReply OnRemoveCacheBar();
 
-	/* Read the duplication count stored in the string */
-	static int ReadDuplicateCount(const FString& ValueWCount, FString& OutValue );
+	/** Callback to return the cache bar's display text, informing the user of the situation */
+	FText GetUncachedBlueprintWarningText() const;
+
+	/** Callback to return the cache bar's current indexing Blueprint name */
+	FText GetCurrentCacheBlueprintName() const;
+
+	/** Callback to cache all uncached Blueprints */
+	FReply OnCacheAllBlueprints();
+
+	/** Callback to cancel the caching process */
+	FReply OnCancelCacheAll();
+
+	/** Retrieves the current index of the Blueprint caching process */
+	int32 GetCurrentCacheIndex() const;
+
+	/** Gets the percent complete of the caching process */
+	TOptional<float> GetPercentCompleteCache() const;
+
+	/** Returns the visibility of the caching progress bar, visible when in progress, hidden when not */
+	EVisibility GetCachingProgressBarVisiblity() const;
+
+	/** Returns the visibility of the "Cache All" button, visible when not caching, collapsed when caching is in progress */
+	EVisibility GetCacheAllButtonVisibility() const;
+
+	/** Returns the caching bar's visibility, it goes invisible when there is nothing to be cached. The next search will remove this bar or make it visible again */
+	EVisibility GetCachingBarVisibility() const;
 	
-	/** Determines if a string matches the search tokens */
-	static bool StringMatchesSearchTokens(const TArray<FString>& Tokens, const FString& ComparisonString);
+	/** Returns the visibility of the caching Blueprint name, visible when in progress, collapsed when not */
+	EVisibility GetCachingBlueprintNameVisiblity() const;
+
+	/** Returns TRUE if Blueprint caching is in progress */
+	bool IsCacheInProgress() const;
+
+	/** Returns the color of the caching bar */
+	FSlateColor GetCachingBarColor() const;
 
 private:
 	/** Pointer back to the blueprint editor that owns us */
@@ -141,4 +331,13 @@ private:
 
 	/** Should we search within the current blueprint only (rather than all blueprints) */
 	bool bIsInFindWithinBlueprintMode;
+
+	/** Thread object that searches through Blueprint data on a separate thread */
+	TSharedPtr< class FStreamSearch> StreamSearch;
+
+	/** Vertical box, used to add and remove widgets dynamically */
+	TWeakPtr< SVerticalBox > MainVerticalBox;
+
+	/** Weak pointer to the cache bar slot, so it can be removed */
+	TWeakPtr< SWidget > CacheBarSlot;
 };
\ No newline at end of file
diff --git a/Engine/Source/Editor/MainFrame/Private/Menus/TranslationEditorMenu.h b/Engine/Source/Editor/MainFrame/Private/Menus/TranslationEditorMenu.h
--- a/Engine/Source/Editor/MainFrame/Private/Menus/TranslationEditorMenu.h
+++ b/Engine/Source/Editor/MainFrame/Private/Menus/TranslationEditorMenu.h
@@ -152,7 +152,7 @@ public:
 				FString ArchiveName = FPaths::GetBaseFilename(FPaths::GetPath(ArchiveFileName));
 				FString ManifestName = FPaths::GetBaseFilename(ManifestFileName);
 
-				TSharedPtr<FCulture> Culture = FInternationalization::Get().GetCulture(ArchiveName);
+				TSharedPtr<FCulture, ESPMode::ThreadSafe> Culture = FInternationalization::Get().GetCulture(ArchiveName);
 				if (Culture.IsValid())
 				{
 					ArchiveName = Culture->GetDisplayName();
diff --git a/Engine/Source/Editor/NewsFeed/Private/Implementation/NewsFeedCache.cpp b/Engine/Source/Editor/NewsFeed/Private/Implementation/NewsFeedCache.cpp
--- a/Engine/Source/Editor/NewsFeed/Private/Implementation/NewsFeedCache.cpp
+++ b/Engine/Source/Editor/NewsFeed/Private/Implementation/NewsFeedCache.cpp
@@ -89,7 +89,7 @@ void FNewsFeedCache::LoadTitleFile( )
 	FString CultureString;
 
 	// get current culture
-	TSharedRef<FCulture> Culture = FInternationalization::Get().GetCurrentCulture();
+	TSharedRef<FCulture, ESPMode::ThreadSafe> Culture = FInternationalization::Get().GetCurrentCulture();
 
 	if (CurrentCultureSpec == ENewsFeedCultureSpec::Full)
 	{
diff --git a/Engine/Source/Editor/UnrealEd/Classes/Commandlets/GatherTextFromAssetsCommandlet.h b/Engine/Source/Editor/UnrealEd/Classes/Commandlets/GatherTextFromAssetsCommandlet.h
--- a/Engine/Source/Editor/UnrealEd/Classes/Commandlets/GatherTextFromAssetsCommandlet.h
+++ b/Engine/Source/Editor/UnrealEd/Classes/Commandlets/GatherTextFromAssetsCommandlet.h
@@ -41,7 +41,7 @@ private:
 		struct FEntry
 		{
 			FString ObjectPath;
-			TSharedPtr<FString> SourceString;
+			TSharedPtr<FString, ESPMode::ThreadSafe> SourceString;
 			EAssetTextGatherStatus::Type Status;
 		};
 
diff --git a/Engine/Source/Editor/UnrealEd/Classes/Commandlets/InternationalizationExportCommandlet.h b/Engine/Source/Editor/UnrealEd/Classes/Commandlets/InternationalizationExportCommandlet.h
--- a/Engine/Source/Editor/UnrealEd/Classes/Commandlets/InternationalizationExportCommandlet.h
+++ b/Engine/Source/Editor/UnrealEd/Classes/Commandlets/InternationalizationExportCommandlet.h
@@ -74,7 +74,7 @@ public:
 	/**
 	 * Retrieves the culture equivalent that matches the set language code.
 	 */
-	TSharedPtr< FCulture > GetCulture() const;
+	TSharedPtr<FCulture, ESPMode::ThreadSafe> GetCulture() const;
 
 	/**
 	 * Retrieves the display name for the language.
@@ -97,7 +97,7 @@ private:
 	FString LanguageCode;
 	FString LanguagePluralForms;
 
-	TSharedPtr< FCulture > Culture;
+	TSharedPtr<FCulture, ESPMode::ThreadSafe> Culture;
 };
 
 
diff --git a/Engine/Source/Editor/UnrealEd/Private/Commandlets/GatherTextFromAssetsCommandlet.cpp b/Engine/Source/Editor/UnrealEd/Private/Commandlets/GatherTextFromAssetsCommandlet.cpp
--- a/Engine/Source/Editor/UnrealEd/Private/Commandlets/GatherTextFromAssetsCommandlet.cpp
+++ b/Engine/Source/Editor/UnrealEd/Private/Commandlets/GatherTextFromAssetsCommandlet.cpp
@@ -519,8 +519,8 @@ bool UGatherTextFromAssetsCommandlet::ProcessTextProperty(UTextProperty* InTextP
 		NewEntry.Status = EAssetTextGatherStatus::None;
 
 		// Fix missing key if broken and allowed.
-		TSharedPtr< FString > Namespace;
-		TSharedPtr< FString > Key;
+		TSharedPtr< FString, ESPMode::ThreadSafe > Namespace;
+		TSharedPtr< FString, ESPMode::ThreadSafe > Key;
 		FTextLocalizationManager::Get().FindKeyNamespaceFromDisplayString(Data->DisplayString, Namespace, Key);
 
 		if( !( Key.IsValid() ) || Key->IsEmpty() )
@@ -564,7 +564,7 @@ bool UGatherTextFromAssetsCommandlet::ProcessTextProperty(UTextProperty* InTextP
 			if( ExistingEntry.IsValid() )
 			{
 				// Fix conflict if present and allowed.
-				if( ExistingEntry->Source.Text != ( Data->GetSourceString().IsValid() ? **(Data->GetSourceString()) : TEXT("") ) )
+				if( ExistingEntry->Source.Text != ( NewEntry.SourceString.IsValid() ? **(NewEntry.SourceString) : TEXT("") ) )
 				{
 					if (bFixBroken)
 					{
@@ -589,7 +589,7 @@ bool UGatherTextFromAssetsCommandlet::ProcessTextProperty(UTextProperty* InTextP
 			if( !( ExistingEntry.IsValid() ) )
 			{
 				// Check for valid string.
-				if( Data->GetSourceString().IsValid() && !( Data->GetSourceString()->IsEmpty() ) )
+				if( NewEntry.SourceString.IsValid() && !( NewEntry.SourceString->IsEmpty() ) )
 				{
 					FString SrcLocation = Object->GetPathName() + TEXT(".") + InTextProp->GetName();
 
@@ -604,8 +604,8 @@ bool UGatherTextFromAssetsCommandlet::ProcessTextProperty(UTextProperty* InTextP
 							{
 								UTextProperty* TextProp	=  Cast<UTextProperty>( *(PropIt) );
 								FText* DataCDO = TextProp->ContainerPtrToValuePtr<FText>( CDO );
-								TSharedPtr< FString > NamespaceCDO;
-								TSharedPtr< FString > KeyCDO;
+								TSharedPtr< FString, ESPMode::ThreadSafe > NamespaceCDO;
+								TSharedPtr< FString, ESPMode::ThreadSafe > KeyCDO;
 								FTextLocalizationManager::Get().FindKeyNamespaceFromDisplayString(DataCDO->DisplayString, NamespaceCDO, KeyCDO);
 
 								if( KeyCDO == Key || ( KeyCDO.Get() && Key.Get() && ( *KeyCDO == *Key ) ) )
@@ -622,7 +622,7 @@ bool UGatherTextFromAssetsCommandlet::ProcessTextProperty(UTextProperty* InTextP
 					Context.SourceLocation = SrcLocation;
 
 					FString EntryDescription = FString::Printf( TEXT("In %s"), *Object->GetFullName());
-					ManifestInfo->AddEntry(EntryDescription, Namespace.Get() ? *Namespace : TEXT(""), Data->GetSourceString().Get() ? *(Data->GetSourceString()) : TEXT(""), Context );
+					ManifestInfo->AddEntry(EntryDescription, Namespace.Get() ? *Namespace : TEXT(""), NewEntry.SourceString.Get() ? *(NewEntry.SourceString) : TEXT(""), Context );
 				}
 			}
 		}
diff --git a/Engine/Source/Editor/UnrealEd/Private/Kismet2/BlueprintEditorUtils.cpp b/Engine/Source/Editor/UnrealEd/Private/Kismet2/BlueprintEditorUtils.cpp
--- a/Engine/Source/Editor/UnrealEd/Private/Kismet2/BlueprintEditorUtils.cpp
+++ b/Engine/Source/Editor/UnrealEd/Private/Kismet2/BlueprintEditorUtils.cpp
@@ -25,6 +25,7 @@
 #include "ScopedTransaction.h"
 #include "Editor/ClassViewer/Public/ClassViewerModule.h"
 #include "Editor/ClassViewer/Public/ClassViewerFilter.h"
+#include "Editor/Kismet/Public/FindInBlueprintManager.h"
 
 #include "BlueprintEditor.h"
 #include "Editor/UnrealEd/Public/Kismet2/Kismet2NameValidators.h"
@@ -1082,6 +1083,13 @@ UClass* FBlueprintEditorUtils::RegenerateBlueprintClass(UBlueprint* Blueprint, U
 			Blueprint->Status = BS_UpToDate;
 		}
 
+#if WITH_EDITOR
+		if(GEditor)
+		{
+			FFindInBlueprintSearchManager::Get().AddOrUpdateBlueprintSearchMetadata(Blueprint);
+		}
+#endif
+
 		if (bReplaceBlueprintWithClass)
 		{
 			FEditoronlyBlueprintHelper::ShouldBeFixed(Blueprint, true);
@@ -2621,7 +2629,7 @@ void FBlueprintEditorUtils::SetBlueprintVariableMetaData(UBlueprint* Blueprint,
 	FBlueprintEditorUtils::MarkBlueprintAsStructurallyModified(Blueprint);
 }
 
-bool FBlueprintEditorUtils::GetBlueprintVariableMetaData(UBlueprint* Blueprint, const FName& VarName, const FName& MetaDataKey, FString& OutMetaDataValue)
+bool FBlueprintEditorUtils::GetBlueprintVariableMetaData(const UBlueprint* Blueprint, const FName& VarName, const FName& MetaDataKey, FString& OutMetaDataValue)
 {
 	const int32 VarIndex = FBlueprintEditorUtils::FindNewVariableIndex(Blueprint, VarName);
 	if (VarIndex == INDEX_NONE)
@@ -2673,7 +2681,7 @@ bool FBlueprintEditorUtils::GetBlueprintVariableMetaData(UBlueprint* Blueprint,
 	}
 	else
 	{
-		FBPVariableDescription& Desc = Blueprint->NewVariables[VarIndex];
+		const FBPVariableDescription& Desc = Blueprint->NewVariables[VarIndex];
 
 		int32 EntryIndex = Desc.FindMetaDataEntryIndexForKey(MetaDataKey);
 		if (EntryIndex != INDEX_NONE)
@@ -3350,7 +3358,7 @@ void FBlueprintEditorUtils::RenameLocalVariable(UBlueprint* InBlueprint, const F
 	}
 }
 
-FBPVariableDescription* FBlueprintEditorUtils::FindLocalVariable(UBlueprint* InBlueprint, const FName& InVariableName)
+FBPVariableDescription* FBlueprintEditorUtils::FindLocalVariable(const UBlueprint* InBlueprint, const FName& InVariableName)
 {
 	// Search through all function entry nodes for a local variable with the passed name
 	TArray<UK2Node_FunctionEntry*> FunctionEntryNodes;
@@ -5970,5 +5978,88 @@ bool FBlueprintEditorUtils::GetFunctionGuidFromClassByFieldName(const UClass* In
 	return false;
 }
 
+void FBlueprintEditorUtils::GetEntryAndResultNodes(const UEdGraph* InGraph, TWeakObjectPtr<class UK2Node_EditablePinBase>& OutEntryNode, TWeakObjectPtr<class UK2Node_EditablePinBase>& OutResultNode)
+{
+	if (InGraph)
+	{
+		// There are a few different potential configurations for editable graphs (FunctionEntry/Result, Tunnel Pairs, etc).
+		// Step through each case until we find one that matches what appears to be in the graph.  This could be improved if
+		// we want to add more robust typing to the graphs themselves
+
+		// Case 1:  Function Entry / Result Pair ------------------
+		TArray<UK2Node_FunctionEntry*> EntryNodes;
+		InGraph->GetNodesOfClass(EntryNodes);
+
+		if (EntryNodes.Num() > 0)
+		{
+			if (EntryNodes[0]->IsEditable())
+			{
+				OutEntryNode = EntryNodes[0];
+
+				// Find a result node
+				TArray<UK2Node_FunctionResult*> ResultNodes;
+				InGraph->GetNodesOfClass(ResultNodes);
+
+				check(ResultNodes.Num() <= 1);
+				UK2Node_FunctionResult* ResultNode = ResultNodes.Num() ? ResultNodes[0] : NULL;
+				// Note:  we assume that if the entry is editable, the result is too (since the entry node is guaranteed to be there on graph creation, but the result isn't)
+				if( ResultNode )
+				{
+					OutResultNode = ResultNode;
+				}
+			}
+		}
+		else
+		{
+			// Case 2:  Tunnel Pair -----------------------------------
+			TArray<UK2Node_Tunnel*> TunnelNodes;
+			InGraph->GetNodesOfClass(TunnelNodes);
+
+			if (TunnelNodes.Num() > 0)
+			{
+				// Iterate over the tunnel nodes, and try to find an entry and exit
+				for (int32 i = 0; i < TunnelNodes.Num(); i++)
+				{
+					UK2Node_Tunnel* Node = TunnelNodes[i];
+					// Composite nodes should never be considered for function entry / exit, since we're searching for a graph's terminals
+					if (Node->IsEditable() && !Node->IsA(UK2Node_Composite::StaticClass()))
+					{
+						if (Node->bCanHaveOutputs)
+						{
+							ensure(!OutEntryNode.IsValid());
+							OutEntryNode = Node;
+						}
+						else if (Node->bCanHaveInputs)
+						{
+							ensure(!OutResultNode.IsValid());
+							OutResultNode = Node;
+						}
+					}
+				}
+			}
+		}
+	}
+}
+
+FText FBlueprintEditorUtils::GetGraphDescription(const UEdGraph* InGraph)
+{
+	TWeakObjectPtr<class UK2Node_EditablePinBase> EntryNode;
+	TWeakObjectPtr<class UK2Node_EditablePinBase> ResultNode;
+	GetEntryAndResultNodes(InGraph, EntryNode, ResultNode);
+
+	UK2Node_EditablePinBase * FunctionEntryNode = EntryNode.Get();
+	if (UK2Node_FunctionEntry* TypedEntryNode = Cast<UK2Node_FunctionEntry>(FunctionEntryNode))
+	{
+		return FText::FromString(TypedEntryNode->MetaData.ToolTip);
+	}
+	else if (UK2Node_Tunnel* TunnelNode = ExactCast<UK2Node_Tunnel>(FunctionEntryNode))
+	{
+		// Must be exactly a tunnel, not a macro instance
+		return FText::FromString(TunnelNode->MetaData.ToolTip);
+	}
+
+	return LOCTEXT( "NoGraphTooltip", "(None)" );
+}
+
 #undef LOCTEXT_NAMESPACE
 
diff --git a/Engine/Source/Editor/UnrealEd/Private/Kismet2/Kismet2.cpp b/Engine/Source/Editor/UnrealEd/Private/Kismet2/Kismet2.cpp
--- a/Engine/Source/Editor/UnrealEd/Private/Kismet2/Kismet2.cpp
+++ b/Engine/Source/Editor/UnrealEd/Private/Kismet2/Kismet2.cpp
@@ -13,8 +13,7 @@
 #include "Kismet2/KismetDebugUtilities.h"
 #include "Toolkits/AssetEditorManager.h"
 #include "Editor/Kismet/Public/BlueprintEditorModule.h"
-#include "Editor/Kismet/Public/FindInBlueprints.h"
-#include "Editor/Kismet/Public/FindInBlueprintUtils.h"
+#include "Editor/Kismet/Public/FindInBlueprintManager.h"
 #include "Toolkits/ToolkitManager.h"
 #include "Editor/KismetCompiler/Public/KismetCompilerModule.h"
 #include "Kismet2/CompilerResultsLog.h"
@@ -309,6 +308,17 @@ void FKismetEditorUtilities::CompileBlueprint(UBlueprint* BlueprintObj, bool bIs
 	// compiling the blueprint will inherently dirty the package, but if there 
 	// weren't any changes to save before, there shouldn't be after
 	bool const bStartedWithUnsavedChanges = (BlueprintPackage != NULL) ? BlueprintPackage->IsDirty() : true;
+#if WITH_EDITOR
+	if(GEditor)
+	{
+		// We do not want to regenerate a search Guid during loads, nothing has changed in the Blueprint and it is cached elsewhere
+		if(!bIsRegeneratingOnLoad)
+		{
+			BlueprintObj->SearchGuid = FGuid::NewGuid();
+			FFindInBlueprintSearchManager::Get().AddOrUpdateBlueprintSearchMetadata(BlueprintObj);
+		}
+	}
+#endif
 
 	// The old class is either the GeneratedClass if we had an old successful compile, or the SkeletonGeneratedClass stub if there were previously fatal errors
 	UClass* OldClass = (BlueprintObj->GeneratedClass != NULL && (BlueprintObj->GeneratedClass != BlueprintObj->SkeletonGeneratedClass)) ? BlueprintObj->GeneratedClass : NULL;
@@ -1325,103 +1335,6 @@ void FKismetEditorUtilities::StripExternalComponents(class UBlueprint* Blueprint
 	GeneratedCDO->Serialize(InvalidateRefsAr);
 }
 
-
-/* Add search meta information for a given class type*/
-template<class NodeType>
-inline void AddSearchMetaInfo( TArray<UObject::FAssetRegistryTag> &OutTags, const FName& Category, const UBlueprint* Blueprint)
-{
-	TArray<NodeType*> Nodes;
-	FBlueprintEditorUtils::GetAllNodesOfClass(Blueprint, Nodes );
-
-	//To avoid storing duplicate entries in the asset registry, if duplicates are found we indicate how many by appending "::#"
-	TMap<FString,int> AllPaths;
-	for(auto It(Nodes.CreateConstIterator());It;++It)
-	{
-		NodeType* Node = *It;
-
-		// Make sure we don't add MD for nodes that are going away soon
-		if( Node->GetOuter()->IsPendingKill() )
-		{
-			continue;
-		}
-
-		// If the result is empty, the node has no data to offer.
-		FString Result = FindInBlueprintsUtil::GetNodeTypePath(Node);
-		if(Result.IsEmpty())
-		{
-			continue;
-		}
-		Result += FString("::") + Node->GetNodeTitle(ENodeTitleType::ListView).ToString();
-
-		if(int* Count = AllPaths.Find(Result))
-		{
-			(*Count)++;
-		}
-		else
-		{
-			AllPaths.Add(Result,1);
-		}
-	}
-
-	FString FinalString;
-	for(auto It(AllPaths.CreateConstIterator());It;++It)
-	{
-		FString Result = *It.Key();
-		if(It.Value()>1)
-		{
-			Result = FString::Printf(TEXT("%s::%i"), *It.Key(), int(It.Value()));
-		}
-		FinalString += Result + ",";
-
-	}
-	OutTags.Add( UObject::FAssetRegistryTag(Category, FinalString, UObject::FAssetRegistryTag::TT_Hidden) );
-}
-
-void FKismetEditorUtilities::AddInterfaceTags(const UBlueprint* Blueprint, TArray<UObject::FAssetRegistryTag>& OutTags)
-{
-	if (Blueprint->ImplementedInterfaces.Num() == 0)
-	{
-		return;
-	}
-	for (auto Iter(Blueprint->ImplementedInterfaces.CreateConstIterator()); Iter; Iter++)
-	{
-		const FBPInterfaceDescription& Description = *Iter;
-		if (Description.Interface->GetClass() != NULL)
-		{
-			FString Result = FString::Printf(TEXT("[%s] %s"), *Iter->Interface->GetSuperClass()->GetName(), *Iter->Interface->GetOuter()->GetName());
-			OutTags.Add(UObject::FAssetRegistryTag(TEXT("Interfaces"), Result, UObject::FAssetRegistryTag::TT_Hidden));
-		}
-	}
-}
-
-void FKismetEditorUtilities::GetAssetRegistryTagsForBlueprint(const UBlueprint* Blueprint, TArray<UObject::FAssetRegistryTag>& OutTags)
-{
-	//Add information about which functions are called
-	AddSearchMetaInfo<UK2Node_CallFunction>(OutTags, TEXT("CallsFunctions"), Blueprint);
-
-	//Add information about which macro instances are used
-	AddSearchMetaInfo<UK2Node_MacroInstance>(OutTags, TEXT("MacroInstances"), Blueprint);
-
-	//Add information about which events are used 
-	AddSearchMetaInfo<UK2Node_Event>(OutTags, TEXT("ImplementsFunction"), Blueprint);
-
-	//Add information about which variable gets are used
-	AddSearchMetaInfo<UK2Node_VariableGet>(OutTags, TEXT("VariableGet"), Blueprint);
-
-	//Add information about which variable sets are used
-	AddSearchMetaInfo<UK2Node_VariableSet>(OutTags, TEXT("VariableSet"), Blueprint);
-
-	//Add information about which delegates
-	AddSearchMetaInfo<UK2Node_BaseMCDelegate>(OutTags, TEXT("MulticastDelegate"), Blueprint);
-	AddSearchMetaInfo<UK2Node_DelegateSet>(OutTags, TEXT("DelegateSet"), Blueprint);
-
-	//Add information for all node comments
-	AddSearchMetaInfo<UEdGraphNode>(OutTags, TEXT("Comments"), Blueprint);
-	
-	//Add information about interfaces (these are not nodes as the other tags are and so require a custom search)
-	AddInterfaceTags(Blueprint, OutTags);
-}
-
 bool FKismetEditorUtilities::IsTrackedBlueprintParent(const UClass* ParentClass)
 {
 	if (ParentClass->ClassGeneratedBy == NULL)
diff --git a/Engine/Source/Editor/UnrealEd/Public/Kismet2/BlueprintEditorUtils.h b/Engine/Source/Editor/UnrealEd/Public/Kismet2/BlueprintEditorUtils.h
--- a/Engine/Source/Editor/UnrealEd/Public/Kismet2/BlueprintEditorUtils.h
+++ b/Engine/Source/Editor/UnrealEd/Public/Kismet2/BlueprintEditorUtils.h
@@ -27,7 +27,6 @@ namespace EGraphRemoveFlags
 	};
 };
 
-
 class UNREALED_API FBlueprintEditorUtils
 {
 public:
@@ -270,6 +269,29 @@ public:
 		}
 	}
 
+	/**
+	 * Searches all nodes in a Blueprint and checks for a matching Guid
+	 *
+	 * @param InBlueprint			The Blueprint to search
+	 * @param InNodeGuid			The Guid to check Blueprints against
+	 *
+	 * @return						Returns a Node with a matching Guid
+	 */
+	static UEdGraphNode* GetNodeByGUID(const UBlueprint* InBlueprint, const FGuid& InNodeGuid)
+	{
+		TArray<UEdGraphNode*> GraphNodes;
+		GetAllNodesOfClass(InBlueprint, GraphNodes);
+
+		for(UEdGraphNode* GraphNode : GraphNodes)
+		{
+			if(GraphNode->NodeGuid == InNodeGuid)
+			{
+				return GraphNode;
+			}
+		}
+		return nullptr;
+	}
+
 	/** @return true if this Blueprint is dependent on the given blueprint, false otherwise */
 	static bool IsBlueprintDependentOn(UBlueprint const* Blueprint, UBlueprint const* TestBlueprint);
 
@@ -488,7 +510,7 @@ public:
 	 * @param InVariableName	Name of the variable to search for
 	 * @return					The local variable description
 	 */
-	static FBPVariableDescription* FindLocalVariable(UBlueprint* InBlueprint, const FName& InVariableName);
+	static FBPVariableDescription* FindLocalVariable(const UBlueprint* InBlueprint, const FName& InVariableName);
 
 	/**
 	 * Finds a local variable name using the variable's Guid
@@ -574,7 +596,7 @@ public:
 	static void SetBlueprintVariableMetaData(UBlueprint* Blueprint, const FName& VarName, const FName& MetaDataKey, const FString& MetaDataValue);
 
 	/** Get a metadata key/value on the specified variable, or timeline if it exists, returning false if it does not exist */
-	static bool GetBlueprintVariableMetaData(UBlueprint* Blueprint, const FName& VarName, const FName& MetaDataKey, FString& OutMetaDataValue);
+	static bool GetBlueprintVariableMetaData(const UBlueprint* Blueprint, const FName& VarName, const FName& MetaDataKey, FString& OutMetaDataValue);
 
 	/** Clear metadata key on specified variable, or timeline */
 	static void RemoveBlueprintVariableMetaData(UBlueprint* Blueprint, const FName& VarName, const FName& MetaDataKey);
@@ -852,6 +874,22 @@ public:
 	 */
 	static bool IsPaletteActionReadOnly(TSharedPtr<FEdGraphSchemaAction> ActionIn, TSharedPtr<class FBlueprintEditor> const BlueprintEditorIn);
 
+	/**
+	 * Finds the entry and result nodes for a function or macro graph
+	 *
+	 * @param InGraph			The graph to search through
+	 * @param OutEntryNode		The found entry node for the graph
+	 * @param OutResultNode		The found result node for the graph
+	 */
+	static void GetEntryAndResultNodes(const UEdGraph* InGraph, TWeakObjectPtr<class UK2Node_EditablePinBase>& OutEntryNode, TWeakObjectPtr<class UK2Node_EditablePinBase>& OutResultNode);
+
+	/**
+	 * Returns the description of the graph from the metadata
+	 *
+	 * @param InGraph			Graph to find the description of
+	 * @return					The description of the graph
+	 */
+	static FText GetGraphDescription(const UEdGraph* InGraph);
 protected:
 	// Removes all NULL graph references from the SubGraphs array and recurses thru the non-NULL ones
 	static void CleanNullGraphReferencesRecursive(UEdGraph* Graph);
diff --git a/Engine/Source/Editor/UnrealEd/Public/Kismet2/KismetEditorUtilities.h b/Engine/Source/Editor/UnrealEd/Public/Kismet2/KismetEditorUtilities.h
--- a/Engine/Source/Editor/UnrealEd/Public/Kismet2/KismetEditorUtilities.h
+++ b/Engine/Source/Editor/UnrealEd/Public/Kismet2/KismetEditorUtilities.h
@@ -131,9 +131,6 @@ public:
 	/** Run over the components in the blueprint, and then remove any that fall outside this blueprint's scope (e.g. components brought over after reparenting from another class) */
 	static void StripExternalComponents(class UBlueprint* Blueprint);
 
-	/** Generate asset registry tags for this Blueprint */
-	static void GetAssetRegistryTagsForBlueprint(const class UBlueprint* Blueprint, TArray<UObject::FAssetRegistryTag>& OutTags);
-
 	/** Whether or not the specified actor is a valid target for bound events */
 	static bool IsActorValidForLevelScript(const AActor* Actor);
 
diff --git a/Engine/Source/Editor/UnrealEd/UnrealEd.Build.cs b/Engine/Source/Editor/UnrealEd/UnrealEd.Build.cs
--- a/Engine/Source/Editor/UnrealEd/UnrealEd.Build.cs
+++ b/Engine/Source/Editor/UnrealEd/UnrealEd.Build.cs
@@ -107,7 +107,8 @@ public class UnrealEd : ModuleRules
                 "EditorWidgets",
 				"GraphEditor",
 				"Kismet",
-                "InternationalizationSettings"
+                "InternationalizationSettings",
+                "JsonUtilities",
 			}
 		);
 
diff --git a/Engine/Source/Runtime/Core/Private/Internationalization/Cultures/AmericanEnglish.h b/Engine/Source/Runtime/Core/Private/Internationalization/Cultures/AmericanEnglish.h
--- a/Engine/Source/Runtime/Core/Private/Internationalization/Cultures/AmericanEnglish.h
+++ b/Engine/Source/Runtime/Core/Private/Internationalization/Cultures/AmericanEnglish.h
@@ -8,10 +8,10 @@
 class FAmericanEnglishCulture
 {
 public:
-	static TSharedRef< FCulture > Create()
+	static TSharedRef<FCulture, ESPMode::ThreadSafe> Create()
 	{
 #if UE_ENABLE_ICU
-		TSharedRef< FCulture > Culture = MakeShareable( new FCulture( FString("en-US") ) );
+		TSharedRef<FCulture, ESPMode::ThreadSafe> Culture = MakeShareable( new FCulture( FString("en-US") ) );
 #else
 		TArray<FString> InNativeDigits;
 
@@ -131,7 +131,7 @@ public:
 			);
 
 
-		TSharedRef< FCulture > Culture = MakeShareable( new FCulture(
+		TSharedRef<FCulture, ESPMode::ThreadSafe> Culture = MakeShareable( new FCulture(
 			  LOCTEXT("AmericanEnglishDisplayName", "English (United States)")		//const FText DisplayName
 			, FString(TEXT("English (United States)"))			//const FString EnglishName
 			, 1033												//const int KeyboardLayoutId
diff --git a/Engine/Source/Runtime/Core/Private/Internationalization/Cultures/IndiaHindi.h b/Engine/Source/Runtime/Core/Private/Internationalization/Cultures/IndiaHindi.h
--- a/Engine/Source/Runtime/Core/Private/Internationalization/Cultures/IndiaHindi.h
+++ b/Engine/Source/Runtime/Core/Private/Internationalization/Cultures/IndiaHindi.h
@@ -8,10 +8,10 @@
 class FIndiaHindiCulture
 {
 public:
-	static TSharedRef< FCulture > Create()
+	static TSharedRef<FCulture, ESPMode::ThreadSafe> Create()
 	{
 #if UE_ENABLE_ICU
-		TSharedRef< FCulture > Culture = MakeShareable( new FCulture( FString("hi-IN") ) );
+		TSharedRef<FCulture, ESPMode::ThreadSafe> Culture = MakeShareable( new FCulture( FString("hi-IN") ) );
 #else
 		TArray<FString> InNativeDigits;
 
@@ -133,7 +133,7 @@ public:
 			);
 
 
-		TSharedRef< FCulture > Culture = MakeShareable( new FCulture(
+		TSharedRef<FCulture, ESPMode::ThreadSafe> Culture = MakeShareable( new FCulture(
 			  LOCTEXT("HindiDisplayName", "Hindi (India)")					//const FText DisplayName
 			, FString(TEXT("Hindi (India)"))		//const FString EnglishName
 			, 1081									//const int KeyboardLayoutId
diff --git a/Engine/Source/Runtime/Core/Private/Internationalization/Cultures/InvariantCulture.h b/Engine/Source/Runtime/Core/Private/Internationalization/Cultures/InvariantCulture.h
--- a/Engine/Source/Runtime/Core/Private/Internationalization/Cultures/InvariantCulture.h
+++ b/Engine/Source/Runtime/Core/Private/Internationalization/Cultures/InvariantCulture.h
@@ -8,10 +8,10 @@
 class FInvariantCulture
 {
 public:
-	static TSharedRef< FCulture > Create()
+	static TSharedRef<FCulture, ESPMode::ThreadSafe> Create()
 	{
 #if UE_ENABLE_ICU
-		TSharedRef< FCulture > Culture = MakeShareable( new FCulture( FString() ) );
+		TSharedRef<FCulture, ESPMode::ThreadSafe> Culture = MakeShareable( new FCulture( FString() ) );
 #else
 		TArray<FString> InNativeDigits;
 
@@ -131,7 +131,7 @@ public:
 			);
 
 
-		TSharedRef< FCulture > Culture = MakeShareable( new FCulture(
+		TSharedRef<FCulture, ESPMode::ThreadSafe> Culture = MakeShareable( new FCulture(
 			  LOCTEXT("InvariantCultureDisplayName", "Invariant Language (Invariant Country)")		//const FText DisplayName
 			, FString(TEXT("Invariant Language (Invariant Country)"))			//const FString EnglishName
 			, 1033												//const int KeyboardLayoutId
diff --git a/Engine/Source/Runtime/Core/Private/Internationalization/Cultures/JapaneseCulture.h b/Engine/Source/Runtime/Core/Private/Internationalization/Cultures/JapaneseCulture.h
--- a/Engine/Source/Runtime/Core/Private/Internationalization/Cultures/JapaneseCulture.h
+++ b/Engine/Source/Runtime/Core/Private/Internationalization/Cultures/JapaneseCulture.h
@@ -8,10 +8,10 @@
 class FJapaneseCulture
 {
 public:
-	static TSharedRef< FCulture > Create()
+	static TSharedRef<FCulture, ESPMode::ThreadSafe> Create()
 	{
 #if UE_ENABLE_ICU
-		TSharedRef< FCulture > Culture = MakeShareable( new FCulture( FString("ja-JP") ) );
+		TSharedRef<FCulture, ESPMode::ThreadSafe> Culture = MakeShareable( new FCulture( FString("ja-JP") ) );
 #else
 		TArray<FString> NativeDigits;
 
@@ -147,7 +147,7 @@ public:
 			);
 
 
-		TSharedRef< FCulture > Culture = MakeShareable( new FCulture(
+		TSharedRef<FCulture, ESPMode::ThreadSafe> Culture = MakeShareable( new FCulture(
 			  LOCTEXT("JapaneseDisplayName", "Japanese (Japan)")	//const FText DisplayName
 			, FString(TEXT("Japanese (Japan)"))	//const FString EnglishName
 			, 1041	//const int KeyboardLayoutId
diff --git a/Engine/Source/Runtime/Core/Private/Internationalization/Cultures/KoreanCulture.h b/Engine/Source/Runtime/Core/Private/Internationalization/Cultures/KoreanCulture.h
--- a/Engine/Source/Runtime/Core/Private/Internationalization/Cultures/KoreanCulture.h
+++ b/Engine/Source/Runtime/Core/Private/Internationalization/Cultures/KoreanCulture.h
@@ -8,10 +8,10 @@
 class FKoreanCulture
 {
 public:
-	static TSharedRef< FCulture > Create()
+	static TSharedRef<FCulture, ESPMode::ThreadSafe> Create()
 	{
 #if UE_ENABLE_ICU
-		TSharedRef< FCulture > Culture = MakeShareable( new FCulture( FString("ko-KR") ) );
+		TSharedRef<FCulture, ESPMode::ThreadSafe> Culture = MakeShareable( new FCulture( FString("ko-KR") ) );
 #else
 		TArray<FString> NativeDigits;
 
@@ -147,7 +147,7 @@ public:
 			);
 
 
-		TSharedRef< FCulture > Culture = MakeShareable( new FCulture(
+		TSharedRef<FCulture, ESPMode::ThreadSafe> Culture = MakeShareable( new FCulture(
 			  LOCTEXT("KoreanDisplayName", "Korean (Korea)")	//const FText DisplayName
 			, FString(TEXT("Korean (Korea)"))	//const FString EnglishName
 			, 1042	//const int KeyboardLayoutId
diff --git a/Engine/Source/Runtime/Core/Private/Internationalization/Cultures/LEETCulture.h b/Engine/Source/Runtime/Core/Private/Internationalization/Cultures/LEETCulture.h
--- a/Engine/Source/Runtime/Core/Private/Internationalization/Cultures/LEETCulture.h
+++ b/Engine/Source/Runtime/Core/Private/Internationalization/Cultures/LEETCulture.h
@@ -8,10 +8,10 @@
 class FLEETCulture
 {
 public:
-	static TSharedRef< FCulture > Create()
+	static TSharedRef<FCulture, ESPMode::ThreadSafe> Create()
 	{
 #if UE_ENABLE_ICU
-		TSharedRef< FCulture > Culture = MakeShareable( new FCulture( FString() ) );
+		TSharedRef<FCulture, ESPMode::ThreadSafe> Culture = MakeShareable( new FCulture( FString() ) );
 #else
 		TArray<FString> InNativeDigits;
 
@@ -135,7 +135,7 @@ public:
 			);
 
 
-		TSharedRef< FCulture > Culture = MakeShareable( new FCulture(
+		TSharedRef<FCulture, ESPMode::ThreadSafe> Culture = MakeShareable( new FCulture(
 			  LOCTEXT("LEETCultureDisplayName", "LEET Language (LEET Country)")		//const FText DisplayName
 			, FString(TEXT("LEET Language (LEET Country)"))			//const FString EnglishName
 			, 1033												//const int KeyboardLayoutId
diff --git a/Engine/Source/Runtime/Core/Private/Internationalization/ICUCulture.cpp b/Engine/Source/Runtime/Core/Private/Internationalization/ICUCulture.cpp
--- a/Engine/Source/Runtime/Core/Private/Internationalization/ICUCulture.cpp
+++ b/Engine/Source/Runtime/Core/Private/Internationalization/ICUCulture.cpp
@@ -8,10 +8,10 @@
 
 namespace
 {
-	TSharedRef<const icu::Collator> CreateCollator( const icu::Locale& ICULocale )
+	TSharedRef<const icu::Collator, ESPMode::ThreadSafe> CreateCollator( const icu::Locale& ICULocale )
 	{
 		UErrorCode ICUStatus = U_ZERO_ERROR;
-		TSharedPtr<const icu::Collator> Ptr = MakeShareable( icu::Collator::createInstance( ICULocale, ICUStatus ) );
+		TSharedPtr<const icu::Collator, ESPMode::ThreadSafe> Ptr = MakeShareable( icu::Collator::createInstance( ICULocale, ICUStatus ) );
 		checkf(Ptr.IsValid(), TEXT("Creating a collator object failed using locale %s. Perhaps this locale has no data."), StringCast<TCHAR>(ICULocale.getName()).Get());
 		return Ptr.ToSharedRef();
 	}
@@ -196,18 +196,18 @@ FString FCulture::FICUCultureImplementation::GetVariant() const
 	return ICULocale.getVariant();
 }
 
-TSharedRef<const icu::Collator> FCulture::FICUCultureImplementation::GetCollator(const ETextComparisonLevel::Type ComparisonLevel) const
+TSharedRef<const icu::Collator, ESPMode::ThreadSafe> FCulture::FICUCultureImplementation::GetCollator(const ETextComparisonLevel::Type ComparisonLevel) const
 {
 	UErrorCode ICUStatus = U_ZERO_ERROR;
 	const bool bIsDefault = (ComparisonLevel == ETextComparisonLevel::Default);
-	const TSharedRef<const icu::Collator> DefaultCollator( ICUCollator );
+	const TSharedRef<const icu::Collator, ESPMode::ThreadSafe> DefaultCollator( ICUCollator );
 	if(bIsDefault)
 	{
 		return DefaultCollator;
 	}
 	else
 	{
-		const TSharedRef<icu::Collator> Collator( DefaultCollator->clone() );
+		const TSharedRef<icu::Collator, ESPMode::ThreadSafe> Collator( DefaultCollator->clone() );
 		Collator->setAttribute(UColAttribute::UCOL_STRENGTH, UEToICU(ComparisonLevel), ICUStatus);
 		return Collator;
 	}
diff --git a/Engine/Source/Runtime/Core/Private/Internationalization/ICUCulture.h b/Engine/Source/Runtime/Core/Private/Internationalization/ICUCulture.h
--- a/Engine/Source/Runtime/Core/Private/Internationalization/ICUCulture.h
+++ b/Engine/Source/Runtime/Core/Private/Internationalization/ICUCulture.h
@@ -132,7 +132,7 @@ class FCulture::FICUCultureImplementation
 
 	FString GetVariant() const;
 
-	TSharedRef<const icu::Collator> GetCollator(const ETextComparisonLevel::Type ComparisonLevel) const;
+	TSharedRef<const icu::Collator, ESPMode::ThreadSafe> GetCollator(const ETextComparisonLevel::Type ComparisonLevel) const;
 	TSharedRef<const icu::DecimalFormat> GetDecimalFormatter(const FNumberFormattingOptions* const Options = NULL) const;
 	TSharedRef<const icu::DecimalFormat> GetCurrencyFormatter(const FNumberFormattingOptions* const Options = NULL) const;
 	TSharedRef<const icu::DecimalFormat> GetPercentFormatter(const FNumberFormattingOptions* const Options = NULL) const;
@@ -141,7 +141,7 @@ class FCulture::FICUCultureImplementation
 	TSharedRef<const icu::DateFormat> GetDateTimeFormatter(const EDateTimeStyle::Type DateStyle, const EDateTimeStyle::Type TimeStyle, const FString& TimeZone) const;
 
 	icu::Locale ICULocale;
-	const TSharedRef<const icu::Collator> ICUCollator;
+	const TSharedRef<const icu::Collator, ESPMode::ThreadSafe> ICUCollator;
 	const TSharedRef<const icu::DecimalFormat> ICUDecimalFormat;
 	const TSharedRef<const icu::DecimalFormat> ICUCurrencyFormat;
 	const TSharedRef<const icu::DecimalFormat> ICUPercentFormat;
diff --git a/Engine/Source/Runtime/Core/Private/Internationalization/ICUText.cpp b/Engine/Source/Runtime/Core/Private/Internationalization/ICUText.cpp
--- a/Engine/Source/Runtime/Core/Private/Internationalization/ICUText.cpp
+++ b/Engine/Source/Runtime/Core/Private/Internationalization/ICUText.cpp
@@ -12,11 +12,11 @@
 #include <unicode/msgfmt.h>
 #include "ICUUtilities.h"
 
-FText FText::AsDate(const FDateTime& DateTime, const EDateTimeStyle::Type DateStyle, const TSharedPtr<FCulture>& TargetCulture)
+FText FText::AsDate(const FDateTime& DateTime, const EDateTimeStyle::Type DateStyle, const TSharedPtr<FCulture, ESPMode::ThreadSafe>& TargetCulture)
 {
 	FInternationalization& I18N = FInternationalization::Get();
 	checkf(I18N.IsInitialized() == true, TEXT("FInternationalization is not initialized. An FText formatting method was likely used in static object initialization - this is not supported."));
-	TSharedPtr<FCulture> Culture = TargetCulture.IsValid() ? TargetCulture : I18N.GetCurrentCulture();
+	TSharedPtr<FCulture, ESPMode::ThreadSafe> Culture = TargetCulture.IsValid() ? TargetCulture : I18N.GetCurrentCulture();
 
 	int32 UNIXTimestamp = DateTime.ToUnixTimestamp();
 	UDate ICUDate = static_cast<double>(UNIXTimestamp) * U_MILLIS_PER_SECOND;
@@ -35,11 +35,11 @@ FText FText::AsDate(const FDateTime& DateTime, const EDateTimeStyle::Type DateSt
 	return ResultText;
 }
 
-FText FText::AsTime(const FDateTime& DateTime, const EDateTimeStyle::Type TimeStyle, const FString& TimeZone, const TSharedPtr<FCulture>& TargetCulture)
+FText FText::AsTime(const FDateTime& DateTime, const EDateTimeStyle::Type TimeStyle, const FString& TimeZone, const TSharedPtr<FCulture, ESPMode::ThreadSafe>& TargetCulture)
 {
 	FInternationalization& I18N = FInternationalization::Get();
 	checkf(I18N.IsInitialized() == true, TEXT("FInternationalization is not initialized. An FText formatting method was likely used in static object initialization - this is not supported."));
-	TSharedPtr<FCulture> Culture = TargetCulture.IsValid() ? TargetCulture : I18N.GetCurrentCulture();
+	TSharedPtr<FCulture, ESPMode::ThreadSafe> Culture = TargetCulture.IsValid() ? TargetCulture : I18N.GetCurrentCulture();
 
 	int32 UNIXTimestamp = DateTime.ToUnixTimestamp();
 	UDate ICUDate = static_cast<double>(UNIXTimestamp) * U_MILLIS_PER_SECOND;
@@ -58,11 +58,11 @@ FText FText::AsTime(const FDateTime& DateTime, const EDateTimeStyle::Type TimeSt
 	return ResultText;
 }
 
-FText FText::AsTimespan(const FTimespan& Time, const TSharedPtr<FCulture>& TargetCulture)
+FText FText::AsTimespan(const FTimespan& Time, const TSharedPtr<FCulture, ESPMode::ThreadSafe>& TargetCulture)
 {
 	FInternationalization& I18N = FInternationalization::Get();
 	checkf(I18N.IsInitialized() == true, TEXT("FInternationalization is not initialized. An FText formatting method was likely used in static object initialization - this is not supported."));
-	TSharedPtr<FCulture> Culture = TargetCulture.IsValid() ? TargetCulture : I18N.GetCurrentCulture();
+	TSharedPtr<FCulture, ESPMode::ThreadSafe> Culture = TargetCulture.IsValid() ? TargetCulture : I18N.GetCurrentCulture();
 
 	FText TimespanFormatPattern = NSLOCTEXT("Timespan", "FormatPattern", "{Hours}:{Minutes}:{Seconds}");
 
@@ -82,11 +82,11 @@ FText FText::AsTimespan(const FTimespan& Time, const TSharedPtr<FCulture>& Targe
 	return FText::Format(TimespanFormatPattern, TimeArguments);
 }
 
-FText FText::AsDateTime(const FDateTime& DateTime, const EDateTimeStyle::Type DateStyle, const EDateTimeStyle::Type TimeStyle, const FString& TimeZone, const TSharedPtr<FCulture>& TargetCulture)
+FText FText::AsDateTime(const FDateTime& DateTime, const EDateTimeStyle::Type DateStyle, const EDateTimeStyle::Type TimeStyle, const FString& TimeZone, const TSharedPtr<FCulture, ESPMode::ThreadSafe>& TargetCulture)
 {
 	FInternationalization& I18N = FInternationalization::Get();
 	checkf(I18N.IsInitialized() == true, TEXT("FInternationalization is not initialized. An FText formatting method was likely used in static object initialization - this is not supported."));
-	TSharedPtr<FCulture> Culture = TargetCulture.IsValid() ? TargetCulture : I18N.GetCurrentCulture();
+	TSharedPtr<FCulture, ESPMode::ThreadSafe> Culture = TargetCulture.IsValid() ? TargetCulture : I18N.GetCurrentCulture();
 
 	int32 UNIXTimestamp = DateTime.ToUnixTimestamp();
 	UDate ICUDate = static_cast<double>(UNIXTimestamp) * U_MILLIS_PER_SECOND;
@@ -105,7 +105,7 @@ FText FText::AsDateTime(const FDateTime& DateTime, const EDateTimeStyle::Type Da
 	return ResultText;
 }
 
-FText FText::AsMemory(SIZE_T NumBytes, const FNumberFormattingOptions* const Options, const TSharedPtr<FCulture>& TargetCulture)
+FText FText::AsMemory(SIZE_T NumBytes, const FNumberFormattingOptions* const Options, const TSharedPtr<FCulture, ESPMode::ThreadSafe>& TargetCulture)
 {
 	checkf(FInternationalization::Get().IsInitialized() == true, TEXT("FInternationalization is not initialized. An FText formatting method was likely used in static object initialization - this is not supported."));
 	FFormatNamedArguments Args;
@@ -135,7 +135,7 @@ int32 FText::CompareTo( const FText& Other, const ETextComparisonLevel::Type Com
 {
 	UErrorCode ICUStatus = U_ZERO_ERROR;
 
-	const TSharedRef<const icu::Collator> Collator( FInternationalization::Get().GetCurrentCulture()->Implementation->GetCollator(ComparisonLevel) );
+	const TSharedRef<const icu::Collator, ESPMode::ThreadSafe> Collator( FInternationalization::Get().GetCurrentCulture()->Implementation->GetCollator(ComparisonLevel) );
 
 	icu::UnicodeString A;
 	ICUUtilities::Convert(DisplayString.Get(), A);
@@ -155,7 +155,7 @@ bool FText::EqualTo( const FText& Other, const ETextComparisonLevel::Type Compar
 {
 	UErrorCode ICUStatus = U_ZERO_ERROR;
 
-	const TSharedRef<const icu::Collator> Collator( FInternationalization::Get().GetCurrentCulture()->Implementation->GetCollator(ComparisonLevel) );
+	const TSharedRef<const icu::Collator, ESPMode::ThreadSafe> Collator( FInternationalization::Get().GetCurrentCulture()->Implementation->GetCollator(ComparisonLevel) );
 
 	icu::UnicodeString A;
 	ICUUtilities::Convert(DisplayString.Get(), A);
@@ -200,7 +200,7 @@ public:
 
 private:
 	const ETextComparisonLevel::Type ComparisonLevel;
-	const TSharedRef<const icu::Collator> ICUCollator;
+	const TSharedRef<const icu::Collator, ESPMode::ThreadSafe> ICUCollator;
 };
 
 FText::FSortPredicate::FSortPredicate(const ETextComparisonLevel::Type ComparisonLevel)
@@ -258,7 +258,7 @@ void FText::GetFormatPatternParameters(const FText& Pattern, TArray<FString>& Pa
 	}
 }
 
-FText FText::FormatInternal(const FText& Pattern, const FFormatNamedArguments& Arguments, bool bInRebuildText)
+FText FText::FormatInternal(const FText& Pattern, const FFormatNamedArguments& Arguments, bool bInRebuildText, bool bInRebuildAsSource)
 {
 	checkf(FInternationalization::Get().IsInitialized() == true, TEXT("FInternationalization is not initialized. An FText formatting method was likely used in static object initialization - this is not supported."));
 	//SCOPE_CYCLE_COUNTER( STAT_TextFormat );
@@ -310,7 +310,7 @@ FText FText::FormatInternal(const FText& Pattern, const FFormatNamedArguments& A
 					ArgumentValue.TextValue->Rebuild();
 				}
 
-				FString NativeStringValue = ArgumentValue.TextValue->ToString();
+				FString NativeStringValue = bInRebuildAsSource? ArgumentValue.TextValue->BuildSourceString() : ArgumentValue.TextValue->ToString();
 				icu::UnicodeString ICUStringValue;
 				ICUUtilities::Convert(NativeStringValue, ICUStringValue, false);
 				ArgumentValues.Add( icu::Formattable( ICUStringValue ) );
@@ -319,7 +319,7 @@ FText FText::FormatInternal(const FText& Pattern, const FFormatNamedArguments& A
 		}
 	}
 
-	const FString& NativePatternString = Pattern.ToString();
+	const FString& NativePatternString = bInRebuildAsSource? Pattern.BuildSourceString() : Pattern.ToString();
 	icu::UnicodeString ICUPatternString;
 	ICUUtilities::Convert(NativePatternString, ICUPatternString);
 
@@ -347,7 +347,7 @@ FText FText::FormatInternal(const FText& Pattern, const FFormatNamedArguments& A
 	return ResultText;
 }
 
-FText FText::FormatInternal(const FText& Pattern, const FFormatOrderedArguments& Arguments, bool bInRebuildText)
+FText FText::FormatInternal(const FText& Pattern, const FFormatOrderedArguments& Arguments, bool bInRebuildText, bool bInRebuildAsSource)
 {
 	checkf(FInternationalization::Get().IsInitialized() == true, TEXT("FInternationalization is not initialized. An FText formatting method was likely used in static object initialization - this is not supported."));
 	//SCOPE_CYCLE_COUNTER( STAT_TextFormat );
@@ -391,7 +391,7 @@ FText FText::FormatInternal(const FText& Pattern, const FFormatOrderedArguments&
 					ArgumentValue.TextValue->Rebuild();
 				}
 
-				FString NativeStringValue = ArgumentValue.TextValue->ToString();
+				FString NativeStringValue = bInRebuildAsSource? ArgumentValue.TextValue->BuildSourceString() : ArgumentValue.TextValue->ToString();
 				icu::UnicodeString ICUStringValue;
 				ICUUtilities::Convert(NativeStringValue, ICUStringValue, false);
 				ArgumentValues.Add( icu::Formattable( ICUStringValue ) );
@@ -400,7 +400,7 @@ FText FText::FormatInternal(const FText& Pattern, const FFormatOrderedArguments&
 		}
 	}
 
-	const FString& NativePatternString = Pattern.ToString();
+	const FString& NativePatternString = bInRebuildAsSource? Pattern.BuildSourceString() : Pattern.ToString();
 	icu::UnicodeString ICUPatternString;
 	ICUUtilities::Convert(NativePatternString, ICUPatternString);
 
@@ -433,7 +433,7 @@ FText FText::FormatInternal(const FText& Pattern, const FFormatOrderedArguments&
 	return ResultText;
 }
 
-FText FText::FormatInternal(const FText& Pattern, const TArray< FFormatArgumentData > InArguments, bool bInRebuildText)
+FText FText::FormatInternal(const FText& Pattern, const TArray< FFormatArgumentData > InArguments, bool bInRebuildText, bool bInRebuildAsSource)
 {
 	checkf(FInternationalization::Get().IsInitialized() == true, TEXT("FInternationalization is not initialized. An FText formatting method was likely used in static object initialization - this is not supported."));
 	//SCOPE_CYCLE_COUNTER( STAT_TextFormat );
@@ -485,7 +485,7 @@ FText FText::FormatInternal(const FText& Pattern, const TArray< FFormatArgumentD
 					ArgumentValue.TextValue->Rebuild();
 				}
 
-				FString StringValue = ArgumentValue.TextValue->ToString();
+				FString StringValue = bInRebuildAsSource? ArgumentValue.TextValue->BuildSourceString() : ArgumentValue.TextValue->ToString();
 				icu::UnicodeString ICUStringValue;
 				ICUUtilities::Convert(StringValue, ICUStringValue, false);
 				ArgumentValues.Add( ICUStringValue );
@@ -494,7 +494,7 @@ FText FText::FormatInternal(const FText& Pattern, const TArray< FFormatArgumentD
 		}
 	}
 
-	const FString& PatternString = Pattern.ToString();
+	const FString& PatternString = bInRebuildAsSource? Pattern.BuildSourceString() : Pattern.ToString();
 	icu::UnicodeString ICUPatternString;
 	ICUUtilities::Convert(PatternString, ICUPatternString, false);
 
@@ -529,17 +529,17 @@ FText FText::FormatInternal(const FText& Pattern, const TArray< FFormatArgumentD
 
 FText FText::Format(const FText& Pattern, const FFormatNamedArguments& Arguments)
 {
-	return FormatInternal(Pattern, Arguments, false);
+	return FormatInternal(Pattern, Arguments, false, false);
 }
 
 FText FText::Format(const FText& Pattern, const FFormatOrderedArguments& Arguments)
 {
-	return FormatInternal(Pattern, Arguments, false);
+	return FormatInternal(Pattern, Arguments, false, false);
 }
 
 FText FText::Format(const FText& Pattern, const TArray< struct FFormatArgumentData > InArguments)
 {
-	return FormatInternal(Pattern, InArguments, false);
+	return FormatInternal(Pattern, InArguments, false, false);
 }
 
 #endif
diff --git a/Engine/Source/Runtime/Core/Private/Internationalization/ICUText.h b/Engine/Source/Runtime/Core/Private/Internationalization/ICUText.h
--- a/Engine/Source/Runtime/Core/Private/Internationalization/ICUText.h
+++ b/Engine/Source/Runtime/Core/Private/Internationalization/ICUText.h
@@ -11,11 +11,11 @@
 #include "unicode/unistr.h"
 
 template<typename T1, typename T2>
-FText FText::AsNumberTemplate(T1 Val, const FNumberFormattingOptions* const Options, const TSharedPtr<FCulture>& TargetCulture)
+FText FText::AsNumberTemplate(T1 Val, const FNumberFormattingOptions* const Options, const TSharedPtr<FCulture, ESPMode::ThreadSafe>& TargetCulture)
 {
 	FInternationalization& I18N = FInternationalization::Get();
 	checkf(I18N.IsInitialized() == true, TEXT("FInternationalization is not initialized. An FText formatting method was likely used in static object initialization - this is not supported."));
-	const TSharedRef<FCulture> Culture = TargetCulture.IsValid() ? TargetCulture.ToSharedRef() : I18N.GetCurrentCulture();
+	const TSharedRef<FCulture, ESPMode::ThreadSafe> Culture = TargetCulture.IsValid() ? TargetCulture.ToSharedRef() : I18N.GetCurrentCulture();
 	UErrorCode ICUStatus = U_ZERO_ERROR;
 	const TSharedRef<const icu::DecimalFormat> ICUDecimalFormat( Culture->Implementation->GetDecimalFormatter(Options) );
 	icu::Formattable FormattableVal(static_cast<T2>(Val));
@@ -31,11 +31,11 @@ FText FText::AsNumberTemplate(T1 Val, const FNumberFormattingOptions* const Opti
 }
 
 template<typename T1, typename T2>
-FText FText::AsCurrencyTemplate(T1 Val, const FNumberFormattingOptions* const Options, const TSharedPtr<FCulture>& TargetCulture)
+FText FText::AsCurrencyTemplate(T1 Val, const FNumberFormattingOptions* const Options, const TSharedPtr<FCulture, ESPMode::ThreadSafe>& TargetCulture)
 {
 	FInternationalization& I18N = FInternationalization::Get();
 	checkf(I18N.IsInitialized() == true, TEXT("FInternationalization is not initialized. An FText formatting method was likely used in static object initialization - this is not supported."));
-	const TSharedRef<FCulture> Culture = TargetCulture.IsValid() ? TargetCulture.ToSharedRef() : I18N.GetCurrentCulture();
+	const TSharedRef<FCulture, ESPMode::ThreadSafe> Culture = TargetCulture.IsValid() ? TargetCulture.ToSharedRef() : I18N.GetCurrentCulture();
 	UErrorCode ICUStatus = U_ZERO_ERROR;
 	const TSharedRef<const icu::DecimalFormat> ICUDecimalFormat( Culture->Implementation->GetCurrencyFormatter(Options) );
 	icu::Formattable FormattableVal(static_cast<T2>(Val));
@@ -51,11 +51,11 @@ FText FText::AsCurrencyTemplate(T1 Val, const FNumberFormattingOptions* const Op
 }
 
 template<typename T1, typename T2>
-FText FText::AsPercentTemplate(T1 Val, const FNumberFormattingOptions* const Options, const TSharedPtr<FCulture>& TargetCulture)
+FText FText::AsPercentTemplate(T1 Val, const FNumberFormattingOptions* const Options, const TSharedPtr<FCulture, ESPMode::ThreadSafe>& TargetCulture)
 {
 	FInternationalization& I18N = FInternationalization::Get();
 	checkf(I18N.IsInitialized() == true, TEXT("FInternationalization is not initialized. An FText formatting method was likely used in static object initialization - this is not supported."));
-	const TSharedRef<FCulture> Culture = TargetCulture.IsValid() ? TargetCulture.ToSharedRef() : I18N.GetCurrentCulture();
+	const TSharedRef<FCulture, ESPMode::ThreadSafe> Culture = TargetCulture.IsValid() ? TargetCulture.ToSharedRef() : I18N.GetCurrentCulture();
 	UErrorCode ICUStatus = U_ZERO_ERROR;
 	const TSharedRef<const icu::DecimalFormat> ICUDecimalFormat( Culture->Implementation->GetPercentFormatter(Options) );
 	icu::Formattable FormattableVal(static_cast<T2>(Val));
@@ -68,4 +68,4 @@ FText FText::AsPercentTemplate(T1 Val, const FNumberFormattingOptions* const Opt
 	FText ReturnText = FText::CreateNumericalText(NativeString);
 	ReturnText.History = MakeShareable(new FTextHistory_AsPercent(Val, Options, TargetCulture));
 	return ReturnText;
-}
+}
\ No newline at end of file
diff --git a/Engine/Source/Runtime/Core/Private/Internationalization/Internationalization.cpp b/Engine/Source/Runtime/Core/Private/Internationalization/Internationalization.cpp
--- a/Engine/Source/Runtime/Core/Private/Internationalization/Internationalization.cpp
+++ b/Engine/Source/Runtime/Core/Private/Internationalization/Internationalization.cpp
@@ -107,15 +107,15 @@ void FInternationalization::SetCurrentCulture(const FString& Name)
 	}
 }
 
-TSharedRef< FCulture > FInternationalization::GetCurrentCulture() const
+TSharedRef<FCulture, ESPMode::ThreadSafe> FInternationalization::GetCurrentCulture() const
 {
 	return AllCultures[CurrentCultureIndex];
 }
 
-TSharedPtr< FCulture > FInternationalization::GetCulture(const FString& Name) const
+TSharedPtr<FCulture, ESPMode::ThreadSafe> FInternationalization::GetCulture(const FString& Name) const
 {
 	int32 CultureIndex = GetCultureIndex(Name);
-	return CultureIndex != -1 ? AllCultures[CultureIndex] : TSharedPtr<FCulture>();
+	return CultureIndex != -1 ? AllCultures[CultureIndex] : TSharedPtr<FCulture, ESPMode::ThreadSafe>();
 }
 
 int32 FInternationalization::GetCultureIndex(const FString& Name) const
@@ -553,7 +553,7 @@ void FInternationalization::GetCultureNames(TArray<FString>& CultureNames) const
 	}
 }
 
-void FInternationalization::GetCulturesWithAvailableLocalization(const TArray<FString>& InLocalizationPaths, TArray< TSharedPtr<FCulture> >& OutAvailableCultures) const
+void FInternationalization::GetCulturesWithAvailableLocalization(const TArray<FString>& InLocalizationPaths, TArray< TSharedPtr<FCulture, ESPMode::ThreadSafe> >& OutAvailableCultures) const
 {
 	OutAvailableCultures.Reset();
 
@@ -628,11 +628,11 @@ void FInternationalization::PopulateAllCultures(void)
 	const FString DefaultLocaleName = FPlatformMisc::GetDefaultLocale();
 
 	// Find the default locale as a culture, if present..
-	auto FindDefaultPredicate = [DefaultLocaleName](const TSharedRef< FCulture >& Culture) -> bool
+	auto FindDefaultPredicate = [DefaultLocaleName](const TSharedRef<FCulture, ESPMode::ThreadSafe>& Culture) -> bool
 	{
 		return Culture->GetName() == DefaultLocaleName;
 	};
-	TSharedRef< FCulture>* FoundDefaultCulture = AllCultures.FindByPredicate(FindDefaultPredicate);
+	TSharedRef<FCulture, ESPMode::ThreadSafe>* FoundDefaultCulture = AllCultures.FindByPredicate(FindDefaultPredicate);
 	// If present, set default culture variable.
 	if(FoundDefaultCulture)
 	{
@@ -661,4 +661,4 @@ void FInternationalization::PopulateAllCultures(void)
 	InvariantCulture = GetCulture(TEXT(""));
 }
 
-#undef LOCTEXT_NAMESPACE
\ No newline at end of file
+#undef LOCTEXT_NAMESPACE
diff --git a/Engine/Source/Runtime/Core/Private/Internationalization/LegacyText.cpp b/Engine/Source/Runtime/Core/Private/Internationalization/LegacyText.cpp
--- a/Engine/Source/Runtime/Core/Private/Internationalization/LegacyText.cpp
+++ b/Engine/Source/Runtime/Core/Private/Internationalization/LegacyText.cpp
@@ -5,32 +5,32 @@
 #if !UE_ENABLE_ICU
 #include "Text.h"
 
-FText FText::AsDate( const FDateTime& DateTime, const EDateTimeStyle::Type DateStyle, const TSharedPtr<FCulture>& TargetCulture )
+FText FText::AsDate( const FDateTime& DateTime, const EDateTimeStyle::Type DateStyle, const TSharedPtr<FCulture, ESPMode::ThreadSafe>& TargetCulture )
 {
 	checkf(FInternationalization::Get().IsInitialized() == true, TEXT("FInternationalization is not initialized. An FText formatting method was likely used in static object initialization - this is not supported."));
 	return FText::FromString( DateTime.ToString( TEXT("%Y.%m.%d") ) );
 }
 
-FText FText::AsTime( const FDateTime& DateTime, const EDateTimeStyle::Type TimeStyle, const FString& TimeZone, const TSharedPtr<FCulture>& TargetCulture )
+FText FText::AsTime( const FDateTime& DateTime, const EDateTimeStyle::Type TimeStyle, const FString& TimeZone, const TSharedPtr<FCulture, ESPMode::ThreadSafe>& TargetCulture )
 {
 	checkf(FInternationalization::Get().IsInitialized() == true, TEXT("FInternationalization is not initialized. An FText formatting method was likely used in static object initialization - this is not supported."));
 	return FText::FromString( DateTime.ToString( TEXT("%H.%M.%S") ) );
 }
 
-FText FText::AsTimespan( const FTimespan& Time, const TSharedPtr<FCulture>& TargetCulture)
+FText FText::AsTimespan( const FTimespan& Time, const TSharedPtr<FCulture, ESPMode::ThreadSafe>& TargetCulture)
 {
 	checkf(FInternationalization::Get().IsInitialized() == true, TEXT("FInternationalization is not initialized. An FText formatting method was likely used in static object initialization - this is not supported."));
 	FDateTime DateTime(Time.GetTicks());
 	return FText::FromString( DateTime.ToString( TEXT("%H.%M.%S") ) );
 }
 
-FText FText::AsDateTime( const FDateTime& DateTime, const EDateTimeStyle::Type DateStyle, const EDateTimeStyle::Type TimeStyle, const FString& TimeZone, const TSharedPtr<FCulture>& TargetCulture )
+FText FText::AsDateTime( const FDateTime& DateTime, const EDateTimeStyle::Type DateStyle, const EDateTimeStyle::Type TimeStyle, const FString& TimeZone, const TSharedPtr<FCulture, ESPMode::ThreadSafe>& TargetCulture )
 {
 	checkf(FInternationalization::Get().IsInitialized() == true, TEXT("FInternationalization is not initialized. An FText formatting method was likely used in static object initialization - this is not supported."));
 	return FText::FromString(DateTime.ToString(TEXT("%Y.%m.%d-%H.%M.%S")));
 }
 
-FText FText::AsMemory( SIZE_T NumBytes, const FNumberFormattingOptions* const Options, const TSharedPtr<FCulture>& TargetCulture )
+FText FText::AsMemory( SIZE_T NumBytes, const FNumberFormattingOptions* const Options, const TSharedPtr<FCulture, ESPMode::ThreadSafe>& TargetCulture )
 {
 	checkf(FInternationalization::Get().IsInitialized() == true, TEXT("FInternationalization is not initialized. An FText formatting method was likely used in static object initialization - this is not supported."));
 	FFormatNamedArguments Args;
@@ -479,17 +479,17 @@ FText FText::Format(const FText& Pattern, const TArray< FFormatArgumentData > In
 	return FLegacyTextHelper::Format(Pattern, FGetArgumentValue::CreateRaw(&ArgumentGetter, &FArgumentGetter::GetArgumentValue));
 }
 
-FText FText::FormatInternal(const FText& Pattern, const FFormatNamedArguments& Arguments, bool bInRebuildText)
+FText FText::FormatInternal(const FText& Pattern, const FFormatNamedArguments& Arguments, bool bInRebuildText, bool bInRebuildAsSource)
 {
 	return Format(Pattern, Arguments);
 }
 
-FText FText::FormatInternal(const FText& Pattern, const FFormatOrderedArguments& Arguments, bool bInRebuildText)
+FText FText::FormatInternal(const FText& Pattern, const FFormatOrderedArguments& Arguments, bool bInRebuildText, bool bInRebuildAsSource)
 {
 	return Format(Pattern, Arguments);
 }
 
-FText FText::FormatInternal(const FText& Pattern, const TArray< struct FFormatArgumentData > InArguments, bool bInRebuildText)
+FText FText::FormatInternal(const FText& Pattern, const TArray< struct FFormatArgumentData > InArguments, bool bInRebuildText, bool bInRebuildAsSource)
 {
 	return Format(Pattern, InArguments);
 }
diff --git a/Engine/Source/Runtime/Core/Private/Internationalization/LegacyText.h b/Engine/Source/Runtime/Core/Private/Internationalization/LegacyText.h
--- a/Engine/Source/Runtime/Core/Private/Internationalization/LegacyText.h
+++ b/Engine/Source/Runtime/Core/Private/Internationalization/LegacyText.h
@@ -3,28 +3,28 @@
 #pragma once
 
 template<typename T1, typename T2>
-FText FText::AsNumberTemplate(T1 Val, const FNumberFormattingOptions* const Options, const TSharedPtr<FCulture>& TargetCulture)
+FText FText::AsNumberTemplate(T1 Val, const FNumberFormattingOptions* const Options, const TSharedPtr<FCulture, ESPMode::ThreadSafe>& TargetCulture)
 {
 	FInternationalization& I18N = FInternationalization::Get();
 	checkf(I18N.IsInitialized() == true, TEXT("FInternationalization is not initialized. An FText formatting method was likely used in static object initialization - this is not supported."));
-	const TSharedRef<FCulture> Culture = TargetCulture.IsValid() ? TargetCulture.ToSharedRef() : I18N.GetCurrentCulture();
+	const TSharedRef<FCulture, ESPMode::ThreadSafe> Culture = TargetCulture.IsValid() ? TargetCulture.ToSharedRef() : I18N.GetCurrentCulture();
 	return FText::CreateNumericalText( Culture->NumberFormattingRule.AsNumber(Val) );
 }
 
 template<typename T1, typename T2>
-FText FText::AsCurrencyTemplate(T1 Val, const FNumberFormattingOptions* const Options, const TSharedPtr<FCulture>& TargetCulture)
+FText FText::AsCurrencyTemplate(T1 Val, const FNumberFormattingOptions* const Options, const TSharedPtr<FCulture, ESPMode::ThreadSafe>& TargetCulture)
 {
 	FInternationalization& I18N = FInternationalization::Get();
 	checkf(I18N.IsInitialized() == true, TEXT("FInternationalization is not initialized. An FText formatting method was likely used in static object initialization - this is not supported."));
-	const TSharedRef<FCulture> Culture = TargetCulture.IsValid() ? TargetCulture.ToSharedRef() : I18N.GetCurrentCulture();
+	const TSharedRef<FCulture, ESPMode::ThreadSafe> Culture = TargetCulture.IsValid() ? TargetCulture.ToSharedRef() : I18N.GetCurrentCulture();
 	return FText::CreateNumericalText( Culture->NumberFormattingRule.AsCurrency(Val) );
 }
 
 template<typename T1, typename T2>
-FText FText::AsPercentTemplate(T1 Val, const FNumberFormattingOptions* const Options, const TSharedPtr<FCulture>& TargetCulture)
+FText FText::AsPercentTemplate(T1 Val, const FNumberFormattingOptions* const Options, const TSharedPtr<FCulture, ESPMode::ThreadSafe>& TargetCulture)
 {
 	FInternationalization& I18N = FInternationalization::Get();
 	checkf(I18N.IsInitialized() == true, TEXT("FInternationalization is not initialized. An FText formatting method was likely used in static object initialization - this is not supported."));
-	const TSharedRef<FCulture> Culture = TargetCulture.IsValid() ? TargetCulture.ToSharedRef() : I18N.GetCurrentCulture();
+	const TSharedRef<FCulture, ESPMode::ThreadSafe> Culture = TargetCulture.IsValid() ? TargetCulture.ToSharedRef() : I18N.GetCurrentCulture();
 	return FText::CreateNumericalText( Culture->NumberFormattingRule.AsPercent(Val) );
 }
\ No newline at end of file
diff --git a/Engine/Source/Runtime/Core/Private/Internationalization/Text.cpp b/Engine/Source/Runtime/Core/Private/Internationalization/Text.cpp
--- a/Engine/Source/Runtime/Core/Private/Internationalization/Text.cpp
+++ b/Engine/Source/Runtime/Core/Private/Internationalization/Text.cpp
@@ -25,16 +25,16 @@ FArchive& operator<<( FArchive& Ar, FFormatArgumentData& Value )
 
 const FString* FTextInspector::GetNamespace(const FText& Text)
 {
-	TSharedPtr< FString > Namespace;
-	TSharedPtr< FString > Key;
+	TSharedPtr< FString, ESPMode::ThreadSafe > Namespace;
+	TSharedPtr< FString, ESPMode::ThreadSafe > Key;
 	FTextLocalizationManager::Get().FindKeyNamespaceFromDisplayString(Text.DisplayString, Namespace, Key);
 	return Namespace.Get();
 }
 
 const FString* FTextInspector::GetKey(const FText& Text)
 {
-	TSharedPtr< FString > Namespace;
-	TSharedPtr< FString > Key;
+	TSharedPtr< FString, ESPMode::ThreadSafe > Namespace;
+	TSharedPtr< FString, ESPMode::ThreadSafe > Key;
 	FTextLocalizationManager::Get().FindKeyNamespaceFromDisplayString(Text.DisplayString, Namespace, Key);
 	return Key.Get();
 }
@@ -461,7 +461,7 @@ FText FText::Format(const FText& Fmt,const FText& v1,const FText& v2,const FText
 // on some platforms (PS4) int64_t is a typedef of long.  However, UE4 typedefs int64 as long long.  Since these are distinct types, and ICU only has a constructor for int64_t, casting to int64 causes a compile error from ambiguous function call, 
 // so cast to int64_t's where appropriate here to avoid problems.
 
-#define DEF_ASNUMBER_CAST(T1, T2) FText FText::AsNumber(T1 Val, const FNumberFormattingOptions* const Options, const TSharedPtr<FCulture>& TargetCulture) { return FText::AsNumberTemplate<T1, T2>(Val, Options, TargetCulture); }
+#define DEF_ASNUMBER_CAST(T1, T2) FText FText::AsNumber(T1 Val, const FNumberFormattingOptions* const Options, const TSharedPtr<FCulture, ESPMode::ThreadSafe>& TargetCulture) { return FText::AsNumberTemplate<T1, T2>(Val, Options, TargetCulture); }
 #define DEF_ASNUMBER(T) DEF_ASNUMBER_CAST(T, T)
 DEF_ASNUMBER(float)
 DEF_ASNUMBER(double)
@@ -479,7 +479,7 @@ DEF_ASNUMBER_CAST(uint64, int64_t)
 * Generate an FText that represents the passed number as currency in the current culture
 */
 
-#define DEF_ASCURRENCY_CAST(T1, T2) FText FText::AsCurrency(T1 Val, const FNumberFormattingOptions* const Options, const TSharedPtr<FCulture>& TargetCulture) { return FText::AsCurrencyTemplate<T1, T2>(Val, Options, TargetCulture); }
+#define DEF_ASCURRENCY_CAST(T1, T2) FText FText::AsCurrency(T1 Val, const FNumberFormattingOptions* const Options, const TSharedPtr<FCulture, ESPMode::ThreadSafe>& TargetCulture) { return FText::AsCurrencyTemplate<T1, T2>(Val, Options, TargetCulture); }
 #define DEF_ASCURRENCY(T) DEF_ASCURRENCY_CAST(T, T)
 DEF_ASCURRENCY(float)
 	DEF_ASCURRENCY(double)
@@ -497,7 +497,7 @@ DEF_ASCURRENCY(float)
 * Generate an FText that represents the passed number as a percentage in the current culture
 */
 
-#define DEF_ASPERCENT_CAST(T1, T2) FText FText::AsPercent(T1 Val, const FNumberFormattingOptions* const Options, const TSharedPtr<FCulture>& TargetCulture) { return FText::AsPercentTemplate<T1, T2>(Val, Options, TargetCulture); }
+#define DEF_ASPERCENT_CAST(T1, T2) FText FText::AsPercent(T1 Val, const FNumberFormattingOptions* const Options, const TSharedPtr<FCulture, ESPMode::ThreadSafe>& TargetCulture) { return FText::AsPercentTemplate<T1, T2>(Val, Options, TargetCulture); }
 #define DEF_ASPERCENT(T) DEF_ASPERCENT_CAST(T, T)
 DEF_ASPERCENT(double)
 DEF_ASPERCENT(float)
@@ -505,7 +505,7 @@ DEF_ASPERCENT(float)
 
 bool FText::FindText( const FString& Namespace, const FString& Key, FText& OutText, const FString* const SourceString )
 {
-	TSharedPtr< FString > FoundString = FTextLocalizationManager::Get().FindString( Namespace, Key );
+	TSharedPtr< FString, ESPMode::ThreadSafe > FoundString = FTextLocalizationManager::Get().FindString( Namespace, Key );
 
 	if ( FoundString.IsValid() )
 	{
@@ -534,8 +534,8 @@ CORE_API FArchive& operator<<( FArchive& Ar, FText& Value )
 		// Namespaces and keys are no longer stored in the FText, we need to read them in and discard
 		if( Ar.UE4Ver() >= VER_UE4_ADDED_NAMESPACE_AND_KEY_DATA_TO_FTEXT )
 		{
-			TSharedPtr< FString > Namespace;
-			TSharedPtr< FString > Key;
+			TSharedPtr< FString, ESPMode::ThreadSafe > Namespace;
+			TSharedPtr< FString, ESPMode::ThreadSafe > Key;
 
 			Namespace = MakeShareable( new FString() );
 			Ar << *Namespace;
@@ -715,24 +715,29 @@ const FString& FText::ToString() const
 	return DisplayString.Get();
 }
 
+FString FText::BuildSourceString() const
+{
+	return History->ToText(true).ToString();
+}
+
 void FText::Rebuild() const
 {
 	if(History.IsValid() && History->IsOutOfDate(Revision))
 	{
 		Revision = FTextLocalizationManager::Get().GetHeadCultureRevision();
 		
-		DisplayString.Get() = History->ToText().DisplayString.Get();
+		DisplayString.Get() = History->ToText(false).DisplayString.Get();
 	}
 }
 
-TSharedPtr<FString> FText::GetSourceString() const
+TSharedPtr< FString, ESPMode::ThreadSafe > FText::GetSourceString() const
 {
 	if(History.IsValid())
 	{
-		TSharedPtr<FString> SourceString = History->GetSourceString();
+		TSharedPtr< FString, ESPMode::ThreadSafe > SourceString = History->GetSourceString();
 		if(SourceString.IsValid())
 		{
-			return History->GetSourceString();
+			return SourceString;
 		}
 	}
 
diff --git a/Engine/Source/Runtime/Core/Private/Internationalization/TextHistory.cpp b/Engine/Source/Runtime/Core/Private/Internationalization/TextHistory.cpp
--- a/Engine/Source/Runtime/Core/Private/Internationalization/TextHistory.cpp
+++ b/Engine/Source/Runtime/Core/Private/Internationalization/TextHistory.cpp
@@ -13,7 +13,7 @@ bool FTextHistory::IsOutOfDate(int32 InRevision)
 	return InRevision < FTextLocalizationManager::Get().GetHeadCultureRevision();
 }
 
-TSharedPtr< FString > FTextHistory::GetSourceString()
+TSharedPtr< FString, ESPMode::ThreadSafe > FTextHistory::GetSourceString() const
 {
 	return NULL;
 }
@@ -27,21 +27,30 @@ FTextHistory_Base::FTextHistory_Base(FString InSourceString)
 
 }
 
-FTextHistory_Base::FTextHistory_Base(TSharedPtr< FString > InSourceString)
+FTextHistory_Base::FTextHistory_Base(TSharedPtr< FString, ESPMode::ThreadSafe > InSourceString)
 	: SourceString(InSourceString)
 {
 
 }
 
-FText FTextHistory_Base::ToText() const
+FText FTextHistory_Base::ToText(bool bInAsSource) const
 {
-	// This should never be called
+	if(bInAsSource)
+	{
+		if(SourceString.IsValid())
+		{
+			return FText::FromString(*SourceString.Get());
+		}
+		return FText::GetEmpty();
+	}
+
+	// This should never be called except to build as source
 	check(0);
 
 	return FText::GetEmpty();
 }
 
-TSharedPtr< FString > FTextHistory_Base::GetSourceString()
+TSharedPtr< FString, ESPMode::ThreadSafe > FTextHistory_Base::GetSourceString() const
 {
 	return SourceString;
 }
@@ -56,7 +65,7 @@ void FTextHistory_Base::Serialize( FArchive& Ar )
 	}
 }
 
-void FTextHistory_Base::SerializeForDisplayString(FArchive& Ar, TSharedRef<FString>& InOutDisplayString)
+void FTextHistory_Base::SerializeForDisplayString(FArchive& Ar, TSharedRef<FString, ESPMode::ThreadSafe>& InOutDisplayString)
 {
 	if(Ar.IsLoading())
 	{
@@ -79,8 +88,8 @@ void FTextHistory_Base::SerializeForDisplayString(FArchive& Ar, TSharedRef<FStri
 	}
 	else if(Ar.IsSaving())
 	{
-		TSharedPtr< FString > Namespace;
-		TSharedPtr< FString > Key;
+		TSharedPtr< FString, ESPMode::ThreadSafe > Namespace;
+		TSharedPtr< FString, ESPMode::ThreadSafe > Key;
 
 		FTextLocalizationManager::Get().FindKeyNamespaceFromDisplayString(InOutDisplayString, Namespace, Key);
 
@@ -137,9 +146,9 @@ FTextHistory_NamedFormat::FTextHistory_NamedFormat(const FText& InSourceText, co
 
 }
 
-FText FTextHistory_NamedFormat::ToText() const
+FText FTextHistory_NamedFormat::ToText(bool bInAsSource) const
 {
-	return FText::FormatInternal(SourceText, Arguments, true);
+	return FText::FormatInternal(SourceText, Arguments, true, bInAsSource);
 }
 
 void FTextHistory_NamedFormat::Serialize( FArchive& Ar )
@@ -164,9 +173,9 @@ FTextHistory_OrderedFormat::FTextHistory_OrderedFormat(const FText& InSourceText
 
 }
 
-FText FTextHistory_OrderedFormat::ToText() const
+FText FTextHistory_OrderedFormat::ToText(bool bInAsSource) const
 {
-	return FText::FormatInternal(SourceText, Arguments, true);
+	return FText::FormatInternal(SourceText, Arguments, true, bInAsSource);
 }
 
 void FTextHistory_OrderedFormat::Serialize( FArchive& Ar )
@@ -191,9 +200,9 @@ FTextHistory_ArgumentDataFormat::FTextHistory_ArgumentDataFormat(const FText& In
 
 }
 
-FText FTextHistory_ArgumentDataFormat::ToText() const
+FText FTextHistory_ArgumentDataFormat::ToText(bool bInAsSource) const
 {
-	return FText::FormatInternal(SourceText, Arguments, true);
+	return FText::FormatInternal(SourceText, Arguments, true, bInAsSource);
 }
 
 void FTextHistory_ArgumentDataFormat::Serialize( FArchive& Ar )
@@ -216,7 +225,7 @@ FTextHistory_FormatNumber::FTextHistory_FormatNumber()
 {
 }
 
-FTextHistory_FormatNumber::FTextHistory_FormatNumber(const FFormatArgumentValue& InSourceValue, const FNumberFormattingOptions* const InFormatOptions, const TSharedPtr<FCulture> InTargetCulture)
+FTextHistory_FormatNumber::FTextHistory_FormatNumber(const FFormatArgumentValue& InSourceValue, const FNumberFormattingOptions* const InFormatOptions, const TSharedPtr<FCulture, ESPMode::ThreadSafe> InTargetCulture)
 	: SourceValue(InSourceValue)
 	, FormatOptions(nullptr)
 	, TargetCulture(InTargetCulture)
@@ -270,30 +279,32 @@ void FTextHistory_FormatNumber::Serialize(FArchive& Ar)
 ///////////////////////////////////////
 // FTextHistory_AsNumber
 
-FTextHistory_AsNumber::FTextHistory_AsNumber(const FFormatArgumentValue& InSourceValue, const FNumberFormattingOptions* const InFormatOptions, const TSharedPtr<FCulture> InTargetCulture)
+FTextHistory_AsNumber::FTextHistory_AsNumber(const FFormatArgumentValue& InSourceValue, const FNumberFormattingOptions* const InFormatOptions, const TSharedPtr<FCulture, ESPMode::ThreadSafe> InTargetCulture)
 	: FTextHistory_FormatNumber(InSourceValue, InFormatOptions, InTargetCulture)
 {
 }
 
-FText FTextHistory_AsNumber::ToText() const
+FText FTextHistory_AsNumber::ToText(bool bInAsSource) const
 {
+	TSharedPtr< FCulture, ESPMode::ThreadSafe > CurrentCulture = bInAsSource? FInternationalization::Get().GetInvariantCulture() : TargetCulture;
+
 	switch(SourceValue.Type)
 	{
 	case EFormatArgumentType::UInt:
 		{
-			return FText::AsNumber(SourceValue.UIntValue, FormatOptions, TargetCulture);
+			return FText::AsNumber(SourceValue.UIntValue, FormatOptions, CurrentCulture);
 		}
 	case EFormatArgumentType::Int:
 		{
-			return FText::AsNumber(SourceValue.IntValue, FormatOptions, TargetCulture);
+			return FText::AsNumber(SourceValue.IntValue, FormatOptions, CurrentCulture);
 		}
 	case EFormatArgumentType::Float:
 		{
-			return FText::AsNumber(SourceValue.FloatValue, FormatOptions, TargetCulture);
+			return FText::AsNumber(SourceValue.FloatValue, FormatOptions, CurrentCulture);
 		}
 	case EFormatArgumentType::Double:
 		{
-			return FText::AsNumber(SourceValue.DoubleValue, FormatOptions, TargetCulture);
+			return FText::AsNumber(SourceValue.DoubleValue, FormatOptions, CurrentCulture);
 		}
 	default:
 		{
@@ -318,22 +329,24 @@ void FTextHistory_AsNumber::Serialize( FArchive& Ar )
 ///////////////////////////////////////
 // FTextHistory_AsPercent
 
-FTextHistory_AsPercent::FTextHistory_AsPercent(const FFormatArgumentValue& InSourceValue, const FNumberFormattingOptions* const InFormatOptions, const TSharedPtr<FCulture> InTargetCulture)
+FTextHistory_AsPercent::FTextHistory_AsPercent(const FFormatArgumentValue& InSourceValue, const FNumberFormattingOptions* const InFormatOptions, const TSharedPtr<FCulture, ESPMode::ThreadSafe> InTargetCulture)
 	: FTextHistory_FormatNumber(InSourceValue, InFormatOptions, InTargetCulture)
 {
 }
 
-FText FTextHistory_AsPercent::ToText() const
+FText FTextHistory_AsPercent::ToText(bool bInAsSource) const
 {
+	TSharedPtr< FCulture, ESPMode::ThreadSafe > CurrentCulture = bInAsSource? FInternationalization::Get().GetInvariantCulture() : TargetCulture;
+
 	switch(SourceValue.Type)
 	{
 	case EFormatArgumentType::Float:
 		{
-			return FText::AsPercent(SourceValue.FloatValue, FormatOptions, TargetCulture);
+			return FText::AsPercent(SourceValue.FloatValue, FormatOptions, CurrentCulture);
 		}
 	case EFormatArgumentType::Double:
 		{
-			return FText::AsPercent(SourceValue.DoubleValue, FormatOptions, TargetCulture);
+			return FText::AsPercent(SourceValue.DoubleValue, FormatOptions, CurrentCulture);
 		}
 	default:
 		{
@@ -358,30 +371,32 @@ void FTextHistory_AsPercent::Serialize( FArchive& Ar )
 ///////////////////////////////////////
 // FTextHistory_AsCurrency
 
-FTextHistory_AsCurrency::FTextHistory_AsCurrency(const FFormatArgumentValue& InSourceValue, const FNumberFormattingOptions* const InFormatOptions, const TSharedPtr<FCulture> InTargetCulture)
+FTextHistory_AsCurrency::FTextHistory_AsCurrency(const FFormatArgumentValue& InSourceValue, const FNumberFormattingOptions* const InFormatOptions, const TSharedPtr<FCulture, ESPMode::ThreadSafe> InTargetCulture)
 	: FTextHistory_FormatNumber(InSourceValue, InFormatOptions, InTargetCulture)
 {
 }
 
-FText FTextHistory_AsCurrency::ToText() const
+FText FTextHistory_AsCurrency::ToText(bool bInAsSource) const
 {
+	TSharedPtr< FCulture, ESPMode::ThreadSafe > CurrentCulture = bInAsSource? FInternationalization::Get().GetInvariantCulture() : TargetCulture;
+
 	switch(SourceValue.Type)
 	{
 	case EFormatArgumentType::UInt:
 		{
-			return FText::AsCurrency(SourceValue.UIntValue, FormatOptions, TargetCulture);
+			return FText::AsCurrency(SourceValue.UIntValue, FormatOptions, CurrentCulture);
 		}
 	case EFormatArgumentType::Int:
 		{
-			return FText::AsCurrency(SourceValue.IntValue, FormatOptions, TargetCulture);
+			return FText::AsCurrency(SourceValue.IntValue, FormatOptions, CurrentCulture);
 		}
 	case EFormatArgumentType::Float:
 		{
-			return FText::AsCurrency(SourceValue.FloatValue, FormatOptions, TargetCulture);
+			return FText::AsCurrency(SourceValue.FloatValue, FormatOptions, CurrentCulture);
 		}
 	case EFormatArgumentType::Double:
 		{
-			return FText::AsCurrency(SourceValue.DoubleValue, FormatOptions, TargetCulture);
+			return FText::AsCurrency(SourceValue.DoubleValue, FormatOptions, CurrentCulture);
 		}
 	default:
 		{
@@ -406,7 +421,7 @@ void FTextHistory_AsCurrency::Serialize( FArchive& Ar )
 ///////////////////////////////////////
 // FTextHistory_AsDate
 
-FTextHistory_AsDate::FTextHistory_AsDate(const FDateTime& InSourceDateTime, const EDateTimeStyle::Type InDateStyle, const TSharedPtr<FCulture> InTargetCulture)
+FTextHistory_AsDate::FTextHistory_AsDate(const FDateTime& InSourceDateTime, const EDateTimeStyle::Type InDateStyle, const TSharedPtr<FCulture, ESPMode::ThreadSafe> InTargetCulture)
 	: SourceDateTime(InSourceDateTime)
 	, DateStyle(InDateStyle)
 	, TargetCulture(InTargetCulture)
@@ -444,15 +459,16 @@ void FTextHistory_AsDate::Serialize(FArchive& Ar)
 	}
 }
 
-FText FTextHistory_AsDate::ToText() const
+FText FTextHistory_AsDate::ToText(bool bInAsSource) const
 {
-	return FText::AsDate(SourceDateTime, DateStyle, TargetCulture);
-}
+	TSharedPtr< FCulture, ESPMode::ThreadSafe > CurrentCulture = bInAsSource? FInternationalization::Get().GetInvariantCulture() : TargetCulture;
 
+	return FText::AsDate(SourceDateTime, DateStyle, CurrentCulture);
+}
 ///////////////////////////////////////
 // FTextHistory_AsTime
 
-FTextHistory_AsTime::FTextHistory_AsTime(const FDateTime& InSourceDateTime, const EDateTimeStyle::Type InTimeStyle, const FString& InTimeZone, const TSharedPtr<FCulture> InTargetCulture)
+FTextHistory_AsTime::FTextHistory_AsTime(const FDateTime& InSourceDateTime, const EDateTimeStyle::Type InTimeStyle, const FString& InTimeZone, const TSharedPtr<FCulture, ESPMode::ThreadSafe> InTargetCulture)
 	: SourceDateTime(InSourceDateTime)
 	, TimeStyle(InTimeStyle)
 	, TimeZone(InTimeZone)
@@ -493,15 +509,17 @@ void FTextHistory_AsTime::Serialize(FArchive& Ar)
 	}
 }
 
-FText FTextHistory_AsTime::ToText() const
+FText FTextHistory_AsTime::ToText(bool bInAsSource) const
 {
+	TSharedPtr< FCulture, ESPMode::ThreadSafe > CurrentCulture = bInAsSource? FInternationalization::Get().GetInvariantCulture() : TargetCulture;
+	
 	return FText::AsTime(SourceDateTime, TimeStyle, TimeZone, TargetCulture);
 }
 
 ///////////////////////////////////////
 // FTextHistory_AsDateTime
 
-FTextHistory_AsDateTime::FTextHistory_AsDateTime(const FDateTime& InSourceDateTime, const EDateTimeStyle::Type InDateStyle, const EDateTimeStyle::Type InTimeStyle, const FString& InTimeZone, const TSharedPtr<FCulture> InTargetCulture)
+FTextHistory_AsDateTime::FTextHistory_AsDateTime(const FDateTime& InSourceDateTime, const EDateTimeStyle::Type InDateStyle, const EDateTimeStyle::Type InTimeStyle, const FString& InTimeZone, const TSharedPtr<FCulture, ESPMode::ThreadSafe> InTargetCulture)
 	: SourceDateTime(InSourceDateTime)
 	, DateStyle(InDateStyle)
 	, TimeStyle(InTimeStyle)
@@ -547,7 +565,9 @@ void FTextHistory_AsDateTime::Serialize(FArchive& Ar)
 	}
 }
 
-FText FTextHistory_AsDateTime::ToText() const
+FText FTextHistory_AsDateTime::ToText(bool bInAsSource) const
 {
+	TSharedPtr< FCulture, ESPMode::ThreadSafe > CurrentCulture = bInAsSource? FInternationalization::Get().GetInvariantCulture() : TargetCulture;
+	
 	return FText::AsDateTime(SourceDateTime, DateStyle, TimeStyle, TimeZone, TargetCulture);
-}
\ No newline at end of file
+}
diff --git a/Engine/Source/Runtime/Core/Private/Internationalization/TextHistory.h b/Engine/Source/Runtime/Core/Private/Internationalization/TextHistory.h
--- a/Engine/Source/Runtime/Core/Private/Internationalization/TextHistory.h
+++ b/Engine/Source/Runtime/Core/Private/Internationalization/TextHistory.h
@@ -24,19 +24,19 @@ public:
 	virtual ~FTextHistory() {};
 
 	/** Rebuilds the FText from the hierarchical history, the result should be in the current locale */
-	virtual FText ToText() const = 0;
+	virtual FText ToText(bool bInAsSource) const = 0;
 	
 	/** Serializes the history to/from an FArchive */
 	virtual void Serialize(FArchive& Ar) = 0;
 
 	/** Serializes data needed to get the FText's DisplayString, most history classes do not need to do anything here */
-	virtual void SerializeForDisplayString(FArchive& Ar, TSharedRef<FString>& InOutDisplayString) {};
+	virtual void SerializeForDisplayString(FArchive& Ar, TSharedRef<FString, ESPMode::ThreadSafe>& InOutDisplayString) {};
 
 	/** Returns TRUE if the Revision is out of date */
 	virtual bool IsOutOfDate(int32 InRevision);
 
 	/** Returns the source string managed by the history (if any). */
-	virtual TSharedPtr< FString > GetSourceString();
+	virtual TSharedPtr< FString, ESPMode::ThreadSafe > GetSourceString() const;
 };
 
 /** No complexity to it, just holds the source string. */
@@ -45,20 +45,20 @@ class CORE_API FTextHistory_Base : public FTextHistory
 public:
 	FTextHistory_Base() {};
 	FTextHistory_Base(FString InSourceString);
-	FTextHistory_Base(TSharedPtr< FString > InSourceString);
+	FTextHistory_Base(TSharedPtr< FString, ESPMode::ThreadSafe > InSourceString);
 
 
 	// Begin FTextHistory interface
-	virtual FText ToText() const override;
+	virtual FText ToText(bool bInAsSource) const override;
 	virtual void Serialize(FArchive& Ar) override;
-	virtual void SerializeForDisplayString(FArchive& Ar, TSharedRef<FString>& InOutDisplayString) override;
+	virtual void SerializeForDisplayString(FArchive& Ar, TSharedRef<FString, ESPMode::ThreadSafe>& InOutDisplayString) override;
 	virtual bool IsOutOfDate(int32 InRevision) override { return false; }
-	virtual TSharedPtr< FString > GetSourceString() override;
+	virtual TSharedPtr< FString, ESPMode::ThreadSafe > GetSourceString() const override;
 	// End FTextHistory interface
 
 private:
 	/** The source string for an FText */
-	TSharedPtr<FString> SourceString;
+	TSharedPtr<FString, ESPMode::ThreadSafe> SourceString;
 };
 
 /** Handles history for FText::Format when passing named arguments */
@@ -69,7 +69,7 @@ public:
 	FTextHistory_NamedFormat(const FText& InSourceText, const FFormatNamedArguments& InArguments);
 
 	// Begin FTextHistory interface
-	virtual FText ToText() const override;
+	virtual FText ToText(bool bInAsSource) const override;
 	virtual void Serialize(FArchive& Ar) override;
 	// End FTextHistory interface
 
@@ -88,7 +88,7 @@ public:
 	FTextHistory_OrderedFormat(const FText& InSourceText, const FFormatOrderedArguments& InArguments);
 
 	// Begin FTextHistory interface
-	virtual FText ToText() const override;
+	virtual FText ToText(bool bInAsSource) const override;
 	virtual void Serialize(FArchive& Ar) override;
 	// End FTextHistory interface
 
@@ -107,7 +107,7 @@ public:
 	FTextHistory_ArgumentDataFormat(const FText& InSourceText, const TArray< struct FFormatArgumentData >& InArguments);
 
 	// Begin FTextHistory interface
-	virtual FText ToText() const override;
+	virtual FText ToText(bool bInAsSource) const override;
 	virtual void Serialize(FArchive& Ar) override;
 	// End FTextHistory interface
 
@@ -123,7 +123,7 @@ class CORE_API FTextHistory_FormatNumber : public FTextHistory
 {
 public:
 	FTextHistory_FormatNumber();
-	FTextHistory_FormatNumber(const FFormatArgumentValue& InSourceValue, const FNumberFormattingOptions* const InFormatOptions, const TSharedPtr<FCulture> InTargetCulture);
+	FTextHistory_FormatNumber(const FFormatArgumentValue& InSourceValue, const FNumberFormattingOptions* const InFormatOptions, const TSharedPtr<FCulture, ESPMode::ThreadSafe> InTargetCulture);
 
 	~FTextHistory_FormatNumber();
 
@@ -136,7 +136,7 @@ protected:
 	/** All the formatting options available to format using. This can be NULL */
 	FNumberFormattingOptions* FormatOptions;
 	/** The culture to format using */
-	TSharedPtr<FCulture> TargetCulture;
+	TSharedPtr<FCulture, ESPMode::ThreadSafe> TargetCulture;
 };
 
 /**  Handles history for formatting using AsNumber */
@@ -144,10 +144,10 @@ class CORE_API FTextHistory_AsNumber : public FTextHistory_FormatNumber
 {
 public:
 	FTextHistory_AsNumber() {};
-	FTextHistory_AsNumber(const FFormatArgumentValue& InSourceValue, const FNumberFormattingOptions* const InFormatOptions, const TSharedPtr<FCulture> InTargetCulture);
+	FTextHistory_AsNumber(const FFormatArgumentValue& InSourceValue, const FNumberFormattingOptions* const InFormatOptions, const TSharedPtr<FCulture, ESPMode::ThreadSafe> InTargetCulture);
 
 	// Begin FTextHistory interface
-	virtual FText ToText() const override;
+	virtual FText ToText(bool bInAsSource) const override;
 	virtual void Serialize(FArchive& Ar) override;
 	// End FTextHistory interface
 };
@@ -157,10 +157,10 @@ class CORE_API FTextHistory_AsPercent : public FTextHistory_FormatNumber
 {
 public:
 	FTextHistory_AsPercent() {};
-	FTextHistory_AsPercent(const FFormatArgumentValue& InSourceValue, const FNumberFormattingOptions* const InFormatOptions, const TSharedPtr<FCulture> InTargetCulture);
+	FTextHistory_AsPercent(const FFormatArgumentValue& InSourceValue, const FNumberFormattingOptions* const InFormatOptions, const TSharedPtr<FCulture, ESPMode::ThreadSafe> InTargetCulture);
 
 	// Begin FTextHistory interface
-	virtual FText ToText() const override;
+	virtual FText ToText(bool bInAsSource) const override;
 	virtual void Serialize(FArchive& Ar) override;
 	// End FTextHistory interface
 };
@@ -170,10 +170,10 @@ class CORE_API FTextHistory_AsCurrency : public FTextHistory_FormatNumber
 {
 public:
 	FTextHistory_AsCurrency() {};
-	FTextHistory_AsCurrency(const FFormatArgumentValue& InSourceValue, const FNumberFormattingOptions* const InFormatOptions, const TSharedPtr<FCulture> InTargetCulture);
+	FTextHistory_AsCurrency(const FFormatArgumentValue& InSourceValue, const FNumberFormattingOptions* const InFormatOptions, const TSharedPtr<FCulture, ESPMode::ThreadSafe> InTargetCulture);
 
 	// Begin FTextHistory interface
-	virtual FText ToText() const override;
+	virtual FText ToText(bool bInAsSource) const override;
 	virtual void Serialize(FArchive& Ar) override;
 	// End FTextHistory interface
 };
@@ -183,10 +183,10 @@ class CORE_API FTextHistory_AsDate : public FTextHistory
 {
 public:
 	FTextHistory_AsDate() {};
-	FTextHistory_AsDate(const FDateTime& InSourceDateTime, const EDateTimeStyle::Type InDateStyle, const TSharedPtr<FCulture> InTargetCulture);
+	FTextHistory_AsDate(const FDateTime& InSourceDateTime, const EDateTimeStyle::Type InDateStyle, const TSharedPtr<FCulture, ESPMode::ThreadSafe> InTargetCulture);
 
 	// Begin FTextHistory interface
-	virtual FText ToText() const override;
+	virtual FText ToText(bool bInAsSource) const override;
 	virtual void Serialize(FArchive& Ar) override;
 	// End FTextHistory interface
 
@@ -196,7 +196,7 @@ private:
 	/** Style to format the date using */
 	EDateTimeStyle::Type DateStyle;
 	/** Culture to format the date in */
-	TSharedPtr<FCulture> TargetCulture;
+	TSharedPtr<FCulture, ESPMode::ThreadSafe> TargetCulture;
 };
 
 /**  Handles history for formatting using AsTime */
@@ -204,10 +204,10 @@ class CORE_API FTextHistory_AsTime : public FTextHistory
 {
 public:
 	FTextHistory_AsTime() {};
-	FTextHistory_AsTime(const FDateTime& InSourceDateTime, const EDateTimeStyle::Type InTimeStyle, const FString& InTimeZone, const TSharedPtr<FCulture> InTargetCulture);
+	FTextHistory_AsTime(const FDateTime& InSourceDateTime, const EDateTimeStyle::Type InTimeStyle, const FString& InTimeZone, const TSharedPtr<FCulture, ESPMode::ThreadSafe> InTargetCulture);
 
 	// Begin FTextHistory interface
-	virtual FText ToText() const override;
+	virtual FText ToText(bool bInAsSource) const override;
 	virtual void Serialize(FArchive& Ar) override;
 	// End FTextHistory interface
 
@@ -219,7 +219,7 @@ private:
 	/** Timezone to put the time in */
 	FString TimeZone;
 	/** Culture to format the time in */
-	TSharedPtr<FCulture> TargetCulture;
+	TSharedPtr<FCulture, ESPMode::ThreadSafe> TargetCulture;
 };
 
 /**  Handles history for formatting using AsDateTime */
@@ -227,10 +227,10 @@ class CORE_API FTextHistory_AsDateTime : public FTextHistory
 {
 public:
 	FTextHistory_AsDateTime() {};
-	FTextHistory_AsDateTime(const FDateTime& InSourceDateTime, const EDateTimeStyle::Type InDateStyle, const EDateTimeStyle::Type InTimeStyle, const FString& InTimeZone, const TSharedPtr<FCulture> InTargetCulture);
+	FTextHistory_AsDateTime(const FDateTime& InSourceDateTime, const EDateTimeStyle::Type InDateStyle, const EDateTimeStyle::Type InTimeStyle, const FString& InTimeZone, const TSharedPtr<FCulture, ESPMode::ThreadSafe> InTargetCulture);
 
 	// Begin FTextHistory interface
-	virtual FText ToText() const override;
+	virtual FText ToText(bool bInAsSource) const override;
 	virtual void Serialize(FArchive& Ar) override;
 	// End FTextHistory interfaces
 
@@ -244,5 +244,5 @@ private:
 	/** Timezone to put the time in */
 	FString TimeZone;
 	/** Culture to format the time in */
-	TSharedPtr<FCulture> TargetCulture;
-};
\ No newline at end of file
+	TSharedPtr<FCulture, ESPMode::ThreadSafe> TargetCulture;
+};
diff --git a/Engine/Source/Runtime/Core/Private/Internationalization/TextLocalizationManager.cpp b/Engine/Source/Runtime/Core/Private/Internationalization/TextLocalizationManager.cpp
--- a/Engine/Source/Runtime/Core/Private/Internationalization/TextLocalizationManager.cpp
+++ b/Engine/Source/Runtime/Core/Private/Internationalization/TextLocalizationManager.cpp
@@ -87,11 +87,11 @@ void EndInitTextLocalization()
 
 			LocalizationPaths += FPaths::GetEngineLocalizationPaths();
 
-			TArray< TSharedPtr<FCulture> > AvailableCultures;
+			TArray< TSharedPtr<FCulture, ESPMode::ThreadSafe> > AvailableCultures;
 			I18N.GetCulturesWithAvailableLocalization(LocalizationPaths, AvailableCultures);
 
 			// If we don't have an available translation for the current culture, fallback to generic English
-			TSharedPtr<FCulture> TargetCulture = I18N.GetCulture(CultureName);
+			TSharedPtr<FCulture, ESPMode::ThreadSafe> TargetCulture = I18N.GetCulture(CultureName);
 			if(!TargetCulture.IsValid() || !AvailableCultures.Contains(TargetCulture))
 			{
 				UE_LOG(LogTextLocalizationManager, Warning, TEXT("The selected culture '%s' is not available; falling back to 'en'"), *CultureName);
@@ -117,7 +117,7 @@ FTextLocalizationManager& FTextLocalizationManager::Get()
 	return *GTextLocalizationManager;
 }
 
-const TSharedRef<FString>* FTextLocalizationManager::FTextLookupTable::GetString(const FString& Namespace, const FString& Key, const uint32 SourceStringHash) const
+const TSharedRef<FString, ESPMode::ThreadSafe>* FTextLocalizationManager::FTextLookupTable::GetString(const FString& Namespace, const FString& Key, const uint32 SourceStringHash) const
 {
 	const FKeyTable* const KeyTable = NamespaceTable.Find(Namespace);
 	if( KeyTable )
@@ -426,7 +426,7 @@ void FTextLocalizationManager::OnCultureChanged()
 	LoadResources(ShouldLoadEditor, ShouldLoadGame);
 }
 
-TSharedPtr<FString> FTextLocalizationManager::FindString( const FString& Namespace, const FString& Key, const FString* const SourceString )
+TSharedPtr<FString, ESPMode::ThreadSafe> FTextLocalizationManager::FindString( const FString& Namespace, const FString& Key, const FString* const SourceString )
 {
 	FScopeLock ScopeLock( &SynchronizationObject );
 
@@ -444,7 +444,7 @@ TSharedPtr<FString> FTextLocalizationManager::FindString( const FString& Namespa
 	return NULL;
 }
 
-TSharedRef<FString> FTextLocalizationManager::GetString(const FString& Namespace, const FString& Key, const FString* const SourceString)
+TSharedRef<FString, ESPMode::ThreadSafe> FTextLocalizationManager::GetString(const FString& Namespace, const FString& Key, const FString* const SourceString)
 {
 	FScopeLock ScopeLock( &SynchronizationObject );
 
@@ -508,7 +508,7 @@ TSharedRef<FString> FTextLocalizationManager::GetString(const FString& Namespace
 			UE_LOG(LogTextLocalizationManager, Verbose, TEXT("An attempt was made to get a localized string (Namespace:%s, Key:%s, Source:%s), but it did not exist."), *Namespace, *Key, SourceString ? **SourceString : TEXT(""));
 		}
 
-		const TSharedRef<FString> UnlocalizedString = MakeShareable( new FString( SourceString ? **SourceString : TEXT("") ) );
+		const TSharedRef<FString, ESPMode::ThreadSafe> UnlocalizedString = MakeShareable( new FString( SourceString ? **SourceString : TEXT("") ) );
 
 		// If live-culture-swap is enabled or the system is uninitialized - make entries so that they can be updated when system is initialized or a culture swap occurs.
 		CA_SUPPRESS(6236)
@@ -560,8 +560,10 @@ TSharedRef<FString> FTextLocalizationManager::GetString(const FString& Namespace
 	}
 }
 
-void FTextLocalizationManager::FindKeyNamespaceFromDisplayString(TSharedRef<FString> InDisplayString, TSharedPtr<FString>& OutNamespace, TSharedPtr<FString>& OutKey)
+void FTextLocalizationManager::FindKeyNamespaceFromDisplayString(TSharedRef<FString, ESPMode::ThreadSafe> InDisplayString, TSharedPtr<FString, ESPMode::ThreadSafe>& OutNamespace, TSharedPtr<FString, ESPMode::ThreadSafe>& OutKey)
 {
+	FScopeLock ScopeLock( &SynchronizationObject );
+
 	FNamespaceKeyEntry* NamespaceKeyEntry = ReverseLiveTable.Find(InDisplayString);
 
 	if(NamespaceKeyEntry)
diff --git a/Engine/Source/Runtime/Core/Private/Tests/Internationalization/NumberFormatingRulesTest.cpp b/Engine/Source/Runtime/Core/Private/Tests/Internationalization/NumberFormatingRulesTest.cpp
--- a/Engine/Source/Runtime/Core/Private/Tests/Internationalization/NumberFormatingRulesTest.cpp
+++ b/Engine/Source/Runtime/Core/Private/Tests/Internationalization/NumberFormatingRulesTest.cpp
@@ -253,7 +253,7 @@ bool FNumberFormattingRulesTest::RunTest (const FString& Parameters)
 
 #if UE_ENABLE_ICU
 	{
-		const TSharedPtr<FCulture>& InvariantCulture = I18N.GetInvariantCulture();
+		const TSharedPtr<FCulture, ESPMode::ThreadSafe>& InvariantCulture = I18N.GetInvariantCulture();
 		Test(this, TEXT("Convert a Double to a number formatted correct for hi-IN but as invariant"),				FText::AsNumber(DoubleValue, NULL, InvariantCulture),			FText::FromString(TEXT("12345678.901")));
 		Test(this, TEXT("Convert a Float to a number formatted correct for hi-IN but as invariant"),				FText::AsNumber(FloatValue, NULL, InvariantCulture),			FText::FromString(TEXT("1234.567017")));
 		Test(this, TEXT("Convert a Negative Double to a number formatted correct for hi-IN but as invariant"),	FText::AsNumber(DoubleNegativeValue, NULL, InvariantCulture),	FText::FromString(TEXT("-12345678.901")));
diff --git a/Engine/Source/Runtime/Core/Public/Containers/UnrealString.h b/Engine/Source/Runtime/Core/Public/Containers/UnrealString.h
--- a/Engine/Source/Runtime/Core/Public/Containers/UnrealString.h
+++ b/Engine/Source/Runtime/Core/Public/Containers/UnrealString.h
@@ -1420,6 +1420,59 @@ inline FString BytesToHex(const uint8* In, int32 Count)
 	return Result;
 }
 
+/**
+ * Checks if the TChar is a valid hex character
+ * @param Char		The character
+ * @return	True if in 0-9 and A-F ranges
+ */
+inline const bool CheckTCharIsHex( const TCHAR Char )
+{
+	return ( Char >= TEXT('0') && Char <= TEXT('9') ) || ( Char >= TEXT('A') && Char <= TEXT('F') ) || ( Char >= TEXT('a') && Char <= TEXT('f') );
+}
+
+/**
+ * Convert a TChar to equivalent hex value as a uint8
+ * @param Char		The character
+ * @return	The uint8 value of a hex character
+ */
+inline const uint8 TCharToNibble( const TCHAR Char )
+{
+	check( CheckTCharIsHex( Char ) );
+	if( Char >= TEXT('0') && Char <= TEXT('9') )
+	{
+		return Char - TEXT('0');
+	}
+	else if( Char >= TEXT('A') && Char <= TEXT('F') )
+	{
+		return ( Char - TEXT('A') ) + 10;
+	}
+	return ( Char - TEXT('a') ) + 10;
+}
+
+/** 
+ * Convert FString of Hex digits into the byte array.
+ * @param HexString		The FString of Hex values
+ * @param OutBytes		Ptr to memory must be preallocated large enough
+ * @return	The number of bytes copied
+ */
+inline int32 HexToBytes( const FString& HexString, uint8* OutBytes )
+{
+	int32 NumBytes = 0;
+	const bool bPaddNibble = ( HexString.Len() % 2 ) == 1;
+	const TCHAR* CharPos = *HexString;
+	if( bPaddNibble )
+	{
+		OutBytes[ NumBytes++ ] = TCharToNibble( *CharPos++ );
+	}
+	while( *CharPos )
+	{
+		OutBytes[ NumBytes ] = TCharToNibble( *CharPos++ ) << 4;
+		OutBytes[ NumBytes ] += TCharToNibble( *CharPos++ );
+		++NumBytes;
+	}
+	return NumBytes - 1;
+}
+
 /** A little helper to avoid trying to compile a printf with types that are not accepted by printf **/
 
 template<typename T, bool TIsNumeric>
diff --git a/Engine/Source/Runtime/Core/Public/Internationalization/Internationalization.h b/Engine/Source/Runtime/Core/Public/Internationalization/Internationalization.h
--- a/Engine/Source/Runtime/Core/Public/Internationalization/Internationalization.h
+++ b/Engine/Source/Runtime/Core/Public/Internationalization/Internationalization.h
@@ -42,19 +42,19 @@ public:
 	CORE_API void SetCurrentCulture(const FString& Name);
 
 	//@return the current culture
-	CORE_API TSharedRef< FCulture > GetCurrentCulture() const;
+	CORE_API TSharedRef<FCulture, ESPMode::ThreadSafe> GetCurrentCulture() const;
 
 	//@return culture object by given name, or NULL if not found
-	CORE_API TSharedPtr< FCulture > GetCulture(const FString& Name) const;
+	CORE_API TSharedPtr<FCulture, ESPMode::ThreadSafe> GetCulture(const FString& Name) const;
 
 	//@return the default culture
-	CORE_API TSharedRef< FCulture > GetDefaultCulture() const
+	CORE_API TSharedRef< FCulture, ESPMode::ThreadSafe > GetDefaultCulture() const
 	{
 		return DefaultCulture.ToSharedRef();
 	}
 
 	//@return the invariant culture
-	CORE_API TSharedRef< FCulture > GetInvariantCulture() const
+	CORE_API TSharedRef< FCulture, ESPMode::ThreadSafe > GetInvariantCulture() const
 	{
 		return InvariantCulture.ToSharedRef();
 	}
@@ -69,7 +69,7 @@ public:
 	CORE_API void GetCultureNames(TArray<FString>& CultureNames) const;
 
 	// Given some paths to look at, populate a list of cultures that we have available localization information for
-	CORE_API void GetCulturesWithAvailableLocalization(const TArray<FString>& InLocalizationPaths, TArray< TSharedPtr<FCulture> >& OutAvailableCultures) const;
+	CORE_API void GetCulturesWithAvailableLocalization(const TArray<FString>& InLocalizationPaths, TArray< TSharedPtr<FCulture, ESPMode::ThreadSafe> >& OutAvailableCultures) const;
 
 private:
 	FInternationalization();
@@ -92,12 +92,16 @@ private:
 	static FInternationalization* Instance;
 	bool bIsInitialized;
 
-	TArray< TSharedRef< FCulture > > AllCultures;
+	/** Must be thread safe, cultures may be accessed from more than one thread */
+	TArray< TSharedRef< FCulture, ESPMode::ThreadSafe > > AllCultures;
 
 	int CurrentCultureIndex;
 
-	TSharedPtr< FCulture > DefaultCulture;
-	TSharedPtr< FCulture > InvariantCulture;
+	/** Must be thread safe, cultures may be accessed from more than one thread */
+	TSharedPtr< FCulture, ESPMode::ThreadSafe > DefaultCulture;
+	/** Must be thread safe, cultures may be accessed from more than one thread */
+	TSharedPtr< FCulture, ESPMode::ThreadSafe > InvariantCulture;
+
 	TArray< void* > DLLHandles;
 };
 
diff --git a/Engine/Source/Runtime/Core/Public/Internationalization/Text.h b/Engine/Source/Runtime/Core/Public/Internationalization/Text.h
--- a/Engine/Source/Runtime/Core/Public/Internationalization/Text.h
+++ b/Engine/Source/Runtime/Core/Public/Internationalization/Text.h
@@ -123,51 +123,51 @@ public:
 	/**
 	 * Generate an FText that represents the passed number in the current culture
 	 */
-	static FText AsNumber(float Val,	const FNumberFormattingOptions* const Options = NULL, const TSharedPtr<FCulture>& TargetCulture = NULL);
-	static FText AsNumber(double Val,	const FNumberFormattingOptions* const Options = NULL, const TSharedPtr<FCulture>& TargetCulture = NULL);
-	static FText AsNumber(int8 Val,		const FNumberFormattingOptions* const Options = NULL, const TSharedPtr<FCulture>& TargetCulture = NULL);
-	static FText AsNumber(int16 Val,	const FNumberFormattingOptions* const Options = NULL, const TSharedPtr<FCulture>& TargetCulture = NULL);
-	static FText AsNumber(int32 Val,	const FNumberFormattingOptions* const Options = NULL, const TSharedPtr<FCulture>& TargetCulture = NULL);
-	static FText AsNumber(int64 Val,	const FNumberFormattingOptions* const Options = NULL, const TSharedPtr<FCulture>& TargetCulture = NULL);
-	static FText AsNumber(uint8 Val,	const FNumberFormattingOptions* const Options = NULL, const TSharedPtr<FCulture>& TargetCulture = NULL);
-	static FText AsNumber(uint16 Val,	const FNumberFormattingOptions* const Options = NULL, const TSharedPtr<FCulture>& TargetCulture = NULL);
-	static FText AsNumber(uint32 Val,	const FNumberFormattingOptions* const Options = NULL, const TSharedPtr<FCulture>& TargetCulture = NULL);
-	static FText AsNumber(uint64 Val,	const FNumberFormattingOptions* const Options = NULL, const TSharedPtr<FCulture>& TargetCulture = NULL);
-	static FText AsNumber(long Val,		const FNumberFormattingOptions* const Options = NULL, const TSharedPtr<FCulture>& TargetCulture = NULL);
+	static FText AsNumber(float Val,	const FNumberFormattingOptions* const Options = NULL, const TSharedPtr<FCulture, ESPMode::ThreadSafe>& TargetCulture = NULL);
+	static FText AsNumber(double Val,	const FNumberFormattingOptions* const Options = NULL, const TSharedPtr<FCulture, ESPMode::ThreadSafe>& TargetCulture = NULL);
+	static FText AsNumber(int8 Val,		const FNumberFormattingOptions* const Options = NULL, const TSharedPtr<FCulture, ESPMode::ThreadSafe>& TargetCulture = NULL);
+	static FText AsNumber(int16 Val,	const FNumberFormattingOptions* const Options = NULL, const TSharedPtr<FCulture, ESPMode::ThreadSafe>& TargetCulture = NULL);
+	static FText AsNumber(int32 Val,	const FNumberFormattingOptions* const Options = NULL, const TSharedPtr<FCulture, ESPMode::ThreadSafe>& TargetCulture = NULL);
+	static FText AsNumber(int64 Val,	const FNumberFormattingOptions* const Options = NULL, const TSharedPtr<FCulture, ESPMode::ThreadSafe>& TargetCulture = NULL);
+	static FText AsNumber(uint8 Val,	const FNumberFormattingOptions* const Options = NULL, const TSharedPtr<FCulture, ESPMode::ThreadSafe>& TargetCulture = NULL);
+	static FText AsNumber(uint16 Val,	const FNumberFormattingOptions* const Options = NULL, const TSharedPtr<FCulture, ESPMode::ThreadSafe>& TargetCulture = NULL);
+	static FText AsNumber(uint32 Val,	const FNumberFormattingOptions* const Options = NULL, const TSharedPtr<FCulture, ESPMode::ThreadSafe>& TargetCulture = NULL);
+	static FText AsNumber(uint64 Val,	const FNumberFormattingOptions* const Options = NULL, const TSharedPtr<FCulture, ESPMode::ThreadSafe>& TargetCulture = NULL);
+	static FText AsNumber(long Val,		const FNumberFormattingOptions* const Options = NULL, const TSharedPtr<FCulture, ESPMode::ThreadSafe>& TargetCulture = NULL);
 
 	/**
 	 * Generate an FText that represents the passed number as currency in the current culture
 	 */
-	static FText AsCurrency(float Val,	const FNumberFormattingOptions* const Options = NULL, const TSharedPtr<FCulture>& TargetCulture = NULL);
-	static FText AsCurrency(double Val,	const FNumberFormattingOptions* const Options = NULL, const TSharedPtr<FCulture>& TargetCulture = NULL);
-	static FText AsCurrency(int8 Val,	const FNumberFormattingOptions* const Options = NULL, const TSharedPtr<FCulture>& TargetCulture = NULL);
-	static FText AsCurrency(int16 Val,	const FNumberFormattingOptions* const Options = NULL, const TSharedPtr<FCulture>& TargetCulture = NULL);
-	static FText AsCurrency(int32 Val,	const FNumberFormattingOptions* const Options = NULL, const TSharedPtr<FCulture>& TargetCulture = NULL);
-	static FText AsCurrency(int64 Val,	const FNumberFormattingOptions* const Options = NULL, const TSharedPtr<FCulture>& TargetCulture = NULL);
-	static FText AsCurrency(uint8 Val,	const FNumberFormattingOptions* const Options = NULL, const TSharedPtr<FCulture>& TargetCulture = NULL);
-	static FText AsCurrency(uint16 Val,	const FNumberFormattingOptions* const Options = NULL, const TSharedPtr<FCulture>& TargetCulture = NULL);
-	static FText AsCurrency(uint32 Val,	const FNumberFormattingOptions* const Options = NULL, const TSharedPtr<FCulture>& TargetCulture = NULL);
-	static FText AsCurrency(uint64 Val,	const FNumberFormattingOptions* const Options = NULL, const TSharedPtr<FCulture>& TargetCulture = NULL);
-	static FText AsCurrency(long Val,	const FNumberFormattingOptions* const Options = NULL, const TSharedPtr<FCulture>& TargetCulture = NULL);
+	static FText AsCurrency(float Val,	const FNumberFormattingOptions* const Options = NULL, const TSharedPtr<FCulture, ESPMode::ThreadSafe>& TargetCulture = NULL);
+	static FText AsCurrency(double Val,	const FNumberFormattingOptions* const Options = NULL, const TSharedPtr<FCulture, ESPMode::ThreadSafe>& TargetCulture = NULL);
+	static FText AsCurrency(int8 Val,	const FNumberFormattingOptions* const Options = NULL, const TSharedPtr<FCulture, ESPMode::ThreadSafe>& TargetCulture = NULL);
+	static FText AsCurrency(int16 Val,	const FNumberFormattingOptions* const Options = NULL, const TSharedPtr<FCulture, ESPMode::ThreadSafe>& TargetCulture = NULL);
+	static FText AsCurrency(int32 Val,	const FNumberFormattingOptions* const Options = NULL, const TSharedPtr<FCulture, ESPMode::ThreadSafe>& TargetCulture = NULL);
+	static FText AsCurrency(int64 Val,	const FNumberFormattingOptions* const Options = NULL, const TSharedPtr<FCulture, ESPMode::ThreadSafe>& TargetCulture = NULL);
+	static FText AsCurrency(uint8 Val,	const FNumberFormattingOptions* const Options = NULL, const TSharedPtr<FCulture, ESPMode::ThreadSafe>& TargetCulture = NULL);
+	static FText AsCurrency(uint16 Val,	const FNumberFormattingOptions* const Options = NULL, const TSharedPtr<FCulture, ESPMode::ThreadSafe>& TargetCulture = NULL);
+	static FText AsCurrency(uint32 Val,	const FNumberFormattingOptions* const Options = NULL, const TSharedPtr<FCulture, ESPMode::ThreadSafe>& TargetCulture = NULL);
+	static FText AsCurrency(uint64 Val,	const FNumberFormattingOptions* const Options = NULL, const TSharedPtr<FCulture, ESPMode::ThreadSafe>& TargetCulture = NULL);
+	static FText AsCurrency(long Val,	const FNumberFormattingOptions* const Options = NULL, const TSharedPtr<FCulture, ESPMode::ThreadSafe>& TargetCulture = NULL);
 
 	/**
 	 * Generate an FText that represents the passed number as a percentage in the current culture
 	 */
-	static FText AsPercent(float Val,	const FNumberFormattingOptions* const Options = NULL, const TSharedPtr<FCulture>& TargetCulture = NULL);
-	static FText AsPercent(double Val,	const FNumberFormattingOptions* const Options = NULL, const TSharedPtr<FCulture>& TargetCulture = NULL);
+	static FText AsPercent(float Val,	const FNumberFormattingOptions* const Options = NULL, const TSharedPtr<FCulture, ESPMode::ThreadSafe>& TargetCulture = NULL);
+	static FText AsPercent(double Val,	const FNumberFormattingOptions* const Options = NULL, const TSharedPtr<FCulture, ESPMode::ThreadSafe>& TargetCulture = NULL);
 
 	/**
 	 * Generate an FText that represents the passed number as a date and/or time in the current culture
 	 */
-	static FText AsDate(const FDateTime& DateTime, const EDateTimeStyle::Type DateStyle = EDateTimeStyle::Default, const TSharedPtr<FCulture>& TargetCulture = NULL);
-	static FText AsDateTime(const FDateTime& DateTime, const EDateTimeStyle::Type DateStyle = EDateTimeStyle::Default, const EDateTimeStyle::Type TimeStyle = EDateTimeStyle::Default, const FString& TimeZone = TEXT(""), const TSharedPtr<FCulture>& TargetCulture = NULL);
-	static FText AsTime(const FDateTime& DateTime, const EDateTimeStyle::Type TimeStyle = EDateTimeStyle::Default, const FString& TimeZone = TEXT(""), const TSharedPtr<FCulture>& TargetCulture = NULL);
-	static FText AsTimespan(const FTimespan& Time, const TSharedPtr<FCulture>& TargetCulture = NULL);
+	static FText AsDate(const FDateTime& DateTime, const EDateTimeStyle::Type DateStyle = EDateTimeStyle::Default, const TSharedPtr<FCulture, ESPMode::ThreadSafe>& TargetCulture = NULL);
+	static FText AsDateTime(const FDateTime& DateTime, const EDateTimeStyle::Type DateStyle = EDateTimeStyle::Default, const EDateTimeStyle::Type TimeStyle = EDateTimeStyle::Default, const FString& TimeZone = TEXT(""), const TSharedPtr<FCulture, ESPMode::ThreadSafe>& TargetCulture = NULL);
+	static FText AsTime(const FDateTime& DateTime, const EDateTimeStyle::Type TimeStyle = EDateTimeStyle::Default, const FString& TimeZone = TEXT(""), const TSharedPtr<FCulture, ESPMode::ThreadSafe>& TargetCulture = NULL);
+	static FText AsTimespan(const FTimespan& Time, const TSharedPtr<FCulture, ESPMode::ThreadSafe>& TargetCulture = NULL);
 
 	/**
 	 * Generate an FText that represents the passed number as a memory size in the current culture
 	 */
-	static FText AsMemory(SIZE_T NumBytes, const FNumberFormattingOptions* const Options = NULL, const TSharedPtr<FCulture>& TargetCulture = NULL);
+	static FText AsMemory(SIZE_T NumBytes, const FNumberFormattingOptions* const Options = NULL, const TSharedPtr<FCulture, ESPMode::ThreadSafe>& TargetCulture = NULL);
 
 	/**
 	 * Attempts to find an existing FText using the representation found in the loc tables for the specified namespace and key
@@ -187,6 +187,9 @@ public:
 
 	const FString& ToString() const;
 
+	/** Deep build of the source string for this FText, climbing the history hierarchy */
+	FString BuildSourceString() const;
+
 	bool IsNumeric() const
 	{
 		return DisplayString.Get().IsNumeric();
@@ -283,29 +286,29 @@ private:
 	static FText CreateChronologicalText(FString InSourceString);
 
 	/** Returns the source string of the FText */
-	TSharedPtr<FString> GetSourceString() const;
+	TSharedPtr< FString, ESPMode::ThreadSafe > GetSourceString() const;
 
 	/** Rebuilds the FText under the current culture if needed */
 	void Rebuild() const;
 
-	static FText FormatInternal(const FText& Pattern, const FFormatNamedArguments& Arguments, bool bInRebuildText);
-	static FText FormatInternal(const FText& Pattern, const FFormatOrderedArguments& Arguments, bool bInRebuildText);
-	static FText FormatInternal(const FText& Pattern, const TArray< struct FFormatArgumentData > InArguments, bool bInRebuildText);
+	static FText FormatInternal(const FText& Pattern, const FFormatNamedArguments& Arguments, bool bInRebuildText, bool bInRebuildAsSource);
+	static FText FormatInternal(const FText& Pattern, const FFormatOrderedArguments& Arguments, bool bInRebuildText, bool bInRebuildAsSource);
+	static FText FormatInternal(const FText& Pattern, const TArray< struct FFormatArgumentData > InArguments, bool bInRebuildText, bool bInRebuildAsSource);
 
 private:
 	template<typename T1, typename T2>
-	static FText AsNumberTemplate(T1 Val, const FNumberFormattingOptions* const Options = NULL, const TSharedPtr<FCulture>& TargetCulture = NULL);
+	static FText AsNumberTemplate(T1 Val, const FNumberFormattingOptions* const Options = NULL, const TSharedPtr<FCulture, ESPMode::ThreadSafe>& TargetCulture = NULL);
 	template<typename T1, typename T2>
-	static FText AsCurrencyTemplate(T1 Val, const FNumberFormattingOptions* const Options = NULL, const TSharedPtr<FCulture>& TargetCulture = NULL);
+	static FText AsCurrencyTemplate(T1 Val, const FNumberFormattingOptions* const Options = NULL, const TSharedPtr<FCulture, ESPMode::ThreadSafe>& TargetCulture = NULL);
 	template<typename T1, typename T2>
-	static FText AsPercentTemplate(T1 Val, const FNumberFormattingOptions* const Options = NULL, const TSharedPtr<FCulture>& TargetCulture = NULL);
+	static FText AsPercentTemplate(T1 Val, const FNumberFormattingOptions* const Options = NULL, const TSharedPtr<FCulture, ESPMode::ThreadSafe>& TargetCulture = NULL);
 
 private:
 	/** The visible display string for this FText */
-	TSharedRef<FString> DisplayString;
+	TSharedRef<FString, ESPMode::ThreadSafe> DisplayString;
 
 	/** The FText's history, to allow it to rebuild under a new culture */
-	TSharedPtr<class FTextHistory> History;
+	TSharedPtr<class FTextHistory, ESPMode::ThreadSafe> History;
 
 	/** Flags with various information on what sort of FText this is */
 	int32 Flags;
diff --git a/Engine/Source/Runtime/Core/Public/Internationalization/TextLocalizationManager.h b/Engine/Source/Runtime/Core/Public/Internationalization/TextLocalizationManager.h
--- a/Engine/Source/Runtime/Core/Public/Internationalization/TextLocalizationManager.h
+++ b/Engine/Source/Runtime/Core/Public/Internationalization/TextLocalizationManager.h
@@ -43,9 +43,9 @@ private:
 		bool bIsLocalized;
 		FString TableName;
 		uint32 SourceStringHash;
-		TSharedRef<FString> String;
+		TSharedRef<FString, ESPMode::ThreadSafe> String;
 
-		FStringEntry(const bool InIsLocalized, const FString& InTableName, const uint32 InSourceStringHash, const TSharedRef<FString>& InString)
+		FStringEntry(const bool InIsLocalized, const FString& InTableName, const uint32 InSourceStringHash, const TSharedRef<FString, ESPMode::ThreadSafe>& InString)
 			:	bIsLocalized(InIsLocalized), TableName(InTableName), SourceStringHash(InSourceStringHash), String(InString)
 		{
 		}
@@ -57,16 +57,16 @@ private:
 		typedef TMap<FString, FKeyTable> FNamespaceTable;
 		FNamespaceTable NamespaceTable;
 
-		const TSharedRef<FString>* GetString(const FString& Namespace, const FString& Key, const uint32 SourceStringHash) const;
+		const TSharedRef<FString, ESPMode::ThreadSafe>* GetString(const FString& Namespace, const FString& Key, const uint32 SourceStringHash) const;
 	};
 
 	/** Stores a Namespace and Key, used in reverse lookups from a DisplayString to find the Namespace and Key */
 	struct FNamespaceKeyEntry
 	{
-		TSharedPtr<FString> Namespace;
-		TSharedPtr<FString> Key;
+		TSharedPtr<FString, ESPMode::ThreadSafe> Namespace;
+		TSharedPtr<FString, ESPMode::ThreadSafe> Key;
 
-		FNamespaceKeyEntry(TSharedPtr<FString> InNamespace, TSharedPtr<FString> InKey)
+		FNamespaceKeyEntry(TSharedPtr<FString, ESPMode::ThreadSafe> InNamespace, TSharedPtr<FString, ESPMode::ThreadSafe> InKey)
 			: Namespace(InNamespace)
 			, Key(InKey)
 		{}
@@ -75,8 +75,8 @@ private:
 public:
 	static CORE_API FTextLocalizationManager& Get();
 
-	TSharedPtr<FString> FindString( const FString& Namespace, const FString& Key, const FString* const SourceString = nullptr );
-	TSharedRef<FString> GetString(const FString& Namespace, const FString& Key, const FString* const SourceString);
+	TSharedPtr<FString, ESPMode::ThreadSafe> FindString( const FString& Namespace, const FString& Key, const FString* const SourceString = nullptr );
+	TSharedRef<FString, ESPMode::ThreadSafe> GetString(const FString& Namespace, const FString& Key, const FString* const SourceString);
 
 	/**
 	 * Using an FText SourceString's, returns the associated Namespace and Key
@@ -85,7 +85,7 @@ public:
 	 * @param OutNamespace			Returns with the associated Namespace
 	 * @param OutKey				Returns with the associated Key
 	 */
-	void CORE_API FindKeyNamespaceFromDisplayString(TSharedRef<FString> InDisplayString, TSharedPtr<FString>& OutNamespace, TSharedPtr<FString>& OutKey);
+	void CORE_API FindKeyNamespaceFromDisplayString(TSharedRef<FString, ESPMode::ThreadSafe> InDisplayString, TSharedPtr<FString, ESPMode::ThreadSafe>& OutNamespace, TSharedPtr<FString, ESPMode::ThreadSafe>& OutKey);
 
 	void CORE_API RegenerateResources(const FString& ConfigFilePath);
 
@@ -108,7 +108,7 @@ private:
 	FTextLookupTable LiveTable;
 	
 	/** Table for looking up namespaces and keys using the DisplayString of an FText */
-	TMap< TSharedPtr< FString >, FNamespaceKeyEntry > ReverseLiveTable;
+	TMap< TSharedPtr< FString, ESPMode::ThreadSafe >, FNamespaceKeyEntry > ReverseLiveTable;
 
 	FCriticalSection SynchronizationObject;
 
diff --git a/Engine/Source/Runtime/Engine/Classes/EdGraph/EdGraphNode.h b/Engine/Source/Runtime/Engine/Classes/EdGraph/EdGraphNode.h
--- a/Engine/Source/Runtime/Engine/Classes/EdGraph/EdGraphNode.h
+++ b/Engine/Source/Runtime/Engine/Classes/EdGraph/EdGraphNode.h
@@ -374,6 +374,12 @@ public:
 	/** Return whether to draw this node as a comment node */
 	virtual bool ShouldDrawNodeAsComment() const { return false; }
 
+	/**
+	 * Add's node data to the search metadata, override to collect more data that may be desirable to search for
+	 *
+	 * @param OutTaggedMetaData		Built array of tagged meta data for the node
+	 */
+	virtual void AddSearchMetaDataInfo(TArray<struct FSearchTagDataPair>& OutTaggedMetaData) const;
 #endif // WITH_EDITOR
 
 };
diff --git a/Engine/Source/Runtime/Engine/Classes/Engine/Blueprint.h b/Engine/Source/Runtime/Engine/Classes/Engine/Blueprint.h
--- a/Engine/Source/Runtime/Engine/Classes/Engine/Blueprint.h
+++ b/Engine/Source/Runtime/Engine/Classes/Engine/Blueprint.h
@@ -163,7 +163,7 @@ struct FBPVariableDescription
 	/** Set a metadata value on the variable */
 	ENGINE_API void SetMetaData(const FName& Key, const FString& Value);
 	/** Gets a metadata value on the variable; asserts if the value isn't present.  Check for validiy using FindMetaDataEntryIndexForKey. */
-	ENGINE_API FString GetMetaData(const FName& Key);
+	ENGINE_API FString GetMetaData(const FName& Key) const;
 	/** Clear metadata value on the variable */
 	ENGINE_API void RemoveMetaData(const FName& Key);
 	/** Find the index in the array of a metadata entry */
@@ -296,6 +296,10 @@ class ENGINE_API UBlueprint : public UBlueprintCore
 	UPROPERTY(transient)
 	bool bDisplayCompilePIEWarning;
 
+	/** Guid key for finding searchable data for Blueprint in the DDC */
+	UPROPERTY()
+	FGuid  SearchGuid;
+
 #endif //WITH_EDITORONLY_DATA
 
 	/** 'Simple' construction script - graph of components to instance */
diff --git a/Engine/Source/Runtime/Engine/Engine.Build.cs b/Engine/Source/Runtime/Engine/Engine.Build.cs
--- a/Engine/Source/Runtime/Engine/Engine.Build.cs
+++ b/Engine/Source/Runtime/Engine/Engine.Build.cs
@@ -158,8 +158,20 @@ public class Engine : ModuleRules
 
 		if (UEBuildConfiguration.bBuildEditor == true)
 		{
-			PublicDependencyModuleNames.Add("UnrealEd");	// @todo api: Only public because of WITH_EDITOR and UNREALED_API
-			CircularlyReferencedDependentModules.Add("UnrealEd");
+                
+			PublicDependencyModuleNames.AddRange(
+                new string[] {
+                    "UnrealEd", 
+                    "Kismet"
+                }
+            );	// @todo api: Only public because of WITH_EDITOR and UNREALED_API
+
+			CircularlyReferencedDependentModules.AddRange(
+                new string[] {
+                    "UnrealEd",
+                    "Kismet"
+                }
+            );
 
 			PrivateIncludePathModuleNames.Add("TextureCompressor");
 			PrivateIncludePaths.Add("Developer/TextureCompressor/Public");
diff --git a/Engine/Source/Runtime/Engine/Private/Blueprint.cpp b/Engine/Source/Runtime/Engine/Private/Blueprint.cpp
--- a/Engine/Source/Runtime/Engine/Private/Blueprint.cpp
+++ b/Engine/Source/Runtime/Engine/Private/Blueprint.cpp
@@ -36,7 +36,7 @@ bool FBPVariableDescription::HasMetaData(const FName& Key) const
 }
 
 /** Gets a metadata value on the variable; asserts if the value isn't present.  Check for validiy using FindMetaDataEntryIndexForKey. */
-FString FBPVariableDescription::GetMetaData(const FName& Key)
+FString FBPVariableDescription::GetMetaData(const FName& Key) const
 {
 	int32 EntryIndex = FindMetaDataEntryIndexForKey(Key);
 	check(EntryIndex != INDEX_NONE);
@@ -556,12 +556,14 @@ void UBlueprint::GetAssetRegistryTags(TArray<FAssetRegistryTag>& OutTags) const
 	}
 	OutTags.Add( FAssetRegistryTag("ClassFlags", FString::FromInt(ClassFlagsTagged), FAssetRegistryTag::TT_Hidden) );
 
-	FKismetEditorUtilities::GetAssetRegistryTagsForBlueprint(this, OutTags);
-
 	OutTags.Add( FAssetRegistryTag( "IsDataOnly",
 		FBlueprintEditorUtils::IsDataOnlyBlueprint(this) ? TEXT("True") : TEXT("False"),
 		FAssetRegistryTag::TT_Alphabetical ) );
 
+	if(SearchGuid.IsValid())
+	{
+		OutTags.Add( FAssetRegistryTag("SearchGuid", SearchGuid.ToString(), FAssetRegistryTag::TT_Hidden) );
+	}
 }
 
 FString UBlueprint::GetFriendlyName() const
diff --git a/Engine/Source/Runtime/Engine/Private/EdGraph/EdGraphNode.cpp b/Engine/Source/Runtime/Engine/Private/EdGraph/EdGraphNode.cpp
--- a/Engine/Source/Runtime/Engine/Private/EdGraph/EdGraphNode.cpp
+++ b/Engine/Source/Runtime/Engine/Private/EdGraph/EdGraphNode.cpp
@@ -8,6 +8,7 @@
 #include "Slate.h"
 #include "ScopedTransaction.h"
 #include "Editor/UnrealEd/Public/Kismet2/Kismet2NameValidators.h"
+#include "Editor/Kismet/Public/FindInBlueprintManager.h"
 #endif
 
 #define LOCTEXT_NAMESPACE "EdGraph"
@@ -333,6 +334,28 @@ int32 UEdGraphNode::GetPinIndex(UEdGraphPin* Pin) const
 	return Pins.Find(Pin);
 }
 
+void UEdGraphNode::AddSearchMetaDataInfo(TArray<struct FSearchTagDataPair>& OutTaggedMetaData) const
+{
+	// Searchable - Primary label for the item in the search results
+	OutTaggedMetaData.Add(FSearchTagDataPair(FFindInBlueprintSearchTags::FiB_Name, GetNodeTitle(ENodeTitleType::ListView)));
+
+	// Searchable - As well as being searchable, this displays in the tooltip for the node
+	OutTaggedMetaData.Add(FSearchTagDataPair(FFindInBlueprintSearchTags::FiB_ClassName, FText::FromString(GetClass()->GetName())));
+
+	// Non-searchable - Used to lookup the node when attempting to jump to it
+	OutTaggedMetaData.Add(FSearchTagDataPair(FFindInBlueprintSearchTags::FiB_NodeGuid, FText::FromString(NodeGuid.ToString(EGuidFormats::Digits))));
+
+	// Non-searchable - Important for matching pin types with icons and colors, stored here so that each pin does not store it
+	OutTaggedMetaData.Add(FSearchTagDataPair(FFindInBlueprintSearchTags::FiB_SchemaName, FText::FromString(GetSchema()->GetClass()->GetName())));
+
+	// Non-Searchable - Used to display the icon and color for this node for better visual identification.
+	FLinearColor GlyphColor = FLinearColor::White;
+	OutTaggedMetaData.Add(FSearchTagDataPair(FFindInBlueprintSearchTags::FiB_Glyph, FText::FromString(GetPaletteIcon(GlyphColor).ToString())));
+	OutTaggedMetaData.Add(FSearchTagDataPair(FFindInBlueprintSearchTags::FiB_GlyphColor, FText::FromString(GlyphColor.ToString())));
+	OutTaggedMetaData.Add(FSearchTagDataPair(FFindInBlueprintSearchTags::FiB_Comment, FText::FromString(NodeComment)));
+}
+
+
 #endif	//#if WITH_EDITOR
 
 /////////////////////////////////////////////////////
diff --git a/Engine/Source/Runtime/JsonUtilities/Private/JsonObjectConverter.cpp b/Engine/Source/Runtime/JsonUtilities/Private/JsonObjectConverter.cpp
--- a/Engine/Source/Runtime/JsonUtilities/Private/JsonObjectConverter.cpp
+++ b/Engine/Source/Runtime/JsonUtilities/Private/JsonObjectConverter.cpp
@@ -12,7 +12,7 @@ FString FJsonObjectConverter::StandardizeCase(const FString &StringIn)
 	return FixedString;
 }
 
-static TSharedPtr<FJsonValue> UPropertyToJsonValue(UProperty* Property, const void* Value, int64 CheckFlags, int64 SkipFlags)
+TSharedPtr<FJsonValue> FJsonObjectConverter::UPropertyToJsonValue(UProperty* Property, const void* Value, int64 CheckFlags, int64 SkipFlags)
 {
 	if (UNumericProperty *NumericProperty = Cast<UNumericProperty>(Property))
 	{
diff --git a/Engine/Source/Runtime/JsonUtilities/Public/JsonObjectConverter.h b/Engine/Source/Runtime/JsonUtilities/Public/JsonObjectConverter.h
--- a/Engine/Source/Runtime/JsonUtilities/Public/JsonObjectConverter.h
+++ b/Engine/Source/Runtime/JsonUtilities/Public/JsonObjectConverter.h
@@ -140,5 +140,16 @@ public:
 		}
 		return true;
 	}
+
+	/* * Converts from a UProperty to a Json Value using exportText
+	 *
+	 * @param Property			The property to export
+	 * @param Value				Pointer to the value of the property
+	 * @param CheckFlags		Only convert properties that match at least one of these flags. If 0 check all properties.
+	 * @param SkipFlags			Skip properties that match any of these flags
+	 *
+	 * @return					The constructed JsonValue from the property
+	 */
+	static TSharedPtr<FJsonValue> UPropertyToJsonValue(UProperty* Property, const void* Value, int64 CheckFlags, int64 SkipFlags);
 };
 
