Fixed missing #undef LOCTEXT_NAMESPACE in FindInBlueprintManager.cpp.

diff --git a/Engine/Source/Editor/Kismet/Private/FindInBlueprintManager.cpp b/Engine/Source/Editor/Kismet/Private/FindInBlueprintManager.cpp
--- a/Engine/Source/Editor/Kismet/Private/FindInBlueprintManager.cpp
+++ b/Engine/Source/Editor/Kismet/Private/FindInBlueprintManager.cpp
@@ -1128,4 +1128,6 @@ FString FFindInBlueprintSearchManager::ConvertFTextToHexString(FText InValue)
 	Ar.Close();
 
 	return BytesToHex(SerializedData.GetData(), SerializedData.Num());
-}
\ No newline at end of file
+}
+
+#undef LOCTEXT_NAMESPACE
\ No newline at end of file
