Changed the MyBlueprint graphaction gather to split variables and UActorComponent properties. This means that the variable order can still control the category layout as it did previously and the category sorting works as expected.
#TTP 296001 - ROCKET: TASK: K2: PUNTABLE: Allow double clicking on comment bubble to edit it
#Branch UE4
#Proj Kismet

#change Modified SMyBlueprint::CollectAllActions to split variable actions from UActorComponent actions so the variable order can control the category order as originally intended.

#notes whilst this fix doesn't add a great deal of extra work to the CollectAllActions call I can possible see issues with it. if for instance a UActorComponent is in a category on its own you will not be able to sort that ( because re-ordering variables in the blueprint just won't affect the gather order ).

diff --git a/Engine/Source/Editor/Kismet/Private/SMyBlueprint.cpp b/Engine/Source/Editor/Kismet/Private/SMyBlueprint.cpp
--- a/Engine/Source/Editor/Kismet/Private/SMyBlueprint.cpp
+++ b/Engine/Source/Editor/Kismet/Private/SMyBlueprint.cpp
@@ -607,6 +607,8 @@ void SMyBlueprint::CollectAllActions(FGraphActionListBuilderBase& OutAllActions)
 
 	// List of names of functions we implement
 	TArray<FName> ImplementedFunctions;
+	TArray<TSharedPtr<FEdGraphSchemaAction>> VariableActions;
+	TArray<TSharedPtr<FEdGraphSchemaAction>> ComponentPropertyActions;
 
 	// Grab Variables
 	for (TFieldIterator<UProperty> PropertyIt(Blueprint->SkeletonGeneratedClass, FieldIteratorSuperFlag); PropertyIt; ++PropertyIt)
@@ -620,6 +622,9 @@ void SMyBlueprint::CollectAllActions(FGraphActionListBuilderBase& OutAllActions)
 		const bool bShouldShowAsVar = (!Property->HasAnyPropertyFlags(CPF_Parm) && Property->HasAllPropertyFlags(CPF_BlueprintVisible)) && !bDelegateProp;
 		const bool bShouldShowAsDelegate = !Property->HasAnyPropertyFlags(CPF_Parm) && bMulticastDelegateProp 
 			&& Property->HasAnyPropertyFlags(CPF_BlueprintAssignable | CPF_BlueprintCallable);
+		UObjectPropertyBase* Obj = Cast<UObjectPropertyBase>(Property);
+		const bool bComponentProperty = Obj && Obj->PropertyClass ? Obj->PropertyClass->IsChildOf<UActorComponent>() : false;
+		TArray<TSharedPtr<FEdGraphSchemaAction>>& ActionList = bComponentProperty ? ComponentPropertyActions : VariableActions;
 		
 		if(!bShouldShowAsVar && !bShouldShowAsDelegate)
 		{
@@ -643,7 +648,7 @@ void SMyBlueprint::CollectAllActions(FGraphActionListBuilderBase& OutAllActions)
 			TSharedPtr<FEdGraphSchemaAction_K2Var> NewVarAction = MakeShareable(new FEdGraphSchemaAction_K2Var(PropertyCategory, PropertyDesc, PropertyTooltip.ToString(), 0));
 			NewVarAction->SetVariableInfo(PropertyName, Blueprint->SkeletonGeneratedClass);
 			NewVarAction->SectionID = NodeSectionID::VARIABLE;
-			OutAllActions.AddAction(NewVarAction);
+			ActionList.Add(NewVarAction);
 		}
 		else if (bShouldShowAsDelegate)
 		{
@@ -654,8 +659,8 @@ void SMyBlueprint::CollectAllActions(FGraphActionListBuilderBase& OutAllActions)
 				NewFuncAction = MakeShareable(new FEdGraphSchemaAction_K2Delegate(PropertyCategory, PropertyDesc, PropertyTooltip.ToString(), 0));
 				NewFuncAction->SetDelegateInfo(PropertyName, Blueprint->SkeletonGeneratedClass);
 				NewFuncAction->EdGraph = NULL;
-				NewFuncAction->SectionID = NodeSectionID::DELEGATE;				
-				OutAllActions.AddAction(NewFuncAction);
+				NewFuncAction->SectionID = NodeSectionID::DELEGATE;
+				ActionList.Add(NewFuncAction);
 			}
 
 			UClass* OwnerClass = CastChecked<UClass>(Property->GetOuter());
@@ -671,6 +676,16 @@ void SMyBlueprint::CollectAllActions(FGraphActionListBuilderBase& OutAllActions)
 		}
 	}
 
+	// Add all variable actions first so their order can control the categories
+	for( int32 i = 0; i < VariableActions.Num(); ++i)
+	{
+		OutAllActions.AddAction( VariableActions[ i ] );
+	}
+	for( int32 i = 0; i < ComponentPropertyActions.Num(); ++i)
+	{
+		OutAllActions.AddAction( ComponentPropertyActions[ i ] );
+	}
+
 	// Grab functions implemented by the blueprint
 	for (int32 i = 0; i < Blueprint->FunctionGraphs.Num(); i++)
 	{
