Bluepirnt UI: crash fix while renaming

diff --git a/Engine/Source/Editor/Kismet/Private/BlueprintDetailsCustomization.cpp b/Engine/Source/Editor/Kismet/Private/BlueprintDetailsCustomization.cpp
--- a/Engine/Source/Editor/Kismet/Private/BlueprintDetailsCustomization.cpp
+++ b/Engine/Source/Editor/Kismet/Private/BlueprintDetailsCustomization.cpp
@@ -637,7 +637,7 @@ bool FBlueprintVarActionDetails::IsABlueprintVariable(UProperty* VariablePropert
 
 bool FBlueprintVarActionDetails::IsALocalVariable(UProperty* VariableProperty) const
 {
-	return Cast<UFunction>(VariableProperty->GetOuter()) != NULL;
+	return VariableProperty && (Cast<UFunction>(VariableProperty->GetOuter()) != NULL);
 }
 
 bool FBlueprintVarActionDetails::GetVariableNameChangeEnabled() const
