UDS: EDITOR: Default Values window should be scrollable ttp339547 Panels are horizontal.

diff --git a/Engine/Source/Editor/Kismet/Private/UserDefinedStructureEditor.cpp b/Engine/Source/Editor/Kismet/Private/UserDefinedStructureEditor.cpp
--- a/Engine/Source/Editor/Kismet/Private/UserDefinedStructureEditor.cpp
+++ b/Engine/Source/Editor/Kismet/Private/UserDefinedStructureEditor.cpp
@@ -255,7 +255,7 @@ TSharedRef<SDockTab> FUserDefinedStructureEditor::SpawnStructureTab(const FSpawn
 		EditedStruct = Cast<UUserDefinedStruct>(EditingObjects[ 0 ]);
 	}
 
-	auto Box = SNew(SVerticalBox);
+	auto Box = SNew(SHorizontalBox);
 
 	{
 		// Create a property view
@@ -268,7 +268,6 @@ TSharedRef<SDockTab> FUserDefinedStructureEditor::SpawnStructureTab(const FSpawn
 		PropertyView->RegisterInstancedCustomPropertyLayout(UUserDefinedStruct::StaticClass(), LayoutStructDetails);
 		PropertyView->SetObject(EditedStruct);
 		Box->AddSlot()
-		.AutoHeight()
 		[
 			PropertyView.ToSharedRef()
 		];
@@ -292,7 +291,6 @@ TSharedRef<SDockTab> FUserDefinedStructureEditor::SpawnStructureTab(const FSpawn
 		if (DefaultValueWidget.IsValid())
 		{
 			Box->AddSlot()
-			.AutoHeight()
 			.VAlign(EVerticalAlignment::VAlign_Top)
 			[
 				DefaultValueWidget.ToSharedRef()
