Fix children of struct properties in the blueprint editor not always passing visibility tests
- Refactored IsPropertyVisible delegate to pass in a parent property pointer so complex visibility tests with struct parents can be performed

diff --git a/Engine/Source/Editor/BehaviorTreeEditor/Private/BehaviorTreeEditor.cpp b/Engine/Source/Editor/BehaviorTreeEditor/Private/BehaviorTreeEditor.cpp
--- a/Engine/Source/Editor/BehaviorTreeEditor/Private/BehaviorTreeEditor.cpp
+++ b/Engine/Source/Editor/BehaviorTreeEditor/Private/BehaviorTreeEditor.cpp
@@ -1055,9 +1055,9 @@ void FBehaviorTreeEditor::BindDebuggerToolbarCommands()
 		FExecuteAction::CreateStatic(&FBehaviorTreeDebugger::StopPlaySession));
 }
 
-bool FBehaviorTreeEditor::IsPropertyVisible(UProperty const * const InProperty) const
+bool FBehaviorTreeEditor::IsPropertyVisible( const FPropertyAndParent& PropertyAndParent) const
 {
-	return !InProperty->HasAnyPropertyFlags(CPF_DisableEditOnInstance);
+	return !PropertyAndParent.Property.HasAnyPropertyFlags(CPF_DisableEditOnInstance);
 }
 
 bool FBehaviorTreeEditor::IsPropertyEditable() const
diff --git a/Engine/Source/Editor/BehaviorTreeEditor/Private/BehaviorTreeEditor.h b/Engine/Source/Editor/BehaviorTreeEditor/Private/BehaviorTreeEditor.h
--- a/Engine/Source/Editor/BehaviorTreeEditor/Private/BehaviorTreeEditor.h
+++ b/Engine/Source/Editor/BehaviorTreeEditor/Private/BehaviorTreeEditor.h
@@ -86,7 +86,7 @@ public:
 
 	void JumpToNode(const UEdGraphNode* Node);
 
-	bool IsPropertyVisible(UProperty const * const InProperty) const;
+	bool IsPropertyVisible( const struct FPropertyAndParent& PropertyAndParent ) const;
 	bool IsPropertyEditable() const;
 	void OnPackageSaved(const FString& PackageFileName, UObject* Outer);
 	void OnFinishedChangingProperties(const FPropertyChangedEvent& PropertyChangedEvent);
diff --git a/Engine/Source/Editor/BehaviorTreeEditor/Private/SBehaviorTreeDiff.cpp b/Engine/Source/Editor/BehaviorTreeEditor/Private/SBehaviorTreeDiff.cpp
--- a/Engine/Source/Editor/BehaviorTreeEditor/Private/SBehaviorTreeDiff.cpp
+++ b/Engine/Source/Editor/BehaviorTreeEditor/Private/SBehaviorTreeDiff.cpp
@@ -559,9 +559,9 @@ void SBehaviorTreeDiff::FBehaviorTreeDiffPanel::OnSelectionChanged( const FGraph
 	}
 }
 
-bool SBehaviorTreeDiff::FBehaviorTreeDiffPanel::IsPropertyVisible(UProperty const * const InProperty) const
+bool SBehaviorTreeDiff::FBehaviorTreeDiffPanel::IsPropertyVisible(const FPropertyAndParent& PropertyAndParent) const
 {
-	return !InProperty->HasAnyPropertyFlags(CPF_DisableEditOnInstance);
+	return !PropertyAndParent.Property.HasAnyPropertyFlags(CPF_DisableEditOnInstance);
 }
 
 bool SBehaviorTreeDiff::FBehaviorTreeDiffPanel::IsPropertyEditable()
diff --git a/Engine/Source/Editor/BehaviorTreeEditor/Private/SBehaviorTreeDiff.h b/Engine/Source/Editor/BehaviorTreeEditor/Private/SBehaviorTreeDiff.h
--- a/Engine/Source/Editor/BehaviorTreeEditor/Private/SBehaviorTreeDiff.h
+++ b/Engine/Source/Editor/BehaviorTreeEditor/Private/SBehaviorTreeDiff.h
@@ -56,7 +56,7 @@ private:
 		/**
 		 * Delegate to say if a node property should be visible
 		 */
-		bool IsPropertyVisible(UProperty const * const InProperty) const;
+		bool IsPropertyVisible(const struct FPropertyAndParent& PropertyAndParen) const;
 
 		/**
 		 * Delegate to say if a node property should be editable
diff --git a/Engine/Source/Editor/Blutility/Private/GlobalBlutilityDialog.cpp b/Engine/Source/Editor/Blutility/Private/GlobalBlutilityDialog.cpp
--- a/Engine/Source/Editor/Blutility/Private/GlobalBlutilityDialog.cpp
+++ b/Engine/Source/Editor/Blutility/Private/GlobalBlutilityDialog.cpp
@@ -128,10 +128,10 @@ void FGlobalBlutilityDialog::CreateInternalWidgets()
 	// Helper for hiding private blueprint variables
 	struct Local
 	{
-		static bool IsPropertyVisible(const UProperty* InProperty)
+		static bool IsPropertyVisible(const FPropertyAndParent& PropertyAndParent)
 		{
 			// For details views in the level editor all properties are the instanced versions
-			if ((InProperty != NULL) && InProperty->HasAllPropertyFlags(CPF_DisableEditOnInstance))
+			if (PropertyAndParent.Property.HasAllPropertyFlags(CPF_DisableEditOnInstance))
 			{
 				return false;
 			}
diff --git a/Engine/Source/Editor/DetailCustomizations/Private/ComponentsTree.cpp b/Engine/Source/Editor/DetailCustomizations/Private/ComponentsTree.cpp
--- a/Engine/Source/Editor/DetailCustomizations/Private/ComponentsTree.cpp
+++ b/Engine/Source/Editor/DetailCustomizations/Private/ComponentsTree.cpp
@@ -79,9 +79,10 @@ void SComponentRowWidget::Construct( const FArguments& InArgs, FComponentTreeNod
 //////////////////////////////////////////////////////////////////////////
 
 // Don't show any 'disable on instance' properties, these are instances of components we are seeing
-static bool IsPropertyVisible( const UProperty* InProperty )
+static bool IsPropertyVisible( const FPropertyAndParent& PropertyAndParent )
 {
-	if ((InProperty != NULL) && InProperty->HasAllPropertyFlags(CPF_DisableEditOnInstance))
+
+	if ( PropertyAndParent.Property.HasAllPropertyFlags(CPF_DisableEditOnInstance) )
 	{
 		return false;
 	}
diff --git a/Engine/Source/Editor/FoliageEdit/Private/SFoliageEditMeshDisplayItem.cpp b/Engine/Source/Editor/FoliageEdit/Private/SFoliageEditMeshDisplayItem.cpp
--- a/Engine/Source/Editor/FoliageEdit/Private/SFoliageEditMeshDisplayItem.cpp
+++ b/Engine/Source/Editor/FoliageEdit/Private/SFoliageEditMeshDisplayItem.cpp
@@ -1851,9 +1851,9 @@ FString SFoliageEditMeshDisplayItem::GetSaveRemoveSettingsTooltip() const
 	return NSLOCTEXT("FoliageEdMode", "SaveSettings_Tooltip", "Save these settings as an InstancedFoliageSettings object stored in a package.").ToString();
 }
 
-bool SFoliageEditMeshDisplayItem::IsPropertyVisible(UProperty const * const InProperty) const
+bool SFoliageEditMeshDisplayItem::IsPropertyVisible( const FPropertyAndParent& PropertyAndParent ) const
 {
-	const FString Category = FObjectEditorUtils::GetCategory(InProperty);
+	const FString Category = FObjectEditorUtils::GetCategory(&PropertyAndParent.Property);
 	return Category == TEXT("Clustering") || Category == TEXT("Culling") || Category == TEXT("Lighting") || Category == TEXT("Collision") || Category == TEXT("Rendering");
 }
 
diff --git a/Engine/Source/Editor/FoliageEdit/Private/SFoliageEditMeshDisplayItem.h b/Engine/Source/Editor/FoliageEdit/Private/SFoliageEditMeshDisplayItem.h
--- a/Engine/Source/Editor/FoliageEdit/Private/SFoliageEditMeshDisplayItem.h
+++ b/Engine/Source/Editor/FoliageEdit/Private/SFoliageEditMeshDisplayItem.h
@@ -120,7 +120,7 @@ private:
 	FString GetSaveRemoveSettingsTooltip() const;
 
 	/** Determines what properties to display in the details panel. */
-	bool IsPropertyVisible(UProperty const * const InProperty) const;
+	bool IsPropertyVisible( const struct FPropertyAndParent& PropertyAndParent ) const;
 
 	/** Cluster Accessors */
 	FString GetInstanceCountString() const;
diff --git a/Engine/Source/Editor/Kismet/Private/SKismetInspector.cpp b/Engine/Source/Editor/Kismet/Private/SKismetInspector.cpp
--- a/Engine/Source/Editor/Kismet/Private/SKismetInspector.cpp
+++ b/Engine/Source/Editor/Kismet/Private/SKismetInspector.cpp
@@ -401,8 +401,9 @@ void SKismetInspector::UpdateFromObjects(const TArray<UObject*>& PropertyObjects
 				UProperty* Property = *PropIt;
 				check(Property != NULL);
 
-				SelectedObjectProperties.AddUnique(Property);
+				SelectedObjectProperties.Add(Property);
 
+				UArrayProperty* ArrayProperty = Cast<UArrayProperty>(Property);
 				// If 'showing inners' on a struct, add them to the list as well
 				UStructProperty* StructProperty = Cast<UStructProperty>(Property);
 				static FName ShowOnlyInners("ShowOnlyInnerProperties");
@@ -414,9 +415,16 @@ void SKismetInspector::UpdateFromObjects(const TArray<UObject*>& PropertyObjects
 					{
 						UProperty* InsideStructProperty = *StructPropIt;
 						check(InsideStructProperty != NULL);
-						SelectedObjectProperties.AddUnique(InsideStructProperty);
+						SelectedObjectProperties.Add(InsideStructProperty);
 					}
 				}
+				else if( ArrayProperty && ArrayProperty->Inner->IsA<UStructProperty>() )
+				{
+					// Array property inners with complex children should be tested for visibility since the children of the inner could have different visibility requirements
+					SelectedObjectProperties.Add(ArrayProperty->Inner);
+				}
+
+
 			}
 
 			// Attempt to locate a matching property for the current component template
@@ -435,7 +443,7 @@ void SKismetInspector::UpdateFromObjects(const TArray<UObject*>& PropertyObjects
 						// If the property value matches the current component template, add it as a selected property for filtering
 						if(EditableComponentTemplate == ObjectProperty->GetObjectPropertyValue_InContainer(Object))
 						{
-							SelectedObjectProperties.AddUnique(ObjectProperty);
+							SelectedObjectProperties.Add(ObjectProperty);
 						}
 					}
 				}
@@ -449,19 +457,20 @@ void SKismetInspector::UpdateFromObjects(const TArray<UObject*>& PropertyObjects
 	ContextualEditingBorderWidget->SetContent( MakeContextualEditingWidget(SelectionInfo, Options) );
 }
 
-bool SKismetInspector::IsPropertyVisible(UProperty const * const InProperty) const
+bool SKismetInspector::IsPropertyVisible( const FPropertyAndParent& PropertyAndParent ) const
 {
-	check(InProperty);
+	const UProperty& Property = PropertyAndParent.Property;
+
 
 	// If we are in 'instance preview' - hide anything marked 'disabled edit on instance'
-	if ((ESlateCheckBoxState::Checked == PublicViewState) && InProperty->HasAnyPropertyFlags(CPF_DisableEditOnInstance))
+	if ((ESlateCheckBoxState::Checked == PublicViewState) && Property.HasAnyPropertyFlags(CPF_DisableEditOnInstance))
 	{
 		return false;
 	}
 
-	bool bEditOnTemplateDisabled = InProperty->HasAnyPropertyFlags(CPF_DisableEditOnTemplate);
+	bool bEditOnTemplateDisabled = Property.HasAnyPropertyFlags(CPF_DisableEditOnTemplate);
 
-	if(const UClass* OwningClass = Cast<UClass>(InProperty->GetOuter()))
+	if(const UClass* OwningClass = Cast<UClass>(Property.GetOuter()))
 	{
 		const UBlueprint* BP = Kismet2Ptr.IsValid() ? Kismet2Ptr.Pin()->GetBlueprintObj() : NULL;
 		const bool VariableAddedInCurentBlueprint = (OwningClass->ClassGeneratedBy == BP);
@@ -469,7 +478,7 @@ bool SKismetInspector::IsPropertyVisible(UProperty const * const InProperty) con
 		// If we did not add this var, hide it!
 		if(!VariableAddedInCurentBlueprint)
 		{
-			if (bEditOnTemplateDisabled || InProperty->GetBoolMetaData(FBlueprintMetadata::MD_Private))
+			if (bEditOnTemplateDisabled || Property.GetBoolMetaData(FBlueprintMetadata::MD_Private))
 			{
 				return false;
 			}
@@ -477,8 +486,8 @@ bool SKismetInspector::IsPropertyVisible(UProperty const * const InProperty) con
 	}
 
 	// figure out if this Blueprint variable is an Actor variable
-	const UArrayProperty* ArrayProperty = Cast<const UArrayProperty>(InProperty);
-	const UProperty* TestProperty = ArrayProperty ? ArrayProperty->Inner : InProperty;
+	const UArrayProperty* ArrayProperty = Cast<const UArrayProperty>(&Property);
+	const UProperty* TestProperty = ArrayProperty ? ArrayProperty->Inner : &Property;
 	const UObjectPropertyBase* ObjectProperty = Cast<const UObjectPropertyBase>(TestProperty);
 	bool bIsActorProperty = (ObjectProperty != NULL && ObjectProperty->PropertyClass->IsChildOf(AActor::StaticClass()));
 
@@ -489,15 +498,14 @@ bool SKismetInspector::IsPropertyVisible(UProperty const * const InProperty) con
 		return false;
 	}
 
-	// Filter down to selected properties only if set
-	for(auto PropIt = SelectedObjectProperties.CreateConstIterator(); PropIt; ++PropIt)
+	// Filter down to selected properties only if set.
+	// If the current property is selected then it is visible or if its parent is selected and the current property did not fail any of the above tests it should be visible.
+	if( SelectedObjectProperties.Find( &Property ) || ( PropertyAndParent.ParentProperty && SelectedObjectProperties.Find( PropertyAndParent.ParentProperty ) ) )
 	{
-		if(PropIt->IsValid() && PropIt->Get() == InProperty)
-		{
-			return true;
-		}
+		return true;
 	}
 
+
 	return !SelectedObjectProperties.Num();
 }
 
diff --git a/Engine/Source/Editor/Kismet/Public/SKismetInspector.h b/Engine/Source/Editor/Kismet/Public/SKismetInspector.h
--- a/Engine/Source/Editor/Kismet/Public/SKismetInspector.h
+++ b/Engine/Source/Editor/Kismet/Public/SKismetInspector.h
@@ -55,7 +55,7 @@ public:
 	void ShowDetailsForObjects(const TArray<UObject*>& PropertyObjects, const FShowDetailsOptions& Options = FShowDetailsOptions());
 
 	/** Used to control visibility of a property in the property window */
-	bool IsPropertyVisible(UProperty const * const InProperty) const;
+	bool IsPropertyVisible( const struct FPropertyAndParent& PropertyAndParent ) const;
 
 	/** Turns on or off details customization for components */
 	void EnableComponentDetailsCustomization(bool bEnable);
@@ -96,8 +96,8 @@ protected:
 	/** If true show the kismet inspector title widget */
 	bool bShowTitleArea;
 
-	/** Array of object properties that should be visible */
-	TArray<TWeakObjectPtr<UProperty> > SelectedObjectProperties;
+	/** Set of object properties that should be visible */
+	TSet<TWeakObjectPtr<UProperty> > SelectedObjectProperties;
 	
 protected:
 	/** Show properties of the selected object */
diff --git a/Engine/Source/Editor/LandscapeEditor/Private/SLandscapeEditor.cpp b/Engine/Source/Editor/LandscapeEditor/Private/SLandscapeEditor.cpp
--- a/Engine/Source/Editor/LandscapeEditor/Private/SLandscapeEditor.cpp
+++ b/Engine/Source/Editor/LandscapeEditor/Private/SLandscapeEditor.cpp
@@ -239,12 +239,14 @@ bool SLandscapeEditor::GetLandscapeEditorIsEnabled() const
 	return false;
 }
 
-bool SLandscapeEditor::GetIsPropertyVisible(const UProperty* const Property) const
+bool SLandscapeEditor::GetIsPropertyVisible(const FPropertyAndParent& PropertyAndParent) const
 {
+	const UProperty& Property = PropertyAndParent.Property;
+
 	FEdModeLandscape* LandscapeEdMode = GetEditorMode();
 	if (LandscapeEdMode != NULL && LandscapeEdMode->CurrentToolSet != NULL)
 	{
-		if (Property->HasMetaData("ShowForMask"))
+		if (Property.HasMetaData("ShowForMask"))
 		{
 			const bool bMaskEnabled = LandscapeEdMode->CurrentToolSet->GetTool() &&
 				LandscapeEdMode->CurrentToolSet->GetTool()->SupportsMask() &&
@@ -256,36 +258,36 @@ bool SLandscapeEditor::GetIsPropertyVisible(const UProperty* const Property) con
 				return true;
 			}
 		}
-		if (Property->HasMetaData("ShowForTools"))
+		if (Property.HasMetaData("ShowForTools"))
 		{
 			const FName CurrentToolName = LandscapeEdMode->CurrentToolSet->GetToolSetName();
 
 			TArray<FString> ShowForTools;
-			Property->GetMetaData("ShowForTools").ParseIntoArray(&ShowForTools, TEXT(","), true);
+			Property.GetMetaData("ShowForTools").ParseIntoArray(&ShowForTools, TEXT(","), true);
 			if (!ShowForTools.Contains(CurrentToolName.ToString()))
 			{
 				return false;
 			}
 		}
-		if (Property->HasMetaData("ShowForBrushes"))
+		if (Property.HasMetaData("ShowForBrushes"))
 		{
 			const FName CurrentBrushSetName = LandscapeEdMode->LandscapeBrushSets[LandscapeEdMode->CurrentBrushSetIndex].BrushSetName;
 	//		const FName CurrentBrushName = LandscapeEdMode->CurrentBrush->GetBrushName();
 
 			TArray<FString> ShowForBrushes;
-			Property->GetMetaData("ShowForBrushes").ParseIntoArray(&ShowForBrushes, TEXT(","), true);
+			Property.GetMetaData("ShowForBrushes").ParseIntoArray(&ShowForBrushes, TEXT(","), true);
 			if (!ShowForBrushes.Contains(CurrentBrushSetName.ToString()))
 				//&& !ShowForBrushes.Contains(CurrentBrushName.ToString())
 			{
 				return false;
 			}
 		}
-		if (Property->HasMetaData("ShowForTargetTypes"))
+		if (Property.HasMetaData("ShowForTargetTypes"))
 		{
 			static const TCHAR* TargetTypeNames[] = { TEXT("Heightmap"), TEXT("Weightmap"), TEXT("Visibility") };
 
 			TArray<FString> ShowForTargetTypes;
-			Property->GetMetaData("ShowForTargetTypes").ParseIntoArray(&ShowForTargetTypes, TEXT(","), true);
+			Property.GetMetaData("ShowForTargetTypes").ParseIntoArray(&ShowForTargetTypes, TEXT(","), true);
 
 			const ELandscapeToolTargetType::Type CurrentTargetType = LandscapeEdMode->CurrentToolTarget.TargetType;
 			if (CurrentTargetType == ELandscapeToolTargetType::Invalid ||
diff --git a/Engine/Source/Editor/LandscapeEditor/Private/SLandscapeEditor.h b/Engine/Source/Editor/LandscapeEditor/Private/SLandscapeEditor.h
--- a/Engine/Source/Editor/LandscapeEditor/Private/SLandscapeEditor.h
+++ b/Engine/Source/Editor/LandscapeEditor/Private/SLandscapeEditor.h
@@ -74,7 +74,7 @@ protected:
 
 	bool GetLandscapeEditorIsEnabled() const;
 
-	bool GetIsPropertyVisible(const UProperty* const Property) const;
+	bool GetIsPropertyVisible(const struct FPropertyAndParent& PropertyAndParent) const;
 
 	void OnChangeMode(FName ModeName);
 	bool IsModeEnabled(FName ModeName) const;
diff --git a/Engine/Source/Editor/LevelEditor/Private/SLevelEditor.cpp b/Engine/Source/Editor/LevelEditor/Private/SLevelEditor.cpp
--- a/Engine/Source/Editor/LevelEditor/Private/SLevelEditor.cpp
+++ b/Engine/Source/Editor/LevelEditor/Private/SLevelEditor.cpp
@@ -508,10 +508,10 @@ static TSharedRef<SDockTab> SummonDetailsPanel( FName TabIdentifier )
 {
 	struct Local
 	{
-		static bool IsPropertyVisible( const UProperty* InProperty )
+		static bool IsPropertyVisible( const FPropertyAndParent& PropertyAndParent )
 		{
 			// For details views in the level editor all properties are the instanced versions
-			if ((InProperty != NULL) && InProperty->HasAllPropertyFlags(CPF_DisableEditOnInstance))
+			if (PropertyAndParent.Property.HasAllPropertyFlags(CPF_DisableEditOnInstance))
 			{
 				return false;
 			}
diff --git a/Engine/Source/Editor/PropertyEditor/Private/DetailLayoutBuilderImpl.cpp b/Engine/Source/Editor/PropertyEditor/Private/DetailLayoutBuilderImpl.cpp
--- a/Engine/Source/Editor/PropertyEditor/Private/DetailLayoutBuilderImpl.cpp
+++ b/Engine/Source/Editor/PropertyEditor/Private/DetailLayoutBuilderImpl.cpp
@@ -409,12 +409,19 @@ TSharedPtr<FAssetThumbnailPool> FDetailLayoutBuilderImpl::GetThumbnailPool() con
 
 bool FDetailLayoutBuilderImpl::IsPropertyVisible( TSharedRef<IPropertyHandle> PropertyHandle ) const
 {
-	return PropertyHandle->IsValidHandle() ? IsPropertyVisible( PropertyHandle->GetProperty() ) : false;
+	if( PropertyHandle->IsValidHandle() )
+	{
+		FPropertyAndParent PropertyAndParent(*PropertyHandle->GetProperty(), PropertyHandle->GetParentHandle().IsValid() ? PropertyHandle->GetParentHandle()->GetProperty() : nullptr );
+
+		return IsPropertyVisible(PropertyAndParent);
+	}
+	
+	return false;
 }
 
-bool FDetailLayoutBuilderImpl::IsPropertyVisible( UProperty* Property ) const
+bool FDetailLayoutBuilderImpl::IsPropertyVisible( const FPropertyAndParent& PropertyAndParent ) const
 {
-	return DetailsView.IsPropertyVisible( Property );
+	return DetailsView.IsPropertyVisible( PropertyAndParent );
 }
 
 const IDetailsView& FDetailLayoutBuilderImpl::GetDetailsView() const
diff --git a/Engine/Source/Editor/PropertyEditor/Private/DetailLayoutBuilderImpl.h b/Engine/Source/Editor/PropertyEditor/Private/DetailLayoutBuilderImpl.h
--- a/Engine/Source/Editor/PropertyEditor/Private/DetailLayoutBuilderImpl.h
+++ b/Engine/Source/Editor/PropertyEditor/Private/DetailLayoutBuilderImpl.h
@@ -18,7 +18,7 @@ public:
 	virtual void ForceRefreshDetails() override;
 	virtual TSharedPtr<FAssetThumbnailPool> GetThumbnailPool() const override;
 	virtual bool IsPropertyVisible( TSharedRef<IPropertyHandle> PropertyHandle ) const override;
-	virtual bool IsPropertyVisible( UProperty* PropertyHandle ) const override;
+	virtual bool IsPropertyVisible( const struct FPropertyAndParent& PropertyAndParent ) const override;
 
 	virtual const TSharedRef< IPropertyUtilities >& GetPropertyUtilities() const override; 
 
diff --git a/Engine/Source/Editor/PropertyEditor/Private/DetailPropertyRow.cpp b/Engine/Source/Editor/PropertyEditor/Private/DetailPropertyRow.cpp
--- a/Engine/Source/Editor/PropertyEditor/Private/DetailPropertyRow.cpp
+++ b/Engine/Source/Editor/PropertyEditor/Private/DetailPropertyRow.cpp
@@ -242,9 +242,9 @@ void FDetailPropertyRow::GenerateChildrenForPropertyNode( TSharedPtr<FPropertyNo
 	{
 		TSharedRef<FDetailCategoryImpl> ParentCategoryRef = ParentCategory.Pin().ToSharedRef();
 		IDetailLayoutBuilder& LayoutBuilder = ParentCategoryRef->GetParentLayout();
-		UProperty* Property = RootPropertyNode->GetProperty();
+		UProperty* ParentProperty = RootPropertyNode->GetProperty();
 
-		const bool bStructProperty = Property && Property->IsA<UStructProperty>();
+		const bool bStructProperty = ParentProperty && ParentProperty->IsA<UStructProperty>();
 
 		for( int32 ChildIndex = 0; ChildIndex < RootPropertyNode->GetNumChildNodes(); ++ChildIndex )
 		{
@@ -258,7 +258,7 @@ void FDetailPropertyRow::GenerateChildrenForPropertyNode( TSharedPtr<FPropertyNo
 					GenerateChildrenForPropertyNode( ChildNode, OutChildren );
 				}
 				// Only struct children can have custom visibility that is different from their parent.
-				else if ( !bStructProperty || LayoutBuilder.IsPropertyVisible( ChildNode->GetProperty() ) )
+				else if ( !bStructProperty || LayoutBuilder.IsPropertyVisible( FPropertyAndParent(*ChildNode->GetProperty(), ParentProperty ) ) )
 				{			
 					FDetailLayoutCustomization Customization;
 					Customization.PropertyRow = MakeShareable( new FDetailPropertyRow( ChildNode, ParentCategoryRef ) );
diff --git a/Engine/Source/Editor/PropertyEditor/Private/IDetailsViewPrivate.h b/Engine/Source/Editor/PropertyEditor/Private/IDetailsViewPrivate.h
--- a/Engine/Source/Editor/PropertyEditor/Private/IDetailsViewPrivate.h
+++ b/Engine/Source/Editor/PropertyEditor/Private/IDetailsViewPrivate.h
@@ -75,7 +75,7 @@ public:
 	/**
 	 * @return True if the property is visible
 	 */
-	virtual bool IsPropertyVisible(const UProperty* Property) const = 0;
+	virtual bool IsPropertyVisible( const struct FPropertyAndParent& PropertyAndParent ) const = 0;
 
 	/**
 	 * @return The thumbnail pool that should be used for thumbnails being rendered in this view
diff --git a/Engine/Source/Editor/PropertyEditor/Private/PropertyEditorModule.cpp b/Engine/Source/Editor/PropertyEditor/Private/PropertyEditorModule.cpp
--- a/Engine/Source/Editor/PropertyEditor/Private/PropertyEditorModule.cpp
+++ b/Engine/Source/Editor/PropertyEditor/Private/PropertyEditorModule.cpp
@@ -106,14 +106,16 @@ void FPropertyEditorModule::NotifyCustomizationModuleChanged()
 	}
 }
 
-static bool ShouldShowProperty(UProperty const * const Property, bool bHaveTemplate)
+static bool ShouldShowProperty(const FPropertyAndParent& PropertyAndParent, bool bHaveTemplate)
 {
-	if ((Property != NULL) && bHaveTemplate)
+	const UProperty& Property = PropertyAndParent.Property;
+
+	if ( bHaveTemplate )
 	{
-		const UClass* PropertyOwnerClass = Cast<const UClass>(Property->GetOuter());
+		const UClass* PropertyOwnerClass = Cast<const UClass>(Property.GetOuter());
 		const bool bDisableEditOnTemplate = PropertyOwnerClass 
 			&& PropertyOwnerClass->HasAllFlags(RF_Native) 
-			&& Property->HasAnyPropertyFlags(CPF_DisableEditOnTemplate);
+			&& Property.HasAnyPropertyFlags(CPF_DisableEditOnTemplate);
 		if(bDisableEditOnTemplate)
 		{
 			return false;
@@ -644,9 +646,9 @@ TSharedRef<class IStructureDetailsView> FPropertyEditorModule::CreateStructureDe
 	{
 		struct FDontShowObjects
 		{
-			static bool CanShow(const UProperty* Property)
+			static bool CanShow( const FPropertyAndParent& PropertyAndParent )
 			{
-				return Property && !Property->IsA<UObjectPropertyBase>() && !Property->IsA<UInterfaceProperty>();
+				return !PropertyAndParent.Property.IsA<UObjectPropertyBase>() && !PropertyAndParent.Property.IsA<UInterfaceProperty>();
 			}
 		};
 
diff --git a/Engine/Source/Editor/PropertyEditor/Private/SDetailsViewBase.cpp b/Engine/Source/Editor/PropertyEditor/Private/SDetailsViewBase.cpp
--- a/Engine/Source/Editor/PropertyEditor/Private/SDetailsViewBase.cpp
+++ b/Engine/Source/Editor/PropertyEditor/Private/SDetailsViewBase.cpp
@@ -270,9 +270,9 @@ bool SDetailsViewBase::GetCustomSavedExpansionState(const FString& NodePath) con
 	return ExpandedDetailNodes.Contains(NodePath);
 }
 
-bool SDetailsViewBase::IsPropertyVisible(const UProperty* Property) const
+bool SDetailsViewBase::IsPropertyVisible( const FPropertyAndParent& PropertyAndParent ) const
 {
-	return IsPropertyVisibleDelegate.IsBound() ? IsPropertyVisibleDelegate.Execute(Property) : true;
+	return IsPropertyVisibleDelegate.IsBound() ? IsPropertyVisibleDelegate.Execute(PropertyAndParent) : true;
 }
 
 TSharedPtr<IPropertyUtilities> SDetailsViewBase::GetPropertyUtilities()
@@ -910,8 +910,8 @@ void SDetailsViewBase::UpdatePropertyMapRecursive(FPropertyNode& InNode, FDetail
 				// Children of arrays are not visible directly,
 				bVisibleByDefault &= !bIsChildOfArray;
 
-				// See if the user requested specific property visibility 
-				const bool bIsUserVisible = IsPropertyVisible(Property);
+				FPropertyAndParent PropertyAndParent(*Property, ParentProperty);
+				const bool bIsUserVisible = IsPropertyVisible(PropertyAndParent);
 
 				// Inners of customized in structs should not be taken into consideration for customizing.  They are not designed to be individually customized when their parent is already customized
 				if (!bIsChildOfCustomizedStruct)
diff --git a/Engine/Source/Editor/PropertyEditor/Private/SDetailsViewBase.h b/Engine/Source/Editor/PropertyEditor/Private/SDetailsViewBase.h
--- a/Engine/Source/Editor/PropertyEditor/Private/SDetailsViewBase.h
+++ b/Engine/Source/Editor/PropertyEditor/Private/SDetailsViewBase.h
@@ -200,7 +200,7 @@ public:
 	/**
 	 * @return True if the property is visible
 	 */
-	bool IsPropertyVisible(const UProperty* Property) const override;
+	bool IsPropertyVisible( const struct FPropertyAndParent& PropertyAndParent ) const override;
 
 	/**
 	* Sets a delegate which is regardless of the objects being viewed to lay out generic details not specific to any object
diff --git a/Engine/Source/Editor/PropertyEditor/Private/SPropertyTreeViewImpl.cpp b/Engine/Source/Editor/PropertyEditor/Private/SPropertyTreeViewImpl.cpp
--- a/Engine/Source/Editor/PropertyEditor/Private/SPropertyTreeViewImpl.cpp
+++ b/Engine/Source/Editor/PropertyEditor/Private/SPropertyTreeViewImpl.cpp
@@ -657,7 +657,9 @@ void SPropertyTreeViewImpl::OnGetChildrenForPropertyNode( TSharedPtr<FPropertyNo
 		UProperty* Property = ChildNode->GetProperty();
 		if(Property != NULL && IsPropertyVisible.IsBound())
 		{
-			bPropertyVisible = IsPropertyVisible.Execute(Property);
+			FPropertyAndParent PropertyAndParent(*Property, InPropertyNode->GetProperty());
+
+			bPropertyVisible = IsPropertyVisible.Execute( PropertyAndParent );
 		}
 
 		if(bPropertyVisible)
diff --git a/Engine/Source/Editor/PropertyEditor/Public/DetailLayoutBuilder.h b/Engine/Source/Editor/PropertyEditor/Public/DetailLayoutBuilder.h
--- a/Engine/Source/Editor/PropertyEditor/Public/DetailLayoutBuilder.h
+++ b/Engine/Source/Editor/PropertyEditor/Public/DetailLayoutBuilder.h
@@ -138,5 +138,5 @@ public:
 	/**
 	 * @return true if the property should be visible in the details panel or false if the specific details panel is not showing this property
 	 */
-	virtual bool IsPropertyVisible( UProperty* Property ) const = 0;
+	virtual bool IsPropertyVisible( const struct FPropertyAndParent& PropertyAndParent ) const = 0;
 };
diff --git a/Engine/Source/Editor/PropertyEditor/Public/PropertyEditorDelegates.h b/Engine/Source/Editor/PropertyEditor/Public/PropertyEditorDelegates.h
--- a/Engine/Source/Editor/PropertyEditor/Public/PropertyEditorDelegates.h
+++ b/Engine/Source/Editor/PropertyEditor/Public/PropertyEditorDelegates.h
@@ -3,8 +3,23 @@
 #pragma once
 
 
-/** Delegate called to see if a property should be drawn */
-DECLARE_DELEGATE_RetVal_OneParam( bool, FIsPropertyVisible, const UProperty* const );
+
+struct FPropertyAndParent
+{
+	FPropertyAndParent( const UProperty& InProperty, const UProperty* const InParentProperty )
+		: Property( InProperty )
+		, ParentProperty( InParentProperty )
+	{}
+
+	/** The property always exists */
+	const UProperty& Property;
+
+	/** The parent property may not exist */
+	const UProperty* const ParentProperty;
+};
+
+/** Delegate called to see if a property should be visible */
+DECLARE_DELEGATE_RetVal_OneParam( bool, FIsPropertyVisible, const FPropertyAndParent& );
 
 /** Delegate called to get a detail layout for a specific object class */
 DECLARE_DELEGATE_RetVal( TSharedRef<class IDetailCustomization>, FOnGetDetailCustomizationInstance );
diff --git a/Engine/Source/Editor/UnrealEd/Private/DebugToolExec.cpp b/Engine/Source/Editor/UnrealEd/Private/DebugToolExec.cpp
--- a/Engine/Source/Editor/UnrealEd/Private/DebugToolExec.cpp
+++ b/Engine/Source/Editor/UnrealEd/Private/DebugToolExec.cpp
@@ -23,7 +23,7 @@ void FDebugToolExec::EditObject(UObject* Object, bool bShouldShowNonEditable)
 	struct Local
 	{
 		/** Delegate to show all properties */
-		static bool IsPropertyVisible( const UProperty* const Property, bool bInShouldShowNonEditable )
+		static bool IsPropertyVisible(  const FPropertyAndParent& PropertyAndParent, bool bInShouldShowNonEditable )
 		{
 			return bInShouldShowNonEditable;
 		}
