UDS Editor: Wider space for type selector.
 ttp#335015 k2: custom structs: the custom struct type field is too small to display many of the structs names. The names also do not appear in tool-tips making it difficult to find out what the field is without placing the struct in a blueprint

diff --git a/Engine/Source/Editor/Kismet/Private/UserDefinedStructureEditor.cpp b/Engine/Source/Editor/Kismet/Private/UserDefinedStructureEditor.cpp
--- a/Engine/Source/Editor/Kismet/Private/UserDefinedStructureEditor.cpp
+++ b/Engine/Source/Editor/Kismet/Private/UserDefinedStructureEditor.cpp
@@ -676,6 +676,8 @@ public:
 
 		TSharedPtr<SImage> ErrorIcon;
 
+		const float ValueContentWidth = 250.0f;
+
 		NodeRow
 		.NameContent()
 		[
@@ -701,10 +703,11 @@ public:
 			]
 		]
 		.ValueContent()
+		.MaxDesiredWidth(ValueContentWidth)
+		.MinDesiredWidth(ValueContentWidth)
 		[
 			SNew(SHorizontalBox)
 			+SHorizontalBox::Slot()
-			.FillWidth(1)
 			.VAlign(VAlign_Center)
 			.Padding(0.0f, 0.0f, 4.0f, 0.0f)
 			[
