UDS Property Editor for Default Values enabled. ttp335009, 335507

diff --git a/Engine/Config/BaseEditor.ini b/Engine/Config/BaseEditor.ini
--- a/Engine/Config/BaseEditor.ini
+++ b/Engine/Config/BaseEditor.ini
@@ -251,7 +251,7 @@ SupportWebsiteURL="http://unrealengine.com/support"
 
 [UserDefinedStructure]
 bUseUserDefinedStructure=true
-bShowDefaultValuePropertyEditor=false
+bShowDefaultValuePropertyEditor=true
 
 [EpicSurvey]
 Source="CDN"
diff --git a/Engine/Source/Editor/Kismet/Private/UserDefinedStructureEditor.cpp b/Engine/Source/Editor/Kismet/Private/UserDefinedStructureEditor.cpp
--- a/Engine/Source/Editor/Kismet/Private/UserDefinedStructureEditor.cpp
+++ b/Engine/Source/Editor/Kismet/Private/UserDefinedStructureEditor.cpp
@@ -514,19 +514,6 @@ public:
 		return false;
 	}
 
-	FText OnGetArgDefaultValueText() const
-	{
-		auto StructureDetailsSP = StructureDetails.Pin();
-		if(StructureDetailsSP.IsValid())
-		{
-			if(const FStructVariableDescription* FieldDesc = StructureDetailsSP->FindStructureFieldByGuid(FieldGuid))
-			{
-				return FText::FromString(FieldDesc->DefaultValue);
-			}
-		}
-		return FText();
-	}
-
 	FText OnGetTooltipText() const
 	{
 		auto StructureDetailsSP = StructureDetails.Pin();
@@ -601,15 +588,6 @@ public:
 		}
 	}
 
-	void OnArgDefaultValueCommitted(const FText& NewText, ETextCommit::Type InTextCommit)
-	{
-		auto StructureDetailsSP = StructureDetails.Pin();
-		if(StructureDetailsSP.IsValid())
-		{
-			FStructureEditorUtils::ChangeVariableDefaultValue(StructureDetailsSP->GetUserDefinedStruct(), FieldGuid, NewText.ToString());
-		}
-	}
-
 	/** IDetailCustomNodeBuilder Interface*/
 	virtual void SetOnRebuildChildren( FSimpleDelegate InOnRegenerateChildren ) override 
 	{
@@ -743,21 +721,6 @@ public:
 
 	virtual void GenerateChildContent( IDetailChildrenBuilder& ChildrenBuilder ) override 
 	{
-		ChildrenBuilder.AddChildContent( *LOCTEXT( "UDSVarDefaultValue", "Default Value" ).ToString() )
-		.NameContent()
-		[
-			SNew(STextBlock)
-			.Text( LOCTEXT( "DefaultValue", "Default Value" ) )
-			.Font( IDetailLayoutBuilder::GetDetailFont() )
-		]
-		.ValueContent()
-		[
-			SNew(SEditableTextBox)
-			.Text( this, &FUserDefinedStructureFieldLayout::OnGetArgDefaultValueText )
-			.OnTextCommitted( this, &FUserDefinedStructureFieldLayout::OnArgDefaultValueCommitted )
-			.Font( IDetailLayoutBuilder::GetDetailFont() )
-		];
-
 		ChildrenBuilder.AddChildContent(*LOCTEXT("Tooltip", "Tooltip").ToString())
 		.NameContent()
 		[
