UDS PropertyEditor Default Values - values are updated at initialization

diff --git a/Engine/Source/Editor/Kismet/Private/UserDefinedStructureEditor.cpp b/Engine/Source/Editor/Kismet/Private/UserDefinedStructureEditor.cpp
--- a/Engine/Source/Editor/Kismet/Private/UserDefinedStructureEditor.cpp
+++ b/Engine/Source/Editor/Kismet/Private/UserDefinedStructureEditor.cpp
@@ -25,6 +25,7 @@ public:
 	void Initialize()
 	{
 		StructData = MakeShareable(new FStructOnScope(GetUserDefinedStruct()));
+		FStructureEditorUtils::Fill_MakeStructureDefaultValue(GetUserDefinedStruct(), StructData->GetStructMemory());
 
 		FPropertyEditorModule& PropertyModule = FModuleManager::LoadModuleChecked<FPropertyEditorModule>("PropertyEditor");
 		FDetailsViewArgs ViewArgs;
