BP: Node enabled, one can set a single variable in a structure instance.
"HasNativeMake" and "HasNativeBreak" are no longer boolean - fix.

diff --git a/Engine/Source/Editor/BlueprintGraph/Classes/K2Node_SetFieldsInStruct.h b/Engine/Source/Editor/BlueprintGraph/Classes/K2Node_SetFieldsInStruct.h
--- a/Engine/Source/Editor/BlueprintGraph/Classes/K2Node_SetFieldsInStruct.h
+++ b/Engine/Source/Editor/BlueprintGraph/Classes/K2Node_SetFieldsInStruct.h
@@ -20,4 +20,16 @@ class UK2Node_SetFieldsInStruct : public UK2Node_MakeStruct
 	virtual bool IsNodePure() const override { return false; }
 	virtual class FNodeHandlingFunctor* CreateNodeHandler(class FKismetCompilerContext& CompilerContext) const override;
 	// End K2Node interface
+
+	BLUEPRINTGRAPH_API static bool ShowCustomPinActions(const UEdGraphPin* InGraphPin, bool bIgnorePinsNum);
+
+	enum EPinsToRemove
+	{
+		GivenPin,
+		AllOtherPins
+	};
+
+	BLUEPRINTGRAPH_API void RemoveFieldPins(const UEdGraphPin* InGraphPin, EPinsToRemove Selection);
+	BLUEPRINTGRAPH_API bool AllPinsAreShown() const;
+	BLUEPRINTGRAPH_API void RestoreAllPins();
 };
diff --git a/Engine/Source/Editor/BlueprintGraph/Private/EdGraphSchema_K2.cpp b/Engine/Source/Editor/BlueprintGraph/Private/EdGraphSchema_K2.cpp
--- a/Engine/Source/Editor/BlueprintGraph/Private/EdGraphSchema_K2.cpp
+++ b/Engine/Source/Editor/BlueprintGraph/Private/EdGraphSchema_K2.cpp
@@ -32,6 +32,7 @@
 #include "K2Node_GetEnumeratorName.h"
 #include "K2Node_GetEnumeratorNameAsString.h"
 #include "K2Node_Tunnel.h"
+#include "K2Node_SetFieldsInStruct.h"
 
 //////////////////////////////////////////////////////////////////////////
 // FBlueprintMetadata
@@ -831,6 +832,12 @@ void UEdGraphSchema_K2::GetContextMenuActions(const UEdGraph* CurrentGraph, cons
 					    MenuBuilder->AddMenuEntry( FGraphEditorCommands::Get().RemoveExecutionPin );
 				    }
 			    }
+
+				if (UK2Node_SetFieldsInStruct::ShowCustomPinActions(InGraphPin, true))
+				{
+					MenuBuilder->AddMenuEntry(FGraphEditorCommands::Get().RemoveThisStructVarPin);
+					MenuBuilder->AddMenuEntry(FGraphEditorCommands::Get().RemoveOtherStructVarPins);
+				}
 		    }
 		}
 		MenuBuilder->EndSection();
@@ -930,6 +937,11 @@ void UEdGraphSchema_K2::GetContextMenuActions(const UEdGraph* CurrentGraph, cons
 						MenuBuilder->AddMenuEntry(FGraphEditorCommands::Get().GotoNativeVariableDefinition);
 					}
 
+					if (InGraphNode->IsA(UK2Node_SetFieldsInStruct::StaticClass()))
+					{
+						MenuBuilder->AddMenuEntry(FGraphEditorCommands::Get().RestoreAllStructVarPins);
+					}
+
 					MenuBuilder->AddMenuEntry(FGenericCommands::Get().Rename, NAME_None, LOCTEXT("Rename", "Rename"), LOCTEXT("Rename_Tooltip", "Renames selected function or variable in blueprint.") );
 				}
 
@@ -3892,7 +3904,7 @@ void UEdGraphSchema_K2::BackwardCompatibilityNodeConversion(UEdGraph* Graph, boo
 			}
 
 			// Check to see if the struct has a native make/break that we should try to convert to.
-			if(OldMakeStructNode->StructType && OldMakeStructNode->StructType->GetBoolMetaData(TEXT("HasNativeMake")))
+			if (OldMakeStructNode->StructType && OldMakeStructNode->StructType->HasMetaData(TEXT("HasNativeMake")))
 			{
 				UFunction* MakeNodeFunction = NULL;
 
@@ -3973,7 +3985,7 @@ void UEdGraphSchema_K2::BackwardCompatibilityNodeConversion(UEdGraph* Graph, boo
 			}
 
 				// Check to see if the struct has a native make/break that we should try to convert to.
-			if(OldBreakStructNode->StructType && OldBreakStructNode->StructType->GetBoolMetaData(TEXT("HasNativeBreak")))
+			if (OldBreakStructNode->StructType && OldBreakStructNode->StructType->HasMetaData(TEXT("HasNativeBreak")))
 			{
 				UFunction* BreakNodeFunction = NULL;
 
diff --git a/Engine/Source/Editor/BlueprintGraph/Private/K2ActionMenuBuilder.cpp b/Engine/Source/Editor/BlueprintGraph/Private/K2ActionMenuBuilder.cpp
--- a/Engine/Source/Editor/BlueprintGraph/Private/K2ActionMenuBuilder.cpp
+++ b/Engine/Source/Editor/BlueprintGraph/Private/K2ActionMenuBuilder.cpp
@@ -24,6 +24,7 @@
 #include "K2Node_SwitchEnum.h"
 #include "K2Node_SwitchString.h"
 #include "K2Node_VariableSetRef.h"
+#include "K2Node_SetFieldsInStruct.h"
 
 #include "EdGraphSchema_K2_Actions.h"
 
@@ -1309,6 +1310,7 @@ void FK2ActionMenuBuilder::GetPinAllowedNodeTypes(FBlueprintGraphActionListBuild
 			if ((K2Schema->PC_Struct == FromPin.PinType.PinCategory) && !FromPin.PinType.bIsArray)
 			{
 				UScriptStruct* PinStruct = Cast<UScriptStruct>(FromPin.PinType.PinSubCategoryObject.Get());
+				const bool bCanBeMade = UK2Node_MakeStruct::CanBeMade(PinStruct);
 				if ((FromPin.Direction == EGPD_Output) && UK2Node_BreakStruct::CanBeBroken(PinStruct))
 				{
 					UK2Node_BreakStruct* NodeTemplate = ContextMenuBuilder.CreateTemplateNode<UK2Node_BreakStruct>();
@@ -1316,13 +1318,21 @@ void FK2ActionMenuBuilder::GetPinAllowedNodeTypes(FBlueprintGraphActionListBuild
 					TSharedPtr<FEdGraphSchemaAction_K2NewNode> Action = FK2ActionMenuBuilder::AddNewNodeAction(ContextMenuBuilder, FString(), NodeTemplate->GetNodeTitle(ENodeTitleType::ListView), NodeTemplate->GetTooltip(), 0, NodeTemplate->GetKeywords());
 					Action->NodeTemplate = NodeTemplate;
 				}
-				else if ((FromPin.Direction == EGPD_Input) && UK2Node_MakeStruct::CanBeMade(PinStruct))
+				else if ((FromPin.Direction == EGPD_Input) && bCanBeMade)
 				{
 					UK2Node_MakeStruct* NodeTemplate = ContextMenuBuilder.CreateTemplateNode<UK2Node_MakeStruct>();
 					NodeTemplate->StructType = PinStruct;
 					TSharedPtr<FEdGraphSchemaAction_K2NewNode> Action = FK2ActionMenuBuilder::AddNewNodeAction(ContextMenuBuilder, FString(), NodeTemplate->GetNodeTitle(ENodeTitleType::ListView), NodeTemplate->GetTooltip(), 0, NodeTemplate->GetKeywords());
 					Action->NodeTemplate = NodeTemplate;
 				}
+
+				if ((FromPin.Direction == EGPD_Output) && bCanBeMade)
+				{
+					UK2Node_SetFieldsInStruct* NodeTemplate = ContextMenuBuilder.CreateTemplateNode<UK2Node_SetFieldsInStruct>();
+					NodeTemplate->StructType = PinStruct;
+					TSharedPtr<FEdGraphSchemaAction_K2NewNode> Action = FK2ActionMenuBuilder::AddNewNodeAction(ContextMenuBuilder, FString(), NodeTemplate->GetNodeTitle(ENodeTitleType::ListView), NodeTemplate->GetTooltip(), 0, NodeTemplate->GetKeywords());
+					Action->NodeTemplate = NodeTemplate;
+				}
 			}
 
 			if ((FromPin.Direction == EGPD_Input) && (FromPin.PinType.PinCategory == K2Schema->PC_Delegate))
@@ -1929,12 +1939,24 @@ void FK2ActionMenuBuilder::GetStructActions(FBlueprintGraphActionListBuilder& Co
 
 			if (UK2Node_MakeStruct::CanBeMade(Struct))
 			{
-				UK2Node_MakeStruct* MakeNodeTemplate = ContextMenuBuilder.CreateTemplateNode<UK2Node_MakeStruct>();
-				if( MakeNodeTemplate != NULL )
 				{
-					MakeNodeTemplate->StructType = Struct;
-					TSharedPtr<FEdGraphSchemaAction_K2NewNode> Action = FK2ActionMenuBuilder::AddNewNodeAction(ContextMenuBuilder, K2ActionCategories::MakeStructCategory, MakeNodeTemplate->GetNodeTitle(ENodeTitleType::ListView), MakeNodeTemplate->GetTooltip(), 0, MakeNodeTemplate->GetKeywords());
-					Action->NodeTemplate = MakeNodeTemplate;
+					UK2Node_MakeStruct* MakeNodeTemplate = ContextMenuBuilder.CreateTemplateNode<UK2Node_MakeStruct>();
+					if (MakeNodeTemplate != NULL)
+					{
+						MakeNodeTemplate->StructType = Struct;
+						TSharedPtr<FEdGraphSchemaAction_K2NewNode> Action = FK2ActionMenuBuilder::AddNewNodeAction(ContextMenuBuilder, K2ActionCategories::MakeStructCategory, MakeNodeTemplate->GetNodeTitle(ENodeTitleType::ListView), MakeNodeTemplate->GetTooltip(), 0, MakeNodeTemplate->GetKeywords());
+						Action->NodeTemplate = MakeNodeTemplate;
+					}
+				}
+
+				{
+					UK2Node_SetFieldsInStruct* SetFieldsTemplate = ContextMenuBuilder.CreateTemplateNode<UK2Node_SetFieldsInStruct>();
+					if (SetFieldsTemplate != NULL)
+					{
+						SetFieldsTemplate->StructType = Struct;
+						TSharedPtr<FEdGraphSchemaAction_K2NewNode> Action = FK2ActionMenuBuilder::AddNewNodeAction(ContextMenuBuilder, K2ActionCategories::MakeStructCategory, SetFieldsTemplate->GetNodeTitle(ENodeTitleType::ListView), SetFieldsTemplate->GetTooltip(), 0, SetFieldsTemplate->GetKeywords());
+						Action->NodeTemplate = SetFieldsTemplate;
+					}
 				}
 			}
 		}
@@ -2483,8 +2505,8 @@ void FK2ActionMenuBuilder::GetFuncNodesWithPinType(FBlueprintGraphActionListBuil
 		static const FName HasNativeMake(TEXT("HasNativeMake"));
 		static const FName HasNativeBreak(TEXT("HasNativeBreak"));
 		const UScriptStruct* ScriptStruct = Cast<const UScriptStruct>(DesiredPinType.PinSubCategoryObject.Get());
-		const bool bUseNativeMake = bWantOutput && ScriptStruct && ScriptStruct->GetBoolMetaData(HasNativeMake);
-		const bool bUseNativeBrake = !bWantOutput && ScriptStruct && ScriptStruct->GetBoolMetaData(HasNativeBreak);
+		const bool bUseNativeMake = bWantOutput && ScriptStruct && ScriptStruct->HasMetaData(HasNativeMake);
+		const bool bUseNativeBrake = !bWantOutput && ScriptStruct && ScriptStruct->HasMetaData(HasNativeBreak);
 
 		for (TFieldIterator<UFunction> FunctionIt(Class, EFieldIteratorFlags::IncludeSuper); FunctionIt; ++FunctionIt)
 		{
diff --git a/Engine/Source/Editor/BlueprintGraph/Private/K2Node_SetFieldsInStruct.cpp b/Engine/Source/Editor/BlueprintGraph/Private/K2Node_SetFieldsInStruct.cpp
--- a/Engine/Source/Editor/BlueprintGraph/Private/K2Node_SetFieldsInStruct.cpp
+++ b/Engine/Source/Editor/BlueprintGraph/Private/K2Node_SetFieldsInStruct.cpp
@@ -69,7 +69,7 @@ FString UK2Node_SetFieldsInStruct::GetTooltip() const
 
 FName UK2Node_SetFieldsInStruct::GetPaletteIcon(FLinearColor& OutColor) const
 { 
-	return UK2Node::GetPaletteIcon(OutColor); 
+	return UK2Node_Variable::GetPaletteIcon(OutColor);
 }
 
 FNodeHandlingFunctor* UK2Node_SetFieldsInStruct::CreateNodeHandler(class FKismetCompilerContext& CompilerContext) const
@@ -77,4 +77,69 @@ FNodeHandlingFunctor* UK2Node_SetFieldsInStruct::CreateNodeHandler(class FKismet
 	return new FKCHandler_SetFieldsInStruct(CompilerContext);
 }
 
+bool UK2Node_SetFieldsInStruct::ShowCustomPinActions(const UEdGraphPin* Pin, bool bIgnorePinsNum)
+{
+	const int32 MinimalPinsNum = 4;
+	const UEdGraphSchema_K2* Schema = GetDefault<UEdGraphSchema_K2>();
+	const auto Node = Pin ? Cast<const UK2Node_SetFieldsInStruct>(Pin->GetOwningNodeUnchecked()) : NULL;
+	return Node
+		&& ((Node->Pins.Num() > MinimalPinsNum) || bIgnorePinsNum)
+		&& (EGPD_Input == Pin->Direction)
+		&& (Pin->PinName != SetFieldsInStructHelper::StructRefPinName())
+		&& !Schema->IsMetaPin(*Pin);
+}
+
+void UK2Node_SetFieldsInStruct::RemoveFieldPins(const UEdGraphPin* Pin, EPinsToRemove Selection)
+{
+	if (ShowCustomPinActions(Pin, false) && (Pin->GetOwningNodeUnchecked() == this))
+	{
+		bool bWasChanged = false;
+		for (FOptionalPinFromProperty& OptionalProperty : ShowPinForProperties)
+		{
+			const bool bHide = ((Pin->PinName == OptionalProperty.PropertyName.ToString()) == (Selection == EPinsToRemove::GivenPin));
+			if (OptionalProperty.bShowPin && bHide)
+			{
+				bWasChanged = true;
+				OptionalProperty.bShowPin = false;
+			}
+		}
+
+		if (bWasChanged)
+		{
+			ReconstructNode();
+		}
+	}
+}
+
+bool UK2Node_SetFieldsInStruct::AllPinsAreShown() const
+{
+	for (const FOptionalPinFromProperty& OptionalProperty : ShowPinForProperties)
+	{
+		if (!OptionalProperty.bShowPin)
+		{
+			return false;
+		}
+	}
+
+	return true;
+}
+
+void UK2Node_SetFieldsInStruct::RestoreAllPins()
+{
+	bool bWasChanged = false;
+	for (FOptionalPinFromProperty& OptionalProperty : ShowPinForProperties)
+	{
+		if (!OptionalProperty.bShowPin)
+		{
+			bWasChanged = true;
+			OptionalProperty.bShowPin = true;
+		}
+	}
+
+	if (bWasChanged)
+	{
+		ReconstructNode();
+	}
+}
+
 #undef LOCTEXT_NAMESPACE
\ No newline at end of file
diff --git a/Engine/Source/Editor/BlueprintGraph/Public/BlueprintGraphDefinitions.h b/Engine/Source/Editor/BlueprintGraph/Public/BlueprintGraphDefinitions.h
--- a/Engine/Source/Editor/BlueprintGraph/Public/BlueprintGraphDefinitions.h
+++ b/Engine/Source/Editor/BlueprintGraph/Public/BlueprintGraphDefinitions.h
@@ -58,6 +58,6 @@
 #include "K2Node_Variable.h"
 #include "K2Node_VariableGet.h"
 #include "K2Node_VariableSet.h"
-
+#include "K2Node_SetFieldsInStruct.h"
 
 #include "NodeDependingOnEnumInterface.h"
diff --git a/Engine/Source/Editor/GraphEditor/Private/GraphEditorActions.cpp b/Engine/Source/Editor/GraphEditor/Private/GraphEditorActions.cpp
--- a/Engine/Source/Editor/GraphEditor/Private/GraphEditorActions.cpp
+++ b/Engine/Source/Editor/GraphEditor/Private/GraphEditorActions.cpp
@@ -12,6 +12,11 @@ void FGraphEditorCommandsImpl::RegisterCommands()
 	UI_COMMAND( AddExecutionPin, "Add execution pin", "Adds another execution output pin to an execution sequence or switch node", EUserInterfaceActionType::Button, FInputGesture() )
 	UI_COMMAND( RemoveExecutionPin, "Remove execution pin", "Removes an execution output pin from an execution sequence or switch node", EUserInterfaceActionType::Button, FInputGesture() )
 
+	UI_COMMAND( RemoveThisStructVarPin, "Remove this struct variable pin", "Removes the selected input pin", EUserInterfaceActionType::Button, FInputGesture() )
+	UI_COMMAND( RemoveOtherStructVarPins, "Remove all other pins", "Removes all variable input pins, except for the selected one", EUserInterfaceActionType::Button, FInputGesture() )
+
+	UI_COMMAND( RestoreAllStructVarPins, "Restore all structure pins", "Restore all structure pins", EUserInterfaceActionType::Button, FInputGesture() )
+
 	UI_COMMAND( AddOptionPin, "Add Option Pin", "Adds another option input pin to the node", EUserInterfaceActionType::Button, FInputGesture() )
 	UI_COMMAND( RemoveOptionPin, "Remove Option Pin", "Removes the last option input pin from the node", EUserInterfaceActionType::Button, FInputGesture() )
 
diff --git a/Engine/Source/Editor/GraphEditor/Public/GraphEditorActions.h b/Engine/Source/Editor/GraphEditor/Public/GraphEditorActions.h
--- a/Engine/Source/Editor/GraphEditor/Public/GraphEditorActions.h
+++ b/Engine/Source/Editor/GraphEditor/Public/GraphEditorActions.h
@@ -26,6 +26,11 @@ public:
 	TSharedPtr< FUICommandInfo > AddExecutionPin;
 	TSharedPtr< FUICommandInfo > RemoveExecutionPin;
 
+	// SetFieldsInStruct specific commands
+	TSharedPtr< FUICommandInfo > RemoveThisStructVarPin;
+	TSharedPtr< FUICommandInfo > RemoveOtherStructVarPins;
+	TSharedPtr< FUICommandInfo > RestoreAllStructVarPins;
+
 	// Select node specific commands
 	TSharedPtr< FUICommandInfo > AddOptionPin;
 	TSharedPtr< FUICommandInfo > RemoveOptionPin;
diff --git a/Engine/Source/Editor/Kismet/Private/BlueprintEditor.cpp b/Engine/Source/Editor/Kismet/Private/BlueprintEditor.cpp
--- a/Engine/Source/Editor/Kismet/Private/BlueprintEditor.cpp
+++ b/Engine/Source/Editor/Kismet/Private/BlueprintEditor.cpp
@@ -463,6 +463,21 @@ TSharedRef<SGraphEditor> FBlueprintEditor::CreateGraphEditorWidget(TSharedRef<FT
 				FCanExecuteAction::CreateSP( this, &FBlueprintEditor::CanRemoveExecutionPin )
 				);
 
+			GraphEditorCommands->MapAction(FGraphEditorCommands::Get().RemoveThisStructVarPin,
+				FExecuteAction::CreateSP(this, &FBlueprintEditor::OnRemoveThisStructVarPin),
+				FCanExecuteAction::CreateSP(this, &FBlueprintEditor::CanRemoveThisStructVarPin)
+				);
+
+			GraphEditorCommands->MapAction(FGraphEditorCommands::Get().RemoveOtherStructVarPins,
+				FExecuteAction::CreateSP(this, &FBlueprintEditor::OnRemoveOtherStructVarPins),
+				FCanExecuteAction::CreateSP(this, &FBlueprintEditor::CanRemoveOtherStructVarPins)
+				);
+
+			GraphEditorCommands->MapAction(FGraphEditorCommands::Get().RestoreAllStructVarPins,
+				FExecuteAction::CreateSP(this, &FBlueprintEditor::OnRestoreAllStructVarPins),
+				FCanExecuteAction::CreateSP(this, &FBlueprintEditor::CanRestoreAllStructVarPins)
+				);
+
 			GraphEditorCommands->MapAction( FGraphEditorCommands::Get().AddOptionPin,
 				FExecuteAction::CreateSP( this, &FBlueprintEditor::OnAddOptionPin ),
 				FCanExecuteAction::CreateSP( this, &FBlueprintEditor::CanAddOptionPin )
@@ -2880,6 +2895,96 @@ bool FBlueprintEditor::CanRemoveExecutionPin() const
 	return bReturnValue;
 }
 
+void  FBlueprintEditor::OnRemoveThisStructVarPin()
+{
+	TSharedPtr<SGraphEditor> FocusedGraphEd = FocusedGraphEdPtr.Pin();
+	UEdGraphPin* SelectedPin = FocusedGraphEd.IsValid() ? FocusedGraphEd->GetGraphPinForMenu() : NULL;
+	UEdGraphNode* OwningNode = SelectedPin ? SelectedPin->GetOwningNodeUnchecked() : NULL;
+	if (auto SetFilestInStructNode = Cast<UK2Node_SetFieldsInStruct>(OwningNode))
+	{
+		const FScopedTransaction Transaction(LOCTEXT("RemoveThisStructVarPin", "Remove Struct Var Pin"));
+		SetFilestInStructNode->Modify();
+		SelectedPin->Modify();
+		SetFilestInStructNode->RemoveFieldPins(SelectedPin, UK2Node_SetFieldsInStruct::EPinsToRemove::GivenPin);
+
+		// Update the graph so that the node will be refreshed
+		FocusedGraphEd->NotifyGraphChanged();
+
+		UEdGraph* CurrentGraph = FocusedGraphEd->GetCurrentGraph();
+		UBlueprint* Blueprint = FBlueprintEditorUtils::FindBlueprintForGraphChecked(CurrentGraph);
+		FBlueprintEditorUtils::MarkBlueprintAsModified(Blueprint);
+	}
+}
+
+bool FBlueprintEditor::CanRemoveThisStructVarPin() const
+{
+	auto FocusedGraphEd = FocusedGraphEdPtr.Pin();
+	const auto SelectedPin = FocusedGraphEd.IsValid() ? FocusedGraphEd->GetGraphPinForMenu() : NULL;
+	return UK2Node_SetFieldsInStruct::ShowCustomPinActions(SelectedPin, false);
+}
+
+void  FBlueprintEditor::OnRemoveOtherStructVarPins()
+{
+	TSharedPtr<SGraphEditor> FocusedGraphEd = FocusedGraphEdPtr.Pin();
+	UEdGraphPin* SelectedPin = FocusedGraphEd.IsValid() ? FocusedGraphEd->GetGraphPinForMenu() : NULL;
+	UEdGraphNode* OwningNode = SelectedPin ? SelectedPin->GetOwningNodeUnchecked() : NULL;
+	if (auto SetFilestInStructNode = Cast<UK2Node_SetFieldsInStruct>(OwningNode))
+	{
+		const FScopedTransaction Transaction(LOCTEXT("RemoveOtherStructVarPins", "Remove Other Struct Var Pins"));
+		SetFilestInStructNode->Modify();
+		SelectedPin->Modify();
+		SetFilestInStructNode->RemoveFieldPins(SelectedPin, UK2Node_SetFieldsInStruct::EPinsToRemove::AllOtherPins);
+
+		// Update the graph so that the node will be refreshed
+		FocusedGraphEd->NotifyGraphChanged();
+
+		UEdGraph* CurrentGraph = FocusedGraphEd->GetCurrentGraph();
+		UBlueprint* Blueprint = FBlueprintEditorUtils::FindBlueprintForGraphChecked(CurrentGraph);
+		FBlueprintEditorUtils::MarkBlueprintAsModified(Blueprint);
+	}
+}
+
+bool FBlueprintEditor::CanRemoveOtherStructVarPins() const
+{
+	auto FocusedGraphEd = FocusedGraphEdPtr.Pin();
+	const auto SelectedPin = FocusedGraphEd.IsValid() ? FocusedGraphEd->GetGraphPinForMenu() : NULL;
+	return UK2Node_SetFieldsInStruct::ShowCustomPinActions(SelectedPin, false);
+}
+
+void FBlueprintEditor::OnRestoreAllStructVarPins()
+{
+	const FGraphPanelSelectionSet& SelectedNodes = GetSelectedNodes();
+	FGraphPanelSelectionSet::TConstIterator It(SelectedNodes);
+	UK2Node_SetFieldsInStruct* Node = (!!It) ? Cast<UK2Node_SetFieldsInStruct>(*It) : NULL;
+	if (Node && !Node->AllPinsAreShown())
+	{
+		const FScopedTransaction Transaction(LOCTEXT("RestoreAllStructVarPins", "Restore all struct var pins"));
+		Node->Modify();
+		Node->RestoreAllPins();
+
+		// Refresh the current graph, so the pins can be updated
+		TSharedPtr<SGraphEditor> FocusedGraphEd = FocusedGraphEdPtr.Pin();
+		if (FocusedGraphEd.IsValid())
+		{
+			UEdGraph* CurrentGraph = FocusedGraphEd->GetCurrentGraph();
+			UBlueprint* Blueprint = FBlueprintEditorUtils::FindBlueprintForGraphChecked(CurrentGraph);
+			FBlueprintEditorUtils::MarkBlueprintAsModified(Blueprint);
+
+			FocusedGraphEd->NotifyGraphChanged();
+		}
+	}
+
+	
+}
+
+bool FBlueprintEditor::CanRestoreAllStructVarPins() const
+{
+	const FGraphPanelSelectionSet& SelectedNodes = GetSelectedNodes();
+	FGraphPanelSelectionSet::TConstIterator It(SelectedNodes);
+	UK2Node_SetFieldsInStruct* Node = (!!It) ? Cast<UK2Node_SetFieldsInStruct>(*It) : NULL;
+	return Node && !Node->AllPinsAreShown();
+}
+
 void FBlueprintEditor::OnAddOptionPin()
 {
 	const FGraphPanelSelectionSet& SelectedNodes = GetSelectedNodes();
diff --git a/Engine/Source/Editor/Kismet/Public/BlueprintEditor.h b/Engine/Source/Editor/Kismet/Public/BlueprintEditor.h
--- a/Engine/Source/Editor/Kismet/Public/BlueprintEditor.h
+++ b/Engine/Source/Editor/Kismet/Public/BlueprintEditor.h
@@ -580,6 +580,15 @@ protected:
 	void OnRemoveExecutionPin();
 	bool CanRemoveExecutionPin() const;
 
+	void OnRemoveThisStructVarPin();
+	bool CanRemoveThisStructVarPin() const;
+
+	void OnRemoveOtherStructVarPins();
+	bool CanRemoveOtherStructVarPins() const;
+
+	void OnRestoreAllStructVarPins();
+	bool CanRestoreAllStructVarPins() const;
+
 	void OnAddOptionPin();
 	bool CanAddOptionPin() const;
 
