UMG - Need to force the kismet inspector to clear focus immediately when the selection is about to change to immediately commit any textbox edits that need to be committed.  Fixes a crash with property migration getting out of sync, where the values commit after selection has already changed.

diff --git a/Engine/Source/Editor/Kismet/Private/SKismetInspector.cpp b/Engine/Source/Editor/Kismet/Private/SKismetInspector.cpp
--- a/Engine/Source/Editor/Kismet/Private/SKismetInspector.cpp
+++ b/Engine/Source/Editor/Kismet/Private/SKismetInspector.cpp
@@ -292,6 +292,8 @@ void SKismetInspector::ShowDetailsForSingleObject(UObject* Object, const FShowDe
 
 void SKismetInspector::ShowDetailsForObjects(const TArray<UObject*>& PropertyObjects, const FShowDetailsOptions& Options)
 {
+	FSlateApplication::Get().ClearKeyboardFocus(EKeyboardFocusCause::SetDirectly);
+
 	FKismetSelectionInfo SelectionInfo;
 	UpdateFromObjects(PropertyObjects, SelectionInfo, Options);
 }
