Fixed performance issues with Find in Blueprints

Refactored Json serialization for Find-in-Blueprints to achieve performance increases

#ttp 340452 - BP:  Need to rewrite the UProp -> JSON conversion funcitons to ensure that FIB is performant with compilation

diff --git a/Engine/Source/Editor/Kismet/Private/FindInBlueprintManager.cpp b/Engine/Source/Editor/Kismet/Private/FindInBlueprintManager.cpp
--- a/Engine/Source/Editor/Kismet/Private/FindInBlueprintManager.cpp
+++ b/Engine/Source/Editor/Kismet/Private/FindInBlueprintManager.cpp
@@ -43,13 +43,135 @@ const FText FFindInBlueprintSearchTags::FiB_GlyphColor = LOCTEXT("GlyphColor", "
 
 namespace BlueprintSearchMetaDataHelpers
 {
-	typedef TJsonWriter<TCHAR, TCondensedJsonPrintPolicy<TCHAR>> SearchMetaDataWriter;
-
-	/** Helper function to set a Json string field using FText's, the FText's will be serialized and converted into hex strings for storage */
-	void SetJsonTextField(TSharedRef< FJsonObject > InOutJsonObject, FText InFieldName, FText InFieldValue)
+	/** Json Writer used for serializing FText's in the correct format for Find-in-Blueprints */
+	template < class PrintPolicy = TPrettyJsonPrintPolicy<TCHAR> >
+	class TJsonFindInBlueprintStringWriter : public TJsonStringWriter<PrintPolicy>
 	{
-		InOutJsonObject->SetStringField(FFindInBlueprintSearchManager::ConvertFTextToHexString(InFieldName), FFindInBlueprintSearchManager::ConvertFTextToHexString(InFieldValue));
-	}
+	public:
+		static TSharedRef< TJsonFindInBlueprintStringWriter > Create( FString* const InStream )
+		{
+			return MakeShareable( new TJsonFindInBlueprintStringWriter( InStream ) );
+		}
+
+		using TJsonStringWriter<PrintPolicy>::WriteObjectStart;
+
+		void WriteObjectStart( const FText& Identifier )
+		{
+			check( Stack.Top() == EJson::Object );
+			WriteIdentifier( Identifier );
+
+			PrintPolicy::WriteLineTerminator(Stream);
+			PrintPolicy::WriteTabs(Stream, IndentLevel);
+			PrintPolicy::WriteChar(Stream, TCHAR('{'));
+			++IndentLevel;
+			Stack.Push( EJson::Object );
+			PreviousTokenWritten = EJsonToken::CurlyOpen;
+		}
+
+		void WriteArrayStart( const FText& Identifier )
+		{
+			check( Stack.Top() == EJson::Object );
+			WriteIdentifier( Identifier );
+
+			PrintPolicy::WriteSpace( Stream );
+			PrintPolicy::WriteChar(Stream, TCHAR('['));
+			++IndentLevel;
+			Stack.Push( EJson::Array );
+			PreviousTokenWritten = EJsonToken::SquareOpen;
+		}
+
+		void WriteValue( const FText& Identifier, const bool Value )
+		{
+			check( Stack.Top() == EJson::Object );
+			WriteIdentifier( Identifier );
+
+			PrintPolicy::WriteSpace(Stream);
+			WriteBoolValue( Value );
+			PreviousTokenWritten = Value ? EJsonToken::True : EJsonToken::False;
+		}
+
+		void WriteValue( const FText& Identifier, const double Value )
+		{
+			check( Stack.Top() == EJson::Object );
+			WriteIdentifier( Identifier );
+
+			PrintPolicy::WriteSpace(Stream);
+			WriteNumberValue( Value );
+			PreviousTokenWritten = EJsonToken::Number;
+		}
+
+		void WriteValue( const FText& Identifier, const int32 Value )
+		{
+			check( Stack.Top() == EJson::Object );
+			WriteIdentifier( Identifier );
+
+			PrintPolicy::WriteSpace(Stream);
+			WriteIntegerValue( Value );
+			PreviousTokenWritten = EJsonToken::Number;
+		}
+
+		void WriteValue( const FText& Identifier, const int64 Value )
+		{
+			check( Stack.Top() == EJson::Object );
+			WriteIdentifier( Identifier );
+
+			PrintPolicy::WriteSpace(Stream);
+			WriteIntegerValue( Value );
+			PreviousTokenWritten = EJsonToken::Number;
+		}
+
+		void WriteValue( const FText& Identifier, const FString& Value )
+		{
+			check( Stack.Top() == EJson::Object );
+			WriteIdentifier( Identifier );
+
+			PrintPolicy::WriteSpace(Stream);
+			WriteStringValue( Value );
+			PreviousTokenWritten = EJsonToken::String;
+		}
+
+		void WriteValue( const FText& Identifier, const FText& Value )
+		{
+			check( Stack.Top() == EJson::Object );
+			WriteIdentifier( Identifier );
+
+			PrintPolicy::WriteSpace(Stream);
+			WriteTextValue( Value );
+			PreviousTokenWritten = EJsonToken::String;
+		}
+
+	protected:
+		TJsonFindInBlueprintStringWriter( FString* const InOutString )
+			: TJsonStringWriter<PrintPolicy>( InOutString )
+		{
+		}
+
+		virtual void WriteStringValue( const FString& String ) override
+		{
+			// We just want to make sure all strings are converted into FText hex strings, used by the FiB system
+			TJsonStringWriter::WriteStringValue( FFindInBlueprintSearchManager::ConvertFTextToHexString(FText::FromString(String)) );
+		}
+
+		void WriteTextValue( const FText& Text )
+		{
+			// We just want to make sure all strings are converted into FText hex strings, used by the FiB system
+			TJsonStringWriter::WriteStringValue( FFindInBlueprintSearchManager::ConvertFTextToHexString(Text) );
+		}
+
+		FORCEINLINE void WriteIdentifier( const FText& Identifier )
+		{
+			WriteCommaIfNeeded();
+			PrintPolicy::WriteLineTerminator(Stream);
+
+			PrintPolicy::WriteTabs(Stream, IndentLevel);
+
+			WriteTextValue( Identifier );
+			PrintPolicy::WriteChar(Stream, TCHAR(':'));
+		}
+	};
+
+	typedef TJsonWriter<TCHAR, TCondensedJsonPrintPolicy<TCHAR>> SearchMetaDataWriterParentClass;
+	typedef TJsonFindInBlueprintStringWriter<TCondensedJsonPrintPolicy<TCHAR>> SearchMetaDataWriter;
 
 	/**
 	 * Checks if Json value is searchable, eliminating data that not considered useful to search for
@@ -126,74 +248,29 @@ namespace BlueprintSearchMetaDataHelpers
 	/**
 	 * Saves a graph pin type to a Json object
 	 *
+	 * @param InWriter				Writer used for saving the Json
 	 * @param InPinType				The pin type to save
-	 * @param InOutJsonObject		The Json object to save the pintype to
 	 */
-	void SavePinTypeToJson(const FEdGraphPinType& InPinType, TSharedRef< FJsonObject > InOutJsonObject)
+	void SavePinTypeToJson(TSharedRef< SearchMetaDataWriter>& InWriter, const FEdGraphPinType& InPinType)
 	{
 		// Only save strings that are not empty
 
 		if(!InPinType.PinCategory.IsEmpty())
 		{
-			SetJsonTextField(InOutJsonObject, FFindInBlueprintSearchTags::FiB_PinCategory, FText::FromString(InPinType.PinCategory));
+			InWriter->WriteValue(FFindInBlueprintSearchTags::FiB_PinCategory, InPinType.PinCategory);
 		}
-		
+
 		if(!InPinType.PinSubCategory.IsEmpty())
 		{
-			SetJsonTextField(InOutJsonObject, FFindInBlueprintSearchTags::FiB_PinSubCategory, FText::FromString(InPinType.PinSubCategory));
+			InWriter->WriteValue(FFindInBlueprintSearchTags::FiB_PinSubCategory, InPinType.PinSubCategory);
 		}
 
 		if(InPinType.PinSubCategoryObject.IsValid())
 		{
-			SetJsonTextField(InOutJsonObject, FFindInBlueprintSearchTags::FiB_ObjectClass, FText::FromString(InPinType.PinSubCategoryObject->GetName()));
-		}
-		InOutJsonObject->SetBoolField(FFindInBlueprintSearchManager::ConvertFTextToHexString(FFindInBlueprintSearchTags::FiB_IsArray), InPinType.bIsArray);
-		InOutJsonObject->SetBoolField(FFindInBlueprintSearchManager::ConvertFTextToHexString(FFindInBlueprintSearchTags::FiB_IsReference), InPinType.bIsReference);
-	}
-
-	/** Recursively digs into a JsonValue to correct all Strings into FTexts, including all fields keys. */
-	void FixUpJsonValueToStrings(TSharedRef< FJsonValue >& InJsonValue)
-	{
-		if(InJsonValue->Type == EJson::String)
-		{
-			InJsonValue = MakeShareable( new FJsonValueString(FFindInBlueprintSearchManager::ConvertFTextToHexString(FText::FromString(InJsonValue->AsString()))));
-		}
-		else if(InJsonValue->Type == EJson::Object)
-		{
-			TSharedRef< FJsonObject > NewJsonObject = MakeShareable(new FJsonObject);
-
-			TSharedPtr< FJsonObject > ParentObject = InJsonValue->AsObject();
-
-			if(ParentObject.IsValid())
-			{
-				// Go through all object sub-values and attempt to potentially fix up any strings
-				for( auto MapValues : ParentObject->Values )
-				{
-					TSharedRef< FJsonValue > JsonValueAsSharedRef = MapValues.Value.ToSharedRef();
-					FixUpJsonValueToStrings(JsonValueAsSharedRef);
-
-					NewJsonObject->SetField(FFindInBlueprintSearchManager::ConvertFTextToHexString(FText::FromString(MapValues.Key)), JsonValueAsSharedRef);
-				}
-			}
-
-
-			InJsonValue = MakeShareable( new FJsonValueObject(NewJsonObject));
-		}
-		else if( InJsonValue->Type == EJson::Array)
-		{
-			// Build a new Json array and replace the old one
-			TArray< TSharedPtr< FJsonValue > > JsonArrayData = InJsonValue->AsArray();
-			TArray< TSharedPtr< FJsonValue > > NewJsonArrayData;
-			for( TSharedPtr< FJsonValue >& JsonArrayItem : JsonArrayData)
-			{
-				TSharedRef< FJsonValue > JsonValueAsSharedRef = JsonArrayItem.ToSharedRef();
-				FixUpJsonValueToStrings(JsonValueAsSharedRef);
-
-				NewJsonArrayData.Add(JsonValueAsSharedRef);
-			}
-
-			InJsonValue = MakeShareable(new FJsonValueArray(NewJsonArrayData));
+			InWriter->WriteValue(FFindInBlueprintSearchTags::FiB_ObjectClass, FText::FromString(InPinType.PinSubCategoryObject->GetName()));
 		}
+		InWriter->WriteValue(FFindInBlueprintSearchTags::FiB_IsArray, InPinType.bIsArray);
+		InWriter->WriteValue(FFindInBlueprintSearchTags::FiB_IsReference, InPinType.bIsReference);
 	}
 
 	/**
@@ -207,16 +284,17 @@ namespace BlueprintSearchMetaDataHelpers
 	{
 		FEdGraphPinType VariableType = InVariableDescription.VarType;
 
-		TSharedRef< FJsonObject > JsonObject = MakeShareable(new FJsonObject);
-		SetJsonTextField(JsonObject, FFindInBlueprintSearchTags::FiB_Name, FText::FromString(InVariableDescription.FriendlyName));
+		InWriter->WriteObjectStart();
+
+		InWriter->WriteValue(FFindInBlueprintSearchTags::FiB_Name, InVariableDescription.FriendlyName);
 
 		// Find the variable's tooltip
 		FString TooltipResult;
 		FBlueprintEditorUtils::GetBlueprintVariableMetaData(InBlueprint, InVariableDescription.VarName, FBlueprintMetadata::MD_Tooltip, TooltipResult);
-		SetJsonTextField(JsonObject, FFindInBlueprintSearchTags::FiB_Tooltip, FText::FromString(TooltipResult));
+		InWriter->WriteValue(FFindInBlueprintSearchTags::FiB_Tooltip, TooltipResult);
 
 		// Save the variable's pin type
-		BlueprintSearchMetaDataHelpers::SavePinTypeToJson(VariableType, JsonObject);
+		SavePinTypeToJson(InWriter, VariableType);
 
 		// Find the UProperty and convert it into a Json value.
 		UProperty* VariableProperty = FindField<UProperty>(InBlueprint->GeneratedClass, InVariableDescription.VarName);
@@ -225,20 +303,15 @@ namespace BlueprintSearchMetaDataHelpers
 			const uint8* PropData = VariableProperty->ContainerPtrToValuePtr<uint8>(InBlueprint->GeneratedClass->GetDefaultObject());
 			auto JsonValue = FJsonObjectConverter::UPropertyToJsonValue(VariableProperty, PropData, 0, 0);
 
-			
 			// Only use the value if it is searchable
 			if(BlueprintSearchMetaDataHelpers::CheckIfJsonValueIsSearchable(JsonValue))
 			{
 				TSharedRef< FJsonValue > JsonValueAsSharedRef = JsonValue.ToSharedRef();
-				FixUpJsonValueToStrings(JsonValueAsSharedRef);
-				//FixUpJsonValueToStrings(JsonValue.ToSharedRef());
-				JsonObject->SetField(FFindInBlueprintSearchManager::ConvertFTextToHexString(FFindInBlueprintSearchTags::FiB_DefaultValue), JsonValueAsSharedRef);
-				//JsonObject->SetStringField(FFindInBlueprintSearchManager::ConvertFTextToHexString(FFindInBlueprintSearchTags::FiB_DefaultValue), FFindInBlueprintSearchManager::ConvertFTextToHexString(FText::FromString(JsonValue->AsString())));
+				FJsonSerializer::Serialize(JsonValue, FFindInBlueprintSearchTags::FiB_DefaultValue.ToString(), StaticCastSharedRef<SearchMetaDataWriterParentClass>(InWriter) );
 			}
 		}
 
-		// Serialize the object out using the writer, so that it turns into raw text in the metadata string being made
-		FJsonSerializer::Serialize(JsonObject, InWriter);
+		InWriter->WriteObjectEnd();
 	}
 
 	/**
@@ -250,7 +323,7 @@ namespace BlueprintSearchMetaDataHelpers
 	void GatherNodesFromGraph(TSharedRef< SearchMetaDataWriter>& InWriter, const UEdGraph* InGraph)
 	{
 		// Collect all macro graphs
-		InWriter->WriteArrayStart(FFindInBlueprintSearchManager::ConvertFTextToHexString(FFindInBlueprintSearchTags::FiB_Nodes));
+		InWriter->WriteArrayStart(FFindInBlueprintSearchTags::FiB_Nodes);
 		{
 			for(auto* Node : InGraph->Nodes)
 			{
@@ -262,7 +335,7 @@ namespace BlueprintSearchMetaDataHelpers
 						continue;
 					}
 
-					TSharedRef< FJsonObject > JsonNodeObject = MakeShareable(new FJsonObject);
+					InWriter->WriteObjectStart();
 
 					// Retrieve the search metadata from the node, some node types may have extra metadata to be searchable.
 					TArray<struct FSearchTagDataPair> Tags;
@@ -271,28 +344,26 @@ namespace BlueprintSearchMetaDataHelpers
 					// Go through the node metadata tags and put them into the Json object.
 					for(const FSearchTagDataPair& SearchData : Tags)
 					{
-						FString Key = FFindInBlueprintSearchManager::ConvertFTextToHexString(SearchData.Key);
-						FString Value = FFindInBlueprintSearchManager::ConvertFTextToHexString(SearchData.Value);
-						JsonNodeObject->SetStringField(Key, Value);
+						InWriter->WriteValue(SearchData.Key, SearchData.Value);
 					}
 
 					// Find all the pins and extract their metadata
-					TArray<TSharedPtr< FJsonValue > > JsonPinList;
+					InWriter->WriteArrayStart(FFindInBlueprintSearchTags::FiB_Pins);
 					for(UEdGraphPin* Pin : Node->Pins)
 					{
-						TSharedRef< FJsonObject > JsonPinObject = MakeShareable(new FJsonObject);
-						SetJsonTextField(JsonPinObject, FFindInBlueprintSearchTags::FiB_Name, FText::FromString(Pin->GetSchema()->GetPinDisplayName(Pin)));
-						SetJsonTextField(JsonPinObject, FFindInBlueprintSearchTags::FiB_DefaultValue, FText::FromString(Pin->GetDefaultAsString()));
-						BlueprintSearchMetaDataHelpers::SavePinTypeToJson(Pin->PinType, JsonPinObject);
-
-						// Add the Json pin data to the pin list
-						JsonPinList.Add(MakeShareable(new FJsonValueObject(JsonPinObject)));
+						InWriter->WriteObjectStart();
+						{
+							InWriter->WriteValue(FFindInBlueprintSearchTags::FiB_Name, FText::FromString(Pin->GetSchema()->GetPinDisplayName(Pin)));
+							InWriter->WriteValue(FFindInBlueprintSearchTags::FiB_DefaultValue, FText::FromString(Pin->GetDefaultAsString()));
+						}
+						SavePinTypeToJson(InWriter, Pin->PinType);
+						InWriter->WriteObjectEnd();
 					}
-					JsonNodeObject->SetArrayField(FFindInBlueprintSearchManager::ConvertFTextToHexString(FFindInBlueprintSearchTags::FiB_Pins), JsonPinList);
+					InWriter->WriteArrayEnd();
 
-					// Serialize the object out using the writer, so that it turns into raw text in the metadata string being made
-					FJsonSerializer::Serialize(JsonNodeObject, InWriter);
+					InWriter->WriteObjectEnd();
 				}
+				
 			}
 		}
 		InWriter->WriteArrayEnd();
@@ -311,7 +382,7 @@ namespace BlueprintSearchMetaDataHelpers
 		if(InGraphArray.Num() > 0)
 		{
 			// Collect all graphs
-			InWriter->WriteArrayStart(FFindInBlueprintSearchManager::ConvertFTextToHexString(InTitle));
+			InWriter->WriteArrayStart(InTitle);
 			{
 				for(const UEdGraph* Graph : InGraphArray)
 				{
@@ -319,14 +390,13 @@ namespace BlueprintSearchMetaDataHelpers
 
 					FGraphDisplayInfo DisplayInfo;
 					Graph->GetSchema()->GetGraphDisplayInformation(*Graph, DisplayInfo);
-					InWriter->WriteValue(FFindInBlueprintSearchManager::ConvertFTextToHexString(FFindInBlueprintSearchTags::FiB_Name), FFindInBlueprintSearchManager::ConvertFTextToHexString(DisplayInfo.PlainName));
+					InWriter->WriteValue(FFindInBlueprintSearchTags::FiB_Name, DisplayInfo.PlainName);
 
 					FString GraphDescription = FBlueprintEditorUtils::GetGraphDescription(Graph).ToString();
 					if(!GraphDescription.IsEmpty())
 					{
-						InWriter->WriteValue(FFindInBlueprintSearchManager::ConvertFTextToHexString(FFindInBlueprintSearchTags::FiB_Description), FFindInBlueprintSearchManager::ConvertFTextToHexString(FText::FromString(GraphDescription)));
+						InWriter->WriteValue(FFindInBlueprintSearchTags::FiB_Description, FText::FromString(GraphDescription));
 					}
-
 					// All nodes will appear as children to the graph in search results
 					GatherNodesFromGraph(InWriter, Graph);
 
@@ -334,7 +404,7 @@ namespace BlueprintSearchMetaDataHelpers
 					TArray<UK2Node_FunctionEntry*> FunctionEntryNodes;
 					Graph->GetNodesOfClass<UK2Node_FunctionEntry>(FunctionEntryNodes);
 
-					InWriter->WriteArrayStart(FFindInBlueprintSearchManager::ConvertFTextToHexString(FFindInBlueprintSearchTags::FiB_Properties));
+					InWriter->WriteArrayStart(FFindInBlueprintSearchTags::FiB_Properties);
 					{
 						// Search in all FunctionEntry nodes for their local variables and add them to the list
 						FString ActionCategory;
@@ -629,19 +699,15 @@ FString FFindInBlueprintSearchManager::GatherBlueprintSearchMetadata(const UBlue
 	FString SearchMetaData;
 
 	// The search registry tags for a Blueprint are all in Json
-	TSharedRef< BlueprintSearchMetaDataHelpers::SearchMetaDataWriter > Writer = TJsonWriterFactory<TCHAR, TCondensedJsonPrintPolicy<TCHAR>>::Create( &SearchMetaData );
+	TSharedRef< BlueprintSearchMetaDataHelpers::TJsonFindInBlueprintStringWriter<TCondensedJsonPrintPolicy<TCHAR>> > Writer = BlueprintSearchMetaDataHelpers::TJsonFindInBlueprintStringWriter<TCondensedJsonPrintPolicy<TCHAR>>::Create( &SearchMetaData );
 
 	TMap<FString, TMap<FString,int>> AllPaths;
 	Writer->WriteObjectStart();
 
-	// Pull out all nodes and their pins to be searchable
-	TArray<UEdGraphNode*> Nodes;
-	FBlueprintEditorUtils::GetAllNodesOfClass(Blueprint, Nodes );
-
 	// Only pull properties if the Blueprint has been compiled
 	if(Blueprint->SkeletonGeneratedClass)
 	{
-		Writer->WriteArrayStart(FFindInBlueprintSearchManager::ConvertFTextToHexString(FFindInBlueprintSearchTags::FiB_Properties));
+		Writer->WriteArrayStart(FFindInBlueprintSearchTags::FiB_Properties);
 		{
 			for( const FBPVariableDescription& Variable : Blueprint->NewVariables )
 			{
@@ -667,6 +733,8 @@ FString FFindInBlueprintSearchManager::GatherBlueprintSearchMetadata(const UBlue
 
 void FFindInBlueprintSearchManager::AddOrUpdateBlueprintSearchMetadata(UBlueprint* InBlueprint, bool bInForceReCache/* = false*/)
 {
+	double StartSeconds = FPlatformTime::Seconds();
+
 	check(InBlueprint);
 
 	// Allow only one thread to query the DDC at a single time, this is important in the case that two threads are waiting for the same DDC information, once the information is pulled, it is deleted and is no longer accessible.
@@ -791,17 +859,7 @@ bool FFindInBlueprintSearchManager::ContinueSearchQuery(const FStreamSearch* InS
 			}
 			else if(SearchArray[SearchIdx].Blueprint)
 			{
-				// Don't rebuild the search data if the Blueprint isn't dirty (the package) or possibly dirty (compilation).
-				if(Cast<UPackage>(SearchArray[SearchIdx].Blueprint->GetOutermost())->IsDirty() || SearchArray[SearchIdx].Blueprint->IsPossiblyDirty())
-				{
-					// Update the Blueprint's search metadata so our search is up-to-date.
-					SearchArray[SearchIdx].Blueprint->SearchGuid = FGuid::NewGuid();
-					AddOrUpdateBlueprintSearchMetadata(SearchArray[SearchIdx].Blueprint);
-				}
-				else
-				{
-					RetrieveDDCData(SearchArray[SearchIdx], *SearchArray[SearchIdx].Blueprint->GetPathName());
-				}
+				RetrieveDDCData(SearchArray[SearchIdx], *SearchArray[SearchIdx].Blueprint->GetPathName());
 
  				OutSearchData = SearchArray[SearchIdx++];
 				return true;
@@ -1130,4 +1188,4 @@ FString FFindInBlueprintSearchManager::ConvertFTextToHexString(FText InValue)
 	return BytesToHex(SerializedData.GetData(), SerializedData.Num());
 }
 
-#undef LOCTEXT_NAMESPACE
\ No newline at end of file
+#undef LOCTEXT_NAMESPACE
diff --git a/Engine/Source/Runtime/Core/Public/Serialization/Json/JsonSerializer.h b/Engine/Source/Runtime/Core/Public/Serialization/Json/JsonSerializer.h
--- a/Engine/Source/Runtime/Core/Public/Serialization/Json/JsonSerializer.h
+++ b/Engine/Source/Runtime/Core/Public/Serialization/Json/JsonSerializer.h
@@ -54,6 +54,13 @@ public:
 		TSharedRef< FElement > StartingElement = MakeShareable( new FElement( Object ) );
 		return FJsonSerializer::Serialize<CharType, PrintPolicy>( StartingElement, Writer );
 	}
+
+	template <class CharType, class PrintPolicy >
+	static bool Serialize( const TSharedPtr<FJsonValue >& Value, const FString& Identifier, const TSharedRef< TJsonWriter< CharType, PrintPolicy > >& Writer )
+	{
+		TSharedRef< FElement > StartingElement = MakeShareable( new FElement( Identifier, Value ) );
+		return FJsonSerializer::Serialize<CharType, PrintPolicy>( StartingElement, Writer );
+	}
 	
 
 private:
diff --git a/Engine/Source/Runtime/Core/Public/Serialization/Json/JsonWriter.h b/Engine/Source/Runtime/Core/Public/Serialization/Json/JsonWriter.h
--- a/Engine/Source/Runtime/Core/Public/Serialization/Json/JsonWriter.h
+++ b/Engine/Source/Runtime/Core/Public/Serialization/Json/JsonWriter.h
@@ -296,7 +296,7 @@ protected:
 	}
 
 
-private:
+protected:
 
 	FORCEINLINE bool CanWriteValueWithoutIdentifier() const 
 	{ 
@@ -341,7 +341,7 @@ private:
 		PrintPolicy::WriteString(Stream, TEXT("null"));
 	}
 
-	void WriteStringValue( const FString& String )
+	virtual void WriteStringValue( const FString& String )
 	{
 		FString OutString;
 		OutString += TEXT("\"");
@@ -405,7 +405,7 @@ public:
 	}
 
 
-private:
+protected:
 
 	TJsonStringWriter( FString* const InOutString )
 		: TJsonWriter<TCHAR, PrintPolicy>( new FMemoryWriter(Bytes) )
