Fix for the generic key bindings working intermittently in the blueprint editor.
#TTP 340565 - Regression: BP: Rename/delete node shortcut keys do not function
#Branch UE4
#Proj Kismet

#Change Slightly modified what Nick Darnell added in cl 2126370 so instead of clearing ketboard focus, it grabs the keyboard captor widget and calls OnKeyboardFocusLost on it. Explanation why in info.

Info:
The previous fix caused an issue where dropping or changing the selection of nodes in the BP graph editor calls SKismetInspector::ShowDetailsForObjects which was then clearing keyboard focus ( when the graph editor had focus )
resulting in the GraphEditor not responding to keyboard shortcuts ( unless refocused ). Hopefully this change performs the same function as Nicks changes ( flushes any text changes ) but doesn't force the GraphEditor to lose focus as well.

diff --git a/Engine/Source/Editor/Kismet/Private/SKismetInspector.cpp b/Engine/Source/Editor/Kismet/Private/SKismetInspector.cpp
--- a/Engine/Source/Editor/Kismet/Private/SKismetInspector.cpp
+++ b/Engine/Source/Editor/Kismet/Private/SKismetInspector.cpp
@@ -292,7 +292,12 @@ void SKismetInspector::ShowDetailsForSingleObject(UObject* Object, const FShowDe
 
 void SKismetInspector::ShowDetailsForObjects(const TArray<UObject*>& PropertyObjects, const FShowDetailsOptions& Options)
 {
-	FSlateApplication::Get().ClearKeyboardFocus(EKeyboardFocusCause::SetDirectly);
+	TSharedPtr<SWidget> KeyCaptor = FSlateApplication::Get().GetKeyboardFocusedWidget();
+	// Force any widgets to commit text changes
+	if( KeyCaptor.IsValid() )
+	{
+		KeyCaptor->OnKeyboardFocusLost( FKeyboardFocusEvent( EKeyboardFocusCause::Mouse ) );
+	}
 
 	FKismetSelectionInfo SelectionInfo;
 	UpdateFromObjects(PropertyObjects, SelectionInfo, Options);
