UMG/Blueprint - Fixing the fix of a fix. CL 2137188 by [REDACTED_PERSON] fixed my fix preemptively clearing the details panel when it may not even have focus. This change makes it so we only clear focus if we own focus, and there's a guard to prevent reentrantance which appeared as a problem after the previous fixes.

diff --git a/Engine/Source/Editor/Kismet/Private/SKismetInspector.cpp b/Engine/Source/Editor/Kismet/Private/SKismetInspector.cpp
--- a/Engine/Source/Editor/Kismet/Private/SKismetInspector.cpp
+++ b/Engine/Source/Editor/Kismet/Private/SKismetInspector.cpp
@@ -292,11 +292,17 @@ void SKismetInspector::ShowDetailsForSingleObject(UObject* Object, const FShowDe
 
 void SKismetInspector::ShowDetailsForObjects(const TArray<UObject*>& PropertyObjects, const FShowDetailsOptions& Options)
 {
-	TSharedPtr<SWidget> KeyCaptor = FSlateApplication::Get().GetKeyboardFocusedWidget();
-	// Force any widgets to commit text changes
-	if( KeyCaptor.IsValid() )
+	static bool bIsReentrant = false;
+	if ( !bIsReentrant )
 	{
-		KeyCaptor->OnKeyboardFocusLost( FKeyboardFocusEvent( EKeyboardFocusCause::Mouse ) );
+		bIsReentrant = true;
+		// When the selection is changed, we may be potentially actively editing a property,
+		// if this occurs we need need to immediately clear keyboard focus
+		if ( FSlateApplication::Get().HasFocusedDescendants(AsShared()) )
+		{
+			FSlateApplication::Get().ClearKeyboardFocus(EKeyboardFocusCause::Mouse);
+		}
+		bIsReentrant = false;
 	}
 
 	FKismetSelectionInfo SelectionInfo;
