Fix for Clang compilers and FindInBP JSON parser

diff --git a/Engine/Source/Editor/Kismet/Private/FindInBlueprintManager.cpp b/Engine/Source/Editor/Kismet/Private/FindInBlueprintManager.cpp
--- a/Engine/Source/Editor/Kismet/Private/FindInBlueprintManager.cpp
+++ b/Engine/Source/Editor/Kismet/Private/FindInBlueprintManager.cpp
@@ -57,87 +57,87 @@ namespace BlueprintSearchMetaDataHelpers
 
 		void WriteObjectStart( const FText& Identifier )
 		{
-			check( Stack.Top() == EJson::Object );
+			check( this->Stack.Top() == EJson::Object );
 			WriteIdentifier( Identifier );
 
-			PrintPolicy::WriteLineTerminator(Stream);
-			PrintPolicy::WriteTabs(Stream, IndentLevel);
-			PrintPolicy::WriteChar(Stream, TCHAR('{'));
-			++IndentLevel;
-			Stack.Push( EJson::Object );
-			PreviousTokenWritten = EJsonToken::CurlyOpen;
+			PrintPolicy::WriteLineTerminator(this->Stream);
+			PrintPolicy::WriteTabs(this->Stream, this->IndentLevel);
+			PrintPolicy::WriteChar(this->Stream, TCHAR('{'));
+			++(this->IndentLevel);
+			this->Stack.Push( EJson::Object );
+			this->PreviousTokenWritten = EJsonToken::CurlyOpen;
 		}
 
 		void WriteArrayStart( const FText& Identifier )
 		{
-			check( Stack.Top() == EJson::Object );
+			check( this->Stack.Top() == EJson::Object );
 			WriteIdentifier( Identifier );
 
-			PrintPolicy::WriteSpace( Stream );
-			PrintPolicy::WriteChar(Stream, TCHAR('['));
-			++IndentLevel;
-			Stack.Push( EJson::Array );
-			PreviousTokenWritten = EJsonToken::SquareOpen;
+			PrintPolicy::WriteSpace( this->Stream );
+			PrintPolicy::WriteChar(this->Stream, TCHAR('['));
+			++(this->IndentLevel);
+			this->Stack.Push( EJson::Array );
+			this->PreviousTokenWritten = EJsonToken::SquareOpen;
 		}
 
 		void WriteValue( const FText& Identifier, const bool Value )
 		{
-			check( Stack.Top() == EJson::Object );
+			check( this->Stack.Top() == EJson::Object );
 			WriteIdentifier( Identifier );
 
-			PrintPolicy::WriteSpace(Stream);
-			WriteBoolValue( Value );
-			PreviousTokenWritten = Value ? EJsonToken::True : EJsonToken::False;
+			PrintPolicy::WriteSpace(this->Stream);
+			this->WriteBoolValue( Value );
+			this->PreviousTokenWritten = Value ? EJsonToken::True : EJsonToken::False;
 		}
 
 		void WriteValue( const FText& Identifier, const double Value )
 		{
-			check( Stack.Top() == EJson::Object );
+			check( this->Stack.Top() == EJson::Object );
 			WriteIdentifier( Identifier );
 
-			PrintPolicy::WriteSpace(Stream);
-			WriteNumberValue( Value );
-			PreviousTokenWritten = EJsonToken::Number;
+			PrintPolicy::WriteSpace(this->Stream);
+			this->WriteNumberValue( Value );
+			this->PreviousTokenWritten = EJsonToken::Number;
 		}
 
 		void WriteValue( const FText& Identifier, const int32 Value )
 		{
-			check( Stack.Top() == EJson::Object );
+			check( this->Stack.Top() == EJson::Object );
 			WriteIdentifier( Identifier );
 
-			PrintPolicy::WriteSpace(Stream);
-			WriteIntegerValue( Value );
-			PreviousTokenWritten = EJsonToken::Number;
+			PrintPolicy::WriteSpace(this->Stream);
+			this->WriteIntegerValue( Value );
+			this->PreviousTokenWritten = EJsonToken::Number;
 		}
 
 		void WriteValue( const FText& Identifier, const int64 Value )
 		{
-			check( Stack.Top() == EJson::Object );
+			check( this->Stack.Top() == EJson::Object );
 			WriteIdentifier( Identifier );
 
-			PrintPolicy::WriteSpace(Stream);
-			WriteIntegerValue( Value );
-			PreviousTokenWritten = EJsonToken::Number;
+			PrintPolicy::WriteSpace(this->Stream);
+			this->WriteIntegerValue( Value );
+			this->PreviousTokenWritten = EJsonToken::Number;
 		}
 
 		void WriteValue( const FText& Identifier, const FString& Value )
 		{
-			check( Stack.Top() == EJson::Object );
+			check( this->Stack.Top() == EJson::Object );
 			WriteIdentifier( Identifier );
 
-			PrintPolicy::WriteSpace(Stream);
+			PrintPolicy::WriteSpace(this->Stream);
 			WriteStringValue( Value );
-			PreviousTokenWritten = EJsonToken::String;
+			this->PreviousTokenWritten = EJsonToken::String;
 		}
 
 		void WriteValue( const FText& Identifier, const FText& Value )
 		{
-			check( Stack.Top() == EJson::Object );
+			check( this->Stack.Top() == EJson::Object );
 			WriteIdentifier( Identifier );
 
-			PrintPolicy::WriteSpace(Stream);
+			PrintPolicy::WriteSpace(this->Stream);
 			WriteTextValue( Value );
-			PreviousTokenWritten = EJsonToken::String;
+			this->PreviousTokenWritten = EJsonToken::String;
 		}
 
 	protected:
@@ -149,24 +149,24 @@ namespace BlueprintSearchMetaDataHelpers
 		virtual void WriteStringValue( const FString& String ) override
 		{
 			// We just want to make sure all strings are converted into FText hex strings, used by the FiB system
-			TJsonStringWriter::WriteStringValue( FFindInBlueprintSearchManager::ConvertFTextToHexString(FText::FromString(String)) );
+			TJsonStringWriter<PrintPolicy>::WriteStringValue( FFindInBlueprintSearchManager::ConvertFTextToHexString(FText::FromString(String)) );
 		}
 
 		void WriteTextValue( const FText& Text )
 		{
 			// We just want to make sure all strings are converted into FText hex strings, used by the FiB system
-			TJsonStringWriter::WriteStringValue( FFindInBlueprintSearchManager::ConvertFTextToHexString(Text) );
+			TJsonStringWriter<PrintPolicy>::WriteStringValue( FFindInBlueprintSearchManager::ConvertFTextToHexString(Text) );
 		}
 
 		FORCEINLINE void WriteIdentifier( const FText& Identifier )
 		{
-			WriteCommaIfNeeded();
-			PrintPolicy::WriteLineTerminator(Stream);
+			this->WriteCommaIfNeeded();
+			PrintPolicy::WriteLineTerminator(this->Stream);
 
-			PrintPolicy::WriteTabs(Stream, IndentLevel);
+			PrintPolicy::WriteTabs(this->Stream, this->IndentLevel);
 
 			WriteTextValue( Identifier );
-			PrintPolicy::WriteChar(Stream, TCHAR(':'));
+			PrintPolicy::WriteChar(this->Stream, TCHAR(':'));
 		}
 	};
 
