#UE4 fix for TTP#340084, reviewd by [REDACTED_PERSON]... We were collapsing nodes (based on ::CanCollapseNodes) but then not expanding them in ExpandNodes (CanCollapseNodes does not check CanPasteHere)

diff --git a/Engine/Source/Editor/Kismet/Private/BlueprintEditor.cpp b/Engine/Source/Editor/Kismet/Private/BlueprintEditor.cpp
--- a/Engine/Source/Editor/Kismet/Private/BlueprintEditor.cpp
+++ b/Engine/Source/Editor/Kismet/Private/BlueprintEditor.cpp
@@ -5287,12 +5287,16 @@ void FBlueprintEditor::ExpandNode(UEdGraphNode* InNodeToExpand, UEdGraph* InSour
 	UEdGraphNode* Entry = NULL;
 	UEdGraphNode* Result = NULL;
 
+	const bool bIsCollapsedGraph = InNodeToExpand->IsA<UK2Node_Composite>();
+
 	// Move the nodes over, remembering any that are boundary nodes
 	while (SourceGraph->Nodes.Num())
 	{
 		UEdGraphNode* Node = SourceGraph->Nodes.Pop();
 
-		if(!Node->CanPasteHere(DestinationGraph, Node->GetSchema()))
+		// We do not check CanPasteHere when determining CanCollapseNodes, unlike CanCollapseSelectionToFunction/Macro,
+		// so when expanding a collapsed graph we don't want to check the CanPasteHere function:
+		if (!bIsCollapsedGraph && !Node->CanPasteHere(DestinationGraph, Node->GetSchema()))
 		{
 			continue;
 		}
