GameplayAbility Editor improvements
 - properly hook up schema
 - custom nodes for gameplay effects

diff --git a/Engine/Source/Editor/BlueprintGraph/Classes/EdGraphSchema_K2.h b/Engine/Source/Editor/BlueprintGraph/Classes/EdGraphSchema_K2.h
--- a/Engine/Source/Editor/BlueprintGraph/Classes/EdGraphSchema_K2.h
+++ b/Engine/Source/Editor/BlueprintGraph/Classes/EdGraphSchema_K2.h
@@ -498,6 +498,38 @@ public:
 	virtual void RecombinePin(UEdGraphPin* Pin) const override;
 	// End EdGraphSchema Interface
 
+	/**
+	* Configure the supplied variable node based on the supplied info
+	*
+	* @param	InVarNode			The variable node to be configured
+	* @param	InVariableName		The name of the current variable
+	* @param	InVaraiableSource	The source of the variable
+	* @param	InTargetBlueprint	The blueprint this node will be used on
+	*/
+	static void ConfigureVarNode(class UK2Node_Variable* InVarNode, FName InVariableName, UStruct* InVariableSource, UBlueprint* InTargetBlueprint);
+
+	/**
+	* Creates a new variable getter node and adds it to ParentGraph
+	*
+	* @param	GraphPosition		The location of the new node inside the graph
+	* @param	ParentGraph			The graph to spawn the new node in
+	* @param	VariableName		The name of the variable
+	* @param	Source				The source of the variable
+	* @return	A pointer to the newly spawned node
+	*/
+	virtual class UK2Node_VariableGet* SpawnVariableGetNode(const FVector2D GraphPosition, class UEdGraph* ParentGraph, FName VariableName, UStruct* Source) const;
+
+	/**
+	* Creates a new variable setter node and adds it to ParentGraph
+	*
+	* @param	GraphPosition		The location of the new node inside the graph
+	* @param	ParentGraph			The graph to spawn the new node in
+	* @param	VariableName		The name of the variable
+	* @param	Source				The source of the variable
+	* @return	A pointer to the newly spawned node
+	*/
+	virtual class UK2Node_VariableSet* SpawnVariableSetNode(const FVector2D GraphPosition, class UEdGraph* ParentGraph, FName VariableName, UStruct* Source) const;
+
 	// Returns whether the supplied Pin is a splittable struct
 	bool PinHasSplittableStructType(const UEdGraphPin* InGraphPin) const;
 
diff --git a/Engine/Source/Editor/BlueprintGraph/Classes/K2Node_Variable.h b/Engine/Source/Editor/BlueprintGraph/Classes/K2Node_Variable.h
--- a/Engine/Source/Editor/BlueprintGraph/Classes/K2Node_Variable.h
+++ b/Engine/Source/Editor/BlueprintGraph/Classes/K2Node_Variable.h
@@ -15,8 +15,8 @@ namespace ESelfContextInfo
 	};
 }
 
-UCLASS(MinimalAPI, abstract)
-class UK2Node_Variable : public UK2Node
+UCLASS(abstract)
+class BLUEPRINTGRAPH_API UK2Node_Variable : public UK2Node
 {
 	GENERATED_UCLASS_BODY()
 
@@ -60,11 +60,11 @@ public:
 	virtual ERedirectType DoPinsMatchForReconstruction(const UEdGraphPin* NewPin, int32 NewPinIndex, const UEdGraphPin* OldPin, int32 OldPinIndex)  const override;
 	virtual void ValidateNodeDuringCompilation(class FCompilerResultsLog& MessageLog) const override;
 	virtual FText GetToolTipHeading() const override;
-	virtual void GetNodeAttributes( TArray<TKeyValuePair<FString, FString>>& OutNodeAttributes ) const override;
+	virtual void GetNodeAttributes(TArray<TKeyValuePair<FString, FString>>& OutNodeAttributes) const override;
 	// End K2Node interface
 
 	/** Set up this variable node from the supplied UProperty */
-	BLUEPRINTGRAPH_API void SetFromProperty(const UProperty* Property, bool bSelfContext);
+	void SetFromProperty(const UProperty* Property, bool bSelfContext);
 
 	/** Util to get variable name as a string */
 	FString GetVarNameString() const
@@ -107,13 +107,13 @@ public:
 	bool RecreatePinForVariable(EEdGraphPinDirection Direction, TArray<UEdGraphPin*>& OldPins);
 
 	/** Get the class to look for this variable in */
-	BLUEPRINTGRAPH_API UClass* GetVariableSourceClass() const;
+	UClass* GetVariableSourceClass() const;
 
 	/** Get the UProperty for this variable node */
-	BLUEPRINTGRAPH_API UProperty* GetPropertyForVariable() const;
+	UProperty* GetPropertyForVariable() const;
 
 	/** Accessor for the value output pin of the node */
-	BLUEPRINTGRAPH_API UEdGraphPin* GetValuePin() const;
+	UEdGraphPin* GetValuePin() const;
 
 	/** Validates there are no errors in the node */
 	void CheckForErrors(const class UEdGraphSchema_K2* Schema, class FCompilerResultsLog& MessageLog);
@@ -127,7 +127,7 @@ public:
 	 * @param  IconColorOut	A color out, further discerning the variable's type.
 	 * @return A icon representing the specified variable's type.
 	 */
-	BLUEPRINTGRAPH_API static FName GetVariableIconAndColor(UStruct* VarScope, FName VarName, FLinearColor& IconColorOut);
+	static FName GetVariableIconAndColor(UStruct* VarScope, FName VarName, FLinearColor& IconColorOut);
 
 	/**
 	 * Utility method intended to serve as a choke point for various slate 
@@ -137,6 +137,6 @@ public:
 	 * @param  IconColorOut	A color out, further discerning the variable's type.
 	 * @return A icon representing the specified variable's type.
 	 */
-	BLUEPRINTGRAPH_API static FName GetVarIconFromPinType(const FEdGraphPinType& InPinType, FLinearColor& IconColorOut);
+	static FName GetVarIconFromPinType(const FEdGraphPinType& InPinType, FLinearColor& IconColorOut);
 };
 
diff --git a/Engine/Source/Editor/BlueprintGraph/Classes/K2Node_VariableGet.h b/Engine/Source/Editor/BlueprintGraph/Classes/K2Node_VariableGet.h
--- a/Engine/Source/Editor/BlueprintGraph/Classes/K2Node_VariableGet.h
+++ b/Engine/Source/Editor/BlueprintGraph/Classes/K2Node_VariableGet.h
@@ -5,8 +5,8 @@
 #include "K2Node_Variable.h"
 #include "K2Node_VariableGet.generated.h"
 
-UCLASS(MinimalAPI)
-class UK2Node_VariableGet : public UK2Node_Variable
+UCLASS()
+class BLUEPRINTGRAPH_API UK2Node_VariableGet : public UK2Node_Variable
 {
 	GENERATED_UCLASS_BODY()
 
diff --git a/Engine/Source/Editor/GameplayAbilitiesEditor/Classes/GameplayAbilityGraph.h b/Engine/Source/Editor/GameplayAbilitiesEditor/Classes/GameplayAbilityGraph.h
new file mode 100644
--- /dev/null
+++ b/Engine/Source/Editor/GameplayAbilitiesEditor/Classes/GameplayAbilityGraph.h
@@ -0,0 +1,12 @@
+// Copyright 1998-2014 Epic Games, Inc. All Rights Reserved.
+
+#pragma once
+
+#include "GameplayAbilityGraph.generated.h"
+
+UCLASS(MinimalAPI)
+class UGameplayAbilityGraph : public UEdGraph
+{
+	GENERATED_UCLASS_BODY()
+};
+
diff --git a/Engine/Source/Editor/GameplayAbilitiesEditor/Classes/GameplayAbilityGraphSchema.h b/Engine/Source/Editor/GameplayAbilitiesEditor/Classes/GameplayAbilityGraphSchema.h
--- a/Engine/Source/Editor/GameplayAbilitiesEditor/Classes/GameplayAbilityGraphSchema.h
+++ b/Engine/Source/Editor/GameplayAbilitiesEditor/Classes/GameplayAbilityGraphSchema.h
@@ -10,5 +10,25 @@ class UGameplayAbilityGraphSchema : public UEdGraphSchema_K2
 {
 	GENERATED_UCLASS_BODY()
 
-	// Dummy class needed for possible extensions later
+	/**
+	* Creates a new variable getter node and adds it to ParentGraph
+	*
+	* @param	GraphPosition		The location of the new node inside the graph
+	* @param	ParentGraph			The graph to spawn the new node in
+	* @param	VariableName		The name of the variable
+	* @param	Source				The source of the variable
+	* @return	A pointer to the newly spawned node
+	*/
+	virtual class UK2Node_VariableGet* SpawnVariableGetNode(const FVector2D GraphPosition, class UEdGraph* ParentGraph, FName VariableName, UStruct* Source) const override;
+
+	/**
+	* Creates a new variable setter node and adds it to ParentGraph
+	*
+	* @param	GraphPosition		The location of the new node inside the graph
+	* @param	ParentGraph			The graph to spawn the new node in
+	* @param	VariableName		The name of the variable
+	* @param	Source				The source of the variable
+	* @return	A pointer to the newly spawned node
+	*/
+	virtual class UK2Node_VariableSet* SpawnVariableSetNode(const FVector2D GraphPosition, class UEdGraph* ParentGraph, FName VariableName, UStruct* Source) const override;
 };
\ No newline at end of file
diff --git a/Engine/Source/Editor/GameplayAbilitiesEditor/Classes/K2Node_GameplayEffectVariable.h b/Engine/Source/Editor/GameplayAbilitiesEditor/Classes/K2Node_GameplayEffectVariable.h
new file mode 100644
--- /dev/null
+++ b/Engine/Source/Editor/GameplayAbilitiesEditor/Classes/K2Node_GameplayEffectVariable.h
@@ -0,0 +1,40 @@
+// Copyright 1998-2014 Epic Games, Inc. All Rights Reserved.
+
+#pragma once
+
+#include "K2Node.h"
+#include "K2Node_VariableGet.h"
+#include "GameplayEffect.h"
+
+#include "K2Node_GameplayEffectVariable.generated.h"
+
+UCLASS(MinimalAPI)
+class UK2Node_GameplayEffectVariable : public UK2Node_VariableGet // only need a custom node for the getter
+{
+	GENERATED_UCLASS_BODY()
+
+public:
+	class FGameplayEffectDisplayInfo
+	{
+	public:
+		struct ModifierDisplayInfo
+		{
+			FString DisplayName;
+			float Magnitude;
+			EGameplayModOp::Type ModOp;
+		};
+
+		FGameplayEffectDisplayInfo() {}
+
+		void Init(const UGameplayEffect& GameplayEffect);
+
+		float Duration;
+		float Period;
+		float ChanceToApplyToTarget;
+		float ChanceToExecuteOnGameplayEffect;
+
+		TArray<ModifierDisplayInfo> AttributeModifiers;
+	};
+
+	FGameplayEffectDisplayInfo GameplayEffectInfo;
+};
diff --git a/Engine/Source/Editor/GameplayAbilitiesEditor/GameplayAbilitiesEditor.Build.cs b/Engine/Source/Editor/GameplayAbilitiesEditor/GameplayAbilitiesEditor.Build.cs
--- a/Engine/Source/Editor/GameplayAbilitiesEditor/GameplayAbilitiesEditor.Build.cs
+++ b/Engine/Source/Editor/GameplayAbilitiesEditor/GameplayAbilitiesEditor.Build.cs
@@ -8,6 +8,8 @@ namespace UnrealBuildTool.Rules
 		{
             PrivateIncludePaths.AddRange(
                 new string[] {
+                    "Editor/GraphEditor/Private",
+				    "Editor/Kismet/Private",
 					"Editor/GameplayAbilitiesEditor/Private",
 				}
 			);
diff --git a/Engine/Source/Editor/GameplayAbilitiesEditor/Private/AbilitySystemEditorPrivatePCH.h b/Engine/Source/Editor/GameplayAbilitiesEditor/Private/AbilitySystemEditorPrivatePCH.h
--- a/Engine/Source/Editor/GameplayAbilitiesEditor/Private/AbilitySystemEditorPrivatePCH.h
+++ b/Engine/Source/Editor/GameplayAbilitiesEditor/Private/AbilitySystemEditorPrivatePCH.h
@@ -12,4 +12,5 @@
 
 #include "Core.h"
 #include "Engine.h"
+#include "Editor/BlueprintGraph/Classes/K2Node_Variable.h"
 #include "BlueprintGraphDefinitions.h"
\ No newline at end of file
diff --git a/Engine/Source/Editor/GameplayAbilitiesEditor/Private/GameplayAbilitiesBlueprintFactory.cpp b/Engine/Source/Editor/GameplayAbilitiesEditor/Private/GameplayAbilitiesBlueprintFactory.cpp
--- a/Engine/Source/Editor/GameplayAbilitiesEditor/Private/GameplayAbilitiesBlueprintFactory.cpp
+++ b/Engine/Source/Editor/GameplayAbilitiesEditor/Private/GameplayAbilitiesBlueprintFactory.cpp
@@ -7,6 +7,8 @@
 
 #include "Abilities/GameplayAbility.h"
 #include "GameplayAbilityBlueprint.h"
+#include "GameplayAbilityGraph.h"
+#include "GameplayAbilityGraphSchema.h"
 
 #include "Editor/ClassViewer/Public/ClassViewerFilter.h"
 
@@ -269,6 +271,21 @@ UObject* UGameplayAbilitiesBlueprintFactory::FactoryCreateNew(UClass* Class, UOb
 	{
 		UGameplayAbilityBlueprint* NewBP = CastChecked<UGameplayAbilityBlueprint>(FKismetEditorUtilities::CreateBlueprint(ParentClass, InParent, Name, BlueprintType, UGameplayAbilityBlueprint::StaticClass(), UGameplayAbilityBlueprintGeneratedClass::StaticClass(), CallingContext));
 
+		if (NewBP)
+		{
+			UGameplayAbilityBlueprint* AbilityBP = UGameplayAbilityBlueprint::FindRootGameplayAbilityBlueprint(NewBP);
+			if (AbilityBP == NULL)
+			{
+				const UEdGraphSchema_K2* K2Schema = GetDefault<UEdGraphSchema_K2>();
+
+				// Only allow a gameplay ability graph if there isn't one in a parent blueprint
+				UEdGraph* NewGraph = FBlueprintEditorUtils::CreateNewGraph(NewBP, TEXT("Gameplay Ability Graph"), UGameplayAbilityGraph::StaticClass(), UGameplayAbilityGraphSchema::StaticClass());
+				FBlueprintEditorUtils::AddDomainSpecificGraph(NewBP, NewGraph);
+				NewBP->LastEditedDocuments.Add(NewGraph);
+				NewGraph->bAllowDeletion = false;
+			}
+		}
+
 		return NewBP;
 	}
 }
diff --git a/Engine/Source/Editor/GameplayAbilitiesEditor/Private/GameplayAbilitiesEditorModule.cpp b/Engine/Source/Editor/GameplayAbilitiesEditor/Private/GameplayAbilitiesEditorModule.cpp
--- a/Engine/Source/Editor/GameplayAbilitiesEditor/Private/GameplayAbilitiesEditorModule.cpp
+++ b/Engine/Source/Editor/GameplayAbilitiesEditor/Private/GameplayAbilitiesEditorModule.cpp
@@ -9,6 +9,7 @@
 #include "AssetToolsModule.h"
 #include "AssetTypeActions_GameplayAbilitiesBlueprint.h"
 #include "GameplayAbilitiesGraphPanelPinFactory.h"
+#include "GameplayAbilitiesGraphPanelNodeFactory.h"
 
 class FGameplayAbilitiesEditorModule : public IGameplayAbilitiesEditorModule
 {
@@ -42,8 +43,12 @@ void FGameplayAbilitiesEditorModule::StartupModule()
 	AssetTools.RegisterAssetTypeActions(Action);
 	CreatedAssetTypeActions.Add(Action);
 
+	// Register factories for pins and nodes
 	TSharedPtr<FGameplayAbilitiesGraphPanelPinFactory> GameplayAbilitiesGraphPanelPinFactory = MakeShareable(new FGameplayAbilitiesGraphPanelPinFactory());
 	FEdGraphUtilities::RegisterVisualPinFactory(GameplayAbilitiesGraphPanelPinFactory);
+
+	TSharedPtr<FGameplayAbilitiesGraphPanelNodeFactory> GameplayAbilitiesGraphPanelNodeFactory = MakeShareable(new FGameplayAbilitiesGraphPanelNodeFactory());
+	FEdGraphUtilities::RegisterVisualNodeFactory(GameplayAbilitiesGraphPanelNodeFactory);
 }
 
 void FGameplayAbilitiesEditorModule::RegisterAssetTypeAction(IAssetTools& AssetTools, TSharedRef<IAssetTypeActions> Action)
diff --git a/Engine/Source/Editor/GameplayAbilitiesEditor/Private/GameplayAbilitiesGraphPanelNodeFactory.h b/Engine/Source/Editor/GameplayAbilitiesEditor/Private/GameplayAbilitiesGraphPanelNodeFactory.h
new file mode 100644
--- /dev/null
+++ b/Engine/Source/Editor/GameplayAbilitiesEditor/Private/GameplayAbilitiesGraphPanelNodeFactory.h
@@ -0,0 +1,20 @@
+// Copyright 1998-2014 Epic Games, Inc. All Rights Reserved.
+
+#pragma once
+
+#include "EdGraphUtilities.h"
+#include "SGraphNodeK2GameplayEffectVar.h"
+#include "K2Node_GameplayEffectVariable.h"
+
+class FGameplayAbilitiesGraphPanelNodeFactory: public FGraphPanelNodeFactory
+{
+	virtual TSharedPtr<class SGraphNode> CreateNode(class UEdGraphNode* InNode) const override
+	{
+		if (UK2Node_GameplayEffectVariable* GameplayEffectNode = Cast<UK2Node_GameplayEffectVariable>(InNode))
+		{
+			return SNew(SGraphNodeK2GameplayEffectVar, GameplayEffectNode);
+		}
+
+		return NULL;
+	}
+};
diff --git a/Engine/Source/Editor/GameplayAbilitiesEditor/Private/GameplayAbilityGraph.cpp b/Engine/Source/Editor/GameplayAbilitiesEditor/Private/GameplayAbilityGraph.cpp
new file mode 100644
--- /dev/null
+++ b/Engine/Source/Editor/GameplayAbilitiesEditor/Private/GameplayAbilityGraph.cpp
@@ -0,0 +1,16 @@
+// Copyright 1998-2014 Epic Games, Inc. All Rights Reserved.
+
+#include "AbilitySystemEditorPrivatePCH.h"
+#include "GameplayAbilityGraph.h"
+
+#define LOCTEXT_NAMESPACE "GameplayAbilityGraph"
+
+/////////////////////////////////////////////////////
+// UGameplayAbilityGraph
+
+UGameplayAbilityGraph::UGameplayAbilityGraph(const FPostConstructInitializeProperties& PCIP)
+	: Super(PCIP)
+{
+}
+
+#undef LOCTEXT_NAMESPACE
diff --git a/Engine/Source/Editor/GameplayAbilitiesEditor/Private/GameplayAbilityGraphSchema.cpp b/Engine/Source/Editor/GameplayAbilitiesEditor/Private/GameplayAbilityGraphSchema.cpp
--- a/Engine/Source/Editor/GameplayAbilitiesEditor/Private/GameplayAbilityGraphSchema.cpp
+++ b/Engine/Source/Editor/GameplayAbilitiesEditor/Private/GameplayAbilityGraphSchema.cpp
@@ -1,10 +1,47 @@
 // Copyright 1998-2014 Epic Games, Inc. All Rights Reserved.
 
 #include "AbilitySystemEditorPrivatePCH.h"
-
+#include "K2Node_GameplayEffectVariable.h"
+#include "GameplayEffect.h"
 #include "GameplayAbilityGraphSchema.h"
+#include "Kismet/KismetSystemLibrary.h"
 
 UGameplayAbilityGraphSchema::UGameplayAbilityGraphSchema(const FPostConstructInitializeProperties& PCIP)
 : Super(PCIP)
 {
+}
+
+UK2Node_VariableGet* UGameplayAbilityGraphSchema::SpawnVariableGetNode(const FVector2D GraphPosition, class UEdGraph* ParentGraph, FName VariableName, UStruct* Source) const
+{
+	// Perform handling to create custom nodes for some classes
+	UProperty* VarProp = FindField<UProperty>(Source, VariableName);
+	UObjectProperty* ObjProp = Cast<UObjectProperty>(VarProp);
+	if (ObjProp)
+	{
+		UBlueprint* Blueprint = FBlueprintEditorUtils::FindBlueprintForGraph(ParentGraph);
+		TSubclassOf<UObject> GenClass = Blueprint->GeneratedClass;
+
+		UObject* ActiveObject = GenClass->GetDefaultObject();
+
+		// If the variable is a GameplayEffect create a custom node to show it
+		UGameplayEffect* GameplayEffect = Cast<UGameplayEffect>(ObjProp->GetObjectPropertyValue_InContainer(ActiveObject));
+		if (GameplayEffect)
+		{
+			UK2Node_GameplayEffectVariable* NodeTemplate = NewObject<UK2Node_GameplayEffectVariable>();
+			UEdGraphSchema_K2::ConfigureVarNode(NodeTemplate, VariableName, Source, Blueprint);
+
+			UK2Node_GameplayEffectVariable* VariableNode = FEdGraphSchemaAction_K2NewNode::SpawnNodeFromTemplate<UK2Node_GameplayEffectVariable>(ParentGraph, NodeTemplate, GraphPosition);
+			VariableNode->GameplayEffectInfo.Init(*GameplayEffect);
+
+			return VariableNode;
+		}
+	}
+
+	// Couldn't find an appropriate custom node for this variable, use the generic case
+	return Super::SpawnVariableGetNode(GraphPosition, ParentGraph, VariableName, Source);
+}
+
+UK2Node_VariableSet* UGameplayAbilityGraphSchema::SpawnVariableSetNode(const FVector2D GraphPosition, class UEdGraph* ParentGraph, FName VariableName, UStruct* Source) const
+{
+	return Super::SpawnVariableSetNode(GraphPosition, ParentGraph, VariableName, Source);
 }
\ No newline at end of file
diff --git a/Engine/Source/Editor/GameplayAbilitiesEditor/Private/K2Node_GameplayEffectVariable.cpp b/Engine/Source/Editor/GameplayAbilitiesEditor/Private/K2Node_GameplayEffectVariable.cpp
new file mode 100644
--- /dev/null
+++ b/Engine/Source/Editor/GameplayAbilitiesEditor/Private/K2Node_GameplayEffectVariable.cpp
@@ -0,0 +1,35 @@
+// Copyright 1998-2014 Epic Games, Inc. All Rights Reserved.
+
+
+#include "AbilitySystemEditorPrivatePCH.h"
+#include "K2Node_GameplayEffectVariable.h"
+
+#define LOCTEXT_NAMESPACE "K2Node_GameplayEffectVariable"
+
+UK2Node_GameplayEffectVariable::UK2Node_GameplayEffectVariable(const class FPostConstructInitializeProperties& PCIP)
+	: Super(PCIP)
+{
+}
+
+void UK2Node_GameplayEffectVariable::FGameplayEffectDisplayInfo::Init(const UGameplayEffect& GameplayEffect)
+{
+	Duration = GameplayEffect.Duration.Value;
+	Period = GameplayEffect.Period.Value;
+	ChanceToApplyToTarget = GameplayEffect.ChanceToApplyToTarget.Value;
+	ChanceToExecuteOnGameplayEffect = GameplayEffect.ChanceToExecuteOnGameplayEffect.Value;
+
+	for (int Idx = 0; Idx < GameplayEffect.Modifiers.Num(); ++Idx)
+	{
+		if (GameplayEffect.Modifiers[Idx].ModifierType == EGameplayMod::Attribute)
+		{
+			ModifierDisplayInfo ModInfo;
+			ModInfo.DisplayName = GameplayEffect.Modifiers[Idx].Attribute.GetName();
+			ModInfo.Magnitude = GameplayEffect.Modifiers[Idx].Magnitude.Value;
+			ModInfo.ModOp = GameplayEffect.Modifiers[Idx].ModifierOp;
+
+			AttributeModifiers.Add(ModInfo);
+		}
+	}
+}
+
+#undef LOCTEXT_NAMESPACE
diff --git a/Engine/Source/Editor/GraphEditor/Private/KismetNodes/SGraphNodeK2Var.h b/Engine/Source/Editor/GraphEditor/Private/KismetNodes/SGraphNodeK2Var.h
--- a/Engine/Source/Editor/GraphEditor/Private/KismetNodes/SGraphNodeK2Var.h
+++ b/Engine/Source/Editor/GraphEditor/Private/KismetNodes/SGraphNodeK2Var.h
@@ -4,7 +4,7 @@
 #ifndef __SGraphNodeK2Var_h__
 #define __SGraphNodeK2Var_h__
 
-class SGraphNodeK2Var : public SGraphNodeK2Base
+class GRAPHEDITOR_API SGraphNodeK2Var : public SGraphNodeK2Base
 {
 public:
 	SLATE_BEGIN_ARGS(SGraphNodeK2Var){}
diff --git a/Engine/Source/Editor/Kismet/Private/BPDelegateDragDropAction.h b/Engine/Source/Editor/Kismet/Private/BPDelegateDragDropAction.h
--- a/Engine/Source/Editor/Kismet/Private/BPDelegateDragDropAction.h
+++ b/Engine/Source/Editor/Kismet/Private/BPDelegateDragDropAction.h
@@ -6,7 +6,7 @@
 #include "BPVariableDragDropAction.h"
 
 /** DragDropAction class for dropping a Variable onto a graph */
-class FKismetDelegateDragDropAction : public FKismetVariableDragDropAction
+class KISMET_API FKismetDelegateDragDropAction : public FKismetVariableDragDropAction
 {
 public:
 	DRAG_DROP_OPERATOR_TYPE(FKismetDelegateDragDropAction, FKismetVariableDragDropAction)
diff --git a/Engine/Source/Editor/Kismet/Private/BPFunctionDragDropAction.h b/Engine/Source/Editor/Kismet/Private/BPFunctionDragDropAction.h
--- a/Engine/Source/Editor/Kismet/Private/BPFunctionDragDropAction.h
+++ b/Engine/Source/Editor/Kismet/Private/BPFunctionDragDropAction.h
@@ -9,7 +9,7 @@
 * FKismetDragDropAction
 *******************************************************************************/
 
-class FKismetDragDropAction : public FGraphSchemaActionDragDropAction
+class KISMET_API FKismetDragDropAction : public FGraphSchemaActionDragDropAction
 {
 public:
 	DRAG_DROP_OPERATOR_TYPE(FKismetDragDropAction, FGraphSchemaActionDragDropAction)
@@ -45,7 +45,7 @@ protected:
 * FKismetFunctionDragDropAction
 *******************************************************************************/
 
-class FKismetFunctionDragDropAction : public FKismetDragDropAction
+class KISMET_API FKismetFunctionDragDropAction : public FKismetDragDropAction
 {
 public:
 	DRAG_DROP_OPERATOR_TYPE(FKismetFunctionDragDropAction, FKismetDragDropAction)
@@ -78,7 +78,7 @@ protected:
 * FKismetMacroDragDropAction
 *******************************************************************************/
 
-class FKismetMacroDragDropAction : public FKismetDragDropAction
+class KISMET_API FKismetMacroDragDropAction : public FKismetDragDropAction
 {
 public:
 	DRAG_DROP_OPERATOR_TYPE(FKismetMacroDragDropAction, FKismetDragDropAction)
diff --git a/Engine/Source/Editor/Kismet/Private/BPVariableDragDropAction.cpp b/Engine/Source/Editor/Kismet/Private/BPVariableDragDropAction.cpp
--- a/Engine/Source/Editor/Kismet/Private/BPVariableDragDropAction.cpp
+++ b/Engine/Source/Editor/Kismet/Private/BPVariableDragDropAction.cpp
@@ -312,7 +312,7 @@ FReply FKismetVariableDragDropAction::DroppedOnPin(FVector2D ScreenPosition, FVe
 					Action.NodeTemplate = VarNode;
 
 					UBlueprint* DropOnBlueprint = FBlueprintEditorUtils::FindBlueprintForGraph(TargetPin->GetOwningNode()->GetGraph());
-					ConfigureVarNode(VarNode, VariableName, VariableSource.Get(), DropOnBlueprint);
+					UEdGraphSchema_K2::ConfigureVarNode(VarNode, VariableName, VariableSource.Get(), DropOnBlueprint);
 
 					Action.PerformAction(TargetPin->GetOwningNode()->GetGraph(), TargetPin, GraphPosition);
 				}
@@ -348,7 +348,7 @@ FReply FKismetVariableDragDropAction::DroppedOnNode(FVector2D ScreenPosition, FV
 			TargetNode->Modify();
 			Pin->Modify();
 
-			ConfigureVarNode(TargetNode, VariableName, VariableSource.Get(), DropOnBlueprint);
+			UEdGraphSchema_K2::ConfigureVarNode(TargetNode, VariableName, VariableSource.Get(), DropOnBlueprint);
 
 
 			if ((Pin == NULL) || (Pin->LinkedTo.Num() == BadLinks.Num()) || (Schema == NULL))
@@ -376,46 +376,26 @@ FReply FKismetVariableDragDropAction::DroppedOnNode(FVector2D ScreenPosition, FV
 	return FReply::Unhandled();
 }
 
-void FKismetVariableDragDropAction::ConfigureVarNode(UK2Node_Variable* InVarNode, FName InVariableName, UStruct* InVariableSource, UBlueprint* InTargetBlueprint)
-{
-	// See if this is a 'self context' (ie. blueprint class is owner (or child of owner) of dropped var class)
-	if ((InVariableSource == NULL) || InTargetBlueprint->SkeletonGeneratedClass->IsChildOf(InVariableSource))
-	{
-		InVarNode->VariableReference.SetSelfMember(InVariableName);
-	}
-	else if(InVariableSource->IsA(UStruct::StaticClass()))
-	{
-		InVarNode->VariableReference.SetLocalMember(InVariableName, InVariableSource, FBlueprintEditorUtils::FindLocalVariableGuidByName(InTargetBlueprint, InVariableName));
-	}
-	else
-	{
-		InVarNode->VariableReference.SetExternalMember(InVariableName, CastChecked<UClass>(InVariableSource));
-	}
-}
-
-
 void FKismetVariableDragDropAction::MakeGetter(FNodeConstructionParams InParams)
 {
 	check(InParams.Graph);
 
-	UK2Node_VariableGet* GetVarNodeTemplate = NewObject<UK2Node_VariableGet>();
-	UBlueprint* DropOnBlueprint = FBlueprintEditorUtils::FindBlueprintForGraph(InParams.Graph);
-
-	ConfigureVarNode(GetVarNodeTemplate, InParams.VariableName, InParams.VariableSource.Get(), DropOnBlueprint);
-
-	FEdGraphSchemaAction_K2NewNode::SpawnNodeFromTemplate<UK2Node_VariableGet>(InParams.Graph, GetVarNodeTemplate, InParams.GraphPosition);
+	const UEdGraphSchema_K2* K2_Schema = Cast<const UEdGraphSchema_K2>(InParams.Graph->GetSchema());
+	if (K2_Schema)
+	{
+		K2_Schema->SpawnVariableGetNode(InParams.GraphPosition, InParams.Graph, InParams.VariableName, InParams.VariableSource.Get());
+	}
 }
 
 void FKismetVariableDragDropAction::MakeSetter(FNodeConstructionParams InParams)
 {
 	check(InParams.Graph);
 
-	UK2Node_VariableSet* SetVarNodeTemplate = NewObject<UK2Node_VariableSet>();
-	UBlueprint* DropOnBlueprint = FBlueprintEditorUtils::FindBlueprintForGraph(InParams.Graph);
-
-	ConfigureVarNode(SetVarNodeTemplate, InParams.VariableName, InParams.VariableSource.Get(), DropOnBlueprint);
-
-	FEdGraphSchemaAction_K2NewNode::SpawnNodeFromTemplate<UK2Node_VariableSet>(InParams.Graph, SetVarNodeTemplate, InParams.GraphPosition);
+	const UEdGraphSchema_K2* K2_Schema = Cast<const UEdGraphSchema_K2>(InParams.Graph->GetSchema());
+	if (K2_Schema)
+	{
+		K2_Schema->SpawnVariableSetNode(InParams.GraphPosition, InParams.Graph, InParams.VariableName, InParams.VariableSource.Get());
+	}
 }
 
 bool FKismetVariableDragDropAction::CanExecuteMakeSetter(FNodeConstructionParams InParams, UProperty* InVariableProperty)
diff --git a/Engine/Source/Editor/Kismet/Private/BPVariableDragDropAction.h b/Engine/Source/Editor/Kismet/Private/BPVariableDragDropAction.h
--- a/Engine/Source/Editor/Kismet/Private/BPVariableDragDropAction.h
+++ b/Engine/Source/Editor/Kismet/Private/BPVariableDragDropAction.h
@@ -6,16 +6,16 @@
 #include "Editor/GraphEditor/Public/GraphEditorDragDropAction.h"
 
 /** DragDropAction class for dropping a Variable onto a graph */
-class FKismetVariableDragDropAction : public FGraphEditorDragDropAction
+class KISMET_API FKismetVariableDragDropAction : public FGraphEditorDragDropAction
 {
 public:
 	DRAG_DROP_OPERATOR_TYPE(FKismetVariableDragDropAction, FGraphEditorDragDropAction)
 
 	// FGraphEditorDragDropAction interface
-	virtual void HoverTargetChanged();
+	virtual void HoverTargetChanged() override;
 	virtual FReply DroppedOnPin(FVector2D ScreenPosition, FVector2D GraphPosition) override;
 	virtual FReply DroppedOnNode(FVector2D ScreenPosition, FVector2D GraphPosition) override;
-	virtual FReply DroppedOnPanel( const TSharedRef< class SWidget >& Panel, FVector2D ScreenPosition, FVector2D GraphPosition, UEdGraph& Graph) override;
+	virtual FReply DroppedOnPanel(const TSharedRef< class SWidget >& Panel, FVector2D ScreenPosition, FVector2D GraphPosition, UEdGraph& Graph) override;
 	virtual FReply DroppedOnAction(TSharedRef<struct FEdGraphSchemaAction> Action) override;
 	virtual FReply DroppedOnCategory(FString Category) override;
 	// End of FGraphEditorDragDropAction
@@ -58,9 +58,6 @@ public:
 		return VariableProperty;
 	}
 
-	/** Configure the supplied variable node based on this action's info */
-	static void ConfigureVarNode(UK2Node_Variable* InVarNode, FName InVariableName, UStruct* InVariableSource, UBlueprint* InTargetBlueprint);
-
 	/** Set if operation is modified by alt */
 	void SetAltDrag(bool InIsAltDrag) {	bAltDrag = InIsAltDrag;}
 
