#summary Can now deprecate Blueprint classes via a checkbox in the Blueprint properties

#add Added a checkbox to deprecate a Blueprint, the Blueprint will no longer be placeable, children classes can no longer be made of it (or any of it's children) and no new variables can be designated in Blueprints of the BP type.

#change AddSubMenu will now take the tooltip as an FText attribute so the tooltip can be hooked up to a function

#ttp 304153 - LICENSEE: K2:  Add support for deprecating a blueprint class

diff --git a/Engine/Source/Developer/AssetTools/Private/AssetTypeActions/AssetTypeActions_Blueprint.cpp b/Engine/Source/Developer/AssetTools/Private/AssetTypeActions/AssetTypeActions_Blueprint.cpp
--- a/Engine/Source/Developer/AssetTools/Private/AssetTypeActions/AssetTypeActions_Blueprint.cpp
+++ b/Engine/Source/Developer/AssetTools/Private/AssetTypeActions/AssetTypeActions_Blueprint.cpp
@@ -38,13 +38,17 @@ void FAssetTypeActions_Blueprint::GetActions( const TArray<UObject*>& InObjects,
 
 	if ( Blueprints.Num() == 1 && CanCreateNewDerivedBlueprint() )
 	{
+		TAttribute<FText>::FGetter DynamicTooltipGetter;
+		DynamicTooltipGetter.BindSP(this, &FAssetTypeActions_Blueprint::GetNewDerivedBlueprintTooltip, Blueprints[0]);
+		TAttribute<FText> DynamicTooltipAttribute = TAttribute<FText>::Create(DynamicTooltipGetter);
+
 		MenuBuilder.AddMenuEntry(
 			LOCTEXT("Blueprint_NewDerivedBlueprint", "Create Blueprint based on this"),
-			LOCTEXT("Blueprint_NewDerivedBlueprintTooltip", "Creates a blueprint based on the selected blueprint."),
+			DynamicTooltipAttribute,
 			FSlateIcon(),
 			FUIAction(
 				FExecuteAction::CreateSP( this, &FAssetTypeActions_Blueprint::ExecuteNewDerivedBlueprint, Blueprints[0] ),
-				FCanExecuteAction()
+				FCanExecuteAction::CreateSP( this, &FAssetTypeActions_Blueprint::CanExecuteNewDerivedBlueprint, Blueprints[0] )
 				)
 			);
 	}
@@ -166,6 +170,23 @@ void FAssetTypeActions_Blueprint::ExecuteNewDerivedBlueprint(TWeakObjectPtr<UBlu
 	}
 }
 
+FText FAssetTypeActions_Blueprint::GetNewDerivedBlueprintTooltip(TWeakObjectPtr<UBlueprint> InObject)
+{
+	if(InObject->GeneratedClass->HasAnyClassFlags(CLASS_Deprecated))
+	{
+		return LOCTEXT("Blueprint_NewDerivedBlueprintIsDeprecatedTooltip", "Blueprint class is deprecated, cannot derive a child Blueprint!");
+	}
+	else
+	{
+		return LOCTEXT("Blueprint_NewDerivedBlueprintTooltip", "Creates a blueprint based on the selected blueprint.");
+	}
+}
+
+bool FAssetTypeActions_Blueprint::CanExecuteNewDerivedBlueprint(TWeakObjectPtr<UBlueprint> InObject)
+{
+	return !InObject->GeneratedClass->HasAnyClassFlags(CLASS_Deprecated);
+}
+
 bool FAssetTypeActions_Blueprint::ShouldUseDataOnlyEditor( const UBlueprint* Blueprint ) const
 {
 	return FBlueprintEditorUtils::IsDataOnlyBlueprint(Blueprint) 
diff --git a/Engine/Source/Developer/AssetTools/Private/AssetTypeActions/AssetTypeActions_Blueprint.h b/Engine/Source/Developer/AssetTools/Private/AssetTypeActions/AssetTypeActions_Blueprint.h
--- a/Engine/Source/Developer/AssetTools/Private/AssetTypeActions/AssetTypeActions_Blueprint.h
+++ b/Engine/Source/Developer/AssetTools/Private/AssetTypeActions/AssetTypeActions_Blueprint.h
@@ -37,4 +37,10 @@ private:
 
 	/** Returns true if the blueprint is data only */
 	bool ShouldUseDataOnlyEditor( const UBlueprint* Blueprint ) const;
+
+	/** Returns the tooltip to display when attempting to derive a Blueprint */
+	FText GetNewDerivedBlueprintTooltip(TWeakObjectPtr<UBlueprint> InObject);
+
+	/** Returns TRUE if you can derive a Blueprint */
+	bool CanExecuteNewDerivedBlueprint(TWeakObjectPtr<UBlueprint> InObject);
 };
diff --git a/Engine/Source/Developer/AssetTools/Private/AssetTypeActions/AssetTypeActions_BlueprintGeneratedClass.cpp b/Engine/Source/Developer/AssetTools/Private/AssetTypeActions/AssetTypeActions_BlueprintGeneratedClass.cpp
--- a/Engine/Source/Developer/AssetTools/Private/AssetTypeActions/AssetTypeActions_BlueprintGeneratedClass.cpp
+++ b/Engine/Source/Developer/AssetTools/Private/AssetTypeActions/AssetTypeActions_BlueprintGeneratedClass.cpp
@@ -38,13 +38,17 @@ void FAssetTypeActions_BlueprintGeneratedClass::GetActions( const TArray<UObject
 
 	if ( BlueprintGeneratedClasses.Num() == 1 )
 	{
+		TAttribute<FText>::FGetter DynamicTooltipGetter;
+		DynamicTooltipGetter.BindSP(this, &FAssetTypeActions_BlueprintGeneratedClass::GetNewDerivedBlueprintTooltip, BlueprintGeneratedClasses[0]);
+		TAttribute<FText> DynamicTooltipAttribute = TAttribute<FText>::Create(DynamicTooltipGetter);
+
 		MenuBuilder.AddMenuEntry(
 			LOCTEXT("Blueprint_NewDerivedBlueprint", "Create Blueprint based on this"),
-			LOCTEXT("Blueprint_NewDerivedBlueprintTooltip", "Creates a blueprint based on the selected blueprint."),
+			DynamicTooltipAttribute,
 			FSlateIcon(),
 			FUIAction(
 				FExecuteAction::CreateSP( this, &FAssetTypeActions_BlueprintGeneratedClass::ExecuteNewDerivedBlueprint, BlueprintGeneratedClasses[0] ),
-				FCanExecuteAction()
+				FCanExecuteAction::CreateSP( this, &FAssetTypeActions_BlueprintGeneratedClass::CanExecuteNewDerivedBlueprint, BlueprintGeneratedClasses[0] )
 				)
 			);
 	}
@@ -163,6 +167,23 @@ void FAssetTypeActions_BlueprintGeneratedClass::ExecuteNewDerivedBlueprint(TWeak
 	}
 }
 
+FText FAssetTypeActions_BlueprintGeneratedClass::GetNewDerivedBlueprintTooltip(TWeakObjectPtr<UBlueprintGeneratedClass> InObject)
+{
+	if(InObject->HasAnyClassFlags(CLASS_Deprecated))
+	{
+		return LOCTEXT("Blutility_NewDerivedBlueprintIsDeprecatedTooltip", "Blueprint class is deprecated, cannot derive a child Blueprint!");
+	}
+	else
+	{
+		return LOCTEXT("Blutility_NewDerivedBlueprintTooltip", "Creates a blueprint based on the selected blueprint.");
+	}
+}
+
+bool FAssetTypeActions_BlueprintGeneratedClass::CanExecuteNewDerivedBlueprint(TWeakObjectPtr<UBlueprintGeneratedClass> InObject)
+{
+	return !InObject->HasAnyClassFlags(CLASS_Deprecated);
+}
+
 bool FAssetTypeActions_BlueprintGeneratedClass::ShouldUseDataOnlyEditor( const UBlueprint* Blueprint ) const
 {
 	return FBlueprintEditorUtils::IsDataOnlyBlueprint(Blueprint) 
diff --git a/Engine/Source/Developer/AssetTools/Private/AssetTypeActions/AssetTypeActions_BlueprintGeneratedClass.h b/Engine/Source/Developer/AssetTools/Private/AssetTypeActions/AssetTypeActions_BlueprintGeneratedClass.h
--- a/Engine/Source/Developer/AssetTools/Private/AssetTypeActions/AssetTypeActions_BlueprintGeneratedClass.h
+++ b/Engine/Source/Developer/AssetTools/Private/AssetTypeActions/AssetTypeActions_BlueprintGeneratedClass.h
@@ -33,4 +33,10 @@ private:
 
 	/** Returns true if the blueprint is data only */
 	bool ShouldUseDataOnlyEditor( const UBlueprint* Blueprint ) const;
+
+	/** Returns the tooltip to display when attempting to derive a Blueprint */
+	FText GetNewDerivedBlueprintTooltip(TWeakObjectPtr<UBlueprintGeneratedClass> InObject);
+
+	/** Returns TRUE if you can derive a Blueprint */
+	bool CanExecuteNewDerivedBlueprint(TWeakObjectPtr<UBlueprintGeneratedClass> InObject);
 };
diff --git a/Engine/Source/Editor/BlueprintGraph/Private/EdGraphSchema_K2.cpp b/Engine/Source/Editor/BlueprintGraph/Private/EdGraphSchema_K2.cpp
--- a/Engine/Source/Editor/BlueprintGraph/Private/EdGraphSchema_K2.cpp
+++ b/Engine/Source/Editor/BlueprintGraph/Private/EdGraphSchema_K2.cpp
@@ -2607,7 +2607,7 @@ void UEdGraphSchema_K2::GetVariableSubtypes(const FString& Type, TArray<UObject*
 		for (TObjectIterator<UClass> ClassIt; ClassIt; ++ClassIt)
 		{
 			UClass* CurrentClass = *ClassIt;
-			if (UEdGraphSchema_K2::IsAllowableBlueprintVariableType(CurrentClass))
+			if (UEdGraphSchema_K2::IsAllowableBlueprintVariableType(CurrentClass) && !CurrentClass->HasAnyClassFlags(CLASS_Deprecated))
 			{
 				SubtypesList.Add(CurrentClass);
 			}
@@ -2619,7 +2619,7 @@ void UEdGraphSchema_K2::GetVariableSubtypes(const FString& Type, TArray<UObject*
 		for (TObjectIterator<UClass> ClassIt; ClassIt; ++ClassIt)
 		{
 			UClass* CurrentClass = *ClassIt;
-			if (!CurrentClass->IsChildOf(UInterface::StaticClass()) && UEdGraphSchema_K2::IsAllowableBlueprintVariableType(CurrentClass))
+			if (!CurrentClass->IsChildOf(UInterface::StaticClass()) && UEdGraphSchema_K2::IsAllowableBlueprintVariableType(CurrentClass) && !CurrentClass->HasAnyClassFlags(CLASS_Deprecated))
 			{
 				SubtypesList.Add(CurrentClass);
 			}
diff --git a/Engine/Source/Editor/Blutility/Private/AssetTypeActions_EditorUtilityBlueprint.cpp b/Engine/Source/Editor/Blutility/Private/AssetTypeActions_EditorUtilityBlueprint.cpp
--- a/Engine/Source/Editor/Blutility/Private/AssetTypeActions_EditorUtilityBlueprint.cpp
+++ b/Engine/Source/Editor/Blutility/Private/AssetTypeActions_EditorUtilityBlueprint.cpp
@@ -82,13 +82,17 @@ void FAssetTypeActions_EditorUtilityBlueprint::GetActions(const TArray<UObject*>
 
 	if (Blueprints.Num() == 1)
 	{
+		TAttribute<FText>::FGetter DynamicTooltipGetter;
+		DynamicTooltipGetter.BindSP(this, &FAssetTypeActions_EditorUtilityBlueprint::GetNewDerivedBlueprintTooltip, Blueprints[0]);
+		TAttribute<FText> DynamicTooltipAttribute = TAttribute<FText>::Create(DynamicTooltipGetter);
+
 		MenuBuilder.AddMenuEntry(
 			LOCTEXT("Blutility_NewDerivedBlueprint", "Create Blueprint based on this"),
-			LOCTEXT("Blutility_NewDerivedBlueprintTooltip", "Creates a blueprint based on the selected blueprint."),
+			DynamicTooltipAttribute,
 			FSlateIcon(),
 			FUIAction(
 				FExecuteAction::CreateSP( this, &FAssetTypeActions_EditorUtilityBlueprint::ExecuteNewDerivedBlueprint, Blueprints[0] ),
-				FCanExecuteAction()
+				FCanExecuteAction::CreateSP( this, &FAssetTypeActions_EditorUtilityBlueprint::CanExecuteNewDerivedBlueprint, Blueprints[0] )
 				)
 			);
 	}
@@ -193,6 +197,23 @@ void FAssetTypeActions_EditorUtilityBlueprint::ExecuteNewDerivedBlueprint(TWeakO
 	}
 }
 
+FText FAssetTypeActions_EditorUtilityBlueprint::GetNewDerivedBlueprintTooltip(TWeakObjectPtr<class UEditorUtilityBlueprint> InObject)
+{
+	if(InObject->GeneratedClass->HasAnyClassFlags(CLASS_Deprecated))
+	{
+		return LOCTEXT("Blutility_NewDerivedBlueprintIsDeprecatedTooltip", "Blueprint class is deprecated, cannot derive a child Blueprint!");
+	}
+	else
+	{
+		return LOCTEXT("Blutility_NewDerivedBlueprintTooltip", "Creates a blueprint based on the selected blueprint.");
+	}
+}
+
+bool FAssetTypeActions_EditorUtilityBlueprint::CanExecuteNewDerivedBlueprint(TWeakObjectPtr<class UEditorUtilityBlueprint> InObject)
+{
+	return !InObject->GeneratedClass->HasAnyClassFlags(CLASS_Deprecated);
+}
+
 void FAssetTypeActions_EditorUtilityBlueprint::ExecuteGlobalBlutility(TWeakObjectPtr<UEditorUtilityBlueprint> InObject)
 {
 	if (auto Object = InObject.Get())
@@ -214,4 +235,4 @@ void FAssetTypeActions_EditorUtilityBlueprint::ExecuteGlobalBlutility(TWeakObjec
 	}
 }
 
-#undef LOCTEXT_NAMESPACE
\ No newline at end of file
+#undef LOCTEXT_NAMESPACE
diff --git a/Engine/Source/Editor/Blutility/Private/AssetTypeActions_EditorUtilityBlueprint.h b/Engine/Source/Editor/Blutility/Private/AssetTypeActions_EditorUtilityBlueprint.h
--- a/Engine/Source/Editor/Blutility/Private/AssetTypeActions_EditorUtilityBlueprint.h
+++ b/Engine/Source/Editor/Blutility/Private/AssetTypeActions_EditorUtilityBlueprint.h
@@ -24,4 +24,10 @@ protected:
 	void ExecuteEditDefaults(FWeakBlueprintPointerArray Objects);
 	void ExecuteNewDerivedBlueprint(TWeakObjectPtr<class UEditorUtilityBlueprint> InObject);
 	void ExecuteGlobalBlutility(TWeakObjectPtr<class UEditorUtilityBlueprint> InObject);
+
+	/** Returns the tooltip to display when attempting to derive a Blueprint */
+	FText GetNewDerivedBlueprintTooltip(TWeakObjectPtr<class UEditorUtilityBlueprint> InObject);
+
+	/** Returns TRUE if you can derive a Blueprint */
+	bool CanExecuteNewDerivedBlueprint(TWeakObjectPtr<class UEditorUtilityBlueprint> InObject);
 };
\ No newline at end of file
diff --git a/Engine/Source/Editor/ClassViewer/Private/SClassViewer.cpp b/Engine/Source/Editor/ClassViewer/Private/SClassViewer.cpp
--- a/Engine/Source/Editor/ClassViewer/Private/SClassViewer.cpp
+++ b/Engine/Source/Editor/ClassViewer/Private/SClassViewer.cpp
@@ -443,6 +443,28 @@ namespace ClassViewer
 		static UBlueprint* GetBlueprint( UClass* InClass );
 		static void UpdateClassInNode(const FString& InGeneratedClassPackageName, UClass* InNewClass, UBlueprint* InNewBluePrint );
 
+		/** Util class to checks if a particular class can be made into a Blueprint, ignores deprecation
+		 *
+		 * @param InClass					The class to verify can be made into a Blueprint
+		 * @return							TRUE if the class can be made into a Blueprint
+		 */
+		bool CanCreateBlueprintOfClass_IgnoreDeprecation(UClass* InClass)
+		{
+			// Temporarily remove the deprecated flag so we can check if it is valid for
+			bool bIsClassDeprecated = InClass->HasAnyClassFlags(CLASS_Deprecated);
+			InClass->ClassFlags &= ~CLASS_Deprecated;
+
+			bool bCanCreateBlueprintOfClass = FKismetEditorUtilities::CanCreateBlueprintOfClass( InClass );
+
+			// Reassign the deprecated flag if it was previously assigned
+			if(bIsClassDeprecated)
+			{
+				InClass->ClassFlags |= CLASS_Deprecated;
+			}
+
+			return bCanCreateBlueprintOfClass;
+		}
+
 		/** Checks if a particular class is a brush.
 		 *	@param InClass				The Class to check.
 		 *	@return Returns true if the class is a brush.
@@ -786,7 +808,7 @@ namespace ClassViewer
 		{
 			if (UClass* Class = InClass.Get())
 			{
-				bInOutIsBlueprintBase = FKismetEditorUtilities::CanCreateBlueprintOfClass( Class );
+				bInOutIsBlueprintBase = CanCreateBlueprintOfClass_IgnoreDeprecation( Class );
 				bInOutHasBlueprint = Class->ClassGeneratedBy != NULL;
 			}
 			else
@@ -806,7 +828,7 @@ namespace ClassViewer
 			// If there is no class, it may be an unloaded blueprint.
 			if(UClass* Class = InNode->Class.Get())
 			{
-				return FKismetEditorUtilities::CanCreateBlueprintOfClass( Class );
+				return CanCreateBlueprintOfClass_IgnoreDeprecation(Class);
 			}
 			else if(InNode->bIsBPNormalType)
 			{
@@ -1026,7 +1048,7 @@ namespace ClassViewer
 		 */
 		static void OpenCreateBlueprintMenu(FMenuBuilder& MenuBuilder, UClass* InCreationClass)
 		{
-			if(InCreationClass == NULL || !FKismetEditorUtilities::CanCreateBlueprintOfClass(InCreationClass))
+			if(InCreationClass == NULL || !CanCreateBlueprintOfClass_IgnoreDeprecation(InCreationClass))
 			{
 				MenuBuilder.AddWidget( BlueprintNameEntry::MakeErrorWarningWidget(), FText::GetEmpty() );
 			}
@@ -1064,6 +1086,25 @@ namespace ClassViewer
 			}
 		}
 
+		/** Returns the tooltip to display when attempting to derive a Blueprint */
+		FText GetCreateBlueprintTooltip(UClass* InCreationClass)
+		{
+			if(InCreationClass->HasAnyClassFlags(CLASS_Deprecated))
+			{
+				return LOCTEXT("ClassViewerMenuCreateDeprecatedBlueprint_Tooltip", "Blueprint class is deprecated!");
+			}
+			else
+			{
+				return LOCTEXT("ClassViewerMenuCreateBlueprint_Tooltip", "Creates a Blueprint using this class as a base.");
+			}
+		}
+
+		/** Returns TRUE if you can derive a Blueprint */
+		bool CanOpenCreateBlueprintMenu(UClass* InCreationClass)
+		{
+			return !InCreationClass->HasAnyClassFlags(CLASS_Deprecated);
+		}
+
 		/**
 		 * Creates a class wizard for creating a new C++ class
 		 *
@@ -1432,7 +1473,21 @@ private:
 				{
 					MenuBuilder.BeginSection("ClassViewerIsBlueprint");
 					{
-						MenuBuilder.AddSubMenu( LOCTEXT("ClassViewerMenuCreateBlueprint", "Create Blueprint"), LOCTEXT("ClassViewerMenuCreateBlueprint_Tooltip", "Creates a Blueprint using this class as a base."), FNewMenuDelegate::CreateStatic( &ClassViewer::Helpers::OpenCreateBlueprintMenu, Class ) );
+						TAttribute<FText>::FGetter DynamicTooltipGetter;
+						DynamicTooltipGetter.BindStatic(&ClassViewer::Helpers::GetCreateBlueprintTooltip, Class);
+						TAttribute<FText> DynamicTooltipAttribute = TAttribute<FText>::Create(DynamicTooltipGetter);
+
+						MenuBuilder.AddSubMenu(
+							LOCTEXT("ClassViewerMenuCreateBlueprint", "Create Blueprint"), 
+							DynamicTooltipAttribute, 
+							FNewMenuDelegate::CreateStatic( &ClassViewer::Helpers::OpenCreateBlueprintMenu, Class ),
+							FUIAction(
+								FExecuteAction(),
+								FCanExecuteAction::CreateStatic( &ClassViewer::Helpers::CanOpenCreateBlueprintMenu, Class )
+								),
+							FName(),
+							EUserInterfaceActionType::Button
+							);
 					}
 					MenuBuilder.EndSection();
 				}
@@ -2278,7 +2333,21 @@ TSharedPtr< SWidget > SClassViewer::BuildMenuWidget()
 	{ 
 		if (bIsBlueprint)
 		{
-			MenuBuilder.AddSubMenu( LOCTEXT("ClassViewerMenuCreateBlueprint", "Create Blueprint"), LOCTEXT("ClassViewerMenuCreateBlueprint_Tooltip", "Creates a Blueprint using this class as a base."), FNewMenuDelegate::CreateStatic( &ClassViewer::Helpers::OpenCreateBlueprintMenu, RightClickClass ) );
+			TAttribute<FText>::FGetter DynamicTooltipGetter;
+			DynamicTooltipGetter.BindStatic(&ClassViewer::Helpers::GetCreateBlueprintTooltip, RightClickClass);
+			TAttribute<FText> DynamicTooltipAttribute = TAttribute<FText>::Create(DynamicTooltipGetter);
+
+			MenuBuilder.AddSubMenu(
+				LOCTEXT("ClassViewerMenuCreateBlueprint", "Create Blueprint"), 
+				DynamicTooltipAttribute, 
+				FNewMenuDelegate::CreateStatic( &ClassViewer::Helpers::OpenCreateBlueprintMenu, RightClickClass ),
+				FUIAction(
+					FExecuteAction(),
+					FCanExecuteAction::CreateStatic( &ClassViewer::Helpers::CanOpenCreateBlueprintMenu, RightClickClass )
+					),
+				FName(),
+				EUserInterfaceActionType::Button
+				);
 		}
 	
 		if(bHasBlueprint)
diff --git a/Engine/Source/Editor/Kismet/Private/BlueprintDetailsCustomization.cpp b/Engine/Source/Editor/Kismet/Private/BlueprintDetailsCustomization.cpp
--- a/Engine/Source/Editor/Kismet/Private/BlueprintDetailsCustomization.cpp
+++ b/Engine/Source/Editor/Kismet/Private/BlueprintDetailsCustomization.cpp
@@ -1619,8 +1619,9 @@ void FBlueprintGraphArgumentLayout::GenerateHeaderRowContent( FDetailWidgetRow&
 		.FillWidth(1)
 		.VAlign(VAlign_Center)
 		[
-			SNew(SEditableTextBox)
+			SAssignNew(ArgumentNameWidget, SEditableTextBox)
 				.Text( this, &FBlueprintGraphArgumentLayout::OnGetArgNameText )
+				.OnTextChanged(this, &FBlueprintGraphArgumentLayout::OnArgNameChange)
 				.OnTextCommitted( this, &FBlueprintGraphArgumentLayout::OnArgNameTextCommitted )
 				.Font( IDetailLayoutBuilder::GetDetailFont() )
 				.IsEnabled(!ShouldPinBeReadOnly())
@@ -1801,6 +1802,33 @@ FText FBlueprintGraphArgumentLayout::OnGetArgNameText() const
 	return FText();
 }
 
+void FBlueprintGraphArgumentLayout::OnArgNameChange(const FText& InNewText)
+{
+	bool bVerified = true;
+
+	FText ErrorMessage;
+
+	if (InNewText.IsEmpty())
+	{
+		ErrorMessage = LOCTEXT("EmptyArgument", "Name cannot be empty!");
+		bVerified = false;
+	}
+	else
+	{
+		const FString& OldName = ParamItemPtr.Pin()->PinName;
+		bVerified = GraphActionDetailsPtr.Pin()->OnVerifyPinRename(TargetNode, OldName, InNewText.ToString(), ErrorMessage);
+	}
+
+	if(!bVerified)
+	{
+		ArgumentNameWidget.Pin()->SetError(ErrorMessage);
+	}
+	else
+	{
+		ArgumentNameWidget.Pin()->SetError(FText::GetEmpty());
+	}
+}
+
 void FBlueprintGraphArgumentLayout::OnArgNameTextCommitted(const FText& NewText, ETextCommit::Type InTextCommit)
 {
 	if (!NewText.IsEmpty() && TargetNode && ParamItemPtr.IsValid() && GraphActionDetailsPtr.IsValid() && !ShouldPinBeReadOnly())
@@ -2834,8 +2862,37 @@ struct FPinRenamedHelper : public FBasePinChangeHelper
 	}
 };
 
+bool FBaseBlueprintGraphActionDetails::OnVerifyPinRename(UK2Node_EditablePinBase* InTargetNode, const FString& InOldName, const FString& InNewName, FText& OutErrorMessage)
+{
+	// If the name is unchanged, allow the name
+	if(InOldName == InNewName)
+	{
+		return true;
+	}
+
+	if (InTargetNode)
+	{
+		// Check if the name conflicts with any of the other internal UFunction's property names (local variables and parameters).
+		const auto FoundFunction = FFunctionFromNodeHelper::FunctionFromNode(InTargetNode);
+		const auto ExistingProperty = FindField<const UProperty>(FoundFunction, *InNewName);
+		if (ExistingProperty)
+		{
+			OutErrorMessage = LOCTEXT("ConflictsWithProperty", "Conflicts with another another local variable or function parameter!");
+			return false;
+		}
+	}
+	return true;
+}
+
 bool FBaseBlueprintGraphActionDetails::OnPinRenamed(UK2Node_EditablePinBase* TargetNode, const FString& OldName, const FString& NewName)
 {
+	// Before changing the name, verify the name
+	FText ErrorMessage;
+	if(!OnVerifyPinRename(TargetNode, OldName, NewName, ErrorMessage))
+	{
+		return false;
+	}
+
 	UEdGraph* Graph = GetGraph();
 
 	if (TargetNode)
@@ -3350,9 +3407,15 @@ bool FBaseBlueprintGraphActionDetails::IsPinNameUnique(const FString& TestName)
 
 void FBaseBlueprintGraphActionDetails::GenerateUniqueParameterName( const FString &BaseName, FString &Result ) const
 {
+	UK2Node_EditablePinBase * FunctionEntryNode = FunctionEntryNodePtr.Get();
+	check(FunctionEntryNode);
+
+	const auto FoundFunction = FFunctionFromNodeHelper::FunctionFromNode(FunctionEntryNode);
+
 	Result = BaseName;
 	int UniqueNum = 0;
-	while(!IsPinNameUnique(Result))
+	// Prevent the unique name from being the same as another of the UFunction's properties
+	while(!IsPinNameUnique(Result) || FindField<const UProperty>(FoundFunction, *Result) != NULL)
 	{
 		Result = FString::Printf(TEXT("%s%d"), *BaseName, ++UniqueNum);
 	}
@@ -3743,6 +3806,38 @@ void FBlueprintGlobalOptionsDetails::OnClassPicked(UClass* PickedClass)
 	}
 }
 
+bool FBlueprintGlobalOptionsDetails::CanDeprecateBlueprint() const
+{
+	// If the parent is deprecated, we cannot modify deprecation on this Blueprint
+	if(GetBlueprintObj()->ParentClass->HasAnyClassFlags(CLASS_Deprecated))
+	{
+		return false;
+	}
+
+	return true;
+}
+
+void FBlueprintGlobalOptionsDetails::OnDeprecateBlueprint(ESlateCheckBoxState::Type InCheckState)
+{
+	GetBlueprintObj()->bDeprecate = InCheckState == ESlateCheckBoxState::Checked? true : false;
+	FBlueprintEditorUtils::MarkBlueprintAsStructurallyModified(GetBlueprintObj());
+}
+
+ESlateCheckBoxState::Type FBlueprintGlobalOptionsDetails::IsDeprecatedBlueprint() const
+{
+	return GetBlueprintObj()->bDeprecate? ESlateCheckBoxState::Checked : ESlateCheckBoxState::Unchecked;
+}
+
+FText FBlueprintGlobalOptionsDetails::GetDeprecatedTooltip() const
+{
+	if(CanDeprecateBlueprint())
+	{
+		return LOCTEXT("DeprecateBlueprintTooltip", "Deprecate the Blueprint and all child Blueprints to make it no longer placeable in the World nor child classes created from it.");
+	}
+	
+	return LOCTEXT("DisabledDeprecateBlueprintTooltip", "This Blueprint is deprecated because of a parent, it is not possible to remove deprecation from it!");
+}
+
 void FBlueprintGlobalOptionsDetails::CustomizeDetails(IDetailLayoutBuilder& DetailLayout)
 {
 	const UBlueprint* Blueprint = GetBlueprintObj();
@@ -3798,12 +3893,36 @@ void FBlueprintGlobalOptionsDetails::CustomizeDetails(IDetailLayoutBuilder& Deta
 			InterfacesCategory.AddCustomBuilder(InheritedInterfaceLayout);
 		}
 
+		// Hide the bDeprecate, we override the functionality.
+		static FName DeprecatePropName(TEXT("bDeprecate"));
+		DetailLayout.HideProperty(DetailLayout.GetProperty(DeprecatePropName));
+
 		// Hide 'run on drag' for LevelBP
 		if (bIsLevelScriptBP)
 		{
 			static FName RunOnDragPropName(TEXT("bRunConstructionScriptOnDrag"));
 			DetailLayout.HideProperty(DetailLayout.GetProperty(RunOnDragPropName));
 		}
+		else
+		{
+			// Only display the ability to deprecate a Blueprint on non-level Blueprints.
+			Category.AddCustomRow( TEXT("Deprecate") )
+				.NameContent()
+				[
+					SNew(STextBlock)
+					.Text( LOCTEXT("DeprecateLabel", "Deprecate") )
+					.ToolTipText( this, &FBlueprintGlobalOptionsDetails::GetDeprecatedTooltip )
+					.Font(IDetailLayoutBuilder::GetDetailFont())
+				]
+				.ValueContent()
+				[
+					SNew(SCheckBox)
+						.IsEnabled( this, &FBlueprintGlobalOptionsDetails::CanDeprecateBlueprint )
+						.IsChecked( this, &FBlueprintGlobalOptionsDetails::IsDeprecatedBlueprint )
+						.OnCheckStateChanged( this, &FBlueprintGlobalOptionsDetails::OnDeprecateBlueprint )
+						.ToolTipText( this, &FBlueprintGlobalOptionsDetails::GetDeprecatedTooltip )
+				];
+		}
 	}
 }
 
diff --git a/Engine/Source/Editor/Kismet/Private/BlueprintDetailsCustomization.h b/Engine/Source/Editor/Kismet/Private/BlueprintDetailsCustomization.h
--- a/Engine/Source/Editor/Kismet/Private/BlueprintDetailsCustomization.h
+++ b/Engine/Source/Editor/Kismet/Private/BlueprintDetailsCustomization.h
@@ -165,6 +165,9 @@ public:
 	/** Refreshes the graph and ensures the target node is up to date */
 	void OnParamsChanged(UK2Node_EditablePinBase* TargetNode, bool bForceRefresh = false);
 
+	/** Checks if the pin rename can occur */
+	bool OnVerifyPinRename(UK2Node_EditablePinBase* InTargetNode, const FString& InOldName, const FString& InNewName, FText& OutErrorMessage);
+
 	/** Refreshes the graph and ensures the target node is up to date */
 	bool OnPinRenamed(UK2Node_EditablePinBase* TargetNode, const FString& OldName, const FString& NewName);
 
@@ -320,6 +323,7 @@ private:
 	FReply OnArgMoveDown();
 
 	FText OnGetArgNameText() const;
+	void OnArgNameChange(const FText& InNewText);
 	void OnArgNameTextCommitted(const FText& NewText, ETextCommit::Type InTextCommit);
 	
 	FEdGraphPinType OnGetPinInfo() const;
@@ -350,6 +354,9 @@ private:
 
 	/** The name of this argument for remembering expansion state */
 	FName ArgumentName;
+
+	/** Holds a weak pointer to the argument name widget, used for error notifications */
+	TWeakPtr<SEditableTextBox> ArgumentNameWidget;
 };
 
 /** Details customization for functions and graphs selected in the MyBlueprint panel */
@@ -535,6 +542,18 @@ protected:
 	/** Delegate called when a class is selected from the class picker */
 	void OnClassPicked(UClass* SelectedClass);
 
+	/** Returns TRUE if the Blueprint can be deprecated */
+	bool CanDeprecateBlueprint() const;
+
+	/** Callback when toggling the Deprecate checkbox, handles marking a Blueprint as deprecated */
+	void OnDeprecateBlueprint(ESlateCheckBoxState::Type InCheckState);
+
+	/** Callback for Deprecate checkbox, returns checked if the Blueprint is deprecated */
+	ESlateCheckBoxState::Type IsDeprecatedBlueprint() const;
+
+	/** Returns the tooltip explaining deprecation */
+	FText GetDeprecatedTooltip() const;
+
 private:
 	/** Weak reference to the Blueprint editor */
 	TWeakPtr<FBlueprintEditor> BlueprintEditorPtr;
diff --git a/Engine/Source/Editor/Kismet/Private/SCSEditorViewportClient.cpp b/Engine/Source/Editor/Kismet/Private/SCSEditorViewportClient.cpp
--- a/Engine/Source/Editor/Kismet/Private/SCSEditorViewportClient.cpp
+++ b/Engine/Source/Editor/Kismet/Private/SCSEditorViewportClient.cpp
@@ -918,8 +918,19 @@ void FSCSEditorViewportClient::UpdatePreviewActorForBlueprint(UBlueprint* InBlue
 			SpawnInfo.bNoCollisionFail = true;
 			SpawnInfo.bNoFail = true;
 			SpawnInfo.ObjectFlags = RF_Transient;
+
+			// Temporarily remove the deprecated flag so we can respawn the Blueprint in the viewport
+			bool bIsClassDeprecated = PreviewBlueprint->GeneratedClass->HasAnyClassFlags(CLASS_Deprecated);
+			PreviewBlueprint->GeneratedClass->ClassFlags &= ~CLASS_Deprecated;
+
 			PreviewActorPtr = PreviewActor = PreviewScene->GetWorld()->SpawnActor( PreviewBlueprint->GeneratedClass, &SpawnLocation, &SpawnRotation, SpawnInfo );
 
+			// Reassign the deprecated flag if it was previously assigned
+			if(bIsClassDeprecated)
+			{
+				PreviewBlueprint->GeneratedClass->ClassFlags |= CLASS_Deprecated;
+			}
+
 			check(PreviewActor);
 
 			// Ensure that the actor is visible
diff --git a/Engine/Source/Editor/KismetCompiler/Private/KismetCompiler.cpp b/Engine/Source/Editor/KismetCompiler/Private/KismetCompiler.cpp
--- a/Engine/Source/Editor/KismetCompiler/Private/KismetCompiler.cpp
+++ b/Engine/Source/Editor/KismetCompiler/Private/KismetCompiler.cpp
@@ -1604,6 +1604,15 @@ void FKismetCompilerContext::FinishCompilingClass(UClass* Class)
 		Class->ClassWithin = ParentClass->ClassWithin ? ParentClass->ClassWithin : UObject::StaticClass();
 		Class->ClassConfigName = ParentClass->ClassConfigName;
 
+		// If the Blueprint was marked as deprecated, then flag the class as deprecated.
+		if(Blueprint->bDeprecate)
+		{
+			Class->ClassFlags |= CLASS_Deprecated;
+		}
+
+		// If the flag is inherited, this will keep the bool up-to-date
+		Blueprint->bDeprecate = (Class->ClassFlags & CLASS_Deprecated) == CLASS_Deprecated;
+
 		// Copy the category info from the parent class
 #if WITH_EDITORONLY_DATA
 		ParentClass->GetHideCategories(AllHideCategories);
diff --git a/Engine/Source/Editor/UnrealEd/Private/Factories/EditorFactories.cpp b/Engine/Source/Editor/UnrealEd/Private/Factories/EditorFactories.cpp
--- a/Engine/Source/Editor/UnrealEd/Private/Factories/EditorFactories.cpp
+++ b/Engine/Source/Editor/UnrealEd/Private/Factories/EditorFactories.cpp
@@ -6034,12 +6034,12 @@ public:
 
 	virtual bool IsClassAllowed(const FClassViewerInitializationOptions& InInitOptions, const UClass* InClass, TSharedRef< FClassViewerFilterFuncs > InFilterFuncs ) override
 	{
-		return InFilterFuncs->IfInChildOfClassesSet(DisallowedChildrenOfClasses, InClass) != EFilterReturn::Passed;
+		return InFilterFuncs->IfInChildOfClassesSet(DisallowedChildrenOfClasses, InClass) != EFilterReturn::Passed && !InClass->HasAnyClassFlags(CLASS_Deprecated);
 	}
 
 	virtual bool IsUnloadedClassAllowed(const FClassViewerInitializationOptions& InInitOptions, const TSharedRef< const IUnloadedBlueprintData > InUnloadedClassData, TSharedRef< FClassViewerFilterFuncs > InFilterFuncs) override
 	{
-		return InFilterFuncs->IfInChildOfClassesSet(DisallowedChildrenOfClasses, InUnloadedClassData) != EFilterReturn::Passed;
+		return InFilterFuncs->IfInChildOfClassesSet(DisallowedChildrenOfClasses, InUnloadedClassData) != EFilterReturn::Passed && !InUnloadedClassData->HasAnyClassFlags(CLASS_Deprecated);
 	}
 };
 
diff --git a/Engine/Source/Editor/UnrealEd/Private/Kismet2/BlueprintEditorUtils.cpp b/Engine/Source/Editor/UnrealEd/Private/Kismet2/BlueprintEditorUtils.cpp
--- a/Engine/Source/Editor/UnrealEd/Private/Kismet2/BlueprintEditorUtils.cpp
+++ b/Engine/Source/Editor/UnrealEd/Private/Kismet2/BlueprintEditorUtils.cpp
@@ -5304,7 +5304,8 @@ public:
 		//		AND it is NOT on the disallowed classes list
 		return InFilterFuncs->IfInChildOfClassesSet( AllowedChildrenOfClasses, InClass) != EFilterReturn::Failed && 
 			InFilterFuncs->IfInChildOfClassesSet(DisallowedChildrenOfClasses, InClass) != EFilterReturn::Passed && 
-			InFilterFuncs->IfInClassesSet(DisallowedClasses, InClass) != EFilterReturn::Passed;
+			InFilterFuncs->IfInClassesSet(DisallowedClasses, InClass) != EFilterReturn::Passed &&
+			!InClass->HasAnyClassFlags(CLASS_Deprecated);
 	}
 
 	virtual bool IsUnloadedClassAllowed(const FClassViewerInitializationOptions& InInitOptions, const TSharedRef< const IUnloadedBlueprintData > InUnloadedClassData, TSharedRef< FClassViewerFilterFuncs > InFilterFuncs) override
@@ -5314,7 +5315,8 @@ public:
 		//		AND it is NOT on the disallowed classes list
 		return InFilterFuncs->IfInChildOfClassesSet( AllowedChildrenOfClasses, InUnloadedClassData) != EFilterReturn::Failed && 
 			InFilterFuncs->IfInChildOfClassesSet(DisallowedChildrenOfClasses, InUnloadedClassData) != EFilterReturn::Passed && 
-			InFilterFuncs->IfInClassesSet(DisallowedClasses, InUnloadedClassData) != EFilterReturn::Passed;
+			InFilterFuncs->IfInClassesSet(DisallowedClasses, InUnloadedClassData) != EFilterReturn::Passed &&
+			!InUnloadedClassData->HasAnyClassFlags(CLASS_Deprecated);
 	}
 };
 
diff --git a/Engine/Source/Editor/UnrealEd/Private/Kismet2/KismetReinstanceUtilities.cpp b/Engine/Source/Editor/UnrealEd/Private/Kismet2/KismetReinstanceUtilities.cpp
--- a/Engine/Source/Editor/UnrealEd/Private/Kismet2/KismetReinstanceUtilities.cpp
+++ b/Engine/Source/Editor/UnrealEd/Private/Kismet2/KismetReinstanceUtilities.cpp
@@ -314,8 +314,19 @@ void FBlueprintCompileReinstancer::ReplaceInstancesOfClass(UClass* OldClass, UCl
 				SpawnInfo.Template = NewArchetype;
 				SpawnInfo.bNoCollisionFail = true;
 				SpawnInfo.bDeferConstruction = true;
+
+				// Temporarily remove the deprecated flag so we can respawn the Blueprint in the level
+				bool bIsClassDeprecated = SpawnClass->HasAnyClassFlags(CLASS_Deprecated);
+				SpawnClass->ClassFlags &= ~CLASS_Deprecated;
+
 				AActor* NewActor = World->SpawnActor(SpawnClass, LocationPtr, RotationPtr, SpawnInfo );
 
+				// Reassign the deprecated flag if it was previously assigned
+				if(bIsClassDeprecated)
+				{
+					SpawnClass->ClassFlags |= CLASS_Deprecated;
+				}
+
 				check(NewActor);
 				NewObject = NewActor;
 
@@ -522,4 +533,4 @@ void FBlueprintCompileReinstancer::ReparentChild(UClass* ChildClass)
 	ChildClass->SetSuperStruct(DuplicatedClass);
 	ChildClass->Bind();
 	ChildClass->StaticLink(true);
-}
\ No newline at end of file
+}
diff --git a/Engine/Source/Runtime/Engine/Classes/Engine/Blueprint.h b/Engine/Source/Runtime/Engine/Classes/Engine/Blueprint.h
--- a/Engine/Source/Runtime/Engine/Classes/Engine/Blueprint.h
+++ b/Engine/Source/Runtime/Engine/Classes/Engine/Blueprint.h
@@ -311,6 +311,9 @@ class ENGINE_API UBlueprint : public UBlueprintCore
 	UPROPERTY()
 	FGuid  SearchGuid;
 
+	/** Deprecates the Blueprint, marking the generated class with the CLASS_Deprecated flag */
+	UPROPERTY(EditAnywhere, Category=BlueprintOption)
+	bool bDeprecate;
 #endif //WITH_EDITORONLY_DATA
 
 	/** 'Simple' construction script - graph of components to instance */
diff --git a/Engine/Source/Runtime/Slate/Private/Framework/MultiBox/MultiBoxBuilder.cpp b/Engine/Source/Runtime/Slate/Private/Framework/MultiBox/MultiBoxBuilder.cpp
--- a/Engine/Source/Runtime/Slate/Private/Framework/MultiBox/MultiBoxBuilder.cpp
+++ b/Engine/Source/Runtime/Slate/Private/Framework/MultiBox/MultiBoxBuilder.cpp
@@ -178,7 +178,7 @@ void FMenuBuilder::AddMenuSeparator(FName InExtensionHook)
 	ApplyHook(InExtensionHook, EExtensionHook::After);
 }
 
-void FMenuBuilder::AddSubMenu( const FText& InMenuLabel, const FText& InToolTip, const FNewMenuDelegate& InSubMenu, const FUIAction& InUIAction, FName InExtensionHook, const EUserInterfaceActionType::Type InUserInterfaceActionType, const bool bInOpenSubMenuOnClick, const FSlateIcon& InIcon )
+void FMenuBuilder::AddSubMenu( const TAttribute<FText>& InMenuLabel, const TAttribute<FText>& InToolTip, const FNewMenuDelegate& InSubMenu, const FUIAction& InUIAction, FName InExtensionHook, const EUserInterfaceActionType::Type InUserInterfaceActionType, const bool bInOpenSubMenuOnClick, const FSlateIcon& InIcon )
 {
 	ApplySectionBeginning();
 
@@ -188,7 +188,7 @@ void FMenuBuilder::AddSubMenu( const FText& InMenuLabel, const FText& InToolTip,
 	MultiBox->AddMultiBlock( NewMenuEntryBlock );
 }
 
-void FMenuBuilder::AddSubMenu( const FText& InMenuLabel, const FText& InToolTip, const FNewMenuDelegate& InSubMenu, const bool bInOpenSubMenuOnClick /*= false*/, const FSlateIcon& InIcon /*= FSlateIcon()*/ )
+void FMenuBuilder::AddSubMenu( const TAttribute<FText>& InMenuLabel, const TAttribute<FText>& InToolTip, const FNewMenuDelegate& InSubMenu, const bool bInOpenSubMenuOnClick /*= false*/, const FSlateIcon& InIcon /*= FSlateIcon()*/ )
 {
 	ApplySectionBeginning();
 
diff --git a/Engine/Source/Runtime/Slate/Public/Framework/MultiBox/MultiBoxBuilder.h b/Engine/Source/Runtime/Slate/Public/Framework/MultiBox/MultiBoxBuilder.h
--- a/Engine/Source/Runtime/Slate/Public/Framework/MultiBox/MultiBoxBuilder.h
+++ b/Engine/Source/Runtime/Slate/Public/Framework/MultiBox/MultiBoxBuilder.h
@@ -214,9 +214,9 @@ public:
 	 * @param	InSubMenu			Sub-Menu object which creates menu entries for the sub-menu
 	 * @param bInOpenSubMenuOnClick Sub-menu will open only if the sub-menu entry is clicked
 	 */
-	void AddSubMenu( const FText& InMenuLabel, const FText& InToolTip, const FNewMenuDelegate& InSubMenu, const FUIAction& InUIAction, FName InExtensionHook, const EUserInterfaceActionType::Type InUserInterfaceActionType, const bool bInOpenSubMenuOnClick = false, const FSlateIcon& InIcon = FSlateIcon() );
+	void AddSubMenu( const TAttribute<FText>& InMenuLabel, const TAttribute<FText>& InToolTip, const FNewMenuDelegate& InSubMenu, const FUIAction& InUIAction, FName InExtensionHook, const EUserInterfaceActionType::Type InUserInterfaceActionType, const bool bInOpenSubMenuOnClick = false, const FSlateIcon& InIcon = FSlateIcon() );
 
-	void AddSubMenu( const FText& InMenuLabel, const FText& InToolTip, const FNewMenuDelegate& InSubMenu, const bool bInOpenSubMenuOnClick = false, const FSlateIcon& InIcon = FSlateIcon() );
+	void AddSubMenu( const TAttribute<FText>& InMenuLabel, const TAttribute<FText>& InToolTip, const FNewMenuDelegate& InSubMenu, const bool bInOpenSubMenuOnClick = false, const FSlateIcon& InIcon = FSlateIcon() );
 
 	void AddSubMenu( const TSharedRef< SWidget > Contents, const FNewMenuDelegate& InSubMenu, const bool bInOpenSubMenuOnClick = false );
 
