Analytics - Rename a blueprint event so that it's distinct from the default editor usage event.

TTP# 340198

diff --git a/Engine/Source/Editor/Kismet/Private/BlueprintEditor.cpp b/Engine/Source/Editor/Kismet/Private/BlueprintEditor.cpp
--- a/Engine/Source/Editor/Kismet/Private/BlueprintEditor.cpp
+++ b/Engine/Source/Editor/Kismet/Private/BlueprintEditor.cpp
@@ -1521,7 +1521,7 @@ FBlueprintEditor::~FBlueprintEditor()
 		BPEditorAttribs.Add( FAnalyticsEventAttribute( FString( "PastedNodesCreated" ), AnalyticsStats.NodePasteCreateCount ));
 
 		BPEditorAttribs.Add( FAnalyticsEventAttribute( FString( "ProjectId" ), ProjectID ) );
-		FEngineAnalytics::GetProvider().RecordEvent( FString( "Editor.Usage.BlueprintEditor" ), BPEditorAttribs );
+		FEngineAnalytics::GetProvider().RecordEvent( FString( "Editor.Usage.BlueprintEditorSummary" ), BPEditorAttribs );
 
 		for (auto Iter = AnalyticsStats.GraphDisallowedPinConnections.CreateConstIterator(); Iter; ++Iter)
 		{
