#UE4 Merge panel more correctly integrated into blueprint editor, ownership cleaned up, cancel action restored...

diff --git a/Engine/Source/Developer/Merge/Private/Merge.cpp b/Engine/Source/Developer/Merge/Private/Merge.cpp
--- a/Engine/Source/Developer/Merge/Private/Merge.cpp
+++ b/Engine/Source/Developer/Merge/Private/Merge.cpp
@@ -1,10 +1,12 @@
 // Copyright 1998-2014 Epic Games, Inc. All Rights Reserved.
 
 #include "MergePrivatePCH.h"
+
+#include "BlueprintEditor.h"
+#include "ISourceControlModule.h"
 #include "SBlueprintMerge.h"
 #include "Toolkits/ToolkitManager.h"
 #include "Toolkits/IToolkit.h"
-#include "ISourceControlModule.h"
 
 #define LOCTEXT_NAMESPACE "Merge"
 
@@ -109,7 +111,8 @@ class FMerge : public IMerge
 	/** IModuleInterface implementation */
 	virtual void StartupModule() override;
 	virtual void ShutdownModule() override;
-	virtual TSharedPtr<SWidget> GenerateMergeWidget(const UBlueprint& Object, FBlueprintEditor& Editor) override;
+	virtual TSharedPtr<SWidget> GenerateMergeWidget(const UBlueprint& Object, TSharedRef<FBlueprintEditor> Editor) override;
+	virtual bool PendingMerge(const UBlueprint& BlueprintObj) const override;
 };
 IMPLEMENT_MODULE( FMerge, Merge )
 
@@ -124,7 +127,7 @@ void FMerge::ShutdownModule()
 	// we call this function before unloading the module.
 }
 
-TSharedPtr<SWidget> FMerge::GenerateMergeWidget(const UBlueprint& Object, FBlueprintEditor& Editor)
+TSharedPtr<SWidget> FMerge::GenerateMergeWidget(const UBlueprint& Object, TSharedRef<FBlueprintEditor> Editor)
 {
 	// merge the local asset with the depot, SCC provides us with the last common revision as
 	// a basis for the merge:
@@ -150,6 +153,7 @@ TSharedPtr<SWidget> FMerge::GenerateMergeWidget(const UBlueprint& Object, FBluep
 
 		return SNew(SBlueprintMerge)
 			.BlueprintLocal(const_cast<UBlueprint*>(&Object))
+			.OwningEditor(Editor)
 			.BaseArgs(BaseArgs);
 	}
 	else
@@ -176,4 +180,19 @@ TSharedPtr<SWidget> FMerge::GenerateMergeWidget(const UBlueprint& Object, FBluep
 	}
 }
 
+bool FMerge::PendingMerge(const UBlueprint& BlueprintObj) const
+{
+	bool bIsMergeEnabled = false;
+	GConfig->GetBool(TEXT("AssetMerge"), TEXT("EnableAssetMerge"), bIsMergeEnabled, GEngineIni);
+	ISourceControlProvider& SourceControlProvider = ISourceControlModule::Get().GetProvider();
+
+	bool bPendingMerge = false;
+	if( bIsMergeEnabled && SourceControlProvider.IsEnabled() )
+	{
+		FSourceControlStatePtr SourceControlState = SourceControlProvider.GetState(BlueprintObj.GetOutermost(), EStateCacheUsage::Use);
+		bPendingMerge = SourceControlState.IsValid() && SourceControlState->IsConflicted();
+	}
+	return bPendingMerge;
+}
+
 #undef LOCTEXT_NAMESPACE
diff --git a/Engine/Source/Developer/Merge/Private/SBlueprintMerge.cpp b/Engine/Source/Developer/Merge/Private/SBlueprintMerge.cpp
--- a/Engine/Source/Developer/Merge/Private/SBlueprintMerge.cpp
+++ b/Engine/Source/Developer/Merge/Private/SBlueprintMerge.cpp
@@ -532,9 +532,10 @@ FReply SBlueprintMerge::OnAcceptResultClicked()
 	// cr doc: this is too simple to be correct..
 	PanelLocal.Blueprint->RemoveGeneratedClasses();
 
-	if (OwningEditor.IsValid())
+	auto Editor = OwningEditor.Pin();
+	if (Editor.IsValid())
 	{
-		OwningEditor->CloseWindow();
+		Editor->CloseWindow();
 	}
 
 	const auto Result = FEditorFileUtils::PromptForCheckoutAndSave(PackagesToSave, /*bCheckDirty=*/ false, /*bPromptToSave=*/ false);
@@ -552,9 +553,10 @@ FReply SBlueprintMerge::OnAcceptResultClicked()
 
 FReply SBlueprintMerge::OnCancelClicked()
 {
-	if (OwningEditor.IsValid())
+	auto Editor = OwningEditor.Pin();
+	if (Editor.IsValid())
 	{
-		OwningEditor->CloseWindow();
+		Editor->CloseMergeTool();
 	}
 	return FReply::Handled();
 }
@@ -580,7 +582,13 @@ FReply SBlueprintMerge::OnTakeBaseClicked()
 void SBlueprintMerge::StageBlueprint(UBlueprint const* DesiredBP)
 {
 	// Take the DesiredBP and applies its changes to the result blueprint:
-	UBlueprint* BlueprintResult = OwningEditor->GetBlueprintObj();
+	auto Editor = OwningEditor.Pin();
+	if( !Editor.IsValid() )
+	{
+		return;
+	}
+
+	UBlueprint* BlueprintResult = Editor->GetBlueprintObj();
 
 	// I want everything in the package... components, CDO, etc:
 
diff --git a/Engine/Source/Developer/Merge/Private/SBlueprintMerge.h b/Engine/Source/Developer/Merge/Private/SBlueprintMerge.h
--- a/Engine/Source/Developer/Merge/Private/SBlueprintMerge.h
+++ b/Engine/Source/Developer/Merge/Private/SBlueprintMerge.h
@@ -13,7 +13,7 @@ public:
 	SLATE_BEGIN_ARGS(SBlueprintMerge){}
 	SLATE_ARGUMENT(class UBlueprint*, BlueprintLocal)
 	SLATE_ARGUMENT( SBlueprintDiff::FArguments, BaseArgs )
-	SLATE_ARGUMENT( TSharedPtr<FBlueprintEditor>, OwningEditor )
+	SLATE_ARGUMENT( TWeakPtr<FBlueprintEditor>, OwningEditor )
 	SLATE_END_ARGS()
 	
 	void Construct(const FArguments& InArgs);
@@ -40,7 +40,7 @@ protected:
 	FReply OnTakeBaseClicked();
 	void StageBlueprint(UBlueprint const* DesiredBP);
 
-	TSharedPtr<FBlueprintEditor> OwningEditor;
+	TWeakPtr<FBlueprintEditor> OwningEditor;
 	FDiffPanel PanelLocal;
 	TWeakPtr<SWindow>		OwningWindow;
 
diff --git a/Engine/Source/Developer/Merge/Public/Merge.h b/Engine/Source/Developer/Merge/Public/Merge.h
--- a/Engine/Source/Developer/Merge/Public/Merge.h
+++ b/Engine/Source/Developer/Merge/Public/Merge.h
@@ -48,6 +48,12 @@ public:
 	 *
 	 * @return The merge view widget
 	 */
-	virtual TSharedPtr<class SWidget> GenerateMergeWidget( const UBlueprint& Object, FBlueprintEditor& Editor ) = 0;
+	virtual TSharedPtr<class SWidget> GenerateMergeWidget( const UBlueprint& Object, TSharedRef<FBlueprintEditor> Editor ) = 0;
+
+	/** 
+	 * @return whether the blueprint is in a conflicted state, and can therefore be merged.
+	 */
+	virtual bool PendingMerge(const UBlueprint& BlueprintObj) const = 0;
+
 };
 
diff --git a/Engine/Source/Editor/Kismet/Private/BlueprintEditor.cpp b/Engine/Source/Editor/Kismet/Private/BlueprintEditor.cpp
--- a/Engine/Source/Editor/Kismet/Private/BlueprintEditor.cpp
+++ b/Engine/Source/Editor/Kismet/Private/BlueprintEditor.cpp
@@ -1592,7 +1592,6 @@ void FBlueprintEditor::CreateDefaultTabContents(const TArray<UBlueprint*>& InBlu
 		];
 
 	FindResults = SNew(SFindInBlueprints, SharedThis(this));
-	MergeTool = SNew(SHorizontalBox);
 	
 	this->Inspector = 
 		SNew(SKismetInspector)
@@ -4653,6 +4652,44 @@ void FBlueprintEditor::UnregisterSCSEditorCustomization(const FName& InComponent
 	SCSEditorCustomizations.Remove(InComponentName);
 }
 
+TSharedRef<SWidget> FBlueprintEditor::ShowMergeTool()
+{
+	/** 
+	 * This function has to handle three case:
+	 *	1. The merge tool is already created, return it
+	 *	2. The merge tool is being restored, if we have a pending merge display the merge tool, otherwise display an empty widget
+	 *  3. The merge tool is being invoked, display the merge tool 
+	 */
+	auto MergeToolPtr = MergeTool.Pin();
+	if( !MergeToolPtr.IsValid() )
+	{
+		if( IMerge::Get().PendingMerge(*GetBlueprintObj()) )
+		{
+			MergeToolPtr = IMerge::Get().GenerateMergeWidget(*GetBlueprintObj(), SharedThis(this) ).ToSharedRef();
+		}
+		else
+		{
+			// We *have* to return a valid reference to a widget, so return an empty box. This is arguabla a deficiency in
+			// our controlling widgets that we can't 'restore' with nothing.
+			MergeToolPtr = SNew(SHorizontalBox);
+		}
+
+		// Keep a weak reference to the widget we've created, widget will be owned by the docktab:
+		MergeTool = MergeToolPtr;
+	}
+
+	return MergeToolPtr.ToSharedRef();
+}
+
+void FBlueprintEditor::CloseMergeTool()
+{
+	auto MergeToolDockTabPtr = MergeToolDockTab.Pin();
+	if( MergeToolDockTabPtr.IsValid() )
+	{
+		MergeToolDockTabPtr->RequestCloseTab();
+	}
+}
+
 TArray<FSCSEditorTreeNodePtrType>  FBlueprintEditor::GetSelectedSCSEditorTreeNodes() const
 {
 	TArray<FSCSEditorTreeNodePtrType>  Nodes;
@@ -5451,12 +5488,7 @@ void FBlueprintEditor::OnRepairCorruptedBlueprint()
 
 void FBlueprintEditor::OnBeginBlueprintMerge()
 {
-	// @cr doc: any way to show this in all of the subviews (defaults + components + graph view, rather than just the active one?)
-	//if( MergeTool->Slots().Num() == 0 && GetBlueprintObj() )
-	{
-		MergeTool->AddSlot()[ IMerge::Get().GenerateMergeWidget(*GetBlueprintObj(), *this).ToSharedRef() ];
-	}
-	TabManager->InvokeTab(FBlueprintEditorTabs::MergeToolID);
+	MergeToolDockTab = TabManager->InvokeTab(FBlueprintEditorTabs::MergeToolID);
 }
 
 void FBlueprintEditor::StartEditingDefaults(bool bAutoFocus, bool bForceRefresh)
diff --git a/Engine/Source/Editor/Kismet/Private/BlueprintEditorTabFactories.h b/Engine/Source/Editor/Kismet/Private/BlueprintEditorTabFactories.h
--- a/Engine/Source/Editor/Kismet/Private/BlueprintEditorTabFactories.h
+++ b/Engine/Source/Editor/Kismet/Private/BlueprintEditorTabFactories.h
@@ -506,7 +506,7 @@ struct FMergeToolSummoner : public FWorkflowTabFactory
 	virtual TSharedRef<SWidget> CreateTabBody(const FWorkflowTabSpawnInfo& Info) const override
 	{
 		TSharedPtr<FBlueprintEditor> BlueprintEditorPtr = StaticCastSharedPtr<FBlueprintEditor>(HostingApp.Pin());
-		return BlueprintEditorPtr->GetMergeTool();
+		return BlueprintEditorPtr->ShowMergeTool();
 	}
 
 	virtual FText GetTabToolTipText(const FWorkflowTabSpawnInfo& Info) const override
diff --git a/Engine/Source/Editor/Kismet/Private/SBlueprintEditorToolbar.cpp b/Engine/Source/Editor/Kismet/Private/SBlueprintEditorToolbar.cpp
--- a/Engine/Source/Editor/Kismet/Private/SBlueprintEditorToolbar.cpp
+++ b/Engine/Source/Editor/Kismet/Private/SBlueprintEditorToolbar.cpp
@@ -22,6 +22,7 @@
 #include "STutorialWrapper.h"
 #include "SBlueprintEditorSelectedDebugObjectWidget.h"
 #include "Engine/LevelScriptBlueprint.h"
+#include "Merge.h"
 
 #define LOCTEXT_NAMESPACE "KismetToolbar"
 
@@ -53,15 +54,6 @@ public:
 	// End of SWidget interface
 };
 
-static bool PendingMerge( const UBlueprint& BlueprintObj )
-{
-	bool bEnableMerge = false;
-	GConfig->GetBool(TEXT("AssetMerge"), TEXT("EnableAssetMerge"), bEnableMerge, GEngineIni);
-	ISourceControlProvider& SourceControlProvider = ISourceControlModule::Get().GetProvider();
-	FSourceControlStatePtr SourceControlState = SourceControlProvider.GetState(BlueprintObj.GetOutermost(), EStateCacheUsage::Use);
-	return bEnableMerge && SourceControlState.IsValid() && SourceControlState->IsConflicted();
-}
-
 //////////////////////////////////////////////////////////////////////////
 // FKismet2Menu
 
@@ -79,7 +71,7 @@ void FKismet2Menu::FillFileMenuBlueprintSection( FMenuBuilder& MenuBuilder, FBlu
 			FSlateIcon());
 
 		UBlueprint* BlueprintObj = Kismet.GetBlueprintObj();
-		if (BlueprintObj && PendingMerge( *BlueprintObj ) )
+		if (BlueprintObj && IMerge::Get().PendingMerge( *BlueprintObj ) )
 		{
 			MenuBuilder.AddMenuEntry(FBlueprintEditorCommands::Get().BeginBlueprintMerge);
 		}
diff --git a/Engine/Source/Editor/Kismet/Public/BlueprintEditor.h b/Engine/Source/Editor/Kismet/Public/BlueprintEditor.h
--- a/Engine/Source/Editor/Kismet/Public/BlueprintEditor.h
+++ b/Engine/Source/Editor/Kismet/Public/BlueprintEditor.h
@@ -191,7 +191,6 @@ public:
 	TSharedRef<class SBlueprintPalette> GetPalette() const {return Palette.ToSharedRef();}
 	TSharedRef<class SWidget> GetCompilerResults() const {return CompilerResults.ToSharedRef();}
 	TSharedRef<class SFindInBlueprints> GetFindResults() const {return FindResults.ToSharedRef();}
-	TSharedRef<class SWidget> GetMergeTool() const { return MergeTool.ToSharedRef(); }
 	TSharedRef<class SSCSEditor> GetSCSEditor() const {return SCSEditor.ToSharedRef();}
 	TSharedPtr<class SSCSEditorViewport> GetSCSViewport() const {return SCSViewport;}
 	TSharedPtr<class SMyBlueprint> GetMyBlueprintWidget() const {return MyBlueprintWidget;}
@@ -449,6 +448,16 @@ public:
 	 */
 	void UnregisterSCSEditorCustomization(const FName& InComponentName);
 
+	/** 
+	 * Forces the merge tool to be shown, if no merge is possible, returns an empty widget:
+	 */
+	TSharedRef<SWidget> ShowMergeTool();
+	
+	/** 
+	 * Closes the merge tool, rather than simply hiding it.
+	 */
+	void CloseMergeTool();
+
 	/** 
 	 * Check to see if we can customize the SCS editor for the passed-in scene component 
 	 * @param	InComponentToCustomize	The component to check to see if a customization exists
@@ -916,7 +925,10 @@ protected:
 	TSharedPtr<class SFindInBlueprints> FindResults;
 
 	/** Merge tool */
-	TSharedPtr<class SHorizontalBox> MergeTool;
+	TWeakPtr<class SWidget> MergeTool;
+
+	/** Docktab that contains the merge tool */
+	TWeakPtr<class SDockTab> MergeToolDockTab;
 
 	/** Reference to owner of the current popup */
 	TWeakPtr<class SWindow> NameEntryPopupWindow;
