Modified the code that updates the replication settings to mark the blueprint as modified when changed.
#TTP 339267 - Editor: If you change a custom event's replication setting, it doesn't show that it needs to be compiled.
#Branch UE4
#Proj Kismet

#Change Modified FBlueprintGraphActionDetails::SetNetFlags to mark the blueprint as structurally modified.

diff --git a/Engine/Source/Editor/Kismet/Private/BlueprintDetailsCustomization.cpp b/Engine/Source/Editor/Kismet/Private/BlueprintDetailsCustomization.cpp
--- a/Engine/Source/Editor/Kismet/Private/BlueprintDetailsCustomization.cpp
+++ b/Engine/Source/Editor/Kismet/Private/BlueprintDetailsCustomization.cpp
@@ -2229,33 +2229,28 @@ void FBlueprintGraphActionDetails::SetNetFlags( TWeakObjectPtr<UK2Node_EditableP
 {
 	if( FunctionEntryNode.IsValid() )
 	{
+		const int32 FlagsToSet = NetFlags ? FUNC_Net|NetFlags : 0;
+		const int32 FlagsToClear = FUNC_Net|FUNC_NetMulticast|FUNC_NetServer|FUNC_NetClient;
 		// Clear all net flags before setting
-		if (UK2Node_FunctionEntry* TypedEntryNode = Cast<UK2Node_FunctionEntry>(FunctionEntryNode.Get()))
+		if( FlagsToSet != FlagsToClear )
 		{
-			TypedEntryNode->ExtraFlags &= ~FUNC_Net;
-			TypedEntryNode->ExtraFlags &= ~FUNC_NetMulticast;
-			TypedEntryNode->ExtraFlags &= ~FUNC_NetServer;
-			TypedEntryNode->ExtraFlags &= ~FUNC_NetClient;
-		}
-		if (UK2Node_CustomEvent * CustomEventNode = Cast<UK2Node_CustomEvent>(FunctionEntryNode.Get()))
-		{
-			CustomEventNode->FunctionFlags &= ~FUNC_Net;
-			CustomEventNode->FunctionFlags &= ~FUNC_NetMulticast;
-			CustomEventNode->FunctionFlags &= ~FUNC_NetServer;
-			CustomEventNode->FunctionFlags &= ~FUNC_NetClient;
-		}
+			bool bBlueprintModified = false;
 
-		if (NetFlags > 0)
-		{
 			if (UK2Node_FunctionEntry* TypedEntryNode = Cast<UK2Node_FunctionEntry>(FunctionEntryNode.Get()))
 			{
-				TypedEntryNode->ExtraFlags |= FUNC_Net;
-				TypedEntryNode->ExtraFlags |= NetFlags;
+				TypedEntryNode->ExtraFlags &= ~FlagsToClear;
+				TypedEntryNode->ExtraFlags |= FlagsToSet;
+				bBlueprintModified = true;
 			}
 			if (UK2Node_CustomEvent * CustomEventNode = Cast<UK2Node_CustomEvent>(FunctionEntryNode.Get()))
 			{
-				CustomEventNode->FunctionFlags |= FUNC_Net;
-				CustomEventNode->FunctionFlags |= NetFlags;
+				CustomEventNode->FunctionFlags &= ~FlagsToClear;
+				CustomEventNode->FunctionFlags |= FlagsToSet;
+				bBlueprintModified = true;
+			}
+			if( bBlueprintModified )
+			{
+				FBlueprintEditorUtils::MarkBlueprintAsStructurallyModified( FunctionEntryNode->GetBlueprint() );
 			}
 		}
 	}
