Components mode will no longer clear selection when you switch back to the mode.

#ttp 340947 - BP:  Switching out of Components mode loses your selection

diff --git a/Engine/Source/Editor/Kismet/Private/BlueprintEditorModes.cpp b/Engine/Source/Editor/Kismet/Private/BlueprintEditorModes.cpp
--- a/Engine/Source/Editor/Kismet/Private/BlueprintEditorModes.cpp
+++ b/Engine/Source/Editor/Kismet/Private/BlueprintEditorModes.cpp
@@ -334,7 +334,6 @@ void FBlueprintComponentsApplicationMode::PreDeactivateMode()
 
 	TSharedPtr<FBlueprintEditor> BP = MyBlueprintEditor.Pin();
 	BP->GetSCSEditor()->SetEnabled(true);
-	BP->GetSCSEditor()->ClearSelection();
 	BP->GetSCSEditor()->UpdateTree();
 	BP->GetInspector()->SetEnabled(true);
 	BP->GetInspector()->EnableComponentDetailsCustomization(false);
@@ -344,7 +343,6 @@ void FBlueprintComponentsApplicationMode::PreDeactivateMode()
 void FBlueprintComponentsApplicationMode::PostActivateMode()
 {
 	TSharedPtr<FBlueprintEditor> BP = MyBlueprintEditor.Pin();
-	BP->GetSCSEditor()->ClearSelection();
 	BP->GetSCSEditor()->UpdateTree();
 	BP->EnableSCSPreview(true);
 	BP->UpdateSCSPreview();
