Can now copy and select all items in the Find-in-Blueprints results window.

#ttp 331279 - Find in blueprints broken?

diff --git a/Engine/Source/Editor/Kismet/Private/FindInBlueprints.cpp b/Engine/Source/Editor/Kismet/Private/FindInBlueprints.cpp
--- a/Engine/Source/Editor/Kismet/Private/FindInBlueprints.cpp
+++ b/Engine/Source/Editor/Kismet/Private/FindInBlueprints.cpp
@@ -934,6 +934,8 @@ void SFindInBlueprints::Construct( const FArguments& InArgs, TSharedPtr<FBluepri
 {
 	BlueprintEditorPtr = InBlueprintEditor;
 
+	RegisterCommands();
+
 	bIsInFindWithinBlueprintMode = true;
 	
 	this->ChildSlot
@@ -978,6 +980,7 @@ void SFindInBlueprints::Construct( const FArguments& InArgs, TSharedPtr<FBluepri
 					.OnGetChildren( this, &SFindInBlueprints::OnGetChildren )
 					.OnMouseButtonDoubleClick(this,&SFindInBlueprints::OnTreeSelectionDoubleClicked)
 					.SelectionMode( ESelectionMode::Multi )
+					.OnContextMenuOpening(this, &SFindInBlueprints::OnContextMenuOpening)
 				]
 			]
 
@@ -1178,6 +1181,17 @@ void SFindInBlueprints::Tick( const FGeometry& AllottedGeometry, const double In
 	}
 }
 
+void SFindInBlueprints::RegisterCommands()
+{
+	TSharedPtr<FUICommandList> ToolKitCommandList = BlueprintEditorPtr.Pin()->GetToolkitCommands();
+
+	ToolKitCommandList->MapAction( FGenericCommands::Get().Copy,
+		FExecuteAction::CreateSP(this, &SFindInBlueprints::OnCopyAction) );
+
+	ToolKitCommandList->MapAction( FGenericCommands::Get().SelectAll,
+		FExecuteAction::CreateSP(this, &SFindInBlueprints::OnSelectAllAction) );
+}
+
 void SFindInBlueprints::FocusForUse(bool bSetFindWithinBlueprint, FString NewSearchTerms, bool bSelectFirstResult)
 {
 	// NOTE: Careful, GeneratePathToWidget can be reentrant in that it can call visibility delegates and such
@@ -1529,4 +1543,69 @@ void SFindInBlueprints::OnCacheComplete()
 	OnSearchTextCommitted(SearchTextField->GetText(), ETextCommit::OnEnter);
 }
 
+TSharedPtr<SWidget> SFindInBlueprints::OnContextMenuOpening()
+{
+	const bool bShouldCloseWindowAfterMenuSelection = true;
+	FMenuBuilder MenuBuilder( bShouldCloseWindowAfterMenuSelection,BlueprintEditorPtr.Pin()->GetToolkitCommands());
+
+	MenuBuilder.BeginSection("BasicOperations");
+	{
+		MenuBuilder.AddMenuEntry(FGenericCommands::Get().SelectAll);
+		MenuBuilder.AddMenuEntry(FGenericCommands::Get().Copy);
+	}
+
+	return MenuBuilder.MakeWidget();
+}
+
+void SFindInBlueprints::SelectAllItemsHelper(FSearchResult InItemToSelect)
+{
+	// Iterates over all children and recursively selects all items in the results
+	TreeView->SetItemSelection(InItemToSelect, true);
+
+	for( const auto Child : InItemToSelect->Children )
+	{
+		SelectAllItemsHelper(Child);
+	}
+}
+
+void SFindInBlueprints::OnSelectAllAction()
+{
+	for( const auto Item : ItemsFound )
+	{
+		SelectAllItemsHelper(Item);
+	}
+}
+
+void SFindInBlueprints::OnCopyAction()
+{
+	TArray< FSearchResult > SelectedItems = TreeView->GetSelectedItems();
+
+	FString SelectedText;
+
+	for( const auto SelectedItem : SelectedItems)
+	{
+		// Add indents for each layer into the tree the item is
+		for(auto ParentItem = SelectedItem->Parent; ParentItem.IsValid(); ParentItem = ParentItem.Pin()->Parent)
+		{
+			SelectedText += TEXT("\t");
+		}
+
+		// Add the display string
+		SelectedText += SelectedItem->GetDisplayString();
+
+		// If there is a comment, add two indents and then the comment
+		FString CommentText = SelectedItem->GetCommentText();
+		if(!CommentText.IsEmpty())
+		{
+			SelectedText += TEXT("\t\t") + CommentText;
+		}
+		
+		// Line terminator so the next item will be on a new line
+		SelectedText += LINE_TERMINATOR;
+	}
+
+	// Copy text to clipboard
+	FPlatformMisc::ClipboardCopy( *SelectedText );
+}
+
 #undef LOCTEXT_NAMESPACE
diff --git a/Engine/Source/Editor/Kismet/Public/FindInBlueprints.h b/Engine/Source/Editor/Kismet/Public/FindInBlueprints.h
--- a/Engine/Source/Editor/Kismet/Public/FindInBlueprints.h
+++ b/Engine/Source/Editor/Kismet/Public/FindInBlueprints.h
@@ -229,6 +229,9 @@ public:
 	void OnCacheComplete();
 private:
 
+	/** Register any Find-in-Blueprint commands */
+	void RegisterCommands();
+
 	/*Called when user changes the text they are searching for */
 	void OnSearchTextChanged(const FText& Text);
 
@@ -307,6 +310,18 @@ private:
 	/** Returns the color of the caching bar */
 	FSlateColor GetCachingBarColor() const;
 
+	/** Callback to build the context menu when right clicking in the tree */
+	TSharedPtr<SWidget> OnContextMenuOpening();
+
+	/** Helper function to select all items */
+	void SelectAllItemsHelper(FSearchResult InItemToSelect);
+
+	/** Callback when user attempts to select all items in the search results */
+	void OnSelectAllAction();
+
+	/** Callback when user attempts to copy their selection in the Find-in-Blueprints */
+	void OnCopyAction();
+
 private:
 	/** Pointer back to the blueprint editor that owns us */
 	TWeakPtr<class FBlueprintEditor> BlueprintEditorPtr;
