Making it so we select the proper PIE instance (not the editor one) when picking an object to actively debug.
(TTP #341188)

diff --git a/Engine/Source/Editor/Kismet/Private/SBlueprintEditorSelectedDebugObjectWidget.cpp b/Engine/Source/Editor/Kismet/Private/SBlueprintEditorSelectedDebugObjectWidget.cpp
--- a/Engine/Source/Editor/Kismet/Private/SBlueprintEditorSelectedDebugObjectWidget.cpp
+++ b/Engine/Source/Editor/Kismet/Private/SBlueprintEditorSelectedDebugObjectWidget.cpp
@@ -412,13 +412,14 @@ void SBlueprintEditorSelectedDebugObjectWidget::GenerateDebugObjectNames(bool bR
 				continue;
 			}
 
-			// We don't have a specific debug world, but the object isnt in a PIE world
-			if ((ObjWorld->WorldType != EWorldType::PIE) && (ObjWorld->WorldType != EWorldType::Editor))
+			if ((ObjWorld->WorldType == EWorldType::Editor) && (GUnrealEd->GetPIEViewport() == nullptr))
 			{
-				continue;
+				AddDebugObject(TestObject);
 			}
-
-			AddDebugObject(TestObject);
+			else if (ObjWorld->WorldType == EWorldType::PIE)
+			{
+				AddDebugObject(TestObject);
+			}			
 		}
 	}
 
