Double clicking anim graph nodes opens respective asset editor

diff --git a/Engine/Source/Editor/AnimGraph/Classes/AnimGraphNode_Base.h b/Engine/Source/Editor/AnimGraph/Classes/AnimGraphNode_Base.h
--- a/Engine/Source/Editor/AnimGraph/Classes/AnimGraphNode_Base.h
+++ b/Engine/Source/Editor/AnimGraph/Classes/AnimGraphNode_Base.h
@@ -113,6 +113,9 @@ class UAnimGraphNode_Base : public UK2Node
 	ANIMGRAPH_API virtual bool CanCreateUnderSpecifiedSchema(const UEdGraphSchema* DesiredSchema) const override;
 	ANIMGRAPH_API virtual void GetMenuEntries(FGraphContextMenuBuilder& ContextMenuBuilder) const override;
 	ANIMGRAPH_API virtual void GetNodeAttributes(TArray<TKeyValuePair<FString, FString>>& OutNodeAttributes) const override;
+
+	// By default return any animation assets we have
+	virtual UObject* GetJumpTargetForDoubleClick() const override { return GetAnimationAsset(); }
 	// End of UK2Node interface
 
 	// UAnimGraphNode_Base interface
diff --git a/Engine/Source/Editor/Kismet/Private/BlueprintEditor.cpp b/Engine/Source/Editor/Kismet/Private/BlueprintEditor.cpp
--- a/Engine/Source/Editor/Kismet/Private/BlueprintEditor.cpp
+++ b/Engine/Source/Editor/Kismet/Private/BlueprintEditor.cpp
@@ -4873,7 +4873,29 @@ void FBlueprintEditor::OnNodeDoubleClicked(class UEdGraphNode* Node)
 	}
 	else if (UObject* HyperlinkTarget = Node->GetJumpTargetForDoubleClick())
 	{
-		FKismetEditorUtilities::BringKismetToFocusAttentionOnObject(HyperlinkTarget);
+		// Check to see if our outer chain contains a blueprint. If we're inside a blueprint (a graph, pin, etc.) then
+		// focus on the target; otherwise open the editor for the target.
+		UBlueprint* TargetBP = Cast<UBlueprint>(const_cast<UObject*>(HyperlinkTarget));
+		if(TargetBP == NULL)
+		{
+			for(UObject* TestOuter = HyperlinkTarget->GetOuter(); TestOuter; TestOuter = TestOuter->GetOuter())
+			{
+				TargetBP = Cast<UBlueprint>(TestOuter);
+				if(TargetBP != NULL)
+				{
+					break;
+				}
+			}
+		}
+
+		if(TargetBP)
+		{
+			FKismetEditorUtilities::BringKismetToFocusAttentionOnObject(HyperlinkTarget);
+		}
+		else
+		{
+			FAssetEditorManager::Get().OpenEditorForAsset(HyperlinkTarget);
+		}
 	}
 }
 
