(WIP) Setting it up so we can switch to the new blueprint menu system using an .ini setting (bUseLegacyActionMenus)... enabled for now (until refactor is finished).

Also, removed bAllowBlueprintComms (as BP comms are stable and enabled fulltime).

diff --git a/Engine/Config/BaseEditor.ini b/Engine/Config/BaseEditor.ini
--- a/Engine/Config/BaseEditor.ini
+++ b/Engine/Config/BaseEditor.ini
@@ -88,7 +88,7 @@ MenuPriority=11
 MenuPriority=10
 
 [/Script/BlueprintGraph.EdGraphSchema_K2]
-bAllowBlueprintComms=true
+bUseLegacyActionMenus=true
 +EditoronlyBPFunctionRedirects=(ClassName="AIBlueprintHelperLibrary", OldFunctionName="SpawnAI", NewFunctionName="SpawnAIFromClass", BlueprintParamName="Pawn", ClassParamName="PawnClass")
 +EditoronlyBPFunctionRedirects=(ClassName="GameplayStatics", OldFunctionName="CreateSaveGameObjectFromBlueprint", NewFunctionName="CreateSaveGameObject", BlueprintParamName="SaveGameBlueprint", ClassParamName="SaveGameClass")
 
diff --git a/Engine/Source/Editor/BlueprintGraph/Classes/EdGraphSchema_K2.h b/Engine/Source/Editor/BlueprintGraph/Classes/EdGraphSchema_K2.h
--- a/Engine/Source/Editor/BlueprintGraph/Classes/EdGraphSchema_K2.h
+++ b/Engine/Source/Editor/BlueprintGraph/Classes/EdGraphSchema_K2.h
@@ -351,9 +351,9 @@ class BLUEPRINTGRAPH_API UEdGraphSchema_K2 : public UEdGraphSchema
 	UPROPERTY()
 	int32 AG_LevelReference;
 
-	/** Whether or not the schema should allow the user to use blueprint communications */
-	UPROPERTY(globalconfig)
-	bool bAllowBlueprintComms;
+	/** Whether this schema should use the old (legacy) menu building, or the new (experimental) menu building */
+	UPROPERTY(GlobalConfig)
+	bool bUseLegacyActionMenus;
 
 	UPROPERTY(globalconfig)
 	TArray<FBlueprintCallableFunctionRedirect> EditoronlyBPFunctionRedirects;
diff --git a/Engine/Source/Editor/BlueprintGraph/Private/BlueprintActionFilter.cpp b/Engine/Source/Editor/BlueprintGraph/Private/BlueprintActionFilter.cpp
--- a/Engine/Source/Editor/BlueprintGraph/Private/BlueprintActionFilter.cpp
+++ b/Engine/Source/Editor/BlueprintGraph/Private/BlueprintActionFilter.cpp
@@ -729,12 +729,12 @@ static bool BlueprintActionFilterImpl::IsNodeUnpastable(FBlueprintActionFilter c
 //------------------------------------------------------------------------------
 static bool BlueprintActionFilterImpl::IsFilteredNodeType(FBlueprintActionFilter const& Filter, UBlueprintNodeSpawner const* BlueprintAction, bool const bExcludeChildClasses)
 {
-	bool bIsFilteredOut = false;
+	bool bIsFilteredOut = (Filter.NodeTypes.Num() > 0);
 
 	UClass* NodeClass = BlueprintAction->NodeClass;
 	if (NodeClass != nullptr)
 	{
-		for (UClass* AllowedClass : Filter.NodeTypes)
+		for (TSubclassOf<UEdGraphNode> AllowedClass : Filter.NodeTypes)
 		{
 			if ((bExcludeChildClasses && (AllowedClass != NodeClass)) ||
 				!NodeClass->IsChildOf(AllowedClass))
@@ -743,13 +743,18 @@ static bool BlueprintActionFilterImpl::IsFilteredNodeType(FBlueprintActionFilter
 				break;
 			}
 		}
+		
+		for (int32 ClassIndx = 0; !bIsFilteredOut && (ClassIndx < Filter.ExcludedNodeTypes.Num()); ++ClassIndx)
+		{
+			TSubclassOf<UEdGraphNode> ExcludedClass = Filter.ExcludedNodeTypes[ClassIndx];
+			if ((bExcludeChildClasses && (ExcludedClass == NodeClass)) ||
+				NodeClass->IsChildOf(ExcludedClass))
+			{
+				bIsFilteredOut = true;
+				break;
+			}
+		}
 	}
-	else
-	{
-		// this node's class type is unknown, filter it out!
-		bIsFilteredOut = (Filter.NodeTypes.Num() > 0);
-	}
-
 	return bIsFilteredOut;
 }
 
diff --git a/Engine/Source/Editor/BlueprintGraph/Private/BlueprintActionMenuBuilder.cpp b/Engine/Source/Editor/BlueprintGraph/Private/BlueprintActionMenuBuilder.cpp
--- a/Engine/Source/Editor/BlueprintGraph/Private/BlueprintActionMenuBuilder.cpp
+++ b/Engine/Source/Editor/BlueprintGraph/Private/BlueprintActionMenuBuilder.cpp
@@ -297,17 +297,6 @@ namespace FBlueprintActionMenuBuilderImpl
 		/** In charge of spawning menu items for this section (holds category/ordering information)*/
 		FBlueprintActionMenuItemFactory ItemFactory;
 	};
-
-	/**
-	 * To offer a fallback in case the action menu refactor is unstable, this
-	 * method implements the old way we used collect blueprint menu actions (for
-	 * both the palette and context menu).
-	 * 
-	 * @param  MenuSections	All the sections that were defined for FBlueprintActionMenuBuilder (we attempt to combine them)
-	 * @param  MenuBuilder	The menu builder we want all the legacy actions appended to.
-	 * @return 
-	 */
-	static void AppendLegacyActions(TArray< TSharedRef<FMenuSectionDefinition> > const& MenuSections, FBlueprintActionMenuBuilder& MenuBuilder);
 }
 
 //------------------------------------------------------------------------------
@@ -317,50 +306,6 @@ FBlueprintActionMenuBuilderImpl::FMenuSectionDefinition::FMenuSectionDefinition(
 {
 }
 
-//------------------------------------------------------------------------------
-void FBlueprintActionMenuBuilderImpl::AppendLegacyActions(TArray< TSharedRef<FMenuSectionDefinition> > const& MenuSections, FBlueprintActionMenuBuilder& MenuBuilder)
-{
-	check(MenuSections.Num() > 0);
-	FBlueprintActionContext const& Context = MenuSections[0]->Filter.Context;
-
-	if (Context.Graphs.Num() > 0)
-	{
-		FBlueprintGraphActionListBuilder ContextMenuBuilder(Context.Graphs[0]);
-		if (Context.Pins.Num() > 0)
-		{
-			ContextMenuBuilder.FromPin = Context.Pins[0];
-		}
-		
-		// @TODO: ContextMenuBuilder.SelectedObjects
-// 		for (TSharedRef<FMenuSectionDefinition> MenuSection : MenuSections)
-// 		{
-// 			FBlueprintActionFilter const& SectionFilter = MenuSection->Filter;
-// 		}
-
-		UEdGraphSchema const* GraphSchema = GetDefault<UEdGraphSchema>(ContextMenuBuilder.CurrentGraph->Schema);
-		GraphSchema->GetGraphContextActions(ContextMenuBuilder);
-
-		MenuBuilder.Append(ContextMenuBuilder);
-	}
-	else if (Context.Blueprints.Num() > 0)
-	{
-		FBlueprintPaletteListBuilder PaletteBuilder(Context.Blueprints[0], MenuSections[0]->ItemFactory.RootCategory.ToString());
-
-		UClass* FilterClass = nullptr;
-
-		FBlueprintActionFilter const& SectionFilter = MenuSections[0]->Filter;
-		if (SectionFilter.OwnerClasses.Num() > 0)
-		{
-			FilterClass = SectionFilter.OwnerClasses[0];
-		}
-
-		UEdGraphSchema_K2 const* K2Schema = GetDefault<UEdGraphSchema_K2>();
-		FK2ActionMenuBuilder(K2Schema).GetPaletteActions(PaletteBuilder, FilterClass); 
-
-		MenuBuilder.Append(PaletteBuilder);
-	}
-}
-
 /*******************************************************************************
  * FBlueprintActionMenuBuilder
  ******************************************************************************/
@@ -396,40 +341,7 @@ void FBlueprintActionMenuBuilder::RebuildActionList()
 	using namespace FBlueprintActionMenuBuilderImpl;
 
 	Empty();
-	if (true)//!GetDefault<UEdGraphSchema_K2>()->bUseLegacyActionMenus)
-	{
-		GenerateMenuSections();
-	}
-	else if (MenuSections.Num() > 0)
-	{
-		// use the context from the first section (the legacy menu building did
-		// it all in one sweep and knew how to mark off sections)
-		AppendLegacyActions(MenuSections, *this);
-	}
-
-	// @TODO: account for all K2ActionMenuBuilder action types...
-	// - FEdGraphSchemaAction_K2AddTimeline
-	// - FEdGraphSchemaAction_K2ViewNode
-	// - FEdGraphSchemaAction_K2AddCustomEvent
-	//   FEdGraphSchemaAction_EventFromFunction
-	//   FEdGraphSchemaAction_K2Var
-	//   FEdGraphSchemaAction_K2Delegate
-	//   FEdGraphSchemaAction_K2AssignDelegate
-	// - FEdGraphSchemaAction_K2AddComment
-	//   FEdGraphSchemaAction_K2PasteHere
-	// - FEdGraphSchemaAction_K2NewNode
-	//   FEdGraphSchemaAction_Dummy
-	//   FEdGraphSchemaAction_K2AddCallOnActor
-	//   FEdGraphSchemaAction_K2AddCallOnVariable
-	// - FEdGraphSchemaAction_K2AddComponent
-	//   FEdGraphSchemaAction_K2AddComment
-}
-
-//------------------------------------------------------------------------------
-void FBlueprintActionMenuBuilder::GenerateMenuSections()
-{
-	using namespace FBlueprintActionMenuBuilderImpl;
-
+	
 	FBlueprintActionDatabase::FClassActionMap const& ClassActions = FBlueprintActionDatabase::Get().GetAllActions();
 	for (auto ActionEntry : ClassActions)
 	{
@@ -443,14 +355,31 @@ void FBlueprintActionMenuBuilder::GenerateMenuSections()
 					// this node spawner doesn't belong in this section of the menu
 					continue;
 				}
-
+				
 				TSharedPtr<FEdGraphSchemaAction> NewActionPtr = MenuSection->ItemFactory.MakeMenuItem(*this, NodeSpawner);
 				check(NewActionPtr.IsValid());
-
+				
 				AddAction(NewActionPtr);
 			}
 		}
 	}
+
+	// @TODO: account for all K2ActionMenuBuilder action types...
+	// - FEdGraphSchemaAction_K2AddTimeline
+	// - FEdGraphSchemaAction_K2ViewNode
+	// - FEdGraphSchemaAction_K2AddCustomEvent
+	//   FEdGraphSchemaAction_EventFromFunction
+	//   FEdGraphSchemaAction_K2Var
+	//   FEdGraphSchemaAction_K2Delegate
+	//   FEdGraphSchemaAction_K2AssignDelegate
+	// - FEdGraphSchemaAction_K2AddComment
+	//   FEdGraphSchemaAction_K2PasteHere
+	// - FEdGraphSchemaAction_K2NewNode
+	//   FEdGraphSchemaAction_Dummy
+	//   FEdGraphSchemaAction_K2AddCallOnActor
+	//   FEdGraphSchemaAction_K2AddCallOnVariable
+	// - FEdGraphSchemaAction_K2AddComponent
+	//   FEdGraphSchemaAction_K2AddComment
 }
 
 #undef LOCTEXT_NAMESPACE
diff --git a/Engine/Source/Editor/BlueprintGraph/Private/K2ActionMenuBuilder.cpp b/Engine/Source/Editor/BlueprintGraph/Private/K2ActionMenuBuilder.cpp
--- a/Engine/Source/Editor/BlueprintGraph/Private/K2ActionMenuBuilder.cpp
+++ b/Engine/Source/Editor/BlueprintGraph/Private/K2ActionMenuBuilder.cpp
@@ -1043,8 +1043,7 @@ struct FClassDynamicCastHelper
 		check(TestClass && K2Schema);
 		const bool bIsSkeletonClass = FKismetEditorUtilities::IsClassABlueprintSkeleton(TestClass);
 		const bool bProperClass = TestClass->HasAnyFlags(RF_Public) && !TestClass->HasAnyClassFlags(CLASS_Deprecated | CLASS_NewerVersionExists);
-		const bool bAllowedByBPComms = K2Schema->bAllowBlueprintComms || !TestClass->ClassGeneratedBy;
-		return !bIsSkeletonClass && bProperClass && bAllowedByBPComms;
+		return !bIsSkeletonClass && bProperClass;
 	}
 
 	static void SpawnNode(UClass* TargetType, FBlueprintGraphActionListBuilder& ContextMenuBuilder)
@@ -1170,15 +1169,6 @@ void FK2ActionMenuBuilder::GetPinAllowedNodeTypes(FBlueprintGraphActionListBuild
 
 				if (FromPin.Direction == EGPD_Output)
 				{
-					// If this is a call to another generated class, find the first parent that is natively-defined, to prevent calls to improper function headers
-					if( !K2Schema->bAllowBlueprintComms &&(PinClass->ClassGeneratedBy && PinClass->ClassGeneratedBy != Blueprint) && !FBlueprintEditorUtils::IsLevelScriptBlueprint(Blueprint) )
-					{
-						while( PinClass->ClassGeneratedBy != NULL )
-						{
-							PinClass = PinClass->GetSuperClass();
-						}
-					}
-
 					const bool bWantBindableDelegates = (GraphType == GT_Ubergraph) || (GraphType == GT_Function) ||
 						((bAllowUnsafeCommands) && (Blueprint->BlueprintType == BPTYPE_MacroLibrary));
 
@@ -1236,7 +1226,7 @@ void FK2ActionMenuBuilder::GetPinAllowedNodeTypes(FBlueprintGraphActionListBuild
 						UClass* TestClass = *ClassIt;
 						bool bIsSkeletonClass = FKismetEditorUtilities::IsClassABlueprintSkeleton(TestClass);
 
-						if (TestClass->HasAnyFlags(RF_Public) && !bIsSkeletonClass && !TestClass->HasAnyClassFlags(CLASS_Deprecated|CLASS_NewerVersionExists) && (K2Schema->bAllowBlueprintComms || !TestClass->ClassGeneratedBy))
+						if (TestClass->HasAnyFlags(RF_Public) && !bIsSkeletonClass && !TestClass->HasAnyClassFlags(CLASS_Deprecated|CLASS_NewerVersionExists))
 						{
 							bool bIsCastable = false;
 							if (TestClass != PinClass)
diff --git a/Engine/Source/Editor/BlueprintGraph/Public/BlueprintActionFilter.h b/Engine/Source/Editor/BlueprintGraph/Public/BlueprintActionFilter.h
--- a/Engine/Source/Editor/BlueprintGraph/Public/BlueprintActionFilter.h
+++ b/Engine/Source/Editor/BlueprintGraph/Public/BlueprintActionFilter.h
@@ -78,7 +78,14 @@ public:
 	 * these types, then it is filtered out. Use the "ExcludeChildNodeTypes"
 	 * flag to aggressively filter out child classes (must be an explicit match).
 	 */
-	TArray<UClass*> NodeTypes;
+	TArray< TSubclassOf<UEdGraphNode> > NodeTypes;
+	
+	/**
+	 * A list of node types that should be filtered out. If a node class is 
+	 * listed both here and in NodeTypes, then the exclusion wins (and it will 
+	 * be filtered out).
+	 */
+	TArray< TSubclassOf<UEdGraphNode> > ExcludedNodeTypes;
 	
 	/**
 	 * A list of classes that you want members for. If an action does not have 
diff --git a/Engine/Source/Editor/BlueprintGraph/Public/BlueprintActionMenuBuilder.h b/Engine/Source/Editor/BlueprintGraph/Public/BlueprintActionMenuBuilder.h
--- a/Engine/Source/Editor/BlueprintGraph/Public/BlueprintActionMenuBuilder.h
+++ b/Engine/Source/Editor/BlueprintGraph/Public/BlueprintActionMenuBuilder.h
@@ -53,17 +53,13 @@ public:
 	void AddMenuSection(FBlueprintActionFilter const& Filter, FText const& Heading = FText::GetEmpty(), int32 MenuOrder = 0);
 	
 	/**
-	 * Regenerates the entire menu list from the cached menu sections.
+	 * Regenerates the entire menu list from the cached menu sections. Filters 
+	 * and adds action items from the blueprint action database (as defined by 
+	 * the MenuSections list).
 	 */
 	void RebuildActionList();
 
 private:
-	/**
-	 * Filters and adds action items from the blueprint action database (as 
-	 * defined by the MenuSections list).
-	 */
-	void GenerateMenuSections();
-	
 	/** 
 	 * Defines all the separate sections of the menu (filter, sort order, etc.).
 	 * Defined as a TSharedRef<> so as to hide the implementation details (keep 
diff --git a/Engine/Source/Editor/Kismet/Private/SBlueprintActionMenu.cpp b/Engine/Source/Editor/Kismet/Private/SBlueprintActionMenu.cpp
--- a/Engine/Source/Editor/Kismet/Private/SBlueprintActionMenu.cpp
+++ b/Engine/Source/Editor/Kismet/Private/SBlueprintActionMenu.cpp
@@ -10,6 +10,8 @@
 #include "SMyBlueprint.h"
 #include "BlueprintEditorUtils.h"
 #include "K2ActionMenuBuilder.h" // for FBlueprintPaletteListBuilder
+#include "BlueprintActionMenuBuilder.h"
+#include "BlueprintActionFilter.h"
 #include "BlueprintPaletteFavorites.h"
 #include "IDocumentation.h"
 
@@ -384,44 +386,103 @@ ESlateCheckBoxState::Type SBlueprintActionMenu::ContextToggleIsChecked() const
 
 void SBlueprintActionMenu::CollectAllActions(FGraphActionListBuilderBase& OutAllActions)
 {
-	// Build up the context object
-	FBlueprintGraphActionListBuilder ContextMenuBuilder(GraphObj);
-	bool bIsContextSensitive = EditorPtr.Pin()->GetIsContextSensitive();
-	bool bHasPins = DraggedFromPins.Num() > 0;
-	if (!bIsContextSensitive)
-	{
-		FBlueprintPaletteListBuilder PaletteBuilder(EditorPtr.Pin()->GetBlueprintObj());
-		UEdGraphSchema_K2::GetAllActions(PaletteBuilder);
-		//@TODO: Avoid this copy
-		OutAllActions.Append(PaletteBuilder);
-	}
-	else if (!bHasPins)
+	UBlueprint* Blueprint = EditorPtr.Pin()->GetBlueprintObj();
+	
+	check(EditorPtr.IsValid());
+	TSharedPtr<FBlueprintEditor> BlueprintEditor = EditorPtr.Pin();
+	bool const bIsContextSensitive = BlueprintEditor->GetIsContextSensitive();
+	
+	if (GetDefault<UEdGraphSchema_K2>()->bUseLegacyActionMenus)
 	{
-		// Pass in selected property
-		FEdGraphSchemaAction_K2Var* SelectedVar = EditorPtr.Pin()->GetMyBlueprintWidget()->SelectionAsVar();
-		if (SelectedVar != NULL)
+		// Build up the context object
+		FBlueprintGraphActionListBuilder ContextMenuBuilder(GraphObj);
+		bool bHasPins = DraggedFromPins.Num() > 0;
+		if (!bIsContextSensitive)
+		{
+			FBlueprintPaletteListBuilder PaletteBuilder(Blueprint);
+			UEdGraphSchema_K2::GetAllActions(PaletteBuilder);
+			//@TODO: Avoid this copy
+			OutAllActions.Append(PaletteBuilder);
+		}
+		else if (!bHasPins)
+		{
+			// Pass in selected property
+			FEdGraphSchemaAction_K2Var* SelectedVar = EditorPtr.Pin()->GetMyBlueprintWidget()->SelectionAsVar();
+			if (SelectedVar != NULL)
+			{
+				ContextMenuBuilder.SelectedObjects.Add(SelectedVar->GetProperty());
+			}
+		}
+		else
+		{
+			ContextMenuBuilder.FromPin = DraggedFromPins[0];
+		}
+		
+		// Note: Cannot call GetGraphContextActions() during serialization and GC due to its use of FindObject()
+		if(!GIsSavingPackage && !GIsGarbageCollecting)
 		{
-			ContextMenuBuilder.SelectedObjects.Add(SelectedVar->GetProperty());
+			// Determine all possible actions
+			GraphObj->GetSchema()->GetGraphContextActions(ContextMenuBuilder);
 		}
+		
+		// Copy the added options back to the main list
+		//@TODO: Avoid this copy
+		OutAllActions.Append(ContextMenuBuilder);
 	}
 	else
 	{
-		ContextMenuBuilder.FromPin = DraggedFromPins[0];
+		FBlueprintActionFilter MenuFilter;
+		MenuFilter.Context.Blueprints.Add(Blueprint);
+		// we still want context from the graph (even if the user has unchecked
+		// "Context Sensitive"), otherwise the user would be presented with nodes
+		// that can't be placed in the graph... if the user isn't being presented
+		// with a valid node, then fix it up in filtering
+		MenuFilter.Context.Graphs.Add(GraphObj);
+		
+		FBlueprintActionMenuBuilder MenuBuilder;
+		if (bIsContextSensitive)
+		{
+			MenuFilter.Context.Pins = DraggedFromPins;
+			
+			FEdGraphSchemaAction_K2Var* SelectedVar = BlueprintEditor->GetMyBlueprintWidget()->SelectionAsVar();
+			if (SelectedVar != nullptr)
+			{
+				if (UObjectProperty* SelectedProperty = Cast<UObjectProperty>(SelectedVar->GetProperty()))
+				{
+					MenuFilter.OwnerClasses.Add(SelectedProperty->PropertyClass);
+					
+					bool const bUseFriendlyName = GetDefault<UEditorStyleSettings>()->bShowFriendlyNames;
+					FText const PropertyName = bUseFriendlyName ? FText::FromString(UEditorEngine::GetFriendlyName(SelectedProperty)) : FText::FromName(SelectedProperty->GetFName());
+					
+					MenuFilter.NodeTypes.Add(UK2Node_CallFunction::StaticClass());
+					FText const FuncSectionName = FText::Format(LOCTEXT("ComponentFuncSection", "Call Function on {0}"), PropertyName);
+					MenuBuilder.AddMenuSection(MenuFilter, FuncSectionName);
+					
+					MenuFilter.NodeTypes[0] = UK2Node_Event::StaticClass();
+					FText const EventSectionName = FText::Format(LOCTEXT("ComponentEventSection", "Add Event for {0}"), PropertyName);
+					MenuBuilder.AddMenuSection(MenuFilter, EventSectionName);
+					
+					MenuFilter.NodeTypes.Empty();
+					MenuFilter.OwnerClasses.Empty();
+				}
+			}
+			
+			// only allow members of this blueprint to be called (along with static
+			// library members)
+			MenuFilter.OwnerClasses.Add(Blueprint->GeneratedClass);
+		}	
+		MenuBuilder.AddMenuSection(MenuFilter);
+		
+		MenuBuilder.RebuildActionList();
+		// copy the added options back to the main list
+		OutAllActions.Append(MenuBuilder); // @TODO: Avoid this copy
 	}
-
-	// Note: Cannot call GetGraphContextActions() during serialization and GC due to its use of FindObject()
-	if(!GIsSavingPackage && !GIsGarbageCollecting)
+	
+	// also try adding promote to variable if we can do so.
+	if (DraggedFromPins.Num() > 0)
 	{
-		// Determine all possible actions
-		GraphObj->GetSchema()->GetGraphContextActions(ContextMenuBuilder);
-
-		// Also try adding promote to variable if we can do so.
-		TryInsertPromoteToVariable(ContextMenuBuilder, OutAllActions);
+		TryInsertPromoteToVariable(DraggedFromPins[0], OutAllActions);
 	}
-
-	// Copy the added options back to the main list
-	//@TODO: Avoid this copy
-	OutAllActions.Append(ContextMenuBuilder);
 }
 
 TSharedRef<SEditableTextBox> SBlueprintActionMenu::GetFilterTextBox()
@@ -459,18 +520,15 @@ void SBlueprintActionMenu::OnActionSelected( const TArray< TSharedPtr<FEdGraphSc
 	}
 }
 
-void SBlueprintActionMenu::TryInsertPromoteToVariable( FGraphContextMenuBuilder &ContextMenuBuilder, FGraphActionListBuilderBase &OutAllActions )
+void SBlueprintActionMenu::TryInsertPromoteToVariable(UEdGraphPin const* ContextPin, FGraphActionListBuilderBase& OutAllActions)
 {
 	// If we can promote this to a variable add a menu entry to do so.
 	const UEdGraphSchema_K2* K2Schema = Cast<const UEdGraphSchema_K2>(GraphObj->GetSchema());
-	if ( (K2Schema != NULL) && ( ContextMenuBuilder.FromPin != NULL ) )
+	if ((K2Schema != nullptr) && K2Schema->CanPromotePinToVariable(*ContextPin))
 	{
-		if( K2Schema ->CanPromotePinToVariable(*ContextMenuBuilder.FromPin) )
-		{	
-			TSharedPtr<FBlueprintAction_PromoteVariable> PromoteAction = TSharedPtr<FBlueprintAction_PromoteVariable>(new FBlueprintAction_PromoteVariable());
-			PromoteAction->MyBlueprintEditor = EditorPtr;
-			OutAllActions.AddAction( PromoteAction );
-		}
+		TSharedPtr<FBlueprintAction_PromoteVariable> PromoteAction = TSharedPtr<FBlueprintAction_PromoteVariable>(new FBlueprintAction_PromoteVariable());
+		PromoteAction->MyBlueprintEditor = EditorPtr;
+		OutAllActions.AddAction( PromoteAction );
 	}
 }
 
diff --git a/Engine/Source/Editor/Kismet/Private/SBlueprintActionMenu.h b/Engine/Source/Editor/Kismet/Private/SBlueprintActionMenu.h
--- a/Engine/Source/Editor/Kismet/Private/SBlueprintActionMenu.h
+++ b/Engine/Source/Editor/Kismet/Private/SBlueprintActionMenu.h
@@ -50,7 +50,7 @@ protected:
 	void CollectAllActions(FGraphActionListBuilderBase& OutAllActions);
 
 	/** Functioin to try to insert a promote to variable entry if it is possible to do so. */
-	void TryInsertPromoteToVariable( FGraphContextMenuBuilder &ContextMenuBuilder, FGraphActionListBuilderBase &OutAllActions );
+	void TryInsertPromoteToVariable(UEdGraphPin const* ContextPin, FGraphActionListBuilderBase& OutAllActions);
 
 private:
 	UEdGraph* GraphObj;
diff --git a/Engine/Source/Editor/Kismet/Private/SBlueprintLibraryPalette.cpp b/Engine/Source/Editor/Kismet/Private/SBlueprintLibraryPalette.cpp
--- a/Engine/Source/Editor/Kismet/Private/SBlueprintLibraryPalette.cpp
+++ b/Engine/Source/Editor/Kismet/Private/SBlueprintLibraryPalette.cpp
@@ -6,6 +6,8 @@
 #include "Editor/ClassViewer/Public/ClassViewerFilter.h" // for FClassViewerFilterFuncs
 #include "K2ActionMenuBuilder.h"
 #include "BlueprintPaletteFavorites.h"
+#include "BlueprintActionFilter.h"
+#include "BlueprintActionMenuBuilder.h"
 
 #define LOCTEXT_NAMESPACE "BlueprintLibraryPalette"
 
@@ -205,12 +207,38 @@ void SBlueprintLibraryPalette::CollectAllActions(FGraphActionListBuilderBase& Ou
 	{
 		RootCategory = TEXT("");
 	}
-
-	UBlueprint const* const Blueprint = GetBlueprint();
-	FBlueprintPaletteListBuilder PaletteBuilder(Blueprint, RootCategory);
-
-	UEdGraphSchema_K2::GetPaletteActions(PaletteBuilder, FilterClass);
-	OutAllActions.Append(PaletteBuilder);
+	
+	if (GetDefault<UEdGraphSchema_K2>()->bUseLegacyActionMenus)
+	{
+		UBlueprint const* const Blueprint = GetBlueprint();
+		FBlueprintPaletteListBuilder PaletteBuilder(Blueprint, RootCategory);
+		
+		UEdGraphSchema_K2::GetPaletteActions(PaletteBuilder, FilterClass);
+		OutAllActions.Append(PaletteBuilder);
+	}
+	else
+	{
+		uint32 FilterFlags = 0x00;
+		if (FilterClass.IsValid())
+		{
+			FilterFlags = FBlueprintActionFilter::BPFILTER_ExcludeGlobalFields;
+		}
+		
+		FBlueprintActionFilter PaletteFilter(FilterFlags);
+		PaletteFilter.Context.Blueprints.Add(GetBlueprint());
+		PaletteFilter.ExcludedNodeTypes.Add(UK2Node_Variable::StaticClass());
+		
+		if (FilterClass.IsValid())
+		{
+			PaletteFilter.OwnerClasses.Add(FilterClass.Get());
+		}
+		
+		FBlueprintActionMenuBuilder PaletteBuilder;
+		PaletteBuilder.AddMenuSection(PaletteFilter, FText::FromString(RootCategory));
+		PaletteBuilder.RebuildActionList();
+		
+		OutAllActions.Append(PaletteBuilder);
+	}
 }
 
 //------------------------------------------------------------------------------
