Changed the size of thumbnail for asset picker for content reference
- Added support for it and set it for persona toolbar

diff --git a/Engine/Source/Editor/Kismet/Private/WorkflowOrientedApp/SContentReference.cpp b/Engine/Source/Editor/Kismet/Private/WorkflowOrientedApp/SContentReference.cpp
--- a/Engine/Source/Editor/Kismet/Private/WorkflowOrientedApp/SContentReference.cpp
+++ b/Engine/Source/Editor/Kismet/Private/WorkflowOrientedApp/SContentReference.cpp
@@ -19,6 +19,7 @@ void SContentReference::Construct(const FArguments& InArgs)
 	AllowSelectingNewAsset = InArgs._AllowSelectingNewAsset;
 	AllowClearingReference = InArgs._AllowClearingReference;
 	AllowedClass = InArgs._AllowedClass;
+	AssetPickerSizeOverride = InArgs._AssetPickerSizeOverride;
 
 	// Save off delegates
 	OnShouldFilterAsset = InArgs._OnShouldFilterAsset;
@@ -226,8 +227,8 @@ TSharedRef<SWidget> SContentReference::MakeAssetPickerMenu()
 	AssetPickerConfig.ThumbnailLabel = EThumbnailLabel::ClassName;
 
 	return SNew(SBox)
-		.WidthOverride(384)
-		.HeightOverride(768)
+		.WidthOverride(AssetPickerSizeOverride.Get().X)
+		.HeightOverride(AssetPickerSizeOverride.Get().Y)
 		[
 			ContentBrowserModule.Get().CreateAssetPicker(AssetPickerConfig)
 		];
diff --git a/Engine/Source/Editor/Kismet/Public/WorkflowOrientedApp/SContentReference.h b/Engine/Source/Editor/Kismet/Public/WorkflowOrientedApp/SContentReference.h
--- a/Engine/Source/Editor/Kismet/Public/WorkflowOrientedApp/SContentReference.h
+++ b/Engine/Source/Editor/Kismet/Public/WorkflowOrientedApp/SContentReference.h
@@ -21,6 +21,7 @@ public:
 		, _AllowClearingReference(false)
 		, _AllowedClass((UClass*)NULL)
 		, _WidthOverride(FOptionalSize())
+		, _AssetPickerSizeOverride(FVector2D(384, 768))
 	{}
 		// The style of the content reference widget (optional)
 		SLATE_ARGUMENT(FName, Style)
@@ -55,6 +56,9 @@ public:
 		/** When specified, the path box will request this fixed size. */
 		SLATE_ATTRIBUTE(FOptionalSize, WidthOverride)
 
+		// when you'd like to modify size of the asset picker window
+		SLATE_ATTRIBUTE(FVector2D, AssetPickerSizeOverride)
+
 		// The event to call when the user selects a new asset or clears the existing reference
 		//@TODO: Support that stuff!
 	SLATE_END_ARGS()
@@ -91,6 +95,7 @@ protected:
 	TAttribute<bool> AllowSelectingNewAsset;
 	TAttribute<bool> AllowClearingReference;
 	TAttribute<UClass*> AllowedClass;
+	TAttribute<FVector2D>	AssetPickerSizeOverride;
 
 	// Delegates
 	FOnShouldFilterAsset OnShouldFilterAsset;
diff --git a/Engine/Source/Editor/Persona/Private/SPersonaToolbar.cpp b/Engine/Source/Editor/Persona/Private/SPersonaToolbar.cpp
--- a/Engine/Source/Editor/Persona/Private/SPersonaToolbar.cpp
+++ b/Engine/Source/Editor/Persona/Private/SPersonaToolbar.cpp
@@ -15,6 +15,7 @@
 #define LOCTEXT_NAMESPACE "PersonaToolbar"
 
 const float ContentRefWidth = 80.0f;
+const FVector2D AssetPickerSize(200, 700);
 
 //////////////////////////////////////////////////////////////////////////
 // SPipelineSeparator
@@ -126,6 +127,7 @@ void FPersonaToolbar::FillPersonaModeToolbar(FToolBarBuilder& ParentToolbarBuild
 				.AllowedClass(USkeletalMesh::StaticClass())
 				.OnShouldFilterAsset(this, &FPersonaToolbar::ShouldFilterAssetBasedOnSkeleton)
 				.OnSetReference(this, &FPersonaToolbar::OnSetSkeletalMeshReference)
+				.AssetPickerSizeOverride(AssetPickerSize)
 			]
 		);
 
@@ -171,6 +173,7 @@ void FPersonaToolbar::FillPersonaModeToolbar(FToolBarBuilder& ParentToolbarBuild
 				.AllowedClass(UAnimationAsset::StaticClass())
 				.OnShouldFilterAsset(this, &FPersonaToolbar::ShouldFilterAssetBasedOnSkeleton)
 				.OnSetReference(this, &FPersonaToolbar::OnSetAnimationAssetReference)
+				.AssetPickerSizeOverride(AssetPickerSize)
 			]
 		);
 
