#UE4 removing this intialization logic until I can find a clean way to bind File=>Merge to create a dockable window..

diff --git a/Engine/Source/Editor/Kismet/Private/BlueprintEditor.cpp b/Engine/Source/Editor/Kismet/Private/BlueprintEditor.cpp
--- a/Engine/Source/Editor/Kismet/Private/BlueprintEditor.cpp
+++ b/Engine/Source/Editor/Kismet/Private/BlueprintEditor.cpp
@@ -993,10 +993,6 @@ void FBlueprintEditor::CommonInitialization(const TArray<UBlueprint*>& InitBluep
 	// Make sure we know when tabs become active to update details tab
 	FGlobalTabmanager::Get()->OnActiveTabChanged_Subscribe( FOnActiveTabChanged::FDelegate::CreateRaw(this, &FBlueprintEditor::OnActiveTabChanged) );
 
-	FGlobalTabmanager::Get()->RegisterNomadTabSpawner(MergeToolTabId, FOnSpawnTab::CreateRaw(this, &FBlueprintEditor::CreateMergeToolTab))
-		.SetDisplayName(NSLOCTEXT("MergeTool", "TabTitle", "Merge Tool"))
-		.SetTooltipText(NSLOCTEXT("MergeTool", "TooltipText", "Open the Blueprint Merge Tool."));
-
 	if (InitBlueprints.Num() == 1)
 	{
 		// Load blueprint libraries
@@ -1510,7 +1506,6 @@ FBlueprintEditor::~FBlueprintEditor()
 	}
 
 	FGlobalTabmanager::Get()->OnActiveTabChanged_Unsubscribe( FOnActiveTabChanged::FDelegate::CreateRaw(this, &FBlueprintEditor::OnActiveTabChanged) );
-	FGlobalTabmanager::Get()->UnregisterNomadTabSpawner(MergeToolTabId);
 
 	if (FEngineAnalytics::IsAvailable())
 	{
