- Changed to list view by default
- added back extensibility menu and removed adding our menu to it.

diff --git a/Engine/Source/Editor/Kismet/Private/WorkflowOrientedApp/SContentReference.cpp b/Engine/Source/Editor/Kismet/Private/WorkflowOrientedApp/SContentReference.cpp
--- a/Engine/Source/Editor/Kismet/Private/WorkflowOrientedApp/SContentReference.cpp
+++ b/Engine/Source/Editor/Kismet/Private/WorkflowOrientedApp/SContentReference.cpp
@@ -20,6 +20,7 @@ void SContentReference::Construct(const FArguments& InArgs)
 	AllowClearingReference = InArgs._AllowClearingReference;
 	AllowedClass = InArgs._AllowedClass;
 	AssetPickerSizeOverride = InArgs._AssetPickerSizeOverride;
+	InitialAssetViewType = InArgs._InitialAssetViewType;
 
 	// Save off delegates
 	OnShouldFilterAsset = InArgs._OnShouldFilterAsset;
@@ -225,6 +226,7 @@ TSharedRef<SWidget> SContentReference::MakeAssetPickerMenu()
 	AssetPickerConfig.OnShouldFilterAsset = OnShouldFilterAsset;
 	AssetPickerConfig.bAllowNullSelection = true;
 	AssetPickerConfig.ThumbnailLabel = EThumbnailLabel::ClassName;
+	AssetPickerConfig.InitialAssetViewType = InitialAssetViewType;
 
 	return SNew(SBox)
 		.WidthOverride(AssetPickerSizeOverride.Get().X)
diff --git a/Engine/Source/Editor/Kismet/Public/WorkflowOrientedApp/SContentReference.h b/Engine/Source/Editor/Kismet/Public/WorkflowOrientedApp/SContentReference.h
--- a/Engine/Source/Editor/Kismet/Public/WorkflowOrientedApp/SContentReference.h
+++ b/Engine/Source/Editor/Kismet/Public/WorkflowOrientedApp/SContentReference.h
@@ -3,6 +3,7 @@
 #pragma once
 
 #include "Editor/ContentBrowser/Public/ContentBrowserDelegates.h"
+#include "Editor/ContentBrowser/Public/IContentBrowserSingleton.h"
 
 // This widget is a content reference, which can optionally show buttons for interacting with the reference
 class KISMET_API SContentReference : public SCompoundWidget
@@ -22,6 +23,7 @@ public:
 		, _AllowedClass((UClass*)NULL)
 		, _WidthOverride(FOptionalSize())
 		, _AssetPickerSizeOverride(FVector2D(384, 768))
+		, _InitialAssetViewType(EAssetViewType::Tile)
 	{}
 		// The style of the content reference widget (optional)
 		SLATE_ARGUMENT(FName, Style)
@@ -59,6 +61,8 @@ public:
 		// when you'd like to modify size of the asset picker window
 		SLATE_ATTRIBUTE(FVector2D, AssetPickerSizeOverride)
 
+		// initial view type
+		SLATE_ARGUMENT(EAssetViewType::Type, InitialAssetViewType)
 		// The event to call when the user selects a new asset or clears the existing reference
 		//@TODO: Support that stuff!
 	SLATE_END_ARGS()
@@ -95,7 +99,10 @@ protected:
 	TAttribute<bool> AllowSelectingNewAsset;
 	TAttribute<bool> AllowClearingReference;
 	TAttribute<UClass*> AllowedClass;
+
+	// Asset picker configuration
 	TAttribute<FVector2D>	AssetPickerSizeOverride;
+	EAssetViewType::Type	InitialAssetViewType;
 
 	// Delegates
 	FOnShouldFilterAsset OnShouldFilterAsset;
diff --git a/Engine/Source/Editor/Persona/Private/Persona.cpp b/Engine/Source/Editor/Persona/Private/Persona.cpp
--- a/Engine/Source/Editor/Persona/Private/Persona.cpp
+++ b/Engine/Source/Editor/Persona/Private/Persona.cpp
@@ -453,8 +453,6 @@ FPersona::~FPersona()
 	FEditorDelegates::OnAssetPostImport.RemoveAll(this);
 	FReimportManager::Instance()->OnPostReimport().RemoveAll(this);
 
-	FPersonaModule* PersonaModule = &FModuleManager::LoadModuleChecked<FPersonaModule>( "Persona" );
-	PersonaModule->GetMenuExtensibilityManager()->RemoveExtender(PersonaMenuExtender);
 
 	if(TargetSkeleton)
 	{
@@ -858,9 +856,8 @@ void FPersona::InitPersona(const EToolkitMode::Type Mode, const TSharedPtr< clas
 			FMenuExtensionDelegate::CreateStatic( &Local::AddPersonaAssetMenu )
 			);
 					
-	PersonaModule->GetMenuExtensibilityManager()->AddExtender(PersonaMenuExtender);					
-
 	AddMenuExtender(PersonaMenuExtender);
+	AddMenuExtender(PersonaModule->GetMenuExtensibilityManager()->GetAllExtenders(GetToolkitCommands(), GetEditingObjects()));
 
 	check(PreviewComponent == NULL);
 
diff --git a/Engine/Source/Editor/Persona/Private/SPersonaToolbar.cpp b/Engine/Source/Editor/Persona/Private/SPersonaToolbar.cpp
--- a/Engine/Source/Editor/Persona/Private/SPersonaToolbar.cpp
+++ b/Engine/Source/Editor/Persona/Private/SPersonaToolbar.cpp
@@ -15,7 +15,7 @@
 #define LOCTEXT_NAMESPACE "PersonaToolbar"
 
 const float ContentRefWidth = 80.0f;
-const FVector2D AssetPickerSize(200, 700);
+const FVector2D AssetPickerSize(250, 700);
 
 //////////////////////////////////////////////////////////////////////////
 // SPipelineSeparator
@@ -128,6 +128,7 @@ void FPersonaToolbar::FillPersonaModeToolbar(FToolBarBuilder& ParentToolbarBuild
 				.OnShouldFilterAsset(this, &FPersonaToolbar::ShouldFilterAssetBasedOnSkeleton)
 				.OnSetReference(this, &FPersonaToolbar::OnSetSkeletalMeshReference)
 				.AssetPickerSizeOverride(AssetPickerSize)
+				.InitialAssetViewType(EAssetViewType::List)
 			]
 		);
 
@@ -174,6 +175,7 @@ void FPersonaToolbar::FillPersonaModeToolbar(FToolBarBuilder& ParentToolbarBuild
 				.OnShouldFilterAsset(this, &FPersonaToolbar::ShouldFilterAssetBasedOnSkeleton)
 				.OnSetReference(this, &FPersonaToolbar::OnSetAnimationAssetReference)
 				.AssetPickerSizeOverride(AssetPickerSize)
+				.InitialAssetViewType(EAssetViewType::List)
 			]
 		);
 
