#UE4 SBlueprintDiff no longer resets itself when we PIE or compile a blueprint that is open in the blueprint diff tool, the crash this reset was originally introduced for is no longer reproducable (likely fix: 1895532

diff --git a/Engine/Source/Developer/AssetTools/Private/AssetTypeActions/AssetTypeActions_Blueprint.cpp b/Engine/Source/Developer/AssetTools/Private/AssetTypeActions/AssetTypeActions_Blueprint.cpp
--- a/Engine/Source/Developer/AssetTools/Private/AssetTypeActions/AssetTypeActions_Blueprint.cpp
+++ b/Engine/Source/Developer/AssetTools/Private/AssetTypeActions/AssetTypeActions_Blueprint.cpp
@@ -251,7 +251,7 @@ UThumbnailInfo* FAssetTypeActions_Blueprint::GetThumbnailInfo(UObject* Asset) co
 	return ThumbnailInfo;
 }
 
-void FAssetTypeActions_Blueprint::OpenInDefaults( class UBlueprint* OldBlueprint, class UBlueprint* NewBlueprint ) const
+void FAssetTypeActions_Blueprint::OpenInDefaults( const UBlueprint* OldBlueprint, const UBlueprint* NewBlueprint ) const
 {
 	const bool bComparedBlueprintsHaveGeneratedClasses = *(OldBlueprint->GeneratedClass) && *(NewBlueprint->GeneratedClass);
 	ensure(bComparedBlueprintsHaveGeneratedClasses);
diff --git a/Engine/Source/Developer/AssetTools/Private/AssetTypeActions/AssetTypeActions_Blueprint.h b/Engine/Source/Developer/AssetTools/Private/AssetTypeActions/AssetTypeActions_Blueprint.h
--- a/Engine/Source/Developer/AssetTools/Private/AssetTypeActions/AssetTypeActions_Blueprint.h
+++ b/Engine/Source/Developer/AssetTools/Private/AssetTypeActions/AssetTypeActions_Blueprint.h
@@ -27,7 +27,7 @@ protected:
 private:
 
 	/* Called to open the blueprint defaults view, this opens whatever text diff tool the user has */
-	void OpenInDefaults(class UBlueprint* OldBlueprint, class UBlueprint* NewBlueprint ) const;
+	void OpenInDefaults(const class UBlueprint* OldBlueprint, const class UBlueprint* NewBlueprint ) const;
 
 	/** Handler for when EditDefaults is selected */
 	void ExecuteEditDefaults(TArray<TWeakObjectPtr<UBlueprint>> Objects);
diff --git a/Engine/Source/Developer/AssetTools/Private/AssetTypeActions/AssetTypeActions_BlueprintGeneratedClass.cpp b/Engine/Source/Developer/AssetTools/Private/AssetTypeActions/AssetTypeActions_BlueprintGeneratedClass.cpp
--- a/Engine/Source/Developer/AssetTools/Private/AssetTypeActions/AssetTypeActions_BlueprintGeneratedClass.cpp
+++ b/Engine/Source/Developer/AssetTools/Private/AssetTypeActions/AssetTypeActions_BlueprintGeneratedClass.cpp
@@ -232,7 +232,7 @@ UThumbnailInfo* FAssetTypeActions_BlueprintGeneratedClass::GetThumbnailInfo(UObj
 	return ThumbnailInfo;
 }
 
-void FAssetTypeActions_BlueprintGeneratedClass::OpenInDefaults( class UBlueprint* OldBlueprint, class UBlueprint* NewBlueprint ) const
+void FAssetTypeActions_BlueprintGeneratedClass::OpenInDefaults( const UBlueprint* OldBlueprint, const UBlueprint* NewBlueprint ) const
 {
 	FString OldTextFilename = DumpAssetToTempFile(OldBlueprint->GeneratedClass->GetDefaultObject());
 	FString NewTextFilename = DumpAssetToTempFile(NewBlueprint->GeneratedClass->GetDefaultObject());
diff --git a/Engine/Source/Developer/AssetTools/Private/AssetTypeActions/AssetTypeActions_BlueprintGeneratedClass.h b/Engine/Source/Developer/AssetTools/Private/AssetTypeActions/AssetTypeActions_BlueprintGeneratedClass.h
--- a/Engine/Source/Developer/AssetTools/Private/AssetTypeActions/AssetTypeActions_BlueprintGeneratedClass.h
+++ b/Engine/Source/Developer/AssetTools/Private/AssetTypeActions/AssetTypeActions_BlueprintGeneratedClass.h
@@ -23,7 +23,7 @@ protected:
 private:
 
 	/* Called to open the blueprint defaults view, this opens whatever text dif tool the user has */
-	void OpenInDefaults(class UBlueprint* OldBlueprint, class UBlueprint* NewBlueprint ) const;
+	void OpenInDefaults(const class UBlueprint* OldBlueprint, const class UBlueprint* NewBlueprint ) const;
 
 	/** Handler for when EditDefaults is selected */
 	void ExecuteEditDefaults(TArray<TWeakObjectPtr<UBlueprintGeneratedClass>> Objects);
diff --git a/Engine/Source/Editor/Kismet/Private/SBlueprintDiff.cpp b/Engine/Source/Editor/Kismet/Private/SBlueprintDiff.cpp
--- a/Engine/Source/Editor/Kismet/Private/SBlueprintDiff.cpp
+++ b/Engine/Source/Editor/Kismet/Private/SBlueprintDiff.cpp
@@ -361,9 +361,6 @@ void SBlueprintDiff::Construct( const FArguments& InArgs)
 	TArray<UEdGraph*> GraphsOld,GraphsNew;
 	PanelOld.Blueprint->GetAllGraphs(GraphsOld);
 	PanelNew.Blueprint->GetAllGraphs(GraphsNew);
-
-	// Add handler for blueprint changing.
-	PanelNew.Blueprint->OnChanged().AddSP( this, &SBlueprintDiff::OnBlueprintChanged );
 	
 	//Add Graphs that exist in both blueprints, or in blueprint 1 only
 	for(auto It(GraphsOld.CreateConstIterator());It;It++)
@@ -840,18 +837,5 @@ void SBlueprintDiff::DisablePinDiffFocus()
 	}
 }
 
-void SBlueprintDiff::OnBlueprintChanged(UBlueprint* InBlueprint)
-{
-	DisablePinDiffFocus();
-
- 	TArray<UEdGraph*> GraphsOld,GraphsNew;
- 	PanelOld.Blueprint->GetAllGraphs(GraphsOld);
- 	PanelNew.Blueprint->GetAllGraphs(GraphsNew);
- 	PanelOld.GeneratePanel(GraphsOld[0],GraphsNew[0]);
-  	PanelNew.GeneratePanel(GraphsNew[0], GraphsOld[0]);
- 
- 	ResetGraphEditors();
-}
-
 #undef LOCTEXT_NAMESPACE
 
diff --git a/Engine/Source/Editor/Kismet/Public/SBlueprintDiff.h b/Engine/Source/Editor/Kismet/Public/SBlueprintDiff.h
--- a/Engine/Source/Editor/Kismet/Public/SBlueprintDiff.h
+++ b/Engine/Source/Editor/Kismet/Public/SBlueprintDiff.h
@@ -123,7 +123,7 @@ struct KISMET_API FDiffPanel
 	bool CanCopyNodes() const;
 
 	/*The blueprint that owns the graph we are showing*/
-	class UBlueprint*				Blueprint;
+	const class UBlueprint*				Blueprint;
 
 	/*The border around the graph editor, used to change the content when new graphs are set */
 	TSharedPtr<SBorder>				GraphEditorBorder;
@@ -147,7 +147,7 @@ class  KISMET_API SBlueprintDiff: public SCompoundWidget
 {
 
 public:
-	DECLARE_DELEGATE_TwoParams( FOpenInDefaults, class UBlueprint* , class UBlueprint* );
+	DECLARE_DELEGATE_TwoParams( FOpenInDefaults, const class UBlueprint* , const class UBlueprint* );
 
 	SLATE_BEGIN_ARGS( SBlueprintDiff ){}
 			SLATE_ARGUMENT( class UBlueprint*, BlueprintOld )
@@ -184,9 +184,6 @@ protected:
 	virtual void OnSelectionChanged(FGraphToDiff Item, ESelectInfo::Type SelectionType);
 
 	void OnDiffListSelectionChanged(const TSharedPtr<struct FDiffResultItem>& TheDiff, FListItemGraphToDiff* GraphDiffer);
-
-	/* Handler for blueprint changing. */
-	void OnBlueprintChanged(UBlueprint* InBlueprint);
 		
 	/** Disable the focus on a particular pin */
 	void DisablePinDiffFocus();
