(WIP) Handling consolidated property menu actions for the palette in the new blueprint action system (adding drag-n-drop functionality).

diff --git a/Engine/Source/Editor/BlueprintGraph/Classes/BlueprintPropertyNodeSpawner.h b/Engine/Source/Editor/BlueprintGraph/Classes/BlueprintPropertyNodeSpawner.h
--- a/Engine/Source/Editor/BlueprintGraph/Classes/BlueprintPropertyNodeSpawner.h
+++ b/Engine/Source/Editor/BlueprintGraph/Classes/BlueprintPropertyNodeSpawner.h
@@ -26,7 +26,7 @@ public:
 	 * @param  NodeClass	The type of node you want spawned (can't distinct a getter from a setter from the property alone).
 	 * @return A newly allocated instance of this class.
 	 */
-	static UBlueprintPropertyNodeSpawner* Create(UProperty const* const Property, UClass* const NodeClass = nullptr);
+	BLUEPRINTGRAPH_API static UBlueprintPropertyNodeSpawner* Create(UProperty const* const Property, UClass* const NodeClass = nullptr);
 
 	// UBlueprintNodeSpawner interface
 	virtual UEdGraphNode* Invoke(UEdGraph* ParentGraph) const override;
@@ -39,14 +39,14 @@ public:
 	 *
 	 * @return The property that this class was initialized with.
 	 */
-	UProperty const* GetProperty() const;
+	BLUEPRINTGRAPH_API UProperty const* GetProperty() const;
 
 	/**
 	 * Checks to see if the associated property is a delegate property.
 	 *
 	 * @return True if the property is a UMulticastDelegateProperty (otherwise false).
 	 */
-	bool IsDelegateProperty() const;
+	BLUEPRINTGRAPH_API bool IsDelegateProperty() const;
 	
 private:
 	/** The property to configure new nodes with. */
diff --git a/Engine/Source/Editor/BlueprintGraph/Private/BlueprintActionDatabase.cpp b/Engine/Source/Editor/BlueprintGraph/Private/BlueprintActionDatabase.cpp
--- a/Engine/Source/Editor/BlueprintGraph/Private/BlueprintActionDatabase.cpp
+++ b/Engine/Source/Editor/BlueprintGraph/Private/BlueprintActionDatabase.cpp
@@ -27,6 +27,10 @@
 // used below in FBlueprintActionDatabaseImpl::AddClassPropertyActions()
 #include "K2Node_VariableGet.h"
 #include "K2Node_VariableSet.h"
+#include "K2Node_AddDelegate.h"
+#include "K2Node_CallDelegate.h"
+#include "K2Node_RemoveDelegate.h"
+#include "K2Node_ClearDelegate.h"
 // used below in FBlueprintActionDatabaseImpl::AddClassCastActions()
 #include "K2Node_DynamicCast.h"
 #include "K2Node_ClassDynamicCast.h"
@@ -92,6 +96,17 @@ namespace FBlueprintNodeSpawnerFactory
 	 * @return A new node-spawner, setup to spawn a UEdGraphNode_Comment.
 	 */
 	static UBlueprintNodeSpawner* MakeCommentNodeSpawner();
+	
+	/**
+	 * A templatized method which constructs a node-spawner for various delegate
+	 * node types (any node with a UMulticastDelegateProperty). Takes the
+	 * specified property and assigns it to the node post-spawn.
+	 *
+	 * @param  Property	 The delegate you want set for the spawned node.
+	 * @return A new node-spawner, setup to spawn some delegate node (defined by the DelegateNodeType template param).
+	 */
+	template<class DelegateNodeType>
+	static UBlueprintNodeSpawner* MakeDelegateNodeSpawner(UMulticastDelegateProperty* Property);
 };
 
 //------------------------------------------------------------------------------
@@ -226,6 +241,28 @@ static UBlueprintNodeSpawner* FBlueprintNodeSpawnerFactory::MakeCommentNodeSpawn
 	return NodeSpawner;
 }
 
+//------------------------------------------------------------------------------
+template<class DelegateNodeType>
+static UBlueprintNodeSpawner* FBlueprintNodeSpawnerFactory::MakeDelegateNodeSpawner(UMulticastDelegateProperty* Property)
+{
+	UBlueprintNodeSpawner* NodeSpawner = UBlueprintPropertyNodeSpawner::Create(Property, DelegateNodeType::StaticClass());
+	check(NodeSpawner != nullptr);
+	
+	auto CustomizeDelegateNodeLambda = [](UEdGraphNode* NewNode, bool bIsTemplateNode, UMulticastDelegateProperty* Property)
+	{
+		DelegateNodeType* DelegateNode = CastChecked<DelegateNodeType>(NewNode);
+		
+		UBlueprint* Blueprint = DelegateNode->GetBlueprint();
+		check(Blueprint != nullptr);
+		bool const bIsSelfContext = Blueprint->GeneratedClass->IsChildOf(Property->GetOwnerClass());
+		
+		DelegateNode->SetFromProperty(Property, bIsSelfContext);
+	};
+	NodeSpawner->CustomizeNodeDelegate = UBlueprintNodeSpawner::FCustomizeNodeDelegate::CreateStatic(CustomizeDelegateNodeLambda, Property);
+	
+	return NodeSpawner;
+}
+
 /*******************************************************************************
  * Static FBlueprintActionDatabase Helpers
  ******************************************************************************/
@@ -397,6 +434,8 @@ static void FBlueprintActionDatabaseImpl::AddClassFunctionActions(UClass const*
 //------------------------------------------------------------------------------
 static void FBlueprintActionDatabaseImpl::AddClassPropertyActions(UClass const* const Class, FActionList& ActionListOut)
 {
+	using namespace FBlueprintNodeSpawnerFactory; // for MakeDelegateNodeSpawner()
+	
 	// loop over all the properties in the specified class; exclude-super because 
 	// we can always get the super properties by looking up that class separately 
 	for (TFieldIterator<UProperty> PropertyIt(Class, EFieldIteratorFlags::ExcludeSuper); PropertyIt; ++PropertyIt)
@@ -407,25 +446,28 @@ static void FBlueprintActionDatabaseImpl::AddClassPropertyActions(UClass const*
 			continue;
 		}
 
-		// add an empty node spawner (that actually won't spawn anything) to 
-		// represent the property as a whole (used for the "collapsed" property
-		// entries in the palette; the ones that don't immediatly spawn a node, 
-		// but instead a sub-menu when you drop them in the graph)
-		//
-		// @TODO: consider moving this out into FBlueprintActionMenuBuilder
-		if (UBlueprintPropertyNodeSpawner* NodeSpawner = UBlueprintPropertyNodeSpawner::Create(Property))
-		{
-			ActionListOut.Add(NodeSpawner);
-		}
-
  		bool bIsDelegate = Property->IsA(UMulticastDelegateProperty::StaticClass());
  		if (bIsDelegate)
  		{
-			// @TODO: GetEventDispatcherNodesForClass()
-			//   UK2Node_AddDelegate
-			//   UK2Node_RemoveDelegate
-			//   UK2Node_ClearDelegate
-			//   UK2Node_CallDelegate
+			UMulticastDelegateProperty* DelegateProperty = CastChecked<UMulticastDelegateProperty>(Property);
+			if (DelegateProperty->HasAnyPropertyFlags(CPF_BlueprintAssignable))
+			{
+				UBlueprintNodeSpawner* AddSpawner = MakeDelegateNodeSpawner<UK2Node_AddDelegate>(DelegateProperty);
+				ActionListOut.Add(AddSpawner);
+				
+				// @TODO: account for: GetEventDispatcherNodesForClass() - FEdGraphSchemaAction_K2AssignDelegate
+			}
+			
+			if (DelegateProperty->HasAnyPropertyFlags(CPF_BlueprintCallable))
+			{
+				UBlueprintNodeSpawner* CallSpawner = MakeDelegateNodeSpawner<UK2Node_CallDelegate>(DelegateProperty);
+				ActionListOut.Add(CallSpawner);
+			}
+			
+			UBlueprintNodeSpawner* RemoveSpawner = MakeDelegateNodeSpawner<UK2Node_RemoveDelegate>(DelegateProperty);
+			ActionListOut.Add(RemoveSpawner);
+			UBlueprintNodeSpawner* ClearSpawner  = MakeDelegateNodeSpawner<UK2Node_ClearDelegate>(DelegateProperty);
+			ActionListOut.Add(ClearSpawner);
 
 			// @TODO: AddBoundEventActionsForClass()
 			//   UK2Node_ComponentBoundEvent 
diff --git a/Engine/Source/Editor/BlueprintGraph/Private/BlueprintActionMenuBuilder.cpp b/Engine/Source/Editor/BlueprintGraph/Private/BlueprintActionMenuBuilder.cpp
deleted file mode 100644
diff --git a/Engine/Source/Editor/BlueprintGraph/Private/BlueprintActionMenuItem.cpp b/Engine/Source/Editor/BlueprintGraph/Private/BlueprintActionMenuItem.cpp
deleted file mode 100644
diff --git a/Engine/Source/Editor/BlueprintGraph/Private/BlueprintEventNodeSpawner.cpp b/Engine/Source/Editor/BlueprintGraph/Private/BlueprintEventNodeSpawner.cpp
--- a/Engine/Source/Editor/BlueprintGraph/Private/BlueprintEventNodeSpawner.cpp
+++ b/Engine/Source/Editor/BlueprintGraph/Private/BlueprintEventNodeSpawner.cpp
@@ -2,8 +2,7 @@
 
 #include "BlueprintGraphPrivatePCH.h"
 #include "BlueprintEventNodeSpawner.h"
-#include "EdGraphSchema_K2.h"			// for GetFriendlySignitureName()
-#include "BlueprintActionMenuBuilder.h"	// for AddEventCategory
+#include "EdGraphSchema_K2.h" // for GetFriendlySignitureName()
 
 #define LOCTEXT_NAMESPACE "BlueprintEventNodeSpawner"
 
diff --git a/Engine/Source/Editor/BlueprintGraph/Private/BlueprintPropertyNodeSpawner.cpp b/Engine/Source/Editor/BlueprintGraph/Private/BlueprintPropertyNodeSpawner.cpp
--- a/Engine/Source/Editor/BlueprintGraph/Private/BlueprintPropertyNodeSpawner.cpp
+++ b/Engine/Source/Editor/BlueprintGraph/Private/BlueprintPropertyNodeSpawner.cpp
@@ -63,11 +63,15 @@ FText UBlueprintPropertyNodeSpawner::GetDefaultMenuName() const
 {
 	check(Property != nullptr);
 	
-	FText MenuName = FText::FromName(Property->GetFName());
+	FText MenuName;
 	if (GetDefault<UEditorStyleSettings>()->bShowFriendlyNames)
 	{
 		MenuName = FText::FromString(UEditorEngine::GetFriendlyName(Property));
 	}
+	else
+	{
+		MenuName = FText::FromName(Property->GetFName());
+	}
 
 	if (NodeClass != nullptr)
 	{
diff --git a/Engine/Source/Editor/BlueprintGraph/Private/K2Node_Self.cpp b/Engine/Source/Editor/BlueprintGraph/Private/K2Node_Self.cpp
--- a/Engine/Source/Editor/BlueprintGraph/Private/K2Node_Self.cpp
+++ b/Engine/Source/Editor/BlueprintGraph/Private/K2Node_Self.cpp
@@ -4,7 +4,7 @@
 #include "BlueprintGraphPrivatePCH.h"
 #include "KismetCompiler.h"
 #include "BlueprintNodeSpawner.h"
-#include "BlueprintActionMenuBuilder.h"
+#include "EditorCategoryUtils.h"
 
 #define LOCTEXT_NAMESPACE "K2Node_Self"
 
@@ -93,7 +93,7 @@ void UK2Node_Self::GetMenuActions(TArray<UBlueprintNodeSpawner*>& ActionListOut)
 
 FText UK2Node_Self::GetMenuCategory() const
 {
-	return FBlueprintActionMenuBuilder::VariablesCategory;
+	return FEditorCategoryUtils::GetCommonCategory(FCommonEditorCategory::Variables);
 }
 
 #undef LOCTEXT_NAMESPACE
diff --git a/Engine/Source/Editor/BlueprintGraph/Private/K2Node_VariableSetRef.cpp b/Engine/Source/Editor/BlueprintGraph/Private/K2Node_VariableSetRef.cpp
--- a/Engine/Source/Editor/BlueprintGraph/Private/K2Node_VariableSetRef.cpp
+++ b/Engine/Source/Editor/BlueprintGraph/Private/K2Node_VariableSetRef.cpp
@@ -6,7 +6,7 @@
 #include "VariableSetHandler.h"
 #include "K2Node_VariableSetRef.h"
 #include "BlueprintNodeSpawner.h"
-#include "BlueprintActionMenuBuilder.h" // for VariablesCategory
+#include "EditorCategoryUtils.h"
 
 static FString TargetVarPinName(TEXT("Target"));
 static FString VarValuePinName(TEXT("Value"));
@@ -214,7 +214,7 @@ void UK2Node_VariableSetRef::GetMenuActions(TArray<UBlueprintNodeSpawner*>& Acti
 
 FText UK2Node_VariableSetRef::GetMenuCategory() const
 {
-	return FBlueprintActionMenuBuilder::VariablesCategory;
+	return FEditorCategoryUtils::GetCommonCategory(FCommonEditorCategory::Variables);
 }
 
 
diff --git a/Engine/Source/Editor/BlueprintGraph/Public/BlueprintActionFilter.h b/Engine/Source/Editor/BlueprintGraph/Public/BlueprintActionFilter.h
--- a/Engine/Source/Editor/BlueprintGraph/Public/BlueprintActionFilter.h
+++ b/Engine/Source/Editor/BlueprintGraph/Public/BlueprintActionFilter.h
@@ -117,8 +117,8 @@ public:
 	 * Appends another filter to be utilized in IsFiltered() queries, extending  
 	 * the query to be: IsFilteredByThis() || Rhs.IsFiltered()
 	 *
-	 * NOTE: Appending additional filters can hinder performance (as the 
-	 *       they have to loop through all their tests on top of these)
+	 * NOTE: Appending additional filters can hinder performance (as the
+	 *       additional filter has to loop over its tests foreach database entry)
 	 * 
 	 * @param  Rhs	The filter to append to this one.
 	 * @return This.
@@ -129,8 +129,8 @@ public:
 	 * Appends another filter to be utilized in IsFiltered() queries, extending  
 	 * the query to be: IsFilteredByThis() && Rhs.IsFiltered()
 	 *
-	 * NOTE: Appending additional filters can hinder performance (as the 
-	 *       they have to loop through all their tests on top of these)
+	 * NOTE: Appending additional filters can hinder performance (as the
+	 *       additional filter has to loop over its tests foreach database entry)
 	 * 
 	 * @param  Rhs	The filter to append to this one.
 	 * @return This.
diff --git a/Engine/Source/Editor/BlueprintGraph/Public/BlueprintActionMenuBuilder.h b/Engine/Source/Editor/BlueprintGraph/Public/BlueprintActionMenuBuilder.h
deleted file mode 100644
diff --git a/Engine/Source/Editor/BlueprintGraph/Public/BlueprintActionMenuItem.h b/Engine/Source/Editor/BlueprintGraph/Public/BlueprintActionMenuItem.h
deleted file mode 100644
diff --git a/Engine/Source/Editor/Kismet/Private/BPDelegateDragDropAction.cpp b/Engine/Source/Editor/Kismet/Private/BPDelegateDragDropAction.cpp
--- a/Engine/Source/Editor/Kismet/Private/BPDelegateDragDropAction.cpp
+++ b/Engine/Source/Editor/Kismet/Private/BPDelegateDragDropAction.cpp
@@ -10,12 +10,6 @@
 
 #define LOCTEXT_NAMESPACE "DelegateDragDropAction"
 
-FKismetDelegateDragDropAction::FKismetDelegateDragDropAction( const TSharedRef< SWidget >& InOwner ) 
-	: FKismetVariableDragDropAction()
-	, Owner( InOwner )
-{
-}
-
 void FKismetDelegateDragDropAction::MakeEvent(FNodeConstructionParams Params)
 {
 	check(Params.Graph && Params.Property);
@@ -136,7 +130,7 @@ FReply FKismetDelegateDragDropAction::DroppedOnPanel(const TSharedRef< SWidget >
 			}
 		}
 		MenuBuilder.EndSection();
-		FSlateApplication::Get().PushMenu(Owner, MenuBuilder.MakeWidget(), ScreenPosition, FPopupTransitionEffect( FPopupTransitionEffect::ContextMenu));
+		FSlateApplication::Get().PushMenu(Panel, MenuBuilder.MakeWidget(), ScreenPosition, FPopupTransitionEffect( FPopupTransitionEffect::ContextMenu));
 	}
 	return FReply::Handled();
 }
diff --git a/Engine/Source/Editor/Kismet/Private/BPDelegateDragDropAction.h b/Engine/Source/Editor/Kismet/Private/BPDelegateDragDropAction.h
--- a/Engine/Source/Editor/Kismet/Private/BPDelegateDragDropAction.h
+++ b/Engine/Source/Editor/Kismet/Private/BPDelegateDragDropAction.h
@@ -21,9 +21,9 @@ public:
 
 	bool IsValid() const;
 	
-	static TSharedRef<FKismetDelegateDragDropAction> New( TSharedRef< SWidget > Owner, FName InVariableName, UStruct* InSource, FNodeCreationAnalytic AnalyticCallback)
+	static TSharedRef<FKismetDelegateDragDropAction> New( FName InVariableName, UStruct* InSource, FNodeCreationAnalytic AnalyticCallback)
 	{
-		TSharedRef<FKismetDelegateDragDropAction> Operation = MakeShareable(new FKismetDelegateDragDropAction( Owner ) );
+		TSharedRef<FKismetDelegateDragDropAction> Operation = MakeShareable(new FKismetDelegateDragDropAction);
 		Operation->VariableName = InVariableName;
 		Operation->VariableSource = InSource;
 		Operation->AnalyticCallback = AnalyticCallback;
@@ -57,9 +57,5 @@ public:
 	static void AssignEvent(FNodeConstructionParams Params);
 
 protected:
-	FKismetDelegateDragDropAction( const TSharedRef< SWidget >& InOwner );
-
-private:
-
-	TSharedRef< SWidget > Owner;
+	FKismetDelegateDragDropAction() {}
 };
diff --git a/Engine/Source/Editor/Kismet/Private/BlueprintActionMenuBuilder.cpp b/Engine/Source/Editor/Kismet/Private/BlueprintActionMenuBuilder.cpp
new file mode 100644
--- /dev/null
+++ b/Engine/Source/Editor/Kismet/Private/BlueprintActionMenuBuilder.cpp
@@ -0,0 +1,619 @@
+// Copyright 1998-2014 Epic Games, Inc. All Rights Reserved.
+
+#include "BlueprintEditorPrivatePCH.h"
+#include "BlueprintActionMenuBuilder.h"
+#include "BlueprintActionMenuItem.h"
+#include "BlueprintActionDatabase.h"
+#include "BlueprintActionFilter.h"
+#include "BlueprintNodeSpawner.h"
+#include "BlueprintFunctionNodeSpawner.h"
+#include "BlueprintPropertyNodeSpawner.h"
+#include "K2ActionMenuBuilder.h"		// for FBlueprintPaletteListBuilder/FBlueprintGraphActionListBuilder
+#include "EdGraphSchema_K2.h"			// for StaticClass(), bUseLegacyActionMenus, etc.
+#include "BlueprintEditor.h"			// for GetMyBlueprintWidget(), and GetIsContextSensitive()
+#include "K2ActionMenuBuilder.h"
+#include "SMyBlueprint.h"				// for SelectionAsVar()
+#include "EdGraphSchema_K2_Actions.h"	// for FEdGraphSchemaAction_K2Var
+
+#define LOCTEXT_NAMESPACE "BlueprintActionMenuBuilder"
+
+/*******************************************************************************
+ * FBlueprintActionMenuItemFactory
+ ******************************************************************************/
+
+class FBlueprintActionMenuItemFactory
+{
+public:
+	/** 
+	 * Menu item factory constructor. Sets up the blueprint context, which
+	 * is utilized when configuring blueprint menu items' names/tooltips/etc.
+	 *
+	 * @param  Context	The blueprint context for the menu being built.
+	 */
+	FBlueprintActionMenuItemFactory(FBlueprintActionContext const& Context);
+
+	/** A root category to perpend every menu item with */
+	FText RootCategory;
+	/** The menu sort order to set every menu item with */
+	int32 MenuGrouping;
+	
+	/**
+	 * Spawns a new FBlueprintActionMenuItem with the node-spawner. Constructs
+	 * the menu item's category, name, tooltip, etc.
+	 * 
+	 * @param  ItemOwner	The menu item owner that this will be assigned to.
+	 * @param  Action		The node-spawner that the new menu item should wrap.
+	 * @return A newly allocated FBlueprintActionMenuItem (which wraps the supplied action).
+	 */
+	TSharedPtr<FEdGraphSchemaAction> MakeMenuItem(FBlueprintActionMenuBuilder& ItemOwner, UBlueprintNodeSpawner* Action);
+
+	/**
+	 * Clones the property spawner and clears its NodeClass field (used to 
+	 * represent the property as a whole, encompasing all the node types that
+	 * can be associated with it, like: getters, setters, etc.
+	 *
+	 * @param  ItemOwner	The menu item owner that this will be assigned to.
+	 * @param  NodeSpawner	The node-spawner that we want to clone (and create a proxy for).
+	 * @return A cloned UBlueprintPropertyNodeSpawner (representing the property).
+	 */
+	UBlueprintPropertyNodeSpawner* MakePropertyActionProxy(FBlueprintActionMenuBuilder& ItemOwner, UBlueprintPropertyNodeSpawner* NodeSpawner);
+	
+	/**
+	 * Checks to see if a proxy has already been made for the specified property
+	 * (checks to see if someone has already called MakePropertyActionProxy()
+	 * for a spawner associated with this property).
+	 *
+	 * @param  Property		The property you want to find a proxy for.
+	 * @return True if there's a proxy already cached for the property, otherwise false.
+	 */
+	bool HasProxy(UProperty const* Property) const;
+	
+	/**
+	 * Checks to see if the supplied property node spawner is a proxy itself 
+	 * (was it spawned from MakePropertyActionProxy()? does it have an empty 
+	 * NodeClass field?)
+	 *
+	 * @param  PropertyNodeSpawner	The spawner you want to check.
+	 * @return True if the supplied node-spawner is a proxy, false otherwise.
+	 */
+	bool IsProxy(UBlueprintPropertyNodeSpawner* PropertyNodeSpawner) const;
+	
+protected:
+	/**
+	 * Attempts to pull a menu name from the supplied spawner. If one isn't 
+	 * provided, then it spawns a temporary node and pulls one from that node's 
+	 * title.
+	 * 
+	 * @param  ItemOwner	The menu builder that this name will be listed from.
+	 * @param  Action		The action you want to suss name information from.
+	 * @return A name for the menu item wrapping this action.
+	 */
+	FText GetMenuNameForAction(FBlueprintActionMenuBuilder& ItemOwner, UBlueprintNodeSpawner* Action);
+
+	/**
+	 * Attempts to pull a menu category from the supplied spawner. If one isn't 
+	 * provided, then it spawns a temporary node and pulls one from that node.
+	 * 
+	 * @param  ItemOwner	The menu builder that the item will be listed from.
+	 * @param  Action		The action you want to suss category information from.
+	 * @return A category for the menu item wrapping this action.
+	 */
+	FText GetCategoryForAction(FBlueprintActionMenuBuilder& ItemOwner, UBlueprintNodeSpawner* Action);
+
+	/**
+	 * Attempts to pull a menu tooltip from the supplied spawner. If one isn't 
+	 * provided, then it spawns a temporary node and pulls one from that node.
+	 * 
+	 * @param  ItemOwner	The menu builder that the item will be listed from.
+	 * @param  Action		The action you want to suss tooltip information from.
+	 * @return A tooltip for the menu item wrapping this action.
+	 */
+	FText GetTooltipForAction(FBlueprintActionMenuBuilder& ItemOwner, UBlueprintNodeSpawner* Action);
+
+	/**
+	 * Attempts to pull a keywords from the supplied spawner. If one isn't 
+	 * provided, then it spawns a temporary node and pulls them from that.
+	 * 
+	 * @param  ItemOwner	The menu builder that the item will be listed from.
+	 * @param  Action		The action you want to suss keyword information from.
+	 * @return A keyword text string for the menu item wrapping this action.
+	 */
+	FText GetSearchKeywordsForAction(FBlueprintActionMenuBuilder& ItemOwner, UBlueprintNodeSpawner* Action);
+
+	/**
+	 * Utility getter function that retrieves the blueprint context for the menu
+	 * items being made.
+	 * 
+	 * @return The first blueprint out of the cached FBlueprintActionContext.
+	 */
+	UBlueprint* GetBlueprint() const;
+
+	/**
+	 * Utility function that will pull the "OwnerOfTemporaries" field from 
+	 * ItemOwner. If the object is null, then this will allocate one on behalf
+	 * of ItemOwner.
+	 * 
+	 * @param  The FBlueprintActionMenuBuilder you want an "OwnerOfTemporaries" from. 
+	 * @return The FBlueprintActionMenuBuilder's OwnerOfTemporaries.
+	 */
+	UEdGraph* GetOwnerOfTemporaries(FBlueprintActionMenuBuilder& ItemOwner) const;
+
+	/** Cached context for the blueprint menu being built */
+	FBlueprintActionContext const& Context;
+	
+	/** Tracks intermediate property proxies (if any were created) */
+	TMap<UProperty const*, UBlueprintPropertyNodeSpawner*> PropertyProxyMap;
+};
+
+//------------------------------------------------------------------------------
+FBlueprintActionMenuItemFactory::FBlueprintActionMenuItemFactory(FBlueprintActionContext const& ContextIn)
+	: Context(ContextIn)
+	, RootCategory(FText::GetEmpty())
+	, MenuGrouping(0)
+{
+}
+
+//------------------------------------------------------------------------------
+TSharedPtr<FEdGraphSchemaAction> FBlueprintActionMenuItemFactory::MakeMenuItem(FBlueprintActionMenuBuilder& ItemOwner, UBlueprintNodeSpawner* Action)
+{
+	FBlueprintActionMenuItem* NewMenuItem = new FBlueprintActionMenuItem(Action, MenuGrouping);
+	
+	NewMenuItem->MenuDescription    = GetMenuNameForAction(ItemOwner, Action);
+	NewMenuItem->TooltipDescription = GetTooltipForAction(ItemOwner, Action).ToString();
+	NewMenuItem->Category           = GetCategoryForAction(ItemOwner, Action).ToString();
+	NewMenuItem->Keywords           = GetSearchKeywordsForAction(ItemOwner, Action).ToString();
+
+	if (NewMenuItem->Category.IsEmpty())
+	{
+		NewMenuItem->Category = RootCategory.ToString();
+	}
+	else if (!RootCategory.IsEmpty())
+	{
+		NewMenuItem->Category = FString::Printf(TEXT("%s|%s"), *RootCategory.ToString(), *NewMenuItem->Category);
+	}
+	
+	return MakeShareable(NewMenuItem);
+}
+
+//------------------------------------------------------------------------------
+UBlueprintPropertyNodeSpawner* FBlueprintActionMenuItemFactory::MakePropertyActionProxy(FBlueprintActionMenuBuilder& ItemOwner, UBlueprintPropertyNodeSpawner* PropertyNodeSpawner)
+{
+	UBlueprintPropertyNodeSpawner* ActionProxy = nullptr;
+	
+	UProperty const* Property = PropertyNodeSpawner->GetProperty();
+	if (UBlueprintPropertyNodeSpawner** FoundProxy = PropertyProxyMap.Find(Property))
+	{
+		ActionProxy = *FoundProxy;
+	}
+	else
+	{
+		ActionProxy = DuplicateObject(PropertyNodeSpawner, GetOwnerOfTemporaries(ItemOwner));
+		// we use a null node class to flag the proxy, FBlueprintActionMenuItem
+		// queues off of this and will provide a drag-drop menu to the user
+		ActionProxy->NodeClass = nullptr;
+		PropertyProxyMap.Add(Property, ActionProxy);
+	}
+	
+	return ActionProxy;
+}
+
+//------------------------------------------------------------------------------
+bool FBlueprintActionMenuItemFactory::HasProxy(UProperty const* Property) const
+{
+	return PropertyProxyMap.Contains(Property);
+}
+
+//------------------------------------------------------------------------------
+bool FBlueprintActionMenuItemFactory::IsProxy(UBlueprintPropertyNodeSpawner* NodeSpawner) const
+{
+	bool bIsProxy = false;
+	
+	UProperty const* Property = NodeSpawner->GetProperty();
+	if (UBlueprintPropertyNodeSpawner* const* FoundProxy = PropertyProxyMap.Find(Property))
+	{
+		bIsProxy = (*FoundProxy == NodeSpawner);
+	}
+	check(!bIsProxy || (NodeSpawner->NodeClass == nullptr));
+	
+	return bIsProxy;
+}
+
+//------------------------------------------------------------------------------
+FText FBlueprintActionMenuItemFactory::GetMenuNameForAction(FBlueprintActionMenuBuilder& ItemOwner, UBlueprintNodeSpawner* Action)
+{
+	check(Action != nullptr);
+	// give the action the chance to save on performance (to keep from having to spawn a template node)
+	FText MenuName = Action->GetDefaultMenuName();
+	
+	if (MenuName.IsEmpty())
+	{
+		if (UEdGraphNode* NodeTemplate = Action->MakeTemplateNode(GetOwnerOfTemporaries(ItemOwner)))
+		{
+			MenuName = NodeTemplate->GetNodeTitle(ENodeTitleType::ListView);
+		}
+		else
+		{
+			// need to give it some name, this is as good as any I guess
+			MenuName = FText::FromName(Action->GetFName());
+		}
+	}
+	
+	return MenuName;
+}
+
+//------------------------------------------------------------------------------
+FText FBlueprintActionMenuItemFactory::GetCategoryForAction(FBlueprintActionMenuBuilder& ItemOwner, UBlueprintNodeSpawner* Action)
+{
+	check(Action != nullptr);
+	// give the action the chance to save on performance (to keep from having to spawn a template node)
+	FText MenuCategory = Action->GetDefaultMenuCategory();
+	
+	if (MenuCategory.IsEmpty())
+	{
+		UEdGraph* OwnerOfTemporaries = GetOwnerOfTemporaries(ItemOwner);
+		// @TODO: consider moving GetMenuCategory() up into UEdGraphNode
+		if (UK2Node* NodeTemplate = Cast<UK2Node>(Action->MakeTemplateNode(OwnerOfTemporaries)))
+		{
+			MenuCategory = NodeTemplate->GetMenuCategory();
+		}
+		
+		// if the category is still empty, then try and construct the category
+		// from the spawner's type
+		if (MenuCategory.IsEmpty())
+		{
+			// put uncategorized function calls in a member function category
+			// (sorted by their respective classes)
+			if (UBlueprintFunctionNodeSpawner* FuncSpawner = Cast<UBlueprintFunctionNodeSpawner>(Action))
+			{
+				UFunction const* Function = FuncSpawner->GetFunction();
+				check(Function != nullptr);
+				UClass* FuncOwner = Function->GetOuterUClass();
+				
+				UBlueprint* Blueprint = GetBlueprint();
+				check(Blueprint != nullptr);
+				
+				// if this is NOT a self function call (self function calls
+				// don't get nested any deeper)
+				if (!Blueprint->GeneratedClass->IsChildOf(FuncOwner))
+				{
+					MenuCategory = FText::FromName(FuncOwner->GetFName());
+				}
+				MenuCategory = FText::Format(LOCTEXT("MemberFunctionsCategory", "Call Function|{0}"), MenuCategory);
+			}
+		}
+	}
+	
+	return MenuCategory;
+}
+
+//------------------------------------------------------------------------------
+FText FBlueprintActionMenuItemFactory::GetTooltipForAction(FBlueprintActionMenuBuilder& ItemOwner, UBlueprintNodeSpawner* Action)
+{
+	check(Action != nullptr);
+	// give the action the chance to save on performance (to keep from having to spawn a template node)
+	FText Tooltip = Action->GetDefaultMenuTooltip();
+	
+	if (Tooltip.IsEmpty())
+	{
+		if (UEdGraphNode* NodeTemplate = Action->MakeTemplateNode(GetOwnerOfTemporaries(ItemOwner)))
+		{
+			Tooltip = FText::FromString(NodeTemplate->GetTooltip());
+		}
+	}
+	
+	return Tooltip;
+}
+
+//------------------------------------------------------------------------------
+FText FBlueprintActionMenuItemFactory::GetSearchKeywordsForAction(FBlueprintActionMenuBuilder& ItemOwner, UBlueprintNodeSpawner* Action)
+{
+	check(Action != nullptr);
+	// give the action the chance to save on performance (to keep from having to spawn a template node)
+	FText SearchKeywords = Action->GetDefaultSearchKeywords();
+	
+	if (SearchKeywords.IsEmpty())
+	{
+		if (UEdGraphNode* NodeTemplate = Action->MakeTemplateNode(GetOwnerOfTemporaries(ItemOwner)))
+		{
+			SearchKeywords = FText::FromString(NodeTemplate->GetKeywords());
+		}
+	}
+	
+	return SearchKeywords;
+}
+
+//------------------------------------------------------------------------------
+UBlueprint* FBlueprintActionMenuItemFactory::GetBlueprint() const
+{
+	UBlueprint* Blueprint = nullptr;
+	if (Context.Blueprints.Num() > 0)
+	{
+		Blueprint = Context.Blueprints[0];
+	}
+	return Blueprint;
+}
+
+//------------------------------------------------------------------------------
+UEdGraph* FBlueprintActionMenuItemFactory::GetOwnerOfTemporaries(FBlueprintActionMenuBuilder& ItemOwner) const
+{
+	UBlueprint* Blueprint = GetBlueprint();
+	check(Blueprint != nullptr);
+	
+	UEdGraph*& OwnerOfTemporaries = ItemOwner.OwnerOfTemporaries;
+	if (OwnerOfTemporaries == nullptr)
+	{
+		// unfortunately the OwnerOfTemporaries cannont be directly assigned to
+		// the transient package because various nodes relay on the blueprint
+		// outer (just like they rely on this graph outer); for example, the
+		// call-function node uses the blueprint to see if the function has
+		// self-context
+		OwnerOfTemporaries = NewObject<UEdGraph>(Blueprint);
+		OwnerOfTemporaries->SetFlags(RF_Transient);
+		
+		OwnerOfTemporaries->Schema = UEdGraphSchema_K2::StaticClass();
+		if (Context.Graphs.Num() > 0)
+		{
+			OwnerOfTemporaries->Schema = Context.Graphs[0]->Schema;
+		}
+		
+	}
+	return OwnerOfTemporaries;
+}
+
+/*******************************************************************************
+ * Static FBlueprintActionMenuBuilder Helpers
+ ******************************************************************************/
+
+namespace FBlueprintActionMenuBuilderImpl
+{
+	/** 
+	 * Defines a sub-section of the overall blueprint menu (filter, heading, etc.)
+	 */
+	struct FMenuSectionDefinition
+	{
+		FMenuSectionDefinition(FBlueprintActionFilter const& SectionFilter, uint32 const Flags);
+		
+		/** Sets the root category for menu items in this section. */
+		void SetSectionHeading(FText const& RootCategory);
+		/** Gets the root category for menu items in this section. */
+		FText const& GetSectionHeading() const;
+
+		/** Sets the grouping for menu items belonging to this section. */
+		void SetSectionSortOrder(int32 const MenuGrouping);
+		
+		/**
+		 * Filters the supplied action and if it passes, spawns a new 
+		 * FBlueprintActionMenuItem for the specified menu (does not add the 
+		 * item to the menu-builder itself).
+		 *
+		 * @param  ItemOwner	The menu item owner that this will be assigned to.
+		 * @param  Action		The node-spawner that the new menu item should wrap.
+		 * @return An empty TSharedPtr if the action was filtered out, otherwise a newly allocated FBlueprintActionMenuItem.
+		 */
+		TSharedPtr<FEdGraphSchemaAction> MakeMenuItem(FBlueprintActionMenuBuilder& ItemOwner, UBlueprintNodeSpawner* Action);
+		
+		
+		/** Series of ESectionFlags, aimed at customizing how we construct this menu section */
+		uint32 Flags;
+		/** A filter for this section of the menu */
+		FBlueprintActionFilter Filter;
+	private:
+		/** In charge of spawning menu items for this section (holds category/ordering information)*/
+		FBlueprintActionMenuItemFactory ItemFactory;
+	};
+	
+	/**
+	 * To offer a fallback in case the action menu refactor is unstable, this
+	 * method implements the old way we used collect blueprint menu actions (for
+	 * both the palette and context menu).
+	 *
+	 * @param  MenuSection	The primary section for the FBlueprintActionMenuBuilder.
+	 * @param  MenuOut  	The menu builder we want all the legacy actions appended to.
+	 */
+	static void AppendLegacyItems(FMenuSectionDefinition const& MenuDef, FBlueprintActionMenuBuilder& MenuOut);
+}
+
+//------------------------------------------------------------------------------
+FBlueprintActionMenuBuilderImpl::FMenuSectionDefinition::FMenuSectionDefinition(FBlueprintActionFilter const& SectionFilterIn, uint32 const FlagsIn)
+	: Flags(FlagsIn)
+	, Filter(SectionFilterIn)
+	, ItemFactory(Filter.Context)
+{
+}
+
+//------------------------------------------------------------------------------
+void FBlueprintActionMenuBuilderImpl::FMenuSectionDefinition::SetSectionHeading(FText const& RootCategory)
+{
+	ItemFactory.RootCategory = RootCategory;
+}
+
+//------------------------------------------------------------------------------
+FText const& FBlueprintActionMenuBuilderImpl::FMenuSectionDefinition::GetSectionHeading() const
+{
+	return ItemFactory.RootCategory;
+}
+
+//------------------------------------------------------------------------------
+void FBlueprintActionMenuBuilderImpl::FMenuSectionDefinition::SetSectionSortOrder(int32 const MenuGrouping)
+{
+	ItemFactory.MenuGrouping = MenuGrouping;
+}
+
+//------------------------------------------------------------------------------
+TSharedPtr<FEdGraphSchemaAction> FBlueprintActionMenuBuilderImpl::FMenuSectionDefinition::MakeMenuItem(FBlueprintActionMenuBuilder& ItemOwner, UBlueprintNodeSpawner* Action)
+{
+	if (Flags & FBlueprintActionMenuBuilder::ConsolidatePropertyActions)
+	{
+		if (UBlueprintPropertyNodeSpawner* PropertySpawner = Cast<UBlueprintPropertyNodeSpawner>(Action))
+		{
+			UProperty const* Property = PropertySpawner->GetProperty();
+			if (ItemFactory.HasProxy(Property))
+			{
+				// we've already filtered a proxy for this property, we don't
+				// want to spawn another
+				Action = nullptr;
+			}
+			else
+			{
+				// make a proxy action (that won't spawn anything)...
+				// FBlueprintActionMenuItem should interpret this and offer the
+				// user a drag-drop menu (with a list of multiple property nodes)
+				Action = ItemFactory.MakePropertyActionProxy(ItemOwner, PropertySpawner);
+			}
+		}
+	}
+	
+	TSharedPtr<FEdGraphSchemaAction> MenuEntry = nullptr;
+	if ((Action != nullptr) && !Filter.IsFiltered(Action))
+	{
+		MenuEntry = ItemFactory.MakeMenuItem(ItemOwner, Action);
+	}
+	return MenuEntry;
+}
+
+//------------------------------------------------------------------------------
+void FBlueprintActionMenuBuilderImpl::AppendLegacyItems(FMenuSectionDefinition const& MenuDef, FBlueprintActionMenuBuilder& MenuOut)
+{
+	FBlueprintActionFilter const&  MenuFilter  = MenuDef.Filter;
+	FBlueprintActionContext const& MenuContext = MenuFilter.Context;
+	
+	// if this is for the context menu
+	if (MenuContext.Graphs.Num() > 0)
+	{
+		UEdGraph* Graph = MenuContext.Graphs[0];
+		UEdGraphSchema const* GraphSchema = GetDefault<UEdGraphSchema>(Graph->Schema);
+		
+		FBlueprintGraphActionListBuilder LegacyBuilder(Graph);
+		if (MenuContext.Pins.Num() > 0)
+		{
+			LegacyBuilder.FromPin = MenuContext.Pins[0];
+		}
+		
+		UBlueprint* Blueprint = FBlueprintEditorUtils::FindBlueprintForGraphChecked(Graph);
+		FBlueprintEditor* EditorInst = (FBlueprintEditor*)FAssetEditorManager::Get().FindEditorForAsset(Blueprint, /*bFocusIfOpen =*/false);
+		
+		bool bIsContextSensitive = true;
+		if (EditorInst != nullptr)
+		{
+			bIsContextSensitive = EditorInst->GetIsContextSensitive();
+			
+			FEdGraphSchemaAction_K2Var* SelectedVar = EditorInst->GetMyBlueprintWidget()->SelectionAsVar();
+			if ((SelectedVar != nullptr) && (SelectedVar->GetProperty() != nullptr))
+			{
+				LegacyBuilder.SelectedObjects.Add(SelectedVar->GetProperty());
+			}
+		}
+		
+		if (bIsContextSensitive)
+		{
+			GraphSchema->GetGraphContextActions(LegacyBuilder);
+			MenuOut.Append(LegacyBuilder);
+		}
+		else
+		{
+			FBlueprintPaletteListBuilder ContextlessLegacyBuilder(Blueprint);
+			UEdGraphSchema_K2::GetAllActions(ContextlessLegacyBuilder);
+			MenuOut.Append(ContextlessLegacyBuilder);
+		}
+	}
+	else if (MenuContext.Blueprints.Num() > 0)
+	{
+		UBlueprint* Blueprint = MenuContext.Blueprints[0];
+		FBlueprintPaletteListBuilder LegacyBuilder(Blueprint, MenuDef.GetSectionHeading().ToString());
+		
+		UClass* ClassFilter = nullptr;
+		if (MenuFilter.OwnerClasses.Num() > 0)
+		{
+			ClassFilter = MenuFilter.OwnerClasses[0];
+		}
+		
+		UEdGraphSchema_K2 const* K2Schema = GetDefault<UEdGraphSchema_K2>();
+		FK2ActionMenuBuilder(K2Schema).GetPaletteActions(LegacyBuilder, ClassFilter);
+		
+		MenuOut.Append(LegacyBuilder);
+	}
+}
+
+/*******************************************************************************
+ * FBlueprintActionMenuBuilder
+ ******************************************************************************/
+
+//------------------------------------------------------------------------------
+FBlueprintActionMenuBuilder::FBlueprintActionMenuBuilder(FBlueprintActionFilter const& Filter, bool bAutoBuild/*= true*/)
+{
+	AddMenuSection(Filter);
+	if (bAutoBuild)
+	{
+		RebuildActionList();
+	}
+}
+
+//------------------------------------------------------------------------------
+void FBlueprintActionMenuBuilder::Empty()
+{
+	FGraphActionListBuilderBase::Empty();
+	MenuSections.Empty();
+}
+
+//------------------------------------------------------------------------------
+void FBlueprintActionMenuBuilder::AddMenuSection(FBlueprintActionFilter const& Filter, FText const& Heading/* = FText::GetEmpty()*/, int32 MenuOrder/* = 0*/, uint32 const Flags/* = 0*/)
+{
+	using namespace FBlueprintActionMenuBuilderImpl;
+	
+	TSharedRef<FMenuSectionDefinition> SectionDescRef = MakeShareable(new FMenuSectionDefinition(Filter, Flags));
+	SectionDescRef->SetSectionHeading(Heading);
+	SectionDescRef->SetSectionSortOrder(MenuOrder);
+
+	MenuSections.Add(SectionDescRef);
+}
+
+//------------------------------------------------------------------------------
+void FBlueprintActionMenuBuilder::RebuildActionList()
+{
+	using namespace FBlueprintActionMenuBuilderImpl;
+
+	FGraphActionListBuilderBase::Empty();
+	
+	if (!GetDefault<UEdGraphSchema_K2>()->bUseLegacyActionMenus)
+	{
+		FBlueprintActionDatabase::FClassActionMap const& ClassActions = FBlueprintActionDatabase::Get().GetAllActions();
+		for (auto ActionEntry : ClassActions)
+		{
+			UClass* Class = ActionEntry.Key;
+			for (UBlueprintNodeSpawner* NodeSpawner : ActionEntry.Value)
+			{
+				for (TSharedRef<FMenuSectionDefinition> MenuSection : MenuSections)
+				{
+					TSharedPtr<FEdGraphSchemaAction> MenuEntry = MenuSection->MakeMenuItem(*this, NodeSpawner);
+					if (MenuEntry.IsValid())
+					{
+						AddAction(MenuEntry);
+					}
+				}
+			}
+		}
+	}
+	else if (MenuSections.Num() > 0)
+	{
+		AppendLegacyItems(*MenuSections[0], *this);
+	}
+
+	// @TODO: account for all K2ActionMenuBuilder action types...
+	// - FEdGraphSchemaAction_K2AddTimeline
+	// - FEdGraphSchemaAction_K2ViewNode
+	// - FEdGraphSchemaAction_K2AddCustomEvent
+	//   FEdGraphSchemaAction_EventFromFunction
+	//   FEdGraphSchemaAction_K2Var
+	//   FEdGraphSchemaAction_K2Delegate
+	//   FEdGraphSchemaAction_K2AssignDelegate
+	// - FEdGraphSchemaAction_K2AddComment
+	//   FEdGraphSchemaAction_K2PasteHere
+	// - FEdGraphSchemaAction_K2NewNode
+	//   FEdGraphSchemaAction_Dummy
+	//   FEdGraphSchemaAction_K2AddCallOnActor
+	//   FEdGraphSchemaAction_K2AddCallOnVariable
+	// - FEdGraphSchemaAction_K2AddComponent
+	//   FEdGraphSchemaAction_K2AddComment
+}
+
+#undef LOCTEXT_NAMESPACE
diff --git a/Engine/Source/Editor/Kismet/Private/BlueprintActionMenuItem.cpp b/Engine/Source/Editor/Kismet/Private/BlueprintActionMenuItem.cpp
new file mode 100644
--- /dev/null
+++ b/Engine/Source/Editor/Kismet/Private/BlueprintActionMenuItem.cpp
@@ -0,0 +1,185 @@
+// Copyright 1998-2014 Epic Games, Inc. All Rights Reserved.
+
+#include "BlueprintEditorPrivatePCH.h"
+#include "BlueprintActionMenuItem.h"
+#include "BlueprintNodeSpawner.h"
+#include "KismetEditorUtilities.h"	// for BringKismetToFocusAttentionOnObject()
+#include "BlueprintEditorUtils.h"	// for AnalyticsTrackNewNode(), MarkBlueprintAsModified(), etc.
+#include "ScopedTransaction.h"
+#include "BlueprintPropertyNodeSpawner.h"
+#include "BPVariableDragDropAction.h"
+#include "BPDelegateDragDropAction.h"
+
+#define LOCTEXT_NAMESPACE "BlueprintActionMenuItem"
+
+/*******************************************************************************
+ * Static FBlueprintMenuActionItem Helpers
+ ******************************************************************************/
+
+namespace FBlueprintMenuActionItemImpl
+{
+	/**
+	 * Utility function for marking blueprints dirty and recompiling them after 
+	 * a node has been added.
+	 * 
+	 * @param  SpawnedNode	The node that was just added to the blueprint.
+	 */
+	static void DirtyBlueprintFromNewNode(UEdGraphNode* SpawnedNode);
+}
+
+//------------------------------------------------------------------------------
+static void FBlueprintMenuActionItemImpl::DirtyBlueprintFromNewNode(UEdGraphNode* SpawnedNode)
+{
+	UEdGraph const* const NodeGraph = SpawnedNode->GetGraph();
+	check(NodeGraph != nullptr);
+	
+	UBlueprint* Blueprint = FBlueprintEditorUtils::FindBlueprintForGraphChecked(NodeGraph);
+	check(Blueprint != nullptr);
+	
+	UK2Node* K2Node = Cast<UK2Node>(SpawnedNode);
+	// see if we need to recompile skeleton after adding this node, or just mark
+	// it dirty (default to rebuilding the skel, since there is no way to if
+	// non-k2 nodes structurally modify the blueprint)
+	if ((K2Node == nullptr) || K2Node->NodeCausesStructuralBlueprintChange())
+	{
+		FBlueprintEditorUtils::MarkBlueprintAsStructurallyModified(Blueprint);
+	}
+	else
+	{
+		FBlueprintEditorUtils::MarkBlueprintAsModified(Blueprint);
+	}
+}
+
+/*******************************************************************************
+ * FBlueprintMenuActionItem
+ ******************************************************************************/
+
+//------------------------------------------------------------------------------
+FBlueprintActionMenuItem::FBlueprintActionMenuItem(UBlueprintNodeSpawner* NodeSpawner, int32 MenuGrouping/* = 0*/)
+	: Action(NodeSpawner)
+{
+	check(Action != nullptr);
+	Grouping = MenuGrouping;
+}
+
+//------------------------------------------------------------------------------
+UEdGraphNode* FBlueprintActionMenuItem::PerformAction(UEdGraph* ParentGraph, UEdGraphPin* FromPin, FVector2D const Location, bool bSelectNewNode/* = true*/)
+{
+	using namespace FBlueprintMenuActionItemImpl;
+	FScopedTransaction Transaction(LOCTEXT("AddNodeTransaction", "Add Node"));
+	
+	// this could return an existing node
+	UEdGraphNode* SpawnedNode = Action->Invoke(ParentGraph);
+	
+	// if a returned node hasn't been added to the graph yet (it must have been freshly spawned)
+	if (ParentGraph->Nodes.Find(SpawnedNode) == INDEX_NONE)
+	{
+		check(SpawnedNode != nullptr);
+		
+		ParentGraph->Modify();
+		ParentGraph->AddNode(SpawnedNode, /*bool bFromUI =*/true, bSelectNewNode);
+		
+		SpawnedNode->PostPlacedNewNode();
+		SpawnedNode->NodePosX = Location.X;
+		SpawnedNode->NodePosY = Location.Y;
+		
+		if (FromPin != nullptr)
+		{
+			// for input pins, a new node will generally overlap the node being
+			// dragged from... work out if we want add in some spacing from the connecting node
+			if (FromPin->Direction == EGPD_Input)
+			{
+				UEdGraphNode* FromNode = FromPin->GetOwningNode();
+				check(FromNode != nullptr);
+				float const FromNodeX = FromNode->NodePosX;
+				
+				static const float MinNodeDistance = 60.f; // min distance between spawned nodes (to keep them from overlapping)
+				if (MinNodeDistance > FMath::Abs(FromNodeX - Location.X))
+				{
+					SpawnedNode->NodePosX = FromNodeX - MinNodeDistance;
+				}
+			}
+			
+			// modify before the call to AutowireNewNode() below
+			FromPin->Modify();
+		}
+		
+		static const float GridSnapSize = 16.f; // @TODO: ensure this is the same as SNodePanel::GetSnapGridSize()
+		SpawnedNode->SnapToGrid(GridSnapSize);
+		
+		// make sure to auto-wire after we position the new node (in case
+		// the auto-wire creates a conversion node to put between them)
+		SpawnedNode->AutowireNewNode(FromPin);
+		
+		FBlueprintEditorUtils::AnalyticsTrackNewNode(SpawnedNode);
+		DirtyBlueprintFromNewNode(SpawnedNode);
+	}
+	// if this node already existed, then we just want to focus on that node...
+	// some node types are only allowed one instance per blueprint (like events)
+	else if (SpawnedNode != nullptr)
+	{
+		FKismetEditorUtilities::BringKismetToFocusAttentionOnObject(SpawnedNode);
+	}
+	
+	return SpawnedNode;
+}
+
+//------------------------------------------------------------------------------
+UEdGraphNode* FBlueprintActionMenuItem::PerformAction(UEdGraph* ParentGraph, TArray<UEdGraphPin*>& FromPins, FVector2D const Location, bool bSelectNewNode/* = true*/)
+{
+	UEdGraphPin* FromPin = nullptr;
+	if (FromPins.Num() > 0)
+	{
+		FromPin = FromPins[0];
+	}
+	
+	UEdGraphNode* SpawnedNode = PerformAction(ParentGraph, FromPin, Location, bSelectNewNode);
+	// try auto-wiring the rest of the pins (if there are any)
+	for (int32 PinIndex = 1; PinIndex < FromPins.Num(); ++PinIndex)
+	{
+		SpawnedNode->AutowireNewNode(FromPins[PinIndex]);
+	}
+	
+	return SpawnedNode;
+}
+
+//------------------------------------------------------------------------------
+void FBlueprintActionMenuItem::AddReferencedObjects(FReferenceCollector& Collector)
+{
+	FEdGraphSchemaAction::AddReferencedObjects(Collector);
+	
+	// these don't get saved to disk, but we want to make sure the objects don't
+	// get GC'd while the action array is around
+	Collector.AddReferencedObject(Action);
+}
+
+//------------------------------------------------------------------------------
+TSharedPtr<FDragDropOperation> FBlueprintActionMenuItem::OnDragged(FNodeCreationAnalytic AnalyticsDelegate) const
+{
+	TSharedPtr<FDragDropOperation> DragDropAction = nullptr;
+
+	if (UBlueprintPropertyNodeSpawner* PropertySpawner = Cast<UBlueprintPropertyNodeSpawner>(Action))
+	{
+		if (PropertySpawner->NodeClass == nullptr)
+		{
+			UProperty const* Property = PropertySpawner->GetProperty();
+			FName const PropertyName = Property->GetFName();
+			UStruct* const PropertyOwner = CastChecked<UStruct>(Property->GetOuterUField());
+
+			if (PropertySpawner->IsDelegateProperty())
+			{
+				TSharedRef<FDragDropOperation> DragDropOpRef = FKismetDelegateDragDropAction::New(PropertyName, PropertyOwner, AnalyticsDelegate);
+				DragDropAction = TSharedPtr<FDragDropOperation>(DragDropOpRef);
+			}
+			else
+			{
+				TSharedRef<FDragDropOperation> DragDropOpRef = FKismetDelegateDragDropAction::New(PropertyName, PropertyOwner, AnalyticsDelegate);
+				DragDropAction = TSharedPtr<FDragDropOperation>(DragDropOpRef);
+			}
+
+		}
+	}
+	return DragDropAction;
+}
+
+#undef LOCTEXT_NAMESPACE
diff --git a/Engine/Source/Editor/Kismet/Private/SBlueprintSubPalette.cpp b/Engine/Source/Editor/Kismet/Private/SBlueprintSubPalette.cpp
--- a/Engine/Source/Editor/Kismet/Private/SBlueprintSubPalette.cpp
+++ b/Engine/Source/Editor/Kismet/Private/SBlueprintSubPalette.cpp
@@ -12,6 +12,7 @@
 #include "BlueprintPaletteFavorites.h"
 #include "BlueprintEditorUtils.h"
 #include "SBlueprintActionMenu.h" // for SBlueprintActionMenuExpander
+#include "BlueprintActionMenuItem.h"
 
 #define LOCTEXT_NAMESPACE "BlueprintSubPalette"
 
@@ -319,7 +320,21 @@ FReply SBlueprintSubPalette::OnActionDragged( const TArray< TSharedPtr<FEdGraphS
 			check(VarClass);
 			if(VarClass != NULL)
 			{
-				return FReply::Handled().BeginDragDrop(FKismetDelegateDragDropAction::New( SharedThis( this ), DelegateAction->GetDelegateName(), VarClass, AnalyticsDelegate));
+				return FReply::Handled().BeginDragDrop(FKismetDelegateDragDropAction::New(DelegateAction->GetDelegateName(), VarClass, AnalyticsDelegate));
+			}
+		}
+		else if (InAction->GetTypeId() == FBlueprintActionMenuItem::StaticGetTypeId())
+		{
+			FBlueprintActionMenuItem* BlueprintAction = (FBlueprintActionMenuItem*)InAction.Get();
+
+			TSharedPtr<FDragDropOperation> DragDropOp = BlueprintAction->OnDragged(AnalyticsDelegate);
+			if (DragDropOp.IsValid())
+			{
+				return FReply::Handled().BeginDragDrop(DragDropOp.ToSharedRef());
+			}
+			else
+			{
+				return FReply::Handled().BeginDragDrop(FKismetDragDropAction::New(InAction, AnalyticsDelegate, CanNodeBePlacedDelegate));
 			}
 		}
 		else
diff --git a/Engine/Source/Editor/Kismet/Private/SMyBlueprint.cpp b/Engine/Source/Editor/Kismet/Private/SMyBlueprint.cpp
--- a/Engine/Source/Editor/Kismet/Private/SMyBlueprint.cpp
+++ b/Engine/Source/Editor/Kismet/Private/SMyBlueprint.cpp
@@ -922,7 +922,7 @@ FReply SMyBlueprint::OnActionDragged( const TArray< TSharedPtr<FEdGraphSchemaAct
 				const bool bIsAltDown = MouseEvent.IsAltDown();
 				const bool bIsCtrlDown = MouseEvent.IsLeftControlDown() || MouseEvent.IsRightControlDown();
 				
-				TSharedRef<FKismetVariableDragDropAction> DragOperation = FKismetDelegateDragDropAction::New( SharedThis( this ), DelegateAction->GetDelegateName(), VarClass, AnalyticsDelegate);
+				TSharedRef<FKismetVariableDragDropAction> DragOperation = FKismetDelegateDragDropAction::New(DelegateAction->GetDelegateName(), VarClass, AnalyticsDelegate);
 				DragOperation->SetAltDrag(bIsAltDown);
 				DragOperation->SetCtrlDrag(bIsCtrlDown);
 				return FReply::Handled().BeginDragDrop(DragOperation);
diff --git a/Engine/Source/Editor/Kismet/Public/BlueprintActionMenuBuilder.h b/Engine/Source/Editor/Kismet/Public/BlueprintActionMenuBuilder.h
new file mode 100644
--- /dev/null
+++ b/Engine/Source/Editor/Kismet/Public/BlueprintActionMenuBuilder.h
@@ -0,0 +1,79 @@
+// Copyright 1998-2014 Epic Games, Inc. All Rights Reserved.
+
+#pragma once
+
+#include "EdGraph/EdGraphSchema.h" // for FGraphActionListBuilderBase
+
+// Forward declarations
+namespace FBlueprintActionMenuBuilderImpl
+{
+	// internal type, forward declared to hide implementation details
+	struct FMenuSectionDefinition; 
+};
+class  FBlueprintActionFilter;
+struct FBlueprintActionContext;
+
+/**
+ * Responsible for constructing a list of viable blueprint actions. Runs the 
+ * blueprint actions database through a filter and spawns a series of 
+ * FBlueprintActionMenuItems for actions that pass. Takes care of generating the 
+ * each menu item's category/name/etc.
+ */
+struct KISMET_API FBlueprintActionMenuBuilder : public FGraphActionListBuilderBase
+{
+public:
+	/** Flags used to customize specific sections of the menu. */
+	enum ESectionFlag
+	{
+		// Rolls properties into a single menu item that will spawn a drag-drop
+		// menu for users to pick a node type from.
+		ConsolidatePropertyActions = (1<<0),
+	};
+	
+public:
+	/** Default constructor */
+	FBlueprintActionMenuBuilder() {}
+
+	/**
+	 * Good for constructing a menu built from a single filter. If the 
+	 * bAutoBuild param is set to false, then know you have to manually call
+	 * RebuildActionList() to generate the menu.
+	 *
+	 * @param  Filter		The filter you want applied for this main section of the menu.
+	 * @param  bAutoBuild	If set, will automatically call RebuildActionList() here in the constructor.
+	 */
+	FBlueprintActionMenuBuilder(FBlueprintActionFilter const& Filter, bool bAutoBuild = true);
+	
+	// FGraphActionListBuilderBase interface
+	virtual void Empty() override;
+	// End FGraphActionListBuilderBase interface
+	
+	/**
+	 * Some action menus require multiple sections. One option is to create 
+	 * multiple FBlueprintActionMenuBuilders and append them together, but that
+	 * can be unperformant (each builder will run through the entire database
+	 * separately)... This method provides an alternative, where you can specify
+	 * a separate filter/heading/ordering for a sub-section of the menu.
+	 *
+	 * @param  Filter	 The filter you want applied to this section of the menu.
+	 * @param  Heading	 The root category for this section of the menu (can be left blank).
+	 * @param  MenuOrder The sort order to assign this section of the menu (higher numbers come first).
+	 * @param  Flags	 Set of ESectionFlags to customize this menu section.
+	 */
+	void AddMenuSection(FBlueprintActionFilter const& Filter, FText const& Heading = FText::GetEmpty(), int32 MenuOrder = 0, uint32 const Flags = 0);
+	
+	/**
+	 * Regenerates the entire menu list from the cached menu sections. Filters 
+	 * and adds action items from the blueprint action database (as defined by 
+	 * the MenuSections list).
+	 */
+	void RebuildActionList();
+
+private:
+	/** 
+	 * Defines all the separate sections of the menu (filter, sort order, etc.).
+	 * Defined as a TSharedRef<> so as to hide the implementation details (keep 
+	 * this API clean).
+	 */
+	TArray< TSharedRef<FBlueprintActionMenuBuilderImpl::FMenuSectionDefinition> > MenuSections;
+};
\ No newline at end of file
diff --git a/Engine/Source/Editor/Kismet/Public/BlueprintActionMenuItem.h b/Engine/Source/Editor/Kismet/Public/BlueprintActionMenuItem.h
new file mode 100644
--- /dev/null
+++ b/Engine/Source/Editor/Kismet/Public/BlueprintActionMenuItem.h
@@ -0,0 +1,51 @@
+// Copyright 1998-2014 Epic Games, Inc. All Rights Reserved.
+
+#pragma once
+
+#include "EdGraph/EdGraphSchema.h" // for FEdGraphSchemaAction
+#include "BlueprintEditor.h"       // for FNodeCreationAnalytic
+#include "BlueprintActionMenuItem.generated.h"
+
+// Forward declarations
+class UBlueprintNodeSpawner;
+class FDragDropOperation;
+
+/**
+ * Wrapper around a UBlueprintNodeSpawner, which takes care of specialized
+ * node spawning. This class should not be sub-classed, any special handling
+ * should be done inside a UBlueprintNodeSpawner subclass, which will be
+ * invoked from this class (separated to divide ui and node-spawning).
+ */
+USTRUCT()
+struct FBlueprintActionMenuItem : public FEdGraphSchemaAction
+{
+	GENERATED_USTRUCT_BODY();
+	
+public:
+	static FString StaticGetTypeId() { static FString const TypeId = TEXT("FBlueprintActionMenuItem"); return TypeId; }
+	
+	/** Constructors */
+	FBlueprintActionMenuItem() : Action(nullptr) {}
+	FBlueprintActionMenuItem(UBlueprintNodeSpawner* NodeSpawner, int32 MenuGrouping = 0);
+	
+	// FEdGraphSchemaAction interface
+	virtual FString       GetTypeId() const final { return StaticGetTypeId(); }
+	virtual UEdGraphNode* PerformAction(UEdGraph* ParentGraph, UEdGraphPin* FromPin, FVector2D const Location, bool bSelectNewNode = true) final;
+	virtual UEdGraphNode* PerformAction(UEdGraph* ParentGraph, TArray<UEdGraphPin*>& FromPins, FVector2D const Location, bool bSelectNewNode = true) final;
+	virtual void          AddReferencedObjects(FReferenceCollector& Collector) final;
+	// End FEdGraphSchemaAction interface
+
+	/**
+	 * Attempts to create a specific drag/drop action for this menu entry 
+	 * (certain menu entries require unique drag/drop handlers... like property
+	 * placeholders, where a drop action spawns a sub-menu for the user to pick
+	 * a node type from).
+	 * 
+	 * @param  AnalyticsDelegate	The analytics callback to assign the drag/drop operation.
+	 * @return An empty TSharedPtr if this menu item doesn't require a unique one, else a newly instantiated drag/drop op.
+	 */
+	TSharedPtr<FDragDropOperation> OnDragged(FNodeCreationAnalytic AnalyticsDelegate) const;
+
+	/** Specialized node-spawner, that comprises the action portion of this menu entry. */
+	UBlueprintNodeSpawner* Action;
+};
diff --git a/Engine/Source/Editor/UnrealEd/Private/EditorCategoryUtils.cpp b/Engine/Source/Editor/UnrealEd/Private/EditorCategoryUtils.cpp
--- a/Engine/Source/Editor/UnrealEd/Private/EditorCategoryUtils.cpp
+++ b/Engine/Source/Editor/UnrealEd/Private/EditorCategoryUtils.cpp
@@ -64,6 +64,7 @@ static FFormatNamedArguments& FEditorCategoryUtilsImpl::GetCategoryTable()
 		REGISTER_ROOT_CATEGORY(Pawn);
 		REGISTER_ROOT_CATEGORY(Rendering);
 		REGISTER_ROOT_CATEGORY(Utilities);
+		REGISTER_ROOT_CATEGORY(Variables);
 
 #undef REGISTER_ROOT_CATEGORY
 #undef ASSIGN_ROOT_CATEGORY
@@ -130,6 +131,7 @@ FText const& FEditorCategoryUtils::GetCommonCategory(FCommonEditorCategory::EVal
 		REGISTER_COMMON_CATEGORY(Pawn);
 		REGISTER_COMMON_CATEGORY(Rendering);
 		REGISTER_COMMON_CATEGORY(Utilities);
+		REGISTER_COMMON_CATEGORY(Variables);
 
 		REGISTER_COMMON_CATEGORY(FlowControl);
 		REGISTER_COMMON_CATEGORY(Transformation);
@@ -198,7 +200,7 @@ FString FEditorCategoryUtils::GetCategoryDisplayString(FString const& Unsanitize
 				if (KeyLen > 0)
 				{
 					FString Key(KeyLen, *DisplayString + KeyIndex+1);
-					ReplacementStr = FEditorCategoryUtilsImpl::GetCategory(Key).ToString();
+					ReplacementStr = FEditorCategoryUtilsImpl::GetCategory(Key.Trim()).ToString();
 				}
 				DisplayString.ReplaceInline(*ToReplaceStr, *ReplacementStr);
 			}
diff --git a/Engine/Source/Editor/UnrealEd/Public/EditorCategoryUtils.h b/Engine/Source/Editor/UnrealEd/Public/EditorCategoryUtils.h
--- a/Engine/Source/Editor/UnrealEd/Public/EditorCategoryUtils.h
+++ b/Engine/Source/Editor/UnrealEd/Public/EditorCategoryUtils.h
@@ -41,6 +41,8 @@ struct FCommonEditorCategory
 		Enum,
 		Struct,
 		Macro,
+
+		Variables,
 	};
 };
 
