Moving the new palette/context menu building into a shared utils class (so we can reflect it in the DumpBlueprintsInfo commandlet without having to duplicate code).

diff --git a/Engine/Source/Editor/Kismet/Private/BlueprintActionMenuBuilder.cpp b/Engine/Source/Editor/Kismet/Private/BlueprintActionMenuBuilder.cpp
--- a/Engine/Source/Editor/Kismet/Private/BlueprintActionMenuBuilder.cpp
+++ b/Engine/Source/Editor/Kismet/Private/BlueprintActionMenuBuilder.cpp
@@ -490,6 +490,8 @@ void FBlueprintActionMenuBuilderImpl::AppendLegacyItems(FMenuSectionDefinition c
 		}
 		
 		UBlueprint* Blueprint = FBlueprintEditorUtils::FindBlueprintForGraphChecked(Graph);
+		// not the best way to poll for Selected properties, but since this
+		// function is only temporary until we switch over to the new system...
 		FBlueprintEditor* EditorInst = (FBlueprintEditor*)FAssetEditorManager::Get().FindEditorForAsset(Blueprint, /*bFocusIfOpen =*/false);
 		
 		bool bIsContextSensitive = true;
diff --git a/Engine/Source/Editor/Kismet/Private/BlueprintActionMenuUtils.cpp b/Engine/Source/Editor/Kismet/Private/BlueprintActionMenuUtils.cpp
new file mode 100644
--- /dev/null
+++ b/Engine/Source/Editor/Kismet/Private/BlueprintActionMenuUtils.cpp
@@ -0,0 +1,74 @@
+// Copyright 1998-2014 Epic Games, Inc. All Rights Reserved.
+
+#include "BlueprintEditorPrivatePCH.h"
+#include "BlueprintActionMenuUtils.h"
+#include "BlueprintActionMenuBuilder.h"
+#include "K2Node_CallFunction.h"
+#include "K2Node_Event.h"
+
+#define LOCTEXT_NAMESPACE "BlueprintActionMenuUtils"
+
+//------------------------------------------------------------------------------
+void FBlueprintActionMenuUtils::MakePaletteMenu(FBlueprintActionContext const& Context, UClass* FilterClass, FBlueprintActionMenuBuilder& MenuOut)
+{
+	MenuOut.Empty();
+	
+	uint32 FilterFlags = 0x00;
+	if (FilterClass != nullptr)
+	{
+		FilterFlags = FBlueprintActionFilter::BPFILTER_ExcludeGlobalFields;
+	}
+	
+	FBlueprintActionFilter MenuFilter(FilterFlags);
+	MenuFilter.Context = Context;
+	
+	if (FilterClass != nullptr)
+	{
+		MenuFilter.OwnerClasses.Add(FilterClass);
+	}
+	
+	MenuOut.AddMenuSection(MenuFilter, LOCTEXT("PaletteRoot", "Library"), /*MenuOrder =*/0, FBlueprintActionMenuBuilder::ConsolidatePropertyActions);
+	MenuOut.RebuildActionList();
+}
+
+//------------------------------------------------------------------------------
+void FBlueprintActionMenuUtils::MakeContextMenu(FBlueprintActionContext const& Context, TArray<UProperty*> const& SelectedProperties, FBlueprintActionMenuBuilder& MenuOut)
+{
+	MenuOut.Empty();
+	
+	FBlueprintActionFilter MenuFilter;
+	MenuFilter.Context = Context;
+	
+	for (UProperty* SelectedProperty : SelectedProperties)
+	{
+		static int32 const ComponentsSectionGroup = 100;
+		
+		if (UObjectProperty* ObjProperty = Cast<UObjectProperty>(SelectedProperty))
+		{
+			MenuFilter.OwnerClasses.Add(ObjProperty->PropertyClass);
+			FText const PropertyName = FText::FromName(ObjProperty->GetFName());
+			
+			MenuFilter.NodeTypes.Add(UK2Node_CallFunction::StaticClass());
+			FText const FuncSectionName = FText::Format(LOCTEXT("ComponentFuncSection", "Call Function on {0}"), PropertyName);
+			MenuOut.AddMenuSection(MenuFilter, FuncSectionName, ComponentsSectionGroup);
+			
+			MenuFilter.NodeTypes[0] = UK2Node_Event::StaticClass();
+			FText const EventSectionName = FText::Format(LOCTEXT("ComponentEventSection", "Add Event for {0}"), PropertyName);
+			MenuOut.AddMenuSection(MenuFilter, EventSectionName, ComponentsSectionGroup);
+			
+			MenuFilter.NodeTypes.Empty();
+		}
+	}
+	
+	for (UBlueprint* Blueprint : Context.Blueprints)
+	{
+		MenuFilter.OwnerClasses.Empty();
+		MenuFilter.OwnerClasses.Add(Blueprint->GeneratedClass);
+		
+		MenuOut.AddMenuSection(MenuFilter);
+	}
+	
+	MenuOut.RebuildActionList();
+}
+
+#undef LOCTEXT_NAMESPACE
diff --git a/Engine/Source/Editor/Kismet/Private/SBlueprintActionMenu.cpp b/Engine/Source/Editor/Kismet/Private/SBlueprintActionMenu.cpp
--- a/Engine/Source/Editor/Kismet/Private/SBlueprintActionMenu.cpp
+++ b/Engine/Source/Editor/Kismet/Private/SBlueprintActionMenu.cpp
@@ -12,6 +12,7 @@
 #include "K2ActionMenuBuilder.h" // for FBlueprintPaletteListBuilder
 #include "BlueprintActionMenuBuilder.h"
 #include "BlueprintActionFilter.h"
+#include "BlueprintActionMenuUtils.h"
 #include "BlueprintPaletteFavorites.h"
 #include "IDocumentation.h"
 
@@ -392,97 +393,36 @@ void SBlueprintActionMenu::CollectAllActions(FGraphActionListBuilderBase& OutAll
 	TSharedPtr<FBlueprintEditor> BlueprintEditor = EditorPtr.Pin();
 	bool const bIsContextSensitive = BlueprintEditor->GetIsContextSensitive();
 	
-	if (GetDefault<UEdGraphSchema_K2>()->bUseLegacyActionMenus)
+	FBlueprintActionContext FilterContext;
+	FilterContext.Blueprints.Add(Blueprint);
+	// we still want context from the graph (even if the user has unchecked
+	// "Context Sensitive"), otherwise the user would be presented with nodes
+	// that can't be placed in the graph... if the user isn't being presented
+	// with a valid node, then fix it up in filtering
+	FilterContext.Graphs.Add(GraphObj);
+	
+	TArray<UProperty*> SelectedProperties;
+	if (bIsContextSensitive)
 	{
-		// Build up the context object
-		FBlueprintGraphActionListBuilder ContextMenuBuilder(GraphObj);
-		bool bHasPins = DraggedFromPins.Num() > 0;
-		if (!bIsContextSensitive)
-		{
-			FBlueprintPaletteListBuilder PaletteBuilder(Blueprint);
-			UEdGraphSchema_K2::GetAllActions(PaletteBuilder);
-			//@TODO: Avoid this copy
-			OutAllActions.Append(PaletteBuilder);
-		}
-		else if (!bHasPins)
-		{
-			// Pass in selected property
-			FEdGraphSchemaAction_K2Var* SelectedVar = EditorPtr.Pin()->GetMyBlueprintWidget()->SelectionAsVar();
-			if (SelectedVar != NULL)
-			{
-				ContextMenuBuilder.SelectedObjects.Add(SelectedVar->GetProperty());
-			}
-		}
-		else
-		{
-			ContextMenuBuilder.FromPin = DraggedFromPins[0];
-		}
+		FilterContext.Pins = DraggedFromPins;
 		
-		// Note: Cannot call GetGraphContextActions() during serialization and GC due to its use of FindObject()
-		if(!GIsSavingPackage && !GIsGarbageCollecting)
+		FEdGraphSchemaAction_K2Var* SelectedVar = BlueprintEditor->GetMyBlueprintWidget()->SelectionAsVar();
+		if ((SelectedVar != nullptr) && (SelectedVar->GetProperty() != nullptr))
 		{
-			// Determine all possible actions
-			GraphObj->GetSchema()->GetGraphContextActions(ContextMenuBuilder);
+			SelectedProperties.Add(SelectedVar->GetProperty());
 		}
-		
-		// Copy the added options back to the main list
-		//@TODO: Avoid this copy
-		OutAllActions.Append(ContextMenuBuilder);
-	}
-	else
-	{
-		FBlueprintActionFilter MenuFilter;
-		MenuFilter.Context.Blueprints.Add(Blueprint);
-		// we still want context from the graph (even if the user has unchecked
-		// "Context Sensitive"), otherwise the user would be presented with nodes
-		// that can't be placed in the graph... if the user isn't being presented
-		// with a valid node, then fix it up in filtering
-		MenuFilter.Context.Graphs.Add(GraphObj);
-		
-		FBlueprintActionMenuBuilder MenuBuilder;
-		if (bIsContextSensitive)
-		{
-			MenuFilter.Context.Pins = DraggedFromPins;
-			
-			FEdGraphSchemaAction_K2Var* SelectedVar = BlueprintEditor->GetMyBlueprintWidget()->SelectionAsVar();
-			if (SelectedVar != nullptr)
-			{
-				if (UObjectProperty* SelectedProperty = Cast<UObjectProperty>(SelectedVar->GetProperty()))
-				{
-					MenuFilter.OwnerClasses.Add(SelectedProperty->PropertyClass);
-					
-					bool const bUseFriendlyName = GetDefault<UEditorStyleSettings>()->bShowFriendlyNames;
-					FText const PropertyName = bUseFriendlyName ? FText::FromString(UEditorEngine::GetFriendlyName(SelectedProperty)) : FText::FromName(SelectedProperty->GetFName());
-					
-					MenuFilter.NodeTypes.Add(UK2Node_CallFunction::StaticClass());
-					FText const FuncSectionName = FText::Format(LOCTEXT("ComponentFuncSection", "Call Function on {0}"), PropertyName);
-					MenuBuilder.AddMenuSection(MenuFilter, FuncSectionName);
-					
-					MenuFilter.NodeTypes[0] = UK2Node_Event::StaticClass();
-					FText const EventSectionName = FText::Format(LOCTEXT("ComponentEventSection", "Add Event for {0}"), PropertyName);
-					MenuBuilder.AddMenuSection(MenuFilter, EventSectionName);
-					
-					MenuFilter.NodeTypes.Empty();
-					MenuFilter.OwnerClasses.Empty();
-				}
-			}
-			
-			// only allow members of this blueprint to be called (along with static
-			// library members)
-			MenuFilter.OwnerClasses.Add(Blueprint->GeneratedClass);
-		}	
-		MenuBuilder.AddMenuSection(MenuFilter);
-		
-		MenuBuilder.RebuildActionList();
-		// copy the added options back to the main list
-		OutAllActions.Append(MenuBuilder); // @TODO: Avoid this copy
 	}
 	
-	// also try adding promote to variable if we can do so.
-	if (DraggedFromPins.Num() > 0)
+	FBlueprintActionMenuBuilder MenuBuilder;
+	// NOTE: cannot call GetGraphContextActions() during serialization and GC due to its use of FindObject()
+	if(!GIsSavingPackage && !GIsGarbageCollecting)
 	{
-		TryInsertPromoteToVariable(DraggedFromPins[0], OutAllActions);
+		FBlueprintActionMenuUtils::MakeContextMenu(FilterContext, SelectedProperties, MenuBuilder);
 	}
+	// copy the added options back to the main list
+	OutAllActions.Append(MenuBuilder); // @TODO: Avoid this copy
+	// also try adding promote to variable if we can do so.
+	TryInsertPromoteToVariable(FilterContext, OutAllActions);
 }
 
 TSharedRef<SEditableTextBox> SBlueprintActionMenu::GetFilterTextBox()
@@ -520,15 +460,18 @@ void SBlueprintActionMenu::OnActionSelected( const TArray< TSharedPtr<FEdGraphSc
 	}
 }
 
-void SBlueprintActionMenu::TryInsertPromoteToVariable(UEdGraphPin const* ContextPin, FGraphActionListBuilderBase& OutAllActions)
+void SBlueprintActionMenu::TryInsertPromoteToVariable(FBlueprintActionContext const& MenuContext, FGraphActionListBuilderBase& OutAllActions)
 {
 	// If we can promote this to a variable add a menu entry to do so.
 	const UEdGraphSchema_K2* K2Schema = Cast<const UEdGraphSchema_K2>(GraphObj->GetSchema());
-	if ((K2Schema != nullptr) && K2Schema->CanPromotePinToVariable(*ContextPin))
+	if ((K2Schema != nullptr) && (MenuContext.Pins.Num() > 0))
 	{
-		TSharedPtr<FBlueprintAction_PromoteVariable> PromoteAction = TSharedPtr<FBlueprintAction_PromoteVariable>(new FBlueprintAction_PromoteVariable());
-		PromoteAction->MyBlueprintEditor = EditorPtr;
-		OutAllActions.AddAction( PromoteAction );
+		if (K2Schema->CanPromotePinToVariable(*MenuContext.Pins[0]))
+		{
+			TSharedPtr<FBlueprintAction_PromoteVariable> PromoteAction = TSharedPtr<FBlueprintAction_PromoteVariable>(new FBlueprintAction_PromoteVariable());
+			PromoteAction->MyBlueprintEditor = EditorPtr;
+			OutAllActions.AddAction( PromoteAction );
+		}
 	}
 }
 
diff --git a/Engine/Source/Editor/Kismet/Private/SBlueprintActionMenu.h b/Engine/Source/Editor/Kismet/Private/SBlueprintActionMenu.h
--- a/Engine/Source/Editor/Kismet/Private/SBlueprintActionMenu.h
+++ b/Engine/Source/Editor/Kismet/Private/SBlueprintActionMenu.h
@@ -3,7 +3,8 @@
 
 #pragma once
 
-class SGraphActionMenu;
+class  SGraphActionMenu;
+struct FBlueprintActionContext;
 
 /*******************************************************************************
 * SBlueprintActionMenu
@@ -50,7 +51,7 @@ protected:
 	void CollectAllActions(FGraphActionListBuilderBase& OutAllActions);
 
 	/** Functioin to try to insert a promote to variable entry if it is possible to do so. */
-	void TryInsertPromoteToVariable(UEdGraphPin const* ContextPin, FGraphActionListBuilderBase& OutAllActions);
+	void TryInsertPromoteToVariable(FBlueprintActionContext const& Context, FGraphActionListBuilderBase& OutAllActions);
 
 private:
 	UEdGraph* GraphObj;
diff --git a/Engine/Source/Editor/Kismet/Private/SBlueprintLibraryPalette.cpp b/Engine/Source/Editor/Kismet/Private/SBlueprintLibraryPalette.cpp
--- a/Engine/Source/Editor/Kismet/Private/SBlueprintLibraryPalette.cpp
+++ b/Engine/Source/Editor/Kismet/Private/SBlueprintLibraryPalette.cpp
@@ -8,6 +8,7 @@
 #include "BlueprintPaletteFavorites.h"
 #include "BlueprintActionFilter.h"
 #include "BlueprintActionMenuBuilder.h"
+#include "BlueprintActionMenuUtils.h"
 
 #define LOCTEXT_NAMESPACE "BlueprintLibraryPalette"
 
@@ -208,37 +209,20 @@ void SBlueprintLibraryPalette::CollectAllActions(FGraphActionListBuilderBase& Ou
 		RootCategory = TEXT("");
 	}
 	
-	if (GetDefault<UEdGraphSchema_K2>()->bUseLegacyActionMenus)
-	{
-		UBlueprint const* const Blueprint = GetBlueprint();
-		FBlueprintPaletteListBuilder PaletteBuilder(Blueprint, RootCategory);
-		
-		UEdGraphSchema_K2::GetPaletteActions(PaletteBuilder, FilterClass);
-		OutAllActions.Append(PaletteBuilder);
-	}
-	else
+	FBlueprintActionContext FilterContext;
+	FilterContext.Blueprints.Add(GetBlueprint());
+	
+	UClass* ClassFilter = nullptr;
+	if (FilterClass.IsValid())
 	{
-		uint32 FilterFlags = 0x00;
-		if (FilterClass.IsValid())
-		{
-			FilterFlags = FBlueprintActionFilter::BPFILTER_ExcludeGlobalFields;
-		}
-		
-		FBlueprintActionFilter PaletteFilter(FilterFlags);
-		PaletteFilter.Context.Blueprints.Add(GetBlueprint());
-		PaletteFilter.ExcludedNodeTypes.Add(UK2Node_Variable::StaticClass());
-		
-		if (FilterClass.IsValid())
-		{
-			PaletteFilter.OwnerClasses.Add(FilterClass.Get());
-		}
-		
-		FBlueprintActionMenuBuilder PaletteBuilder;
-		PaletteBuilder.AddMenuSection(PaletteFilter, FText::FromString(RootCategory));
-		PaletteBuilder.RebuildActionList();
-		
-		OutAllActions.Append(PaletteBuilder);
+		ClassFilter = FilterClass.Get();
 	}
+	
+	FBlueprintActionMenuBuilder PaletteBuilder;
+	FBlueprintActionMenuUtils::MakePaletteMenu(FilterContext, ClassFilter, PaletteBuilder);
+	
+	PaletteBuilder.RebuildActionList();
+	OutAllActions.Append(PaletteBuilder);
 }
 
 //------------------------------------------------------------------------------
diff --git a/Engine/Source/Editor/Kismet/Public/BlueprintActionMenuUtils.h b/Engine/Source/Editor/Kismet/Public/BlueprintActionMenuUtils.h
new file mode 100644
--- /dev/null
+++ b/Engine/Source/Editor/Kismet/Public/BlueprintActionMenuUtils.h
@@ -0,0 +1,33 @@
+// Copyright 1998-2014 Epic Games, Inc. All Rights Reserved.
+
+#pragma once
+
+// Forward declarations
+class  FBlueprintActionMenuBuilder;
+struct FBlueprintActionContext;
+
+struct FBlueprintActionMenuUtils
+{
+	/**
+	 * A centralized utility function for constructing blueprint palette menus.
+	 * Rolls the supplied Context and FilterClass into a filter that's used to
+	 * construct the palette.
+	 *
+	 * @param  Context		Contains the blueprint that the palette is for.
+	 * @param  FilterClass	If not null, then this specifies the class whose members we want listed (and nothing else).
+	 * @param  MenuOut		The structure that will be populated with palette menu items.
+	 */
+	static void MakePaletteMenu(FBlueprintActionContext const& Context, UClass* FilterClass, FBlueprintActionMenuBuilder& MenuOut);
+	
+	/**
+	 * A centralized utility function for constructing blueprint context menus.
+	 * Rolls the supplied Context and SelectedProperties into a series of 
+	 * filters that're used to construct the menu.
+	 *
+	 * @param  Context				Contains the blueprint/graph/pin that the menu is for.
+	 * @param  SelectedProperties	A set of selected properties to offer contextual menu options for.
+	 * @param  MenuOut				The structure that will be populated with context menu items.
+	 */
+	static void MakeContextMenu(FBlueprintActionContext const& Context, TArray<UProperty*> const& SelectedProperties, FBlueprintActionMenuBuilder& MenuOut);
+};
+
diff --git a/Engine/Source/Editor/UnrealEd/Private/Commandlets/DumpBlueprintsInfoCommandlet.cpp b/Engine/Source/Editor/UnrealEd/Private/Commandlets/DumpBlueprintsInfoCommandlet.cpp
--- a/Engine/Source/Editor/UnrealEd/Private/Commandlets/DumpBlueprintsInfoCommandlet.cpp
+++ b/Engine/Source/Editor/UnrealEd/Private/Commandlets/DumpBlueprintsInfoCommandlet.cpp
@@ -16,13 +16,10 @@
 #include "AssetSelection.h"
 #include "K2Node_CustomEvent.h"
 #include "ScopedTimers.h"
-
 #include "BlueprintActionMenuBuilder.h"
 #include "BlueprintActionFilter.h"
 #include "BlueprintActionDatabase.h"
-#include "K2Node_CallFunction.h"
-#include "K2Node_Event.h"
-#include "K2Node_Variable.h"
+#include "BlueprintActionMenuUtils.h"
 
 DEFINE_LOG_CATEGORY_STATIC(LogBlueprintInfoDump, Log, All);
 
@@ -985,17 +982,17 @@ static double DumpBlueprintInfoUtils::GetPaletteMenuActions(FBlueprintPaletteLis
 	}
 	else
 	{
-		FBlueprintActionFilter ActionFilter;
-		ActionFilter.Context.Blueprints.Add(const_cast<UBlueprint*>(PaletteBuilder.Blueprint));
-		ActionFilter.ExcludedNodeTypes.Add(UK2Node_Variable::StaticClass());
-
-		// prime the database so it's not recorded in our timing capture
-		FBlueprintActionDatabase::Prime(); 
-
-		FDurationTimer DurationTimer(MenuBuildDuration);
-		FBlueprintActionMenuBuilder MenuBuilder(ActionFilter, /*bAutoBuild =*/true);
-		DurationTimer.Stop();
-
+		FBlueprintActionContext FilterContext;
+		FilterContext.Blueprints.Add(const_cast<UBlueprint*>(PaletteBuilder.Blueprint));
+		
+		FBlueprintActionMenuBuilder MenuBuilder;
+		{
+			// prime the database so it's not recorded in our timing capture
+			FBlueprintActionDatabase::Prime();
+			
+			FScopedDurationTimer DurationTimer(MenuBuildDuration);
+			FBlueprintActionMenuUtils::MakePaletteMenu(FilterContext, PaletteFilter, MenuBuilder);
+		}
 		PaletteBuilder.Append(MenuBuilder);
 	}
 
@@ -1054,7 +1051,7 @@ static void DumpBlueprintInfoUtils::DumpPalette(uint32 Indent, UBlueprint* Bluep
 	}
 	else
 	{
-		FBlueprintPaletteListBuilder PaletteBuilder(Blueprint);
+		FBlueprintPaletteListBuilder PaletteBuilder(Blueprint, TEXT("Library"));
 		double MenuBuildDuration = GetPaletteMenuActions(PaletteBuilder, ClassFilter);
 
 		BeginPaletteEntry += NestedIndent + "\"FilterClass\" : \"" + FilterClassName + "\",\n";
@@ -1525,45 +1522,31 @@ static double DumpBlueprintInfoUtils::GetContextMenuActions(FBlueprintGraphActio
 	}
 	else
 	{
-		UBlueprint* Blueprint = const_cast<UBlueprint*>(ActionBuilder.Blueprint);
-		
-		FBlueprintActionFilter MenuFilter;
-		MenuFilter.Context.Blueprints.Add(Blueprint);
-		MenuFilter.Context.Graphs.Add(const_cast<UEdGraph*>(ActionBuilder.CurrentGraph));
+		FBlueprintActionContext FilterContext;
+		FilterContext.Blueprints.Add(const_cast<UBlueprint*>(ActionBuilder.Blueprint));
+		FilterContext.Graphs.Add(const_cast<UEdGraph*>(ActionBuilder.CurrentGraph));
 		
 		if (ActionBuilder.FromPin != nullptr)
 		{
-			MenuFilter.Context.Pins.Add(const_cast<UEdGraphPin*>(ActionBuilder.FromPin));
+			FilterContext.Pins.Add(const_cast<UEdGraphPin*>(ActionBuilder.FromPin));
 		}
 		
-		FBlueprintActionMenuBuilder MenuBuilder;
+		TArray<UProperty*> SelectedProperties;
 		for (UObject* SelectedObj : ActionBuilder.SelectedObjects)
 		{
-			if (UObjectProperty* SelectedProperty = Cast<UObjectProperty>(SelectedObj))
+			if (UProperty* SelectedProperty = Cast<UObjectProperty>(SelectedObj))
 			{
-				MenuFilter.OwnerClasses.Add(SelectedProperty->PropertyClass);
-				FText const PropertyName = FText::FromName(SelectedProperty->GetFName());
-				
-				MenuFilter.NodeTypes.Add(UK2Node_CallFunction::StaticClass());
-				FText const FuncSectionName = FText::Format(NSLOCTEXT("DumpBpInfo", "ComponentFuncSection", "Call Function on {0}"), PropertyName);
-				MenuBuilder.AddMenuSection(MenuFilter, FuncSectionName);
-				
-				MenuFilter.NodeTypes[0] = UK2Node_Event::StaticClass();
-				FText const EventSectionName = FText::Format(NSLOCTEXT("DumpBpInfo", "ComponentEventSection", "Add Event for {0}"), PropertyName);
-				MenuBuilder.AddMenuSection(MenuFilter, EventSectionName);
-				
-				MenuFilter.NodeTypes.Empty();
-				MenuFilter.OwnerClasses.Empty();
+				SelectedProperties.Add(SelectedProperty);
 			}
 		}
-		MenuFilter.OwnerClasses.Add(Blueprint->GeneratedClass);
-		MenuBuilder.AddMenuSection(MenuFilter);
-
-		// prime the database so it's not recorded in our timing capture
-		FBlueprintActionDatabase::Prime();
+		
+		FBlueprintActionMenuBuilder MenuBuilder;
 		{
+			// prime the database so it's not recorded in our timing capture
+			FBlueprintActionDatabase::Prime();
+			
 			FScopedDurationTimer DurationTimer(MenuBuildDuration);
-			MenuBuilder.RebuildActionList();
+			FBlueprintActionMenuUtils::MakeContextMenu(FilterContext, SelectedProperties, MenuBuilder);
 		}
 		ActionBuilder.Append(MenuBuilder);
 	}
