Apparently, XCode doesn't catch compiler errors like this (mis-forward-declared a struct as a class).

diff --git a/Engine/Source/Editor/Kismet/Private/BlueprintActionMenuBuilder.cpp b/Engine/Source/Editor/Kismet/Private/BlueprintActionMenuBuilder.cpp
--- a/Engine/Source/Editor/Kismet/Private/BlueprintActionMenuBuilder.cpp
+++ b/Engine/Source/Editor/Kismet/Private/BlueprintActionMenuBuilder.cpp
@@ -14,6 +14,7 @@
 #include "K2ActionMenuBuilder.h"
 #include "SMyBlueprint.h"				// for SelectionAsVar()
 #include "EdGraphSchema_K2_Actions.h"	// for FEdGraphSchemaAction_K2Var
+#include "BlueprintEditorUtils.h"		// for FindBlueprintForGraphChecked()
 
 #define LOCTEXT_NAMESPACE "BlueprintActionMenuBuilder"
 
diff --git a/Engine/Source/Editor/Kismet/Public/BlueprintActionMenuUtils.h b/Engine/Source/Editor/Kismet/Public/BlueprintActionMenuUtils.h
--- a/Engine/Source/Editor/Kismet/Public/BlueprintActionMenuUtils.h
+++ b/Engine/Source/Editor/Kismet/Public/BlueprintActionMenuUtils.h
@@ -3,7 +3,7 @@
 #pragma once
 
 // Forward declarations
-class  FBlueprintActionMenuBuilder;
+struct FBlueprintActionMenuBuilder;
 struct FBlueprintActionContext;
 
 struct FBlueprintActionMenuUtils
