(WIP) Menu icons now for the new blueprint action menu system.

diff --git a/Engine/Source/Editor/BlueprintGraph/Private/BlueprintNodeSpawner.cpp b/Engine/Source/Editor/BlueprintGraph/Private/BlueprintNodeSpawner.cpp
--- a/Engine/Source/Editor/BlueprintGraph/Private/BlueprintNodeSpawner.cpp
+++ b/Engine/Source/Editor/BlueprintGraph/Private/BlueprintNodeSpawner.cpp
@@ -88,6 +88,16 @@ FText UBlueprintNodeSpawner::GetDefaultSearchKeywords() const
 	return FText::GetEmpty();
 }
 
+//------------------------------------------------------------------------------
+FName UBlueprintNodeSpawner::GetDefaultMenuIcon(FLinearColor& ColorOut)
+{
+	ColorOut = FLinearColor::White;
+	// DO NOT make a template node and query it here (the separate ui building
+	// code can do that if it likes)... this is meant to be an overridable
+	// alternative to that (for perf reasons)
+	return FName();
+}
+
 //------------------------------------------------------------------------------
 UEdGraphNode* UBlueprintNodeSpawner::MakeTemplateNode(UEdGraph* Outer) const
 {
diff --git a/Engine/Source/Editor/BlueprintGraph/Public/BlueprintNodeSpawner.h b/Engine/Source/Editor/BlueprintGraph/Public/BlueprintNodeSpawner.h
--- a/Engine/Source/Editor/BlueprintGraph/Public/BlueprintNodeSpawner.h
+++ b/Engine/Source/Editor/BlueprintGraph/Public/BlueprintNodeSpawner.h
@@ -113,6 +113,18 @@ public:
 	 */
 	virtual FText GetDefaultSearchKeywords() const;
 
+	/**
+	 * To save on performance, certain sub-classes can pull icons info from
+	 * static data (like function metadata, etc.). This unfortunately doesn't
+	 * completely decouple these spawners from the ui we sought to separate
+	 * them from. Although, this can be thought of as a suggestion that the ui
+	 * building code is free to do with as it pleases (doesn't have to use it).
+	 *
+	 * @param  ColorOut		The color to tint the icon with.
+	 * @return Name of the brush to use (use FEditorStyle::GetBrush() to resolve).
+	 */
+	virtual FName GetDefaultMenuIcon(FLinearColor& ColorOut);
+
 	/**
 	 * Expects the supplied outer to be transient. Will spawn a new template-
 	 * node if the cached one is null, or if its graph outer doesn't match the
diff --git a/Engine/Source/Editor/Kismet/Private/BlueprintActionMenuBuilder.cpp b/Engine/Source/Editor/Kismet/Private/BlueprintActionMenuBuilder.cpp
--- a/Engine/Source/Editor/Kismet/Private/BlueprintActionMenuBuilder.cpp
+++ b/Engine/Source/Editor/Kismet/Private/BlueprintActionMenuBuilder.cpp
@@ -121,6 +121,18 @@ protected:
 	 */
 	FText GetSearchKeywordsForAction(FBlueprintActionMenuBuilder& ItemOwner, UBlueprintNodeSpawner* Action);
 
+	/**
+	 * Attempts to pull a menu icon information from the supplied spawner. If
+	 * info isn't provided, then it spawns a temporary node and pulls data from 
+	 * that node.
+	 * 
+	 * @param  ItemOwner	The menu builder that the item will be listed from.
+	 * @param  Action		The action you want to suss icon information from.
+	 * @param  ColorOut		The color to tint the icon with.
+	 * @return Name of the brush to use (use FEditorStyle::GetBrush() to resolve).
+	 */
+	FName GetMenuIconForAction(FBlueprintActionMenuBuilder& ItemOwner, UBlueprintNodeSpawner* Action, FLinearColor& ColorOut);
+
 	/**
 	 * Utility getter function that retrieves the blueprint context for the menu
 	 * items being made.
@@ -164,6 +176,17 @@ TSharedPtr<FEdGraphSchemaAction> FBlueprintActionMenuItemFactory::MakeMenuItem(F
 	NewMenuItem->Category           = GetCategoryForAction(ItemOwner, Action).ToString();
 	NewMenuItem->Keywords           = GetSearchKeywordsForAction(ItemOwner, Action).ToString();
 
+	FLinearColor IconTint = FLinearColor::White;
+	FName IconBrushName = GetMenuIconForAction(ItemOwner, Action, IconTint);
+	if (FSlateBrush const* LibraryBrush = FEditorStyle::GetBrush(IconBrushName))
+	{
+		NewMenuItem->IconBrush = *LibraryBrush;
+		if (IconTint != FLinearColor::White)
+		{
+			NewMenuItem->IconBrush.TintColor = FSlateColor(IconTint);
+		}
+	}
+
 	if (NewMenuItem->Category.IsEmpty())
 	{
 		NewMenuItem->Category = RootCategory.ToString();
@@ -323,6 +346,20 @@ FText FBlueprintActionMenuItemFactory::GetSearchKeywordsForAction(FBlueprintActi
 	return SearchKeywords;
 }
 
+//------------------------------------------------------------------------------
+FName FBlueprintActionMenuItemFactory::GetMenuIconForAction(FBlueprintActionMenuBuilder& ItemOwner, UBlueprintNodeSpawner* Action, FLinearColor& ColorOut)
+{
+	FName BrushName = Action->GetDefaultMenuIcon(ColorOut);
+	if (BrushName.IsNone())
+	{
+		if (UEdGraphNode* NodeTemplate = Action->MakeTemplateNode(GetOwnerOfTemporaries(ItemOwner)))
+		{
+			BrushName = NodeTemplate->GetPaletteIcon(ColorOut);
+		}
+	}
+	return BrushName;
+}
+
 //------------------------------------------------------------------------------
 UBlueprint* FBlueprintActionMenuItemFactory::GetBlueprint() const
 {
diff --git a/Engine/Source/Editor/Kismet/Private/SBlueprintPalette.cpp b/Engine/Source/Editor/Kismet/Private/SBlueprintPalette.cpp
--- a/Engine/Source/Editor/Kismet/Private/SBlueprintPalette.cpp
+++ b/Engine/Source/Editor/Kismet/Private/SBlueprintPalette.cpp
@@ -27,6 +27,8 @@
 #include "AnimationStateGraph.h"
 #include "AnimStateConduitNode.h"
 #include "AnimationTransitionGraph.h"
+#include "BlueprintActionMenuItem.h"
+#include "BlueprintNodeSpawner.h"
 
 #define LOCTEXT_NAMESPACE "BlueprintPalette"
 
@@ -355,6 +357,11 @@ static void GetPaletteItemIcon(TSharedPtr<FEdGraphSchemaAction> ActionIn, UBluep
 		ToolTipOut = LOCTEXT("Comment_Tooltip", "Create a resizable comment box.").ToString();
 		BrushOut = FEditorStyle::GetBrush(TEXT("GraphEditor.Comment_16x"));
 	}
+	else if (ActionIn->GetTypeId() == FBlueprintActionMenuItem::StaticGetTypeId())
+	{
+		FBlueprintActionMenuItem* NodeSpawnerAction = (FBlueprintActionMenuItem*)ActionIn.Get();
+		BrushOut = &NodeSpawnerAction->IconBrush;
+	}
 }
 
 /**
diff --git a/Engine/Source/Editor/Kismet/Public/BlueprintActionMenuItem.h b/Engine/Source/Editor/Kismet/Public/BlueprintActionMenuItem.h
--- a/Engine/Source/Editor/Kismet/Public/BlueprintActionMenuItem.h
+++ b/Engine/Source/Editor/Kismet/Public/BlueprintActionMenuItem.h
@@ -48,4 +48,7 @@ public:
 
 	/** Specialized node-spawner, that comprises the action portion of this menu entry. */
 	UBlueprintNodeSpawner* Action;
+
+	/** Brush that should be used for the icon on this menu item. */
+	FSlateBrush IconBrush;
 };
