(WIP) Fixing consolidated property enteries in the new BP action menu system (not filtering them out, giving them icons, etc.).

diff --git a/Engine/Source/Editor/BlueprintGraph/Classes/BlueprintPropertyNodeSpawner.h b/Engine/Source/Editor/BlueprintGraph/Classes/BlueprintPropertyNodeSpawner.h
--- a/Engine/Source/Editor/BlueprintGraph/Classes/BlueprintPropertyNodeSpawner.h
+++ b/Engine/Source/Editor/BlueprintGraph/Classes/BlueprintPropertyNodeSpawner.h
@@ -33,6 +33,7 @@ public:
 	virtual UEdGraphNode* Invoke(UEdGraph* ParentGraph) const override;
 	virtual FText GetDefaultMenuName() const override;
 	virtual FText GetDefaultMenuCategory() const override;
+	virtual FName GetDefaultMenuIcon(FLinearColor& ColorOut) override;
 	// End UBlueprintNodeSpawner interface
 
 	/**
diff --git a/Engine/Source/Editor/BlueprintGraph/Private/BlueprintPropertyNodeSpawner.cpp b/Engine/Source/Editor/BlueprintGraph/Private/BlueprintPropertyNodeSpawner.cpp
--- a/Engine/Source/Editor/BlueprintGraph/Private/BlueprintPropertyNodeSpawner.cpp
+++ b/Engine/Source/Editor/BlueprintGraph/Private/BlueprintPropertyNodeSpawner.cpp
@@ -110,6 +110,21 @@ FText UBlueprintPropertyNodeSpawner::GetDefaultMenuCategory() const
 	return PropertyCategory;
 }
 
+//------------------------------------------------------------------------------
+FName UBlueprintPropertyNodeSpawner::GetDefaultMenuIcon(FLinearColor& ColorOut)
+{
+	FName BrushName = Super::GetDefaultMenuIcon(ColorOut);
+	if (NodeClass == nullptr)
+	{
+		if (IsDelegateProperty())
+		{
+			BrushName = TEXT("GraphEditor.Delegate_16x");
+		}
+	}
+
+	return BrushName;
+}
+
 //------------------------------------------------------------------------------
 UProperty const* UBlueprintPropertyNodeSpawner::GetProperty() const
 {
diff --git a/Engine/Source/Editor/Kismet/Private/BlueprintActionMenuBuilder.cpp b/Engine/Source/Editor/Kismet/Private/BlueprintActionMenuBuilder.cpp
--- a/Engine/Source/Editor/Kismet/Private/BlueprintActionMenuBuilder.cpp
+++ b/Engine/Source/Editor/Kismet/Private/BlueprintActionMenuBuilder.cpp
@@ -78,6 +78,11 @@ public:
 	 * @return True if the supplied node-spawner is a proxy, false otherwise.
 	 */
 	bool IsProxy(UBlueprintPropertyNodeSpawner* PropertyNodeSpawner) const;
+
+	/**
+	 * Clears out all the proxy items that this spawned (and tracked).
+	 */
+	void Empty();
 	
 protected:
 	/**
@@ -169,24 +174,16 @@ FBlueprintActionMenuItemFactory::FBlueprintActionMenuItemFactory(FBlueprintActio
 //------------------------------------------------------------------------------
 TSharedPtr<FEdGraphSchemaAction> FBlueprintActionMenuItemFactory::MakeMenuItem(FBlueprintActionMenuBuilder& ItemOwner, UBlueprintNodeSpawner* Action)
 {
-	FBlueprintActionMenuItem* NewMenuItem = new FBlueprintActionMenuItem(Action, MenuGrouping);
+	FLinearColor IconTint = FLinearColor::White;
+	FName IconBrushName = GetMenuIconForAction(ItemOwner, Action, IconTint);
+
+	FBlueprintActionMenuItem* NewMenuItem = new FBlueprintActionMenuItem(Action, FEditorStyle::GetBrush(IconBrushName), IconTint, MenuGrouping);
 	
 	NewMenuItem->MenuDescription    = GetMenuNameForAction(ItemOwner, Action);
 	NewMenuItem->TooltipDescription = GetTooltipForAction(ItemOwner, Action).ToString();
 	NewMenuItem->Category           = GetCategoryForAction(ItemOwner, Action).ToString();
 	NewMenuItem->Keywords           = GetSearchKeywordsForAction(ItemOwner, Action).ToString();
 
-	FLinearColor IconTint = FLinearColor::White;
-	FName IconBrushName = GetMenuIconForAction(ItemOwner, Action, IconTint);
-	if (FSlateBrush const* LibraryBrush = FEditorStyle::GetBrush(IconBrushName))
-	{
-		NewMenuItem->IconBrush = *LibraryBrush;
-		if (IconTint != FLinearColor::White)
-		{
-			NewMenuItem->IconBrush.TintColor = FSlateColor(IconTint);
-		}
-	}
-
 	if (NewMenuItem->Category.IsEmpty())
 	{
 		NewMenuItem->Category = RootCategory.ToString();
@@ -242,6 +239,12 @@ bool FBlueprintActionMenuItemFactory::IsProxy(UBlueprintPropertyNodeSpawner* Nod
 	return bIsProxy;
 }
 
+//------------------------------------------------------------------------------
+void FBlueprintActionMenuItemFactory::Empty()
+{
+	PropertyProxyMap.Empty();
+}
+
 //------------------------------------------------------------------------------
 FText FBlueprintActionMenuItemFactory::GetMenuNameForAction(FBlueprintActionMenuBuilder& ItemOwner, UBlueprintNodeSpawner* Action)
 {
@@ -430,6 +433,11 @@ namespace FBlueprintActionMenuBuilderImpl
 		 */
 		TSharedPtr<FEdGraphSchemaAction> MakeMenuItem(FBlueprintActionMenuBuilder& ItemOwner, UBlueprintNodeSpawner* Action);
 		
+		/**
+		 * Clears out any intermediate UBlueprintNodeSpawner actions that this
+		 * instantiated.
+		 */
+		void Empty();
 		
 		/** Series of ESectionFlags, aimed at customizing how we construct this menu section */
 		uint32 Flags;
@@ -509,6 +517,12 @@ TSharedPtr<FEdGraphSchemaAction> FBlueprintActionMenuBuilderImpl::FMenuSectionDe
 	return MenuEntry;
 }
 
+//------------------------------------------------------------------------------
+void FBlueprintActionMenuBuilderImpl::FMenuSectionDefinition::Empty()
+{
+	ItemFactory.Empty();
+}
+
 //------------------------------------------------------------------------------
 void FBlueprintActionMenuBuilderImpl::AppendLegacyItems(FMenuSectionDefinition const& MenuDef, FBlueprintActionMenuBuilder& MenuOut)
 {
@@ -613,6 +627,12 @@ void FBlueprintActionMenuBuilder::RebuildActionList()
 	using namespace FBlueprintActionMenuBuilderImpl;
 
 	FGraphActionListBuilderBase::Empty();
+	for (TSharedRef<FMenuSectionDefinition> MenuSection : MenuSections)
+	{
+		// clear out intermediate actions that may have been spawned (like 
+		// consolidated property actions).
+		MenuSection->Empty();
+	}
 	
 	if (!GetDefault<UEdGraphSchema_K2>()->bUseLegacyActionMenus)
 	{
diff --git a/Engine/Source/Editor/Kismet/Private/BlueprintActionMenuItem.cpp b/Engine/Source/Editor/Kismet/Private/BlueprintActionMenuItem.cpp
--- a/Engine/Source/Editor/Kismet/Private/BlueprintActionMenuItem.cpp
+++ b/Engine/Source/Editor/Kismet/Private/BlueprintActionMenuItem.cpp
@@ -9,6 +9,7 @@
 #include "BlueprintPropertyNodeSpawner.h"
 #include "BPVariableDragDropAction.h"
 #include "BPDelegateDragDropAction.h"
+#include "BlueprintEditor.h"		// for GetVarIconAndColor()
 
 #define LOCTEXT_NAMESPACE "BlueprintActionMenuItem"
 
@@ -55,11 +56,13 @@ static void FBlueprintMenuActionItemImpl::DirtyBlueprintFromNewNode(UEdGraphNode
  ******************************************************************************/
 
 //------------------------------------------------------------------------------
-FBlueprintActionMenuItem::FBlueprintActionMenuItem(UBlueprintNodeSpawner* NodeSpawner, int32 MenuGrouping/* = 0*/)
+FBlueprintActionMenuItem::FBlueprintActionMenuItem(UBlueprintNodeSpawner* NodeSpawner, FSlateBrush const* MenuIcon, FSlateColor const& IconTintIn, int32 MenuGrouping/* = 0*/)
 	: Action(NodeSpawner)
 {
 	check(Action != nullptr);
-	Grouping = MenuGrouping;
+	Grouping  = MenuGrouping;
+	IconBrush = MenuIcon;
+	IconTint  = IconTintIn;
 }
 
 //------------------------------------------------------------------------------
@@ -153,12 +156,32 @@ void FBlueprintActionMenuItem::AddReferencedObjects(FReferenceCollector& Collect
 	Collector.AddReferencedObject(Action);
 }
 
+//------------------------------------------------------------------------------
+FSlateBrush const* FBlueprintActionMenuItem::GetMenuIcon(FSlateColor& ColorOut)
+{
+	// if this brush is invalid
+	if ((IconBrush == nullptr) || !IconBrush->HasUObject())
+	{
+		if (UBlueprintPropertyNodeSpawner const* PropertySpawner = Cast<UBlueprintPropertyNodeSpawner>(Action))
+		{
+			UProperty const* Property = PropertySpawner->GetProperty();
+			FName const PropertyName = Property->GetFName();
+			UStruct* const PropertyOwner = CastChecked<UStruct>(Property->GetOuterUField());
+
+			IconBrush = FBlueprintEditor::GetVarIconAndColor(PropertyOwner, PropertyName, IconTint);
+		}
+	}
+
+	ColorOut = IconTint;
+	return IconBrush;
+}
+
 //------------------------------------------------------------------------------
 TSharedPtr<FDragDropOperation> FBlueprintActionMenuItem::OnDragged(FNodeCreationAnalytic AnalyticsDelegate) const
 {
 	TSharedPtr<FDragDropOperation> DragDropAction = nullptr;
 
-	if (UBlueprintPropertyNodeSpawner* PropertySpawner = Cast<UBlueprintPropertyNodeSpawner>(Action))
+	if (UBlueprintPropertyNodeSpawner const* PropertySpawner = Cast<UBlueprintPropertyNodeSpawner>(Action))
 	{
 		if (PropertySpawner->NodeClass == nullptr)
 		{
@@ -173,10 +196,9 @@ TSharedPtr<FDragDropOperation> FBlueprintActionMenuItem::OnDragged(FNodeCreation
 			}
 			else
 			{
-				TSharedRef<FDragDropOperation> DragDropOpRef = FKismetDelegateDragDropAction::New(PropertyName, PropertyOwner, AnalyticsDelegate);
+				TSharedRef<FDragDropOperation> DragDropOpRef = FKismetVariableDragDropAction::New(PropertyName, PropertyOwner, AnalyticsDelegate);
 				DragDropAction = TSharedPtr<FDragDropOperation>(DragDropOpRef);
 			}
-
 		}
 	}
 	return DragDropAction;
diff --git a/Engine/Source/Editor/Kismet/Private/SBlueprintPalette.cpp b/Engine/Source/Editor/Kismet/Private/SBlueprintPalette.cpp
--- a/Engine/Source/Editor/Kismet/Private/SBlueprintPalette.cpp
+++ b/Engine/Source/Editor/Kismet/Private/SBlueprintPalette.cpp
@@ -360,7 +360,7 @@ static void GetPaletteItemIcon(TSharedPtr<FEdGraphSchemaAction> ActionIn, UBluep
 	else if (ActionIn->GetTypeId() == FBlueprintActionMenuItem::StaticGetTypeId())
 	{
 		FBlueprintActionMenuItem* NodeSpawnerAction = (FBlueprintActionMenuItem*)ActionIn.Get();
-		BrushOut = &NodeSpawnerAction->IconBrush;
+		BrushOut = NodeSpawnerAction->GetMenuIcon(ColorOut);
 	}
 }
 
diff --git a/Engine/Source/Editor/Kismet/Public/BlueprintActionMenuItem.h b/Engine/Source/Editor/Kismet/Public/BlueprintActionMenuItem.h
--- a/Engine/Source/Editor/Kismet/Public/BlueprintActionMenuItem.h
+++ b/Engine/Source/Editor/Kismet/Public/BlueprintActionMenuItem.h
@@ -4,11 +4,13 @@
 
 #include "EdGraph/EdGraphSchema.h" // for FEdGraphSchemaAction
 #include "BlueprintEditor.h"       // for FNodeCreationAnalytic
+#include "SlateColor.h"
 #include "BlueprintActionMenuItem.generated.h"
 
 // Forward declarations
 class UBlueprintNodeSpawner;
 class FDragDropOperation;
+struct FSlateBrush;
 
 /**
  * Wrapper around a UBlueprintNodeSpawner, which takes care of specialized
@@ -25,8 +27,8 @@ public:
 	static FString StaticGetTypeId() { static FString const TypeId = TEXT("FBlueprintActionMenuItem"); return TypeId; }
 	
 	/** Constructors */
-	FBlueprintActionMenuItem() : Action(nullptr) {}
-	FBlueprintActionMenuItem(UBlueprintNodeSpawner* NodeSpawner, int32 MenuGrouping = 0);
+	FBlueprintActionMenuItem() : Action(nullptr), IconBrush(nullptr), IconTint(FLinearColor::White) {}
+	FBlueprintActionMenuItem(UBlueprintNodeSpawner* NodeSpawner, FSlateBrush const* MenuIcon, FSlateColor const& IconTint, int32 MenuGrouping = 0);
 	
 	// FEdGraphSchemaAction interface
 	virtual FString       GetTypeId() const final { return StaticGetTypeId(); }
@@ -35,6 +37,12 @@ public:
 	virtual void          AddReferencedObjects(FReferenceCollector& Collector) final;
 	// End FEdGraphSchemaAction interface
 
+	/**
+	 * Retrieves the icon brush for this menu entry.
+	 * @return An slate brush to be used for this menu item in the action menu.
+	 */
+	FSlateBrush const* GetMenuIcon(FSlateColor& ColorOut);
+
 	/**
 	 * Attempts to create a specific drag/drop action for this menu entry 
 	 * (certain menu entries require unique drag/drop handlers... like property
@@ -47,8 +55,11 @@ public:
 	TSharedPtr<FDragDropOperation> OnDragged(FNodeCreationAnalytic AnalyticsDelegate) const;
 
 	/** Specialized node-spawner, that comprises the action portion of this menu entry. */
-	UBlueprintNodeSpawner* Action;
+	UBlueprintNodeSpawner const* Action;
 
+private:
 	/** Brush that should be used for the icon on this menu item. */
-	FSlateBrush IconBrush;
+	FSlateBrush const* IconBrush;
+	/** Tint to return along with the icon brush. */
+	FSlateColor IconTint;
 };
