Blueprint history will now leave the user in the graphs that were open when the Blueprint was last closed and will provide history to go back up the breadcrumbs to the top-most graph instead.

Added a QuickNavigateCurrentDocument option, basically prevents the history data from being saved to avoid capturing of bad placement of graphs during this process.

diff --git a/Engine/Source/Editor/Kismet/Private/BlueprintEditor.cpp b/Engine/Source/Editor/Kismet/Private/BlueprintEditor.cpp
--- a/Engine/Source/Editor/Kismet/Private/BlueprintEditor.cpp
+++ b/Engine/Source/Editor/Kismet/Private/BlueprintEditor.cpp
@@ -6087,13 +6087,34 @@ void FBlueprintEditor::RestoreEditedObjectState()
 	{
 		if (UObject* Obj = GetBlueprintObj()->LastEditedDocuments[i].EditedObject)
 		{
-
 			if(UEdGraph* Graph = Cast<UEdGraph>(Obj))
 			{
-				UEdGraph* TopLevelGraph = FBlueprintEditorUtils::GetTopLevelGraph(Graph);
-				
-				FDocumentTracker::EOpenDocumentCause OpenCause = TopLevelGraph != Graph? FDocumentTracker::OpenNewDocument : FDocumentTracker::RestorePreviousDocument;
-				TSharedPtr<SDockTab> TabWithGraph = OpenDocument(TopLevelGraph, OpenCause);
+				struct LocalStruct
+				{
+					static TSharedPtr<SDockTab> OpenGraphTree(FBlueprintEditor* InBlueprintEditor, UEdGraph* InGraph)
+					{
+						FDocumentTracker::EOpenDocumentCause OpenCause = FDocumentTracker::QuickNavigateCurrentDocument;
+
+						for (UObject* OuterObject = InGraph->GetOuter(); OuterObject; OuterObject = OuterObject->GetOuter())
+						{
+							if (UBlueprint* Blueprint = Cast<UBlueprint>(OuterObject))
+							{
+								// reached up to the blueprint for the graph, we are done climbing the tree
+								OpenCause = FDocumentTracker::OpenNewDocument;
+								break;
+							}
+							else if(UEdGraph* OuterGraph = Cast<UEdGraph>(OuterObject))
+							{
+								// Found another graph, open it up
+								OpenGraphTree(InBlueprintEditor, OuterGraph);
+								break;
+							}
+						}
+
+						return InBlueprintEditor->OpenDocument(InGraph, OpenCause);
+					}
+				};
+				TSharedPtr<SDockTab> TabWithGraph = LocalStruct::OpenGraphTree(this, Graph);
 
 				TSharedRef<SGraphEditor> GraphEditor = StaticCastSharedRef<SGraphEditor>(TabWithGraph->GetContent());
 				GraphEditor->SetViewLocation(GetBlueprintObj()->LastEditedDocuments[i].SavedViewOffset, GetBlueprintObj()->LastEditedDocuments[i].SavedZoomAmount);
diff --git a/Engine/Source/Editor/Kismet/Private/WorkflowOrientedApp/WorkflowTabManager.cpp b/Engine/Source/Editor/Kismet/Private/WorkflowOrientedApp/WorkflowTabManager.cpp
--- a/Engine/Source/Editor/Kismet/Private/WorkflowOrientedApp/WorkflowTabManager.cpp
+++ b/Engine/Source/Editor/Kismet/Private/WorkflowOrientedApp/WorkflowTabManager.cpp
@@ -37,10 +37,10 @@ bool FTabInfo::PayloadMatches(const TSharedPtr<FTabPayload> TestPayload) const
 	}
 }
 
-void FTabInfo::AddTabHistory(TSharedPtr< struct FGenericTabHistory > InHistoryNode)
+void FTabInfo::AddTabHistory(TSharedPtr< struct FGenericTabHistory > InHistoryNode, bool bInSaveHistory/* = true*/)
 {
 	// If the tab is not new, save the current history.
-	if(CurrentHistoryIndex >= 0)
+	if(CurrentHistoryIndex >= 0 && bInSaveHistory)
 	{
 		History[CurrentHistoryIndex]->SaveHistory();
 	}
@@ -431,7 +431,7 @@ TSharedPtr<SDockTab> FDocumentTracker::OpenDocument(TSharedPtr<FTabPayload> InPa
 		InOpenCause = FDocumentTracker::ForceOpenNewDocument;
 	}
 
-	if(InOpenCause == NavigatingCurrentDocument)
+	if(InOpenCause == NavigatingCurrentDocument || InOpenCause == QuickNavigateCurrentDocument)
 	{
 		return NavigateCurrentTab(InPayload, InOpenCause);
 	}
@@ -497,7 +497,7 @@ TWeakPtr< FTabInfo > FDocumentTracker::FindTabInForeground()
 
 TSharedPtr<SDockTab> FDocumentTracker::NavigateCurrentTab(TSharedPtr<FTabPayload> InPayload, EOpenDocumentCause InNavigateCause)
 {
-	ensure(InNavigateCause == NavigatingCurrentDocument || InNavigateCause == NavigateBackwards || InNavigateCause == NavigateForwards);
+	ensure(InNavigateCause == NavigatingCurrentDocument || InNavigateCause == QuickNavigateCurrentDocument || InNavigateCause == NavigateBackwards || InNavigateCause == NavigateForwards);
 
 	FTabList& List = GetSpawnedList();
 	if(List.Num())
@@ -521,10 +521,11 @@ TSharedPtr<SDockTab> FDocumentTracker::NavigateCurrentTab(TSharedPtr<FTabPayload
 		SpawnInfo.Payload = InPayload;
 		SpawnInfo.TabInfo = LastEditedTabInfo.Pin();
 
-		if(InNavigateCause == NavigatingCurrentDocument)
+		if(InNavigateCause == NavigatingCurrentDocument || InNavigateCause == QuickNavigateCurrentDocument)
 		{
 			TSharedPtr<FDocumentTabFactory> Factory = FindSupportingFactory(InPayload.ToSharedRef());
-			LastEditedTabInfo.Pin()->AddTabHistory(Factory->CreateTabHistoryNode(InPayload));
+			// If doing a Quick navigate of the document, do not save history data as it's likely still at the default values. The object is always saved
+			LastEditedTabInfo.Pin()->AddTabHistory(Factory->CreateTabHistoryNode(InPayload), InNavigateCause != QuickNavigateCurrentDocument);
 		}
 		else if(InNavigateCause == NavigateBackwards)
 		{
diff --git a/Engine/Source/Editor/Kismet/Public/WorkflowOrientedApp/WorkflowTabManager.h b/Engine/Source/Editor/Kismet/Public/WorkflowOrientedApp/WorkflowTabManager.h
--- a/Engine/Source/Editor/Kismet/Public/WorkflowOrientedApp/WorkflowTabManager.h
+++ b/Engine/Source/Editor/Kismet/Public/WorkflowOrientedApp/WorkflowTabManager.h
@@ -59,8 +59,9 @@ public:
 	 * Adds history data to immediately after the current location in the history list, wiping out any history after in the process
 	 *
 	 * @param InHistoryNode		The history node to add
+	 * @param bInSaveHistory	TRUE if history should be saved
 	 */
-	void AddTabHistory(TSharedPtr< struct FGenericTabHistory > InHistoryNode);
+	void AddTabHistory(TSharedPtr< struct FGenericTabHistory > InHistoryNode, bool bInSaveHistory = true);
 
 	/** Single step forward in history */
 	FReply OnGoForwardInHistory();
@@ -227,6 +228,7 @@ public:
 		ForceOpenNewDocument,
 		RestorePreviousDocument,
 		NavigatingCurrentDocument,
+		QuickNavigateCurrentDocument,
 		NavigateBackwards,
 		NavigateForwards,
 		CreateHistoryEvent
