Functions and macros now will display with their category in the My Blueprint panel.

Also mass delete and changing categories allowed to functions, macros, and delegates.

#ttp 341290 - BP:  Add support for categories for functions in the My Blueprints tab

diff --git a/Engine/Source/Editor/GraphEditor/Private/GraphActionNode.h b/Engine/Source/Editor/GraphEditor/Private/GraphActionNode.h
--- a/Engine/Source/Editor/GraphEditor/Private/GraphActionNode.h
+++ b/Engine/Source/Editor/GraphEditor/Private/GraphActionNode.h
@@ -98,7 +98,7 @@ public:
 				TSharedPtr<FGraphActionNode>& CurrentChild = Children[i];
 				if(CurrentChild->IsCategoryNode())
 				{
-					bFoundMatch = NextCategory == CurrentChild->Category;
+					bFoundMatch = NextCategory == CurrentChild->Category && NodeToAdd->SectionID == CurrentChild->SectionID;
 				}
 				else if(CurrentChild->IsActionNode())
 				{
diff --git a/Engine/Source/Editor/Kismet/Private/BlueprintDetailsCustomization.cpp b/Engine/Source/Editor/Kismet/Private/BlueprintDetailsCustomization.cpp
--- a/Engine/Source/Editor/Kismet/Private/BlueprintDetailsCustomization.cpp
+++ b/Engine/Source/Editor/Kismet/Private/BlueprintDetailsCustomization.cpp
@@ -917,6 +917,9 @@ void FBlueprintVarActionDetails::OnTooltipTextCommitted(const FText& NewText, ET
 
 void FBlueprintVarActionDetails::PopulateCategories(SMyBlueprint* MyBlueprint, TArray<TSharedPtr<FString>>& CategorySource)
 {
+	// Used to compare found categories to prevent double adds
+	TArray<FString> CategoryNameList;
+
 	TArray<FName> VisibleVariables;
 	bool bShowUserVarsOnly = MyBlueprint->ShowUserVarsOnly();
 	UBlueprint* Blueprint = MyBlueprint->GetBlueprintObj();
@@ -956,22 +959,89 @@ void FBlueprintVarActionDetails::PopulateCategories(SMyBlueprint* MyBlueprint, T
 		}
 	}
 
-	// Search through all function entry nodes for local variables to pull their categories
-	TArray<UK2Node_FunctionEntry*> FunctionEntryNodes;
-	FBlueprintEditorUtils::GetAllNodesOfClass<UK2Node_FunctionEntry>(Blueprint, FunctionEntryNodes);
+	// Search through all function graphs for entry nodes to search for local variables to pull their categories
+	for( UEdGraph* FunctionGraph : MyBlueprint->GetBlueprintObj()->FunctionGraphs )
+	{
+		if(UFunction* Function = Blueprint->SkeletonGeneratedClass->FindFunctionByName(FunctionGraph->GetFName()))
+		{
+			FString FunctionCategory = Function->GetMetaData(FBlueprintMetadata::MD_FunctionCategory);
+
+			if(!FunctionCategory.IsEmpty())
+			{
+				bool bNewCategory = true;
+				for (int32 j = 0; j < CategorySource.Num() && bNewCategory; ++j)
+				{
+					bNewCategory &= *CategorySource[j].Get() != FunctionCategory;
+				}
 
-	for (UK2Node_FunctionEntry* const FunctionEntry : FunctionEntryNodes)
+				if(bNewCategory)
+				{
+					CategorySource.Add(MakeShareable(new FString(FunctionCategory)));
+				}
+			}
+		}
+
+		TWeakObjectPtr<UK2Node_EditablePinBase> EntryNode;
+		TWeakObjectPtr<UK2Node_EditablePinBase> ResultNode;
+		FBlueprintEditorUtils::GetEntryAndResultNodes(FunctionGraph, EntryNode, ResultNode);
+		if (UK2Node_FunctionEntry* FunctionEntryNode = Cast<UK2Node_FunctionEntry>(EntryNode.Get()))
+		{
+			for( FBPVariableDescription& Variable : FunctionEntryNode->LocalVariables )
+			{
+				bool bNewCategory = true;
+				for (int32 j = 0; j < CategorySource.Num() && bNewCategory; ++j)
+				{
+					bNewCategory &= *CategorySource[j].Get() != Variable.Category.ToString();
+				}
+				if (bNewCategory)
+				{
+					CategorySource.Add(MakeShareable(new FString(Variable.Category.ToString())));
+				}
+			}
+		}
+	}
+
+	for( UEdGraph* MacroGraph : MyBlueprint->GetBlueprintObj()->MacroGraphs )
 	{
-		for( FBPVariableDescription& Variable : FunctionEntry->LocalVariables )
+		TWeakObjectPtr<UK2Node_EditablePinBase> EntryNode;
+		TWeakObjectPtr<UK2Node_EditablePinBase> ResultNode;
+		FBlueprintEditorUtils::GetEntryAndResultNodes(MacroGraph, EntryNode, ResultNode);
+		if (UK2Node_Tunnel* TypedEntryNode = ExactCast<UK2Node_Tunnel>(EntryNode.Get()))
 		{
 			bool bNewCategory = true;
 			for (int32 j = 0; j < CategorySource.Num() && bNewCategory; ++j)
 			{
-				bNewCategory &= *CategorySource[j].Get() != Variable.Category.ToString();
+				bNewCategory &= *CategorySource[j].Get() != TypedEntryNode->MetaData.Category;
 			}
 			if (bNewCategory)
 			{
-				CategorySource.Add(MakeShareable(new FString(Variable.Category.ToString())));
+				CategorySource.Add(MakeShareable(new FString(TypedEntryNode->MetaData.Category)));
+			}
+		}
+	}
+
+	// Pull categories from overridable functions
+	for (TFieldIterator<UFunction> FunctionIt(MyBlueprint->GetBlueprintObj()->ParentClass, EFieldIteratorFlags::IncludeSuper); FunctionIt; ++FunctionIt)
+	{
+		const UFunction* Function = *FunctionIt;
+		const FName FunctionName = Function->GetFName();
+
+		if (UEdGraphSchema_K2::CanKismetOverrideFunction(Function) && !UEdGraphSchema_K2::FunctionCanBePlacedAsEvent(Function))
+		{
+			FString FunctionCategory = Function->GetMetaData(FBlueprintMetadata::MD_FunctionCategory);
+
+			if(!FunctionCategory.IsEmpty())
+			{
+				bool bNewCategory = true;
+				for (int32 j = 0; j < CategorySource.Num() && bNewCategory; ++j)
+				{
+					bNewCategory &= *CategorySource[j].Get() != FunctionCategory;
+				}
+
+				if(bNewCategory)
+				{
+					CategorySource.Add(MakeShareable(new FString(FunctionCategory)));
+				}
 			}
 		}
 	}
@@ -1049,10 +1119,10 @@ FText FBlueprintVarActionDetails::OnGetCategoryText() const
 
 		FName Category = FBlueprintEditorUtils::GetBlueprintVariableCategory(GetBlueprintObj(), VarName, GetLocalVariableScope(SelectionAsProperty()));
 
-		// Older blueprints will have their name as the default category
-		if( Category == GetBlueprintObj()->GetFName() )
+		// Older blueprints will have their name as the default category and whenever it is the same as the default category, display localized text
+		if( Category == GetBlueprintObj()->GetFName() || Category == K2Schema->VR_DefaultCategory )
 		{
-			return FText::FromName(K2Schema->VR_DefaultCategory);
+			return LOCTEXT("DefaultCategory", "Default");
 		}
 		else
 		{
@@ -1067,7 +1137,10 @@ void FBlueprintVarActionDetails::OnCategoryTextCommitted(const FText& NewText, E
 {
 	if (InTextCommit == ETextCommit::OnEnter || InTextCommit == ETextCommit::OnUserMovedFocus)
 	{
-		FString NewCategory = NewText.ToString();
+		// Remove excess whitespace and prevent categories with just spaces
+		FText CategoryName = FText::TrimPrecedingAndTrailing(NewText);
+
+		FString NewCategory = CategoryName.ToString();
 		if(NewCategory.Len() <= NAME_SIZE)
 		{
 			FBlueprintEditorUtils::SetBlueprintVariableCategory(GetBlueprintObj(), VarName, GetLocalVariableScope(SelectionAsProperty()), FName( *NewCategory ));
@@ -1968,21 +2041,60 @@ void FBlueprintGraphActionDetails::CustomizeDetails( IDetailLayoutBuilder& Detai
 					.OnTextCommitted( this, &FBlueprintGraphActionDetails::OnTooltipTextCommitted )
 					.Font( IDetailLayoutBuilder::GetDetailFont() )
 			];
-			Category.AddCustomRow( LOCTEXT( "Category", "Category" ).ToString() )
-			.NameContent()
-			[
-				SNew(STextBlock)
-					.Text( LOCTEXT( "Category", "Category" ).ToString() )
-					.ToolTipText(LOCTEXT("CategoryTooltip", "The category affects how this function will show up in menus and when searching").ToString())
+
+			FBlueprintVarActionDetails::PopulateCategories(MyBlueprint.Pin().Get(), CategorySource);
+			TSharedPtr<SComboButton> NewComboButton;
+			TSharedPtr<SListView<TSharedPtr<FString>>> NewListView;
+
+			const FString DocLink = TEXT("Shared/Editors/BlueprintEditor/VariableDetails");
+			TSharedPtr<SToolTip> CategoryTooltip = IDocumentation::Get()->CreateToolTip(LOCTEXT("EditCategoryName_Tooltip", "The category of the variable; editing this will place the variable into another category or create a new one."), NULL, DocLink, TEXT("Category"));
+
+			Category.AddCustomRow( TEXT("Category") )
+				.NameContent()
+				[
+					SNew(STextBlock)
+					.Text( LOCTEXT("CategoryLabel", "Category") )
+					.ToolTip(CategoryTooltip)
 					.Font( IDetailLayoutBuilder::GetDetailFont() )
-			]
+				]
 			.ValueContent()
-			[
-				SNew(SEditableTextBox)
-					.Text( this, &FBlueprintGraphActionDetails::OnGetCategoryText )
-					.OnTextCommitted( this, &FBlueprintGraphActionDetails::OnCategoryTextCommitted )
-					.Font( IDetailLayoutBuilder::GetDetailFont() )
-			];
+				[
+					SAssignNew(NewComboButton, SComboButton)
+					.ContentPadding(FMargin(0,0,5,0))
+					.ToolTip(CategoryTooltip)
+					.ButtonContent()
+					[
+						SNew(SBorder)
+						.BorderImage( FEditorStyle::GetBrush("NoBorder") )
+						.Padding(FMargin(0, 0, 5, 0))
+						[
+							SNew(SEditableTextBox)
+							.Text(this, &FBlueprintGraphActionDetails::OnGetCategoryText)
+							.OnTextCommitted(this, &FBlueprintGraphActionDetails::OnCategoryTextCommitted )
+							.ToolTip(CategoryTooltip)
+							.SelectAllTextWhenFocused(true)
+							.RevertTextOnEscape(true)
+							.Font( IDetailLayoutBuilder::GetDetailFont() )
+						]
+					]
+					.MenuContent()
+						[
+							SNew(SVerticalBox)
+							+SVerticalBox::Slot()
+							.AutoHeight()
+							.MaxHeight(400.0f)
+							[
+								SAssignNew(NewListView, SListView<TSharedPtr<FString>>)
+								.ListItemsSource(&CategorySource)
+								.OnGenerateRow(this, &FBlueprintGraphActionDetails::MakeCategoryViewWidget)
+								.OnSelectionChanged(this, &FBlueprintGraphActionDetails::OnCategorySelectionChanged)
+							]
+						]
+				];
+
+			CategoryComboButton = NewComboButton;
+			CategoryListView = NewListView;
+
 			if (IsAccessSpecifierVisible())
 			{
 				Category.AddCustomRow( LOCTEXT( "AccessSpecifier", "Access Specifier" ).ToString() )
@@ -2427,9 +2539,9 @@ FText FBlueprintDelegateActionDetails::OnGetCategoryText() const
 		FName Category = FBlueprintEditorUtils::GetBlueprintVariableCategory(GetBlueprintObj(), DelegateName, NULL);
 
 		// Older blueprints will have their name as the default category
-		if( Category == GetBlueprintObj()->GetFName() )
+		if( Category == GetBlueprintObj()->GetFName() || Category == K2Schema->VR_DefaultCategory )
 		{
-			return FText::FromName(K2Schema->VR_DefaultCategory);
+			return LOCTEXT("DefaultCategory", "Default");
 		}
 		else
 		{
@@ -2446,7 +2558,9 @@ void FBlueprintDelegateActionDetails::OnCategoryTextCommitted(const FText& NewTe
 	{
 		if (UMulticastDelegateProperty* DelegateProperty = GetDelegatePoperty())
 		{
-			FString NewCategory = NewText.ToString();
+			// Remove excess whitespace and prevent categories with just spaces
+			FText CategoryName = FText::TrimPrecedingAndTrailing(NewText);
+			FString NewCategory = CategoryName.ToString();
 			
 			FBlueprintEditorUtils::SetBlueprintVariableCategory(GetBlueprintObj(), DelegateProperty->GetFName(), NULL, FName( *NewCategory ));
 			check(MyBlueprint.IsValid());
@@ -3092,6 +3206,12 @@ FText FBlueprintGraphActionDetails::OnGetCategoryText() const
 {
 	if (FKismetUserDeclaredFunctionMetadata* Metadata = GetMetadataBlock())
 	{
+		const UEdGraphSchema_K2* K2Schema = GetDefault<UEdGraphSchema_K2>();
+		if( Metadata->Category.IsEmpty() || Metadata->Category == K2Schema->VR_DefaultCategory.ToString() )
+		{
+			return LOCTEXT("DefaultCategory", "Default");
+		}
+		
 		return FText::FromString(Metadata->Category);
 	}
 	else
@@ -3102,17 +3222,64 @@ FText FBlueprintGraphActionDetails::OnGetCategoryText() const
 
 void FBlueprintGraphActionDetails::OnCategoryTextCommitted(const FText& NewText, ETextCommit::Type InTextCommit)
 {
-	if (FKismetUserDeclaredFunctionMetadata* Metadata = GetMetadataBlock())
+	if (InTextCommit == ETextCommit::OnEnter || InTextCommit == ETextCommit::OnUserMovedFocus)
 	{
-		Metadata->Category = NewText.ToString();
-		if(auto Function = FindFunction())
+		if (FKismetUserDeclaredFunctionMetadata* Metadata = GetMetadataBlock())
 		{
-			Function->Modify();
-			Function->SetMetaData(FBlueprintMetadata::MD_FunctionCategory, *NewText.ToString());
+			// Remove excess whitespace and prevent categories with just spaces
+			FText CategoryName = FText::TrimPrecedingAndTrailing(NewText);
+
+			if(CategoryName.IsEmpty())
+			{
+				const UEdGraphSchema_K2* K2Schema = GetDefault<UEdGraphSchema_K2>();
+				Metadata->Category = K2Schema->VR_DefaultCategory.ToString();
+			}
+			else
+			{
+				Metadata->Category = CategoryName.ToString();
+			}
+
+			if(auto Function = FindFunction())
+			{
+				Function->Modify();
+				Function->SetMetaData(FBlueprintMetadata::MD_FunctionCategory, *CategoryName.ToString());
+			}
+			MyBlueprint.Pin()->Refresh();
+			FBlueprintEditorUtils::MarkBlueprintAsModified(GetBlueprintObj());
+		}
+	}
+}
+
+void FBlueprintGraphActionDetails::OnCategorySelectionChanged( TSharedPtr<FString> ProposedSelection, ESelectInfo::Type /*SelectInfo*/ )
+{
+	if(ProposedSelection.IsValid())
+	{
+		if (FKismetUserDeclaredFunctionMetadata* Metadata = GetMetadataBlock())
+		{
+			Metadata->Category = *ProposedSelection.Get();
+			if(auto Function = FindFunction())
+			{
+				Function->Modify();
+				Function->SetMetaData(FBlueprintMetadata::MD_FunctionCategory, **ProposedSelection.Get());
+			}
+			MyBlueprint.Pin()->Refresh();
+			FBlueprintEditorUtils::MarkBlueprintAsModified(GetBlueprintObj());
+
+			CategoryListView.Pin()->ClearSelection();
+			CategoryComboButton.Pin()->SetIsOpen(false);
+			MyBlueprint.Pin()->ExpandCategory(*ProposedSelection.Get());
 		}
 	}
 }
 
+TSharedRef< ITableRow > FBlueprintGraphActionDetails::MakeCategoryViewWidget( TSharedPtr<FString> Item, const TSharedRef< STableViewBase >& OwnerTable )
+{
+	return SNew(STableRow<TSharedPtr<FString>>, OwnerTable)
+		[
+			SNew(STextBlock) .Text(*Item.Get())
+		];
+}
+
 FText FBlueprintGraphActionDetails::AccessSpecifierProperName( uint32 AccessSpecifierFlag ) const
 {
 	switch(AccessSpecifierFlag)
diff --git a/Engine/Source/Editor/Kismet/Private/BlueprintDetailsCustomization.h b/Engine/Source/Editor/Kismet/Private/BlueprintDetailsCustomization.h
--- a/Engine/Source/Editor/Kismet/Private/BlueprintDetailsCustomization.h
+++ b/Engine/Source/Editor/Kismet/Private/BlueprintDetailsCustomization.h
@@ -448,6 +448,12 @@ private:
 	/** Called to set the replication type from the details view combo */
 	static void SetNetFlags( TWeakObjectPtr<UK2Node_EditablePinBase> FunctionEntryNode, uint32 NetFlags);
 
+	/** Callback when a graph category is changed */
+	void OnCategorySelectionChanged( TSharedPtr<FString> ProposedSelection, ESelectInfo::Type /*SelectInfo*/ );
+
+	/** Callback to make category widgets */
+	TSharedRef< ITableRow > MakeCategoryViewWidget( TSharedPtr<FString> Item, const TSharedRef< STableViewBase >& OwnerTable );
+
 private:
 
 	/** List of available localized access specifiers names */
@@ -458,6 +464,13 @@ private:
 
 	/** Color block for parenting the color picker */
 	TSharedPtr<class SColorBlock> ColorBlock;
+
+	/** A list of all category names to choose from */
+	TArray<TSharedPtr<FString>> CategorySource;
+
+	/** Widgets for the categories */
+	TWeakPtr<SComboButton> CategoryComboButton;
+	TWeakPtr<SListView<TSharedPtr<FString>>> CategoryListView;
 };
 
 /** Blueprint Interface List Details */
diff --git a/Engine/Source/Editor/Kismet/Private/SMyBlueprint.cpp b/Engine/Source/Editor/Kismet/Private/SMyBlueprint.cpp
--- a/Engine/Source/Editor/Kismet/Private/SMyBlueprint.cpp
+++ b/Engine/Source/Editor/Kismet/Private/SMyBlueprint.cpp
@@ -280,6 +280,9 @@ void SMyBlueprint::Construct(const FArguments& InArgs, TWeakPtr<FBlueprintEditor
 
 void SMyBlueprint::OnCategoryNameCommitted(const FText& InNewText, ETextCommit::Type InTextCommit, TWeakPtr< FGraphActionNode > InAction )
 {
+	// Remove excess whitespace and prevent categories with just spaces
+	FText CategoryName = FText::TrimPrecedingAndTrailing(InNewText);
+
 	TArray<TSharedPtr<FEdGraphSchemaAction>> Actions;
 	GraphActionMenu->GetCategorySubActions(InAction, Actions);
 
@@ -295,16 +298,62 @@ void SMyBlueprint::OnCategoryNameCommitted(const FText& InNewText, ETextCommit::
 			{
 				FEdGraphSchemaAction_K2Var* VarAction = (FEdGraphSchemaAction_K2Var*)Actions[i].Get();
 
-				FBlueprintEditorUtils::SetBlueprintVariableCategory(GetBlueprintObj(), VarAction->GetVariableName(), NULL, FName( *InNewText.ToString() ), true);
+				if(UProperty* TargetProperty = VarAction->GetProperty())
+				{
+					UClass* OuterClass = CastChecked<UClass>(VarAction->GetProperty()->GetOuter());
+					const bool bIsNativeVar = (OuterClass->ClassGeneratedBy == NULL);
+
+					// If the variable is not native and it's outer is the skeleton generated class, we can rename the category
+					if(!bIsNativeVar && OuterClass == GetBlueprintObj()->SkeletonGeneratedClass)
+					{
+						FBlueprintEditorUtils::SetBlueprintVariableCategory(GetBlueprintObj(), VarAction->GetVariableName(), NULL, FName( *CategoryName.ToString() ), true);
+					}
+				}
 			}
 			else if (Actions[i]->GetTypeId() == FEdGraphSchemaAction_K2LocalVar::StaticGetTypeId())
 			{
 				FEdGraphSchemaAction_K2LocalVar* LocalVarAction = (FEdGraphSchemaAction_K2LocalVar*)Actions[i].Get();
 
-				FBlueprintEditorUtils::SetBlueprintVariableCategory(GetBlueprintObj(), LocalVarAction->GetVariableName(), LocalVarAction->GetVariableScope(), FName( *InNewText.ToString() ), true);
+				FBlueprintEditorUtils::SetBlueprintVariableCategory(GetBlueprintObj(), LocalVarAction->GetVariableName(), LocalVarAction->GetVariableScope(), FName( *CategoryName.ToString() ), true);
+			}
+			else if (Actions[i]->GetTypeId() == FEdGraphSchemaAction_K2Delegate::StaticGetTypeId())
+			{
+				FEdGraphSchemaAction_K2Delegate* DelegateAction = (FEdGraphSchemaAction_K2Delegate*)Actions[i].Get();
+				FBlueprintEditorUtils::SetBlueprintVariableCategory(GetBlueprintObj(), DelegateAction->GetDelegatePoperty()->GetFName(), NULL, FName( *CategoryName.ToString() ), true);
+			}
+			else if (Actions[i]->GetTypeId() == FEdGraphSchemaAction_K2Graph::StaticGetTypeId())
+			{
+				// Do not allow renaming of any graph actions outside of the following
+				if(Actions[i]->SectionID == NodeSectionID::FUNCTION || Actions[i]->SectionID == NodeSectionID::MACRO)
+				{
+					FEdGraphSchemaAction_K2Graph* GraphAction = (FEdGraphSchemaAction_K2Graph*)Actions[i].Get();
+
+					// Don't allow changing the category of a graph who's parent is not the current Blueprint
+					if(!FBlueprintEditorUtils::IsPaletteActionReadOnly(Actions[i], BlueprintEditorPtr.Pin()) && FBlueprintEditorUtils::FindBlueprintForGraph(GraphAction->EdGraph) == GetBlueprintObj())
+					{
+						TWeakObjectPtr<UK2Node_EditablePinBase> EntryNode;
+						TWeakObjectPtr<UK2Node_EditablePinBase> ResultNode;
+						FBlueprintEditorUtils::GetEntryAndResultNodes(GraphAction->EdGraph, EntryNode, ResultNode);
+						EntryNode->Modify();
+						if (UK2Node_FunctionEntry* FunctionEntryNode = Cast<UK2Node_FunctionEntry>(EntryNode.Get()))
+						{
+							FunctionEntryNode->MetaData.Category = CategoryName.ToString();
+						}
+						else if (UK2Node_Tunnel* TypedEntryNode = ExactCast<UK2Node_Tunnel>(EntryNode.Get()))
+						{
+							TypedEntryNode->MetaData.Category = CategoryName.ToString();
+						}
+
+						if(UFunction* Function = GetBlueprintObj()->SkeletonGeneratedClass->FindFunctionByName(GraphAction->EdGraph->GetFName()))
+						{
+							Function->SetMetaData(FBlueprintMetadata::MD_FunctionCategory, *CategoryName.ToString());
+						}
+					}
+				}
 			}
 		}
-		FBlueprintEditorUtils::MarkBlueprintAsStructurallyModified(GetBlueprintObj());
+		Refresh();
+		FBlueprintEditorUtils::MarkBlueprintAsModified(GetBlueprintObj());
 	}
 }
 
@@ -354,14 +403,11 @@ FText SMyBlueprint::OnGetSectionTitle( int32 InSectionID )
 
 bool SMyBlueprint::CanRequestRenameOnActionNode(TWeakPtr<FGraphActionNode> InSelectedNode) const
 {
-	bool bHasNonNativeVariables = true;
-	bool bIsReadOnly = false;
+	bool bIsReadOnly = true;
 
 	// If checking if renaming is available on a category node, the category must have a non-native variable.
 	if(InSelectedNode.Pin()->IsCategoryNode())
 	{
-		bHasNonNativeVariables = false;
-
 		TArray<TSharedPtr<FEdGraphSchemaAction>> Actions;
 		GraphActionMenu->GetCategorySubActions(InSelectedNode, Actions);
 
@@ -377,13 +423,36 @@ bool SMyBlueprint::CanRequestRenameOnActionNode(TWeakPtr<FGraphActionNode> InSel
 					UClass* OuterClass = CastChecked<UClass>(TargetProperty->GetOuter());
 					const bool bIsNativeVar = (OuterClass->ClassGeneratedBy == NULL);
 
-					if(!bIsNativeVar)
+					// If the variable is not native and it's outer is the skeleton generated class, we can rename the category
+					if(!bIsNativeVar && OuterClass == GetBlueprintObj()->SkeletonGeneratedClass)
 					{
-						bHasNonNativeVariables = true;
+						bIsReadOnly = false;
 						break;
 					}
 				}
 			}
+			else if (Actions[i]->GetTypeId() == FEdGraphSchemaAction_K2Graph::StaticGetTypeId())
+			{
+				FEdGraphSchemaAction_K2Graph* GraphAction = (FEdGraphSchemaAction_K2Graph*)(Actions[i].Get());
+
+				// Do not allow renaming of any graph actions outside of the following
+				if(Actions[i]->SectionID == NodeSectionID::FUNCTION || Actions[i]->SectionID == NodeSectionID::MACRO)
+				{
+					if(!FBlueprintEditorUtils::IsPaletteActionReadOnly(Actions[i], BlueprintEditorPtr.Pin()) && FBlueprintEditorUtils::FindBlueprintForGraph(GraphAction->EdGraph) == GetBlueprintObj())
+					{
+						bIsReadOnly = false;
+						break;
+					}
+				}
+			}
+			else if(Actions[i]->GetTypeId() == FEdGraphSchemaAction_K2Delegate::StaticGetTypeId())
+			{
+				if(!FBlueprintEditorUtils::IsPaletteActionReadOnly(Actions[i], BlueprintEditorPtr.Pin()))
+				{
+					bIsReadOnly = false;
+					break;
+				}
+			}
 		}
 	}
 	else if(InSelectedNode.Pin()->IsActionNode())
@@ -392,7 +461,7 @@ bool SMyBlueprint::CanRequestRenameOnActionNode(TWeakPtr<FGraphActionNode> InSel
 		bIsReadOnly = FBlueprintEditorUtils::IsPaletteActionReadOnly(InSelectedNode.Pin()->Actions[0], BlueprintEditorPtr.Pin());
 	}
 
-	return BlueprintEditorPtr.Pin()->InEditingMode() && bHasNonNativeVariables && !bIsReadOnly;
+	return BlueprintEditorPtr.Pin()->InEditingMode() && !bIsReadOnly;
 }
 
 void SMyBlueprint::Refresh()
@@ -702,7 +771,13 @@ void SMyBlueprint::CollectAllActions(FGraphActionListBuilderBase& OutAllActions)
 		FGraphDisplayInfo DisplayInfo;
 		Graph->GetSchema()->GetGraphDisplayInformation(*Graph, DisplayInfo);
 
-		TSharedPtr<FEdGraphSchemaAction_K2Graph> NewFuncAction = MakeShareable(new FEdGraphSchemaAction_K2Graph(EEdGraphSchemaAction_K2Graph::Function, FString(), DisplayInfo.PlainName, DisplayInfo.Tooltip, bIsConstructionScript ? 2 : 1));
+		FString FunctionCategory;
+		if(UFunction* Function = Blueprint->SkeletonGeneratedClass->FindFunctionByName(Graph->GetFName()))
+		{
+			FunctionCategory = Function->GetMetaData(FBlueprintMetadata::MD_FunctionCategory);
+		}
+
+		TSharedPtr<FEdGraphSchemaAction_K2Graph> NewFuncAction = MakeShareable(new FEdGraphSchemaAction_K2Graph(EEdGraphSchemaAction_K2Graph::Function, FunctionCategory, DisplayInfo.PlainName, DisplayInfo.Tooltip, bIsConstructionScript ? 2 : 1));
 		NewFuncAction->FuncName = Graph->GetFName();
 		NewFuncAction->EdGraph = Graph;
 		NewFuncAction->SectionID = NodeSectionID::FUNCTION;
@@ -731,7 +806,9 @@ void SMyBlueprint::CollectAllActions(FGraphActionListBuilderBase& OutAllActions)
 		FGraphDisplayInfo DisplayInfo;
 		Graph->GetSchema()->GetGraphDisplayInformation(*Graph, DisplayInfo);
 
-		TSharedPtr<FEdGraphSchemaAction_K2Graph> NewFuncAction = MakeShareable(new FEdGraphSchemaAction_K2Graph(EEdGraphSchemaAction_K2Graph::Macro, FString(), DisplayInfo.PlainName, DisplayInfo.Tooltip, 1));
+		FString MacroCategory = GetGraphCategory(Graph);
+
+		TSharedPtr<FEdGraphSchemaAction_K2Graph> NewFuncAction = MakeShareable(new FEdGraphSchemaAction_K2Graph(EEdGraphSchemaAction_K2Graph::Macro, MacroCategory, DisplayInfo.PlainName, DisplayInfo.Tooltip, 1));
 		NewFuncAction->SectionID = NodeSectionID::MACRO;
 		NewFuncAction->FuncName = FunctionName;
 		NewFuncAction->EdGraph = Graph;
@@ -779,7 +856,14 @@ void SMyBlueprint::CollectAllActions(FGraphActionListBuilderBase& OutAllActions)
 			const FName FunctionName = Graph->GetFName();
 			FString FunctionTooltip = FunctionName.ToString();
 			FString FunctionDesc = FunctionName.ToString();
-			TSharedPtr<FEdGraphSchemaAction_K2Graph> NewFuncAction = MakeShareable(new FEdGraphSchemaAction_K2Graph(EEdGraphSchemaAction_K2Graph::Interface, FString(), FText::FromString(FunctionDesc), FunctionTooltip, 1));
+			
+			FString FunctionCategory;
+			if(UFunction* Function = Blueprint->SkeletonGeneratedClass->FindFunctionByName(Graph->GetFName()))
+			{
+				FunctionCategory = Function->GetMetaData(FBlueprintMetadata::MD_FunctionCategory);
+			}
+
+			TSharedPtr<FEdGraphSchemaAction_K2Graph> NewFuncAction = MakeShareable(new FEdGraphSchemaAction_K2Graph(EEdGraphSchemaAction_K2Graph::Interface, FunctionCategory, FText::FromString(FunctionDesc), FunctionTooltip, 1));
 			NewFuncAction->FuncName = FunctionName;
 			NewFuncAction->EdGraph = Graph;
 			NewFuncAction->SectionID = NodeSectionID::INTERFACE;
@@ -1223,6 +1307,33 @@ bool SMyBlueprint::SelectionHasContextMenu() const
 	return bSelectionHasContextMenu;
 }
 
+FString SMyBlueprint::GetGraphCategory(UEdGraph* InGraph) const
+{
+	FString ReturnCategory;
+
+	// Pull the category from the required metadata based on the types of nodes we can discover in the graph
+	TWeakObjectPtr<UK2Node_EditablePinBase> EntryNode;
+	TWeakObjectPtr<UK2Node_EditablePinBase> ResultNode;
+	FBlueprintEditorUtils::GetEntryAndResultNodes(InGraph, EntryNode, ResultNode);
+	if (UK2Node_FunctionEntry* FunctionEntryNode = Cast<UK2Node_FunctionEntry>(EntryNode.Get()))
+	{
+		ReturnCategory = FunctionEntryNode->MetaData.Category;
+	}
+	else if (UK2Node_Tunnel* TypedEntryNode = ExactCast<UK2Node_Tunnel>(EntryNode.Get()))
+	{
+		ReturnCategory = TypedEntryNode->MetaData.Category;
+	}
+
+	// Empty the category if it's default, we don't want to display the "default" category and items will just appear without a category
+	const UEdGraphSchema_K2* K2Schema = GetDefault<UEdGraphSchema_K2>();
+	if(ReturnCategory == K2Schema->VR_DefaultCategory.ToString())
+	{
+		ReturnCategory.Empty();
+	}
+
+	return ReturnCategory;
+}
+
 void SMyBlueprint::GetSelectedItemsForContextMenu(TArray<FComponentEventConstructionData>& OutSelectedItems) const
 {
 	FEdGraphSchemaAction_K2Var* Var = SelectionAsVar();
@@ -1460,67 +1571,76 @@ bool SMyBlueprint::CanFindEntry() const
 	return true;
 }
 
-void SMyBlueprint::OnDeleteEntry()
+void SMyBlueprint::OnDeleteGraph(UEdGraph* InGraph, EEdGraphSchemaAction_K2Graph::Type InGraphType)
 {
-	if (FEdGraphSchemaAction_K2Graph* GraphAction = SelectionAsGraph())
+	if (InGraph)
 	{
-		UEdGraph* EdGraph = GraphAction->EdGraph;
-		if (EdGraph)
+		const FScopedTransaction Transaction( LOCTEXT("RemoveGraph", "Remove Graph") );
+		GetBlueprintObj()->Modify();
+
+		InGraph->Modify();
+
+		if (InGraphType == EEdGraphSchemaAction_K2Graph::Subgraph)
 		{
-			const FScopedTransaction Transaction( LOCTEXT("RemoveGraph", "Remove Graph") );
-			GetBlueprintObj()->Modify();
+			// Remove any composite nodes bound to this graph
+			TArray<UK2Node_Composite*> AllCompositeNodes;
+			FBlueprintEditorUtils::GetAllNodesOfClass<UK2Node_Composite>(GetBlueprintObj(), AllCompositeNodes);
 
-			EdGraph->Modify();
-				
-			if (GraphAction->GraphType == EEdGraphSchemaAction_K2Graph::Subgraph)
+			const bool bDontRecompile = true;
+			for (auto CompIt = AllCompositeNodes.CreateIterator(); CompIt; ++CompIt)
 			{
-				// Remove any composite nodes bound to this graph
-				TArray<UK2Node_Composite*> AllCompositeNodes;
-				FBlueprintEditorUtils::GetAllNodesOfClass<UK2Node_Composite>(GetBlueprintObj(), AllCompositeNodes);
-
-				const bool bDontRecompile = true;
-				for (auto CompIt = AllCompositeNodes.CreateIterator(); CompIt; ++CompIt)
+				UK2Node_Composite* CompNode = *CompIt;
+				if (CompNode->BoundGraph == InGraph)
 				{
-					UK2Node_Composite* CompNode = *CompIt;
-					if (CompNode->BoundGraph == EdGraph)
-					{
-						FBlueprintEditorUtils::RemoveNode(GetBlueprintObj(), CompNode, bDontRecompile);
-					}
+					FBlueprintEditorUtils::RemoveNode(GetBlueprintObj(), CompNode, bDontRecompile);
 				}
 			}
+		}
 
-			FBlueprintEditorUtils::RemoveGraph(GetBlueprintObj(), EdGraph, EGraphRemoveFlags::Recompile);
-			BlueprintEditorPtr.Pin()->CloseDocumentTab(EdGraph);
-
-			for (TObjectIterator<UK2Node_CreateDelegate> It; It; ++It)
-			{
-				It->HandleAnyChange();
-			}
+		FBlueprintEditorUtils::RemoveGraph(GetBlueprintObj(), InGraph, EGraphRemoveFlags::Recompile);
+		BlueprintEditorPtr.Pin()->CloseDocumentTab(InGraph);
 
-			EdGraph = NULL;
+		for (TObjectIterator<UK2Node_CreateDelegate> It; It; ++It)
+		{
+			It->HandleAnyChange();
 		}
+
+		InGraph = NULL;
 	}
-	if (FEdGraphSchemaAction_K2Delegate* GraphAction = SelectionAsDelegate())
+}
+
+void SMyBlueprint::OnDeleteDelegate(FEdGraphSchemaAction_K2Delegate* InDelegateAction)
+{
+	UEdGraph* EdGraph = InDelegateAction->EdGraph;
+	UBlueprint* Blueprint = GetBlueprintObj();
+	if (EdGraph && Blueprint)
 	{
-		UEdGraph* EdGraph = GraphAction->EdGraph;
-		UBlueprint* Blueprint = GetBlueprintObj();
-		if (EdGraph && Blueprint)
-		{
-			const FScopedTransaction Transaction( LOCTEXT("RemoveDelegate", "Remove Event Dispatcher") );
-			Blueprint->Modify();
+		const FScopedTransaction Transaction( LOCTEXT("RemoveDelegate", "Remove Event Dispatcher") );
+		Blueprint->Modify();
 
-			BlueprintEditorPtr.Pin()->CloseDocumentTab(EdGraph);
-			EdGraph->Modify();
+		BlueprintEditorPtr.Pin()->CloseDocumentTab(EdGraph);
+		EdGraph->Modify();
 
-			FBlueprintEditorUtils::RemoveMemberVariable(Blueprint, EdGraph->GetFName());
-			FBlueprintEditorUtils::RemoveGraph(Blueprint, EdGraph, EGraphRemoveFlags::Recompile);
+		FBlueprintEditorUtils::RemoveMemberVariable(Blueprint, EdGraph->GetFName());
+		FBlueprintEditorUtils::RemoveGraph(Blueprint, EdGraph, EGraphRemoveFlags::Recompile);
 
-			for (TObjectIterator<UK2Node_CreateDelegate> It; It; ++It)
-			{
-				It->HandleAnyChange();
-			}
+		for (TObjectIterator<UK2Node_CreateDelegate> It; It; ++It)
+		{
+			It->HandleAnyChange();
 		}
 	}
+}
+
+void SMyBlueprint::OnDeleteEntry()
+{
+	if (FEdGraphSchemaAction_K2Graph* GraphAction = SelectionAsGraph())
+	{
+		OnDeleteGraph(GraphAction->EdGraph, GraphAction->GraphType);
+	}
+	else if (FEdGraphSchemaAction_K2Delegate* DelegateAction = SelectionAsDelegate())
+	{
+		OnDeleteDelegate(DelegateAction);
+	}
 	else if ( FEdGraphSchemaAction_K2Var* VarAction = SelectionAsVar() )
 	{
 		if(FBlueprintEditorUtils::IsVariableUsed(GetBlueprintObj(), VarAction->GetVariableName()))
@@ -1591,7 +1711,38 @@ void SMyBlueprint::OnDeleteEntry()
 		GraphActionMenu->GetSelectedCategorySubActions(Actions);
 		if (Actions.Num())
 		{
-			const FScopedTransaction Transaction( LOCTEXT( "BulkRemoveVariables", "Bulk Remove Variables" ) );
+			FText TransactionTitle;
+
+			switch((NodeSectionID::Type)Actions[0]->SectionID)
+			{
+			case NodeSectionID::VARIABLE:
+			case NodeSectionID::LOCAL_VARIABLE:
+				{
+					TransactionTitle = LOCTEXT( "BulkRemoveVariables", "Bulk Remove Variables" );
+					break;
+				}
+			case NodeSectionID::DELEGATE:
+				{
+					TransactionTitle = LOCTEXT( "BulkRemoveDelegates", "Bulk Remove Delegates" );
+					break;
+				}
+			case NodeSectionID::FUNCTION:
+				{
+					TransactionTitle = LOCTEXT( "BulkRemoveFunctions", "Bulk Remove Functions" );
+					break;
+				}
+			case NodeSectionID::MACRO:
+				{
+					TransactionTitle = LOCTEXT( "BulkRemoveMacros", "Bulk Remove Macros" );
+					break;
+				}
+			default:
+				{
+					TransactionTitle = LOCTEXT( "BulkRemove", "Bulk Remove Items" );
+				}
+			}
+
+			const FScopedTransaction Transaction( TransactionTitle);
 
 			GetBlueprintObj()->Modify();
 			for (int32 i = 0; i < Actions.Num(); ++i)
@@ -1602,6 +1753,22 @@ void SMyBlueprint::OnDeleteEntry()
 					
 					FBlueprintEditorUtils::RemoveMemberVariable(GetBlueprintObj(), Var->GetVariableName());
 				}
+				else if (Actions[i]->GetTypeId() == FEdGraphSchemaAction_K2LocalVar::StaticGetTypeId())
+				{
+					FEdGraphSchemaAction_K2LocalVar* LocalVarAction = (FEdGraphSchemaAction_K2LocalVar*)Actions[i].Get();
+
+					FBlueprintEditorUtils::RemoveLocalVariable(GetBlueprintObj(), LocalVarAction->GetVariableScope(), LocalVarAction->GetVariableName());
+				}
+				else if (Actions[i]->GetTypeId() == FEdGraphSchemaAction_K2Graph::StaticGetTypeId())
+				{
+					FEdGraphSchemaAction_K2Graph* GraphAction = (FEdGraphSchemaAction_K2Graph*)Actions[i].Get();
+
+					OnDeleteGraph(GraphAction->EdGraph, GraphAction->GraphType);
+				}
+				else if (Actions[i]->GetTypeId() == FEdGraphSchemaAction_K2Delegate::StaticGetTypeId())
+				{
+					OnDeleteDelegate((FEdGraphSchemaAction_K2Delegate*)Actions[i].Get());
+				}
 			}
 		}
 	}
diff --git a/Engine/Source/Editor/Kismet/Private/SMyBlueprint.h b/Engine/Source/Editor/Kismet/Private/SMyBlueprint.h
--- a/Engine/Source/Editor/Kismet/Private/SMyBlueprint.h
+++ b/Engine/Source/Editor/Kismet/Private/SMyBlueprint.h
@@ -189,6 +189,14 @@ private:
 	/** Update Node Create Analytic */
 	void UpdateNodeCreation();
 
+	/** Returns the displayed category, if any, of a graph */
+	FString GetGraphCategory(UEdGraph* InGraph) const;
+
+	/** Helper function to delete a graph in the MyBlueprint window */
+	void OnDeleteGraph(UEdGraph* InGraph, EEdGraphSchemaAction_K2Graph::Type);
+
+	/** Helper function to delete a delegate in the MyBlueprint window */
+	void OnDeleteDelegate(FEdGraphSchemaAction_K2Delegate* InDelegateAction);
 private:
 	/** Pointer back to the blueprint editor that owns us */
 	TWeakPtr<FBlueprintEditor> BlueprintEditorPtr;
