Add compile event analytics

- This was to add how often users trigger compile and the result of it.

diff --git a/Engine/Source/Editor/Kismet/Private/BlueprintEditor.cpp b/Engine/Source/Editor/Kismet/Private/BlueprintEditor.cpp
--- a/Engine/Source/Editor/Kismet/Private/BlueprintEditor.cpp
+++ b/Engine/Source/Editor/Kismet/Private/BlueprintEditor.cpp
@@ -354,6 +354,30 @@ void FBlueprintEditor::AnalyticsTrackNodeEvent( UBlueprint* Blueprint, UEdGraphN
 	}
 }
 
+void FBlueprintEditor::AnalyticsTrackCompileEvent( UBlueprint* Blueprint, int32 NumErrors, int32 NumWarnings ) const
+{
+	if(Blueprint && FEngineAnalytics::IsAvailable())
+	{
+		// we'd like to see if this was happening in normal blueprint editor or persona 
+		const FString EditorName = Cast<UAnimBlueprint>(Blueprint) != NULL ? TEXT("Persona") : TEXT("BlueprintEditor");
+
+		// Build Node Details
+		const UGeneralProjectSettings& ProjectSettings = *GetDefault<UGeneralProjectSettings>();
+		FString ProjectID = ProjectSettings.ProjectID.ToString();
+
+		const bool bSuccess = NumErrors == 0;
+		TArray< FAnalyticsEventAttribute > Attributes;
+		Attributes.Add(FAnalyticsEventAttribute(TEXT("ProjectId"), ProjectID));
+		Attributes.Add(FAnalyticsEventAttribute(TEXT("BlueprintId"), Blueprint->GetBlueprintGuid().ToString()));
+		Attributes.Add(FAnalyticsEventAttribute(TEXT("Success"), bSuccess? TEXT("True") : TEXT("False") ));
+		Attributes.Add(FAnalyticsEventAttribute(TEXT("NumErrors"), FString::FromInt(NumErrors)));
+		Attributes.Add(FAnalyticsEventAttribute(TEXT("NumWarnings"), FString::FromInt(NumWarnings)));
+
+		// Send Analytics event 
+		FEngineAnalytics::GetProvider().RecordEvent(FString::Printf(TEXT("Editor.Usage.%s.Compile"), *EditorName), Attributes);
+	}
+}
+
 void FBlueprintEditor::RefreshEditors()
 {
 	DocumentManager->CleanInvalidTabs();
@@ -2306,6 +2330,10 @@ void FBlueprintEditor::Compile()
 
 		bool bForceMessageDisplay = ((LogResults.NumWarnings > 0) || (LogResults.NumErrors > 0)) && !GetBlueprintObj()->bIsRegeneratingOnLoad;
 		DumpMessagesToCompilerLog(LogResults.Messages, bForceMessageDisplay);
+
+		// send record when player clicks compile and send the result
+		// this will make sure how the users activity is
+		AnalyticsTrackCompileEvent(GetBlueprintObj(), LogResults.NumErrors, LogResults.NumWarnings);
 	}
 }
 
diff --git a/Engine/Source/Editor/Kismet/Public/BlueprintEditor.h b/Engine/Source/Editor/Kismet/Public/BlueprintEditor.h
--- a/Engine/Source/Editor/Kismet/Public/BlueprintEditor.h
+++ b/Engine/Source/Editor/Kismet/Public/BlueprintEditor.h
@@ -156,6 +156,7 @@ public:
 	virtual TSharedPtr<class FSCSEditorTreeNode> FindAndSelectSCSEditorTreeNode(const UActorComponent* InComponent, bool IsCntrlDown) override;
 	virtual int32 GetNumberOfSelectedNodes() const override;
 	virtual void AnalyticsTrackNodeEvent( UBlueprint* Blueprint, UEdGraphNode *GraphNode, bool bNodeDelete = false ) const override;
+	void AnalyticsTrackCompileEvent( UBlueprint* Blueprint, int32 NumErrors, int32 NumWarnings ) const;
 	// End of IBlueprintEditor interface
 
 	// FTickableEditorObject interface
