ttp341325 CRASH: FBlueprintEditor::SaveAsset_Execute() TBaseSPMethodDelegateInstance_RetVal_NoParams<SToolBarButtonBlock,FReply,0>::Execute()

diff --git a/Engine/Source/Editor/Kismet/Private/BlueprintEditor.cpp b/Engine/Source/Editor/Kismet/Private/BlueprintEditor.cpp
--- a/Engine/Source/Editor/Kismet/Private/BlueprintEditor.cpp
+++ b/Engine/Source/Editor/Kismet/Private/BlueprintEditor.cpp
@@ -6546,9 +6546,16 @@ bool FBlueprintEditor::IsFocusedGraphEditable() const
 
 void FBlueprintEditor::SaveAsset_Execute()
 {
-	// Update the Blueprint's search data
-	GetBlueprintObj()->SearchGuid = FGuid::NewGuid();
-	FFindInBlueprintSearchManager::Get().AddOrUpdateBlueprintSearchMetadata(GetBlueprintObj());
+	for (auto Object : GetEditingObjects())
+	{
+		auto Blueprint = Cast<UBlueprint>(Object);
+		if (Blueprint)
+		{
+			// Update the Blueprint's search data
+			Blueprint->SearchGuid = FGuid::NewGuid();
+			FFindInBlueprintSearchManager::Get().AddOrUpdateBlueprintSearchMetadata(Blueprint);
+		}
+	}
 
 	IBlueprintEditor::SaveAsset_Execute();
 }
