Sub-graphs (such as collapsed graphs or state machines) will display under their parent graph in the My Blueprint's tab rather than a category with the same name.

#ttp 334222 - Regression: BP: In the my blueprints tab, collapsed graphs are incorrectly put into a 'Functions' category when created

diff --git a/Engine/Source/Editor/GraphEditor/Private/GraphActionNode.h b/Engine/Source/Editor/GraphEditor/Private/GraphActionNode.h
--- a/Engine/Source/Editor/GraphEditor/Private/GraphActionNode.h
+++ b/Engine/Source/Editor/GraphEditor/Private/GraphActionNode.h
@@ -106,8 +106,9 @@ public:
 					{
 						TSharedPtr<FEdGraphSchemaAction> CurrentAction = NodeToAdd->Actions[0];
 
+						// Make the action's name into a display name, all categories are set as such
 						bFoundMatch = CurrentAction->IsParentable()
-							&& NextCategory == CurrentChild->Actions[0]->MenuDescription.ToString();
+							&& NextCategory == FName::NameToDisplayString(CurrentChild->Actions[0]->MenuDescription.ToString(), false);
 					}
 				}
 				
diff --git a/Engine/Source/Editor/Kismet/Private/BlueprintDetailsCustomization.cpp b/Engine/Source/Editor/Kismet/Private/BlueprintDetailsCustomization.cpp
--- a/Engine/Source/Editor/Kismet/Private/BlueprintDetailsCustomization.cpp
+++ b/Engine/Source/Editor/Kismet/Private/BlueprintDetailsCustomization.cpp
@@ -2042,58 +2042,62 @@ void FBlueprintGraphActionDetails::CustomizeDetails( IDetailLayoutBuilder& Detai
 					.Font( IDetailLayoutBuilder::GetDetailFont() )
 			];
 
-			FBlueprintVarActionDetails::PopulateCategories(MyBlueprint.Pin().Get(), CategorySource);
-			TSharedPtr<SComboButton> NewComboButton;
-			TSharedPtr<SListView<TSharedPtr<FString>>> NewListView;
+			// Composite graphs are auto-categorized into their parent graph
+			if(!GetGraph()->GetOuter()->GetClass()->IsChildOf(UK2Node_Composite::StaticClass()))
+			{
+				FBlueprintVarActionDetails::PopulateCategories(MyBlueprint.Pin().Get(), CategorySource);
+				TSharedPtr<SComboButton> NewComboButton;
+				TSharedPtr<SListView<TSharedPtr<FString>>> NewListView;
 
-			const FString DocLink = TEXT("Shared/Editors/BlueprintEditor/VariableDetails");
-			TSharedPtr<SToolTip> CategoryTooltip = IDocumentation::Get()->CreateToolTip(LOCTEXT("EditCategoryName_Tooltip", "The category of the variable; editing this will place the variable into another category or create a new one."), NULL, DocLink, TEXT("Category"));
+				const FString DocLink = TEXT("Shared/Editors/BlueprintEditor/VariableDetails");
+				TSharedPtr<SToolTip> CategoryTooltip = IDocumentation::Get()->CreateToolTip(LOCTEXT("EditCategoryName_Tooltip", "The category of the variable; editing this will place the variable into another category or create a new one."), NULL, DocLink, TEXT("Category"));
 
-			Category.AddCustomRow( TEXT("Category") )
-				.NameContent()
-				[
-					SNew(STextBlock)
-					.Text( LOCTEXT("CategoryLabel", "Category") )
-					.ToolTip(CategoryTooltip)
-					.Font( IDetailLayoutBuilder::GetDetailFont() )
-				]
-			.ValueContent()
-				[
-					SAssignNew(NewComboButton, SComboButton)
-					.ContentPadding(FMargin(0,0,5,0))
-					.ToolTip(CategoryTooltip)
-					.ButtonContent()
+				Category.AddCustomRow( TEXT("Category") )
+					.NameContent()
 					[
-						SNew(SBorder)
-						.BorderImage( FEditorStyle::GetBrush("NoBorder") )
-						.Padding(FMargin(0, 0, 5, 0))
-						[
-							SNew(SEditableTextBox)
-							.Text(this, &FBlueprintGraphActionDetails::OnGetCategoryText)
-							.OnTextCommitted(this, &FBlueprintGraphActionDetails::OnCategoryTextCommitted )
-							.ToolTip(CategoryTooltip)
-							.SelectAllTextWhenFocused(true)
-							.RevertTextOnEscape(true)
-							.Font( IDetailLayoutBuilder::GetDetailFont() )
-						]
+						SNew(STextBlock)
+						.Text( LOCTEXT("CategoryLabel", "Category") )
+						.ToolTip(CategoryTooltip)
+						.Font( IDetailLayoutBuilder::GetDetailFont() )
 					]
-					.MenuContent()
+				.ValueContent()
+					[
+						SAssignNew(NewComboButton, SComboButton)
+						.ContentPadding(FMargin(0,0,5,0))
+						.ToolTip(CategoryTooltip)
+						.ButtonContent()
 						[
-							SNew(SVerticalBox)
-							+SVerticalBox::Slot()
-							.AutoHeight()
-							.MaxHeight(400.0f)
+							SNew(SBorder)
+							.BorderImage( FEditorStyle::GetBrush("NoBorder") )
+							.Padding(FMargin(0, 0, 5, 0))
 							[
-								SAssignNew(NewListView, SListView<TSharedPtr<FString>>)
-								.ListItemsSource(&CategorySource)
-								.OnGenerateRow(this, &FBlueprintGraphActionDetails::MakeCategoryViewWidget)
-								.OnSelectionChanged(this, &FBlueprintGraphActionDetails::OnCategorySelectionChanged)
+								SNew(SEditableTextBox)
+								.Text(this, &FBlueprintGraphActionDetails::OnGetCategoryText)
+								.OnTextCommitted(this, &FBlueprintGraphActionDetails::OnCategoryTextCommitted )
+								.ToolTip(CategoryTooltip)
+								.SelectAllTextWhenFocused(true)
+								.RevertTextOnEscape(true)
+								.Font( IDetailLayoutBuilder::GetDetailFont() )
 							]
 						]
-				];
+						.MenuContent()
+							[
+								SNew(SVerticalBox)
+								+SVerticalBox::Slot()
+								.AutoHeight()
+								.MaxHeight(400.0f)
+								[
+									SAssignNew(NewListView, SListView<TSharedPtr<FString>>)
+									.ListItemsSource(&CategorySource)
+									.OnGenerateRow(this, &FBlueprintGraphActionDetails::MakeCategoryViewWidget)
+									.OnSelectionChanged(this, &FBlueprintGraphActionDetails::OnCategorySelectionChanged)
+								]
+							]
+					];
 
-			CategoryComboButton = NewComboButton;
-			CategoryListView = NewListView;
+				CategoryComboButton = NewComboButton;
+				CategoryListView = NewListView;
+			}
 
 			if (IsAccessSpecifierVisible())
 			{
