This adds category tracking in the blueprint class to allow sorting of categories in MyBlueprint
#TTP 333452 - BLUEPRINTS: Can't rearrange component categories in the my blueprint tab
#TTP 341990 - LIVE: BLUEPRINTS: Cannot drag and drop into subcategories
#Branch UE4
#Proj Engine, Kismet, UnrealEd, GraphEditor

#Change Added UBlueprint::CategorySorting TArray to store an ordered sort list of categories for variables and components.
#Change Modified dropped on category to update the CategorySorting array in the blueprint
#Change Modified FGraphActionNode to store a categorychain for category nodes only.
#Change Modified FGraphActionNode::AddChild to build and hand off the full category chain to new category nodes.
#Change Modified SMyBlueprint::CollectAllActions to reference the sorting array stored in the blueprint, participating actions are added into a mutimap with the respective sort index and the map is sorted before populating the graph action list.
#Change Added an additional pass when the sorted categories are modified to strip any unused categories out of the blueprint and shrink the array as a lazy maintenance.
#Removed FBlueprintEditorUtils::MoveCategoryBeforeCategory

#info I stopped short of marking the blueprint dirty on category sorting changes after initially having this in. the logic behind this is that marking blueprints dirty for minor changes can be annoying, you can always just manually save your changes and its seems a contrived workflow that you would open a blueprint and change the category sort in my blueprint without making any other changes. Easy to change if there is any real objection to it.

diff --git a/Engine/Source/Editor/GraphEditor/Private/GraphActionNode.h b/Engine/Source/Editor/GraphEditor/Private/GraphActionNode.h
--- a/Engine/Source/Editor/GraphEditor/Private/GraphActionNode.h
+++ b/Engine/Source/Editor/GraphEditor/Private/GraphActionNode.h
@@ -12,6 +12,8 @@ struct FGraphActionNode : TSharedFromThis<FGraphActionNode>
 
 	/** The category this node belongs in. */
 	FString Category;
+	/** The category chain, separated using '|' */
+	FString CategoryChain;
 
 	int32 Grouping;
 	TArray< TSharedPtr<FEdGraphSchemaAction> > Actions;
@@ -54,9 +56,9 @@ public:
 		return MakeShareable( new FGraphActionNode( TEXT(""), Grouping, InActionList, SectionID ) );
 	}
 
-	static TSharedPtr<FGraphActionNode> NewCategory( const FString& InCategory, int32 InGrouping = 0 )
+	static TSharedPtr<FGraphActionNode> NewCategory( const FString& InCategory, const FString& InCategoryChain, int32 InGrouping = 0 )
 	{
-		return MakeShareable( new FGraphActionNode(InCategory, InGrouping, NULL) );
+		return MakeShareable( new FGraphActionNode(InCategory, InCategoryChain, InGrouping, NULL) );
 	}
 
 	static TSharedPtr<FGraphActionNode> NewSeparator( int32 InGrouping = 0 )
@@ -84,7 +86,7 @@ public:
 		return (Actions.Num()==0 || !Actions[0].IsValid()) && Category.IsEmpty();
 	}
 
-	void AddChild(TSharedPtr<FGraphActionNode> NodeToAdd, TArray<FString>& CategoryStack, bool bAlphaSort)
+	void AddChild(TSharedPtr<FGraphActionNode> NodeToAdd, TArray<FString>& CategoryStack, bool bAlphaSort, FString& CategoryChain)
 	{
 		// If we still have some categories in the chain to process...
 		if( CategoryStack.Num() )
@@ -114,18 +116,20 @@ public:
 				
 				if( bFoundMatch )
 				{
-					CurrentChild->AddChild(NodeToAdd, CategoryStack, bAlphaSort);
+					CategoryChain = CategoryChain.IsEmpty() ? NextCategory : FString::Printf( TEXT( "%s|%s" ), *CategoryChain, *NextCategory );
+					CurrentChild->AddChild(NodeToAdd, CategoryStack, bAlphaSort, CategoryChain);
 					break;
 				}
 			}
 
 			if( !bFoundMatch )
 			{
-				TSharedPtr<FGraphActionNode> NewCategoryNode = NewCategory(NextCategory, NodeToAdd->Grouping);
+				CategoryChain = CategoryChain.IsEmpty() ? NextCategory : FString::Printf( TEXT( "%s|%s" ), *CategoryChain, *NextCategory );
+				TSharedPtr<FGraphActionNode> NewCategoryNode = NewCategory(NextCategory, CategoryChain, NodeToAdd->Grouping);
 				// Copy the section ID also 
 				NewCategoryNode->SectionID = NodeToAdd->SectionID;
 				AddToChildrenSorted(NewCategoryNode, bAlphaSort);
-				NewCategoryNode->AddChild(NodeToAdd, CategoryStack, bAlphaSort);
+				NewCategoryNode->AddChild(NodeToAdd, CategoryStack, bAlphaSort, CategoryChain);
 			}
 		}
 		else
@@ -386,6 +390,16 @@ protected:
 		Actions.Add( InAction );
 	}
 
+	FGraphActionNode( FString InCategory, FString InCategoryChain, int32 InGrouping, TSharedPtr<FEdGraphSchemaAction> InAction )
+		: Category( InCategory )
+		, CategoryChain( InCategoryChain )
+		, Grouping( InGrouping )
+		, bIsRenameRequestBeforeReady( false )
+		, SectionID ( 0 )
+	{
+		Actions.Add( InAction );
+	}
+
 	FGraphActionNode( FString InCategory, int32 InGrouping, TArray< TSharedPtr<FEdGraphSchemaAction> >& InActionList, int32 InSectionID = 0 )
 	: Category( InCategory )
 	, Grouping( InGrouping )
diff --git a/Engine/Source/Editor/GraphEditor/Private/SGraphActionMenu.cpp b/Engine/Source/Editor/GraphEditor/Private/SGraphActionMenu.cpp
--- a/Engine/Source/Editor/GraphEditor/Private/SGraphActionMenu.cpp
+++ b/Engine/Source/Editor/GraphEditor/Private/SGraphActionMenu.cpp
@@ -115,7 +115,7 @@ public:
 		TSharedPtr<FGraphEditorDragDropAction> GraphDropOp = DragDropEvent.GetOperationAs<FGraphEditorDragDropAction>();
 		if (GraphDropOp.IsValid())
 		{
-			GraphDropOp->DroppedOnCategory( ActionNode.Pin()->Category );
+			GraphDropOp->DroppedOnCategory( ActionNode.Pin()->CategoryChain );
 			return FReply::Handled();
 		}
 		return FReply::Unhandled();
@@ -174,7 +174,7 @@ void SGraphActionMenu::Construct( const FArguments& InArgs, bool bIsReadOnly/* =
 	this->OnCategoryTextCommitted = InArgs._OnCategoryTextCommitted;
 	this->OnCanRenameSelectedAction = InArgs._OnCanRenameSelectedAction;
 	this->OnGetSectionTitle = InArgs._OnGetSectionTitle;
-	this->FilteredRootAction = FGraphActionNode::NewCategory(TEXT("FILTEREDROOT"));
+	this->FilteredRootAction = FGraphActionNode::NewCategory(TEXT("FILTEREDROOT"), TEXT( "" ));
 	this->OnActionMatchesName = InArgs._OnActionMatchesName;
 	
 	// If a delegate for filtering text is passed in, assign it so that it will be used instead of the built-in filter box
@@ -568,7 +568,8 @@ void SGraphActionMenu::GenerateFilteredItems(bool bPreserveExpansion)
 			CurrentAction.GetCategoryChain(CategoryChain);
 			
 			TSharedPtr<FGraphActionNode> NewNode = FGraphActionNode::NewAction(CurrentAction.Actions);
-			FilteredRootAction->AddChild(NewNode, CategoryChain, bAlphaSortItems);
+			FString EmptyCategoryChain;
+			FilteredRootAction->AddChild(NewNode, CategoryChain, bAlphaSortItems, EmptyCategoryChain);
 		}
 	}
 
@@ -1022,7 +1023,7 @@ FReply SGraphActionMenu::OnItemDragDetected( const FGeometry& MyGeometry, const
 			{
 				if(OnCategoryDragged.IsBound())
 				{
-					return OnCategoryDragged.Execute(Node->Category, MouseEvent);
+					return OnCategoryDragged.Execute(Node->CategoryChain, MouseEvent);
 				}
 			}
 			// Dragging an action
diff --git a/Engine/Source/Editor/Kismet/Private/SMyBlueprint.cpp b/Engine/Source/Editor/Kismet/Private/SMyBlueprint.cpp
--- a/Engine/Source/Editor/Kismet/Private/SMyBlueprint.cpp
+++ b/Engine/Source/Editor/Kismet/Private/SMyBlueprint.cpp
@@ -36,6 +36,7 @@
 #include "Editor/GraphEditor/Private/GraphActionNode.h"
 #include "IDocumentation.h"
 #include "Editor/UnrealEd/Public/SourceCodeNavigation.h"
+#include "EditorCategoryUtils.h"
 
 #define LOCTEXT_NAMESPACE "MyBlueprint"
 
@@ -95,18 +96,12 @@ public:
 	
 	virtual FReply DroppedOnCategory(FString OnCategory) override
 	{
-		// Get the Blueprint via BlueprintEditor
+		// Get MyBlueprint via MyBlueprintPtr
 		TSharedPtr<SMyBlueprint> MyBlueprint = MyBlueprintPtr.Pin();
 		if(MyBlueprint.IsValid())
 		{
-			TSharedPtr<FBlueprintEditor> BPEditor = MyBlueprint->GetBlueprintEditor().Pin();
-			if(BPEditor.IsValid())
-			{
-				UBlueprint* BP = BPEditor->GetBlueprintObj();
-				check(BP);
-				// Move the category
-				FBlueprintEditorUtils::MoveCategoryBeforeCategory(BP, FName(*DraggedCategory), FName(*OnCategory));
-			}
+			// Move the category in the blueprint category sort list
+			MyBlueprint->MoveCategoryBeforeCategory( DraggedCategory, OnCategory );
 		}
 
 		return FReply::Handled();
@@ -363,10 +358,13 @@ FText SMyBlueprint::OnGetSectionTitle( int32 InSectionID )
 	FText SeperatorTitle;
 	/* Setup an appropriate name for the section for this node */
 	switch( InSectionID )
-	{		
+	{
 	case NodeSectionID::VARIABLE:
 		SeperatorTitle = NSLOCTEXT("GraphActionNode", "Variables", "Variables");
 		break;
+	case NodeSectionID::COMPONENT:
+		SeperatorTitle = NSLOCTEXT("GraphActionNode", "Components", "Components");
+		break;
 	case NodeSectionID::FUNCTION:
 		SeperatorTitle = NSLOCTEXT("GraphActionNode", "Functions", "Functions");
 		break;
@@ -680,11 +678,93 @@ void SMyBlueprint::CollectAllActions(FGraphActionListBuilderBase& OutAllActions)
 		FieldIteratorSuperFlag = EFieldIteratorFlags::ExcludeSuper;
 	}
 
+	// Helper structure to aid category sorting
+	struct FGraphActionSort
+	{
+		public:
+			FGraphActionSort( TArray<FName>& BlueprintCategorySorting )
+				: bCategoriesModified( false )
+				, CategorySortIndices( BlueprintCategorySorting )
+			{
+				CategoryUsage.Init( 0, CategorySortIndices.Num() );
+			}
+
+			void AddAction( const FString& Category, TSharedPtr<FEdGraphSchemaAction> Action )
+			{
+				// Find root category
+				int32 RootCategoryDelim = Category.Find( TEXT( "|" ));
+				FName RootCategory = RootCategoryDelim == INDEX_NONE ?	*Category : *Category.Left( RootCategoryDelim );
+				// Get root sort index
+				const int32 SortIndex = GetSortIndex( RootCategory ) + Action->SectionID;
+
+				static const FName DefaultCategory( TEXT( "Default" ));
+				if( RootCategory != DefaultCategory )
+				{
+					// Strip sections out of non default categories to enable them to be grouped together
+					Action->SectionID = NodeSectionID::USER_SORTED;
+				}
+				SortedActions.Add( SortIndex, Action );
+			}
+
+			void GetAllActions( FGraphActionListBuilderBase& OutAllActions )
+			{
+				SortedActions.KeySort( TLess<int32>() );
+
+				for( auto Iter : SortedActions )
+				{
+					OutAllActions.AddAction( Iter.Value );
+				}
+			}
+
+			void CleanupCategories()
+			{
+				// Scrub unused categories from the blueprint
+				if( bCategoriesModified )
+				{
+					for( int32 CategoryIdx = CategoryUsage.Num() - 1; CategoryIdx >= 0; CategoryIdx-- )
+					{
+						if( CategoryUsage[ CategoryIdx ] == 0 )
+						{
+							CategorySortIndices.RemoveAt( CategoryIdx );
+						}
+					}
+					bCategoriesModified = false;
+				}
+			}
+
+		private:
+			const int32 GetSortIndex( FName Category )
+			{
+				int32 SortIndex = CategorySortIndices.Find( Category );
+
+				if( SortIndex == INDEX_NONE )
+				{
+					bCategoriesModified = true;
+					SortIndex = CategorySortIndices.Add( Category );
+					CategoryUsage.Add( 0 );
+				}
+				CategoryUsage[ SortIndex ]++;
+				// Spread the sort values so we can fine tune sorting
+				SortIndex *= 1000;
+
+				return SortIndex + SortedActions.Num();
+			}
+
+		private:
+			/** Signals if the blueprint categories have been modified and require cleanup */
+			bool bCategoriesModified;
+			/** Tracks category usage to aid removal of unused categories */
+			TArray<int32> CategoryUsage;
+			/** Reference to the category sorting in the blueprint */
+			TArray<FName>& CategorySortIndices;
+			/** Map used to sort Graph actions */
+			TMultiMap<int32,TSharedPtr<FEdGraphSchemaAction>> SortedActions;
+	}; 
+
+	// Initialise action sorting instance
+	FGraphActionSort SortList( Blueprint->CategorySorting );
 	// List of names of functions we implement
 	TArray<FName> ImplementedFunctions;
-	TArray<TSharedPtr<FEdGraphSchemaAction>> VariableActions;
-	TArray<TSharedPtr<FEdGraphSchemaAction>> ComponentPropertyActions;
-
 	// Grab Variables
 	for (TFieldIterator<UProperty> PropertyIt(Blueprint->SkeletonGeneratedClass, FieldIteratorSuperFlag); PropertyIt; ++PropertyIt)
 	{
@@ -698,9 +778,6 @@ void SMyBlueprint::CollectAllActions(FGraphActionListBuilderBase& OutAllActions)
 		const bool bShouldShowAsDelegate = !Property->HasAnyPropertyFlags(CPF_Parm) && bMulticastDelegateProp 
 			&& Property->HasAnyPropertyFlags(CPF_BlueprintAssignable | CPF_BlueprintCallable);
 		UObjectPropertyBase* Obj = Cast<UObjectPropertyBase>(Property);
-		const bool bComponentProperty = Obj && Obj->PropertyClass ? Obj->PropertyClass->IsChildOf<UActorComponent>() : false;
-		TArray<TSharedPtr<FEdGraphSchemaAction>>& ActionList = bComponentProperty ? ComponentPropertyActions : VariableActions;
-		
 		if(!bShouldShowAsVar && !bShouldShowAsDelegate)
 		{
 			continue;
@@ -712,6 +789,8 @@ void SMyBlueprint::CollectAllActions(FGraphActionListBuilderBase& OutAllActions)
 
 		FName CategoryName = FObjectEditorUtils::GetCategoryFName(Property);
 		FString PropertyCategory = FObjectEditorUtils::GetCategory(Property);
+		const FString UserCategoryName = FEditorCategoryUtils::GetCategoryDisplayString( PropertyCategory );
+
 		if ((CategoryName == Blueprint->GetFName()) || (CategoryName == K2Schema->VR_DefaultCategory))
 		{
 			CategoryName = NAME_None;		// default, so place in 'non' category
@@ -722,8 +801,9 @@ void SMyBlueprint::CollectAllActions(FGraphActionListBuilderBase& OutAllActions)
 		{
 			TSharedPtr<FEdGraphSchemaAction_K2Var> NewVarAction = MakeShareable(new FEdGraphSchemaAction_K2Var(PropertyCategory, PropertyDesc, PropertyTooltip.ToString(), 0));
 			NewVarAction->SetVariableInfo(PropertyName, Blueprint->SkeletonGeneratedClass);
-			NewVarAction->SectionID = NodeSectionID::VARIABLE;
-			ActionList.Add(NewVarAction);
+			const bool bComponentProperty = Obj && Obj->PropertyClass ? Obj->PropertyClass->IsChildOf<UActorComponent>() : false;
+			NewVarAction->SectionID = bComponentProperty ? NodeSectionID::COMPONENT : NodeSectionID::VARIABLE;
+			SortList.AddAction( UserCategoryName, NewVarAction );
 		}
 		else if (bShouldShowAsDelegate)
 		{
@@ -735,7 +815,7 @@ void SMyBlueprint::CollectAllActions(FGraphActionListBuilderBase& OutAllActions)
 				NewFuncAction->SetDelegateInfo(PropertyName, Blueprint->SkeletonGeneratedClass);
 				NewFuncAction->EdGraph = NULL;
 				NewFuncAction->SectionID = NodeSectionID::DELEGATE;
-				ActionList.Add(NewFuncAction);
+				SortList.AddAction( UserCategoryName, NewFuncAction );
 			}
 
 			UClass* OwnerClass = CastChecked<UClass>(Property->GetOuter());
@@ -751,15 +831,9 @@ void SMyBlueprint::CollectAllActions(FGraphActionListBuilderBase& OutAllActions)
 		}
 	}
 
-	// Add all variable actions first so their order can control the categories
-	for( int32 i = 0; i < VariableActions.Num(); ++i)
-	{
-		OutAllActions.AddAction( VariableActions[ i ] );
-	}
-	for( int32 i = 0; i < ComponentPropertyActions.Num(); ++i)
-	{
-		OutAllActions.AddAction( ComponentPropertyActions[ i ] );
-	}
+	// Add all the sorted variables and components
+	SortList.CleanupCategories();
+	SortList.GetAllActions( OutAllActions );
 
 	// Grab functions implemented by the blueprint
 	for (int32 i = 0; i < Blueprint->FunctionGraphs.Num(); i++)
@@ -2026,4 +2100,37 @@ void SMyBlueprint::ExpandCategory(const FString& CategoryName)
 	}
 }
 
+bool SMyBlueprint::MoveCategoryBeforeCategory( const FString& InCategoryToMove, const FString& InTargetCategory )
+{
+	bool bResult = false;
+	UBlueprint* Blueprint = BlueprintEditorPtr.Pin()->GetBlueprintObj();
+
+	if( Blueprint )
+	{
+		// Find root categories
+		int32 RootCategoryDelim = InCategoryToMove.Find( TEXT( "|" ));
+		FName CategoryToMove = RootCategoryDelim == INDEX_NONE ? *InCategoryToMove : *InCategoryToMove.Left( RootCategoryDelim );
+		RootCategoryDelim = InTargetCategory.Find( TEXT( "|" ));
+		FName TargetCategory = RootCategoryDelim == INDEX_NONE ? *InTargetCategory : *InTargetCategory.Left( RootCategoryDelim );
+
+		TArray<FName>& CategorySort = Blueprint->CategorySorting;
+		const int32 RemovalIndex = CategorySort.Find( CategoryToMove );
+		// Remove existing sort index
+		if( RemovalIndex != INDEX_NONE )
+		{
+			CategorySort.RemoveAt( RemovalIndex );
+		}
+		// Update the Category sort order and refresh ( if the target category has an entry )
+		const int32 InsertIndex = CategorySort.Find( TargetCategory );
+		if( InsertIndex != INDEX_NONE )
+		{
+			CategorySort.Insert( CategoryToMove, InsertIndex );
+			Refresh();
+			bResult = true;
+		}
+	}
+
+	return bResult;
+}
+
 #undef LOCTEXT_NAMESPACE
diff --git a/Engine/Source/Editor/Kismet/Private/SMyBlueprint.h b/Engine/Source/Editor/Kismet/Private/SMyBlueprint.h
--- a/Engine/Source/Editor/Kismet/Private/SMyBlueprint.h
+++ b/Engine/Source/Editor/Kismet/Private/SMyBlueprint.h
@@ -14,10 +14,12 @@ namespace NodeSectionID
 		INTERFACE,				// Interface
 		MACRO,					// Macros
 		VARIABLE,				// Variables
+		COMPONENT,				// Components
 		DELEGATE,				// Delegate/Event
 		USER_ENUM,				// User defined enums
 		LOCAL_VARIABLE,			// Local variables
 		USER_STRUCT,			// User defined structs
+		USER_SORTED				// User sorted categories
 	};
 };
 
@@ -120,6 +122,9 @@ public:
 	/** Expands any category with the associated name */
 	void ExpandCategory(const FString& CategoryName);
 
+	/** Move the category before the target category */
+	bool MoveCategoryBeforeCategory( const FString& CategoryToMove, const FString& TargetCategory );
+
 private:
 	/** Creates widgets for the graph schema actions */
 	TSharedRef<SWidget> OnCreateWidgetForAction(struct FCreateWidgetForActionData* const InCreateData);
diff --git a/Engine/Source/Editor/UnrealEd/Private/Kismet2/BlueprintEditorUtils.cpp b/Engine/Source/Editor/UnrealEd/Private/Kismet2/BlueprintEditorUtils.cpp
--- a/Engine/Source/Editor/UnrealEd/Private/Kismet2/BlueprintEditorUtils.cpp
+++ b/Engine/Source/Editor/UnrealEd/Private/Kismet2/BlueprintEditorUtils.cpp
@@ -2445,51 +2445,6 @@ int32 FBlueprintEditorUtils::FindFirstNewVarOfCategory(const UBlueprint* Bluepri
 }
 
 
-bool FBlueprintEditorUtils::MoveCategoryBeforeCategory(UBlueprint* Blueprint, FName CategoryToMove, FName TargetCategory)
-{
-	// Do nothing if moving to yourself, or no category
-	if(CategoryToMove == TargetCategory || CategoryToMove == NAME_None)
-	{
-		return false;
-	}
-
-	// Before we do anything, check that there is a variable of the target category
-	if(FindFirstNewVarOfCategory(Blueprint, TargetCategory) == INDEX_NONE)
-	{
-		return false;
-	}
-
-	// Move all vars of the category to move into temp array
-	TArray<FBPVariableDescription> VarsToMove;
-	for(int32 VarIdx = Blueprint->NewVariables.Num()-1; VarIdx >= 0; VarIdx--)
-	{
-		if(Blueprint->NewVariables[VarIdx].Category == CategoryToMove)
-		{
-			VarsToMove.Insert(Blueprint->NewVariables[VarIdx], 0);
-			Blueprint->NewVariables.RemoveAt(VarIdx);
-		}
-	}
-
-	// Bail if nothing to move
-	if(VarsToMove.Num() == 0)
-	{
-		return false;
-	}
-
-	// Find insertion point
-	int32 InsertionIndex = FindFirstNewVarOfCategory(Blueprint, TargetCategory);
-
-	// Now insert the nodes we want to move at that point
-	check(InsertionIndex != INDEX_NONE);
-	for(int32 VarIdx=0; VarIdx<VarsToMove.Num(); VarIdx++)
-	{
-		Blueprint->NewVariables.Insert(VarsToMove[VarIdx], InsertionIndex + VarIdx);
-	}
-
-	FBlueprintEditorUtils::MarkBlueprintAsStructurallyModified(Blueprint);
-
-	return true;
-}
 
 int32 FBlueprintEditorUtils::FindTimelineIndex(const UBlueprint* Blueprint, const FName& InName) 
 {
diff --git a/Engine/Source/Editor/UnrealEd/Public/Kismet2/BlueprintEditorUtils.h b/Engine/Source/Editor/UnrealEd/Public/Kismet2/BlueprintEditorUtils.h
--- a/Engine/Source/Editor/UnrealEd/Public/Kismet2/BlueprintEditorUtils.h
+++ b/Engine/Source/Editor/UnrealEd/Public/Kismet2/BlueprintEditorUtils.h
@@ -702,9 +702,6 @@ public:
 	/** Change the order of variables in the Blueprint */
 	static bool MoveVariableBeforeVariable(UBlueprint* Blueprint, FName VarNameToMove, FName TargetVarName, bool bDontRecompile);
 
-	/** Move all vars of the supplied category */
-	static bool MoveCategoryBeforeCategory(UBlueprint* Blueprint, FName CategoryToMove, FName TargetCategory);
-
 	/** Find first variable of the supplied category */
 	static int32 FindFirstNewVarOfCategory(const UBlueprint* Blueprint, FName Category);
 
diff --git a/Engine/Source/Runtime/Engine/Classes/Engine/Blueprint.h b/Engine/Source/Runtime/Engine/Classes/Engine/Blueprint.h
--- a/Engine/Source/Runtime/Engine/Classes/Engine/Blueprint.h
+++ b/Engine/Source/Runtime/Engine/Classes/Engine/Blueprint.h
@@ -364,6 +364,10 @@ class ENGINE_API UBlueprint : public UBlueprintCore
 	UPROPERTY()
 	TArray<struct FBPVariableDescription> NewVariables;
 
+	/** Array of user sorted categories */
+	UPROPERTY()
+	TArray<FName> CategorySorting;
+
 	/** Array of info about the interfaces we implement in this blueprint */
 	UPROPERTY(AssetRegistrySearchable)
 	TArray<struct FBPInterfaceDescription> ImplementedInterfaces;
