#UE4 Diff panel now shows components and defaults (defaults were formerly visible in p4merge or similar third party diff tool)

diff --git a/Engine/Source/Developer/AssetTools/Private/AssetTypeActions/AssetTypeActions_Blueprint.cpp b/Engine/Source/Developer/AssetTools/Private/AssetTypeActions/AssetTypeActions_Blueprint.cpp
--- a/Engine/Source/Developer/AssetTools/Private/AssetTypeActions/AssetTypeActions_Blueprint.cpp
+++ b/Engine/Source/Developer/AssetTools/Private/AssetTypeActions/AssetTypeActions_Blueprint.cpp
@@ -222,8 +222,7 @@ void FAssetTypeActions_Blueprint::PerformAssetDiff(UObject* OldAsset, UObject* N
 					  .BlueprintNew(NewBlueprint)
 					  .OldRevision(OldRevision)
 					  .NewRevision(NewRevision)
-					  .ShowAssetNames(!bIsSingleAsset)
-					  .OpenInDefaults(const_cast<FAssetTypeActions_Blueprint*>(this), &FAssetTypeActions_Blueprint::OpenInDefaults) );
+					  .ShowAssetNames(!bIsSingleAsset) );
 
 	// Make this window a child of the modal window if we've been spawned while one is active.
 	TSharedPtr<SWindow> ActiveModal = FSlateApplication::Get().GetActiveModalWindow();
@@ -251,21 +250,6 @@ UThumbnailInfo* FAssetTypeActions_Blueprint::GetThumbnailInfo(UObject* Asset) co
 	return ThumbnailInfo;
 }
 
-void FAssetTypeActions_Blueprint::OpenInDefaults( const UBlueprint* OldBlueprint, const UBlueprint* NewBlueprint ) const
-{
-	const bool bComparedBlueprintsHaveGeneratedClasses = *(OldBlueprint->GeneratedClass) && *(NewBlueprint->GeneratedClass);
-	ensure(bComparedBlueprintsHaveGeneratedClasses);
-	if (bComparedBlueprintsHaveGeneratedClasses)
-	{
-		const FString OldTextFilename = DumpAssetToTempFile(OldBlueprint->GeneratedClass->GetDefaultObject());
-		const FString NewTextFilename = DumpAssetToTempFile(NewBlueprint->GeneratedClass->GetDefaultObject());
-		const FString DiffCommand = GetDefault<UEditorLoadingSavingSettings>()->TextDiffToolPath.FilePath;
-
-		FAssetToolsModule& AssetToolsModule = FModuleManager::Get().LoadModuleChecked<FAssetToolsModule>("AssetTools");
-     	AssetToolsModule.Get().CreateDiffProcess(DiffCommand, OldTextFilename, NewTextFilename);
-	}
-}
-
 FText FAssetTypeActions_Blueprint::GetAssetDescription(const FAssetData& AssetData) const
 {
 	if (const FString* pDescription = AssetData.TagsAndValues.Find(GET_MEMBER_NAME_CHECKED(UBlueprint, BlueprintDescription)))
diff --git a/Engine/Source/Developer/AssetTools/Private/AssetTypeActions/AssetTypeActions_Blueprint.h b/Engine/Source/Developer/AssetTools/Private/AssetTypeActions/AssetTypeActions_Blueprint.h
--- a/Engine/Source/Developer/AssetTools/Private/AssetTypeActions/AssetTypeActions_Blueprint.h
+++ b/Engine/Source/Developer/AssetTools/Private/AssetTypeActions/AssetTypeActions_Blueprint.h
@@ -25,10 +25,6 @@ protected:
 	virtual bool CanCreateNewDerivedBlueprint() const;
 
 private:
-
-	/* Called to open the blueprint defaults view, this opens whatever text diff tool the user has */
-	void OpenInDefaults(const class UBlueprint* OldBlueprint, const class UBlueprint* NewBlueprint ) const;
-
 	/** Handler for when EditDefaults is selected */
 	void ExecuteEditDefaults(TArray<TWeakObjectPtr<UBlueprint>> Objects);
 
diff --git a/Engine/Source/Developer/AssetTools/Private/AssetTypeActions/AssetTypeActions_BlueprintGeneratedClass.cpp b/Engine/Source/Developer/AssetTools/Private/AssetTypeActions/AssetTypeActions_BlueprintGeneratedClass.cpp
--- a/Engine/Source/Developer/AssetTools/Private/AssetTypeActions/AssetTypeActions_BlueprintGeneratedClass.cpp
+++ b/Engine/Source/Developer/AssetTools/Private/AssetTypeActions/AssetTypeActions_BlueprintGeneratedClass.cpp
@@ -211,8 +211,7 @@ void FAssetTypeActions_BlueprintGeneratedClass::PerformAssetDiff(UObject* OldAss
 					  .BlueprintOld(OldBlueprint)
 					  .BlueprintNew(NewBlueprint)
 					  .OldRevision(OldRevision)
-					  .NewRevision(NewRevision)
-					  .OpenInDefaults(const_cast<FAssetTypeActions_BlueprintGeneratedClass*>(this), &FAssetTypeActions_BlueprintGeneratedClass::OpenInDefaults) );
+					  .NewRevision(NewRevision) );
 
 	FSlateApplication::Get().AddWindow( Window.ToSharedRef(), true );
 }
diff --git a/Engine/Source/Editor/Kismet/Private/BlueprintEditor.cpp b/Engine/Source/Editor/Kismet/Private/BlueprintEditor.cpp
--- a/Engine/Source/Editor/Kismet/Private/BlueprintEditor.cpp
+++ b/Engine/Source/Editor/Kismet/Private/BlueprintEditor.cpp
@@ -899,7 +899,7 @@ void FBlueprintEditor::EnsureBlueprintIsUpToDate(UBlueprint* BlueprintObj)
 			BlueprintObj->SimpleConstructionScript->SetFlags(RF_Transactional);
 
 			// Recreate (or create) any widgets that depend on the SCS
-			SCSEditor = SNew(SSCSEditor, SharedThis(this), BlueprintObj->SimpleConstructionScript);
+			SCSEditor = SNew(SSCSEditor, SharedThis(this), BlueprintObj->SimpleConstructionScript, BlueprintObj);
 			SCSViewport = SAssignNew(SCSViewport, SSCSEditorViewport) .BlueprintEditor(SharedThis(this));
 		}
 
@@ -1655,7 +1655,7 @@ void FBlueprintEditor::CreateDefaultTabContents(const TArray<UBlueprint*>& InBlu
 		InBlueprint->ParentClass->IsChildOf(AActor::StaticClass()) && 
 		InBlueprint->SimpleConstructionScript )
 	{
-		SCSEditor = SNew(SSCSEditor, SharedThis(this), InBlueprint->SimpleConstructionScript);
+		SCSEditor = SNew(SSCSEditor, SharedThis(this), InBlueprint->SimpleConstructionScript, InBlueprint);
 		
 		SCSViewport = SAssignNew(SCSViewport, SSCSEditorViewport) .BlueprintEditor(SharedThis(this));
 	}
diff --git a/Engine/Source/Editor/Kismet/Private/SBlueprintDiff.cpp b/Engine/Source/Editor/Kismet/Private/SBlueprintDiff.cpp
--- a/Engine/Source/Editor/Kismet/Private/SBlueprintDiff.cpp
+++ b/Engine/Source/Editor/Kismet/Private/SBlueprintDiff.cpp
@@ -1,17 +1,83 @@
 // Copyright 1998-2014 Epic Games, Inc. All Rights Reserved.
 
 #include "BlueprintEditorPrivatePCH.h"
+
+#include "BlueprintEditor.h"
+#include "BlueprintEditorCommands.h"
+#include "BlueprintEditorModes.h"
 #include "BlueprintUtilities.h"
+#include "Editor/Kismet/Public/SBlueprintEditorToolbar.h"
+#include "Editor/PropertyEditor/Public/IDetailsView.h"
+#include "Editor/PropertyEditor/Public/PropertyEditorModule.h"
 #include "Editor/UnrealEd/Public/EdGraphUtilities.h"
-#include "BlueprintEditorCommands.h"
-#include "BlueprintEditor.h"
-#include "SBlueprintDiff.h"
+#include "GraphDiffControl.h"
 #include "GraphEditor.h"
+#include "Kismet2/BlueprintEditorUtils.h"
+#include "SBlueprintDiff.h"
 #include "SGraphTitleBar.h"
-#include "GraphDiffControl.h"
+#include "SSCSEditor.h"
+#include "WorkflowOrientedApp/SModeWidget.h"
 
 #define LOCTEXT_NAMESPACE "SBlueprintDif"
 
+struct FPropertyFilter : public TSharedFromThis<FPropertyFilter>
+{
+	FPropertyFilter( TSet< const UProperty*> InHiddenProperties )
+		: HiddenProperties( InHiddenProperties )
+	{
+	}
+
+	~FPropertyFilter()
+	{
+		UE_LOG(LogBlueprintUserMessages, Verbose, TEXT(" Nuking property filter.."));
+	}
+
+	bool IsVisible(const FPropertyAndParent& Property)
+	{
+		return !HiddenProperties.Contains(&Property.Property);
+	}
+
+	TSet< const UProperty*> HiddenProperties;
+};
+
+typedef TMap< FString, const UProperty* > FNamePropertyMap;
+
+static FNamePropertyMap GetCDOProperties( const UBlueprint* ForBlueprint )
+{
+	FNamePropertyMap Ret;
+
+	if( ForBlueprint )
+	{
+		if( const UClass* GeneratedClass = ForBlueprint->GeneratedClass )
+		{
+			if( const UObject* CDO = GeneratedClass->ClassDefaultObject )
+			{
+				if( const UClass* CDOClass = CDO->GetClass() )
+				{
+					check( CDOClass == GeneratedClass );
+					for (TFieldIterator<UProperty> PropertyIt(CDOClass, EFieldIteratorFlags::IncludeSuper); PropertyIt; ++PropertyIt)
+					{
+						Ret.Add(PropertyIt->GetName(), *PropertyIt);
+					}
+				}
+			}
+		}
+	}
+
+	return Ret;
+}
+
+const UObject* GetCDO( const UBlueprint* ForBlueprint )
+{
+	if( !ForBlueprint 
+		|| !ForBlueprint->GeneratedClass )
+	{
+		return NULL;
+	}
+
+	return ForBlueprint->GeneratedClass->ClassDefaultObject;
+}
+
 /** Individual Diff item shown in the list of diffs */
 struct FDiffResultItem: public TSharedFromThis<FDiffResultItem>
 {
@@ -102,7 +168,6 @@ TSharedRef<SWidget> FListItemGraphToDiff::GenerateWidget()
 			.ColorAndOpacity(Color)
 			.Image(GetIconForGraph(GraphNew))
 		];
-
 }
 
 //////////////////////////////////////////////////////////////////////////
@@ -333,6 +398,74 @@ FDiffPanel::FDiffPanel()
 	Blueprint = NULL;
 }
 
+// Generates a widget showing CDO values. NewBlueprint's CDO values will be displayed by the widget (Base just used for determining what to hide):
+static TSharedRef<SWidget> GenerateCDODiff( UBlueprint const* BaseBlueprint, UBlueprint const* NewBlueprint )
+{
+	const auto GenerateEmptyPanel = []( const FText& Message) { return SNew(STextBlock).Text(Message); };
+
+	// filter out properties that have not changed from the GraphToDiff's blueprint:	
+	FNamePropertyMap BaseProperties = GetCDOProperties(BaseBlueprint);
+	FNamePropertyMap BlueprintProperties = GetCDOProperties(NewBlueprint);
+
+	if( BlueprintProperties.Num() == 0 )
+	{
+		if( NewBlueprint )
+		{
+			return GenerateEmptyPanel( FText::Format( LOCTEXT("BPNoCDO", "Blueprint {0} has no Defaults"), FText::FromString( NewBlueprint->GetName() ) ) );
+		}
+		else
+		{
+			return GenerateEmptyPanel( LOCTEXT("BPMissingBP", "Missing Blueprint" ) );
+		}
+	}
+
+	const UObject* BaseBlueprintCDO = GetCDO(BaseBlueprint);
+	const UObject* NewBlueprintCDO = GetCDO(NewBlueprint);
+
+	TSet< const UProperty*> HiddenProperties;
+	for (const auto& Property : BlueprintProperties)
+	{
+		const FString& PropertyName = Property.Key;
+		const UProperty* PropertyValue = Property.Value;
+
+		const UProperty** BaseProperty = BaseProperties.Find(PropertyName);
+		if (BaseProperty && (*BaseProperty)->SameType(PropertyValue))
+		{
+			const void* MyValue = PropertyValue->ContainerPtrToValuePtr<void>(NewBlueprintCDO);
+			const void* BaseValue = (*BaseProperty)->ContainerPtrToValuePtr<void>(BaseBlueprintCDO);
+
+			if ((*BaseProperty)->Identical(MyValue, BaseValue, 0))
+			{
+				HiddenProperties.Add(PropertyValue);
+			}
+		}
+	}
+
+	if( HiddenProperties.Num() == BlueprintProperties.Num() )
+	{
+		return GenerateEmptyPanel(LOCTEXT("BPNoCDODifference", "No differences found in Defaults"));
+	}
+
+	FDetailsViewArgs DetailsViewArgs;
+	DetailsViewArgs.bShowOptions = false;
+	FPropertyEditorModule& EditModule = FModuleManager::Get().GetModuleChecked<FPropertyEditorModule>("PropertyEditor");
+	auto DetailsView = EditModule.CreateDetailView(DetailsViewArgs);
+	DetailsView->SetIsPropertyEditingEnabledDelegate(FIsPropertyEditingEnabled::CreateStatic([]{return false; }));
+	// This filter goes out of scope and is destroyed when this function returns (the control just keeps a weak ptr)
+	// but that seems to be long enough to hide the properties that we want to hide:
+	auto Filter = TSharedPtr<FPropertyFilter>(new FPropertyFilter(HiddenProperties));
+	DetailsView->SetIsPropertyVisibleDelegate(FIsPropertyVisible::CreateSP(Filter.ToSharedRef(), &FPropertyFilter::IsVisible));
+	// This is a read only details view (see the property editing delegate), but it is not const correct:
+	DetailsView->SetObject(const_cast<UObject*>(NewBlueprintCDO) );
+
+	return DetailsView;
+}
+
+static TSharedRef<SWidget> GenerateComponentsDiff( UBlueprint const* BaseBlueprint, UBlueprint const* DisplayedBlueprint )
+{
+	return SNew(SSCSEditor, TSharedPtr<FBlueprintEditor>(), DisplayedBlueprint->SimpleConstructionScript, const_cast<UBlueprint*>(DisplayedBlueprint) );
+}
+
 //////////////////////////////////////////////////////////////////////////
 // SBlueprintDiff
 
@@ -354,8 +487,6 @@ void SBlueprintDiff::Construct( const FArguments& InArgs)
 	PanelOld.bShowAssetName = InArgs._ShowAssetNames;
 	PanelNew.bShowAssetName = InArgs._ShowAssetNames;
 
-	OpenInDefaults = InArgs._OpenInDefaults;
-
 	bLockViews = true;
 
 	TArray<UEdGraph*> GraphsOld,GraphsNew;
@@ -390,100 +521,64 @@ void SBlueprintDiff::Construct( const FArguments& InArgs)
 		}
 	}
 
+	TAttribute<FName> GetActiveMode(this, &SBlueprintDiff::GetCurrentMode);
+	FOnModeChangeRequested SetActiveMode = FOnModeChangeRequested::CreateRaw(this, &SBlueprintDiff::SetCurrentMode);
+
+	auto Widgets = FBlueprintEditorToolbar::GenerateToolbarWidgets( InArgs._BlueprintNew, GetActiveMode, SetActiveMode );
+
+	TSharedRef<SHorizontalBox> MiscWidgets = SNew(SHorizontalBox);
+
+	for (int32 WidgetIdx = 0; WidgetIdx < Widgets.Num(); ++WidgetIdx)
+	{
+		MiscWidgets->AddSlot()
+			.AutoWidth()
+			.VAlign(VAlign_Center)
+			[
+				Widgets[WidgetIdx].ToSharedRef()
+			];
+	}
+
 	this->ChildSlot
-		[	
+		[
 			SNew(SBorder)
-			.BorderImage(FEditorStyle::GetBrush("ToolPanel.GroupBorder"))
-			.Content()
+			.BorderImage(FEditorStyle::GetBrush( "Docking.Tab", ".ContentAreaBrush" ))
 			[
-				SNew(SBorder)
+				SNew(SVerticalBox)
+				+SVerticalBox::Slot()
+				.AutoHeight()
+				.Padding(0.0f, 2.0f, 0.0f, 2.0f)
 				[
-					SNew(SVerticalBox)
-					+SVerticalBox::Slot()
-					.FillHeight(1.f)
+					SNew(SHorizontalBox)
+					+ SHorizontalBox::Slot()
 					[
-						SNew(SSplitter)
-						+SSplitter::Slot()
-						.Value(0.2f)
+						SNew(SSpacer)
+					]
+					+ SHorizontalBox::Slot()
+						.AutoWidth()
 						[
-							SNew(SBorder)
+							// Nested box just give us padding similar to the main blueprint editor,
+							// one border has the dark grey "Toolbar.Background" style:
+							SNew(SHorizontalBox)
+							+ SHorizontalBox::Slot()
+							.Padding(0.0f, 4.f, 0.0f, 12.f)
 							[
-								SNew(SSplitter)
-								.Orientation(Orient_Vertical)
-								+SSplitter::Slot()
-								.Value(0.3f)
+								SNew(SBorder)
+								.BorderImage(FEditorStyle::GetBrush(TEXT("Toolbar.Background")))
 								[
-									SNew(SVerticalBox)
-									+SVerticalBox::Slot()
-									.AutoHeight()
-									.HAlign(HAlign_Left)
-									[
-										//toggle lock button
-										SNew(SButton)
-										.OnClicked(this, &SBlueprintDiff::OnToggleLockView)
-										.Content()
-										[
-											SNew(SImage)
-											.Image(this, &SBlueprintDiff::GetLockViewImage)
-											.ToolTipText(LOCTEXT("BPDifLock", "Lock the blueprint views?"))
-										]
-									]
-									+SVerticalBox::Slot()
-									.FillHeight(1.f)
-									[
-										//graph selection 
-										SNew(SBorder)
-										[
-											SNew(SVerticalBox)
-											+SVerticalBox::Slot()
-											.AutoHeight()
-											[
-												SNew(SBorder)
-												.BorderImage(FEditorStyle::GetBrush("PropertyWindow.CategoryBackground"))
-												.Padding(FMargin(2.0f))
-												.ForegroundColor(FEditorStyle::GetColor("PropertyWindow.CategoryForeground"))
-												.ToolTipText(LOCTEXT("BlueprintDifGraphsToolTip", "Select Graph to Diff"))
-												.HAlign(HAlign_Center)
-												[
-													SNew(STextBlock)
-													.Text(LOCTEXT("BlueprintDifGraphs", "Graphs"))
-												]
-											]
-											+SVerticalBox::Slot()
-											.FillHeight(1.f)
-											[
-												SAssignNew(GraphsToDiff, SListViewType)
-												.ItemHeight(24)
-												.ListItemsSource( &Graphs )
-												.OnGenerateRow( this, &SBlueprintDiff::OnGenerateRow )
-												.SelectionMode( ESelectionMode::Single )
-												.OnSelectionChanged(this, &SBlueprintDiff::OnSelectionChanged)
-											]
-										]
-									]
-								]
-								+SSplitter::Slot()
-								.Value(0.7f)
-								[
-									GenerateToolbar()
+									MiscWidgets
 								]
 							]
 						]
-						+SSplitter::Slot()
-						.Value(0.8f)
-						[
-							SNew(SHorizontalBox)
-							+SHorizontalBox::Slot()
-							.FillWidth(1.f)
-							[
-								//diff window
-								GenerateDiffWindow()
-							]
-						]
-					]
+				]
+				+ SVerticalBox::Slot()
+				[
+					SAssignNew(ModeContents, SBorder)
 				]
 			]
+			
 		];
+
+	SetCurrentMode( FBlueprintEditorApplicationModes::StandardBlueprintEditorMode );
 }
 END_SLATE_FUNCTION_BUILD_OPTIMIZATION
 
@@ -524,7 +619,7 @@ void SBlueprintDiff::OnGraphChanged(FListItemGraphToDiff* Diff)
 void SBlueprintDiff::FocusOnGraphRevisions( class UEdGraph* GraphOld, class UEdGraph* GraphNew, FListItemGraphToDiff* Diff )
 {
 	DisablePinDiffFocus();
-	HandleGraphChanged( GraphOld->GetName() );
+	HandleGraphChanged( GraphOld ? GraphOld->GetName() : GraphNew->GetName() );
 
 	ResetGraphEditors();
 
@@ -714,12 +809,6 @@ bool FDiffPanel::CanCopyNodes() const
 	return false;
 }
 
-FReply SBlueprintDiff::OnOpenInDefaults()
-{
-	OpenInDefaults.ExecuteIfBound(PanelOld.Blueprint, PanelNew.Blueprint);
-	return FReply::Handled();
-}
-
 SGraphEditor* SBlueprintDiff::GetGraphEditorForGraph( UEdGraph* Graph ) const
 {
 	if(PanelOld.GraphEditor.Pin()->GetCurrentGraph() == Graph)
@@ -774,18 +863,6 @@ TSharedRef<SWidget> SBlueprintDiff::GenerateDiffWindow()
 TSharedRef<SWidget> SBlueprintDiff::GenerateToolbar()
 {
 	return SNew(SVerticalBox)
-		+SVerticalBox::Slot()
-		.AutoHeight()
-		[
-			//open in p4dif tool
-			SNew(SButton)
-			.OnClicked(this, &SBlueprintDiff::OnOpenInDefaults)
-			.Content()
-			[
-				SNew(STextBlock)
-				.Text(LOCTEXT("DifBlueprintDefaultsToolTip", "Diff Blueprint Defaults"))
-			]
-		]
 		+SVerticalBox::Slot()
 		.FillHeight(1.f)
 		[
@@ -814,6 +891,171 @@ void SBlueprintDiff::HandleGraphChanged( const FString& GraphName )
 	PanelNew.GeneratePanel(GraphNew, GraphOld);
 }
 
+TSharedRef<SWidget> SBlueprintDiff::GenerateGraphPanel()
+{
+	return SNew(SBorder)
+	.BorderImage(FEditorStyle::GetBrush("ToolPanel.GroupBorder"))
+	.Content()
+	[
+		SNew(SBorder)
+		[
+			SNew(SVerticalBox)
+			+ SVerticalBox::Slot()
+			.FillHeight(1.f)
+			[
+				SNew(SSplitter)
+				+ SSplitter::Slot()
+				.Value(0.2f)
+				[
+					SNew(SBorder)
+					[
+						SNew(SSplitter)
+						.Orientation(Orient_Vertical)
+						+ SSplitter::Slot()
+						.Value(0.3f)
+						[
+							SNew(SVerticalBox)
+							+ SVerticalBox::Slot()
+							.AutoHeight()
+							.HAlign(HAlign_Left)
+							[
+								//toggle lock button
+								SNew(SButton)
+								.OnClicked(this, &SBlueprintDiff::OnToggleLockView)
+								.Content()
+								[
+									SNew(SImage)
+									.Image(this, &SBlueprintDiff::GetLockViewImage)
+									.ToolTipText(LOCTEXT("BPDifLock", "Lock the blueprint views?"))
+								]
+							]
+							+ SVerticalBox::Slot()
+								.FillHeight(1.f)
+								[
+									//graph selection 
+									SNew(SBorder)
+									[
+										SNew(SVerticalBox)
+										+ SVerticalBox::Slot()
+										.AutoHeight()
+										[
+											SNew(SBorder)
+											.BorderImage(FEditorStyle::GetBrush("PropertyWindow.CategoryBackground"))
+											.Padding(FMargin(2.0f))
+											.ForegroundColor(FEditorStyle::GetColor("PropertyWindow.CategoryForeground"))
+											.ToolTipText(LOCTEXT("BlueprintDifGraphsToolTip", "Select Graph to Diff"))
+											.HAlign(HAlign_Center)
+											[
+												SNew(STextBlock)
+												.Text(LOCTEXT("BlueprintDifGraphs", "Graphs"))
+											]
+										]
+										+ SVerticalBox::Slot()
+											.FillHeight(1.f)
+											[
+												SAssignNew(GraphsToDiff, SListViewType)
+												.ItemHeight(24)
+												.ListItemsSource(&Graphs)
+												.OnGenerateRow(this, &SBlueprintDiff::OnGenerateRow)
+												.SelectionMode(ESelectionMode::Single)
+												.OnSelectionChanged(this, &SBlueprintDiff::OnSelectionChanged)
+											]
+									]
+								]
+						]
+						+ SSplitter::Slot()
+							.Value(0.7f)
+							[
+								GenerateToolbar()
+							]
+					]
+				]
+				+ SSplitter::Slot()
+					.Value(0.8f)
+					[
+						SNew(SHorizontalBox)
+						+ SHorizontalBox::Slot()
+						.FillWidth(1.f)
+						[
+							//diff window
+							GenerateDiffWindow()
+						]
+					]
+			]
+		]
+	];
+}
+
+TSharedRef<SWidget> SBlueprintDiff::GenerateDefaultsPanel()
+{
+	//Splitter for left and right blueprint. Current convention is for the local (probably newer?) blueprint to be on the right:
+	return SNew(SSplitter)
+		+ SSplitter::Slot()
+		.Value(0.5f)
+		[
+			SNew(SBorder)
+			.VAlign(VAlign_Fill)
+			[
+				GenerateCDODiff(PanelNew.Blueprint, PanelOld.Blueprint)
+			]
+		]
+	+ SSplitter::Slot()
+		.Value(0.5f)
+		[
+			SNew(SBorder)
+			.VAlign(VAlign_Fill)
+			[
+				GenerateCDODiff(PanelOld.Blueprint, PanelNew.Blueprint)
+			]
+		];
+}
+
+TSharedRef<SWidget> SBlueprintDiff::GenerateComponentsPanel()
+{
+	//Splitter for left and right blueprint. Current convention is for the local (probably newer?) blueprint to be on the right:
+	return SNew(SSplitter)
+		+ SSplitter::Slot()
+		.Value(0.5f)
+		[
+			SNew(SBorder)
+			.VAlign(VAlign_Fill)
+			[
+				GenerateComponentsDiff(PanelNew.Blueprint, PanelOld.Blueprint)
+			]
+		]
+	+ SSplitter::Slot()
+		.Value(0.5f)
+		[
+			SNew(SBorder)
+			.VAlign(VAlign_Fill)
+			[
+				GenerateComponentsDiff(PanelOld.Blueprint, PanelNew.Blueprint)
+			]
+		];
+}
+
+void SBlueprintDiff::SetCurrentMode(FName NewMode)
+{
+	CurrentMode = NewMode;
+
+	if( NewMode == FBlueprintEditorApplicationModes::StandardBlueprintEditorMode)
+	{
+		ModeContents->SetContent( GenerateGraphPanel() );
+	}
+	else if( NewMode == FBlueprintEditorApplicationModes::BlueprintDefaultsMode )
+	{
+		ModeContents->SetContent( GenerateDefaultsPanel() );
+	}
+	else if (NewMode == FBlueprintEditorApplicationModes::BlueprintComponentsMode)
+	{
+		ModeContents->SetContent( GenerateComponentsPanel() );
+	}
+	else
+	{
+		ensureMsgf(false, TEXT("Diff panel does not support mode %s"), *NewMode.ToString() );
+	}
+}
+
 FReply SBlueprintDiff::OnKeyDown( const FGeometry& MyGeometry, const FKeyboardEvent& InKeyboardEvent ) 
 {
 	TArray< FGraphToDiff>  Selected = GraphsToDiff->GetSelectedItems();
diff --git a/Engine/Source/Editor/Kismet/Private/SBlueprintEditorToolbar.cpp b/Engine/Source/Editor/Kismet/Private/SBlueprintEditorToolbar.cpp
--- a/Engine/Source/Editor/Kismet/Private/SBlueprintEditorToolbar.cpp
+++ b/Engine/Source/Editor/Kismet/Private/SBlueprintEditorToolbar.cpp
@@ -629,81 +629,19 @@
 		FToolBarExtensionDelegate::CreateSP( this, &FBlueprintEditorToolbar::FillComponentsToolbar ) );
 }
 
-
 void FBlueprintEditorToolbar::FillBlueprintEditorModesToolbar(FToolBarBuilder& ToolbarBuilder)
 {
 	TSharedPtr<FBlueprintEditor> BlueprintEditorPtr = BlueprintEditor.Pin();
 	UBlueprint* BlueprintObj = BlueprintEditorPtr->GetBlueprintObj();
 
-	if( !BlueprintObj ||
-		(!FBlueprintEditorUtils::IsLevelScriptBlueprint(BlueprintObj) 
-		&& !FBlueprintEditorUtils::IsInterfaceBlueprint(BlueprintObj)
-		&& !BlueprintObj->bIsNewlyCreated)
-		)
+	TAttribute<FName> GetActiveMode(BlueprintEditorPtr.ToSharedRef(), &FBlueprintEditor::GetCurrentMode);
+	FOnModeChangeRequested SetActiveMode = FOnModeChangeRequested::CreateSP(BlueprintEditorPtr.ToSharedRef(), &FBlueprintEditor::SetCurrentMode);
+
+	TArray< TSharedPtr< SWidget > > ToolbarWidgets = GenerateToolbarWidgets( BlueprintObj, GetActiveMode, SetActiveMode );
+	
+	for( const auto& Widget : ToolbarWidgets )
 	{
-		TAttribute<FName> GetActiveMode(BlueprintEditorPtr.ToSharedRef(), &FBlueprintEditor::GetCurrentMode);
-		FOnModeChangeRequested SetActiveMode = FOnModeChangeRequested::CreateSP(BlueprintEditorPtr.ToSharedRef(), &FBlueprintEditor::SetCurrentMode);
-
-		// Left side padding
-		BlueprintEditorPtr->AddToolbarWidget(SNew(SSpacer).Size(FVector2D(4.0f, 1.0f)));
-
-		BlueprintEditorPtr->AddToolbarWidget(
-			SNew( STutorialWrapper, TEXT("DefaultsMode") )
-			[
-				SNew(SModeWidget, FBlueprintEditorApplicationModes::GetLocalizedMode( FBlueprintEditorApplicationModes::BlueprintDefaultsMode ), FBlueprintEditorApplicationModes::BlueprintDefaultsMode)
-				.OnGetActiveMode(GetActiveMode)
-				.OnSetActiveMode(SetActiveMode)
-				.ToolTip(IDocumentation::Get()->CreateToolTip(
-					LOCTEXT("BlueprintDefaultsModeButtonTooltip", "Switch to Blueprint Defaults Mode"),
-					NULL,
-					TEXT("Shared/Editors/BlueprintEditor"),
-					TEXT("DefaultsMode")))
-				.IconImage(FEditorStyle::GetBrush("FullBlueprintEditor.SwitchToBlueprintDefaultsMode"))
-				.SmallIconImage(FEditorStyle::GetBrush("FullBlueprintEditor.SwitchToBlueprintDefaultsMode.Small"))
-			]
-		);
-
-		BlueprintEditorPtr->AddToolbarWidget(SNew(SBlueprintModeSeparator));
-
-		BlueprintEditorPtr->AddToolbarWidget(
-			SNew( STutorialWrapper, TEXT("ComponentsMode") )
-			[
-				SNew(SModeWidget, FBlueprintEditorApplicationModes::GetLocalizedMode( FBlueprintEditorApplicationModes::BlueprintComponentsMode ), FBlueprintEditorApplicationModes::BlueprintComponentsMode)
-				.OnGetActiveMode(GetActiveMode)
-				.OnSetActiveMode(SetActiveMode)
-				.CanBeSelected(BlueprintEditorPtr.Get(), &FBlueprintEditor::CanAccessComponentsMode)
-				.ToolTip(IDocumentation::Get()->CreateToolTip(
-					LOCTEXT("ComponentsModeButtonTooltip", "Switch to Components Mode"),
-					NULL,
-					TEXT("Shared/Editors/BlueprintEditor"),
-					TEXT("ComponentsMode")))
-				.IconImage(FEditorStyle::GetBrush("FullBlueprintEditor.SwitchToComponentsMode"))
-				.SmallIconImage(FEditorStyle::GetBrush("FullBlueprintEditor.SwitchToComponentsMode.Small"))
-			]
-		);
-
-		BlueprintEditorPtr->AddToolbarWidget(SNew(SBlueprintModeSeparator));
-
-		BlueprintEditorPtr->AddToolbarWidget(
-			SNew( STutorialWrapper, TEXT("GraphMode") )
-			[
-				SNew(SModeWidget, FBlueprintEditorApplicationModes::GetLocalizedMode( FBlueprintEditorApplicationModes::StandardBlueprintEditorMode ), FBlueprintEditorApplicationModes::StandardBlueprintEditorMode)
-				.OnGetActiveMode(GetActiveMode)
-				.OnSetActiveMode(SetActiveMode)
-				.CanBeSelected(BlueprintEditorPtr.Get(), &FBlueprintEditor::IsEditingSingleBlueprint)
-				.ToolTip(IDocumentation::Get()->CreateToolTip(
-					LOCTEXT("GraphModeButtonTooltip", "Switch to Graph Editing Mode"),
-					NULL,
-					TEXT("Shared/Editors/BlueprintEditor"),
-					TEXT("GraphMode")))
-				.ToolTipText(LOCTEXT("GraphModeButtonTooltip", "Switch to Graph Editing Mode"))
-				.IconImage(FEditorStyle::GetBrush("FullBlueprintEditor.SwitchToScriptingMode"))
-				.SmallIconImage(FEditorStyle::GetBrush("FullBlueprintEditor.SwitchToScriptingMode.Small"))
-			]
-		);
-		
-		// Right side padding
-		BlueprintEditorPtr->AddToolbarWidget(SNew(SSpacer).Size(FVector2D(4.0f, 1.0f)));
+		BlueprintEditorPtr->AddToolbarWidget( Widget.ToSharedRef() );
 	}
 }
 
@@ -864,5 +802,78 @@
 	}
 }
 
+TArray< TSharedPtr< SWidget> > FBlueprintEditorToolbar::GenerateToolbarWidgets(const UBlueprint* BlueprintObj, TAttribute<FName> ActiveModeGetter, FOnModeChangeRequested ActiveModeSetter)
+{
+	TArray< TSharedPtr< SWidget> > Ret;
+	if (!BlueprintObj ||
+		(!FBlueprintEditorUtils::IsLevelScriptBlueprint(BlueprintObj)
+		&& !FBlueprintEditorUtils::IsInterfaceBlueprint(BlueprintObj)
+		&& !BlueprintObj->bIsNewlyCreated)
+		)
+	{
+		// Left side padding
+		Ret.Add(SNew(SSpacer).Size(FVector2D(4.0f, 1.0f)));
 
+		Ret.Add(
+			SNew(STutorialWrapper, TEXT("DefaultsMode"))
+			[
+				SNew(SModeWidget, FBlueprintEditorApplicationModes::GetLocalizedMode(FBlueprintEditorApplicationModes::BlueprintDefaultsMode), FBlueprintEditorApplicationModes::BlueprintDefaultsMode)
+				.OnGetActiveMode(ActiveModeGetter)
+				.OnSetActiveMode(ActiveModeSetter)
+				.CanBeSelected(BlueprintObj ? FBlueprintEditorUtils::DoesSupportDefaults(BlueprintObj) : false)
+				.ToolTip(IDocumentation::Get()->CreateToolTip(
+					LOCTEXT("BlueprintDefaultsModeButtonTooltip", "Switch to Blueprint Defaults Mode"),
+					NULL,
+					TEXT("Shared/Editors/BlueprintEditor"),
+					TEXT("DefaultsMode")))
+				.IconImage(FEditorStyle::GetBrush("FullBlueprintEditor.SwitchToBlueprintDefaultsMode"))
+				.SmallIconImage(FEditorStyle::GetBrush("FullBlueprintEditor.SwitchToBlueprintDefaultsMode.Small"))
+			]
+		);
+
+		Ret.Add(SNew(SBlueprintModeSeparator));
+
+		Ret.Add(
+			SNew(STutorialWrapper, TEXT("ComponentsMode"))
+			[
+				SNew(SModeWidget, FBlueprintEditorApplicationModes::GetLocalizedMode(FBlueprintEditorApplicationModes::BlueprintComponentsMode), FBlueprintEditorApplicationModes::BlueprintComponentsMode)
+				.OnGetActiveMode(ActiveModeGetter)
+				.OnSetActiveMode(ActiveModeSetter)
+				.CanBeSelected(BlueprintObj ? FBlueprintEditorUtils::DoesSupportComponents(BlueprintObj) : false)
+				.ToolTip(IDocumentation::Get()->CreateToolTip(
+					LOCTEXT("ComponentsModeButtonTooltip", "Switch to Components Mode"),
+					NULL,
+					TEXT("Shared/Editors/BlueprintEditor"),
+					TEXT("ComponentsMode")))
+				.IconImage(FEditorStyle::GetBrush("FullBlueprintEditor.SwitchToComponentsMode"))
+				.SmallIconImage(FEditorStyle::GetBrush("FullBlueprintEditor.SwitchToComponentsMode.Small"))
+			]
+		);
+
+		Ret.Add(SNew(SBlueprintModeSeparator));
+
+		Ret.Add(
+			SNew(STutorialWrapper, TEXT("GraphMode"))
+			[
+				SNew(SModeWidget, FBlueprintEditorApplicationModes::GetLocalizedMode(FBlueprintEditorApplicationModes::StandardBlueprintEditorMode), FBlueprintEditorApplicationModes::StandardBlueprintEditorMode)
+				.OnGetActiveMode(ActiveModeGetter)
+				.OnSetActiveMode(ActiveModeSetter)
+				.CanBeSelected(BlueprintObj != NULL)
+				.ToolTip(IDocumentation::Get()->CreateToolTip(
+					LOCTEXT("GraphModeButtonTooltip", "Switch to Graph Editing Mode"),
+					NULL,
+					TEXT("Shared/Editors/BlueprintEditor"),
+					TEXT("GraphMode")))
+				.ToolTipText(LOCTEXT("GraphModeButtonTooltip", "Switch to Graph Editing Mode"))
+				.IconImage(FEditorStyle::GetBrush("FullBlueprintEditor.SwitchToScriptingMode"))
+				.SmallIconImage(FEditorStyle::GetBrush("FullBlueprintEditor.SwitchToScriptingMode.Small"))
+			]
+		);
+
+		// Right side padding
+		Ret.Add(SNew(SSpacer).Size(FVector2D(4.0f, 1.0f)));
+	}
+	return Ret;
+}
+
 #undef LOCTEXT_NAMESPACE
diff --git a/Engine/Source/Editor/Kismet/Private/SSCSEditor.cpp b/Engine/Source/Editor/Kismet/Private/SSCSEditor.cpp
--- a/Engine/Source/Editor/Kismet/Private/SSCSEditor.cpp
+++ b/Engine/Source/Editor/Kismet/Private/SSCSEditor.cpp
@@ -834,7 +834,7 @@ TSharedRef<SWidget> SSCS_RowWidget::GenerateWidgetForColumn( const FName& Column
 				.OnVerifyTextChanged( this, &SSCS_RowWidget::OnNameTextVerifyChanged )
 				.OnTextCommitted( this, &SSCS_RowWidget::OnNameTextCommit )
 				.IsSelected( this, &SSCS_RowWidget::IsSelectedExclusively )
-				.IsReadOnly( !NodePtr->CanRename() || !SCSEditor.Pin()->InEditingMode() );
+				.IsReadOnly( !NodePtr->CanRename() || (SCSEditor.IsValid() && !SCSEditor.Pin()->InEditingMode()) );
 
 		NodePtr->SetRenameRequestedDelegate(FSCSEditorTreeNode::FOnRenameRequested::CreateSP(InlineWidget.Get(), &SInlineEditableTextBlock::EnterEditingMode));
 
@@ -1825,7 +1825,7 @@ FString SSCS_RowWidget::GetDocumentationLink() const
 {
 	check(SCSEditor.IsValid());
 
-	if(NodePtr == SCSEditor.Pin()->SceneRootNodePtr
+	if( NodePtr == SCSEditor.Pin()->SceneRootNodePtr
 		|| NodePtr->IsNative() || NodePtr->IsInherited())
 	{
 		return TEXT("Shared/Editors/BlueprintEditor/ComponentsMode");
@@ -1838,7 +1838,7 @@ FString SSCS_RowWidget::GetDocumentationExcerptName() const
 {
 	check(SCSEditor.IsValid());
 
-	if(NodePtr == SCSEditor.Pin()->SceneRootNodePtr)
+	if( NodePtr == SCSEditor.Pin()->SceneRootNodePtr)
 	{
 		return TEXT("RootComponent");
 	}
@@ -1856,7 +1856,7 @@ FString SSCS_RowWidget::GetDocumentationExcerptName() const
 
 UBlueprint* SSCS_RowWidget::GetBlueprint() const
 {
-	return SCSEditor.Pin()->Kismet2Ptr.Pin()->GetBlueprintObj();
+	return SCSEditor.Pin()->Blueprint;
 }
 
 bool SSCS_RowWidget::OnNameTextVerifyChanged(const FText& InNewText, FText& OutErrorMessage)
@@ -1901,10 +1901,11 @@ void SSCS_RowWidget::OnNameTextCommit(const FText& InNewName, ETextCommit::Type
 // SSCSEditor
 
 
-void SSCSEditor::Construct( const FArguments& InArgs, TSharedPtr<FBlueprintEditor> InKismet2, USimpleConstructionScript* InSCS )
+void SSCSEditor::Construct( const FArguments& InArgs, TSharedPtr<FBlueprintEditor> InKismet2, USimpleConstructionScript* InSCS, UBlueprint* InBlueprint )
 {
 	Kismet2Ptr = InKismet2;
 	SCS = InSCS;
+	Blueprint = InBlueprint;
 
 	CommandList = MakeShareable( new FUICommandList );
 	CommandList->MapAction( FGenericCommands::Get().Cut,
@@ -1936,69 +1937,78 @@ void SSCSEditor::Construct( const FArguments& InArgs, TSharedPtr<FBlueprintEdito
 
 	FSlateBrush const* MobilityHeaderBrush = FEditorStyle::GetBrush(TEXT("ClassIcon.ComponentMobilityHeaderIcon"));
 
-	this->ChildSlot
-	[
-		SNew( STutorialWrapper, TEXT("ComponentsPanel") )
+	SAssignNew(SCSTreeWidget, SSCSTreeType)
+		.ToolTipText(LOCTEXT("DropAssetToAddComponent", "Drop asset here to add component."))
+		.SCSEditor(this)
+		.TreeItemsSource(&RootNodes)
+		.SelectionMode(ESelectionMode::Multi)
+		.OnGenerateRow(this, &SSCSEditor::MakeTableRowWidget)
+		.OnGetChildren(this, &SSCSEditor::OnGetChildrenForTree)
+		.OnSelectionChanged(this, &SSCSEditor::OnTreeSelectionChanged)
+		.OnContextMenuOpening(this, &SSCSEditor::CreateContextMenu)
+		.OnItemScrolledIntoView(this, &SSCSEditor::OnItemScrolledIntoView)
+		.ItemHeight(24)
+		.HeaderRow
+		(
+		SNew(SHeaderRow)
+
+		+ SHeaderRow::Column(SCS_ColumnName_Mobility)
+		.DefaultLabel(LOCTEXT("MobilityColumnLabel", "Mobility"))
+		.FixedWidth(16.0f) // mobility icons are 16px (16 slate-units = 16px, when application scale == 1)
+		.HeaderContent()
 		[
-			SNew(SVerticalBox)
-
-			+SVerticalBox::Slot()
-			.AutoHeight()
+			SNew(SHorizontalBox)
+			.ToolTip(SNew(SToolTip).Text(LOCTEXT("MobilityColumnTooltip", "Mobility")))
+			+ SHorizontalBox::Slot()
+			.FillWidth(1.0f)
+			.VAlign(VAlign_Center)
+			.HAlign(HAlign_Center)
 			[
-				SAssignNew(ComponentClassCombo, SComponentClassCombo)
-				.OnComponentClassSelected(this, &SSCSEditor::PerformComboAddClass)
+				SNew(SImage).Image(MobilityHeaderBrush)
 			]
+		]
 
-			+SVerticalBox::Slot()
-			.AutoHeight()
-			[
-				SNew(SSpacer)
-				.Size(FVector2D(0.0f,2.0f))
-			]
+	+ SHeaderRow::Column(SCS_ColumnName_ComponentClass)
+		.DefaultLabel(LOCTEXT("Class", "Class"))
+		.FillWidth(4)
 
-			+SVerticalBox::Slot()
-			.FillHeight(.8f)
-			[
-				SAssignNew(SCSTreeWidget, SSCSTreeType)
-				.ToolTipText(LOCTEXT("DropAssetToAddComponent", "Drop asset here to add component."))
-				.SCSEditor( this )
-				.TreeItemsSource( &RootNodes )
-				.SelectionMode(ESelectionMode::Multi)
-				.OnGenerateRow( this, &SSCSEditor::MakeTableRowWidget )
-				.OnGetChildren( this, &SSCSEditor::OnGetChildrenForTree )
-				.OnSelectionChanged( this, &SSCSEditor::OnTreeSelectionChanged )
-				.OnContextMenuOpening(this, &SSCSEditor::CreateContextMenu )
-				.OnItemScrolledIntoView(this, &SSCSEditor::OnItemScrolledIntoView )
-				.ItemHeight( 24 )
-				.HeaderRow
-				(
-					SNew(SHeaderRow)
-
-					+SHeaderRow::Column(SCS_ColumnName_Mobility)
-						.DefaultLabel(LOCTEXT("MobilityColumnLabel", "Mobility"))
-						.FixedWidth(16.0f) // mobility icons are 16px (16 slate-units = 16px, when application scale == 1)
-						.HeaderContent()
-						[
-							SNew(SHorizontalBox)
-								.ToolTip(SNew(SToolTip).Text(LOCTEXT("MobilityColumnTooltip", "Mobility")))
-							+SHorizontalBox::Slot()
-								.FillWidth(1.0f)
-								.VAlign(VAlign_Center)
-								.HAlign(HAlign_Center)
-								[
-									SNew(SImage).Image(MobilityHeaderBrush)
-								]
-						]
-				
-					+ SHeaderRow::Column(SCS_ColumnName_ComponentClass)
-					.DefaultLabel( LOCTEXT("Class", "Class") )
-					.FillWidth(4)
-					
-					+ SHeaderRow::Column(SCS_ColumnName_Asset)
-					.DefaultLabel( LOCTEXT("Asset", "Asset") )
-					.FillWidth(3)
-				)
-			]
+		+ SHeaderRow::Column(SCS_ColumnName_Asset)
+		.DefaultLabel(LOCTEXT("Asset", "Asset"))
+		.FillWidth(3)
+		);
+
+	TSharedPtr<SWidget> Contents;
+	if( InKismet2.IsValid() )
+	{
+		Contents = SNew(SVerticalBox)
+		+ SVerticalBox::Slot()
+		.AutoHeight()
+		[
+			SNew(SComponentClassCombo)
+			.OnComponentClassSelected(this, &SSCSEditor::PerformComboAddClass)
+		]
+		+ SVerticalBox::Slot()
+		.AutoHeight()
+		[
+			SNew(SSpacer)
+			.Size(FVector2D(0.0f, 2.0f))
+		]
+		+ SVerticalBox::Slot()
+		.FillHeight(.8f)
+		[
+			SCSTreeWidget.ToSharedRef()
+		];
+	}
+	else
+	{
+		Contents = SCSTreeWidget;
+	}
+
+	this->ChildSlot
+	[
+		SNew( STutorialWrapper, TEXT("ComponentsPanel") )
+		[
+			Contents.ToSharedRef()
 		]
 	];
 
@@ -2228,15 +2238,7 @@ bool SSCSEditor::CanDuplicateComponent() const
 		return false;
 	}
 
-	TArray<FSCSEditorTreeNodePtrType> SelectedNodes = SCSTreeWidget->GetSelectedItems();
-	for (int32 i = 0; i < SelectedNodes.Num(); ++i)
-	{
-		if(SelectedNodes[i]->GetComponentTemplate() == NULL || SelectedNodes[i]->IsDefaultSceneRoot())
-		{
-			return false;
-		}
-	}
-	return SelectedNodes.Num() > 0;
+	return CanCopyNodes();
 }
 
 void SSCSEditor::OnDuplicateComponent()
@@ -2477,9 +2479,6 @@ void SSCSEditor::UpdateTree(bool bRegenerateTreeNodes)
 		// Reset the scene root node
 		SceneRootNodePtr.Reset();
 
-		// Get the Blueprint that's being edited
-		UBlueprint* Blueprint = Kismet2Ptr.Pin()->GetBlueprintObj();
-
 		// Get the Blueprint class default object as well as the inheritance stack
 		AActor* CDO = NULL;
 		TArray<UBlueprint*> ParentBPStack;
@@ -2588,7 +2587,10 @@ void SSCSEditor::ClearSelection()
 	check(SCSTreeWidget.IsValid());
 	SCSTreeWidget->ClearSelection();
 	
-	Kismet2Ptr.Pin()->GetInspector()->ShowDetailsForObjects(TArray<UObject*>());
+	if( Kismet2Ptr.IsValid() )
+	{
+		Kismet2Ptr.Pin()->GetInspector()->ShowDetailsForObjects(TArray<UObject*>());
+	}
 }
 
 
@@ -2622,7 +2624,7 @@ void SSCSEditor::SaveSCSNode( USCS_Node* Node )
 
 bool SSCSEditor::InEditingMode() const
 {
-	return NULL == GEditor->PlayWorld;
+	return Kismet2Ptr.IsValid() && NULL == GEditor->PlayWorld;
 }
 
 UActorComponent* SSCSEditor::AddNewComponent( UClass* NewComponentClass, UObject* Asset  )
@@ -2630,7 +2632,6 @@ UActorComponent* SSCSEditor::AddNewComponent( UClass* NewComponentClass, UObject
 	const FScopedTransaction Transaction( LOCTEXT("AddComponent", "Add Component") );
 
 	// Create a new template and add it to the blueprint
-	UBlueprint* Blueprint = Kismet2Ptr.Pin()->GetBlueprintObj();
 	Blueprint->Modify();
 
 	SaveSCSCurrentState( Blueprint->SimpleConstructionScript );
@@ -2641,7 +2642,6 @@ UActorComponent* SSCSEditor::AddNewComponent( UClass* NewComponentClass, UObject
 
 UActorComponent* SSCSEditor::AddNewNode(USCS_Node* NewNode,  UObject* Asset, bool bMarkBlueprintModified, bool bSetFocusToNewItem)
 {
-	UBlueprint* Blueprint = Kismet2Ptr.Pin()->GetBlueprintObj();
 	if(Asset)
 	{
 		FComponentAssetBrokerage::AssignAssetToComponent(NewNode->ComponentTemplate, Asset);
@@ -2719,7 +2719,16 @@ void SSCSEditor::CutSelectedNodes()
 
 bool SSCSEditor::CanCopyNodes() const
 {
-	return CanDuplicateComponent();
+	TArray<FSCSEditorTreeNodePtrType> SelectedNodes = SCSTreeWidget->GetSelectedItems();
+	for (int32 i = 0; i < SelectedNodes.Num(); ++i)
+	{
+		if (SelectedNodes[i]->GetComponentTemplate() == NULL || SelectedNodes[i]->IsDefaultSceneRoot())
+		{
+			return false;
+		}
+	}
+
+	return SelectedNodes.Num() > 0;
 }
 
 void SSCSEditor::CopySelectedNodes()
@@ -2816,7 +2825,6 @@ void SSCSEditor::PasteNodes()
 	const FScopedTransaction Transaction(LOCTEXT("PasteComponents", "Paste Component(s)"));
 
 	// Mark the Blueprint as about to be modified
-	UBlueprint* Blueprint = Kismet2Ptr.Pin()->GetBlueprintObj();
 	Blueprint->Modify();
 
 	// Mark the SCS and all SCS nodes as about to be modified
@@ -2902,8 +2910,6 @@ void SSCSEditor::OnDeleteNodes()
 {
 	const FScopedTransaction Transaction( LOCTEXT("RemoveComponent", "Remove Component") );
 
-	UBlueprint* Blueprint = Kismet2Ptr.Pin()->GetBlueprintObj();
-
 	// Get the current render info for the blueprint. If this is NULL then the blueprint is not currently visualizable (no visible primitive components)
 	FThumbnailRenderingInfo* RenderInfo = GUnrealEd->GetThumbnailManager()->GetRenderingInfo( Blueprint );
 
@@ -2942,8 +2948,6 @@ void SSCSEditor::RemoveComponentNode(FSCSEditorTreeNodePtrType InNodePtr)
 {
 	check(InNodePtr.IsValid());
 
-	UBlueprint* Blueprint = Kismet2Ptr.Pin()->GetBlueprintObj();
-
 	// Remove any instances of variable accessors from the blueprint graphs
 	FBlueprintEditorUtils::RemoveVariableNodes(Blueprint, InNodePtr->GetVariableName());
 
@@ -2982,8 +2986,11 @@ void SSCSEditor::UpdateSelectionFromNodes(const TArray<FSCSEditorTreeNodePtrType
 	}
 
 	// Update the details panel
-	SKismetInspector::FShowDetailsOptions Options(InspectorTitle, true);
-	Kismet2Ptr.Pin()->GetInspector()->ShowDetailsForObjects(InspectorObjects, Options);
+	if( Kismet2Ptr.IsValid() )
+	{
+		SKismetInspector::FShowDetailsOptions Options(InspectorTitle, true);
+		Kismet2Ptr.Pin()->GetInspector()->ShowDetailsForObjects(InspectorObjects, Options);
+	}
 }
 
 void SSCSEditor::RefreshSelectionDetails()
@@ -2995,10 +3002,13 @@ void SSCSEditor::OnTreeSelectionChanged(FSCSEditorTreeNodePtrType, ESelectInfo::
 {
 	UpdateSelectionFromNodes(SCSTreeWidget->GetSelectedItems());
 
-	TSharedPtr<class SSCSEditorViewport> ViewportPtr = Kismet2Ptr.Pin()->GetSCSViewport();
-	if (ViewportPtr.IsValid())
+	if( Kismet2Ptr.IsValid() )
 	{
-		ViewportPtr->OnComponentSelectionChanged();
+		TSharedPtr<class SSCSEditorViewport> ViewportPtr = Kismet2Ptr.Pin()->GetSCSViewport();
+		if (ViewportPtr.IsValid())
+		{
+			ViewportPtr->OnComponentSelectionChanged();
+		}
 	}
 
 	// Update the selection visualization
@@ -3047,7 +3057,7 @@ FSCSEditorTreeNodePtrType SSCSEditor::AddTreeNode(USCS_Node* InSCSNode, FSCSEdit
 	
 	// Determine whether or not the given node is inherited from a parent Blueprint
 	USimpleConstructionScript* NodeSCS = InSCSNode->GetSCS();
-	const bool bIsInherited = NodeSCS != Kismet2Ptr.Pin()->GetBlueprintObj()->SimpleConstructionScript;
+	const bool bIsInherited = NodeSCS != Blueprint->SimpleConstructionScript;
 
 	if(InSCSNode->ComponentTemplate->IsA(USceneComponent::StaticClass()))
 	{
@@ -3281,7 +3291,7 @@ void SSCSEditor::OnRenameComponent(bool bTransactional)
 
 bool SSCSEditor::CanRenameComponent() const
 {
-	return SCSTreeWidget->GetSelectedItems().Num() == 1 && SCSTreeWidget->GetSelectedItems()[0]->CanRename();
+	return InEditingMode() && SCSTreeWidget->GetSelectedItems().Num() == 1 && SCSTreeWidget->GetSelectedItems()[0]->CanRename();
 }
 
 void SSCSEditor::GetCollapsedNodes(const FSCSEditorTreeNodePtrType& InNodePtr, TSet<FSCSEditorTreeNodePtrType>& OutCollapsedNodes) const
diff --git a/Engine/Source/Editor/Kismet/Private/SSCSEditor.h b/Engine/Source/Editor/Kismet/Private/SSCSEditor.h
--- a/Engine/Source/Editor/Kismet/Private/SSCSEditor.h
+++ b/Engine/Source/Editor/Kismet/Private/SSCSEditor.h
@@ -376,7 +376,7 @@ public:
 	SLATE_BEGIN_ARGS( SSCSEditor ){}
 	SLATE_END_ARGS()
 
-	void Construct(const FArguments& InArgs, TSharedPtr<FBlueprintEditor> InKismet2, USimpleConstructionScript* InSCS);
+	void Construct(const FArguments& InArgs, TSharedPtr<FBlueprintEditor> InKismet2, USimpleConstructionScript* InSCS, UBlueprint* InBlueprint );
 
 	/** Override OnKeyDown */
 	virtual FReply OnKeyDown( const FGeometry& MyGeometry, const FKeyboardEvent& InKeyboardEvent );
@@ -559,6 +559,9 @@ public:
 	/** Pointer to the script that we are editing */
 	USimpleConstructionScript* SCS;
 
+	/** Pointer to blueprint we were created with */
+	UBlueprint* Blueprint;
+
 	/** Pointer back to owning Kismet 2 tool */
 	TWeakPtr<FBlueprintEditor> Kismet2Ptr;
 
@@ -579,12 +582,8 @@ public:
 
 	/** Whether or not the deferred rename request was flagged as transactional */
 	bool bIsDeferredRenameRequestTransactional;
-
 private:
 
-	/** Pointer to combo box used for choosing component class */
-	TSharedPtr<class SComponentClassCombo> ComponentClassCombo;
-
 	/** Flag to enable/disable component editing */
 	bool	bEnableComponentEditing;
 };
diff --git a/Engine/Source/Editor/Kismet/Public/SBlueprintDiff.h b/Engine/Source/Editor/Kismet/Public/SBlueprintDiff.h
--- a/Engine/Source/Editor/Kismet/Public/SBlueprintDiff.h
+++ b/Engine/Source/Editor/Kismet/Public/SBlueprintDiff.h
@@ -155,7 +155,6 @@ public:
 			SLATE_ARGUMENT( struct FRevisionInfo, OldRevision )
 			SLATE_ARGUMENT( struct FRevisionInfo, NewRevision )
 			SLATE_ARGUMENT( bool, ShowAssetNames )
-			SLATE_EVENT(FOpenInDefaults, OpenInDefaults)
 	SLATE_END_ARGS()
 
 	void Construct(const FArguments& InArgs);
@@ -191,9 +190,6 @@ protected:
 	/*User toggles the option to lock the views between the two blueprints */
 	FReply	OnToggleLockView();
 
-	/*User clicks defaults view button to display defaults in remote diff tool */
-	FReply  OnOpenInDefaults();
-
 	/*Reset the graph editor, called when user switches graphs to display*/
 	virtual void ResetGraphEditors();
 
@@ -218,18 +214,28 @@ protected:
 	/** Event handler that updates the graph view when user selects a new graph */
 	virtual void HandleGraphChanged( const FString& GraphName );
 
+	TSharedRef<SWidget> GenerateGraphPanel();
+	TSharedRef<SWidget> GenerateDefaultsPanel();
+	TSharedRef<SWidget> GenerateComponentsPanel();
+
+	/** Accessor and event handler for toggling between diff view modes (defaults, components, graph view, interface, macro): */
+	void SetCurrentMode(FName NewMode);
+	FName GetCurrentMode() const { return CurrentMode; }
+
+	FName CurrentMode;
+
 	/*The two panels used to show the old & new revision*/ 
 	FDiffPanel				PanelOld, PanelNew;
 	
 	/** If the two views should be locked */
 	bool	bLockViews;
 
-	/*Delegate to call when user wishes to view the defaults*/
-	FOpenInDefaults	 OpenInDefaults;
-
 	/** Border Widget, inside is the current graphs being diffed, we can replace content to change the graph*/
 	TSharedPtr<SBorder>	DiffListBorder;
 
+	/** Contents widget that we swap when mode changes (defaults, components, etc) */
+	TSharedPtr<SBorder> ModeContents;
+
 	/** The ListView containing the graphs the user can select */
 	TSharedPtr<SListViewType>	GraphsToDiff;
 
diff --git a/Engine/Source/Editor/Kismet/Public/SBlueprintEditorToolbar.h b/Engine/Source/Editor/Kismet/Public/SBlueprintEditorToolbar.h
--- a/Engine/Source/Editor/Kismet/Public/SBlueprintEditorToolbar.h
+++ b/Engine/Source/Editor/Kismet/Public/SBlueprintEditorToolbar.h
@@ -2,6 +2,8 @@
 
 #pragma once
 
+#include "WorkflowOrientedApp/SModeWidget.h"
+
 /**
  * Kismet menu
  */
@@ -75,6 +77,9 @@ public:
 	/** Returns the current status as text for the blueprint being edited */
 	FText GetStatusTooltip() const;
 
+	/** Helper function for generating the buttons in the toolbar, reused by merge and diff tools */
+	static TArray< TSharedPtr< class SWidget> > GenerateToolbarWidgets(const class UBlueprint* BlueprintObj, TAttribute<FName> ActiveModeGetter, FOnModeChangeRequested ActiveModeSetter);
+
 private:
 	void FillBlueprintEditorModesToolbar(FToolBarBuilder& ToolbarBuilder);
 	void FillBlueprintGlobalOptionsToolbar(FToolBarBuilder& ToolBarBuilder);
diff --git a/Engine/Source/Editor/UnrealEd/Private/Kismet2/BlueprintEditorUtils.cpp b/Engine/Source/Editor/UnrealEd/Private/Kismet2/BlueprintEditorUtils.cpp
--- a/Engine/Source/Editor/UnrealEd/Private/Kismet2/BlueprintEditorUtils.cpp
+++ b/Engine/Source/Editor/UnrealEd/Private/Kismet2/BlueprintEditorUtils.cpp
@@ -2356,6 +2356,12 @@ bool FBlueprintEditorUtils::DoesSupportComponents(UBlueprint const* Blueprint)
 		&& (Blueprint->BlueprintType != BPTYPE_FunctionLibrary); // Must not be a function library
 }
 
+bool FBlueprintEditorUtils::DoesSupportDefaults(UBlueprint const* Blueprint)
+{
+	return Blueprint->BlueprintType != BPTYPE_MacroLibrary
+		&& Blueprint->BlueprintType != BPTYPE_FunctionLibrary;
+}
+
 // Returns a descriptive name of the type of blueprint passed in
 FString FBlueprintEditorUtils::GetBlueprintTypeDescription(const UBlueprint* Blueprint)
 {
diff --git a/Engine/Source/Editor/UnrealEd/Public/Kismet2/BlueprintEditorUtils.h b/Engine/Source/Editor/UnrealEd/Public/Kismet2/BlueprintEditorUtils.h
--- a/Engine/Source/Editor/UnrealEd/Public/Kismet2/BlueprintEditorUtils.h
+++ b/Engine/Source/Editor/UnrealEd/Public/Kismet2/BlueprintEditorUtils.h
@@ -380,6 +380,9 @@ public:
 	/** Returns whether or not the blueprint supports components */
 	static bool DoesSupportComponents(UBlueprint const* Blueprint);
 
+	/** Returns whether or not the blueprint supports default values (IE has a CDO) */
+	static bool DoesSupportDefaults(UBlueprint const* Blueprint);
+
 	// Returns a descriptive name of the type of blueprint passed in
 	static FString GetBlueprintTypeDescription(const UBlueprint* Blueprint);
 
