Stats system when doingetadata mismatch with the Stats system when doing multiple Find-in-Blueprint searches.

#ttp 340050 - BLUEPRINT:  MetaData mismatch warning when searching

diff --git a/Engine/Source/Editor/Kismet/Private/FindInBlueprints.cpp b/Engine/Source/Editor/Kismet/Private/FindInBlueprints.cpp
--- a/Engine/Source/Editor/Kismet/Private/FindInBlueprints.cpp
+++ b/Engine/Source/Editor/Kismet/Private/FindInBlueprints.cpp
@@ -814,7 +814,8 @@ public:
 		, bThreadCompleted(false)
 		, StopTaskCounter(0)
 	{
-		Thread = FRunnableThread::Create( this, TEXT("FStreamSearch"), 0, TPri_BelowNormal );
+		// Add on a Guid to the thread name to ensure the thread is uniquely named.
+		Thread = FRunnableThread::Create( this, *FString::Printf(TEXT("FStreamSearch%s"), *FGuid::NewGuid().ToString()), 0, TPri_BelowNormal );
 	}
 
 	/** Begin FRunnable Interface */
