#UE4 quick details panel added to blueprint diff too, this 'completes' the diff panel, you can now at least see all the information in the blueprint

diff --git a/Engine/Source/Editor/Kismet/Private/SBlueprintDiff.cpp b/Engine/Source/Editor/Kismet/Private/SBlueprintDiff.cpp
--- a/Engine/Source/Editor/Kismet/Private/SBlueprintDiff.cpp
+++ b/Engine/Source/Editor/Kismet/Private/SBlueprintDiff.cpp
@@ -7,6 +7,7 @@
 #include "BlueprintEditorModes.h"
 #include "BlueprintUtilities.h"
 #include "Editor/Kismet/Public/SBlueprintEditorToolbar.h"
+#include "Editor/Kismet/Public/SKismetInspector.h"
 #include "Editor/PropertyEditor/Public/IDetailsView.h"
 #include "Editor/PropertyEditor/Public/PropertyEditorModule.h"
 #include "Editor/UnrealEd/Public/EdGraphUtilities.h"
@@ -463,7 +464,21 @@ static TSharedRef<SWidget> GenerateCDODiff( UBlueprint const* BaseBlueprint, UBl
 
 static TSharedRef<SWidget> GenerateComponentsDiff( UBlueprint const* BaseBlueprint, UBlueprint const* DisplayedBlueprint )
 {
-	return SNew(SSCSEditor, TSharedPtr<FBlueprintEditor>(), DisplayedBlueprint->SimpleConstructionScript, const_cast<UBlueprint*>(DisplayedBlueprint) );
+	TSharedRef<SKismetInspector> Inspector = SNew(SKismetInspector)
+		.HideNameArea(true)
+		.ViewIdentifier(FName("BlueprintInspector"))
+		.IsPropertyEditingEnabledDelegate(FIsPropertyEditingEnabled::CreateStatic([] { return false; }));
+
+	return SNew(SVerticalBox)
+		+ SVerticalBox::Slot()
+		.AutoHeight()
+		[
+			SNew(SSCSEditor, TSharedPtr<FBlueprintEditor>(), DisplayedBlueprint->SimpleConstructionScript, const_cast<UBlueprint*>(DisplayedBlueprint), Inspector )
+		]
+		+ SVerticalBox::Slot()
+		[
+			Inspector
+		];
 }
 
 //////////////////////////////////////////////////////////////////////////
diff --git a/Engine/Source/Editor/Kismet/Private/SKismetInspector.cpp b/Engine/Source/Editor/Kismet/Private/SKismetInspector.cpp
--- a/Engine/Source/Editor/Kismet/Private/SKismetInspector.cpp
+++ b/Engine/Source/Editor/Kismet/Private/SKismetInspector.cpp
@@ -36,7 +36,6 @@ public:
 TSharedRef<SWidget> SKismetInspector::MakeContextualEditingWidget(struct FKismetSelectionInfo& SelectionInfo, const FShowDetailsOptions& Options)
 {
 	TSharedRef< SVerticalBox > ContextualEditingWidget = SNew( SVerticalBox );
-	TSharedPtr<SMyBlueprint> MyBlueprintWidget = Kismet2Ptr.Pin()->GetMyBlueprintWidget();
 
 	if(bShowTitleArea)
 	{
@@ -210,7 +209,7 @@ void SKismetInspector::Construct(const FArguments& InArgs)
 	PropertyView->SetIsPropertyEditingEnabledDelegate(InArgs._IsPropertyEditingEnabledDelegate);
 	PropertyView->OnFinishedChangingProperties().Add( InArgs._OnFinishedChangingProperties );
 	
-	if (Kismet2->IsEditingSingleBlueprint())
+	if (Kismet2.IsValid() && Kismet2->IsEditingSingleBlueprint())
 	{
 		FOnGetDetailCustomizationInstance LayoutDelegateDetails = FOnGetDetailCustomizationInstance::CreateStatic(&FBlueprintDelegateActionDetails::MakeInstance, TWeakPtr<SMyBlueprint>(Kismet2->GetMyBlueprintWidget()));
 		PropertyView->RegisterInstancedCustomPropertyLayout(UMulticastDelegateProperty::StaticClass(), LayoutDelegateDetails);
diff --git a/Engine/Source/Editor/Kismet/Private/SSCSEditor.cpp b/Engine/Source/Editor/Kismet/Private/SSCSEditor.cpp
--- a/Engine/Source/Editor/Kismet/Private/SSCSEditor.cpp
+++ b/Engine/Source/Editor/Kismet/Private/SSCSEditor.cpp
@@ -1903,9 +1903,10 @@ void SSCS_RowWidget::OnNameTextCommit(const FText& InNewName, ETextCommit::Type
 // SSCSEditor
 
 
-void SSCSEditor::Construct( const FArguments& InArgs, TSharedPtr<FBlueprintEditor> InKismet2, USimpleConstructionScript* InSCS, UBlueprint* InBlueprint )
+void SSCSEditor::Construct( const FArguments& InArgs, TSharedPtr<FBlueprintEditor> InKismet2, USimpleConstructionScript* InSCS, UBlueprint* InBlueprint, TSharedPtr<SKismetInspector> Inspector )
 {
 	Kismet2Ptr = InKismet2;
+	KismetInspectorPtr = Inspector.IsValid() ? Inspector : InKismet2->GetInspector();
 	SCS = InSCS;
 	Blueprint = InBlueprint;
 
@@ -2589,9 +2590,9 @@ void SSCSEditor::ClearSelection()
 	check(SCSTreeWidget.IsValid());
 	SCSTreeWidget->ClearSelection();
 	
-	if( Kismet2Ptr.IsValid() )
+	if( KismetInspectorPtr.IsValid() )
 	{
-		Kismet2Ptr.Pin()->GetInspector()->ShowDetailsForObjects(TArray<UObject*>());
+		KismetInspectorPtr.Pin()->ShowDetailsForObjects(TArray<UObject*>());
 	}
 }
 
@@ -2988,10 +2989,10 @@ void SSCSEditor::UpdateSelectionFromNodes(const TArray<FSCSEditorTreeNodePtrType
 	}
 
 	// Update the details panel
-	if( Kismet2Ptr.IsValid() )
+	if (KismetInspectorPtr.IsValid())
 	{
 		SKismetInspector::FShowDetailsOptions Options(InspectorTitle, true);
-		Kismet2Ptr.Pin()->GetInspector()->ShowDetailsForObjects(InspectorObjects, Options);
+		KismetInspectorPtr.Pin()->ShowDetailsForObjects(InspectorObjects, Options);
 	}
 }
 
diff --git a/Engine/Source/Editor/Kismet/Private/SSCSEditor.h b/Engine/Source/Editor/Kismet/Private/SSCSEditor.h
--- a/Engine/Source/Editor/Kismet/Private/SSCSEditor.h
+++ b/Engine/Source/Editor/Kismet/Private/SSCSEditor.h
@@ -376,7 +376,7 @@ public:
 	SLATE_BEGIN_ARGS( SSCSEditor ){}
 	SLATE_END_ARGS()
 
-	void Construct(const FArguments& InArgs, TSharedPtr<FBlueprintEditor> InKismet2, USimpleConstructionScript* InSCS, UBlueprint* InBlueprint );
+	void Construct(const FArguments& InArgs, TSharedPtr<FBlueprintEditor> InKismet2, USimpleConstructionScript* InSCS, UBlueprint* InBlueprint, TSharedPtr<class SKismetInspector> Inspector = TSharedPtr<class SKismetInspector>() );
 
 	/** Override OnKeyDown */
 	virtual FReply OnKeyDown( const FGeometry& MyGeometry, const FKeyboardEvent& InKeyboardEvent );
@@ -565,6 +565,9 @@ public:
 	/** Pointer back to owning Kismet 2 tool */
 	TWeakPtr<FBlueprintEditor> Kismet2Ptr;
 
+	/** Pointer back to companion inspector tool */
+	TWeakPtr<class SKismetInspector> KismetInspectorPtr;
+
 	/** Tree widget */
 	TSharedPtr<SSCSTreeType> SCSTreeWidget;
 
