#UE4 fix for 328394, 342682. We should not be able to add local variables to graphs that are not editable...

diff --git a/Engine/Source/Editor/Kismet/Private/BlueprintEditor.cpp b/Engine/Source/Editor/Kismet/Private/BlueprintEditor.cpp
--- a/Engine/Source/Editor/Kismet/Private/BlueprintEditor.cpp
+++ b/Engine/Source/Editor/Kismet/Private/BlueprintEditor.cpp
@@ -5784,7 +5784,10 @@ bool FBlueprintEditor::NewDocument_IsVisibleForType(ECreatedDocumentType GraphTy
 	case CGT_NewEventGraph:
 		return FBlueprintEditorUtils::DoesSupportEventGraphs(GetBlueprintObj());
 	case CGT_NewLocalVariable:
-		return (GetBlueprintObj()->BlueprintType != BPTYPE_Interface) && (GetFocusedGraph() && GetFocusedGraph()->GetSchema()->GetGraphType(GetFocusedGraph()) == EGraphType::GT_Function) && !GetFocusedGraph()->IsA(UAnimationTransitionGraph::StaticClass());
+		return (GetBlueprintObj()->BlueprintType != BPTYPE_Interface) 
+			&& (GetFocusedGraph() && GetFocusedGraph()->GetSchema()->GetGraphType(GetFocusedGraph()) == EGraphType::GT_Function)
+			&& !GetFocusedGraph()->IsA(UAnimationTransitionGraph::StaticClass()) 
+			&& IsFocusedGraphEditable();
 	}
 
 	return false;
