#UE4 SBlueprintDiff cleanup

diff --git a/Engine/Source/Editor/Kismet/Private/SBlueprintDiff.cpp b/Engine/Source/Editor/Kismet/Private/SBlueprintDiff.cpp
--- a/Engine/Source/Editor/Kismet/Private/SBlueprintDiff.cpp
+++ b/Engine/Source/Editor/Kismet/Private/SBlueprintDiff.cpp
@@ -867,41 +867,6 @@ FDiffPanel& SBlueprintDiff::GetDiffPanelForNode(UEdGraphNode& Node)
 	return Default;
 }
 
-TSharedRef<SWidget> SBlueprintDiff::GenerateDiffWindow()
-{
-	return SNew(SSplitter)
-		+ SSplitter::Slot()
-		.Value(0.5f)
-		[
-			//left blueprint
-			SAssignNew(PanelOld.GraphEditorBorder, SBorder)
-			.VAlign(VAlign_Fill)
-			[
-				DefaultEmptyPanel()
-			]
-		]
-	+ SSplitter::Slot()
-		.Value(0.5f)
-		[
-			//right blueprint
-			SAssignNew(PanelNew.GraphEditorBorder, SBorder)
-			.VAlign(VAlign_Fill)
-			[
-				DefaultEmptyPanel()
-			]
-		];
-}
-
-TSharedRef<SWidget> SBlueprintDiff::GenerateToolbar()
-{
-	return SNew(SVerticalBox)
-		+SVerticalBox::Slot()
-		.FillHeight(1.f)
-		[
-			SAssignNew(DiffListBorder, SBorder)
-		];
-}
-
 void SBlueprintDiff::HandleGraphChanged( const FString& GraphName )
 {
 	TArray<UEdGraph*> GraphsOld, GraphsNew;
@@ -998,7 +963,12 @@ TSharedRef<SWidget> SBlueprintDiff::GenerateGraphPanel()
 						+ SSplitter::Slot()
 							.Value(0.7f)
 							[
-								GenerateToolbar()
+								SNew(SVerticalBox)
+								+ SVerticalBox::Slot()
+								.FillHeight(1.f)
+								[
+									SAssignNew(DiffListBorder, SBorder)
+								]
 							]
 					]
 				]
@@ -1010,7 +980,27 @@ TSharedRef<SWidget> SBlueprintDiff::GenerateGraphPanel()
 						.FillWidth(1.f)
 						[
 							//diff window
-							GenerateDiffWindow()
+							SNew(SSplitter)
+							+ SSplitter::Slot()
+							.Value(0.5f)
+							[
+								//left blueprint
+								SAssignNew(PanelOld.GraphEditorBorder, SBorder)
+								.VAlign(VAlign_Fill)
+								[
+									DefaultEmptyPanel()
+								]
+							]
+							+ SSplitter::Slot()
+							.Value(0.5f)
+							[
+								//right blueprint
+								SAssignNew(PanelNew.GraphEditorBorder, SBorder)
+								.VAlign(VAlign_Fill)
+								[
+									DefaultEmptyPanel()
+								]
+							]
 						]
 					]
 			]
diff --git a/Engine/Source/Editor/Kismet/Public/SBlueprintDiff.h b/Engine/Source/Editor/Kismet/Public/SBlueprintDiff.h
--- a/Engine/Source/Editor/Kismet/Public/SBlueprintDiff.h
+++ b/Engine/Source/Editor/Kismet/Public/SBlueprintDiff.h
@@ -200,7 +200,7 @@ protected:
 	TSharedRef<ITableRow> OnGenerateRow(FGraphToDiff ParamItem, const TSharedRef<STableViewBase>& OwnerTable );
 
 	/*Called when user clicks on a new graph list item */
-	virtual void OnSelectionChanged(FGraphToDiff Item, ESelectInfo::Type SelectionType);
+	void OnSelectionChanged(FGraphToDiff Item, ESelectInfo::Type SelectionType);
 
 	void OnDiffListSelectionChanged(const TSharedPtr<struct FDiffResultItem>& TheDiff, FListItemGraphToDiff* GraphDiffer);
 		
@@ -211,7 +211,7 @@ protected:
 	FReply	OnToggleLockView();
 
 	/*Reset the graph editor, called when user switches graphs to display*/
-	virtual void ResetGraphEditors();
+	void ResetGraphEditors();
 
 	/*Get the image to show for the toggle lock option*/
 	const FSlateBrush* GetLockViewImage() const;
@@ -222,14 +222,8 @@ protected:
 	/** Get Graph editor associated with this Graph */
 	FDiffPanel& GetDiffPanelForNode(UEdGraphNode& Node);
 
-	/** Generates the widgets used to display the graphs that are being diff'd */
-	virtual TSharedRef<SWidget> GenerateDiffWindow();
-
-	/** Generates the toolbar of buttons (for example, the button that lets you diff the CDO) */
-	virtual TSharedRef<SWidget> GenerateToolbar();
-
 	/** Event handler that updates the graph view when user selects a new graph */
-	virtual void HandleGraphChanged( const FString& GraphName );
+	void HandleGraphChanged( const FString& GraphName );
 
 	TSharedRef<SWidget> GenerateGraphPanel();
 	TSharedRef<SWidget> GenerateDefaultsPanel();
