Fixing issue with duplicate menu entries showing up in the blueprint context menu.

diff --git a/Engine/Source/Editor/Kismet/Private/BlueprintActionMenuUtils.cpp b/Engine/Source/Editor/Kismet/Private/BlueprintActionMenuUtils.cpp
--- a/Engine/Source/Editor/Kismet/Private/BlueprintActionMenuUtils.cpp
+++ b/Engine/Source/Editor/Kismet/Private/BlueprintActionMenuUtils.cpp
@@ -7,7 +7,8 @@
 #include "BlueprintBoundNodeSpawner.h"
 #include "KismetEditorUtilities.h" // for CanPasteNodes()
 #include "K2Node_CallFunction.h"
-#include "K2Node_Event.h"
+#include "K2Node_ComponentBoundEvent.h"
+#include "EdGraphSchema_K2.h"
 
 #define LOCTEXT_NAMESPACE "BlueprintActionMenuUtils"
 
@@ -104,24 +105,27 @@ void FBlueprintActionMenuUtils::MakeContextMenu(FBlueprintActionContext const& C
 	MenuOut.AddMenuSection(MenuFilter);
 	MenuOut.RebuildActionList();
 
-	for (UEdGraph const* Graph : Context.Graphs)
+	if (!GetDefault<UEdGraphSchema_K2>()->bUseLegacyActionMenus)
 	{
-		if (FKismetEditorUtilities::CanPasteNodes(Graph))
+		for (UEdGraph const* Graph : Context.Graphs)
 		{
-			// @TODO: Grey out menu option with tooltip if one of the nodes cannot paste into this graph
-			TSharedPtr<FEdGraphSchemaAction> PasteHereAction(new FEdGraphSchemaAction_K2PasteHere(TEXT(""), LOCTEXT("PasteHereMenuName", "Paste here"), TEXT(""), MainMenuSectionGroup));
-			MenuOut.AddAction(PasteHereAction);
-			break;
+			if (FKismetEditorUtilities::CanPasteNodes(Graph))
+			{
+				// @TODO: Grey out menu option with tooltip if one of the nodes cannot paste into this graph
+				TSharedPtr<FEdGraphSchemaAction> PasteHereAction(new FEdGraphSchemaAction_K2PasteHere(TEXT(""), LOCTEXT("PasteHereMenuName", "Paste here"), TEXT(""), MainMenuSectionGroup));
+				MenuOut.AddAction(PasteHereAction);
+				break;
+			}
 		}
-	}
 
-	if (bIsContextSensitive && (Context.SelectedObjects.Num() == 0))
-	{
-		FText SelectComponentMsg     = LOCTEXT("SelectComponentForEvents",        "Select a Component to see available Events & Functions");
-		FText SelectComponentToolTip = LOCTEXT("SelectComponentForEventsTooltip", "Select a Component in the MyBlueprint tab to see available Events and Functions in this menu.");
-		TSharedPtr<FEdGraphSchemaAction> MsgAction = TSharedPtr<FEdGraphSchemaAction>(new FEdGraphSchemaAction_Dummy(TEXT(""), SelectComponentMsg, SelectComponentToolTip.ToString(), ComponentsSectionGroup));
-		MenuOut.AddAction(MsgAction);
-	}
+		if (bIsContextSensitive && (Context.SelectedObjects.Num() == 0))
+		{
+			FText SelectComponentMsg = LOCTEXT("SelectComponentForEvents", "Select a Component to see available Events & Functions");
+			FText SelectComponentToolTip = LOCTEXT("SelectComponentForEventsTooltip", "Select a Component in the MyBlueprint tab to see available Events and Functions in this menu.");
+			TSharedPtr<FEdGraphSchemaAction> MsgAction = TSharedPtr<FEdGraphSchemaAction>(new FEdGraphSchemaAction_Dummy(TEXT(""), SelectComponentMsg, SelectComponentToolTip.ToString(), ComponentsSectionGroup));
+			MenuOut.AddAction(MsgAction);
+		}
+	}	
 }
 
 #undef LOCTEXT_NAMESPACE
