(WIP) Now taking pin-context into consideration when filtering the new blueprint menu system.

diff --git a/Engine/Source/Editor/BlueprintGraph/Classes/K2Node_EnumEquality.h b/Engine/Source/Editor/BlueprintGraph/Classes/K2Node_EnumEquality.h
--- a/Engine/Source/Editor/BlueprintGraph/Classes/K2Node_EnumEquality.h
+++ b/Engine/Source/Editor/BlueprintGraph/Classes/K2Node_EnumEquality.h
@@ -30,6 +30,7 @@ class UK2Node_EnumEquality : public UK2Node
 	virtual void GetMenuActions(TArray<UBlueprintNodeSpawner*>& ActionListOut) const override;
 	virtual FText GetMenuCategory() const override;
 	virtual int32 GetNodeRefreshPriority() const override { return EBaseNodeRefreshPriority::Low_UsesDependentWildcard; }
+	virtual bool IsConnectionDisallowed(const UEdGraphPin* MyPin, const UEdGraphPin* OtherPin, FString& OutReason) const override;
 	// End UK2Node interface
 
 	/** Get the return value pin */
diff --git a/Engine/Source/Editor/BlueprintGraph/Private/BlueprintActionDatabase.cpp b/Engine/Source/Editor/BlueprintGraph/Private/BlueprintActionDatabase.cpp
--- a/Engine/Source/Editor/BlueprintGraph/Private/BlueprintActionDatabase.cpp
+++ b/Engine/Source/Editor/BlueprintGraph/Private/BlueprintActionDatabase.cpp
@@ -10,29 +10,31 @@
 #include "BlueprintEventNodeSpawner.h"
 #include "BlueprintComponentNodeSpawner.h"
 #include "BlueprintBoundNodeSpawner.h"
+#include "Stats.h"					// for RETURN_QUICK_DECLARE_CYCLE_STAT()
+#include "ScopedTimers.h"			// for FScopedDurationTimer
 
 // used below in FBlueprintNodeSpawnerFactory::MakeMacroNodeSpawner()
 #include "K2Node_MacroInstance.h"
-// used below in FBlueprintActionDatabaseImpl::AddClassEnumActions()
+// used below in BlueprintActionDatabaseImpl::AddClassEnumActions()
 #include "K2Node_GetNumEnumEntries.h"
 #include "K2Node_ForEachElementInEnum.h"
 #include "K2Node_EnumLiteral.h"
 #include "K2Node_CastByteToEnum.h"
 #include "K2Node_SwitchEnum.h"
-// used below in FBlueprintActionDatabaseImpl::AddClassStructActions()
+// used below in BlueprintActionDatabaseImpl::AddClassStructActions()
 #include "K2Node_BreakStruct.h"
 #include "K2Node_MakeStruct.h"
 #include "K2Node_SetFieldsInStruct.h"
 // used below in FBlueprintNodeSpawnerFactory::MakeMessageNodeSpawner()
 #include "K2Node_Message.h"
-// used below in FBlueprintActionDatabaseImpl::AddClassPropertyActions()
+// used below in BlueprintActionDatabaseImpl::AddClassPropertyActions()
 #include "K2Node_VariableGet.h"
 #include "K2Node_VariableSet.h"
 #include "K2Node_AddDelegate.h"
 #include "K2Node_CallDelegate.h"
 #include "K2Node_RemoveDelegate.h"
 #include "K2Node_ClearDelegate.h"
-// used below in FBlueprintActionDatabaseImpl::AddClassCastActions()
+// used below in BlueprintActionDatabaseImpl::AddClassCastActions()
 #include "K2Node_DynamicCast.h"
 #include "K2Node_ClassDynamicCast.h"
 // used below in FBlueprintNodeSpawnerFactory::MakeCommentNodeSpawner()
@@ -239,53 +241,9 @@ static UBlueprintNodeSpawner* FBlueprintNodeSpawnerFactory::MakeCommentNodeSpawn
  * Static FBlueprintActionDatabase Helpers
  ******************************************************************************/
 
-namespace FBlueprintActionDatabaseImpl
+namespace BlueprintActionDatabaseImpl
 {
-	/** 
-	 * Since the action-database is a singleton, we can track its state in a
-	 * static enum value here (hidden as an implementation detail to keep the 
-	 * API clean).
-	 */
-	enum EDatabaseState
-	{
-		DB_Uninitialized,
-		DB_Initializing,
-		DB_Primed,
-	};
-	static EDatabaseState DatabaseState = DB_Uninitialized;
-
-	/** 
-	 * Wrapper around FBlueprintActionDatabase's FActionList. Manages the 
-	 * RF_RootSet flag for the UBlueprintNodeSpawners (ensures they don't get 
-	 * GC'd whilst in the database.
-	 */
-	struct FActionList
-	{
-	public:
-		/** Clears the passed action-list and removes the RF_RootSet from all of its actions */
-		FActionList(FBlueprintActionDatabase::FActionList& DatabaseIn)
-			: ClassDatabase(DatabaseIn)
-		{
-			for (UBlueprintNodeSpawner* Action : ClassDatabase)
-			{
-				check(Action->GetOuter() == GetTransientPackage());
-				Action->ClearFlags(RF_RootSet);
-			}
-			ClassDatabase.Empty();
-		}
-
-		/** 
-		 * Passes the wrapped ClassDatabase this node-spawner to add, after 
-		 * adding the RF_RootSet flag (to keep it from getting GC'd).
-		 *
-		 * @param  NodeSpawner	The action you want added to the database.
-		 */
-		void Add(UBlueprintNodeSpawner* NodeSpawner);
-
-	private:
-		/** */
-		FBlueprintActionDatabase::FActionList& ClassDatabase;
-	};
+	typedef FBlueprintActionDatabase::FActionList FActionList;
 
 	/**
 	 * Mimics UEdGraphSchema_K2::CanUserKismetAccessVariable(); however, this 
@@ -410,24 +368,18 @@ namespace FBlueprintActionDatabaseImpl
 	 */
 	static void AddAutonomousNodeActions(UClass* const Class, FActionList& ActionListOut);
 
+	/**
+	 * Callback to refresh the database when a blueprint has been altered 
+	 * (clears the database entries for the blueprint's classes and recreates 
+	 * them... a property/function could have been added/removed).
+	 * 
+	 * @param  InBlueprint	The blueprint that has been altered.
+	 */
 	static void OnBlueprintChanged(UBlueprint* InBlueprint);
 }
 
 //------------------------------------------------------------------------------
-void FBlueprintActionDatabaseImpl::FActionList::Add(UBlueprintNodeSpawner* NodeSpawner)
-{
-	check(NodeSpawner != nullptr);
-	check(NodeSpawner->GetOuter() == GetTransientPackage());
-	// since this spawner's outer is the transient package, we want to mark it
-	// root so that it doesn't get GC'd (we have to be careful and remove this 
-	// flag when we refresh).
-	NodeSpawner->SetFlags(RF_RootSet);
-
-	ClassDatabase.Add(NodeSpawner);
-}
-
-//------------------------------------------------------------------------------
-static bool FBlueprintActionDatabaseImpl::IsPropertyBlueprintVisible(UProperty const* const Property)
+static bool BlueprintActionDatabaseImpl::IsPropertyBlueprintVisible(UProperty const* const Property)
 {
 	bool const bIsAccessible = Property->HasAllPropertyFlags(CPF_BlueprintVisible);
 
@@ -438,7 +390,7 @@ static bool FBlueprintActionDatabaseImpl::IsPropertyBlueprintVisible(UProperty c
 }
 
 //------------------------------------------------------------------------------
-static void FBlueprintActionDatabaseImpl::AddClassFunctionActions(UClass const* const Class, FActionList& ActionListOut)
+static void BlueprintActionDatabaseImpl::AddClassFunctionActions(UClass const* const Class, FActionList& ActionListOut)
 {
 	using namespace FBlueprintNodeSpawnerFactory; // for MakeMessageNodeSpawner()
 	check(Class != nullptr);
@@ -473,7 +425,7 @@ static void FBlueprintActionDatabaseImpl::AddClassFunctionActions(UClass const*
 }
 
 //------------------------------------------------------------------------------
-static void FBlueprintActionDatabaseImpl::AddClassPropertyActions(UClass const* const Class, FActionList& ActionListOut)
+static void BlueprintActionDatabaseImpl::AddClassPropertyActions(UClass const* const Class, FActionList& ActionListOut)
 {
 	using namespace FBlueprintNodeSpawnerFactory; // for MakeDelegateNodeSpawner()
 	
@@ -531,7 +483,7 @@ static void FBlueprintActionDatabaseImpl::AddClassPropertyActions(UClass const*
 }
 
 //------------------------------------------------------------------------------
-static void FBlueprintActionDatabaseImpl::AddComponentBoundActions(UObjectProperty* ComponentProperty, FActionList& ActionListOut)
+static void BlueprintActionDatabaseImpl::AddComponentBoundActions(UObjectProperty* ComponentProperty, FActionList& ActionListOut)
 {
 	check(ComponentProperty != nullptr);
 
@@ -590,7 +542,7 @@ static void FBlueprintActionDatabaseImpl::AddComponentBoundActions(UObjectProper
 }
 
 //------------------------------------------------------------------------------
-static void FBlueprintActionDatabaseImpl::AddClassCastActions(UClass* const Class, FActionList& ActionListOut)
+static void BlueprintActionDatabaseImpl::AddClassCastActions(UClass* const Class, FActionList& ActionListOut)
 {
 	UEdGraphSchema_K2 const* K2Schema = GetDefault<UEdGraphSchema_K2>();
 	bool bIsCastPermitted  = UEdGraphSchema_K2::IsAllowableBlueprintVariableType(Class);
@@ -614,7 +566,7 @@ static void FBlueprintActionDatabaseImpl::AddClassCastActions(UClass* const Clas
 }
 
 //------------------------------------------------------------------------------
-static void FBlueprintActionDatabaseImpl::AddComponentClassActions(UClass* const Class, FActionList& ActionListOut)
+static void BlueprintActionDatabaseImpl::AddComponentClassActions(UClass* const Class, FActionList& ActionListOut)
 {
 	bool bIsSpawnable = !Class->HasAnyClassFlags(CLASS_Abstract) && Class->HasMetaData(FBlueprintMetadata::MD_BlueprintSpawnableComponent);
 
@@ -628,7 +580,7 @@ static void FBlueprintActionDatabaseImpl::AddComponentClassActions(UClass* const
 }
 
 //------------------------------------------------------------------------------
-static void FBlueprintActionDatabaseImpl::AddClassEnumActions(UClass const* const Class, FActionList& ActionListOut)
+static void BlueprintActionDatabaseImpl::AddClassEnumActions(UClass const* const Class, FActionList& ActionListOut)
 {
 	using namespace FBlueprintNodeSpawnerFactory; // for MakeEnumNodeSpawner()
 
@@ -665,7 +617,7 @@ static void FBlueprintActionDatabaseImpl::AddClassEnumActions(UClass const* cons
 }
 
 //------------------------------------------------------------------------------
-static void FBlueprintActionDatabaseImpl::AddClassStructActions(UClass const* const Class, FActionList& ActionListOut)
+static void BlueprintActionDatabaseImpl::AddClassStructActions(UClass const* const Class, FActionList& ActionListOut)
 {
 	using namespace FBlueprintNodeSpawnerFactory; // for MakeStructNodeSpawner()
 
@@ -694,7 +646,7 @@ static void FBlueprintActionDatabaseImpl::AddClassStructActions(UClass const* co
 }
 
 //------------------------------------------------------------------------------
-static void FBlueprintActionDatabaseImpl::AddBlueprintGraphActions(UClass const* const Class, FActionList& ActionListOut)
+static void BlueprintActionDatabaseImpl::AddBlueprintGraphActions(UClass const* const Class, FActionList& ActionListOut)
 {
 	using namespace FBlueprintNodeSpawnerFactory; // for MakeMacroNodeSpawner()
 
@@ -725,7 +677,7 @@ static void FBlueprintActionDatabaseImpl::AddBlueprintGraphActions(UClass const*
 }
 
 //------------------------------------------------------------------------------
-static void FBlueprintActionDatabaseImpl::AddAutonomousNodeActions(UClass* const Class, FActionList& ActionListOut)
+static void BlueprintActionDatabaseImpl::AddAutonomousNodeActions(UClass* const Class, FActionList& ActionListOut)
 {
 	using namespace FBlueprintNodeSpawnerFactory; // for MakeCommentNodeSpawner()
 
@@ -734,13 +686,7 @@ static void FBlueprintActionDatabaseImpl::AddAutonomousNodeActions(UClass* const
 		UK2Node const* const NodeCDO = Class->GetDefaultObject<UK2Node>();
 		check(NodeCDO != nullptr);
 
-		FBlueprintActionDatabase::FActionList NodeActionList;
-		NodeCDO->GetMenuActions(NodeActionList);
-
-		for (UBlueprintNodeSpawner* Spawner : NodeActionList)
-		{
-			ActionListOut.Add(Spawner);
-		}
+		NodeCDO->GetMenuActions(ActionListOut);
 	}
 	// unfortunately, UEdGraphNode_Comment is not a UK2Node and therefore cannot
 	// leverage UK2Node's GetMenuActions() function, so here we HACK it in
@@ -755,7 +701,7 @@ static void FBlueprintActionDatabaseImpl::AddAutonomousNodeActions(UClass* const
 }
 
 //------------------------------------------------------------------------------
-static void FBlueprintActionDatabaseImpl::OnBlueprintChanged(UBlueprint* Blueprint)
+static void BlueprintActionDatabaseImpl::OnBlueprintChanged(UBlueprint* Blueprint)
 {
 	FBlueprintActionDatabase& ActionDatabase = FBlueprintActionDatabase::Get();
 	ActionDatabase.RefreshClassActions(Blueprint->SkeletonGeneratedClass);
@@ -767,39 +713,79 @@ static void FBlueprintActionDatabaseImpl::OnBlueprintChanged(UBlueprint* Bluepri
  ******************************************************************************/
 
 //------------------------------------------------------------------------------
-void FBlueprintActionDatabase::Prime()
+FBlueprintActionDatabase& FBlueprintActionDatabase::Get()
 {
-	using namespace FBlueprintActionDatabaseImpl; // for DatabaseState
-	if (DatabaseState == DB_Uninitialized)
+	static FBlueprintActionDatabase* DatabaseInst = nullptr;
+	if (DatabaseInst == nullptr)
 	{
-		// set DB_Initializing before Get(), else we end up in an infinite loop
-		DatabaseState = DB_Initializing;
-
-		FBlueprintActionDatabase& DatabaseInst = Get();
-		DatabaseInst.RefreshAll();
+		DatabaseInst = new FBlueprintActionDatabase();
+	}
 
-		for (TObjectIterator<UBlueprint> BpIt; BpIt; ++BpIt)
-		{
-			BpIt->OnChanged().AddStatic(&FBlueprintActionDatabaseImpl::OnBlueprintChanged);
-		}
+	return *DatabaseInst;
+}
 
-		DatabaseState = DB_Primed;
+//------------------------------------------------------------------------------
+FBlueprintActionDatabase::FBlueprintActionDatabase()
+{
+	RefreshAll();
+	for (TObjectIterator<UBlueprint> BpIt; BpIt; ++BpIt)
+	{
+		BpIt->OnChanged().AddStatic(&BlueprintActionDatabaseImpl::OnBlueprintChanged);
 	}
 }
 
 //------------------------------------------------------------------------------
-FBlueprintActionDatabase& FBlueprintActionDatabase::Get()
+void FBlueprintActionDatabase::AddReferencedObjects(FReferenceCollector& Collector)
 {
-	using namespace FBlueprintActionDatabaseImpl; // for DatabaseState/DB_Uninitialized
+	for (auto& ActionListIt : ClassActions)
+	{
+		Collector.AddReferencedObjects(ActionListIt.Value);
+	}
+}
 
-	static FBlueprintActionDatabase DatabaseInst;
-	if (DatabaseState == DB_Uninitialized)
+//------------------------------------------------------------------------------
+void FBlueprintActionDatabase::Tick(float DeltaTime)
+{
+	const double DurationThreshold = FMath::Min(0.003, DeltaTime * 0.01);
+	
+	double CurrentDuration = 0.0;
+	while (!NewActions.IsEmpty())
 	{
-		// prime the database the first time we access it
-		Prime();
+		{
+			FScopedDurationTimer ScopedTimer(CurrentDuration);
+
+			TWeakObjectPtr<UBlueprintNodeSpawner> NewSpawner;
+			NewActions.Dequeue(NewSpawner);
+
+			// @TODO: The weak pointer could still be valid, but the class could  
+			//        have since been trashed (awaiting GC)... need to block on
+			//        spawners that have since been removed from the database
+			if (!NewSpawner.IsStale())
+			{
+				if (UEdGraphNode* TemplateNode = NewSpawner->GetTemplateNode(nullptr))
+				{
+					// since we're doing this incrementally, someone could have 
+					// already requested this template, and allocated its pins 
+					if (TemplateNode->Pins.Num() == 0)
+					{
+						// in certain scenarios we need pin information from the 
+						// spawner (to help filter by pin context)
+						TemplateNode->AllocateDefaultPins();
+					}
+				}
+			}
+		}
+		if (CurrentDuration > DurationThreshold)
+		{
+			break;
+		}
 	}
+}
 
-	return DatabaseInst;
+//------------------------------------------------------------------------------
+TStatId FBlueprintActionDatabase::GetStatId() const
+{
+	RETURN_QUICK_DECLARE_CYCLE_STAT(FBlueprintActionDatabase, STATGROUP_Tickables);
 }
 
 //------------------------------------------------------------------------------
@@ -816,17 +802,19 @@ void FBlueprintActionDatabase::RefreshAll()
 //------------------------------------------------------------------------------
 void FBlueprintActionDatabase::RefreshClassActions(UClass* const Class)
 {
-	using namespace FBlueprintActionDatabaseImpl;
+	using namespace BlueprintActionDatabaseImpl;
 
 	check(Class != nullptr);
 	bool const bIsBlueprintClass = Class->HasAnyClassFlags(CLASS_CompiledFromBlueprint);
 	bool const bIsSkelClass      = bIsBlueprintClass && FKismetEditorUtilities::IsClassABlueprintSkeleton(Class);
 	bool const bIsReinstClass    = Class->HasAnyClassFlags(CLASS_NewerVersionExists);
-	bool const bIsIgnorableClass = bIsReinstClass || (bIsBlueprintClass && !bIsSkelClass);
+	bool const bIsTransientClass = (Class->GetOutermost() == GetTransientPackage());// Class->HasAnyFlags(RF_Transient);
+	bool const bIsIgnorableClass = bIsReinstClass || bIsTransientClass || (bIsBlueprintClass && !bIsSkelClass);
 	
 	if (!bIsIgnorableClass)
 	{
-		FBlueprintActionDatabaseImpl::FActionList ClassActionList(ClassActions.FindOrAdd(Class));
+		FActionList& ClassActionList = ClassActions.FindOrAdd(Class);
+		ClassActionList.Empty();
 		
 		// class field actions (nodes that represent and perform actions on
 		// specific fields of the class... functions, properties, etc.)
@@ -849,6 +837,12 @@ void FBlueprintActionDatabase::RefreshClassActions(UClass* const Class)
 		// global structs, enums, etc.; elements that wouldn't be caught
 		// normally when sifting through fields on all known classes)
 		AddAutonomousNodeActions(Class, ClassActionList);
+
+		// queue all the new actions for priming 
+		for (UBlueprintNodeSpawner* Action : ClassActionList)
+		{
+			NewActions.Enqueue(Action);
+		}
 	}
 
 	// @TODO: account for all K2ActionMenuBuilder methods...
diff --git a/Engine/Source/Editor/BlueprintGraph/Private/BlueprintActionFilter.cpp b/Engine/Source/Editor/BlueprintGraph/Private/BlueprintActionFilter.cpp
--- a/Engine/Source/Editor/BlueprintGraph/Private/BlueprintActionFilter.cpp
+++ b/Engine/Source/Editor/BlueprintGraph/Private/BlueprintActionFilter.cpp
@@ -12,6 +12,7 @@
 #include "BlueprintEditorUtils.h"	// for FindBlueprintForGraph()
 #include "ObjectEditorUtils.h"		// for IsFunctionHiddenFromClass()/IsVariableCategoryHiddenFromClass()
 #include "K2Node_VariableSet.h"
+#include "K2Node_VariableGet.h"
 // "impure" node types (utilized in BlueprintActionFilterImpl::IsImpure)
 #include "K2Node_IfThenElse.h"
 #include "K2Node_MultiGate.h"
@@ -207,6 +208,51 @@ namespace BlueprintActionFilterImpl
 	 * @return True if the action would produce a bound node, tied to an object that isn't selected.
 	 */
 	static bool IsBoundToUnselectedObject(FBlueprintActionFilter const& Filter, UBlueprintNodeSpawner const* BlueprintAction);
+
+	/**
+	 * 
+	 * 
+	 * @param  BlueprintAction	
+	 * @param  Pin	
+	 * @return 
+	 */
+	static bool HasMatchingPin(UBlueprintNodeSpawner const* BlueprintAction, UEdGraphPin const* Pin);
+
+	/**
+	 * 
+	 * 
+	 * @param  Pin	
+	 * @param  MemberField	
+	 * @return 
+	 */
+	static bool IsPinCompatibleWithTargetSelf(UEdGraphPin const* Pin, UField const* MemberField);
+
+	/**
+	 * 
+	 * 
+	 * @param  Filter	
+	 * @param  BlueprintAction	
+	 * @return 
+	 */
+	static bool IsFunctionMissingPinParam(FBlueprintActionFilter const& Filter, UBlueprintNodeSpawner const* BlueprintAction);
+
+	/**
+	 * 
+	 * 
+	 * @param  Filter	
+	 * @param  BlueprintAction	
+	 * @return 
+	 */
+	static bool IsMissmatchedPropertyType(FBlueprintActionFilter const& Filter, UBlueprintNodeSpawner const* BlueprintAction);
+
+	/**
+	 * 
+	 * 
+	 * @param  Filter	
+	 * @param  BlueprintAction	
+	 * @return 
+	 */
+	static bool IsMissingMatchingPinParam(FBlueprintActionFilter const& Filter, UBlueprintNodeSpawner const* BlueprintAction);
 };
 
 //------------------------------------------------------------------------------
@@ -291,22 +337,24 @@ static UFunction const* BlueprintActionFilterImpl::GetAssociatedFunction(UBluepr
 static bool BlueprintActionFilterImpl::IsImpure(UBlueprintNodeSpawner const* NodeSpawner)
 {
 	bool bIsImpure = false;
-
-	UClass* NodeClass = NodeSpawner->NodeClass;
-	if (NodeClass != nullptr)
+	
+	if (UFunction const* Function = GetAssociatedFunction(NodeSpawner))
+	{
+		bIsImpure = (Function->HasAnyFunctionFlags(FUNC_BlueprintPure) == false);
+	}
+	else if (NodeSpawner->NodeClass != nullptr)
 	{
-		// TODO: why are some of these "impure"?
+		UClass* NodeClass = NodeSpawner->NodeClass;
+		// TODO: why are some of these "impure"?... we shouldn't have hardcoded 
+		//       node types here (game modules cannot add their node types here,
+		//       so we should find another way of identifying "pure" node types)
 		bIsImpure = (NodeClass == UK2Node_IfThenElse::StaticClass()) ||
 			(NodeClass == UK2Node_MultiGate::StaticClass()) ||
 			(NodeClass == UK2Node_MakeArray::StaticClass()) ||
 			(NodeClass == UK2Node_Message::StaticClass()) ||
 			(NodeClass == UK2Node_ExecutionSequence::StaticClass());
 	}
-	else if (UFunction const* Function = GetAssociatedFunction(NodeSpawner))
-	{
-		bIsImpure = (Function->HasAnyFunctionFlags(FUNC_BlueprintPure) != false);
-	}
-
+	
 	return bIsImpure;
 }
 
@@ -701,6 +749,198 @@ static bool BlueprintActionFilterImpl::IsBoundToUnselectedObject(FBlueprintActio
 	return bIsFilteredOut;
 }
 
+//------------------------------------------------------------------------------
+static bool BlueprintActionFilterImpl::HasMatchingPin(UBlueprintNodeSpawner const* BlueprintAction, UEdGraphPin const* Pin)
+{
+	bool bHasCompatiblePin = false;
+
+	UEdGraph* OuterGraph = Pin->GetOwningNode()->GetGraph();
+	if (UEdGraphNode* TemplateNode = BlueprintAction->GetTemplateNode(OuterGraph))
+	{
+		if (TemplateNode->Pins.Num() == 0)
+		{
+			TemplateNode->AllocateDefaultPins();
+		}
+
+		UBlueprint* Blueprint = FBlueprintEditorUtils::FindBlueprintForGraph(OuterGraph);
+		check(Blueprint != nullptr);
+		UEdGraphSchema_K2 const* Schema = CastChecked<UEdGraphSchema_K2>(OuterGraph->GetSchema());
+
+		UClass* CallingContext = (Blueprint->SkeletonGeneratedClass != nullptr) ? Blueprint->SkeletonGeneratedClass : Blueprint->ParentClass;
+		UK2Node* K2TemplateNode = Cast<UK2Node>(TemplateNode);
+		
+		for (int32 PinIndex = 0; !bHasCompatiblePin && (PinIndex < TemplateNode->Pins.Num()); ++PinIndex)
+		{
+			UEdGraphPin* TemplatePin = TemplateNode->Pins[PinIndex];
+			if (!Schema->ArePinsCompatible(Pin, TemplatePin, CallingContext))
+			{
+				continue;
+			}
+			bHasCompatiblePin = true;
+
+			if (K2TemplateNode != nullptr)
+			{
+				FString DisallowedReason;
+				// to catch wildcard connections that are prevented
+				bHasCompatiblePin = !K2TemplateNode->IsConnectionDisallowed(TemplatePin, Pin, DisallowedReason);
+			}
+		}
+	}	
+
+	return bHasCompatiblePin;
+}
+
+//------------------------------------------------------------------------------
+static bool BlueprintActionFilterImpl::IsPinCompatibleWithTargetSelf(UEdGraphPin const* Pin, UField const* MemberField)
+{
+	bool bIsCompatible = false;
+
+	UClass* OwnerClass = Cast<UClass>(MemberField->GetOuter());
+	if ((Pin->Direction == EGPD_Output) && (OwnerClass != nullptr))
+	{
+		FEdGraphPinType const& PinType = Pin->PinType;
+		if (PinType.PinSubCategoryObject.IsValid()) 
+		{
+			// if PC_Object/PC_Interface
+			if (UClass const* const PinObjClass = Cast<UClass>(PinType.PinSubCategoryObject.Get()))
+			{
+				bIsCompatible = PinObjClass->IsChildOf(OwnerClass);
+			}
+		}
+	}
+
+	return bIsCompatible;
+}
+
+//------------------------------------------------------------------------------
+static bool BlueprintActionFilterImpl::IsFunctionMissingPinParam(FBlueprintActionFilter const& Filter, UBlueprintNodeSpawner const* BlueprintAction)
+{
+	bool bIsFilteredOut = false;
+	if (UFunction const* AssociatedFunc = GetAssociatedFunction(BlueprintAction))
+	{
+		UEdGraphSchema_K2 const* K2Schema = GetDefault<UEdGraphSchema_K2>();
+		bool const bIsEventSpawner = BlueprintAction->NodeClass->IsChildOf<UK2Node_Event>();
+
+		for (int32 PinIndex = 0; !bIsFilteredOut && (PinIndex < Filter.Context.Pins.Num()); ++PinIndex)
+		{
+			UEdGraphPin const* ContextPin = Filter.Context.Pins[PinIndex];
+
+			FEdGraphPinType const& PinType = ContextPin->PinType;
+			UK2Node const* const K2Node = CastChecked<UK2Node const>(ContextPin->GetOwningNode());
+			EEdGraphPinDirection const PinDir = ContextPin->Direction;
+			
+			if (K2Schema->IsExecPin(*ContextPin))
+			{
+				bIsFilteredOut = (bIsEventSpawner && (PinDir == EGPD_Output)) || !IsImpure(BlueprintAction);
+			}
+			else
+			{
+				bIsFilteredOut = true;
+				if (K2Schema->FunctionHasParamOfType(AssociatedFunc, K2Node->GetBlueprint(), PinType, (PinDir == EGPD_Input)))
+				{
+					bIsFilteredOut = false;
+				}
+				else if (!bIsEventSpawner)
+				{
+					// need to take "Target" self pins into consideration for objects
+					bIsFilteredOut = !IsPinCompatibleWithTargetSelf(ContextPin, AssociatedFunc);
+				}
+			}
+		}
+	}
+
+	return bIsFilteredOut;
+}
+
+//------------------------------------------------------------------------------
+static bool BlueprintActionFilterImpl::IsMissmatchedPropertyType(FBlueprintActionFilter const& Filter, UBlueprintNodeSpawner const* BlueprintAction)
+{
+	bool bIsFilteredOut = false;
+
+	UField const* AssociatedField = GetAssociatedField(BlueprintAction);
+	if (UProperty const* Property = Cast<UProperty const>(AssociatedField))
+	{
+		TArray<UEdGraphPin*> const& ContextPins = Filter.Context.Pins;
+		if (ContextPins.Num() > 0)
+		{
+			bool const bIsDelegate = Property->IsA<UMulticastDelegateProperty>();
+			bool const bIsGetter   = BlueprintAction->NodeClass->IsChildOf<UK2Node_VariableGet>();
+			bool const bIsSetter   = BlueprintAction->NodeClass->IsChildOf<UK2Node_VariableSet>();
+
+			//K2Schema->ConvertPropertyToPinType(VariableProperty, /*out*/ VariablePin->PinType);
+			for (int32 PinIndex = 0; !bIsFilteredOut && PinIndex < ContextPins.Num(); ++PinIndex)
+			{
+				UEdGraphPin const* ContextPin = ContextPins[PinIndex];
+				FEdGraphPinType const& ContextPinType = ContextPin->PinType;
+				UEdGraphSchema_K2 const* K2Schema = CastChecked<UEdGraphSchema_K2 const>(ContextPin->GetSchema());
+
+				// have to account for "self" context pin
+				if (IsPinCompatibleWithTargetSelf(ContextPin, Property))
+				{
+					continue;
+				}
+				else if (bIsDelegate)
+				{
+					// there are a lot of different delegate nodes, so let's  
+					// just iterate over all the pins
+					bIsFilteredOut = !HasMatchingPin(BlueprintAction, ContextPin);
+				}
+				else if (ContextPinType.PinCategory == K2Schema->PC_Exec)
+				{
+					// setters are impure, and therefore should have exec pins
+					bIsFilteredOut = bIsGetter;
+				}
+				else if (bIsGetter || bIsSetter)
+				{
+					bIsFilteredOut = true;
+
+					EEdGraphPinDirection const PinDir = ContextPin->Direction;
+					if ((PinDir == EGPD_Input) && bIsGetter)
+					{
+						FEdGraphPinType InputPinType;
+						K2Schema->ConvertPropertyToPinType(Property, InputPinType);
+						bIsFilteredOut = !K2Schema->ArePinTypesCompatible(ContextPinType, InputPinType);
+					}
+					else if ((PinDir == EGPD_Output) && bIsSetter)
+					{
+						FEdGraphPinType OutputPinType;
+						K2Schema->ConvertPropertyToPinType(Property, OutputPinType);
+						bIsFilteredOut = !K2Schema->ArePinTypesCompatible(OutputPinType, ContextPinType);
+					}					
+				}
+				else
+				{
+					ensureMsg(false, TEXT("Unhandled property/node pair, we've probably made some bad assuptions."));
+				}
+			}
+		}
+	}
+
+	return bIsFilteredOut;
+}
+
+//------------------------------------------------------------------------------
+static bool BlueprintActionFilterImpl::IsMissingMatchingPinParam(FBlueprintActionFilter const& Filter, UBlueprintNodeSpawner const* BlueprintAction)
+{
+	bool bIsFilteredOut = false;
+
+	UField const* AssociatedField = GetAssociatedField(BlueprintAction);
+	// we have a separate tests for field nodes (IsFunctionMissingPinParam/IsMissmatchedPropertyType)
+	if (AssociatedField == nullptr)
+	{
+		for (UEdGraphPin const* ContextPin : Filter.Context.Pins)
+		{
+			if (!HasMatchingPin(BlueprintAction, ContextPin))
+			{
+				bIsFilteredOut = true;
+				break;
+			}
+		}
+	}
+
+	return bIsFilteredOut;
+}
+
 /*******************************************************************************
  * FBlueprintActionFilter
  ******************************************************************************/
@@ -712,28 +952,42 @@ FBlueprintActionFilter::FBlueprintActionFilter(uint32 Flags/*= 0x00*/)
 
 	//
 	// NOTE: The order of these tests can have perf implications, the more one
-	//       rejects on average the sooner it should be added (they're executed
-	//       in order added)
+	//       rejects on average the later it should be added (they're executed
+	//       in reverse order, so user added tests are ran first and the ones 
+	//       here are ran last)
 	//
 
-	if (!(Flags & BPFILTER_IncludeDeprecated))
+	// add first the most expensive tests (they will be ran last, and therefore
+	// should be operating on a smaller subset of node-spawners)
+	//
+	// this test in-particular spawns a template-node and then calls 
+	// AllocateDefaultPins() which is costly, so it should be very last!
+	AddIsFilteredTest(FIsFilteredDelegate::CreateStatic(IsMissingMatchingPinParam));
+	AddIsFilteredTest(FIsFilteredDelegate::CreateStatic(IsMissmatchedPropertyType));
+	AddIsFilteredTest(FIsFilteredDelegate::CreateStatic(IsFunctionMissingPinParam));
+
+	AddIsFilteredTest(FIsFilteredDelegate::CreateStatic(IsIncompatibleImpureNode));
+	
+	AddIsFilteredTest(FIsFilteredDelegate::CreateStatic(IsFieldCategoryHidden));	
+	if (Flags & BPFILTER_ExcludeGlobalFields)
 	{
-		AddIsFilteredTest(FIsFilteredDelegate::CreateStatic(IsDeprecated));
+		AddIsFilteredTest(FIsFilteredDelegate::CreateStatic(IsGlobal));
 	}
-	AddIsFilteredTest(FIsFilteredDelegate::CreateStatic(IsFilteredNodeType, (Flags & BPFILTER_ExcludeChildNodeTypes) != 0));
-	AddIsFilteredTest(FIsFilteredDelegate::CreateStatic(IsBoundToUnselectedObject));
-	AddIsFilteredTest(FIsFilteredDelegate::CreateStatic(IsEventUnimplementable));
+	
 	AddIsFilteredTest(FIsFilteredDelegate::CreateStatic(IsFieldInaccessible));
+	AddIsFilteredTest(FIsFilteredDelegate::CreateStatic(IsEventUnimplementable));
+	AddIsFilteredTest(FIsFilteredDelegate::CreateStatic(IsPermissionNotGranted));
 	AddIsFilteredTest(FIsFilteredDelegate::CreateStatic(IsRestrictedClassMember));
-	if (Flags & BPFILTER_ExcludeGlobalFields)
+	AddIsFilteredTest(FIsFilteredDelegate::CreateStatic(IsIncompatibleWithGraphType));
+
+	if (!(Flags & BPFILTER_IncludeDeprecated))
 	{
-		AddIsFilteredTest(FIsFilteredDelegate::CreateStatic(IsGlobal));
+		AddIsFilteredTest(FIsFilteredDelegate::CreateStatic(IsDeprecated));
 	}
+
+	AddIsFilteredTest(FIsFilteredDelegate::CreateStatic(IsFilteredNodeType, (Flags & BPFILTER_ExcludeChildNodeTypes) != 0));
 	AddIsFilteredTest(FIsFilteredDelegate::CreateStatic(IsExternalField));
-	AddIsFilteredTest(FIsFilteredDelegate::CreateStatic(IsIncompatibleImpureNode));
-	AddIsFilteredTest(FIsFilteredDelegate::CreateStatic(IsPermissionNotGranted));
-	AddIsFilteredTest(FIsFilteredDelegate::CreateStatic(IsIncompatibleWithGraphType));
-	AddIsFilteredTest(FIsFilteredDelegate::CreateStatic(IsFieldCategoryHidden));
+	AddIsFilteredTest(FIsFilteredDelegate::CreateStatic(IsBoundToUnselectedObject));
 
 	// @TODO: account for all K2ActionMenuBuilder checks...
 	//    NodeCDO->CanCreateUnderSpecifiedSchema(this)
@@ -813,7 +1067,9 @@ bool FBlueprintActionFilter::IsFilteredByThis(UBlueprintNodeSpawner const* Bluep
 	FBlueprintActionFilter const& FilterRef = *this;
 
 	bool bIsFiltered = false;
-	for (int32 TestIndex = 0; TestIndex < FilterTests.Num(); ++TestIndex)
+	// iterate backwards so that custom user test are ran first (and the slow
+	// internal tests are ran last).
+	for (int32 TestIndex = FilterTests.Num()-1; TestIndex >= 0; --TestIndex)
 	{
 		FIsFilteredDelegate const& IsFilteredDelegate = FilterTests[TestIndex];
 		check(IsFilteredDelegate.IsBound());
diff --git a/Engine/Source/Editor/BlueprintGraph/Private/BlueprintBoundNodeSpawner.cpp b/Engine/Source/Editor/BlueprintGraph/Private/BlueprintBoundNodeSpawner.cpp
--- a/Engine/Source/Editor/BlueprintGraph/Private/BlueprintBoundNodeSpawner.cpp
+++ b/Engine/Source/Editor/BlueprintGraph/Private/BlueprintBoundNodeSpawner.cpp
@@ -132,7 +132,7 @@ UEdGraphNode* UBlueprintBoundNodeSpawner::Invoke(UEdGraph* ParentGraph, FVector2
 		BoundObj = BoundObjPtr.Get();
 	}
 	
-	bool const bIsTemplateNode = ParentGraph->HasAnyFlags(RF_Transient);
+	bool const bIsTemplateNode = (ParentGraph->GetOutermost() == GetTransientPackage());
 	if (!bIsTemplateNode && NewNode->IsA<UK2Node_AddComponent>())
 	{
 		BindAddComponentNodeWithAsset(NewNode, BoundObj);
diff --git a/Engine/Source/Editor/BlueprintGraph/Private/BlueprintComponentNodeSpawner.cpp b/Engine/Source/Editor/BlueprintGraph/Private/BlueprintComponentNodeSpawner.cpp
--- a/Engine/Source/Editor/BlueprintGraph/Private/BlueprintComponentNodeSpawner.cpp
+++ b/Engine/Source/Editor/BlueprintGraph/Private/BlueprintComponentNodeSpawner.cpp
@@ -54,7 +54,7 @@ UEdGraphNode* UBlueprintComponentNodeSpawner::Invoke(UEdGraph* ParentGraph, FVec
 	
 	UK2Node_AddComponent* NewNode = CastChecked<UK2Node_AddComponent>(Super::Invoke(ParentGraph, Location, PostSpawnDelegate));
 
-	bool bIsTemplateNode = ParentGraph->HasAnyFlags(RF_Transient);
+	bool const bIsTemplateNode = (ParentGraph->GetOutermost() == GetTransientPackage());
 	if (!bIsTemplateNode)
 	{
 		UBlueprint* Blueprint = NewNode->GetBlueprint();
diff --git a/Engine/Source/Editor/BlueprintGraph/Private/BlueprintNodeSpawner.cpp b/Engine/Source/Editor/BlueprintGraph/Private/BlueprintNodeSpawner.cpp
--- a/Engine/Source/Editor/BlueprintGraph/Private/BlueprintNodeSpawner.cpp
+++ b/Engine/Source/Editor/BlueprintGraph/Private/BlueprintNodeSpawner.cpp
@@ -2,10 +2,36 @@
 
 #include "BlueprintGraphPrivatePCH.h"
 #include "BlueprintNodeSpawner.h"
+#include "BlueprintNodeTemplateCache.h"
+
+/*******************************************************************************
+ * Static UBlueprintNodeSpawner Helpers
+ ******************************************************************************/
+
+namespace BlueprintNodeSpawnerImpl
+{
+	/**
+	 * 
+	 * 
+	 * @return 
+	 */
+	FBlueprintNodeTemplateCache* GetSharedTemplateCache(bool const bNoInit = false);
+}
+
+//------------------------------------------------------------------------------
+FBlueprintNodeTemplateCache* BlueprintNodeSpawnerImpl::GetSharedTemplateCache(bool const bNoInit)
+{
+	static FBlueprintNodeTemplateCache* NodeTemplateManager = nullptr;
+	if (!bNoInit && (NodeTemplateManager == nullptr))
+	{
+		NodeTemplateManager = new FBlueprintNodeTemplateCache();
+	}
+	return NodeTemplateManager;
+}
 
 /*******************************************************************************
  * UBlueprintNodeSpawner
-*******************************************************************************/
+ ******************************************************************************/
 
 //------------------------------------------------------------------------------
 UBlueprintNodeSpawner* UBlueprintNodeSpawner::Create(TSubclassOf<UEdGraphNode> const NodeClass, UObject* Outer/* = nullptr*/, FCustomizeNodeDelegate CustomizeNodeDelegate/* = FCustomizeNodeDelegate()*/)
@@ -28,10 +54,20 @@ UBlueprintNodeSpawner* UBlueprintNodeSpawner::Create(TSubclassOf<UEdGraphNode> c
 //------------------------------------------------------------------------------
 UBlueprintNodeSpawner::UBlueprintNodeSpawner(class FPostConstructInitializeProperties const& PCIP)
 	: Super(PCIP)
-	, CachedNodeTemplate(nullptr)
 {
 }
 
+//------------------------------------------------------------------------------
+UBlueprintNodeSpawner::~UBlueprintNodeSpawner()
+{
+	// @TODO: What if, on shutdown, the "SharedTemplateCache" is destroyed 
+	//        first? Then we'd be working with a bad pointer here
+	if (FBlueprintNodeTemplateCache* TemplateCache = BlueprintNodeSpawnerImpl::GetSharedTemplateCache(/*bNoInit =*/true))
+	{
+		TemplateCache->ClearCachedTemplate(this);
+	}
+}
+
 //------------------------------------------------------------------------------
 UEdGraphNode* UBlueprintNodeSpawner::Invoke(UEdGraph* ParentGraph, FVector2D const Location) const
 {
@@ -85,26 +121,10 @@ FName UBlueprintNodeSpawner::GetDefaultMenuIcon(FLinearColor& ColorOut)
 }
 
 //------------------------------------------------------------------------------
-UEdGraphNode* UBlueprintNodeSpawner::MakeTemplateNode(UEdGraph* Outer) const
-{
-	check(Outer != nullptr);
-	check(Outer->HasAnyFlags(RF_Transient));
-	
-	// @TODO: Do we need to respawn when the outers don't match? If not, this'll
-	//        save time for subsequent menu generation... if the old outer gets
-	//        destroyed then CachedNodeTemplate should be nulled (so we should
-	//        be covered on that front)... maybe we should just check the
-	//        graph's class type (and graph type)
-	if (CachedNodeTemplate == nullptr)// || (CachedNodeTemplate->GetOuter() != Outer))
-	{
-		CachedNodeTemplate = Invoke(Outer, FVector2D::ZeroVector);
-		if (CachedNodeTemplate != nullptr)
-		{
-			CachedNodeTemplate->SetFlags(RF_Transient | RF_ArchetypeObject);
-		}
-	}
-	
-	return CachedNodeTemplate;
+UEdGraphNode* UBlueprintNodeSpawner::GetTemplateNode(UEdGraph* Outer) const
+{	
+	UEdGraphNode* TemplateNode = BlueprintNodeSpawnerImpl::GetSharedTemplateCache()->GetNodeTemplate(this, Outer);
+	return TemplateNode;
 }
 
 //------------------------------------------------------------------------------
@@ -121,7 +141,7 @@ UEdGraphNode* UBlueprintNodeSpawner::Invoke(UEdGraph* ParentGraph, FVector2D con
 		NewNode->NodePosX = Location.X;
 		NewNode->NodePosY = Location.Y;
 
-		bool const bIsTemplateNode = ParentGraph->HasAnyFlags(RF_Transient);
+		bool const bIsTemplateNode = (ParentGraph->GetOutermost() == GetTransientPackage());
 		PostSpawnDelegate.ExecuteIfBound(NewNode, bIsTemplateNode);
 
 		if (!bIsTemplateNode)
diff --git a/Engine/Source/Editor/BlueprintGraph/Private/BlueprintNodeTemplateCache.cpp b/Engine/Source/Editor/BlueprintGraph/Private/BlueprintNodeTemplateCache.cpp
new file mode 100644
--- /dev/null
+++ b/Engine/Source/Editor/BlueprintGraph/Private/BlueprintNodeTemplateCache.cpp
@@ -0,0 +1,225 @@
+// Copyright 1998-2014 Epic Games, Inc. All Rights Reserved.
+
+#include "BlueprintGraphPrivatePCH.h"
+#include "BlueprintNodeTemplateCache.h"
+#include "BlueprintEditorUtils.h" // for FindFirstNativeClass(), FindBlueprintForGraph()
+#include "EdGraph/EdGraph.h"
+
+/*******************************************************************************
+ * Static FBlueprintNodeTemplateCache Helpers
+ ******************************************************************************/
+
+namespace BlueprintNodeTemplateCacheImpl
+{
+	/**
+	 * 
+	 * 
+	 * @param  NodeObj	
+	 * @param  Graph	
+	 * @return 
+	 */
+	static bool IsCompatible(UEdGraphNode* NodeObj, UEdGraph* Graph);
+
+	/**
+	 * 
+	 * 
+	 * @param  BlueprintOuter	
+	 * @param  NodeObj	
+	 * @return 
+	 */
+	static UEdGraph* FindCompatibleGraph(UBlueprint* BlueprintOuter, UEdGraphNode* NodeObj);
+
+	/**
+	 * 
+	 * 
+	 * @param  BlueprintClass	
+	 * @param  ParentClass	
+	 * @param  GeneratedClassType	
+	 * @return 
+	 */
+	static UBlueprint* MakeCompatibleBlueprint(TSubclassOf<UBlueprint> BlueprintClass, UClass* ParentClass, TSubclassOf<UBlueprintGeneratedClass> GeneratedClassType);
+
+	/**
+	 * 
+	 * 
+	 * @param  GraphClass	
+	 * @param  SchemaClass	
+	 * @param  BlueprintOuter	
+	 * @param  GraphType	
+	 * @return 
+	 */
+	static UEdGraph* AddGraph(UBlueprint* BlueprintOuter, TSubclassOf<UEdGraphSchema> SchemaClass);
+}
+
+//------------------------------------------------------------------------------
+static bool BlueprintNodeTemplateCacheImpl::IsCompatible(UEdGraphNode* NodeObj, UEdGraph* Graph)
+{
+	return NodeObj->CanCreateUnderSpecifiedSchema(Graph->GetSchema());
+}
+
+//------------------------------------------------------------------------------
+static UEdGraph* BlueprintNodeTemplateCacheImpl::FindCompatibleGraph(UBlueprint* BlueprintOuter, UEdGraphNode* NodeObj)
+{
+	UEdGraph* FoundGraph = nullptr;
+
+	TArray<UObject*> BlueprintChildObjs;
+	GetObjectsWithOuter(BlueprintOuter, BlueprintChildObjs, /*bIncludeNestedObjects =*/false, /*ExclusionFlags =*/RF_PendingKill);
+
+	for (UObject* Child : BlueprintChildObjs)
+	{
+		UEdGraph* ChildGraph = Cast<UEdGraph>(Child);
+		if ((ChildGraph != nullptr) && IsCompatible(NodeObj, ChildGraph))
+		{
+			FoundGraph = ChildGraph;
+			break;
+		}
+	}
+
+	return FoundGraph;
+}
+
+//------------------------------------------------------------------------------
+static UBlueprint* BlueprintNodeTemplateCacheImpl::MakeCompatibleBlueprint(TSubclassOf<UBlueprint> BlueprintClass, UClass* ParentClass, TSubclassOf<UBlueprintGeneratedClass> GeneratedClassType)
+{
+	EBlueprintType BlueprintType = BPTYPE_Normal;
+	// @TODO: BPTYPE_LevelScript requires a level outer, which we don't want to have here... can we get away without it?
+// 	if (BlueprintClass->IsChildOf<ULevelScriptBlueprint>())
+// 	{
+// 		BlueprintType = BPTYPE_LevelScript;
+// 	}
+
+	if (GeneratedClassType == nullptr)
+	{
+		GeneratedClassType = UBlueprintGeneratedClass::StaticClass();
+	}
+
+	UPackage* BlueprintOuter = GetTransientPackage();
+	FString const DesiredName = FString::Printf(TEXT("PROTO_BP_%s"), *BlueprintClass->GetName());
+	FName   const BlueprintName = MakeUniqueObjectName(BlueprintOuter, BlueprintClass, FName(*DesiredName));
+
+	BlueprintClass = FBlueprintEditorUtils::FindFirstNativeClass(BlueprintClass);
+	UBlueprint* NewBlueprint = FKismetEditorUtilities::CreateBlueprint(ParentClass, BlueprintOuter, BlueprintName, BlueprintType, BlueprintClass, GeneratedClassType);
+	NewBlueprint->SetFlags(RF_Transient);
+
+	return NewBlueprint;
+}
+
+//------------------------------------------------------------------------------
+static UEdGraph* BlueprintNodeTemplateCacheImpl::AddGraph(UBlueprint* BlueprintOuter, TSubclassOf<UEdGraphSchema> SchemaClass)
+{
+	UClass* GraphClass = UEdGraph::StaticClass();
+	FName const GraphName = MakeUniqueObjectName(BlueprintOuter, GraphClass, FName(TEXT("TEMPLATE_NODE_OUTER")));
+
+	UEdGraph* NewGraph = ConstructObject<UEdGraph>(GraphClass, BlueprintOuter, GraphName, RF_Transient);
+	NewGraph->Schema = SchemaClass;
+
+	return NewGraph;
+}
+
+//------------------------------------------------------------------------------
+FBlueprintNodeTemplateCache::FBlueprintNodeTemplateCache()
+	: FGCObject()
+{
+	using namespace BlueprintNodeTemplateCacheImpl; // for MakeCompatibleBlueprint()
+
+	TemplateOuters.Add(MakeCompatibleBlueprint(UBlueprint::StaticClass(), AActor::StaticClass(), UBlueprintGeneratedClass::StaticClass()));
+	TemplateOuters.Add(MakeCompatibleBlueprint(UAnimBlueprint::StaticClass(), UAnimInstance::StaticClass(), UAnimBlueprintGeneratedClass::StaticClass()));
+}
+
+//------------------------------------------------------------------------------
+UEdGraphNode* FBlueprintNodeTemplateCache::GetNodeTemplate(UBlueprintNodeSpawner const* NodeSpawner, UEdGraph* TargetGraph)
+{
+	using namespace BlueprintNodeTemplateCacheImpl;
+
+	UEdGraphNode* TemplateNode = nullptr;
+	if (UEdGraphNode** FoundNode = NodeTemplateCache.Find(NodeSpawner))
+	{
+		TemplateNode = *FoundNode;
+	}
+	else if (NodeSpawner->NodeClass != nullptr)
+	{
+		UEdGraphNode* NodeCDO = NodeSpawner->NodeClass->GetDefaultObject<UEdGraphNode>();
+		check(NodeCDO != nullptr);
+
+		UBlueprint* TargetBlueprint = nullptr;
+		TSubclassOf<UBlueprint> BlueprintClass;
+
+		bool const bHasTargetGraph = (TargetGraph != nullptr);
+		if (bHasTargetGraph)
+		{
+			// by the time we're asking for a prototype for this spawner, we should
+			// be sure that it is compatible with the TargetGraph
+			check(IsCompatible(NodeCDO, TargetGraph));
+
+			UBlueprint* TargetBlueprint = FBlueprintEditorUtils::FindBlueprintForGraph(TargetGraph);
+			check(TargetBlueprint != nullptr);
+			TSubclassOf<UBlueprint> BlueprintClass = TargetBlueprint->GetClass();
+		}		
+
+		UBlueprint* CompatibleBlueprint = nullptr;
+		UEdGraph*   CompatibleOuter     = nullptr;
+		// find a compatible outer (don't want to have to create a new one if we don't have to)
+		for (UBlueprint* Blueprint : TemplateOuters)
+		{
+			CompatibleOuter = FindCompatibleGraph(Blueprint, NodeCDO);
+			if (CompatibleOuter != nullptr)
+			{
+				CompatibleBlueprint = Blueprint;
+				break;
+			}
+			else if ((BlueprintClass != nullptr) && Blueprint->GetClass()->IsChildOf(BlueprintClass))
+			{
+				CompatibleBlueprint = Blueprint;
+			}
+		}
+
+		// if a TargetGraph was supplied, and we couldn't find a suitable outer
+		// for this template-node, then attempt to emulate that graph
+		if (bHasTargetGraph)
+		{
+			if (CompatibleBlueprint == nullptr)
+			{
+				TSubclassOf<UBlueprintGeneratedClass> GeneratedClassType = UBlueprintGeneratedClass::StaticClass();
+				if (TargetBlueprint->GeneratedClass != nullptr)
+				{
+					GeneratedClassType = TargetBlueprint->GeneratedClass->GetClass();
+				}
+
+				CompatibleBlueprint = MakeCompatibleBlueprint(BlueprintClass, TargetBlueprint->ParentClass, GeneratedClassType);
+				TemplateOuters.Add(CompatibleBlueprint);
+
+				// this graph may come default with a compatible graph
+				CompatibleOuter = FindCompatibleGraph(CompatibleBlueprint, NodeCDO);
+			}
+
+			if (CompatibleOuter == nullptr)
+			{
+				CompatibleOuter = AddGraph(CompatibleBlueprint, TargetGraph->Schema);
+			}
+		}
+
+		if (CompatibleOuter != nullptr)
+		{
+			TemplateNode = NodeSpawner->Invoke(CompatibleOuter, FVector2D::ZeroVector);
+			NodeTemplateCache.Add(NodeSpawner, TemplateNode);
+		}
+	}
+
+	return TemplateNode;
+}
+
+//------------------------------------------------------------------------------
+void FBlueprintNodeTemplateCache::ClearCachedTemplate(UBlueprintNodeSpawner const* NodeSpawner)
+{
+	NodeTemplateCache.Remove(NodeSpawner); // GC should take care of the rest
+}
+
+//------------------------------------------------------------------------------
+void FBlueprintNodeTemplateCache::AddReferencedObjects(FReferenceCollector& Collector)
+{
+	for (auto& TemplateEntry : NodeTemplateCache)
+	{
+		Collector.AddReferencedObject(TemplateEntry.Value);
+	}
+	Collector.AddReferencedObjects(TemplateOuters);
+}
diff --git a/Engine/Source/Editor/BlueprintGraph/Private/K2ActionMenuBuilder.cpp b/Engine/Source/Editor/BlueprintGraph/Private/K2ActionMenuBuilder.cpp
--- a/Engine/Source/Editor/BlueprintGraph/Private/K2ActionMenuBuilder.cpp
+++ b/Engine/Source/Editor/BlueprintGraph/Private/K2ActionMenuBuilder.cpp
@@ -1049,7 +1049,7 @@ void FK2ActionMenuBuilder::GetContextAllowedNodeTypes(FBlueprintGraphActionListB
 
 //------------------------------------------------------------------------------
 
-struct FClassDynamicCastHelper
+struct FClassDynamicCastMenuUtils
 {
 	static bool CanCastToClass(const UClass* TestClass, const UEdGraphSchema_K2* K2Schema)
 	{
@@ -1359,7 +1359,7 @@ void FK2ActionMenuBuilder::GetPinAllowedNodeTypes(FBlueprintGraphActionListBuild
 			}
 		}
 
-		FClassDynamicCastHelper::GetClassDynamicCastNodes(ContextMenuBuilder, K2Schema);
+		FClassDynamicCastMenuUtils::GetClassDynamicCastNodes(ContextMenuBuilder, K2Schema);
 
 		if ((FromPin.Direction == EGPD_Input) && (FromPin.PinType.PinCategory == K2Schema->PC_Int))
 		{
diff --git a/Engine/Source/Editor/BlueprintGraph/Private/K2Node_EnumEquality.cpp b/Engine/Source/Editor/BlueprintGraph/Private/K2Node_EnumEquality.cpp
--- a/Engine/Source/Editor/BlueprintGraph/Private/K2Node_EnumEquality.cpp
+++ b/Engine/Source/Editor/BlueprintGraph/Private/K2Node_EnumEquality.cpp
@@ -9,6 +9,8 @@
 #include "BlueprintNodeSpawner.h"
 #include "EditorCategoryUtils.h"
 
+#define LOCTEXT_NAMESPACE "K2Node_EnumEquality"
+
 UK2Node_EnumEquality::UK2Node_EnumEquality(const class FPostConstructInitializeProperties& PCIP)
 	: Super(PCIP)
 {
@@ -124,6 +126,24 @@ void UK2Node_EnumEquality::NotifyPinConnectionListChanged(UEdGraphPin* Pin)
 	}
 }
 
+bool UK2Node_EnumEquality::IsConnectionDisallowed(const UEdGraphPin* MyPin, const UEdGraphPin* OtherPin, FString& OutReason) const
+{
+	UEdGraphSchema_K2 const* K2Schema = GetDefault<UEdGraphSchema_K2>();
+
+	bool const bOtherPinIsEnum = (OtherPin->PinType.PinCategory == K2Schema->PC_Byte) &&
+		(OtherPin->PinType.PinSubCategoryObject.IsValid()) &&
+		(OtherPin->PinType.PinSubCategoryObject.Get()->IsA(UEnum::StaticClass()));
+
+
+	bool bIsDisallowed = false;
+	if (!bOtherPinIsEnum && (MyPin->Direction == EGPD_Input))
+	{
+		bIsDisallowed = true;
+		OutReason = LOCTEXT("InputIsNotEnum", "Cannot use the enum equality operator on anything but enums.").ToString();
+	}
+	return bIsDisallowed;
+}
+
 UEdGraphPin* UK2Node_EnumEquality::GetReturnValuePin() const
 {
 	const UEdGraphSchema_K2* K2Schema = GetDefault<UEdGraphSchema_K2>();
@@ -251,3 +271,5 @@ FText UK2Node_EnumEquality::GetMenuCategory() const
 {
 	return FEditorCategoryUtils::GetCommonCategory(FCommonEditorCategory::Enum);
 }
+
+#undef LOCTEXT_NAMESPACE
diff --git a/Engine/Source/Editor/BlueprintGraph/Public/BlueprintActionDatabase.h b/Engine/Source/Editor/BlueprintGraph/Public/BlueprintActionDatabase.h
--- a/Engine/Source/Editor/BlueprintGraph/Public/BlueprintActionDatabase.h
+++ b/Engine/Source/Editor/BlueprintGraph/Public/BlueprintActionDatabase.h
@@ -2,10 +2,14 @@
 
 #pragma once
 
+#include "GCObject.h"
+#include "TickableEditorObject.h"
+
 // Forward declarations
 class UBlueprint;
 class FBlueprintActionFilter;
 class UBlueprintNodeSpawner;
+class FReferenceCollector;
 
 /**
  * Serves as a container for all available blueprint actions (no matter the 
@@ -16,17 +20,13 @@ class UBlueprintNodeSpawner;
  * 
  * @TODO:  Hook up to handle class recompile events, along with enum/struct asset creation events.
  */
-class BLUEPRINTGRAPH_API FBlueprintActionDatabase
+class BLUEPRINTGRAPH_API FBlueprintActionDatabase : public FGCObject, public FTickableEditorObject
 {
 public:
 	typedef TArray<UBlueprintNodeSpawner*> FActionList;
 	typedef TMap<UClass*, FActionList>     FClassActionMap;
 
-	/**
-	 * Ensures that the database singleton is initialized and populated.
-	 */
-	static void Prime();
-
+public:
 	/**
 	 * Getter to access the database singleton. Will populate the database first 
 	 * if this is the first time accessing it.
@@ -35,6 +35,16 @@ public:
 	 */
 	static FBlueprintActionDatabase& Get();
 
+	// FTickableEditorObject interface
+	virtual void Tick(float DeltaTime) override;
+	virtual bool IsTickable() const override { return true; }
+	virtual TStatId GetStatId() const override;
+	// End FTickableEditorObject interface
+
+	// FGCObject interface
+	virtual void AddReferencedObjects(FReferenceCollector& Collector) override;
+	// End FGCObject interface
+
 	/**
 	 * Populates the action database from scratch. Loops over every known class
 	 * and records a set of node-spawners associated with each.
@@ -63,10 +73,10 @@ public:
 	 * @return The entire action database, which maps specific classes to arrays of associated node-spawners.
 	 */
 	FClassActionMap const& GetAllActions();
-	
+
 private:
 	/** Private constructor for singleton purposes. */
-	FBlueprintActionDatabase() {} 
+	FBlueprintActionDatabase();
 
 	/** 
 	 * A map of associated node-spawners for each class. A spawner that 
@@ -76,4 +86,11 @@ private:
 	 * node's type.
 	 */
 	FClassActionMap ClassActions;
+
+	/** 
+	 * Holds newly allocated actions that need to be "primed". Priming is 
+	 * something we do on Tick() aimed at speeding up performance (like 
+	 * pre-caching the spawner's template-node, etc.).
+	 */
+	TQueue<TWeakObjectPtr<UBlueprintNodeSpawner>> NewActions;
 };
diff --git a/Engine/Source/Editor/BlueprintGraph/Public/BlueprintNodeSpawner.h b/Engine/Source/Editor/BlueprintGraph/Public/BlueprintNodeSpawner.h
--- a/Engine/Source/Editor/BlueprintGraph/Public/BlueprintNodeSpawner.h
+++ b/Engine/Source/Editor/BlueprintGraph/Public/BlueprintNodeSpawner.h
@@ -29,6 +29,9 @@ class BLUEPRINTGRAPH_API UBlueprintNodeSpawner : public UObject
 	DECLARE_DELEGATE_TwoParams(FCustomizeNodeDelegate, UEdGraphNode*, bool);
 
 public:
+	/** */
+	virtual ~UBlueprintNodeSpawner();
+
 	/**
 	 * Creates a new UBlueprintNodeSpawner for the specified node class. Sets
 	 * the allocated spawner's NodeClass and CustomizeNodeDelegate fields from
@@ -144,24 +147,15 @@ public:
 	virtual FName GetDefaultMenuIcon(FLinearColor& ColorOut);
 
 	/**
-	 * Expects the supplied outer to be transient. Will spawn a new template-
-	 * node if the cached one is null, or if its graph outer doesn't match the
-	 * one specified here. Once a new node is spawned, this class will cache it
-	 * to save on spawning a new one with later requests (and to keep the node
-	 * from being GC'd).
-	 *
-	 * This should not have to be overridden (it calls Invoke() with the
-	 * specified outer; the transient-ness of the graph should signify that it
-	 * is for a template).
+	 * Retrieves a cached template for the node that this is set to spawn. Will
+	 * instantiate a new template if one didn't previously exist. If the 
+	 * template-node is not compatible with any of our cached UEdGraph outers, 
+	 * then we use TargetGraph as a model to create one that will work.
 	 *
-	 * Unfortunately, the supplied graph requires a Blueprint outer. Certain 
-	 * node types rely on this assumption; this is why the template-node and the
-	 * supplied graph cannot belong to the transient package themselves.
-	 *
-	 * @param  Outer    A transient graph (with a Blueprint outer), to act as the template-node's outer.
+	 * @param  TargetGraph    Optional param that defines a compatible graph outer (used as an achetype if we don't have a compatible outer on hand).
 	 * @return Should return a new/cached template-node (but could be null, or some pre-existing node... depends on the sub-class's Invoke() method).
 	 */
-	UEdGraphNode* MakeTemplateNode(UEdGraph* Outer) const;
+	UEdGraphNode* GetTemplateNode(UEdGraph* TargetGraph = nullptr) const;
 
 protected:
 	/**
@@ -174,18 +168,6 @@ protected:
 	 * @return Null if it failed to spawn a node (if NodeClass is null), otherwise a newly spawned node.
 	 */
 	UEdGraphNode* Invoke(UEdGraph* ParentGraph, FVector2D const Location, FCustomizeNodeDelegate PostSpawnDelegate) const;
-
-private:
-	/**
-	 * We try to avoid the spawning of node-templates (save on perf/mem where we
-	 * can), but sometimes users need to pull non-static data from the node 
-	 * itself (before it's spawned).
-	 *
-	 * Here we store the node template once it's created. Stored so that it can
-	 * be reused (and not reallocated) and not GC'd.
-	 */
-	UPROPERTY(Transient)
-	mutable UEdGraphNode* CachedNodeTemplate;
 };
 
 /*******************************************************************************
diff --git a/Engine/Source/Editor/BlueprintGraph/Public/BlueprintNodeTemplateCache.h b/Engine/Source/Editor/BlueprintGraph/Public/BlueprintNodeTemplateCache.h
new file mode 100644
--- /dev/null
+++ b/Engine/Source/Editor/BlueprintGraph/Public/BlueprintNodeTemplateCache.h
@@ -0,0 +1,66 @@
+// Copyright 1998-2014 Epic Games, Inc. All Rights Reserved.
+
+#pragma once
+
+#include "GCObject.h"
+
+// Forward declarations
+class UBlueprintNodeSpawner;
+class UBlueprint;
+class UEdGraphNode;
+class FReferenceCollector;
+
+/**
+ * Serves as a centralized data-store for all UBlueprintNodeSpawner node-
+ * templates. Implemented this way (rather than internal to UBlueprintNodeSpawner) 
+ * since node-templates require a UEdGraph/UBlueprint outer chain. Instead of 
+ * instantiating a bunch of graphs/blueprints per UBlueprintNodeSpawner, we'd 
+ * rather have a small centralized set here.
+ */
+class FBlueprintNodeTemplateCache : public FGCObject
+{
+public:
+	/** */
+	FBlueprintNodeTemplateCache();
+
+	/**
+	 * Retrieves a cached template associated with the supplied spawner. Will
+	 * instantiate a new template if one didn't already exist. If the
+	 * template-node is not compatible with any cached UEdGraph outer, then
+	 * we use TargetGraph as a model to create one that will work.
+	 * 
+	 * @param  NodeSpawner	Acts as the key for the template lookup; takes care of spawning the template-node if one didn't already exist.
+	 * @param  TargetGraph	Optional param that defines a compatible graph outer (used as an achetype if we don't have a compatible outer on hand).
+	 * @return Should return a new/cached template-node (but could be null, or some pre-existing node... depends on the sub-class's Invoke() method).
+	 */
+	UEdGraphNode* GetNodeTemplate(UBlueprintNodeSpawner const* NodeSpawner, UEdGraph* TargetGraph = nullptr);
+
+	/**
+	 * Wipes any nodes that were cached on behalf of the specified spawner 
+	 * (should be called when NodeSpawner is destroyed, in case 
+	 * GetNodeTemplate() was called for it).
+	 * 
+	 * @param  NodeSpawner	The spawner we want cached node-templates cleared for.
+	 */
+	void ClearCachedTemplate(UBlueprintNodeSpawner const* NodeSpawner);
+
+	// FGCObject interface
+	virtual void AddReferencedObjects(FReferenceCollector& Collector) override;
+	// End FGCObject interface
+
+private:
+	/** 
+	 * Unfortunately, we cannot nest template-nodes in the transient package. 
+	 * Certain nodes operate on the assumption that they have a UEdGraph outer, 
+	 * while a certain subset expect the graph to have a UBlueprint outer. This
+	 * means we cannot spawn templates without a blueprint/graph to add them to.
+	 * 
+	 * This array holds intermediate blueprints that we use to parent the 
+	 * template-nodes. Ideally we only need a small handful that are compatible 
+	 * with all nodes.
+	 */
+	TArray<UBlueprint*> TemplateOuters;
+
+	/** */
+	TMap<UBlueprintNodeSpawner const*, UEdGraphNode*> NodeTemplateCache;
+};
diff --git a/Engine/Source/Editor/Kismet/Private/BlueprintActionMenuBuilder.cpp b/Engine/Source/Editor/Kismet/Private/BlueprintActionMenuBuilder.cpp
--- a/Engine/Source/Editor/Kismet/Private/BlueprintActionMenuBuilder.cpp
+++ b/Engine/Source/Editor/Kismet/Private/BlueprintActionMenuBuilder.cpp
@@ -15,6 +15,7 @@
 #include "SMyBlueprint.h"				// for SelectionAsVar()
 #include "EdGraphSchema_K2_Actions.h"	// for FEdGraphSchemaAction_K2Var
 #include "BlueprintEditorUtils.h"		// for FindBlueprintForGraphChecked()
+#include "BlueprintEditor.h"			// for GetFocusedGraph()
 
 #define LOCTEXT_NAMESPACE "BlueprintActionMenuBuilder"
 
@@ -42,22 +43,21 @@ public:
 	 * Spawns a new FBlueprintActionMenuItem with the node-spawner. Constructs
 	 * the menu item's category, name, tooltip, etc.
 	 * 
-	 * @param  ItemOwner	The menu item owner that this will be assigned to.
+	 * @param  EditorContext
 	 * @param  Action		The node-spawner that the new menu item should wrap.
 	 * @return A newly allocated FBlueprintActionMenuItem (which wraps the supplied action).
 	 */
-	TSharedPtr<FEdGraphSchemaAction> MakeMenuItem(FBlueprintActionMenuBuilder& ItemOwner, UBlueprintNodeSpawner* Action);
+	TSharedPtr<FEdGraphSchemaAction> MakeMenuItem(TWeakPtr<FBlueprintEditor> EditorContext, UBlueprintNodeSpawner* Action);
 
 	/**
 	 * Clones the property spawner and clears its NodeClass field (used to 
-	 * represent the property as a whole, encompasing all the node types that
+	 * represent the property as a whole, encompassing all the node types that
 	 * can be associated with it, like: getters, setters, etc.
 	 *
-	 * @param  ItemOwner	The menu item owner that this will be assigned to.
 	 * @param  NodeSpawner	The node-spawner that we want to clone (and create a proxy for).
 	 * @return A cloned UBlueprintPropertyNodeSpawner (representing the property).
 	 */
-	UBlueprintPropertyNodeSpawner* MakePropertyActionProxy(FBlueprintActionMenuBuilder& ItemOwner, UBlueprintPropertyNodeSpawner* NodeSpawner);
+	UBlueprintPropertyNodeSpawner* MakePropertyActionProxy( UBlueprintPropertyNodeSpawner* NodeSpawner);
 	
 	/**
 	 * Checks to see if a proxy has already been made for the specified property
@@ -90,53 +90,53 @@ protected:
 	 * provided, then it spawns a temporary node and pulls one from that node's 
 	 * title.
 	 * 
-	 * @param  ItemOwner	The menu builder that this name will be listed from.
+	 * @param  EditorContext
 	 * @param  Action		The action you want to suss name information from.
 	 * @return A name for the menu item wrapping this action.
 	 */
-	FText GetMenuNameForAction(FBlueprintActionMenuBuilder& ItemOwner, UBlueprintNodeSpawner* Action);
+	FText GetMenuNameForAction(TWeakPtr<FBlueprintEditor> EditorContext, UBlueprintNodeSpawner* Action);
 
 	/**
 	 * Attempts to pull a menu category from the supplied spawner. If one isn't 
 	 * provided, then it spawns a temporary node and pulls one from that node.
 	 * 
-	 * @param  ItemOwner	The menu builder that the item will be listed from.
+	 * @param  EditorContext
 	 * @param  Action		The action you want to suss category information from.
 	 * @return A category for the menu item wrapping this action.
 	 */
-	FText GetCategoryForAction(FBlueprintActionMenuBuilder& ItemOwner, UBlueprintNodeSpawner* Action);
+	FText GetCategoryForAction(TWeakPtr<FBlueprintEditor> EditorContext, UBlueprintNodeSpawner* Action);
 
 	/**
 	 * Attempts to pull a menu tooltip from the supplied spawner. If one isn't 
 	 * provided, then it spawns a temporary node and pulls one from that node.
 	 * 
-	 * @param  ItemOwner	The menu builder that the item will be listed from.
+	 * @param  EditorContext
 	 * @param  Action		The action you want to suss tooltip information from.
 	 * @return A tooltip for the menu item wrapping this action.
 	 */
-	FText GetTooltipForAction(FBlueprintActionMenuBuilder& ItemOwner, UBlueprintNodeSpawner* Action);
+	FText GetTooltipForAction(TWeakPtr<FBlueprintEditor> EditorContext, UBlueprintNodeSpawner* Action);
 
 	/**
 	 * Attempts to pull a keywords from the supplied spawner. If one isn't 
 	 * provided, then it spawns a temporary node and pulls them from that.
 	 * 
-	 * @param  ItemOwner	The menu builder that the item will be listed from.
+	 * @param  EditorContext
 	 * @param  Action		The action you want to suss keyword information from.
 	 * @return A keyword text string for the menu item wrapping this action.
 	 */
-	FText GetSearchKeywordsForAction(FBlueprintActionMenuBuilder& ItemOwner, UBlueprintNodeSpawner* Action);
+	FText GetSearchKeywordsForAction(TWeakPtr<FBlueprintEditor> EditorContext, UBlueprintNodeSpawner* Action);
 
 	/**
 	 * Attempts to pull a menu icon information from the supplied spawner. If
 	 * info isn't provided, then it spawns a temporary node and pulls data from 
 	 * that node.
 	 * 
-	 * @param  ItemOwner	The menu builder that the item will be listed from.
+	 * @param  EditorContext
 	 * @param  Action		The action you want to suss icon information from.
 	 * @param  ColorOut		The color to tint the icon with.
 	 * @return Name of the brush to use (use FEditorStyle::GetBrush() to resolve).
 	 */
-	FName GetMenuIconForAction(FBlueprintActionMenuBuilder& ItemOwner, UBlueprintNodeSpawner* Action, FLinearColor& ColorOut);
+	FName GetMenuIconForAction(TWeakPtr<FBlueprintEditor> EditorContext, UBlueprintNodeSpawner* Action, FLinearColor& ColorOut);
 
 	/**
 	 * Utility getter function that retrieves the blueprint context for the menu
@@ -144,17 +144,16 @@ protected:
 	 * 
 	 * @return The first blueprint out of the cached FBlueprintActionContext.
 	 */
-	UBlueprint* GetBlueprint() const;
+	UBlueprint* GetTargetBlueprint() const;
 
 	/**
-	 * Utility function that will pull the "OwnerOfTemporaries" field from 
-	 * ItemOwner. If the object is null, then this will allocate one on behalf
-	 * of ItemOwner.
 	 * 
-	 * @param  The FBlueprintActionMenuBuilder you want an "OwnerOfTemporaries" from. 
-	 * @return The FBlueprintActionMenuBuilder's OwnerOfTemporaries.
+	 * 
+	 * @param  Action	
+	 * @param  EditorContext	
+	 * @return 
 	 */
-	UEdGraph* GetOwnerOfTemporaries(FBlueprintActionMenuBuilder& ItemOwner) const;
+	UEdGraphNode* GetTemplateNode(UBlueprintNodeSpawner* Action, TWeakPtr<FBlueprintEditor> EditorContext) const;
 
 	/** Cached context for the blueprint menu being built */
 	FBlueprintActionContext const& Context;
@@ -172,17 +171,17 @@ FBlueprintActionMenuItemFactory::FBlueprintActionMenuItemFactory(FBlueprintActio
 }
 
 //------------------------------------------------------------------------------
-TSharedPtr<FEdGraphSchemaAction> FBlueprintActionMenuItemFactory::MakeMenuItem(FBlueprintActionMenuBuilder& ItemOwner, UBlueprintNodeSpawner* Action)
+TSharedPtr<FEdGraphSchemaAction> FBlueprintActionMenuItemFactory::MakeMenuItem(TWeakPtr<FBlueprintEditor> EditorContext, UBlueprintNodeSpawner* Action)
 {
 	FLinearColor IconTint = FLinearColor::White;
-	FName IconBrushName = GetMenuIconForAction(ItemOwner, Action, IconTint);
+	FName IconBrushName = GetMenuIconForAction(EditorContext, Action, IconTint);
 
 	FBlueprintActionMenuItem* NewMenuItem = new FBlueprintActionMenuItem(Action, FEditorStyle::GetBrush(IconBrushName), IconTint, MenuGrouping);
 	
-	NewMenuItem->MenuDescription    = GetMenuNameForAction(ItemOwner, Action);
-	NewMenuItem->TooltipDescription = GetTooltipForAction(ItemOwner, Action).ToString();
-	NewMenuItem->Category           = GetCategoryForAction(ItemOwner, Action).ToString();
-	NewMenuItem->Keywords           = GetSearchKeywordsForAction(ItemOwner, Action).ToString();
+	NewMenuItem->MenuDescription    = GetMenuNameForAction(EditorContext, Action);
+	NewMenuItem->TooltipDescription = GetTooltipForAction(EditorContext, Action).ToString();
+	NewMenuItem->Category           = GetCategoryForAction(EditorContext, Action).ToString();
+	NewMenuItem->Keywords           = GetSearchKeywordsForAction(EditorContext, Action).ToString();
 
 	if (NewMenuItem->Category.IsEmpty())
 	{
@@ -197,7 +196,7 @@ TSharedPtr<FEdGraphSchemaAction> FBlueprintActionMenuItemFactory::MakeMenuItem(F
 }
 
 //------------------------------------------------------------------------------
-UBlueprintPropertyNodeSpawner* FBlueprintActionMenuItemFactory::MakePropertyActionProxy(FBlueprintActionMenuBuilder& ItemOwner, UBlueprintPropertyNodeSpawner* PropertyNodeSpawner)
+UBlueprintPropertyNodeSpawner* FBlueprintActionMenuItemFactory::MakePropertyActionProxy(UBlueprintPropertyNodeSpawner* PropertyNodeSpawner)
 {
 	UBlueprintPropertyNodeSpawner* ActionProxy = nullptr;
 	
@@ -208,7 +207,7 @@ UBlueprintPropertyNodeSpawner* FBlueprintActionMenuItemFactory::MakePropertyActi
 	}
 	else
 	{
-		ActionProxy = DuplicateObject(PropertyNodeSpawner, GetOwnerOfTemporaries(ItemOwner));
+		ActionProxy = DuplicateObject(PropertyNodeSpawner, GetTransientPackage());
 		// we use a null node class to flag the proxy, FBlueprintActionMenuItem
 		// queues off of this and will provide a drag-drop menu to the user
 		ActionProxy->NodeClass = nullptr;
@@ -246,7 +245,7 @@ void FBlueprintActionMenuItemFactory::Empty()
 }
 
 //------------------------------------------------------------------------------
-FText FBlueprintActionMenuItemFactory::GetMenuNameForAction(FBlueprintActionMenuBuilder& ItemOwner, UBlueprintNodeSpawner* Action)
+FText FBlueprintActionMenuItemFactory::GetMenuNameForAction(TWeakPtr<FBlueprintEditor> EditorContext, UBlueprintNodeSpawner* Action)
 {
 	check(Action != nullptr);
 	// give the action the chance to save on performance (to keep from having to spawn a template node)
@@ -254,7 +253,7 @@ FText FBlueprintActionMenuItemFactory::GetMenuNameForAction(FBlueprintActionMenu
 	
 	if (MenuName.IsEmpty())
 	{
-		if (UEdGraphNode* NodeTemplate = Action->MakeTemplateNode(GetOwnerOfTemporaries(ItemOwner)))
+		if (UEdGraphNode* NodeTemplate = GetTemplateNode(Action, EditorContext))
 		{
 			MenuName = NodeTemplate->GetNodeTitle(ENodeTitleType::ListView);
 		}
@@ -269,7 +268,7 @@ FText FBlueprintActionMenuItemFactory::GetMenuNameForAction(FBlueprintActionMenu
 }
 
 //------------------------------------------------------------------------------
-FText FBlueprintActionMenuItemFactory::GetCategoryForAction(FBlueprintActionMenuBuilder& ItemOwner, UBlueprintNodeSpawner* Action)
+FText FBlueprintActionMenuItemFactory::GetCategoryForAction(TWeakPtr<FBlueprintEditor> EditorContext, UBlueprintNodeSpawner* Action)
 {
 	check(Action != nullptr);
 	// give the action the chance to save on performance (to keep from having to spawn a template node)
@@ -277,9 +276,8 @@ FText FBlueprintActionMenuItemFactory::GetCategoryForAction(FBlueprintActionMenu
 	
 	if (MenuCategory.IsEmpty())
 	{
-		UEdGraph* OwnerOfTemporaries = GetOwnerOfTemporaries(ItemOwner);
 		// @TODO: consider moving GetMenuCategory() up into UEdGraphNode
-		if (UK2Node* NodeTemplate = Cast<UK2Node>(Action->MakeTemplateNode(OwnerOfTemporaries)))
+		if (UK2Node* NodeTemplate = Cast<UK2Node>(GetTemplateNode(Action, EditorContext)))
 		{
 			MenuCategory = NodeTemplate->GetMenuCategory();
 		}
@@ -296,7 +294,7 @@ FText FBlueprintActionMenuItemFactory::GetCategoryForAction(FBlueprintActionMenu
 				check(Function != nullptr);
 				UClass* FuncOwner = Function->GetOuterUClass();
 				
-				UBlueprint* Blueprint = GetBlueprint();
+				UBlueprint* Blueprint = GetTargetBlueprint();
 				check(Blueprint != nullptr);
 				
 				// if this is NOT a self function call (self function calls
@@ -314,7 +312,7 @@ FText FBlueprintActionMenuItemFactory::GetCategoryForAction(FBlueprintActionMenu
 }
 
 //------------------------------------------------------------------------------
-FText FBlueprintActionMenuItemFactory::GetTooltipForAction(FBlueprintActionMenuBuilder& ItemOwner, UBlueprintNodeSpawner* Action)
+FText FBlueprintActionMenuItemFactory::GetTooltipForAction(TWeakPtr<FBlueprintEditor> EditorContext, UBlueprintNodeSpawner* Action)
 {
 	check(Action != nullptr);
 	// give the action the chance to save on performance (to keep from having to spawn a template node)
@@ -322,7 +320,7 @@ FText FBlueprintActionMenuItemFactory::GetTooltipForAction(FBlueprintActionMenuB
 	
 	if (Tooltip.IsEmpty())
 	{
-		if (UEdGraphNode* NodeTemplate = Action->MakeTemplateNode(GetOwnerOfTemporaries(ItemOwner)))
+		if (UEdGraphNode* NodeTemplate = GetTemplateNode(Action, EditorContext))
 		{
 			Tooltip = FText::FromString(NodeTemplate->GetTooltip());
 		}
@@ -332,7 +330,7 @@ FText FBlueprintActionMenuItemFactory::GetTooltipForAction(FBlueprintActionMenuB
 }
 
 //------------------------------------------------------------------------------
-FText FBlueprintActionMenuItemFactory::GetSearchKeywordsForAction(FBlueprintActionMenuBuilder& ItemOwner, UBlueprintNodeSpawner* Action)
+FText FBlueprintActionMenuItemFactory::GetSearchKeywordsForAction(TWeakPtr<FBlueprintEditor> EditorContext, UBlueprintNodeSpawner* Action)
 {
 	check(Action != nullptr);
 	// give the action the chance to save on performance (to keep from having to spawn a template node)
@@ -340,7 +338,7 @@ FText FBlueprintActionMenuItemFactory::GetSearchKeywordsForAction(FBlueprintActi
 	
 	if (SearchKeywords.IsEmpty())
 	{
-		if (UEdGraphNode* NodeTemplate = Action->MakeTemplateNode(GetOwnerOfTemporaries(ItemOwner)))
+		if (UEdGraphNode* NodeTemplate = GetTemplateNode(Action, EditorContext))
 		{
 			SearchKeywords = FText::FromString(NodeTemplate->GetKeywords());
 		}
@@ -350,12 +348,12 @@ FText FBlueprintActionMenuItemFactory::GetSearchKeywordsForAction(FBlueprintActi
 }
 
 //------------------------------------------------------------------------------
-FName FBlueprintActionMenuItemFactory::GetMenuIconForAction(FBlueprintActionMenuBuilder& ItemOwner, UBlueprintNodeSpawner* Action, FLinearColor& ColorOut)
+FName FBlueprintActionMenuItemFactory::GetMenuIconForAction(TWeakPtr<FBlueprintEditor> EditorContext, UBlueprintNodeSpawner* Action, FLinearColor& ColorOut)
 {
 	FName BrushName = Action->GetDefaultMenuIcon(ColorOut);
 	if (BrushName.IsNone())
 	{
-		if (UEdGraphNode* NodeTemplate = Action->MakeTemplateNode(GetOwnerOfTemporaries(ItemOwner)))
+		if (UEdGraphNode* NodeTemplate = GetTemplateNode(Action, EditorContext))
 		{
 			BrushName = NodeTemplate->GetPaletteIcon(ColorOut);
 		}
@@ -364,41 +362,41 @@ FName FBlueprintActionMenuItemFactory::GetMenuIconForAction(FBlueprintActionMenu
 }
 
 //------------------------------------------------------------------------------
-UBlueprint* FBlueprintActionMenuItemFactory::GetBlueprint() const
+UBlueprint* FBlueprintActionMenuItemFactory::GetTargetBlueprint() const
 {
-	UBlueprint* Blueprint = nullptr;
+	UBlueprint* TargetBlueprint = nullptr;
 	if (Context.Blueprints.Num() > 0)
 	{
-		Blueprint = Context.Blueprints[0];
+		TargetBlueprint = Context.Blueprints[0];
 	}
-	return Blueprint;
+	return TargetBlueprint;
 }
 
 //------------------------------------------------------------------------------
-UEdGraph* FBlueprintActionMenuItemFactory::GetOwnerOfTemporaries(FBlueprintActionMenuBuilder& ItemOwner) const
+UEdGraphNode* FBlueprintActionMenuItemFactory::GetTemplateNode(UBlueprintNodeSpawner* Action, TWeakPtr<FBlueprintEditor> EditorContext) const
 {
-	UBlueprint* Blueprint = GetBlueprint();
-	check(Blueprint != nullptr);
-	
-	UEdGraph*& OwnerOfTemporaries = ItemOwner.OwnerOfTemporaries;
-	if (OwnerOfTemporaries == nullptr)
+	UEdGraph* TargetGraph = nullptr;
+	if (Context.Graphs.Num() > 0)
+	{
+		TargetGraph = Context.Graphs[0];
+	}
+	else
 	{
-		// unfortunately the OwnerOfTemporaries cannont be directly assigned to
-		// the transient package because various nodes relay on the blueprint
-		// outer (just like they rely on this graph outer); for example, the
-		// call-function node uses the blueprint to see if the function has
-		// self-context
-		OwnerOfTemporaries = NewObject<UEdGraph>(Blueprint);
-		OwnerOfTemporaries->SetFlags(RF_Transient);
+		UBlueprint* Blueprint = GetTargetBlueprint();
+		check(Blueprint != nullptr);
 		
-		OwnerOfTemporaries->Schema = UEdGraphSchema_K2::StaticClass();
-		if (Context.Graphs.Num() > 0)
+		if (Blueprint->UbergraphPages.Num() > 0)
 		{
-			OwnerOfTemporaries->Schema = Context.Graphs[0]->Schema;
+			TargetGraph = Blueprint->UbergraphPages[0];
+		}
+		else if (EditorContext.IsValid())
+		{
+			TargetGraph = EditorContext.Pin()->GetFocusedGraph();
 		}
-		
 	}
-	return OwnerOfTemporaries;
+
+	check(Action != nullptr);
+	return Action->GetTemplateNode(TargetGraph);
 }
 
 /*******************************************************************************
@@ -427,11 +425,11 @@ namespace FBlueprintActionMenuBuilderImpl
 		 * FBlueprintActionMenuItem for the specified menu (does not add the 
 		 * item to the menu-builder itself).
 		 *
-		 * @param  ItemOwner	The menu item owner that this will be assigned to.
+		 * @param  EditorContext
 		 * @param  Action		The node-spawner that the new menu item should wrap.
 		 * @return An empty TSharedPtr if the action was filtered out, otherwise a newly allocated FBlueprintActionMenuItem.
 		 */
-		TSharedPtr<FEdGraphSchemaAction> MakeMenuItem(FBlueprintActionMenuBuilder& ItemOwner, UBlueprintNodeSpawner* Action);
+		TSharedPtr<FEdGraphSchemaAction> MakeMenuItem(TWeakPtr<FBlueprintEditor> EditorContext, UBlueprintNodeSpawner* Action);
 		
 		/**
 		 * Clears out any intermediate UBlueprintNodeSpawner actions that this
@@ -453,10 +451,11 @@ namespace FBlueprintActionMenuBuilderImpl
 	 * method implements the old way we used collect blueprint menu actions (for
 	 * both the palette and context menu).
 	 *
-	 * @param  MenuSection	The primary section for the FBlueprintActionMenuBuilder.
-	 * @param  MenuOut  	The menu builder we want all the legacy actions appended to.
+	 * @param  MenuSection		The primary section for the FBlueprintActionMenuBuilder.
+	 * @param  BlueprintEditor	
+	 * @param  MenuOut  		The menu builder we want all the legacy actions appended to.
 	 */
-	static void AppendLegacyItems(FMenuSectionDefinition const& MenuDef, FBlueprintActionMenuBuilder& MenuOut);
+	static void AppendLegacyItems(FMenuSectionDefinition const& MenuDef, TWeakPtr<FBlueprintEditor> BlueprintEditor, FBlueprintActionMenuBuilder& MenuOut);
 }
 
 //------------------------------------------------------------------------------
@@ -486,7 +485,7 @@ void FBlueprintActionMenuBuilderImpl::FMenuSectionDefinition::SetSectionSortOrde
 }
 
 //------------------------------------------------------------------------------
-TSharedPtr<FEdGraphSchemaAction> FBlueprintActionMenuBuilderImpl::FMenuSectionDefinition::MakeMenuItem(FBlueprintActionMenuBuilder& ItemOwner, UBlueprintNodeSpawner* Action)
+TSharedPtr<FEdGraphSchemaAction> FBlueprintActionMenuBuilderImpl::FMenuSectionDefinition::MakeMenuItem(TWeakPtr<FBlueprintEditor> EditorContext, UBlueprintNodeSpawner* Action)
 {
 	if (Flags & FBlueprintActionMenuBuilder::ConsolidatePropertyActions)
 	{
@@ -504,7 +503,7 @@ TSharedPtr<FEdGraphSchemaAction> FBlueprintActionMenuBuilderImpl::FMenuSectionDe
 				// make a proxy action (that won't spawn anything)...
 				// FBlueprintActionMenuItem should interpret this and offer the
 				// user a drag-drop menu (with a list of multiple property nodes)
-				Action = ItemFactory.MakePropertyActionProxy(ItemOwner, PropertySpawner);
+				Action = ItemFactory.MakePropertyActionProxy(PropertySpawner);
 			}
 		}
 	}
@@ -512,7 +511,7 @@ TSharedPtr<FEdGraphSchemaAction> FBlueprintActionMenuBuilderImpl::FMenuSectionDe
 	TSharedPtr<FEdGraphSchemaAction> MenuEntry = nullptr;
 	if ((Action != nullptr) && !Filter.IsFiltered(Action))
 	{
-		MenuEntry = ItemFactory.MakeMenuItem(ItemOwner, Action);
+		MenuEntry = ItemFactory.MakeMenuItem(EditorContext, Action);
 	}
 	return MenuEntry;
 }
@@ -524,7 +523,7 @@ void FBlueprintActionMenuBuilderImpl::FMenuSectionDefinition::Empty()
 }
 
 //------------------------------------------------------------------------------
-void FBlueprintActionMenuBuilderImpl::AppendLegacyItems(FMenuSectionDefinition const& MenuDef, FBlueprintActionMenuBuilder& MenuOut)
+void FBlueprintActionMenuBuilderImpl::AppendLegacyItems(FMenuSectionDefinition const& MenuDef, TWeakPtr<FBlueprintEditor> BlueprintEditor, FBlueprintActionMenuBuilder& MenuOut)
 {
 	FBlueprintActionFilter const&  MenuFilter  = MenuDef.Filter;
 	FBlueprintActionContext const& MenuContext = MenuFilter.Context;
@@ -540,22 +539,14 @@ void FBlueprintActionMenuBuilderImpl::AppendLegacyItems(FMenuSectionDefinition c
 		{
 			LegacyBuilder.FromPin = MenuContext.Pins[0];
 		}
+		LegacyBuilder.SelectedObjects.Append(MenuContext.SelectedObjects);
+		
 		
-		UBlueprint* Blueprint = FBlueprintEditorUtils::FindBlueprintForGraphChecked(Graph);
-		// not the best way to poll for Selected properties, but since this
-		// function is only temporary until we switch over to the new system...
-		FBlueprintEditor* EditorInst = (FBlueprintEditor*)FAssetEditorManager::Get().FindEditorForAsset(Blueprint, /*bFocusIfOpen =*/false);
 		
 		bool bIsContextSensitive = true;
-		if (EditorInst != nullptr)
+		if (BlueprintEditor.IsValid())
 		{
-			bIsContextSensitive = EditorInst->GetIsContextSensitive();
-			
-			FEdGraphSchemaAction_K2Var* SelectedVar = EditorInst->GetMyBlueprintWidget()->SelectionAsVar();
-			if ((SelectedVar != nullptr) && (SelectedVar->GetProperty() != nullptr))
-			{
-				LegacyBuilder.SelectedObjects.Add(SelectedVar->GetProperty());
-			}
+			bIsContextSensitive = BlueprintEditor.Pin()->GetIsContextSensitive();
 		}
 		
 		if (bIsContextSensitive)
@@ -565,6 +556,7 @@ void FBlueprintActionMenuBuilderImpl::AppendLegacyItems(FMenuSectionDefinition c
 		}
 		else
 		{
+			UBlueprint* Blueprint = FBlueprintEditorUtils::FindBlueprintForGraphChecked(Graph);
 			FBlueprintPaletteListBuilder ContextlessLegacyBuilder(Blueprint);
 			UEdGraphSchema_K2::GetAllActions(ContextlessLegacyBuilder);
 			MenuOut.Append(ContextlessLegacyBuilder);
@@ -593,13 +585,9 @@ void FBlueprintActionMenuBuilderImpl::AppendLegacyItems(FMenuSectionDefinition c
  ******************************************************************************/
 
 //------------------------------------------------------------------------------
-FBlueprintActionMenuBuilder::FBlueprintActionMenuBuilder(FBlueprintActionFilter const& Filter, bool bAutoBuild/*= true*/)
+FBlueprintActionMenuBuilder::FBlueprintActionMenuBuilder(TWeakPtr<FBlueprintEditor> InBlueprintEditorPtr)
+	: BlueprintEditorPtr(InBlueprintEditorPtr)
 {
-	AddMenuSection(Filter);
-	if (bAutoBuild)
-	{
-		RebuildActionList();
-	}
 }
 
 //------------------------------------------------------------------------------
@@ -644,7 +632,7 @@ void FBlueprintActionMenuBuilder::RebuildActionList()
 			{
 				for (TSharedRef<FMenuSectionDefinition> MenuSection : MenuSections)
 				{
-					TSharedPtr<FEdGraphSchemaAction> MenuEntry = MenuSection->MakeMenuItem(*this, NodeSpawner);
+					TSharedPtr<FEdGraphSchemaAction> MenuEntry = MenuSection->MakeMenuItem(BlueprintEditorPtr, NodeSpawner);
 					if (MenuEntry.IsValid())
 					{
 						AddAction(MenuEntry);
@@ -655,7 +643,7 @@ void FBlueprintActionMenuBuilder::RebuildActionList()
 	}
 	else if (MenuSections.Num() > 0)
 	{
-		AppendLegacyItems(*MenuSections[0], *this);
+		AppendLegacyItems(*MenuSections[0], BlueprintEditorPtr, *this);
 	}
 
 	// @TODO: account for all K2ActionMenuBuilder action types...
diff --git a/Engine/Source/Editor/Kismet/Private/SBlueprintActionMenu.cpp b/Engine/Source/Editor/Kismet/Private/SBlueprintActionMenu.cpp
--- a/Engine/Source/Editor/Kismet/Private/SBlueprintActionMenu.cpp
+++ b/Engine/Source/Editor/Kismet/Private/SBlueprintActionMenu.cpp
@@ -412,7 +412,7 @@ void SBlueprintActionMenu::CollectAllActions(FGraphActionListBuilderBase& OutAll
 		}
 	}
 	
-	FBlueprintActionMenuBuilder MenuBuilder;
+	FBlueprintActionMenuBuilder MenuBuilder(EditorPtr);
 	// NOTE: cannot call GetGraphContextActions() during serialization and GC due to its use of FindObject()
 	if(!GIsSavingPackage && !GIsGarbageCollecting)
 	{
diff --git a/Engine/Source/Editor/Kismet/Private/SBlueprintLibraryPalette.cpp b/Engine/Source/Editor/Kismet/Private/SBlueprintLibraryPalette.cpp
--- a/Engine/Source/Editor/Kismet/Private/SBlueprintLibraryPalette.cpp
+++ b/Engine/Source/Editor/Kismet/Private/SBlueprintLibraryPalette.cpp
@@ -218,7 +218,7 @@ void SBlueprintLibraryPalette::CollectAllActions(FGraphActionListBuilderBase& Ou
 		ClassFilter = FilterClass.Get();
 	}
 	
-	FBlueprintActionMenuBuilder PaletteBuilder;
+	FBlueprintActionMenuBuilder PaletteBuilder(BlueprintEditorPtr);
 	FBlueprintActionMenuUtils::MakePaletteMenu(FilterContext, ClassFilter, PaletteBuilder);
 	
 	PaletteBuilder.RebuildActionList();
diff --git a/Engine/Source/Editor/Kismet/Public/BlueprintActionMenuBuilder.h b/Engine/Source/Editor/Kismet/Public/BlueprintActionMenuBuilder.h
--- a/Engine/Source/Editor/Kismet/Public/BlueprintActionMenuBuilder.h
+++ b/Engine/Source/Editor/Kismet/Public/BlueprintActionMenuBuilder.h
@@ -12,6 +12,7 @@ namespace FBlueprintActionMenuBuilderImpl
 };
 class  FBlueprintActionFilter;
 struct FBlueprintActionContext;
+class  FBlueprintEditor;
 
 /**
  * Responsible for constructing a list of viable blueprint actions. Runs the 
@@ -31,18 +32,7 @@ public:
 	};
 	
 public:
-	/** Default constructor */
-	FBlueprintActionMenuBuilder() {}
-
-	/**
-	 * Good for constructing a menu built from a single filter. If the 
-	 * bAutoBuild param is set to false, then know you have to manually call
-	 * RebuildActionList() to generate the menu.
-	 *
-	 * @param  Filter		The filter you want applied for this main section of the menu.
-	 * @param  bAutoBuild	If set, will automatically call RebuildActionList() here in the constructor.
-	 */
-	FBlueprintActionMenuBuilder(FBlueprintActionFilter const& Filter, bool bAutoBuild = true);
+	FBlueprintActionMenuBuilder(TWeakPtr<FBlueprintEditor> BlueprintEditorPtr);
 	
 	// FGraphActionListBuilderBase interface
 	virtual void Empty() override;
@@ -76,4 +66,7 @@ private:
 	 * this API clean).
 	 */
 	TArray< TSharedRef<FBlueprintActionMenuBuilderImpl::FMenuSectionDefinition> > MenuSections;
+
+	/** */
+	TWeakPtr<FBlueprintEditor> BlueprintEditorPtr;
 };
\ No newline at end of file
diff --git a/Engine/Source/Editor/UnrealEd/Private/Commandlets/DumpBlueprintsInfoCommandlet.cpp b/Engine/Source/Editor/UnrealEd/Private/Commandlets/DumpBlueprintsInfoCommandlet.cpp
--- a/Engine/Source/Editor/UnrealEd/Private/Commandlets/DumpBlueprintsInfoCommandlet.cpp
+++ b/Engine/Source/Editor/UnrealEd/Private/Commandlets/DumpBlueprintsInfoCommandlet.cpp
@@ -20,6 +20,7 @@
 #include "BlueprintActionFilter.h"
 #include "BlueprintActionDatabase.h"
 #include "BlueprintActionMenuUtils.h"
+#include "K2Node_Composite.h"
 
 DEFINE_LOG_CATEGORY_STATIC(LogBlueprintInfoDump, Log, All);
 
@@ -986,10 +987,10 @@ static double DumpBlueprintInfoUtils::GetPaletteMenuActions(FBlueprintPaletteLis
 		FBlueprintActionContext FilterContext;
 		FilterContext.Blueprints.Add(const_cast<UBlueprint*>(PaletteBuilder.Blueprint));
 		
-		FBlueprintActionMenuBuilder MenuBuilder;
+		FBlueprintActionMenuBuilder MenuBuilder(nullptr);
 		{
 			// prime the database so it's not recorded in our timing capture
-			FBlueprintActionDatabase::Prime();
+			FBlueprintActionDatabase::Get();
 			
 			FScopedDurationTimer DurationTimer(MenuBuildDuration);
 			FBlueprintActionMenuUtils::MakePaletteMenu(FilterContext, PaletteFilter, MenuBuilder);
@@ -1462,11 +1463,16 @@ static void DumpBlueprintInfoUtils::DumpContextualPinTypeActions(uint32 Indent,
 {
 	FBlueprintGraphActionListBuilder ContextMenuBuilder(Graph);
 
-	UEdGraphPin* DummyPin = NewObject<UEdGraphPin>(GetTransientPackage());
-	DummyPin->PinName = DummyPin->GetName();
-	DummyPin->Direction = EGPD_Input;
-	DummyPin->PinType = PinType;
-
+	UK2Node_Composite* DummyNode = NewObject<UK2Node_Composite>(Graph);
+	UEdGraphPin* DummyPin = DummyNode->CreatePin(
+		EGPD_Input,
+		PinType.PinCategory,
+		PinType.PinSubCategory,
+		PinType.PinSubCategoryObject.Get(),
+		PinType.bIsArray,
+		PinType.bIsReference,
+		DummyNode->GetName()
+	);
 	ContextMenuBuilder.FromPin = DummyPin;
 
 	DumpContextActionList(Indent, ContextMenuBuilder, FileOutWriter);
@@ -1549,11 +1555,11 @@ static double DumpBlueprintInfoUtils::GetContextMenuActions(FBlueprintGraphActio
 			}
 		}
 		
-		FBlueprintActionMenuBuilder MenuBuilder;
+		FBlueprintActionMenuBuilder MenuBuilder(nullptr);
 		{
 			// prime the database so it's not recorded in our timing capture
-			FBlueprintActionDatabase::Prime();
-			
+			FBlueprintActionDatabase::Get();
+
 			FScopedDurationTimer DurationTimer(MenuBuildDuration);
 			FBlueprintActionMenuUtils::MakeContextMenu(FilterContext, /*bIsContextSensitive =*/true, MenuBuilder);
 		}
