No longer constraining blueprint components visualizers to scene-components (can be actor components)... stemmed from a UDN request

diff --git a/Engine/Source/Editor/Kismet/Private/SCSEditorViewportClient.cpp b/Engine/Source/Editor/Kismet/Private/SCSEditorViewportClient.cpp
--- a/Engine/Source/Editor/Kismet/Private/SCSEditorViewportClient.cpp
+++ b/Engine/Source/Editor/Kismet/Private/SCSEditorViewportClient.cpp
@@ -256,7 +256,7 @@ void FSCSEditorViewportClient::Draw(const FSceneView* View, FPrimitiveDrawInterf
 			{
 				FSCSEditorTreeNodePtrType SelectedNode = SelectedNodes[SelectionIndex];
 
-				UActorComponent* Comp = Cast<USceneComponent>(SelectedNode->FindComponentInstanceInActor(PreviewActor, true));
+				UActorComponent* Comp = SelectedNode->FindComponentInstanceInActor(PreviewActor, true);
 				if(Comp != NULL && Comp->IsRegistered())
 				{
 					// Try and find a visualizer
