Fixing the category that member functions are presented under (for the new BP action menu system).

diff --git a/Engine/Source/Editor/Kismet/Private/BlueprintActionMenuBuilder.cpp b/Engine/Source/Editor/Kismet/Private/BlueprintActionMenuBuilder.cpp
--- a/Engine/Source/Editor/Kismet/Private/BlueprintActionMenuBuilder.cpp
+++ b/Engine/Source/Editor/Kismet/Private/BlueprintActionMenuBuilder.cpp
@@ -296,12 +296,13 @@ FText FBlueprintActionMenuItemFactory::GetCategoryForAction(TWeakPtr<FBlueprintE
 				
 				UBlueprint* Blueprint = GetTargetBlueprint();
 				check(Blueprint != nullptr);
+				UClass* BlueprintClass = (Blueprint->SkeletonGeneratedClass != nullptr) ? Blueprint->SkeletonGeneratedClass : Blueprint->ParentClass;
 				
 				// if this is NOT a self function call (self function calls
 				// don't get nested any deeper)
-				if (!Blueprint->GeneratedClass->IsChildOf(FuncOwner))
+				if (!BlueprintClass->IsChildOf(FuncOwner))
 				{
-					MenuCategory = FText::FromName(FuncOwner->GetFName());
+					MenuCategory = FuncOwner->GetDisplayNameText();
 				}
 				MenuCategory = FText::Format(LOCTEXT("MemberFunctionsCategory", "Call Function|{0}"), MenuCategory);
 			}
@@ -507,7 +508,7 @@ TSharedPtr<FEdGraphSchemaAction> FBlueprintActionMenuBuilderImpl::FMenuSectionDe
 			}
 		}
 	}
-	
+
 	TSharedPtr<FEdGraphSchemaAction> MenuEntry = nullptr;
 	if ((Action != nullptr) && !Filter.IsFiltered(Action))
 	{
