#UE4 Improved view of CDO in diff view, can now cycle through the differences, toggle filtering of properties that differ

diff --git a/Engine/Source/Editor/EditorStyle/Private/SlateEditorStyle.cpp b/Engine/Source/Editor/EditorStyle/Private/SlateEditorStyle.cpp
--- a/Engine/Source/Editor/EditorStyle/Private/SlateEditorStyle.cpp
+++ b/Engine/Source/Editor/EditorStyle/Private/SlateEditorStyle.cpp
@@ -2593,6 +2593,7 @@ void FSlateEditorStyle::FStyle::SetupPropertyEditorStyles()
 		Set( "DetailsView.CategoryBottom", new BOX_BRUSH( "PropertyView/DetailCategoryBottom", FMargin(4.0f/16.0f) ) );
 		Set( "DetailsView.CategoryMiddle", new IMAGE_BRUSH( "PropertyView/DetailCategoryMiddle", FVector2D( 16, 16 ) ) );
 		Set( "DetailsView.CategoryMiddle_Hovered", new IMAGE_BRUSH( "PropertyView/DetailCategoryMiddle_Hovered", FVector2D( 16, 16 ) ) );
+		Set( "DetailsView.CategoryMiddle_Highlighted", new BOX_BRUSH( "Common/TextBox_Special_Active", FMargin(8.0f/32.0f) ) );
 
 		Set( "DetailsView.Splitter", FSplitterStyle()
 			.SetHandleNormalBrush( IMAGE_BRUSH( "Common/SplitterHandleHighlight", Icon8x8, FLinearColor::Black ) )
diff --git a/Engine/Source/Editor/Kismet/Private/SBlueprintDiff.cpp b/Engine/Source/Editor/Kismet/Private/SBlueprintDiff.cpp
--- a/Engine/Source/Editor/Kismet/Private/SBlueprintDiff.cpp
+++ b/Engine/Source/Editor/Kismet/Private/SBlueprintDiff.cpp
@@ -21,51 +21,225 @@
 
 #define LOCTEXT_NAMESPACE "SBlueprintDif"
 
-struct FPropertyFilter : public TSharedFromThis<FPropertyFilter>
+class IDiffHighlighter
 {
-	FPropertyFilter( TSet< const UProperty*> InHiddenProperties )
-		: HiddenProperties( InHiddenProperties )
+public:
+	virtual void NextDiff() = 0;
+	virtual void PrevDiff() = 0;
+	virtual bool HasDifferences() const = 0;
+};
+
+typedef TMap< FString, const UProperty* > FNamePropertyMap;
+
+class CDODiffHighlighter : public TSharedFromThis<CDODiffHighlighter>
+						, public IDiffHighlighter
+{
+public:
+	CDODiffHighlighter( 
+			TWeakPtr<IDetailsView> InOldView
+			, FNamePropertyMap InOldProperties
+			, TWeakPtr<IDetailsView> InNewView
+			, FNamePropertyMap InNewProperties
+			, TArray< FString > InDifferingProperties )
+		: OldView( InOldView )
+		, OldProperties( InOldProperties )
+		, NewView(InNewView)
+		, NewProperties(InNewProperties)
+		, DifferingProperties( InDifferingProperties )
+		, CurrentDifference( -1 )
 	{
 	}
 
-	~FPropertyFilter()
+private:
+	void NextDiff() override;
+	void PrevDiff() override;
+	bool HasDifferences() const override;
+
+	void HighlightCurrentDifference();
+
+	TWeakPtr<IDetailsView> OldView;
+	FNamePropertyMap OldProperties;
+	TWeakPtr<IDetailsView> NewView;
+	FNamePropertyMap NewProperties;
+	TArray< FString > DifferingProperties;
+	int CurrentDifference;
+};
+
+void CDODiffHighlighter::NextDiff()
+{
+	if( DifferingProperties.Num() == 0 )
 	{
-		UE_LOG(LogBlueprintUserMessages, Verbose, TEXT(" Nuking property filter.."));
+		return;
 	}
 
-	bool IsVisible(const FPropertyAndParent& Property)
+	CurrentDifference = (CurrentDifference + 1) % DifferingProperties.Num();
+
+	HighlightCurrentDifference();
+}
+
+void CDODiffHighlighter::PrevDiff()
+{
+	if (DifferingProperties.Num() == 0)
 	{
-		return !HiddenProperties.Contains(&Property.Property);
+		return;
 	}
 
-	TSet< const UProperty*> HiddenProperties;
-};
+	--CurrentDifference;
+	if( CurrentDifference < 0 )
+	{
+		CurrentDifference = DifferingProperties.Num() - 1;
+	}
 
-typedef TMap< FString, const UProperty* > FNamePropertyMap;
+	HighlightCurrentDifference();
+}
 
-static FNamePropertyMap GetCDOProperties( const UBlueprint* ForBlueprint )
+bool CDODiffHighlighter::HasDifferences() const
+{
+	return DifferingProperties.Num() != 0;
+}
+
+void CDODiffHighlighter::HighlightCurrentDifference()
+{
+	const FString& PropertyName = DifferingProperties[CurrentDifference];
+
+	const UProperty** OldProperty = OldProperties.Find(PropertyName);
+	const UProperty** NewProperty = NewProperties.Find(PropertyName);
+
+	OldView.Pin()->HighlightProperty(OldProperty ? *OldProperty : NULL);
+	NewView.Pin()->HighlightProperty(NewProperty ? *NewProperty : NULL);
+}
+
+/*List item that entry for a graph*/
+struct KISMET_API FListItemGraphToDiff	: public TSharedFromThis<FListItemGraphToDiff>
+										, public IDiffHighlighter
+{
+	FListItemGraphToDiff(class SBlueprintDiff* Diff, class UEdGraph* GraphOld, class UEdGraph* GraphNew, const FRevisionInfo& RevisionOld, const FRevisionInfo& RevisionNew);
+	~FListItemGraphToDiff();
+
+	/*Generate Widget for list item*/
+	TSharedRef<SWidget> GenerateWidget();
+
+	/*Get tooltip for list item */
+	FText   GetToolTip();
+
+	/*Get old(left) graph*/
+	UEdGraph* GetGraphOld()const{ return GraphOld; }
+
+	/*Get new(right) graph*/
+	UEdGraph* GetGraphNew()const{ return GraphNew; }
+
+private:
+
+	/*Get icon to use by graph name*/
+	static const FSlateBrush* GetIconForGraph(UEdGraph* Graph);
+
+	/*Diff widget*/
+	class SBlueprintDiff* Diff;
+
+	/*The old graph(left)*/
+	class UEdGraph*	GraphOld;
+
+	/*The new graph(right)*/
+	class UEdGraph* GraphNew;
+
+	/*Description of Old and new graph*/
+	FRevisionInfo	RevisionOld, RevisionNew;
+
+	//////////////////////////////////////////////////////////////////////////
+	// Diff list
+	//////////////////////////////////////////////////////////////////////////
+
+	typedef TSharedPtr<struct FDiffResultItem>	FSharedDiffOnGraph;
+	typedef SListView<FSharedDiffOnGraph >		SListViewType;
+
+public:
+
+	/** Called when the Newer Graph is modified*/
+	void OnGraphChanged(const FEdGraphEditAction& Action);
+
+	/** Generate list of differences*/
+	TSharedRef<SWidget> GenerateDiffListWidget();
+
+	/** Build up the Diff Source Array*/
+	void BuildDiffSourceArray();
+
+	/** Called when user clicks on a new graph list item */
+	void OnSelectionChanged(FSharedDiffOnGraph Item, ESelectInfo::Type SelectionType);
+
+	/** Called when user presses key within the diff view */
+	void KeyWasPressed(const FKeyboardEvent& InKeyboardEvent);
+
+private:
+
+	/** Called when user clicks button to go to next difference in graph */
+	void NextDiff() override;
+
+	/** Called when user clicks button to go to prev difference in graph */
+	void PrevDiff() override;
+
+	/** Called when determining whether the prev/nextbuttons should be enabled */
+	bool HasDifferences() const override;
+
+	/** Get Index of the current diff that is selected */
+	int32 GetCurrentDiffIndex();
+
+	/* Called when a new row is being generated */
+	TSharedRef<ITableRow> OnGenerateRow(FSharedDiffOnGraph ParamItem, const TSharedRef<STableViewBase>& OwnerTable);
+
+	/** ListView of differences */
+	TSharedPtr<SListViewType> DiffList;
+
+	/** Source for list view */
+	TArray<FSharedDiffOnGraph> DiffListSource;
+
+	/** Key commands processed by this widget */
+	TSharedPtr< FUICommandList > KeyCommands;
+};
+
+static FNamePropertyMap GetProperties( const UObject& ForObj )
 {
 	FNamePropertyMap Ret;
+	if (const UClass* Class = ForObj.GetClass())
+	{
+		for (TFieldIterator<UProperty> PropertyIt(Class, EFieldIteratorFlags::IncludeSuper); PropertyIt; ++PropertyIt)
+		{
+			if ((!PropertyIt->HasAnyPropertyFlags(CPF_Parm) && PropertyIt->HasAnyPropertyFlags(CPF_Edit)))
+			{
+				Ret.Add(PropertyIt->GetName(), *PropertyIt);
+			}
+		}
+	}
+	return Ret;
+}
 
+static FNamePropertyMap GetCDOProperties( const UBlueprint* ForBlueprint )
+{
 	if( ForBlueprint )
 	{
 		if( const UClass* GeneratedClass = ForBlueprint->GeneratedClass )
 		{
 			if( const UObject* CDO = GeneratedClass->ClassDefaultObject )
 			{
-				if( const UClass* CDOClass = CDO->GetClass() )
-				{
-					check( CDOClass == GeneratedClass );
-					for (TFieldIterator<UProperty> PropertyIt(CDOClass, EFieldIteratorFlags::IncludeSuper); PropertyIt; ++PropertyIt)
-					{
-						Ret.Add(PropertyIt->GetName(), *PropertyIt);
-					}
-				}
+				return GetProperties(*CDO);
 			}
 		}
 	}
 
-	return Ret;
+	return FNamePropertyMap();
+}
+
+static void GetPropertyNameSet( const UObject& ForObj, TSet< FString >& OutProperties )
+{
+	if( const UClass* Class = ForObj.GetClass() )
+	{
+		for (TFieldIterator<UProperty> PropertyIt(Class, EFieldIteratorFlags::IncludeSuper); PropertyIt; ++PropertyIt)
+		{
+			if ((!PropertyIt->HasAnyPropertyFlags(CPF_Parm) && PropertyIt->HasAnyPropertyFlags(CPF_Edit)))
+			{
+				OutProperties.Add(PropertyIt->GetName());
+			}
+		}
+	}
 }
 
 const UObject* GetCDO( const UBlueprint* ForBlueprint )
@@ -210,16 +384,6 @@ TSharedRef<SWidget> FListItemGraphToDiff::GenerateDiffListWidget()
 		
 		const FDiffListCommands& Commands = FDiffListCommands::Get();
 
-		// Map commands through UI
-		KeyCommands = MakeShareable(new FUICommandList );
-
-		KeyCommands->MapAction( Commands.Previous, FExecuteAction::CreateSP(this, &FListItemGraphToDiff::PrevDiff));
-		KeyCommands->MapAction( Commands.Next, FExecuteAction::CreateSP(this, &FListItemGraphToDiff::NextDiff));
-
-		FToolBarBuilder ToolbarBuilder( KeyCommands.ToSharedRef(), FMultiBoxCustomization::None );
-		ToolbarBuilder.AddToolBarButton(Commands.Previous, NAME_None, TAttribute<FText>(), TAttribute<FText>(), FSlateIcon(FEditorStyle::GetStyleSetName(), "BlueprintDif.PrevDiff"));
-		ToolbarBuilder.AddToolBarButton(Commands.Next, NAME_None, TAttribute<FText>(), TAttribute<FText>(), FSlateIcon(FEditorStyle::GetStyleSetName(), "BlueprintDif.NextDiff"));
-
 		TSharedPtr<SListViewType> DiffListRef;
 		TSharedRef<SHorizontalBox> Result =	SNew(SHorizontalBox)
 			+SHorizontalBox::Slot()
@@ -230,12 +394,6 @@ TSharedRef<SWidget> FListItemGraphToDiff::GenerateDiffListWidget()
 				+SVerticalBox::Slot()
 				.Padding(0.f)
 				.AutoHeight()
-				[
-					ToolbarBuilder.MakeWidget()
-				]
-				+SVerticalBox::Slot()
-				.Padding(0.f)
-				.AutoHeight()
 				[
 					SNew(SBorder)
 					.BorderImage(FEditorStyle::GetBrush("PropertyWindow.CategoryBackground"))
@@ -355,17 +513,25 @@ void FListItemGraphToDiff::PrevDiff()
 	DiffList->SetSelection(DiffListSource[Index]);
 }
 
+bool FListItemGraphToDiff::HasDifferences() const
+{
+	return DiffList.IsValid() && DiffListSource.Num() > 0;
+}
+
 int32 FListItemGraphToDiff::GetCurrentDiffIndex() 
 {
-	auto Selected = DiffList->GetSelectedItems();
-	if(Selected.Num() == 1)
-	{	
-		int32 Index = 0;
-		for(auto It(DiffListSource.CreateIterator());It;++It,Index++)
-		{
-			if(*It == Selected[0])
+	if ( DiffList.IsValid() )
+	{
+		auto Selected = DiffList->GetSelectedItems();
+		if(Selected.Num() == 1)
+		{	
+			int32 Index = 0;
+			for(auto It(DiffListSource.CreateIterator());It;++It,Index++)
 			{
-				return Index;
+				if(*It == Selected[0])
+				{
+					return Index;
+				}
 			}
 		}
 	}
@@ -392,65 +558,62 @@ FDiffPanel::FDiffPanel()
 	LastFocusedPin = NULL;
 }
 
-// Generates a widget showing CDO values. NewBlueprint's CDO values will be displayed by the widget (Base just used for determining what to hide):
-static TSharedRef<SWidget> GenerateCDODiff( UBlueprint const* BaseBlueprint, UBlueprint const* NewBlueprint )
+static void CompareUnrelatedObjects( UObject const* A, const FNamePropertyMap& PropertyMapA, UObject const* B, const FNamePropertyMap& PropertyMapB, TArray<FString> &OutIdenticalProperties, TArray<FString> &OutDifferingProperties )
 {
-	const auto GenerateEmptyPanel = []( const FText& Message) { return SNew(STextBlock).Text(Message); };
-
-	// filter out properties that have not changed from the GraphToDiff's blueprint:	
-	FNamePropertyMap BaseProperties = GetCDOProperties(BaseBlueprint);
-	FNamePropertyMap BlueprintProperties = GetCDOProperties(NewBlueprint);
-
-	if( BlueprintProperties.Num() == 0 )
+	TSet<FString> PropertiesInA, PropertiesInB;
+	if( A )
 	{
-		if( NewBlueprint )
-		{
-			return GenerateEmptyPanel( FText::Format( LOCTEXT("BPNoCDO", "Blueprint {0} has no Defaults"), FText::FromString( NewBlueprint->GetName() ) ) );
-		}
-		else
-		{
-			return GenerateEmptyPanel( LOCTEXT("BPMissingBP", "Missing Blueprint" ) );
-		}
+		GetPropertyNameSet(*A, PropertiesInA);
+	}
+	if( B )
+	{
+		GetPropertyNameSet(*B, PropertiesInB);
 	}
 
-	const UObject* BaseBlueprintCDO = GetCDO(BaseBlueprint);
-	const UObject* NewBlueprintCDO = GetCDO(NewBlueprint);
+	// any properties in A that aren't in B are differing:
+	OutDifferingProperties.Append( PropertiesInA.Difference(PropertiesInB).Array() );
 
-	TSet< const UProperty*> HiddenProperties;
-	for (const auto& Property : BlueprintProperties)
-	{
-		const FString& PropertyName = Property.Key;
-		const UProperty* PropertyValue = Property.Value;
+	// and the converse:
+	OutDifferingProperties.Append( PropertiesInB.Difference(PropertiesInA).Array() );
 
-		const UProperty** BaseProperty = BaseProperties.Find(PropertyName);
-		if (BaseProperty && (*BaseProperty)->SameType(PropertyValue))
+	// for properties in common, dig out the uproperties and determine if they're identical:
+	if( A && B )
+	{
+		TSet<FString> Common = PropertiesInA.Intersect(PropertiesInB);
+		for (const auto& PropertyName : Common)
 		{
-			const void* MyValue = PropertyValue->ContainerPtrToValuePtr<void>(NewBlueprintCDO);
-			const void* BaseValue = (*BaseProperty)->ContainerPtrToValuePtr<void>(BaseBlueprintCDO);
+			const UProperty* const* AProp = PropertyMapA.Find(PropertyName);
+			const UProperty* const* BProp = PropertyMapB.Find(PropertyName);
+			check( AProp && BProp );
 
-			if ((*BaseProperty)->Identical(MyValue, BaseValue, 0))
+			if( (*AProp)->SameType(*BProp) )
 			{
-				HiddenProperties.Add(PropertyValue);
+				const void* AValue = (*AProp)->ContainerPtrToValuePtr<void>(A);
+				const void* BValue = (*BProp)->ContainerPtrToValuePtr<void>(B);
+				if ((*AProp)->Identical(AValue, BValue, 0))
+				{
+					OutIdenticalProperties.Add(PropertyName);
+					continue;
+				}
 			}
-		}
-	}
 
-	if( HiddenProperties.Num() == BlueprintProperties.Num() )
-	{
-		return GenerateEmptyPanel(LOCTEXT("BPNoCDODifference", "No differences found in Defaults"));
+			OutDifferingProperties.Add(PropertyName);
+		}
 	}
+}
 
+// Generates a widget showing CDO values. NewBlueprint's CDO values will be displayed by the widget (Base just used for determining what to hide):
+static TSharedRef<IDetailsView> GenerateCDODiff( UObject const* CDO, TSet<FString> IdenticalPropertiesSet )
+{
 	FDetailsViewArgs DetailsViewArgs;
-	DetailsViewArgs.bShowOptions = false;
+	DetailsViewArgs.bShowDifferingPeropertiesOption = true;
+
 	FPropertyEditorModule& EditModule = FModuleManager::Get().GetModuleChecked<FPropertyEditorModule>("PropertyEditor");
 	auto DetailsView = EditModule.CreateDetailView(DetailsViewArgs);
 	DetailsView->SetIsPropertyEditingEnabledDelegate(FIsPropertyEditingEnabled::CreateStatic([]{return false; }));
-	// This filter goes out of scope and is destroyed when this function returns (the control just keeps a weak ptr)
-	// but that seems to be long enough to hide the properties that we want to hide:
-	auto Filter = TSharedPtr<FPropertyFilter>(new FPropertyFilter(HiddenProperties));
-	DetailsView->SetIsPropertyVisibleDelegate(FIsPropertyVisible::CreateSP(Filter.ToSharedRef(), &FPropertyFilter::IsVisible));
 	// This is a read only details view (see the property editing delegate), but it is not const correct:
-	DetailsView->SetObject(const_cast<UObject*>(NewBlueprintCDO) );
+	DetailsView->SetObject(const_cast<UObject*>(CDO) );
+	DetailsView->UpdateIdenticalProperties(IdenticalPropertiesSet);
 
 	return DetailsView;
 }
@@ -544,6 +707,16 @@ void SBlueprintDiff::Construct( const FArguments& InArgs)
 			];
 	}
 
+	const FDiffListCommands& Commands = FDiffListCommands::Get();
+	KeyCommands = MakeShareable(new FUICommandList);
+
+	KeyCommands->MapAction(Commands.Previous, FExecuteAction::CreateSP(this, &SBlueprintDiff::PrevDiff), FCanExecuteAction::CreateSP( this, &SBlueprintDiff::CanCycleDiffs) );
+	KeyCommands->MapAction(Commands.Next, FExecuteAction::CreateSP(this, &SBlueprintDiff::NextDiff), FCanExecuteAction::CreateSP( this, &SBlueprintDiff::CanCycleDiffs) );
+
+	FToolBarBuilder ToolbarBuilder(KeyCommands.ToSharedRef(), FMultiBoxCustomization::None);
+	ToolbarBuilder.AddToolBarButton(Commands.Previous, NAME_None, TAttribute<FText>(), TAttribute<FText>(), FSlateIcon(FEditorStyle::GetStyleSetName(), "BlueprintDif.PrevDiff"));
+	ToolbarBuilder.AddToolBarButton(Commands.Next, NAME_None, TAttribute<FText>(), TAttribute<FText>(), FSlateIcon(FEditorStyle::GetStyleSetName(), "BlueprintDif.NextDiff"));
+
 	this->ChildSlot
 		[
 			SNew(SBorder)
@@ -556,25 +729,31 @@ void SBlueprintDiff::Construct( const FArguments& InArgs)
 				[
 					SNew(SHorizontalBox)
 					+ SHorizontalBox::Slot()
+					.Padding(16.f)
+					.AutoWidth()
+					[
+						ToolbarBuilder.MakeWidget()
+					]
+					+ SHorizontalBox::Slot()
 					[
 						SNew(SSpacer)
 					]
 					+ SHorizontalBox::Slot()
-						.AutoWidth()
+					.AutoWidth()
+					[
+						// Nested box just give us padding similar to the main blueprint editor,
+						// one border has the dark grey "Toolbar.Background" style:
+						SNew(SHorizontalBox)
+						+ SHorizontalBox::Slot()
+						.Padding(0.0f, 4.f, 0.0f, 12.f)
 						[
-							// Nested box just give us padding similar to the main blueprint editor,
-							// one border has the dark grey "Toolbar.Background" style:
-							SNew(SHorizontalBox)
-							+ SHorizontalBox::Slot()
-							.Padding(0.0f, 4.f, 0.0f, 12.f)
+							SNew(SBorder)
+							.BorderImage(FEditorStyle::GetBrush(TEXT("Toolbar.Background")))
 							[
-								SNew(SBorder)
-								.BorderImage(FEditorStyle::GetBrush(TEXT("Toolbar.Background")))
-								[
-									MiscWidgets
-								]
+								MiscWidgets
 							]
 						]
+					]
 				]
 				+ SVerticalBox::Slot()
 				[
@@ -634,12 +813,34 @@ TSharedRef<SWidget> SBlueprintDiff::DefaultEmptyPanel()
 		];
 }
 
+void SBlueprintDiff::NextDiff()
+{
+	if( DiffHighlighter.IsValid() )
+	{
+		DiffHighlighter->NextDiff();
+	}
+}
+
+void SBlueprintDiff::PrevDiff()
+{
+	if (DiffHighlighter.IsValid())
+	{
+		DiffHighlighter->PrevDiff();
+	}
+}
+
+bool SBlueprintDiff::CanCycleDiffs() const
+{
+	return DiffHighlighter.IsValid() && DiffHighlighter->HasDifferences();
+}
+
 void SBlueprintDiff::FocusOnGraphRevisions( class UEdGraph* GraphOld, class UEdGraph* GraphNew, FListItemGraphToDiff* Diff )
 {
 	HandleGraphChanged( GraphOld ? GraphOld->GetName() : GraphNew->GetName() );
 
 	ResetGraphEditors();
 
+	DiffHighlighter = Diff->AsShared();
 	//set new diff list
 	DiffListBorder->SetContent(Diff->GenerateDiffListWidget());
 }
@@ -973,36 +1174,36 @@ TSharedRef<SWidget> SBlueprintDiff::GenerateGraphPanel()
 					]
 				]
 				+ SSplitter::Slot()
-					.Value(0.8f)
+				.Value(0.8f)
+				[
+					SNew(SHorizontalBox)
+					+ SHorizontalBox::Slot()
+					.FillWidth(1.f)
 					[
-						SNew(SHorizontalBox)
-						+ SHorizontalBox::Slot()
-						.FillWidth(1.f)
+						//diff window
+						SNew(SSplitter)
+						+ SSplitter::Slot()
+						.Value(0.5f)
 						[
-							//diff window
-							SNew(SSplitter)
-							+ SSplitter::Slot()
-							.Value(0.5f)
+							//left blueprint
+							SAssignNew(PanelOld.GraphEditorBorder, SBorder)
+							.VAlign(VAlign_Fill)
 							[
-								//left blueprint
-								SAssignNew(PanelOld.GraphEditorBorder, SBorder)
-								.VAlign(VAlign_Fill)
-								[
-									DefaultEmptyPanel()
-								]
+								DefaultEmptyPanel()
 							]
-							+ SSplitter::Slot()
-							.Value(0.5f)
+						]
+						+ SSplitter::Slot()
+						.Value(0.5f)
+						[
+							//right blueprint
+							SAssignNew(PanelNew.GraphEditorBorder, SBorder)
+							.VAlign(VAlign_Fill)
 							[
-								//right blueprint
-								SAssignNew(PanelNew.GraphEditorBorder, SBorder)
-								.VAlign(VAlign_Fill)
-								[
-									DefaultEmptyPanel()
-								]
+								DefaultEmptyPanel()
 							]
 						]
 					]
+				]
 			]
 		]
 	];
@@ -1010,6 +1211,18 @@ TSharedRef<SWidget> SBlueprintDiff::GenerateGraphPanel()
 
 TSharedRef<SWidget> SBlueprintDiff::GenerateDefaultsPanel()
 {
+	const UObject* A = GetCDO(PanelOld.Blueprint);
+	FNamePropertyMap PropertyMapA = GetProperties(*A);
+	const UObject* B = GetCDO(PanelNew.Blueprint);
+	FNamePropertyMap PropertyMapB = GetProperties(*B);
+	TArray<FString> IdenticalProperties, DifferingProperties;
+	CompareUnrelatedObjects( A, PropertyMapA, B, PropertyMapB, IdenticalProperties, DifferingProperties);
+
+	TSet<FString> IdenticalPropertiesSet(IdenticalProperties);
+	auto LeftDetailsView = GenerateCDODiff(A, IdenticalPropertiesSet );
+	auto RightDetailsView = GenerateCDODiff(B, IdenticalPropertiesSet );
+
+	DiffHighlighter = TSharedPtr<CDODiffHighlighter>(new CDODiffHighlighter(LeftDetailsView, PropertyMapA, RightDetailsView, PropertyMapB, DifferingProperties ) );
 	//Splitter for left and right blueprint. Current convention is for the local (probably newer?) blueprint to be on the right:
 	return SNew(SSplitter)
 		+ SSplitter::Slot()
@@ -1018,7 +1231,7 @@ TSharedRef<SWidget> SBlueprintDiff::GenerateDefaultsPanel()
 			SNew(SBorder)
 			.VAlign(VAlign_Fill)
 			[
-				GenerateCDODiff(PanelNew.Blueprint, PanelOld.Blueprint)
+				LeftDetailsView
 			]
 		]
 	+ SSplitter::Slot()
@@ -1027,7 +1240,7 @@ TSharedRef<SWidget> SBlueprintDiff::GenerateDefaultsPanel()
 			SNew(SBorder)
 			.VAlign(VAlign_Fill)
 			[
-				GenerateCDODiff(PanelOld.Blueprint, PanelNew.Blueprint)
+				RightDetailsView
 			]
 		];
 }
@@ -1045,7 +1258,7 @@ TSharedRef<SWidget> SBlueprintDiff::GenerateComponentsPanel()
 				GenerateComponentsDiff(PanelNew.Blueprint, PanelOld.Blueprint)
 			]
 		]
-	+ SSplitter::Slot()
+		+ SSplitter::Slot()
 		.Value(0.5f)
 		[
 			SNew(SBorder)
@@ -1060,6 +1273,7 @@ void SBlueprintDiff::SetCurrentMode(FName NewMode)
 {
 	CurrentMode = NewMode;
 
+	DiffHighlighter = TSharedPtr< IDiffHighlighter >();
 	if( NewMode == FBlueprintEditorApplicationModes::StandardBlueprintEditorMode)
 	{
 		ModeContents->SetContent( GenerateGraphPanel() );
diff --git a/Engine/Source/Editor/Kismet/Public/SBlueprintDiff.h b/Engine/Source/Editor/Kismet/Public/SBlueprintDiff.h
--- a/Engine/Source/Editor/Kismet/Public/SBlueprintDiff.h
+++ b/Engine/Source/Editor/Kismet/Public/SBlueprintDiff.h
@@ -30,89 +30,6 @@ struct FDiffResultItem : public TSharedFromThis<FDiffResultItem>
 	TSharedRef<SWidget> KISMET_API GenerateWidget() const;
 };
 
-/*List item that entry for a graph*/
-struct KISMET_API FListItemGraphToDiff: public TSharedFromThis<FListItemGraphToDiff>
-{
-	FListItemGraphToDiff(class SBlueprintDiff* Diff, class UEdGraph* GraphOld, class UEdGraph* GraphNew, const FRevisionInfo& RevisionOld, const FRevisionInfo& RevisionNew);
-	~FListItemGraphToDiff();
-
-	/*Generate Widget for list item*/
-	TSharedRef<SWidget> GenerateWidget() ;
-
-	/*Get tooltip for list item */
-	FText   GetToolTip() ;
-
-	/*Get old(left) graph*/
-	UEdGraph* GetGraphOld()const{return GraphOld;}
-
-	/*Get new(right) graph*/
-	UEdGraph* GetGraphNew()const{return GraphNew;}
-
-private:
-
-	/*Get icon to use by graph name*/
-	static const FSlateBrush* GetIconForGraph(UEdGraph* Graph);
-
-	/*Diff widget*/
-	class SBlueprintDiff* Diff;
-
-	/*The old graph(left)*/
-	class UEdGraph*	GraphOld;
-
-	/*The new graph(right)*/
-	class UEdGraph* GraphNew;
-
-	/*Description of Old and new graph*/
-	FRevisionInfo	RevisionOld, RevisionNew;
-
-	//////////////////////////////////////////////////////////////////////////
-	// Diff list
-	//////////////////////////////////////////////////////////////////////////
-
-	typedef TSharedPtr<struct FDiffResultItem>	FSharedDiffOnGraph;
-	typedef SListView<FSharedDiffOnGraph >		SListViewType;
-
-public:
-
-	/** Called when the Newer Graph is modified*/
-	void OnGraphChanged( const FEdGraphEditAction& Action) ;
-
-	/** Generate list of differences*/
-	TSharedRef<SWidget> GenerateDiffListWidget() ;
-
-	/** Build up the Diff Source Array*/
-	void BuildDiffSourceArray();
-
-	/** Called when user clicks on a new graph list item */
-	void OnSelectionChanged(FSharedDiffOnGraph Item, ESelectInfo::Type SelectionType);
-
-	/** Called when user presses key within the diff view */
-	void KeyWasPressed( const FKeyboardEvent& InKeyboardEvent);
-
-private:
-	
-	/** Called when user clicks button to go to next difference in graph */
-	void NextDiff();
-
-	/** Called when user clicks button to go to prev difference in graph */
-	void PrevDiff();
-
-	/** Get Index of the current diff that is selected */
-	int32 GetCurrentDiffIndex() ;
-
-	/* Called when a new row is being generated */
-	TSharedRef<ITableRow> OnGenerateRow(FSharedDiffOnGraph ParamItem, const TSharedRef<STableViewBase>& OwnerTable );
-
-	/** ListView of differences */
-	TSharedPtr<SListViewType> DiffList;
-
-	/** Source for list view */
-	TArray<FSharedDiffOnGraph> DiffListSource;
-
-	/** Key commands processed by this widget */
-	TSharedPtr< FUICommandList > KeyCommands;
-};
-
 /*panel used to display the blueprint*/
 struct KISMET_API FDiffPanel
 {
@@ -177,21 +94,30 @@ public:
 	void Construct(const FArguments& InArgs);
 
 	/** Called when a new Graph is clicked on by user */
-	void OnGraphChanged(FListItemGraphToDiff* Diff);
+	void OnGraphChanged(struct FListItemGraphToDiff* Diff);
 
 	/** Helper function for generating an empty widget */
 	static TSharedRef<SWidget> DefaultEmptyPanel();
 
 protected:
+	/** Called when user clicks button to go to next difference */
+	void NextDiff();
+
+	/** Called when user clicks button to go to prev difference */
+	void PrevDiff();
+
+	/** Called to determine whether we have a list of differences to cycle through */
+	bool CanCycleDiffs() const;
+
 	/* Need to process keys for shortcuts to buttons */
 	virtual FReply OnKeyDown( const FGeometry& MyGeometry, const FKeyboardEvent& InKeyboardEvent ) override;
 
 
-	typedef TSharedPtr<FListItemGraphToDiff>	FGraphToDiff;
+	typedef TSharedPtr<struct FListItemGraphToDiff>	FGraphToDiff;
 	typedef SListView<FGraphToDiff >	SListViewType;
 
 	/** Bring these revisions of graph into focus on main display*/
-	void FocusOnGraphRevisions( class UEdGraph* GraphOld, class UEdGraph* GraphNew , FListItemGraphToDiff* Diff);
+	void FocusOnGraphRevisions( class UEdGraph* GraphOld, class UEdGraph* GraphNew , struct FListItemGraphToDiff* Diff);
 
 	/*Create a list item entry graph that exists in at least one of the blueprints */
 	void CreateGraphEntry(class UEdGraph* GraphOld, class UEdGraph* GraphNew);
@@ -202,7 +128,7 @@ protected:
 	/*Called when user clicks on a new graph list item */
 	void OnSelectionChanged(FGraphToDiff Item, ESelectInfo::Type SelectionType);
 
-	void OnDiffListSelectionChanged(const TSharedPtr<struct FDiffResultItem>& TheDiff, FListItemGraphToDiff* GraphDiffer);
+	void OnDiffListSelectionChanged(const TSharedPtr<struct FDiffResultItem>& TheDiff, struct FListItemGraphToDiff* GraphDiffer);
 		
 	/** Disable the focus on a particular pin */
 	void DisablePinDiffFocus();
@@ -251,6 +177,12 @@ protected:
 	TSharedPtr<SListViewType>	GraphsToDiff;
 
 	friend struct FListItemGraphToDiff;
+
+	/** Key commands processed by this widget */
+	TSharedPtr< FUICommandList > KeyCommands;
+
+	/** Helper class for highlighting diffs in different types of controls (graph view, details view, etc) */
+	TSharedPtr< class IDiffHighlighter > DiffHighlighter;
 };
 
 
diff --git a/Engine/Source/Editor/PropertyEditor/Private/DetailItemNode.cpp b/Engine/Source/Editor/PropertyEditor/Private/DetailItemNode.cpp
--- a/Engine/Source/Editor/PropertyEditor/Private/DetailItemNode.cpp
+++ b/Engine/Source/Editor/PropertyEditor/Private/DetailItemNode.cpp
@@ -259,7 +259,7 @@ static bool PassesAllFilters( const FDetailLayoutCustomization& InCustomization,
 {
 	bool bPassesAllFilters = true;
 	
-	if( InFilter.FilterStrings.Num() > 0 || InFilter.bShowOnlyModifiedProperties == true )
+	if( InFilter.FilterStrings.Num() > 0 || InFilter.bShowOnlyModifiedProperties == true || InFilter.bShowOnlyDiffering == true )
 	{
 		TSharedPtr<FPropertyNode> PropertyNodePin = InCustomization.GetPropertyNode();
 
@@ -273,9 +273,10 @@ static bool PassesAllFilters( const FDetailLayoutCustomization& InCustomization,
 
 			const bool bPassesSearchFilter = bSearchFilterIsEmpty || ( bIsNotBeingFiltered || bIsSeenDueToFiltering || bIsParentSeenDueToFiltering );
 			const bool bPassesModifiedFilter = bPassesSearchFilter && ( InFilter.bShowOnlyModifiedProperties == false || PropertyNodePin->GetDiffersFromDefault() == true );
+			const bool bPassesDifferingFilter = InFilter.bShowOnlyDiffering == false || InFilter.NamesOfMatchingProperties.Find( PropertyNodePin->GetProperty()->GetName() ) == NULL;
 
 			// The property node is visible (note categories are never visible unless they have a child that is visible )
-			bPassesAllFilters = bPassesSearchFilter && bPassesModifiedFilter;
+			bPassesAllFilters = bPassesSearchFilter && bPassesModifiedFilter && bPassesDifferingFilter;
 		}
 		else if( InCustomization.HasCustomWidget() )
 		{
diff --git a/Engine/Source/Editor/PropertyEditor/Private/PropertyNode.cpp b/Engine/Source/Editor/PropertyEditor/Private/PropertyNode.cpp
--- a/Engine/Source/Editor/PropertyEditor/Private/PropertyNode.cpp
+++ b/Engine/Source/Editor/PropertyEditor/Private/PropertyNode.cpp
@@ -46,6 +46,7 @@ FPropertyNode::FPropertyNode(void)
 	, MaxChildDepthAllowed(FPropertyNodeConstants::NoDepthRestrictions)
 	, PropertyNodeFlags (EPropertyNodeFlags::NoFlags)
 	, bRebuildChildrenRequested( false )
+	, bIsHighlighted( false )
 {
 }
 
diff --git a/Engine/Source/Editor/PropertyEditor/Private/PropertyNode.h b/Engine/Source/Editor/PropertyEditor/Private/PropertyNode.h
--- a/Engine/Source/Editor/PropertyEditor/Private/PropertyNode.h
+++ b/Engine/Source/Editor/PropertyEditor/Private/PropertyNode.h
@@ -649,6 +649,12 @@ public:
 
 	FPropertyChangedEvent& FixPropertiesInEvent(FPropertyChangedEvent& Event);
 
+	/** Returns true if this property node should be highlighted by the view */
+	bool IsHighlighted() const { return bIsHighlighted; }
+
+	/** Sets the highlighted flag to the requested value */
+	bool SetIsHighlighted( bool bInIsHighlighted ) { bIsHighlighted = bInIsHighlighted; }
+
 protected:
 
 	TSharedRef<FEditPropertyChain> BuildPropertyChain( UProperty* PropertyAboutToChange );
@@ -763,6 +769,9 @@ protected:
 	/** If true, children of this node will be rebuilt next tick. */
 	bool bRebuildChildrenRequested;
 
+	/** Used to display a highlight on this node */
+	bool bIsHighlighted;
+
 	/** An array of restrictions limiting this property's potential values in property editors.*/
 	TArray<TSharedRef<const class FPropertyRestriction>> Restrictions;
 };
diff --git a/Engine/Source/Editor/PropertyEditor/Private/SDetailSingleItemRow.cpp b/Engine/Source/Editor/PropertyEditor/Private/SDetailSingleItemRow.cpp
--- a/Engine/Source/Editor/PropertyEditor/Private/SDetailSingleItemRow.cpp
+++ b/Engine/Source/Editor/PropertyEditor/Private/SDetailSingleItemRow.cpp
@@ -286,7 +286,11 @@ bool SDetailSingleItemRow::CanPasteProperty() const
 
 const FSlateBrush* SDetailSingleItemRow::GetBorderImage() const
 {
-	if (IsHovered())
+	if( IsHighlighted() )
+	{
+		return FEditorStyle::GetBrush("DetailsView.CategoryMiddle_Highlighted");
+	}
+	else if (IsHovered())
 	{
 		return FEditorStyle::GetBrush("DetailsView.CategoryMiddle_Hovered");
 	}
@@ -336,4 +340,9 @@ FReply SDetailSingleItemRow::OnAddKeyframeClicked()
 	KeyframeHandler->OnKeyPropertyClicked(*Handle);
 
 	return FReply::Handled();
-}
\ No newline at end of file
+}
+
+bool SDetailSingleItemRow::IsHighlighted() const
+{
+	return Customization->IsValidCustomization() && Customization->HasPropertyNode() && Customization->GetPropertyNode()->IsHighlighted();
+}
diff --git a/Engine/Source/Editor/PropertyEditor/Private/SDetailSingleItemRow.h b/Engine/Source/Editor/PropertyEditor/Private/SDetailSingleItemRow.h
--- a/Engine/Source/Editor/PropertyEditor/Private/SDetailSingleItemRow.h
+++ b/Engine/Source/Editor/PropertyEditor/Private/SDetailSingleItemRow.h
@@ -31,6 +31,7 @@ private:
 	const FSlateBrush* GetBorderImage() const;
 	TSharedRef<SWidget> CreateKeyframeButton( FDetailLayoutCustomization& InCustomization, TSharedRef<IDetailTreeNode> InTreeNode );
 	FReply OnAddKeyframeClicked();
+	bool IsHighlighted() const;
 private:
 	/** Customization for this widget */
 	FDetailLayoutCustomization* Customization;
diff --git a/Engine/Source/Editor/PropertyEditor/Private/SDetailsView.cpp b/Engine/Source/Editor/PropertyEditor/Private/SDetailsView.cpp
--- a/Engine/Source/Editor/PropertyEditor/Private/SDetailsView.cpp
+++ b/Engine/Source/Editor/PropertyEditor/Private/SDetailsView.cpp
@@ -71,6 +71,22 @@ void SDetailsView::Construct(const FArguments& InArgs)
 			);
 		}
 
+		if( DetailsViewArgs.bShowDifferingPeropertiesOption )
+		{
+			DetailViewOptions.AddMenuEntry(
+				LOCTEXT("ShowOnlyDiffering", "Show Only Differing Properties"),
+				LOCTEXT("ShowOnlyDiffering_ToolTip", "Displays only properties in this instance which have been changed or added from the instance being compared"),
+				FSlateIcon(),
+				FUIAction(
+					FExecuteAction::CreateSP(this, &SDetailsView::OnShowOnlyDifferingClicked),
+					FCanExecuteAction(),
+					FIsActionChecked::CreateSP(this, &SDetailsView::IsShowOnlyDifferingChecked)
+				),
+				NAME_None,
+				EUserInterfaceActionType::ToggleButton
+			);
+		}
+
 		FUIAction ShowAllAdvancedAction( 
 			FExecuteAction::CreateSP( this, &SDetailsView::OnShowAllAdvancedClicked ),
 			FCanExecuteAction(),
diff --git a/Engine/Source/Editor/PropertyEditor/Private/SDetailsViewBase.cpp b/Engine/Source/Editor/PropertyEditor/Private/SDetailsViewBase.cpp
--- a/Engine/Source/Editor/PropertyEditor/Private/SDetailsViewBase.cpp
+++ b/Engine/Source/Editor/PropertyEditor/Private/SDetailsViewBase.cpp
@@ -86,6 +86,39 @@ void SDetailsViewBase::HideFilterArea(bool bHide)
 	DetailsViewArgs.bAllowSearch = !bHide;
 }
 
+void SDetailsViewBase::HighlightProperty(const UProperty* Property)
+{
+	// Clear the previously enabled highlight:
+	auto PrevPropertyPtr = PrevHighlightedProperty.Pin();
+	if( PrevPropertyPtr.IsValid() )
+	{
+		PrevPropertyPtr->SetIsHighlighted( false );
+	}
+
+	if( !Property )
+	{
+		return;
+	}
+
+	// Find the newly highlighted node, and enable the highlight:
+	FName PropertyOwnerFName = Property->GetOwnerStruct()->GetFName();
+	FClassInstanceToPropertyMap* PropertyNodeMap = ClassToPropertyMap.Find(PropertyOwnerFName);
+	if( PropertyNodeMap )
+	{
+		FName InstanceName = NAME_None;
+		FPropertyNodeMap* NodeMap = PropertyNodeMap->Find(InstanceName);
+		if( NodeMap )
+		{
+			TSharedPtr<FPropertyNode>* PropertyNode = NodeMap->PropertyNameToNode.Find(Property->GetFName());
+			if( PropertyNode )
+			{
+				(*PropertyNode)->SetIsHighlighted( true );
+				PrevHighlightedProperty = *PropertyNode;
+			}
+		}
+	}
+}
+
 EVisibility SDetailsViewBase::GetTreeVisibility() const
 {
 	return DetailLayout.IsValid() && DetailLayout->HasDetails() ? EVisibility::Visible : EVisibility::Collapsed;
@@ -298,6 +331,13 @@ void SDetailsViewBase::OnShowOnlyModifiedClicked()
 	UpdateFilteredDetails();
 }
 
+void SDetailsViewBase::OnShowOnlyDifferingClicked()
+{
+	CurrentFilter.bShowOnlyDiffering = !CurrentFilter.bShowOnlyDiffering;
+
+	UpdateFilteredDetails();
+}
+
 void SDetailsViewBase::OnShowAllAdvancedClicked()
 {
 	CurrentFilter.bShowAllAdvanced = !CurrentFilter.bShowAllAdvanced;
diff --git a/Engine/Source/Editor/PropertyEditor/Private/SDetailsViewBase.h b/Engine/Source/Editor/PropertyEditor/Private/SDetailsViewBase.h
--- a/Engine/Source/Editor/PropertyEditor/Private/SDetailsViewBase.h
+++ b/Engine/Source/Editor/PropertyEditor/Private/SDetailsViewBase.h
@@ -53,9 +53,10 @@ struct FDetailFilter
 	FDetailFilter()
 		: bShowOnlyModifiedProperties(false)
 		, bShowAllAdvanced(false)
+		, bShowOnlyDiffering(false)
 	{}
 
-	bool IsEmptyFilter() const { return FilterStrings.Num() == 0 && bShowOnlyModifiedProperties == false && bShowAllAdvanced == false; }
+	bool IsEmptyFilter() const { return FilterStrings.Num() == 0 && bShowOnlyModifiedProperties == false && bShowAllAdvanced == false && bShowOnlyDiffering == false; }
 
 	/** Any user search terms that items must match */
 	TArray<FString> FilterStrings;
@@ -63,6 +64,9 @@ struct FDetailFilter
 	bool bShowOnlyModifiedProperties;
 	/** If we should show all advanced properties */
 	bool bShowAllAdvanced;
+	/** If we should only show differing properties */
+	bool bShowOnlyDiffering;
+	TSet<FString> NamesOfMatchingProperties;
 };
 
 struct FDetailColumnSizeData
@@ -96,6 +100,11 @@ public:
 	 */
 	virtual void HideFilterArea(bool bHide) override;
 
+	/** 
+	 * Creates a box around the property with PropertyName and scrolls the property into view
+	 */
+	virtual void HighlightProperty(const UProperty* Property) override;
+
 	virtual FOnFinishedChangingProperties& OnFinishedChangingProperties() override
 	{ 
 		return OnFinishedChangingPropertiesDelegate; 
@@ -323,9 +332,15 @@ protected:
 	/** @return true if show all advanced is checked */
 	bool IsShowAllAdvancedChecked() const { return CurrentFilter.bShowAllAdvanced; }
 
+	/** @return true if show only differing is checked */
+	bool IsShowOnlyDifferingChecked() const { return CurrentFilter.bShowOnlyDiffering; }
+
 	/** Called when show only modified is clicked */
 	void OnShowOnlyModifiedClicked();
 
+	/** Called when show only differing is clicked */
+	void OnShowOnlyDifferingClicked();
+
 	/** Called when show all advanced is clicked */
 	void OnShowAllAdvancedClicked();
 
@@ -337,6 +352,9 @@ protected:
 	/** Called when the filter text changes.  This filters specific property nodes out of view */
 	void OnFilterTextChanged(const FText& InFilterText);
 
+	// @todo doc where does this belong?
+	virtual void UpdateIdenticalProperties(const TSet<FString> IdenticalProperties) override { CurrentFilter.NamesOfMatchingProperties = IdenticalProperties; }
+
 	/** 
 	 * Hides or shows properties based on the passed in filter text
 	 * 
@@ -410,4 +428,6 @@ protected:
 
 	/** External property nodes which need to validated each tick */
 	TArray< TWeakPtr<FObjectPropertyNode> > ExternalRootPropertyNodes;
+
+	TWeakPtr< FPropertyNode > PrevHighlightedProperty;
 };
\ No newline at end of file
diff --git a/Engine/Source/Editor/PropertyEditor/Public/IDetailsView.h b/Engine/Source/Editor/PropertyEditor/Public/IDetailsView.h
--- a/Engine/Source/Editor/PropertyEditor/Public/IDetailsView.h
+++ b/Engine/Source/Editor/PropertyEditor/Public/IDetailsView.h
@@ -33,9 +33,18 @@ struct FDetailsViewArgs
 	uint32 bShowModifiedPropertiesOption : 1;
 	/** True if you want to show the actor label */
 	uint32 bShowActorLabel : 1;
+	/** Bind this delegate to hide differing properties */
+	uint32 bShowDifferingPeropertiesOption : 1;
 
 	/** Default constructor */
-	FDetailsViewArgs( const bool InUpdateFromSelection = false, const bool InLockable = false, const bool InAllowSearch = true, const bool InObjectsUseNameArea = false, const bool InHideSelectionTip = false, FNotifyHook* InNotifyHook = NULL, const bool InSearchInitialKeyFocus = false, FName InViewIdentifier = NAME_None )
+	FDetailsViewArgs( const bool InUpdateFromSelection = false
+					, const bool InLockable = false
+					, const bool InAllowSearch = true
+					, const bool InObjectsUseNameArea = false
+					, const bool InHideSelectionTip = false
+					, FNotifyHook* InNotifyHook = NULL
+					, const bool InSearchInitialKeyFocus = false
+					, FName InViewIdentifier = NAME_None )
 		: ViewIdentifier( InViewIdentifier )
 		, NotifyHook( InNotifyHook ) 
 		, bUpdatesFromSelection( InUpdateFromSelection )
@@ -48,6 +57,7 @@ struct FDetailsViewArgs
 		, bShowOptions( true )
 		, bShowModifiedPropertiesOption(true)
 		, bShowActorLabel(true)
+		, bShowDifferingPeropertiesOption(false)
 	{
 	}
 };
@@ -158,4 +168,14 @@ public:
 	 * Sets the visible state of the filter box/property grid area
 	 */
 	virtual void HideFilterArea(bool bIsVisible) = 0;
+
+	/**
+	 * Creates a box around the property with PropertyName and scrolls the property into view
+	 */
+	virtual void HighlightProperty(const UProperty* Property) = 0;
+
+	/**
+	 * Sets the set of properties that are considered identical, used when filtering out identical properties (showing only differing properties)
+	 */
+	virtual void UpdateIdenticalProperties(const TSet<FString> IdenticalProperties) = 0;
 };
