#ue4 clang build fix and an unused var

diff --git a/Engine/Source/Editor/Kismet/Private/SBlueprintDiff.cpp b/Engine/Source/Editor/Kismet/Private/SBlueprintDiff.cpp
--- a/Engine/Source/Editor/Kismet/Private/SBlueprintDiff.cpp
+++ b/Engine/Source/Editor/Kismet/Private/SBlueprintDiff.cpp
@@ -114,7 +114,7 @@ struct KISMET_API FListItemGraphToDiff	: public TSharedFromThis<FListItemGraphTo
 										, public IDiffHighlighter
 {
 	FListItemGraphToDiff(class SBlueprintDiff* Diff, class UEdGraph* GraphOld, class UEdGraph* GraphNew, const FRevisionInfo& RevisionOld, const FRevisionInfo& RevisionNew);
-	~FListItemGraphToDiff();
+	virtual ~FListItemGraphToDiff();
 
 	/*Generate Widget for list item*/
 	TSharedRef<SWidget> GenerateWidget();
@@ -381,9 +381,6 @@ TSharedRef<SWidget> FListItemGraphToDiff::GenerateDiffListWidget()
 
 		Sort(DiffListSource.GetTypedData(),DiffListSource.Num(), SortDiff());
 
-		
-		const FDiffListCommands& Commands = FDiffListCommands::Get();
-
 		TSharedPtr<SListViewType> DiffListRef;
 		TSharedRef<SHorizontalBox> Result =	SNew(SHorizontalBox)
 			+SHorizontalBox::Slot()
