#UE4 another missing destructor..

diff --git a/Engine/Source/Editor/Kismet/Private/SBlueprintDiff.cpp b/Engine/Source/Editor/Kismet/Private/SBlueprintDiff.cpp
--- a/Engine/Source/Editor/Kismet/Private/SBlueprintDiff.cpp
+++ b/Engine/Source/Editor/Kismet/Private/SBlueprintDiff.cpp
@@ -50,6 +50,8 @@ public:
 	{
 	}
 
+	virtual ~CDODiffHighlighter() {}
+
 private:
 	void NextDiff() override;
 	void PrevDiff() override;
