BP: "_C" postfix in BPGC is not displayed in editor.

diff --git a/Engine/Source/Editor/Kismet/Private/BlueprintDetailsCustomization.cpp b/Engine/Source/Editor/Kismet/Private/BlueprintDetailsCustomization.cpp
--- a/Engine/Source/Editor/Kismet/Private/BlueprintDetailsCustomization.cpp
+++ b/Engine/Source/Editor/Kismet/Private/BlueprintDetailsCustomization.cpp
@@ -3755,14 +3755,18 @@ void FBlueprintInterfaceLayout::GenerateChildContent( IDetailChildrenBuilder& Ch
 	UBlueprint* Blueprint = GlobalOptionsDetailsPtr.Pin()->GetBlueprintObj();
 	check(Blueprint);
 
-	TArray<FString> Interfaces;
+	TArray<FInterfaceName> Interfaces;
 
 	if (!bShowsInheritedInterfaces)
 	{
 		// Generate a list of interfaces already implemented
 		for (TArray<FBPInterfaceDescription>::TConstIterator It(Blueprint->ImplementedInterfaces); It; ++It)
 		{
-			Interfaces.AddUnique( (*It).Interface->GetName() );
+			auto Interface = (*It).Interface;
+			if (Interface)
+			{
+				Interfaces.AddUnique(FInterfaceName(Interface->GetFName(), Interface->GetDisplayNameText()));
+			}
 		}
 	}
 	else
@@ -3776,7 +3780,7 @@ void FBlueprintInterfaceLayout::GenerateChildContent( IDetailChildrenBuilder& Ch
 				FImplementedInterface& CurrentInterface = *It;
 				if( CurrentInterface.Class )
 				{
-					Interfaces.Add( CurrentInterface.Class->GetName() );
+					Interfaces.Add(FInterfaceName(CurrentInterface.Class->GetFName(), CurrentInterface.Class->GetDisplayNameText()));
 				}
 			}
 			BlueprintParent = BlueprintParent->GetSuperClass();
@@ -3792,7 +3796,7 @@ void FBlueprintInterfaceLayout::GenerateChildContent( IDetailChildrenBuilder& Ch
 			+SHorizontalBox::Slot()
 			[
 				SNew(STextBlock)
-					.Text(Interfaces[i])
+					.Text(Interfaces[i].DisplayText)
 					.Font( IDetailLayoutBuilder::GetDetailFont() )
 			]
 		];
@@ -3866,12 +3870,12 @@ void FBlueprintInterfaceLayout::OnBrowseToInterface(TWeakObjectPtr<UObject> Asse
 	}
 }
 
-void FBlueprintInterfaceLayout::OnRemoveInterface(FString InterfaceName)
+void FBlueprintInterfaceLayout::OnRemoveInterface(FInterfaceName InterfaceName)
 {
 	UBlueprint* Blueprint = GlobalOptionsDetailsPtr.Pin()->GetBlueprintObj();
 	check(Blueprint);
 
-	const FName InterfaceFName = *InterfaceName;
+	const FName InterfaceFName = InterfaceName.Name;
 
 	// Close all graphs that are about to be removed
 	TArray<UEdGraph*> Graphs;
@@ -3892,13 +3896,13 @@ void FBlueprintInterfaceLayout::OnRemoveInterface(FString InterfaceName)
 /** Helper function for the interface menu */
 bool IsInterfaceImplemented(const UBlueprint* Blueprint, const UClass* TestInterface)
 {
-	const FString InterfaceName = TestInterface->GetName();
+	const auto InterfaceName = TestInterface->GetFName();
 
 	// First look in the blueprint's ImplementedInterfaces list
 	for(TArray<FBPInterfaceDescription>::TConstIterator it(Blueprint->ImplementedInterfaces); it; ++it)
 	{
 		const FBPInterfaceDescription& CurrentInterface = *it;
-		if( CurrentInterface.Interface && CurrentInterface.Interface->GetName() == InterfaceName )
+		if( CurrentInterface.Interface && CurrentInterface.Interface->GetFName() == InterfaceName )
 		{
 			return true;
 		}
@@ -3911,7 +3915,7 @@ bool IsInterfaceImplemented(const UBlueprint* Blueprint, const UClass* TestInter
 		for (TArray<FImplementedInterface>::TIterator It(BlueprintParent->Interfaces); It; ++It)
 		{
 			const FImplementedInterface& CurrentInterface = *It;
-			if (CurrentInterface.Class && (CurrentInterface.Class->GetName() == InterfaceName))
+			if (CurrentInterface.Class && (CurrentInterface.Class->GetFName() == InterfaceName))
 			{
 				return true;
 			}
@@ -3935,13 +3939,13 @@ TSharedRef<SWidget> FBlueprintInterfaceLayout::OnGetAddInterfaceMenuContent()
 			!IsInterfaceImplemented(Blueprint, CurrentInterface) &&
 			!FKismetEditorUtilities::IsClassABlueprintSkeleton(CurrentInterface))
 		{
-			UnimplementedInterfaces.Add( MakeShareable( new FString(CurrentInterface->GetName()) ));
+			UnimplementedInterfaces.Add(MakeShareable(new FInterfaceName(CurrentInterface->GetFName(), CurrentInterface->GetDisplayNameText())));
 		}
 	}
 
 	if (UnimplementedInterfaces.Num() > 0)
 	{
-		return SNew(SListView<TSharedPtr<FString>>)
+		return SNew(SListView<TSharedPtr<FInterfaceName>>)
 			.ListItemsSource(&UnimplementedInterfaces)
 			.OnGenerateRow(this, &FBlueprintInterfaceLayout::GenerateInterfaceListRow)
 			.OnSelectionChanged(this, &FBlueprintInterfaceLayout::OnInterfaceListSelectionChanged);
@@ -3953,23 +3957,23 @@ TSharedRef<SWidget> FBlueprintInterfaceLayout::OnGetAddInterfaceMenuContent()
 	}
 }
 
-TSharedRef<ITableRow> FBlueprintInterfaceLayout::GenerateInterfaceListRow( TSharedPtr<FString> InterfaceName, const TSharedRef<STableViewBase>& OwningList )
+TSharedRef<ITableRow> FBlueprintInterfaceLayout::GenerateInterfaceListRow(TSharedPtr<FInterfaceName> InterfaceName, const TSharedRef<STableViewBase>& OwningList)
 {
-	return SNew(STableRow< TSharedPtr<FString> >, OwningList)
+	return SNew(STableRow< TSharedPtr<FInterfaceName> >, OwningList)
 		[
 			SNew(STextBlock)
-			.Text(*InterfaceName)
+			.Text(InterfaceName.IsValid() ? InterfaceName->DisplayText : FText())
 		];
 }
 
-void FBlueprintInterfaceLayout::OnInterfaceListSelectionChanged(TSharedPtr<FString> Selection, ESelectInfo::Type SelectInfo)
+void FBlueprintInterfaceLayout::OnInterfaceListSelectionChanged(TSharedPtr<FInterfaceName> Selection, ESelectInfo::Type SelectInfo)
 {
 	if (Selection.IsValid())
 	{
 		UBlueprint* Blueprint = GlobalOptionsDetailsPtr.Pin()->GetBlueprintObj();
 		check(Blueprint);
 
-		FName SelectedInterface = FName(*(*Selection.Get()));
+		FName SelectedInterface = Selection->Name;
 		if (!FBlueprintEditorUtils::ImplementNewInterface( Blueprint, SelectedInterface ) )
 		{
 			GlobalOptionsDetailsPtr.Pin()->GetBlueprintEditorPtr().Pin()->LogSimpleMessage( LOCTEXT("ImplementInterface_Error", "Unable to implement interface. Check log for details") );
@@ -3996,16 +4000,11 @@ UBlueprint* FBlueprintGlobalOptionsDetails::GetBlueprintObj() const
 	return NULL;
 }
 
-FString FBlueprintGlobalOptionsDetails::GetParentClassName() const
+FText FBlueprintGlobalOptionsDetails::GetParentClassName() const
 {
-	FName ParentClassName = NAME_None;
 	const UBlueprint* Blueprint = GetBlueprintObj();
-	if(Blueprint != NULL && Blueprint->ParentClass != NULL)
-	{
-		ParentClassName = Blueprint->ParentClass->GetFName();
-	}
-
-	return ParentClassName.ToString();
+	const UClass* ParentClass = Blueprint ? Blueprint->ParentClass : NULL;
+	return ParentClass ? ParentClass->GetDisplayNameText() : FText::FromName(NAME_None);
 }
 
 bool FBlueprintGlobalOptionsDetails::CanReparent() const
diff --git a/Engine/Source/Editor/Kismet/Private/BlueprintDetailsCustomization.h b/Engine/Source/Editor/Kismet/Private/BlueprintDetailsCustomization.h
--- a/Engine/Source/Editor/Kismet/Private/BlueprintDetailsCustomization.h
+++ b/Engine/Source/Editor/Kismet/Private/BlueprintDetailsCustomization.h
@@ -493,6 +493,21 @@ public:
 		: GlobalOptionsDetailsPtr(InGlobalOptionsDetails)
 		, bShowsInheritedInterfaces(bInShowsInheritedInterfaces) {}
 
+	struct FInterfaceName
+	{
+		FName Name;
+		FText DisplayText;
+
+		FInterfaceName() {}
+		FInterfaceName(FName InName, const FText& InDisplayText) 
+			: Name(InName), DisplayText(InDisplayText) {}
+
+		bool operator==(const FInterfaceName& Other) const
+		{
+			return Name == Other.Name;
+		}
+	};
+
 private:
 	/** IDetailCustomNodeBuilder Interface*/
 	virtual void SetOnRebuildChildren( FSimpleDelegate InOnRegenerateChildren ) override {RegenerateChildrenDelegate = InOnRegenerateChildren;}
@@ -506,11 +521,11 @@ private:
 private:
 	/** Callbacks for details UI */
 	void OnBrowseToInterface(TWeakObjectPtr<UObject> Asset);
-	void OnRemoveInterface(FString InterfaceName);
+	void OnRemoveInterface(FInterfaceName InterfaceName);
 
 	TSharedRef<SWidget> OnGetAddInterfaceMenuContent();
-	TSharedRef<ITableRow> GenerateInterfaceListRow( TSharedPtr<FString> InterfaceName, const TSharedRef<STableViewBase>& OwningList );
-	void OnInterfaceListSelectionChanged(TSharedPtr<FString> Selection, ESelectInfo::Type SelectInfo);
+	TSharedRef<ITableRow> GenerateInterfaceListRow(TSharedPtr<FInterfaceName> InterfaceName, const TSharedRef<STableViewBase>& OwningList);
+	void OnInterfaceListSelectionChanged(TSharedPtr<FInterfaceName> Selection, ESelectInfo::Type SelectInfo);
 
 private:
 	/** The parent graph action details customization */
@@ -520,7 +535,7 @@ private:
 	bool bShowsInheritedInterfaces;
 
 	/** List of unimplemented interfaces, for source for a list view */
-	TArray<TSharedPtr<FString>> UnimplementedInterfaces;
+	TArray<TSharedPtr<FInterfaceName>> UnimplementedInterfaces;
 
 	/** The add interface combo button */
 	TSharedPtr<SComboButton> AddInterfaceComboButton;
@@ -557,7 +572,7 @@ public:
 
 protected:
 	/** Gets the Blueprint parent class name text */
-	FString GetParentClassName() const;
+	FText GetParentClassName() const;
 
 	// Determine whether or not we should be allowed to reparent (but still display the parent class regardless)
 	bool CanReparent() const;
diff --git a/Engine/Source/Editor/Kismet/Private/BlueprintEditor.cpp b/Engine/Source/Editor/Kismet/Private/BlueprintEditor.cpp
--- a/Engine/Source/Editor/Kismet/Private/BlueprintEditor.cpp
+++ b/Engine/Source/Editor/Kismet/Private/BlueprintEditor.cpp
@@ -1345,7 +1345,7 @@ void FBlueprintEditor::PostRegenerateMenusAndToolbars()
 FText FBlueprintEditor::GetParentClassNameText() const
 {
 	UClass* ParentClass = (GetBlueprintObj() != NULL) ? GetBlueprintObj()->ParentClass : NULL;
-	return (ParentClass != NULL) ? FText::FromName( ParentClass->GetFName() ) : LOCTEXT("BlueprintEditor_NoParentClass", "None");
+	return (ParentClass != NULL) ? ParentClass->GetDisplayNameText() : LOCTEXT("BlueprintEditor_NoParentClass", "None");
 }
 
 bool FBlueprintEditor::IsParentClassOfObjectABlueprint( const UBlueprint* Blueprint ) const
diff --git a/Engine/Source/Editor/KismetWidgets/Private/SPinTypeSelector.cpp b/Engine/Source/Editor/KismetWidgets/Private/SPinTypeSelector.cpp
--- a/Engine/Source/Editor/KismetWidgets/Private/SPinTypeSelector.cpp
+++ b/Engine/Source/Editor/KismetWidgets/Private/SPinTypeSelector.cpp
@@ -107,15 +107,20 @@ void SPinTypeSelector::Construct(const FArguments& InArgs, FGetPinTypeTree GetPi
 //=======================================================================
 // Attribute Helpers
 
-FString SPinTypeSelector::GetTypeDescription() const
+FText SPinTypeSelector::GetTypeDescription() const
 {
-	if( TargetPinType.Get().PinSubCategoryObject != NULL )
+	const UObject* PinSubCategoryObject = TargetPinType.Get().PinSubCategoryObject.Get();
+	if (PinSubCategoryObject)
 	{
-		return TargetPinType.Get().PinSubCategoryObject->GetName();
+		if (auto Field = Cast<const UField>(PinSubCategoryObject))
+		{
+			return Field->GetDisplayNameText();
+		}
+		return FText::FromString(PinSubCategoryObject->GetName());
 	}
 	else
 	{
-		return TargetPinType.Get().PinCategory;
+		return FText::FromString(TargetPinType.Get().PinCategory);
 	}
 }
 
diff --git a/Engine/Source/Editor/KismetWidgets/Public/SPinTypeSelector.h b/Engine/Source/Editor/KismetWidgets/Public/SPinTypeSelector.h
--- a/Engine/Source/Editor/KismetWidgets/Public/SPinTypeSelector.h
+++ b/Engine/Source/Editor/KismetWidgets/Public/SPinTypeSelector.h
@@ -51,7 +51,7 @@ protected:
 	FSlateColor GetTypeIconColor() const;
 
 	/** Gets a succinct type description for the type being manipulated */
-	FString GetTypeDescription() const;
+	FText GetTypeDescription() const;
 
 	TSharedPtr<SComboButton>		TypeComboButton;
 	TSharedPtr<SSearchBox>			FilterTextBox;
