Revised TableRow widget tag. Dont think the Unique ID is going to be 'safe' to use for id in this case.

diff --git a/Engine/Source/Editor/Kismet/Private/SSCSEditor.cpp b/Engine/Source/Editor/Kismet/Private/SSCSEditor.cpp
--- a/Engine/Source/Editor/Kismet/Private/SSCSEditor.cpp
+++ b/Engine/Source/Editor/Kismet/Private/SSCSEditor.cpp
@@ -2049,14 +2049,7 @@ TSharedRef<ITableRow> SSCSEditor::MakeTableRowWidget( FSCSEditorTreeNodePtrType
 	FString TagName;
 	if (InNodePtr.IsValid() && InNodePtr->GetComponentTemplate() != NULL )
 	{
-		if (InNodePtr->GetSCSNode())
-		{
-			TagName = FString::Printf(TEXT("TableRow,%s,%d"), *InNodePtr->GetComponentTemplate()->GetReadableName(), InNodePtr->GetSCSNode()->GetUniqueID());
-		}
-		else
-		{
-			TagName = FString::Printf(TEXT("TableRow,%s,0"), *InNodePtr->GetComponentTemplate()->GetReadableName());
-		}
+		TagName = FString::Printf(TEXT("TableRow,%s,0"), *InNodePtr->GetComponentTemplate()->GetReadableName());		
 	}
 	return SNew(SSCS_RowWidget, SharedThis(this), InNodePtr, OwnerTable)
 		.Tag(*TagName);
