Modified the node title display and the way changing the excerpt works in the documentation nodes after noticing a few rough edges using it this morning.
#Branch UE4
#Proj GraphEditor, Kismet

#Change made the node title display the current excerpt name
#Change modified the details view so that the documentation excerpt only changes once a valid link and excerpt has been chosen and is valid. Previously this was a little bit dysfunctional .

diff --git a/Engine/Source/Editor/GraphEditor/Private/SGraphNodeDocumentation.cpp b/Engine/Source/Editor/GraphEditor/Private/SGraphNodeDocumentation.cpp
--- a/Engine/Source/Editor/GraphEditor/Private/SGraphNodeDocumentation.cpp
+++ b/Engine/Source/Editor/GraphEditor/Private/SGraphNodeDocumentation.cpp
@@ -84,7 +84,7 @@ void SGraphNodeDocumentation::UpdateGraphNode()
 				[
 					SAssignNew( InlineEditableText,SInlineEditableTextBlock )
 					.Style( FEditorStyle::Get(), "Graph.Node.NodeTitleInlineEditableText" )
-					.Text( GraphNode->GetNodeTitle( ENodeTitleType::EditableTitle ))
+					.Text( this, &SGraphNodeDocumentation::GetDocumentationTitle )
 				]
 				+SVerticalBox::Slot()
 				.AutoHeight()
@@ -157,6 +157,11 @@ void SGraphNodeDocumentation::UpdateGraphNode()
 	];
 }
 
+FText SGraphNodeDocumentation::GetDocumentationTitle() const
+{
+	return FText::Format( LOCTEXT( "DocumentationNode",  "UDN - {0}" ), FText::FromString( GraphNode->GetDocumentationExcerptName() ));
+}
+
 TSharedPtr<SWidget> SGraphNodeDocumentation::CreateDocumentationPage()
 {
 	TSharedPtr<SWidget> DocumentationWidget;
diff --git a/Engine/Source/Editor/GraphEditor/Public/SGraphNodeDocumentation.h b/Engine/Source/Editor/GraphEditor/Public/SGraphNodeDocumentation.h
--- a/Engine/Source/Editor/GraphEditor/Public/SGraphNodeDocumentation.h
+++ b/Engine/Source/Editor/GraphEditor/Public/SGraphNodeDocumentation.h
@@ -35,6 +35,8 @@ protected:
 	virtual bool IsNameReadOnly () const override { return false; }
 	// End SGraphNode Interface
 
+	/** Retrives the current documentation title based on the chosen excerpt */
+	FText GetDocumentationTitle() const;
 	/** Create documentation page from link and excerpt */
 	TSharedPtr<SWidget> CreateDocumentationPage();
 	/** Returns the width the documentation content must adhere to, used as a delegate in child widgets */
diff --git a/Engine/Source/Editor/Kismet/Private/BlueprintDetailsCustomization.cpp b/Engine/Source/Editor/Kismet/Private/BlueprintDetailsCustomization.cpp
--- a/Engine/Source/Editor/Kismet/Private/BlueprintDetailsCustomization.cpp
+++ b/Engine/Source/Editor/Kismet/Private/BlueprintDetailsCustomization.cpp
@@ -4888,6 +4888,10 @@ void FBlueprintDocumentationDetails::CustomizeDetails(IDetailLayoutBuilder& Deta
 
 	if( DocumentationNodePtr.IsValid() )
 	{
+		// Cache Link
+		DocumentationLink = DocumentationNodePtr->GetDocumentationLink();
+		DocumentationExcerpt = DocumentationNodePtr->GetDocumentationExcerptName();
+
 		IDetailCategoryBuilder& DocumentationCategory = DetailLayout.EditCategory("Documentation", LOCTEXT("DocumentationDetailsCategory", "Documentation").ToString(), ECategoryPriority::Default);
 
 		DocumentationCategory.AddCustomRow( TEXT( "Documentation Link" ))
@@ -4970,36 +4974,23 @@ TWeakObjectPtr<UEdGraphNode_Documentation> FBlueprintDocumentationDetails::EdGra
 
 FText FBlueprintDocumentationDetails::OnGetDocumentationLink() const
 {
-	FText Link = FText::GetEmpty();
-	if( DocumentationNodePtr.IsValid() )
-	{
-		Link = FText::FromString( DocumentationNodePtr->GetDocumentationLink() );
-	}
-	return Link;
+	return FText::FromString( DocumentationLink );
 }
 
 FText FBlueprintDocumentationDetails::OnGetDocumentationExcerpt() const
 {
-	FText Excerpt = NSLOCTEXT( "FBlueprintDocumentationDetails", "ExcerptCombo_DefaultText", "Select Excerpt" );
-	if( DocumentationNodePtr.IsValid() )
-	{
-		Excerpt = FText::FromString( DocumentationNodePtr->GetDocumentationExcerptName() );
-	}
-	return Excerpt;
+	return FText::FromString( DocumentationExcerpt );
 }
 
 bool FBlueprintDocumentationDetails::OnExcerptChangeEnabled() const
 {
-	return DocumentationNodePtr.IsValid() && IDocumentation::Get()->PageExists( DocumentationNodePtr->Link );
+	return IDocumentation::Get()->PageExists( DocumentationLink );
 }
 
-void FBlueprintDocumentationDetails::OnDocumentationLinkCommitted( const FText& InNewName, ETextCommit::Type InTextCommit ) const
+void FBlueprintDocumentationDetails::OnDocumentationLinkCommitted( const FText& InNewName, ETextCommit::Type InTextCommit )
 {
-	if( DocumentationNodePtr.IsValid() )
-	{
-		DocumentationNodePtr->Link = InNewName.ToString();
-		DocumentationNodePtr->Excerpt.Empty();
-	}
+	DocumentationLink = InNewName.ToString();
+	DocumentationExcerpt = NSLOCTEXT( "FBlueprintDocumentationDetails", "ExcerptCombo_DefaultText", "Select Excerpt" ).ToString();
 }
 
 TSharedRef< ITableRow > FBlueprintDocumentationDetails::MakeExcerptViewWidget( TSharedPtr<FString> Item, const TSharedRef< STableViewBase >& OwnerTable )
@@ -5016,7 +5007,9 @@ void FBlueprintDocumentationDetails::OnExcerptSelectionChanged( TSharedPtr<FStri
 {
 	if( ProposedSelection.IsValid() && DocumentationNodePtr.IsValid() )
 	{
-		DocumentationNodePtr->Excerpt = *ProposedSelection.Get();
+		DocumentationNodePtr->Link = DocumentationLink;
+		DocumentationExcerpt = *ProposedSelection.Get();
+		DocumentationNodePtr->Excerpt = DocumentationExcerpt;
 		ExcerptComboButton->SetIsOpen( false );
 	}
 }
@@ -5025,9 +5018,9 @@ TSharedRef<SWidget> FBlueprintDocumentationDetails::GenerateExcerptList()
 {
 	ExcerptList.Empty();
 
-	if( DocumentationNodePtr.IsValid() && IDocumentation::Get()->PageExists( DocumentationNodePtr->Link ))
+	if( IDocumentation::Get()->PageExists( DocumentationLink ))
 	{
-		TSharedPtr<IDocumentationPage> DocumentationPage = IDocumentation::Get()->GetPage( DocumentationNodePtr->Link, NULL );
+		TSharedPtr<IDocumentationPage> DocumentationPage = IDocumentation::Get()->GetPage( DocumentationLink, NULL );
 		TArray<FExcerpt> Excerpts;
 		DocumentationPage->GetExcerpts( Excerpts );
 
diff --git a/Engine/Source/Editor/Kismet/Private/BlueprintDetailsCustomization.h b/Engine/Source/Editor/Kismet/Private/BlueprintDetailsCustomization.h
--- a/Engine/Source/Editor/Kismet/Private/BlueprintDetailsCustomization.h
+++ b/Engine/Source/Editor/Kismet/Private/BlueprintDetailsCustomization.h
@@ -779,7 +779,7 @@ protected:
 	bool OnExcerptChangeEnabled() const;
 
 	/** Handler for the documentation link being committed */
-	void OnDocumentationLinkCommitted( const FText& InNewName, ETextCommit::Type InTextCommit ) const;
+	void OnDocumentationLinkCommitted( const FText& InNewName, ETextCommit::Type InTextCommit );
 	
 	/** Generate table row for excerpt combo */
 	TSharedRef< ITableRow > MakeExcerptViewWidget( TSharedPtr<FString> Item, const TSharedRef< STableViewBase >& OwnerTable );
@@ -793,6 +793,10 @@ protected:
 
 private:
 
+	/** Documentation Link */
+	FString DocumentationLink;
+	/** Current Excerpt */
+	FString DocumentationExcerpt;
 	/** Weak reference to the Blueprint editor */
 	TWeakPtr<FBlueprintEditor> BlueprintEditorPtr;
 	/** The editor node we're editing */
