Fixes some crashes when expanding nodes and invalidating opened graphs.

When invalidating currently opened graphs, tab activation is cleared and focus is cleared.

#ttp 343458 - LIVE: CRASH: EDITOR: Editor crashes when attempting to expand a collapsed node in blueprints

diff --git a/Engine/Source/Editor/Kismet/Private/BlueprintEditor.cpp b/Engine/Source/Editor/Kismet/Private/BlueprintEditor.cpp
--- a/Engine/Source/Editor/Kismet/Private/BlueprintEditor.cpp
+++ b/Engine/Source/Editor/Kismet/Private/BlueprintEditor.cpp
@@ -3547,7 +3547,7 @@ void FBlueprintEditor::OnCollapseSelectionToFunction()
 		BlueprintObj->Modify();
 
 		UEdGraphNode* FunctionNode = NULL;
-		UEdGraph* FunctionGraph = CollapseSelectionToFunction(CollapsableNodes, FunctionNode);
+		UEdGraph* FunctionGraph = CollapseSelectionToFunction(FocusedGraphEdPtr.Pin(), CollapsableNodes, FunctionNode);
 
 		FBlueprintEditorUtils::MarkBlueprintAsStructurallyModified( BlueprintObj );
 
@@ -3698,7 +3698,7 @@ void FBlueprintEditor::OnCollapseSelectionToMacro()
 		BlueprintObj->Modify();
 
 		UEdGraphNode* MacroNode = NULL;
-		UEdGraph* MacroGraph = CollapseSelectionToMacro(CollapsableNodes, MacroNode);
+		UEdGraph* MacroGraph = CollapseSelectionToMacro(FocusedGraphEdPtr.Pin(), CollapsableNodes, MacroNode);
 
 		FBlueprintEditorUtils::MarkBlueprintAsStructurallyModified( BlueprintObj );
 
@@ -3803,7 +3803,7 @@ void FBlueprintEditor::OnPromoteSelectionToFunction()
 				FocusedGraphEd->SetNodeSelection(CompositeNode, false);
 
 				UEdGraphNode* FunctionNode = NULL;
-				CollapseSelectionToFunction(ExpandedNodes, FunctionNode);
+				CollapseSelectionToFunction(FocusedGraphEd, ExpandedNodes, FunctionNode);
 				NodesToSelect.Add(FunctionNode);
 			}
 			else
@@ -3886,7 +3886,7 @@ void FBlueprintEditor::OnPromoteSelectionToMacro()
 				FocusedGraphEd->SetNodeSelection(CompositeNode, false);
 
 				UEdGraphNode* MacroNode = NULL;
-				CollapseSelectionToMacro(ExpandedNodes, MacroNode);
+				CollapseSelectionToMacro(FocusedGraphEd, ExpandedNodes, MacroNode);
 				NodesToSelect.Add(MacroNode);
 			}
 			else
@@ -5237,9 +5237,9 @@ void FBlueprintEditor::CollapseNodes(TSet<UEdGraphNode*>& InCollapsableNodes)
 	CollapseNodesIntoGraph(GatewayNode, GatewayNode->GetInputSink(), GatewayNode->GetOutputSource(), SourceGraph, DestinationGraph, InCollapsableNodes);
 }
 
-UEdGraph* FBlueprintEditor::CollapseSelectionToFunction(TSet<class UEdGraphNode*>& InCollapsableNodes, UEdGraphNode*& OutFunctionNode)
+UEdGraph* FBlueprintEditor::CollapseSelectionToFunction(TSharedPtr<SGraphEditor> InRootGraph, TSet<class UEdGraphNode*>& InCollapsableNodes, UEdGraphNode*& OutFunctionNode)
 {
-	TSharedPtr<SGraphEditor> FocusedGraphEd = FocusedGraphEdPtr.Pin();
+	TSharedPtr<SGraphEditor> FocusedGraphEd = InRootGraph;
 	if (!FocusedGraphEd.IsValid())
 	{
 		return NULL;
@@ -5298,9 +5298,9 @@ UEdGraph* FBlueprintEditor::CollapseSelectionToFunction(TSet<class UEdGraphNode*
 	return NewGraph;
 }
 
-UEdGraph* FBlueprintEditor::CollapseSelectionToMacro(TSet<UEdGraphNode*>& InCollapsableNodes, UEdGraphNode*& OutMacroNode)
+UEdGraph* FBlueprintEditor::CollapseSelectionToMacro(TSharedPtr<SGraphEditor> InRootGraph, TSet<UEdGraphNode*>& InCollapsableNodes, UEdGraphNode*& OutMacroNode)
 {
-	TSharedPtr<SGraphEditor> FocusedGraphEd = FocusedGraphEdPtr.Pin();
+	TSharedPtr<SGraphEditor> FocusedGraphEd = InRootGraph;
 	if (!FocusedGraphEd.IsValid())
 	{
 		return NULL;
@@ -6274,16 +6274,20 @@ bool FBlueprintEditor::IsSelectionNativeFunction()
 	{
 		FGraphPanelSelectionSet SelectedNodes = GraphEditor->GetSelectedNodes();
 		FGraphPanelSelectionSet::TIterator NodeIter( SelectedNodes );
-		const UK2Node_CallFunction* FunctionNode = Cast<UK2Node_CallFunction>( *NodeIter );
 
-		if( FunctionNode && SelectedNodes.Num() == 1 )
+		if(NodeIter)
 		{
-			UFunction* FunctionPtr = FunctionNode->FunctionReference.ResolveMember<UFunction>( FunctionNode );
-			UClass* OwningClass = FunctionPtr ? FunctionPtr->GetOuterUClass() : NULL;
+			const UK2Node_CallFunction* FunctionNode = Cast<UK2Node_CallFunction>( *NodeIter );
 
-			if( OwningClass && OwningClass->HasAllClassFlags( CLASS_Native ))
+			if( FunctionNode && SelectedNodes.Num() == 1 )
 			{
-				return true;
+				UFunction* FunctionPtr = FunctionNode->FunctionReference.ResolveMember<UFunction>( FunctionNode );
+				UClass* OwningClass = FunctionPtr ? FunctionPtr->GetOuterUClass() : NULL;
+
+				if( OwningClass && OwningClass->HasAllClassFlags( CLASS_Native ))
+				{
+					return true;
+				}
 			}
 		}
 	}
@@ -6325,15 +6329,19 @@ bool FBlueprintEditor::IsSelectionNativeVariable()
 	{
 		FGraphPanelSelectionSet SelectedNodes = GraphEditor->GetSelectedNodes();
 		FGraphPanelSelectionSet::TIterator NodeIter( SelectedNodes );
-		const UK2Node_Variable* VarNode = Cast<UK2Node_Variable>( *NodeIter );
 
-		if( VarNode && SelectedNodes.Num() == 1 )
+		if(NodeIter)
 		{
-			UProperty* VariableProperty = VarNode->VariableReference.ResolveMember<UProperty>( VarNode );
+			const UK2Node_Variable* VarNode = Cast<UK2Node_Variable>( *NodeIter );
 
-			if( VariableProperty && VariableProperty->HasAllFlags( RF_Native ))
+			if( VarNode && SelectedNodes.Num() == 1 )
 			{
-				return true;
+				UProperty* VariableProperty = VarNode->VariableReference.ResolveMember<UProperty>( VarNode );
+
+				if( VariableProperty && VariableProperty->HasAllFlags( RF_Native ))
+				{
+					return true;
+				}
 			}
 		}
 	}
diff --git a/Engine/Source/Editor/Kismet/Private/WorkflowOrientedApp/WorkflowTabManager.cpp b/Engine/Source/Editor/Kismet/Private/WorkflowOrientedApp/WorkflowTabManager.cpp
--- a/Engine/Source/Editor/Kismet/Private/WorkflowOrientedApp/WorkflowTabManager.cpp
+++ b/Engine/Source/Editor/Kismet/Private/WorkflowOrientedApp/WorkflowTabManager.cpp
@@ -141,6 +141,8 @@ void FTabInfo::JumpToNearestValidHistoryData()
 			History[CurrentHistoryIndex]->EvokeHistory(AsShared());
 			History[CurrentHistoryIndex]->RestoreHistory();
 			History[CurrentHistoryIndex]->GetFactory().Pin()->OnTabActivated(Tab.Pin());
+			FGlobalTabmanager::Get()->SetActiveTab(nullptr);
+			FSlateApplication::Get().ClearKeyboardFocus(EKeyboardFocusCause::SetDirectly);
 		}
 	}
 }
diff --git a/Engine/Source/Editor/Kismet/Public/BlueprintEditor.h b/Engine/Source/Editor/Kismet/Public/BlueprintEditor.h
--- a/Engine/Source/Editor/Kismet/Public/BlueprintEditor.h
+++ b/Engine/Source/Editor/Kismet/Public/BlueprintEditor.h
@@ -748,10 +748,10 @@ protected:
 	void CollapseNodes(TSet<class UEdGraphNode*>& InCollapsableNodes);
 
 	/** Called when a selection of nodes are being collapsed into a function */
-	UEdGraph* CollapseSelectionToFunction(TSet<class UEdGraphNode*>& InCollapsableNodes, UEdGraphNode*& OutFunctionNode);
+	UEdGraph* CollapseSelectionToFunction(TSharedPtr<SGraphEditor> InRootGraph, TSet<class UEdGraphNode*>& InCollapsableNodes, UEdGraphNode*& OutFunctionNode);
 
 	/** Called when a selection of nodes are being collapsed into a macro */
-	UEdGraph* CollapseSelectionToMacro(TSet<class UEdGraphNode*>& InCollapsableNodes, UEdGraphNode*& OutMacroNode);
+	UEdGraph* CollapseSelectionToMacro(TSharedPtr<SGraphEditor> InRootGraph, TSet<class UEdGraphNode*>& InCollapsableNodes, UEdGraphNode*& OutMacroNode);
 
 	/**
 	 * Called when a selection of nodes is being collapsed into a function
