Making the old property node spawner into a delegate spawner (replacing specific variable spawners with the new var spawner type).

diff --git a/Engine/Source/Editor/BlueprintGraph/Classes/BlueprintDelegateNodeSpawner.h b/Engine/Source/Editor/BlueprintGraph/Classes/BlueprintDelegateNodeSpawner.h
new file mode 100644
--- /dev/null
+++ b/Engine/Source/Editor/BlueprintGraph/Classes/BlueprintDelegateNodeSpawner.h
@@ -0,0 +1,53 @@
+// Copyright 1998-2014 Epic Games, Inc. All Rights Reserved.
+
+#pragma once
+
+#include "BlueprintNodeSpawner.h"
+#include "BlueprintDelegateNodeSpawner.generated.h"
+
+// Forward declarations
+class UK2Node_BaseMCDelegate;
+
+/**
+ * Takes care of spawning various nodes associated with delegates. Serves as the 
+ * "action" portion for certain FBlueprintActionMenuItems. Evolved from 
+ * FEdGraphSchemaAction_K2Delegate, FEdGraphSchemaAction_K2AssignDelegate, etc.
+ */
+UCLASS(Transient)
+class BLUEPRINTGRAPH_API UBlueprintDelegateNodeSpawner : public UBlueprintNodeSpawner
+{
+	GENERATED_UCLASS_BODY()
+
+public:
+	/**
+	 * Creates a new UBlueprintDelegateNodeSpawner for the specified property.
+	 * Does not do any compatibility checking to ensure that the property is
+	 * accessible from blueprints (do that before calling this).
+	 *
+	 * @param  NodeClass	The node type that you want the spawner to spawn.
+	 * @param  Property		The property you want assigned to spawned nodes.
+	 * @param  Outer		Optional outer for the new spawner (if left null, the transient package will be used).
+	 * @return A newly allocated instance of this class.
+	 */
+	static UBlueprintDelegateNodeSpawner* Create(TSubclassOf<UK2Node_BaseMCDelegate> NodeClass, UMulticastDelegateProperty const* const Property, UObject* Outer = nullptr);
+
+	// UBlueprintNodeSpawner interface
+	virtual UEdGraphNode* Invoke(UEdGraph* ParentGraph, FVector2D const Location) const override;
+	virtual FText GetDefaultMenuName() const override;
+	virtual FText GetDefaultMenuCategory() const override;
+	virtual FName GetDefaultMenuIcon(FLinearColor& ColorOut) const override;
+	// End UBlueprintNodeSpawner interface
+
+	/**
+	 * Accessor to the delegate property that this spawner wraps (the delegate
+	 * that this will assign spawned nodes with).
+	 *
+	 * @return The delegate property that this was initialized with.
+	 */
+	UMulticastDelegateProperty const* GetProperty() const;
+	
+private:
+	/** */
+	UPROPERTY()
+	TWeakObjectPtr<UMulticastDelegateProperty> Property;
+};
\ No newline at end of file
diff --git a/Engine/Source/Editor/BlueprintGraph/Classes/BlueprintPropertyNodeSpawner.h b/Engine/Source/Editor/BlueprintGraph/Classes/BlueprintPropertyNodeSpawner.h
deleted file mode 100644
diff --git a/Engine/Source/Editor/BlueprintGraph/Classes/BlueprintVariableNodeSpawner.h b/Engine/Source/Editor/BlueprintGraph/Classes/BlueprintVariableNodeSpawner.h
--- a/Engine/Source/Editor/BlueprintGraph/Classes/BlueprintVariableNodeSpawner.h
+++ b/Engine/Source/Editor/BlueprintGraph/Classes/BlueprintVariableNodeSpawner.h
@@ -30,7 +30,7 @@ public:
 	 * @param  Outer		Optional outer for the new spawner (if left null, the transient package will be used).
 	 * @return A newly allocated instance of this class.
 	 */
-	static UBlueprintVariableNodeSpawner* Create(TSubclassOf<UK2Node_Variable> NodeClass, UProperty* VarProperty, UObject* Outer = nullptr);
+	static UBlueprintVariableNodeSpawner* Create(TSubclassOf<UK2Node_Variable> NodeClass, UProperty const* VarProperty, UObject* Outer = nullptr);
 
 	/**
 	 * Creates a new UBlueprintVariableNodeSpawner, charged with spawning
diff --git a/Engine/Source/Editor/BlueprintGraph/Private/BlueprintActionDatabase.cpp b/Engine/Source/Editor/BlueprintGraph/Private/BlueprintActionDatabase.cpp
--- a/Engine/Source/Editor/BlueprintGraph/Private/BlueprintActionDatabase.cpp
+++ b/Engine/Source/Editor/BlueprintGraph/Private/BlueprintActionDatabase.cpp
@@ -6,7 +6,7 @@
 #include "KismetEditorUtilities.h"	// for IsClassABlueprintSkeleton(), IsClassABlueprintInterface(), GetBoundsForSelectedNodes(), etc.
 #include "BlueprintNodeSpawner.h"
 #include "BlueprintFunctionNodeSpawner.h"
-#include "BlueprintPropertyNodeSpawner.h"
+#include "BlueprintDelegateNodeSpawner.h"
 #include "BlueprintEventNodeSpawner.h"
 #include "BlueprintComponentNodeSpawner.h"
 #include "BlueprintBoundNodeSpawner.h"
@@ -261,7 +261,7 @@ static UBlueprintNodeSpawner* FBlueprintNodeSpawnerFactory::MakeAssignDelegateNo
 {
 	// @TODO: it'd be awesome to have both nodes spawned by this available for 
 	//        context pin matching (the delegate inputs and the event outputs)
-	return UBlueprintPropertyNodeSpawner::Create<UK2Node_AssignDelegate>(DelegateProperty);
+	return UBlueprintDelegateNodeSpawner::Create(UK2Node_AssignDelegate::StaticClass(), DelegateProperty);
 }
 
 /*******************************************************************************
@@ -496,7 +496,7 @@ static void BlueprintActionDatabaseImpl::AddClassPropertyActions(UClass const* c
 			UMulticastDelegateProperty* DelegateProperty = CastChecked<UMulticastDelegateProperty>(Property);
 			if (DelegateProperty->HasAnyPropertyFlags(CPF_BlueprintAssignable))
 			{
-				UBlueprintNodeSpawner* AddSpawner = UBlueprintPropertyNodeSpawner::Create<UK2Node_AddDelegate>(DelegateProperty);
+				UBlueprintNodeSpawner* AddSpawner = UBlueprintDelegateNodeSpawner::Create(UK2Node_AddDelegate::StaticClass(), DelegateProperty);
 				ActionListOut.Add(AddSpawner);
 				
 				UBlueprintNodeSpawner* AssignSpawner = MakeAssignDelegateNodeSpawner(DelegateProperty);
@@ -505,13 +505,13 @@ static void BlueprintActionDatabaseImpl::AddClassPropertyActions(UClass const* c
 			
 			if (DelegateProperty->HasAnyPropertyFlags(CPF_BlueprintCallable))
 			{
-				UBlueprintNodeSpawner* CallSpawner = UBlueprintPropertyNodeSpawner::Create<UK2Node_CallDelegate>(DelegateProperty);
+				UBlueprintNodeSpawner* CallSpawner = UBlueprintDelegateNodeSpawner::Create(UK2Node_CallDelegate::StaticClass(), DelegateProperty);
 				ActionListOut.Add(CallSpawner);
 			}
 			
-			UBlueprintNodeSpawner* RemoveSpawner = UBlueprintPropertyNodeSpawner::Create<UK2Node_RemoveDelegate>(DelegateProperty);
+			UBlueprintNodeSpawner* RemoveSpawner = UBlueprintDelegateNodeSpawner::Create(UK2Node_RemoveDelegate::StaticClass(), DelegateProperty);
 			ActionListOut.Add(RemoveSpawner);
-			UBlueprintNodeSpawner* ClearSpawner  = UBlueprintPropertyNodeSpawner::Create<UK2Node_ClearDelegate>(DelegateProperty);
+			UBlueprintNodeSpawner* ClearSpawner = UBlueprintDelegateNodeSpawner::Create(UK2Node_ClearDelegate::StaticClass(), DelegateProperty);
 			ActionListOut.Add(ClearSpawner);
 
 			// @TODO: AddBoundEventActionsForClass()
@@ -519,14 +519,14 @@ static void BlueprintActionDatabaseImpl::AddClassPropertyActions(UClass const* c
  		}
 		else
 		{
-			if (UObjectProperty* ObjProperty = Cast<UObjectProperty>(Property))
-			{
-				AddComponentBoundActions(ObjProperty, ActionListOut);
-			}
+// 			if (UObjectProperty* ObjProperty = Cast<UObjectProperty>(Property))
+// 			{
+// 				AddComponentBoundActions(ObjProperty, ActionListOut);
+// 			}
 
-			UBlueprintPropertyNodeSpawner* GetterSpawner = UBlueprintPropertyNodeSpawner::Create<UK2Node_VariableGet>(Property);
+			UBlueprintVariableNodeSpawner* GetterSpawner = UBlueprintVariableNodeSpawner::Create(UK2Node_VariableGet::StaticClass(), Property);
 			ActionListOut.Add(GetterSpawner);
-			UBlueprintPropertyNodeSpawner* SetterSpawner = UBlueprintPropertyNodeSpawner::Create<UK2Node_VariableSet>(Property);
+			UBlueprintVariableNodeSpawner* SetterSpawner = UBlueprintVariableNodeSpawner::Create(UK2Node_VariableSet::StaticClass(), Property);
 			ActionListOut.Add(SetterSpawner);
 		}
 	}
@@ -549,8 +549,8 @@ static void BlueprintActionDatabaseImpl::AddComponentBoundActions(UObjectPropert
 				continue;
 			}
 
-			UBlueprintPropertyNodeSpawner* PropertySpawner = UBlueprintPropertyNodeSpawner::Create<UK2Node_ComponentBoundEvent>(Property);
-			UBlueprintBoundNodeSpawner* BoundEventSpawner = UBlueprintBoundNodeSpawner::Create(PropertySpawner, ComponentProperty);
+			UBlueprintDelegateNodeSpawner* DelegateSpawner = UBlueprintDelegateNodeSpawner::Create(/*NodeClass =*/nullptr, Property);
+			DelegateSpawner->NodeClass = UK2Node_ComponentBoundEvent::StaticClass();
 
 			auto BindComponentLambda = [](UEdGraphNode* NewNode, bool /*bIsTemplateNode*/, TWeakObjectPtr<UObjectProperty> ComponentProp, TWeakObjectPtr<UMulticastDelegateProperty> DelegateProp)
 			{
@@ -561,6 +561,7 @@ static void BlueprintActionDatabaseImpl::AddComponentBoundActions(UObjectPropert
 				}
 			};
 
+			UBlueprintBoundNodeSpawner* BoundEventSpawner = UBlueprintBoundNodeSpawner::Create(DelegateSpawner, ComponentProperty);
 			TWeakObjectPtr<UObjectProperty> ComponentPropPtr = ComponentProperty;
 			TWeakObjectPtr<UMulticastDelegateProperty> DelegatePropPtr = Property;
 			BoundEventSpawner->CustomizeNodeDelegate = UBlueprintNodeSpawner::FCustomizeNodeDelegate::CreateStatic(BindComponentLambda, ComponentPropPtr, DelegatePropPtr);
diff --git a/Engine/Source/Editor/BlueprintGraph/Private/BlueprintActionFilter.cpp b/Engine/Source/Editor/BlueprintGraph/Private/BlueprintActionFilter.cpp
--- a/Engine/Source/Editor/BlueprintGraph/Private/BlueprintActionFilter.cpp
+++ b/Engine/Source/Editor/BlueprintGraph/Private/BlueprintActionFilter.cpp
@@ -6,7 +6,7 @@
 #include "BlueprintFunctionNodeSpawner.h"
 #include "BlueprintEventNodeSpawner.h"
 #include "BlueprintComponentNodeSpawner.h"
-#include "BlueprintPropertyNodeSpawner.h"
+#include "BlueprintDelegateNodeSpawner.h"
 #include "BlueprintBoundNodeSpawner.h"
 #include "BlueprintVariableNodeSpawner.h"
 #include "EdGraphSchema_K2.h"		// for FBlueprintMetadata
@@ -366,7 +366,7 @@ static UProperty const* BlueprintActionFilterImpl::GetAssociatedProperty(UBluepr
 {
 	UProperty const* Property = nullptr;
 
-	if (UBlueprintPropertyNodeSpawner const* PropertySpawner = Cast<UBlueprintPropertyNodeSpawner>(BlueprintAction))
+	if (UBlueprintDelegateNodeSpawner const* PropertySpawner = Cast<UBlueprintDelegateNodeSpawner>(BlueprintAction))
 	{
 		Property = PropertySpawner->GetProperty();
 	}
diff --git a/Engine/Source/Editor/BlueprintGraph/Private/BlueprintBoundNodeSpawner.cpp b/Engine/Source/Editor/BlueprintGraph/Private/BlueprintBoundNodeSpawner.cpp
--- a/Engine/Source/Editor/BlueprintGraph/Private/BlueprintBoundNodeSpawner.cpp
+++ b/Engine/Source/Editor/BlueprintGraph/Private/BlueprintBoundNodeSpawner.cpp
@@ -2,9 +2,10 @@
 
 #include "BlueprintGraphPrivatePCH.h"
 #include "BlueprintBoundNodeSpawner.h"
-#include "K2Node_ComponentBoundEvent.h"
-#include "BlueprintPropertyNodeSpawner.h"
+#include "BlueprintVariableNodeSpawner.h"
+#include "BlueprintDelegateNodeSpawner.h"
 #include "K2Node_AddComponent.h"
+#include "K2Node_ComponentBoundEvent.h"
 #include "K2Node_CallFunction.h"
 #include "ComponentAssetBroker.h"
 #include "Editor/EditorEngine.h"	// for GetFriendlyName()
@@ -34,16 +35,7 @@ namespace BlueprintBoundNodeSpawnerImpl
 	 * @param  NewNode		The newly spawned node that needs to be bound.
 	 * @param  PropertyObj	The property object you want bound to the node.
 	 */
-	static void BindFunctionNodeWithProperty(UEdGraphNode* NewNode, UObject* PropertyObj);
-
-	/**
-	 * 
-	 * 
-	 * @param  NewNode	
-	 * @param  DelegateProperty	
-	 * @return 
-	 */
-	static void BindEventNodeToDelegate(UEdGraphNode* NewNode, UObject* DelegateProperty);
+	static void BindFunctionNodeWithVariable(UEdGraphNode* NewNode, UObject* VariableObj);
 };
 
 //------------------------------------------------------------------------------
@@ -59,13 +51,13 @@ static void BlueprintBoundNodeSpawnerImpl::BindAddComponentNodeWithAsset(UEdGrap
 }
 
 //------------------------------------------------------------------------------
-static void BlueprintBoundNodeSpawnerImpl::BindFunctionNodeWithProperty(UEdGraphNode* NewNode, UObject* PropertyObj)
+static void BlueprintBoundNodeSpawnerImpl::BindFunctionNodeWithVariable(UEdGraphNode* NewNode, UObject* PropertyObj)
 {
 	UK2Node_CallFunction* FuncNode = CastChecked<UK2Node_CallFunction>(NewNode);
 	if (UProperty const* BindingProperty = Cast<UProperty>(PropertyObj))
 	{
 		// @TODO: unnecessary allocation, could do this all by hand, or with a local
-		UBlueprintNodeSpawner* TempNodeSpawner = UBlueprintPropertyNodeSpawner::Create<UK2Node_VariableGet>(BindingProperty);
+		UBlueprintNodeSpawner* TempNodeSpawner = UBlueprintVariableNodeSpawner::Create(UK2Node_VariableGet::StaticClass(), BindingProperty);
 		
 		float const EstimatedVarNodeWidth = 224.0f;
 		FVector2D VarNodePos;
@@ -93,42 +85,6 @@ static void BlueprintBoundNodeSpawnerImpl::BindFunctionNodeWithProperty(UEdGraph
 	}
 }
 
-//------------------------------------------------------------------------------
-static void BlueprintBoundNodeSpawnerImpl::BindEventNodeToDelegate(UEdGraphNode* NewNode, UObject* InDelegateProperty)
-{
-	UMulticastDelegateProperty* DelegateProperty = CastChecked<UMulticastDelegateProperty>(InDelegateProperty);
-	if (UK2Node_Event* EventNode = Cast<UK2Node_Event>(NewNode))
-	{
-		FVector2D const SpawnLocation(EventNode->NodePosX, EventNode->NodePosY);
-		FVector2D const EventNodeOffset(-150.f, -150.f);
-		EventNode->NodePosX += EventNodeOffset.X;
-		EventNode->NodePosY += EventNodeOffset.Y;
-
-		// @TODO: unnecessary allocation, could do this all by hand, or with a local
-		UBlueprintNodeSpawner* DelegateNodeSpawner = UBlueprintPropertyNodeSpawner::Create<UK2Node_AddDelegate>(DelegateProperty);
-
-		UEdGraph* ParentGraph = EventNode->GetGraph();
-		UK2Node_AddDelegate* DelegateNode = CastChecked<UK2Node_AddDelegate>(DelegateNodeSpawner->Invoke(ParentGraph, SpawnLocation));
-
-		ParentGraph->Modify();
-		ParentGraph->AddNode(DelegateNode, /*bFromUI =*/false, /*bSelectNewNode =*/false);
-
-		UEdGraphPin* DelegateInputPin = DelegateNode->GetDelegatePin();
-		UEdGraphPin* EventDelegatePin = EventNode->FindPinChecked(UK2Node_Event::DelegateOutputName);
-
-		if ((DelegateInputPin != nullptr) && (EventDelegatePin != nullptr))
-		{
-			UEdGraphSchema_K2 const* K2Schema = GetDefault<UEdGraphSchema_K2>();
-			if (K2Schema->TryCreateConnection(EventDelegatePin, DelegateInputPin))
-			{
-				// refresh the node once connected to the delegate (it should
-				// conform it's pins to match the delegate signature)
-				EventNode->ReconstructNode();
-			}
-		}
-	}
-}
-
 /*******************************************************************************
  * UBlueprintBoundNodeSpawner
  ******************************************************************************/
@@ -180,22 +136,17 @@ UEdGraphNode* UBlueprintBoundNodeSpawner::Invoke(UEdGraph* ParentGraph, FVector2
 	}
 	
 	bool const bIsTemplateNode = (ParentGraph->GetOutermost() == GetTransientPackage());
-	if (!bIsTemplateNode && NewNode->IsA<UK2Node_AddComponent>())
+	if (!bIsTemplateNode)
 	{
-		BindAddComponentNodeWithAsset(NewNode, BoundObj);
-	}
-	else if (NewNode->IsA<UK2Node_CallFunction>())
-	{
-		BindFunctionNodeWithProperty(NewNode, BoundObj);
-	}
-	else if (BoundObj != nullptr)
-	{
-		if (!bIsTemplateNode && BoundObj->IsA<UMulticastDelegateProperty>())
+		if (NewNode->IsA<UK2Node_AddComponent>())
+		{
+			BindAddComponentNodeWithAsset(NewNode, BoundObj);
+		}
+		else if (NewNode->IsA<UK2Node_CallFunction>())
 		{
-			BindEventNodeToDelegate(NewNode, BoundObj);
+			BindFunctionNodeWithVariable(NewNode, BoundObj);
 		}
 	}
-	
 	return NewNode;
 }
 
@@ -209,19 +160,12 @@ FText UBlueprintBoundNodeSpawner::GetDefaultMenuName() const
 	FText MenuName;
 	if (NodeClass == UK2Node_ComponentBoundEvent::StaticClass())
 	{
-		if (UBlueprintPropertyNodeSpawner* PropSpawner = Cast<UBlueprintPropertyNodeSpawner>(SubSpawner))
+		if (UBlueprintDelegateNodeSpawner* PropSpawner = Cast<UBlueprintDelegateNodeSpawner>(SubSpawner))
 		{
 			FText const PropertyName = FText::FromName(PropSpawner->GetProperty()->GetFName());
 			MenuName = FText::Format(LOCTEXT("ComponentEventName", "Add {0}"), PropertyName);
 		}
 	}
-	else if (BoundObjPtr.IsValid() && BoundObjPtr->IsA<UMulticastDelegateProperty>())
-	{
-		UMulticastDelegateProperty* DelegateProperty = CastChecked<UMulticastDelegateProperty>(BoundObjPtr.Get());
-		FText const PropertyName = FText::FromString(bShowFriendlyNames ? UEditorEngine::GetFriendlyName(DelegateProperty) : DelegateProperty->GetName());
-
-		MenuName = FText::Format(LOCTEXT("ComponentEventName", "Assign {0}"), PropertyName);
-	}
 	else
 	{
 		MenuName = SubSpawner->GetDefaultMenuName();
@@ -246,14 +190,6 @@ FText UBlueprintBoundNodeSpawner::GetDefaultMenuCategory() const
 	{
 		FText const PropertyName = FText::FromName(BoundObjPtr->GetFName());
 		DefaultCategory = FText::Format(LOCTEXT("ComponentEventCategory", "Call Function on {0}|{1}"), PropertyName, DefaultCategory);
-	}
-	else if (UMulticastDelegateProperty* DelegateProperty = Cast<UMulticastDelegateProperty>(BoundObjPtr.Get()))
-	{
-		DefaultCategory = FText::FromString(FObjectEditorUtils::GetCategory(DelegateProperty));
-		if (DefaultCategory.IsEmpty())
-		{
-			DefaultCategory = LOCTEXT("BoundDelegateCategory", "Event Dispatchers");
-		}
 	}	
 	return DefaultCategory;
 }
diff --git a/Engine/Source/Editor/BlueprintGraph/Private/BlueprintDelegateNodeSpawner.cpp b/Engine/Source/Editor/BlueprintGraph/Private/BlueprintDelegateNodeSpawner.cpp
new file mode 100644
--- /dev/null
+++ b/Engine/Source/Editor/BlueprintGraph/Private/BlueprintDelegateNodeSpawner.cpp
@@ -0,0 +1,115 @@
+// Copyright 1998-2014 Epic Games, Inc. All Rights Reserved.
+
+#include "BlueprintGraphPrivatePCH.h"
+#include "BlueprintDelegateNodeSpawner.h"
+#include "K2Node_Variable.h"
+#include "EditorStyleSettings.h"	// for bShowFriendlyNames
+#include "ObjectEditorUtils.h"		// for GetCategory()
+#include "BlueprintEditorUtils.h"	// for FindBlueprintForNodeChecked()
+#include "EditorCategoryUtils.h"
+
+#define LOCTEXT_NAMESPACE "BlueprintDelegateNodeSpawner"
+
+/*******************************************************************************
+ * UBlueprintDelegateNodeSpawner
+ ******************************************************************************/
+
+//------------------------------------------------------------------------------
+UBlueprintDelegateNodeSpawner* UBlueprintDelegateNodeSpawner::Create(TSubclassOf<UK2Node_BaseMCDelegate> NodeClass, UMulticastDelegateProperty const* const Property, UObject* Outer/* = nullptr*/)
+{
+	check(Property != nullptr);
+	if (Outer == nullptr)
+	{
+		Outer = GetTransientPackage();
+	}
+
+	UBlueprintDelegateNodeSpawner* NodeSpawner = NewObject<UBlueprintDelegateNodeSpawner>(Outer);
+	NodeSpawner->Property  = Property;
+	NodeSpawner->NodeClass = NodeClass;
+	return NodeSpawner;
+}
+
+//------------------------------------------------------------------------------
+UBlueprintDelegateNodeSpawner::UBlueprintDelegateNodeSpawner(class FPostConstructInitializeProperties const& PCIP)
+	: Super(PCIP)
+{
+}
+
+//------------------------------------------------------------------------------
+UEdGraphNode* UBlueprintDelegateNodeSpawner::Invoke(UEdGraph* ParentGraph, FVector2D const Location) const
+{
+	auto PostSpawnSetupLambda = [](UEdGraphNode* NewNode, bool bIsTemplateNode, UMulticastDelegateProperty const* Property, FCustomizeNodeDelegate UserDelegate)
+	{
+		UK2Node_BaseMCDelegate* DelegateNode = Cast<UK2Node_BaseMCDelegate>(NewNode);
+		if ((DelegateNode != nullptr) && (Property != nullptr))
+		{
+			UBlueprint* Blueprint = FBlueprintEditorUtils::FindBlueprintForNodeChecked(NewNode);
+			UClass* OwnerClass = Property->GetOwnerClass();
+			bool const bIsSelfContext = Blueprint->SkeletonGeneratedClass->IsChildOf(OwnerClass);
+
+			DelegateNode->SetFromProperty(Property, bIsSelfContext);
+		}
+		UserDelegate.ExecuteIfBound(NewNode, bIsTemplateNode);
+	};
+
+	FCustomizeNodeDelegate PostSpawnSetupDelegate = FCustomizeNodeDelegate::CreateStatic(PostSpawnSetupLambda, GetProperty(), CustomizeNodeDelegate);
+	return Super::Invoke(ParentGraph, Location, PostSpawnSetupDelegate);
+}
+
+//------------------------------------------------------------------------------
+FText UBlueprintDelegateNodeSpawner::GetDefaultMenuName() const
+{	
+	FText MenuName = Super::GetDefaultMenuName();
+
+	UMulticastDelegateProperty const* Property = GetProperty();
+	if ((NodeClass == nullptr) && (Property != nullptr))
+	{
+		bool const bShowFriendlyNames = GetDefault<UEditorStyleSettings>()->bShowFriendlyNames;
+		MenuName = bShowFriendlyNames ? FText::FromString(UEditorEngine::GetFriendlyName(Property)) : FText::FromName(Property->GetFName());
+	}
+
+	return MenuName;
+}
+
+//------------------------------------------------------------------------------
+FText UBlueprintDelegateNodeSpawner::GetDefaultMenuCategory() const
+{
+	FText PropertyCategory;
+	if (UMulticastDelegateProperty const* Property = GetProperty())
+	{
+		PropertyCategory = FText::FromString(FObjectEditorUtils::GetCategory(Property));
+	}
+
+	if (PropertyCategory.IsEmpty())
+	{
+		PropertyCategory = FEditorCategoryUtils::GetCommonCategory(FCommonEditorCategory::Delegates);
+	}
+	return PropertyCategory;
+}
+
+//------------------------------------------------------------------------------
+FName UBlueprintDelegateNodeSpawner::GetDefaultMenuIcon(FLinearColor& ColorOut) const
+{
+	FName BrushName = Super::GetDefaultMenuIcon(ColorOut);
+	if (NodeClass == nullptr)
+	{
+		BrushName = TEXT("GraphEditor.Delegate_16x");
+	}
+	else if (UMulticastDelegateProperty const* Property = GetProperty())
+	{
+		FName    const PropertyName  = Property->GetFName();
+		UStruct* const PropertyOwner = CastChecked<UStruct>(Property->GetOuterUField());
+
+		BrushName = UK2Node_Variable::GetVariableIconAndColor(PropertyOwner, PropertyName, ColorOut);
+	}
+	return BrushName;
+}
+
+//------------------------------------------------------------------------------
+UMulticastDelegateProperty const* UBlueprintDelegateNodeSpawner::GetProperty() const
+{
+	// Get() does IsValid() checks for us
+	return Property.Get();
+}
+
+#undef LOCTEXT_NAMESPACE
diff --git a/Engine/Source/Editor/BlueprintGraph/Private/BlueprintPropertyNodeSpawner.cpp b/Engine/Source/Editor/BlueprintGraph/Private/BlueprintPropertyNodeSpawner.cpp
deleted file mode 100644
diff --git a/Engine/Source/Editor/BlueprintGraph/Private/BlueprintVariableNodeSpawner.cpp b/Engine/Source/Editor/BlueprintGraph/Private/BlueprintVariableNodeSpawner.cpp
--- a/Engine/Source/Editor/BlueprintGraph/Private/BlueprintVariableNodeSpawner.cpp
+++ b/Engine/Source/Editor/BlueprintGraph/Private/BlueprintVariableNodeSpawner.cpp
@@ -7,6 +7,7 @@
 #include "Editor/EditorEngine.h"	// for GetFriendlyName()
 #include "ObjectEditorUtils.h"		// for GetCategory()
 #include "EdGraphSchema_K2.h"		// for ConvertPropertyToPinType()
+#include "EditorCategoryUtils.h"	// for BuildCategoryString()
 
 #define LOCTEXT_NAMESPACE "BlueprintVariableNodeSpawner"
 
@@ -15,7 +16,7 @@
  ******************************************************************************/
 
 //------------------------------------------------------------------------------
-UBlueprintVariableNodeSpawner* UBlueprintVariableNodeSpawner::Create(TSubclassOf<UK2Node_Variable> NodeClass, UProperty* VarProperty, UObject* Outer/*= nullptr*/)
+UBlueprintVariableNodeSpawner* UBlueprintVariableNodeSpawner::Create(TSubclassOf<UK2Node_Variable> NodeClass, UProperty const* VarProperty, UObject* Outer/*= nullptr*/)
 {
 	check(VarProperty != nullptr);
 	if (Outer == nullptr)
@@ -121,7 +122,7 @@ FText UBlueprintVariableNodeSpawner::GetDefaultMenuCategory() const
 	{
 		VarSubCategory = FText::FromString(FObjectEditorUtils::GetCategory(MemberVariable.Get()));
 	}
-	return FText::Format(LOCTEXT("VariableCategory", "Variables|{0}"), VarSubCategory);
+	return FEditorCategoryUtils::BuildCategoryString(FCommonEditorCategory::Variables, VarSubCategory);
 }
 
 //------------------------------------------------------------------------------
@@ -157,12 +158,8 @@ UObject const* UBlueprintVariableNodeSpawner::GetVarOuter() const
 //------------------------------------------------------------------------------
 UProperty const* UBlueprintVariableNodeSpawner::GetVarProperty() const
 {
-	UProperty const* VarProperty = nullptr;
-	if (MemberVariable.IsValid())
-	{
-		VarProperty = MemberVariable.Get();
-	}
-	return VarProperty;
+	// Get() does IsValid() checks for us
+	return MemberVariable.Get();
 }
 
 //------------------------------------------------------------------------------
diff --git a/Engine/Source/Editor/Kismet/Private/BlueprintActionMenuBuilder.cpp b/Engine/Source/Editor/Kismet/Private/BlueprintActionMenuBuilder.cpp
--- a/Engine/Source/Editor/Kismet/Private/BlueprintActionMenuBuilder.cpp
+++ b/Engine/Source/Editor/Kismet/Private/BlueprintActionMenuBuilder.cpp
@@ -3,17 +3,17 @@
 #include "BlueprintEditorPrivatePCH.h"
 #include "BlueprintActionMenuBuilder.h"
 #include "BlueprintActionMenuItem.h"
-#include "BlueprintActionDatabase.h"
+#include "BlueprintDragDropMenuItem.h"
 #include "BlueprintActionFilter.h"
+#include "BlueprintActionDatabase.h"
 #include "BlueprintNodeSpawner.h"
 #include "BlueprintFunctionNodeSpawner.h"
-#include "BlueprintPropertyNodeSpawner.h"
+#include "BlueprintDelegateNodeSpawner.h"
+#include "BlueprintVariableNodeSpawner.h"
 #include "K2ActionMenuBuilder.h"		// for FBlueprintPaletteListBuilder/FBlueprintGraphActionListBuilder
 #include "EdGraphSchema_K2.h"			// for StaticClass(), bUseLegacyActionMenus, etc.
 #include "BlueprintEditor.h"			// for GetMyBlueprintWidget(), and GetIsContextSensitive()
-#include "K2ActionMenuBuilder.h"
 #include "SMyBlueprint.h"				// for SelectionAsVar()
-#include "EdGraphSchema_K2_Actions.h"	// for FEdGraphSchemaAction_K2Var
 #include "BlueprintEditorUtils.h"		// for FindBlueprintForGraphChecked()
 #include "BlueprintEditor.h"			// for GetFocusedGraph()
 
@@ -43,48 +43,22 @@ public:
 	 * Spawns a new FBlueprintActionMenuItem with the node-spawner. Constructs
 	 * the menu item's category, name, tooltip, etc.
 	 * 
-	 * @param  EditorContext
-	 * @param  Action		The node-spawner that the new menu item should wrap.
+	 * @param  EditorContext	
+	 * @param  Action			The node-spawner that the new menu item should wrap.
 	 * @return A newly allocated FBlueprintActionMenuItem (which wraps the supplied action).
 	 */
-	TSharedPtr<FEdGraphSchemaAction> MakeMenuItem(TWeakPtr<FBlueprintEditor> EditorContext, UBlueprintNodeSpawner* Action);
-
-	/**
-	 * Clones the property spawner and clears its NodeClass field (used to 
-	 * represent the property as a whole, encompassing all the node types that
-	 * can be associated with it, like: getters, setters, etc.
-	 *
-	 * @param  NodeSpawner	The node-spawner that we want to clone (and create a proxy for).
-	 * @return A cloned UBlueprintPropertyNodeSpawner (representing the property).
-	 */
-	UBlueprintPropertyNodeSpawner* MakePropertyActionProxy( UBlueprintPropertyNodeSpawner* NodeSpawner);
-	
-	/**
-	 * Checks to see if a proxy has already been made for the specified property
-	 * (checks to see if someone has already called MakePropertyActionProxy()
-	 * for a spawner associated with this property).
-	 *
-	 * @param  Property		The property you want to find a proxy for.
-	 * @return True if there's a proxy already cached for the property, otherwise false.
-	 */
-	bool HasProxy(UProperty const* Property) const;
-	
-	/**
-	 * Checks to see if the supplied property node spawner is a proxy itself 
-	 * (was it spawned from MakePropertyActionProxy()? does it have an empty 
-	 * NodeClass field?)
-	 *
-	 * @param  PropertyNodeSpawner	The spawner you want to check.
-	 * @return True if the supplied node-spawner is a proxy, false otherwise.
-	 */
-	bool IsProxy(UBlueprintPropertyNodeSpawner* PropertyNodeSpawner) const;
+	TSharedPtr<FEdGraphSchemaAction> MakeActionMenuItem(TWeakPtr<FBlueprintEditor> EditorContext, UBlueprintNodeSpawner* Action);
 
 	/**
-	 * Clears out all the proxy items that this spawned (and tracked).
+	 * Spawns a new FBlueprintDragDropMenuItem with the node-spawner. Constructs
+	 * the menu item's category, name, tooltip, etc.
+	 * 
+	 * @param  SampleAction	One of the (possibly) many node-spawners that this menu item is set to represent.
+	 * @return A newly allocated FBlueprintActionMenuItem (which wraps the supplied action).
 	 */
-	void Empty();
+	TSharedPtr<FEdGraphSchemaAction> MakeDragDropMenuItem(UBlueprintNodeSpawner* SampleAction);
 	
-protected:
+private:
 	/**
 	 * Attempts to pull a menu name from the supplied spawner. If one isn't 
 	 * provided, then it spawns a temporary node and pulls one from that node's 
@@ -155,11 +129,9 @@ protected:
 	 */
 	UEdGraphNode* GetTemplateNode(UBlueprintNodeSpawner* Action, TWeakPtr<FBlueprintEditor> EditorContext) const;
 
+private:
 	/** Cached context for the blueprint menu being built */
 	FBlueprintActionContext const& Context;
-	
-	/** Tracks intermediate property proxies (if any were created) */
-	TMap<UProperty const*, UBlueprintPropertyNodeSpawner*> PropertyProxyMap;
 };
 
 //------------------------------------------------------------------------------
@@ -171,77 +143,29 @@ FBlueprintActionMenuItemFactory::FBlueprintActionMenuItemFactory(FBlueprintActio
 }
 
 //------------------------------------------------------------------------------
-TSharedPtr<FEdGraphSchemaAction> FBlueprintActionMenuItemFactory::MakeMenuItem(TWeakPtr<FBlueprintEditor> EditorContext, UBlueprintNodeSpawner* Action)
+TSharedPtr<FEdGraphSchemaAction> FBlueprintActionMenuItemFactory::MakeActionMenuItem(TWeakPtr<FBlueprintEditor> EditorContext, UBlueprintNodeSpawner* Action)
 {
 	FLinearColor IconTint = FLinearColor::White;
 	FName IconBrushName = GetMenuIconForAction(EditorContext, Action, IconTint);
 
-	FBlueprintActionMenuItem* NewMenuItem = new FBlueprintActionMenuItem(Action, FEditorStyle::GetBrush(IconBrushName), IconTint, MenuGrouping);
-	
+	FBlueprintActionMenuItem * NewMenuItem = new FBlueprintActionMenuItem(Action, FEditorStyle::GetBrush(IconBrushName), IconTint, MenuGrouping);
 	NewMenuItem->MenuDescription    = GetMenuNameForAction(EditorContext, Action);
 	NewMenuItem->TooltipDescription = GetTooltipForAction(EditorContext, Action).ToString();
-	NewMenuItem->Category           = GetCategoryForAction(EditorContext, Action).ToString();
-	NewMenuItem->Keywords           = GetSearchKeywordsForAction(EditorContext, Action).ToString();
+	NewMenuItem->Category			= GetCategoryForAction(EditorContext, Action).ToString();
+	NewMenuItem->Keywords			= GetSearchKeywordsForAction(EditorContext, Action).ToString();
 
-	if (NewMenuItem->Category.IsEmpty())
-	{
-		NewMenuItem->Category = RootCategory.ToString();
-	}
-	else if (!RootCategory.IsEmpty())
-	{
-		NewMenuItem->Category = FString::Printf(TEXT("%s|%s"), *RootCategory.ToString(), *NewMenuItem->Category);
-	}
-	
+	NewMenuItem->Category = FString::Printf(TEXT("%s|%s"), *RootCategory.ToString(), *NewMenuItem->Category);	
 	return MakeShareable(NewMenuItem);
 }
 
 //------------------------------------------------------------------------------
-UBlueprintPropertyNodeSpawner* FBlueprintActionMenuItemFactory::MakePropertyActionProxy(UBlueprintPropertyNodeSpawner* PropertyNodeSpawner)
-{
-	UBlueprintPropertyNodeSpawner* ActionProxy = nullptr;
-	
-	UProperty const* Property = PropertyNodeSpawner->GetProperty();
-	if (UBlueprintPropertyNodeSpawner** FoundProxy = PropertyProxyMap.Find(Property))
-	{
-		ActionProxy = *FoundProxy;
-	}
-	else
-	{
-		ActionProxy = DuplicateObject(PropertyNodeSpawner, GetTransientPackage());
-		// we use a null node class to flag the proxy, FBlueprintActionMenuItem
-		// queues off of this and will provide a drag-drop menu to the user
-		ActionProxy->NodeClass = nullptr;
-		PropertyProxyMap.Add(Property, ActionProxy);
-	}
-	
-	return ActionProxy;
-}
-
-//------------------------------------------------------------------------------
-bool FBlueprintActionMenuItemFactory::HasProxy(UProperty const* Property) const
-{
-	return PropertyProxyMap.Contains(Property);
-}
-
-//------------------------------------------------------------------------------
-bool FBlueprintActionMenuItemFactory::IsProxy(UBlueprintPropertyNodeSpawner* NodeSpawner) const
+TSharedPtr<FEdGraphSchemaAction> FBlueprintActionMenuItemFactory::MakeDragDropMenuItem(UBlueprintNodeSpawner* SampleAction)
 {
-	bool bIsProxy = false;
-	
-	UProperty const* Property = NodeSpawner->GetProperty();
-	if (UBlueprintPropertyNodeSpawner* const* FoundProxy = PropertyProxyMap.Find(Property))
-	{
-		bIsProxy = (*FoundProxy == NodeSpawner);
-	}
-	check(!bIsProxy || (NodeSpawner->NodeClass == nullptr));
-	
-	return bIsProxy;
-}
+	// FBlueprintDragDropMenuItem takes care of its own menu MenuDescription, etc.
+	FBlueprintDragDropMenuItem* NewMenuItem = new FBlueprintDragDropMenuItem(Context, SampleAction, MenuGrouping);
 
-//------------------------------------------------------------------------------
-void FBlueprintActionMenuItemFactory::Empty()
-{
-	PropertyProxyMap.Empty();
+	NewMenuItem->Category = FString::Printf(TEXT("%s|%s"), *RootCategory.ToString(), *NewMenuItem->Category);
+	return MakeShareable(NewMenuItem);
 }
 
 //------------------------------------------------------------------------------
@@ -406,12 +330,16 @@ UEdGraphNode* FBlueprintActionMenuItemFactory::GetTemplateNode(UBlueprintNodeSpa
 
 namespace FBlueprintActionMenuBuilderImpl
 {
-	/** 
-	 * Defines a sub-section of the overall blueprint menu (filter, heading, etc.)
-	 */
+	/** Defines a sub-section of the overall blueprint menu (filter, heading, etc.) */
 	struct FMenuSectionDefinition
 	{
+	public:
 		FMenuSectionDefinition(FBlueprintActionFilter const& SectionFilter, uint32 const Flags);
+
+		/** Series of ESectionFlags, aimed at customizing how we construct this menu section */
+		uint32 Flags;
+		/** A filter for this section of the menu */
+		FBlueprintActionFilter Filter;
 		
 		/** Sets the root category for menu items in this section. */
 		void SetSectionHeading(FText const& RootCategory);
@@ -427,24 +355,22 @@ namespace FBlueprintActionMenuBuilderImpl
 		 * item to the menu-builder itself).
 		 *
 		 * @param  EditorContext
-		 * @param  Action		The node-spawner that the new menu item should wrap.
+		 * @param  Action			The node-spawner that the new menu item should wrap.
 		 * @return An empty TSharedPtr if the action was filtered out, otherwise a newly allocated FBlueprintActionMenuItem.
 		 */
 		TSharedPtr<FEdGraphSchemaAction> MakeMenuItem(TWeakPtr<FBlueprintEditor> EditorContext, UBlueprintNodeSpawner* Action);
 		
 		/**
-		 * Clears out any intermediate UBlueprintNodeSpawner actions that this
-		 * instantiated.
+		 * Clears out any consolidated properties that this may have been 
+		 * tracking (so we can start a new and spawn new consolidated menu items).
 		 */
 		void Empty();
 		
-		/** Series of ESectionFlags, aimed at customizing how we construct this menu section */
-		uint32 Flags;
-		/** A filter for this section of the menu */
-		FBlueprintActionFilter Filter;
 	private:
 		/** In charge of spawning menu items for this section (holds category/ordering information)*/
 		FBlueprintActionMenuItemFactory ItemFactory;
+		/** Tracks the properties that we've already consolidated and passed (when using the ConsolidatePropertyActions flag)*/
+		TSet<UProperty const*> ConsolidateProperties;
 	};
 	
 	/**
@@ -488,31 +414,48 @@ void FBlueprintActionMenuBuilderImpl::FMenuSectionDefinition::SetSectionSortOrde
 //------------------------------------------------------------------------------
 TSharedPtr<FEdGraphSchemaAction> FBlueprintActionMenuBuilderImpl::FMenuSectionDefinition::MakeMenuItem(TWeakPtr<FBlueprintEditor> EditorContext, UBlueprintNodeSpawner* Action)
 {
-	if (Flags & FBlueprintActionMenuBuilder::ConsolidatePropertyActions)
+	bool bPassedFilter = (Action != nullptr) && !Filter.IsFiltered(Action);
+
+	bool bCreateDragDropItem = false;
+	// if the user wants to consolidate all property actions, then we have to 
+	// check and see if this is one of those that needs consolidating (that 
+	// needs bCreateDragDropItem)
+	if (bPassedFilter && (Flags & FBlueprintActionMenuBuilder::ConsolidatePropertyActions))
 	{
-		if (UBlueprintPropertyNodeSpawner* PropertySpawner = Cast<UBlueprintPropertyNodeSpawner>(Action))
+		UProperty const* ActionProperty = nullptr;
+		if (UBlueprintVariableNodeSpawner* VariableSpawner = Cast<UBlueprintVariableNodeSpawner>(Action))
 		{
-			UProperty const* Property = PropertySpawner->GetProperty();
-			if (ItemFactory.HasProxy(Property))
-			{
-				// we've already filtered a proxy for this property, we don't
-				// want to spawn another
-				Action = nullptr;
-			}
-			else
-			{
-				// make a proxy action (that won't spawn anything)...
-				// FBlueprintActionMenuItem should interpret this and offer the
-				// user a drag-drop menu (with a list of multiple property nodes)
-				Action = ItemFactory.MakePropertyActionProxy(PropertySpawner);
-			}
+			bCreateDragDropItem  = true;
+			ActionProperty = VariableSpawner->GetVarProperty();
+			bPassedFilter = (ActionProperty != nullptr);
+		}
+		else if (UBlueprintDelegateNodeSpawner* DelegateSpawner = Cast<UBlueprintDelegateNodeSpawner>(Action))
+		{
+			bCreateDragDropItem  = true;
+			ActionProperty = DelegateSpawner->GetProperty();
+			bPassedFilter = (ActionProperty != nullptr);
+		}
+
+		if (ActionProperty != nullptr)
+		{
+			// @TODO: could be added as a filter test instead
+			bool const bAlreadyConsolidated = ConsolidateProperties.Contains(ActionProperty);
+			bPassedFilter = !bAlreadyConsolidated;
+			ConsolidateProperties.Add(ActionProperty);
 		}
 	}
 
 	TSharedPtr<FEdGraphSchemaAction> MenuEntry = nullptr;
-	if ((Action != nullptr) && !Filter.IsFiltered(Action))
+	if (bPassedFilter)
 	{
-		MenuEntry = ItemFactory.MakeMenuItem(EditorContext, Action);
+		if (bCreateDragDropItem)
+		{
+			MenuEntry = ItemFactory.MakeDragDropMenuItem(Action);
+		}
+		else
+		{
+			MenuEntry = ItemFactory.MakeActionMenuItem(EditorContext, Action);
+		}
 	}
 	return MenuEntry;
 }
@@ -520,7 +463,7 @@ TSharedPtr<FEdGraphSchemaAction> FBlueprintActionMenuBuilderImpl::FMenuSectionDe
 //------------------------------------------------------------------------------
 void FBlueprintActionMenuBuilderImpl::FMenuSectionDefinition::Empty()
 {
-	ItemFactory.Empty();
+	ConsolidateProperties.Empty();
 }
 
 //------------------------------------------------------------------------------
@@ -542,8 +485,6 @@ void FBlueprintActionMenuBuilderImpl::AppendLegacyItems(FMenuSectionDefinition c
 		}
 		LegacyBuilder.SelectedObjects.Append(MenuContext.SelectedObjects);
 		
-		
-		
 		bool bIsContextSensitive = true;
 		if (BlueprintEditor.IsValid())
 		{
@@ -652,17 +593,17 @@ void FBlueprintActionMenuBuilder::RebuildActionList()
 	// - FEdGraphSchemaAction_K2ViewNode
 	// - FEdGraphSchemaAction_K2AddCustomEvent
 	//   FEdGraphSchemaAction_EventFromFunction
-	//   FEdGraphSchemaAction_K2Var
-	//   FEdGraphSchemaAction_K2Delegate
-	//   FEdGraphSchemaAction_K2AssignDelegate
+	// - FEdGraphSchemaAction_K2Var
+	// - FEdGraphSchemaAction_K2Delegate
+	// - FEdGraphSchemaAction_K2AssignDelegate
 	// - FEdGraphSchemaAction_K2AddComment
-	//   FEdGraphSchemaAction_K2PasteHere
+	// - FEdGraphSchemaAction_K2PasteHere
 	// - FEdGraphSchemaAction_K2NewNode
-	//   FEdGraphSchemaAction_Dummy
+	// - FEdGraphSchemaAction_Dummy
 	//   FEdGraphSchemaAction_K2AddCallOnActor
 	//   FEdGraphSchemaAction_K2AddCallOnVariable
 	// - FEdGraphSchemaAction_K2AddComponent
-	//   FEdGraphSchemaAction_K2AddComment
+	// - FEdGraphSchemaAction_K2AddComment
 }
 
 #undef LOCTEXT_NAMESPACE
diff --git a/Engine/Source/Editor/Kismet/Private/BlueprintActionMenuItem.cpp b/Engine/Source/Editor/Kismet/Private/BlueprintActionMenuItem.cpp
--- a/Engine/Source/Editor/Kismet/Private/BlueprintActionMenuItem.cpp
+++ b/Engine/Source/Editor/Kismet/Private/BlueprintActionMenuItem.cpp
@@ -6,10 +6,6 @@
 #include "KismetEditorUtilities.h"	// for BringKismetToFocusAttentionOnObject()
 #include "BlueprintEditorUtils.h"	// for AnalyticsTrackNewNode(), MarkBlueprintAsModified(), etc.
 #include "ScopedTransaction.h"
-#include "BlueprintPropertyNodeSpawner.h"
-#include "BPVariableDragDropAction.h"
-#include "BPDelegateDragDropAction.h"
-#include "BlueprintEditor.h"		// for GetVarIconAndColor()
 #include "SNodePanel.h"				// for GetSnapGridSize()
 
 #define LOCTEXT_NAMESPACE "BlueprintActionMenuItem"
@@ -57,7 +53,7 @@ static void FBlueprintMenuActionItemImpl::DirtyBlueprintFromNewNode(UEdGraphNode
  ******************************************************************************/
 
 //------------------------------------------------------------------------------
-FBlueprintActionMenuItem::FBlueprintActionMenuItem(UBlueprintNodeSpawner* NodeSpawner, FSlateBrush const* MenuIcon, FSlateColor const& IconTintIn, int32 MenuGrouping/* = 0*/)
+FBlueprintActionMenuItem::FBlueprintActionMenuItem(UBlueprintNodeSpawner const* NodeSpawner, FSlateBrush const* MenuIcon, FSlateColor const& IconTintIn, int32 MenuGrouping/* = 0*/)
 	: Action(NodeSpawner)
 {
 	check(Action != nullptr);
@@ -160,49 +156,8 @@ void FBlueprintActionMenuItem::AddReferencedObjects(FReferenceCollector& Collect
 //------------------------------------------------------------------------------
 FSlateBrush const* FBlueprintActionMenuItem::GetMenuIcon(FSlateColor& ColorOut)
 {
-	// if this brush is invalid
-	if (IconBrush == nullptr)
-	{
-		if (UBlueprintPropertyNodeSpawner const* PropertySpawner = Cast<UBlueprintPropertyNodeSpawner>(Action))
-		{
-			UProperty const* Property = PropertySpawner->GetProperty();
-			FName const PropertyName = Property->GetFName();
-			UStruct* const PropertyOwner = CastChecked<UStruct>(Property->GetOuterUField());
-
-			IconBrush = FBlueprintEditor::GetVarIconAndColor(PropertyOwner, PropertyName, IconTint);
-		}
-	}
-
 	ColorOut = IconTint;
 	return IconBrush;
 }
 
-//------------------------------------------------------------------------------
-TSharedPtr<FDragDropOperation> FBlueprintActionMenuItem::OnDragged(FNodeCreationAnalytic AnalyticsDelegate) const
-{
-	TSharedPtr<FDragDropOperation> DragDropAction = nullptr;
-
-	if (UBlueprintPropertyNodeSpawner const* PropertySpawner = Cast<UBlueprintPropertyNodeSpawner>(Action))
-	{
-		if (PropertySpawner->NodeClass == nullptr)
-		{
-			UProperty const* Property = PropertySpawner->GetProperty();
-			FName const PropertyName = Property->GetFName();
-			UStruct* const PropertyOwner = CastChecked<UStruct>(Property->GetOuterUField());
-
-			if (PropertySpawner->IsDelegateProperty())
-			{
-				TSharedRef<FDragDropOperation> DragDropOpRef = FKismetDelegateDragDropAction::New(PropertyName, PropertyOwner, AnalyticsDelegate);
-				DragDropAction = TSharedPtr<FDragDropOperation>(DragDropOpRef);
-			}
-			else
-			{
-				TSharedRef<FDragDropOperation> DragDropOpRef = FKismetVariableDragDropAction::New(PropertyName, PropertyOwner, AnalyticsDelegate);
-				DragDropAction = TSharedPtr<FDragDropOperation>(DragDropOpRef);
-			}
-		}
-	}
-	return DragDropAction;
-}
-
 #undef LOCTEXT_NAMESPACE
diff --git a/Engine/Source/Editor/Kismet/Private/BlueprintDragDropMenuItem.cpp b/Engine/Source/Editor/Kismet/Private/BlueprintDragDropMenuItem.cpp
new file mode 100644
--- /dev/null
+++ b/Engine/Source/Editor/Kismet/Private/BlueprintDragDropMenuItem.cpp
@@ -0,0 +1,175 @@
+// Copyright 1998-2014 Epic Games, Inc. All Rights Reserved.
+
+#include "BlueprintEditorPrivatePCH.h"
+#include "BlueprintDragDropMenuItem.h"
+#include "BlueprintActionMenuItem.h" // for PerformAction()
+#include "BlueprintNodeSpawner.h"
+#include "BlueprintDelegateNodeSpawner.h"
+#include "BlueprintVariableNodeSpawner.h"
+#include "BPDelegateDragDropAction.h"
+#include "BPVariableDragDropAction.h"
+#include "BlueprintEditor.h"		// for GetVarIconAndColor()
+#include "ObjectEditorUtils.h"
+#include "SlateColor.h"
+#include "EdGraph/EdGraph.h"
+#include "EdGraph/EdGraphNode.h"
+#include "BlueprintActionFilter.h"	// for FBlueprintActionContext
+#include "EditorCategoryUtils.h"
+#include "EditorStyleSettings.h"	// for bShowFriendlyNames
+
+#define LOCTEXT_NAMESPACE "BlueprintDragDropMenuItem"
+DEFINE_LOG_CATEGORY_STATIC(LogBlueprintDragDropMenuItem, Log, All);
+
+//------------------------------------------------------------------------------
+FBlueprintDragDropMenuItem::FBlueprintDragDropMenuItem(FBlueprintActionContext const& Context, UBlueprintNodeSpawner* SampleActionIn, int32 MenuGrouping/* = 0*/)
+	: SampleAction(SampleActionIn)
+{
+	check(SampleAction != nullptr);
+	Grouping  = MenuGrouping;
+
+	UProperty const* SampleProperty = nullptr;
+	if (UBlueprintDelegateNodeSpawner const* DelegateSpawner = Cast<UBlueprintDelegateNodeSpawner const>(SampleAction))
+	{
+		SampleProperty = DelegateSpawner->GetProperty();
+	}
+	else if (UBlueprintVariableNodeSpawner const* VariableSpawner = Cast<UBlueprintVariableNodeSpawner const>(SampleAction))
+	{
+		SampleProperty = VariableSpawner->GetVarProperty();
+	}
+
+	if (SampleProperty != nullptr)
+	{
+		bool const bShowFriendlyNames = GetDefault<UEditorStyleSettings>()->bShowFriendlyNames;
+		MenuDescription = bShowFriendlyNames ? FText::FromString(UEditorEngine::GetFriendlyName(SampleProperty)) : FText::FromName(SampleProperty->GetFName());
+
+		TooltipDescription = SampleProperty->GetToolTipText().ToString();
+		Category = FObjectEditorUtils::GetCategory(SampleProperty);
+
+		bool const bIsDelegateProperty = SampleProperty->IsA<UMulticastDelegateProperty>();
+		if (bIsDelegateProperty && Category.IsEmpty())
+		{
+			Category = FEditorCategoryUtils::GetCommonCategory(FCommonEditorCategory::Delegates).ToString();
+		}
+		else if (!bIsDelegateProperty)
+		{
+			check(Context.Blueprints.Num() > 0);
+			UBlueprint const* Blueprint = Context.Blueprints[0];
+			UClass const*     BlueprintClass = (Blueprint->SkeletonGeneratedClass != nullptr) ? Blueprint->SkeletonGeneratedClass : Blueprint->ParentClass;
+
+			UClass const* PropertyClass = SampleProperty->GetOwnerClass();
+			checkSlow(PropertyClass != nullptr);
+			bool const bIsMemberProperty = BlueprintClass->IsChildOf(PropertyClass);
+
+			FText TextCategory;
+			if (Category.IsEmpty())
+			{
+				TextCategory = FEditorCategoryUtils::GetCommonCategory(FCommonEditorCategory::Variables);
+			}
+			else if (bIsMemberProperty)
+			{
+				TextCategory = FEditorCategoryUtils::BuildCategoryString(FCommonEditorCategory::Variables, FText::FromString(Category));
+			}
+			
+			if (!bIsMemberProperty)
+			{
+				TextCategory = FText::Format(LOCTEXT("NonMemberVarCategory", "Class|{0}|{1}"), PropertyClass->GetDisplayNameText(), TextCategory);
+			}
+			Category = TextCategory.ToString();
+		}
+	}
+	else
+	{
+		UE_LOG(LogBlueprintDragDropMenuItem, Warning, TEXT("Unhandled (or invalid) spawner: '%s'"), *SampleAction->GetName());
+	}
+}
+
+//------------------------------------------------------------------------------
+UEdGraphNode* FBlueprintDragDropMenuItem::PerformAction(UEdGraph* ParentGraph, UEdGraphPin* FromPin, FVector2D const Location, bool bSelectNewNode/* = true*/)
+{
+	// we shouldn't get here (this should just be used for drag/drop ops), but just in case we do...
+	FBlueprintActionMenuItem BlueprintActionItem(SampleAction, nullptr, FSlateColor(FLinearColor::White));
+	return BlueprintActionItem.PerformAction(ParentGraph, FromPin, Location, bSelectNewNode);
+}
+
+//------------------------------------------------------------------------------
+UEdGraphNode* FBlueprintDragDropMenuItem::PerformAction(UEdGraph* ParentGraph, TArray<UEdGraphPin*>& FromPins, FVector2D const Location, bool bSelectNewNode/* = true*/)
+{
+	UEdGraphPin* FromPin = nullptr;
+	if (FromPins.Num() > 0)
+	{
+		FromPin = FromPins[0];
+	}
+	
+	UEdGraphNode* SpawnedNode = PerformAction(ParentGraph, FromPin, Location, bSelectNewNode);
+	// try auto-wiring the rest of the pins (if there are any)
+	for (int32 PinIndex = 1; PinIndex < FromPins.Num(); ++PinIndex)
+	{
+		SpawnedNode->AutowireNewNode(FromPins[PinIndex]);
+	}
+	
+	return SpawnedNode;
+}
+
+//------------------------------------------------------------------------------
+void FBlueprintDragDropMenuItem::AddReferencedObjects(FReferenceCollector& Collector)
+{
+	FEdGraphSchemaAction::AddReferencedObjects(Collector);
+	
+	// these don't get saved to disk, but we want to make sure the objects don't
+	// get GC'd while the action array is around
+	Collector.AddReferencedObject(SampleAction);
+}
+
+//------------------------------------------------------------------------------
+FSlateBrush const* FBlueprintDragDropMenuItem::GetMenuIcon(FSlateColor& ColorOut)
+{
+	FSlateBrush const* IconBrush = nullptr;
+	ColorOut = FSlateColor(FLinearColor::White);
+
+	if (UBlueprintDelegateNodeSpawner const* DelegateSpawner = Cast<UBlueprintDelegateNodeSpawner const>(SampleAction))
+	{
+		IconBrush = FEditorStyle::GetBrush(TEXT("GraphEditor.Delegate_16x"));
+	}
+	else if (UBlueprintVariableNodeSpawner const* VariableSpawner = Cast<UBlueprintVariableNodeSpawner const>(SampleAction))
+	{
+		if (UProperty const* Property = VariableSpawner->GetVarProperty())
+		{
+			UStruct* const PropertyOwner = CastChecked<UStruct>(Property->GetOuterUField());
+			IconBrush = FBlueprintEditor::GetVarIconAndColor(PropertyOwner, Property->GetFName(), ColorOut);
+		}
+	}
+	return IconBrush;
+}
+
+//------------------------------------------------------------------------------
+TSharedPtr<FDragDropOperation> FBlueprintDragDropMenuItem::OnDragged(FNodeCreationAnalytic AnalyticsDelegate) const
+{
+	TSharedPtr<FDragDropOperation> DragDropAction = nullptr;
+
+	if (UBlueprintDelegateNodeSpawner const* DelegateSpawner = Cast<UBlueprintDelegateNodeSpawner const>(SampleAction))
+	{
+		if (UMulticastDelegateProperty const* Property = DelegateSpawner->GetProperty())
+		{
+			UStruct* const PropertyOwner = CastChecked<UStruct>(Property->GetOuterUField());
+			TSharedRef<FDragDropOperation> DragDropOpRef = FKismetDelegateDragDropAction::New(Property->GetFName(), PropertyOwner, AnalyticsDelegate);
+			DragDropAction = TSharedPtr<FDragDropOperation>(DragDropOpRef);
+		}
+	}
+	else if (UBlueprintVariableNodeSpawner const* VariableSpawner = Cast<UBlueprintVariableNodeSpawner const>(SampleAction))
+	{
+		if (UProperty const* Property = VariableSpawner->GetVarProperty())
+		{
+			UStruct* const PropertyOwner = CastChecked<UStruct>(Property->GetOuterUField());
+			TSharedRef<FDragDropOperation> DragDropOpRef = FKismetVariableDragDropAction::New(Property->GetFName(), PropertyOwner, AnalyticsDelegate);
+			DragDropAction = TSharedPtr<FDragDropOperation>(DragDropOpRef);
+		}
+		// @TODO: handle local variables as well
+	}
+	else
+	{
+		UE_LOG(LogBlueprintDragDropMenuItem, Warning, TEXT("Unhandled spawner type: '%s'"), *SampleAction->GetClass()->GetName());
+	}
+	return DragDropAction;
+}
+
+#undef LOCTEXT_NAMESPACE
diff --git a/Engine/Source/Editor/Kismet/Private/SBlueprintPalette.cpp b/Engine/Source/Editor/Kismet/Private/SBlueprintPalette.cpp
--- a/Engine/Source/Editor/Kismet/Private/SBlueprintPalette.cpp
+++ b/Engine/Source/Editor/Kismet/Private/SBlueprintPalette.cpp
@@ -27,6 +27,7 @@
 #include "AnimStateConduitNode.h"
 #include "AnimationTransitionGraph.h"
 #include "BlueprintActionMenuItem.h"
+#include "BlueprintDragDropMenuItem.h"
 #include "BlueprintNodeSpawner.h"
 
 #define LOCTEXT_NAMESPACE "BlueprintPalette"
@@ -361,6 +362,11 @@ static void GetPaletteItemIcon(TSharedPtr<FEdGraphSchemaAction> ActionIn, UBluep
 		FBlueprintActionMenuItem* NodeSpawnerAction = (FBlueprintActionMenuItem*)ActionIn.Get();
 		BrushOut = NodeSpawnerAction->GetMenuIcon(ColorOut);
 	}
+	else if (ActionIn->GetTypeId() == FBlueprintDragDropMenuItem::StaticGetTypeId())
+	{
+		FBlueprintDragDropMenuItem* DragDropAction = (FBlueprintDragDropMenuItem*)ActionIn.Get();
+		BrushOut = DragDropAction->GetMenuIcon(ColorOut);
+	}
 	else if (ActionIn->GetTypeId() == FEdGraphSchemaAction_K2AddDocumentation::StaticGetTypeId())
 	{
 		ToolTipOut = LOCTEXT("Documentation_Tooltip", "Create a documentation node.").ToString();
diff --git a/Engine/Source/Editor/Kismet/Private/SBlueprintSubPalette.cpp b/Engine/Source/Editor/Kismet/Private/SBlueprintSubPalette.cpp
--- a/Engine/Source/Editor/Kismet/Private/SBlueprintSubPalette.cpp
+++ b/Engine/Source/Editor/Kismet/Private/SBlueprintSubPalette.cpp
@@ -12,7 +12,7 @@
 #include "BlueprintPaletteFavorites.h"
 #include "BlueprintEditorUtils.h"
 #include "SBlueprintActionMenu.h" // for SBlueprintActionMenuExpander
-#include "BlueprintActionMenuItem.h"
+#include "BlueprintDragDropMenuItem.h"
 
 #define LOCTEXT_NAMESPACE "BlueprintSubPalette"
 
@@ -323,9 +323,9 @@ FReply SBlueprintSubPalette::OnActionDragged( const TArray< TSharedPtr<FEdGraphS
 				return FReply::Handled().BeginDragDrop(FKismetDelegateDragDropAction::New(DelegateAction->GetDelegateName(), VarClass, AnalyticsDelegate));
 			}
 		}
-		else if (InAction->GetTypeId() == FBlueprintActionMenuItem::StaticGetTypeId())
+		else if (InAction->GetTypeId() == FBlueprintDragDropMenuItem::StaticGetTypeId())
 		{
-			FBlueprintActionMenuItem* BlueprintAction = (FBlueprintActionMenuItem*)InAction.Get();
+			FBlueprintDragDropMenuItem* BlueprintAction = (FBlueprintDragDropMenuItem*)InAction.Get();
 
 			TSharedPtr<FDragDropOperation> DragDropOp = BlueprintAction->OnDragged(AnalyticsDelegate);
 			if (DragDropOp.IsValid())
diff --git a/Engine/Source/Editor/Kismet/Public/BlueprintActionMenuItem.h b/Engine/Source/Editor/Kismet/Public/BlueprintActionMenuItem.h
--- a/Engine/Source/Editor/Kismet/Public/BlueprintActionMenuItem.h
+++ b/Engine/Source/Editor/Kismet/Public/BlueprintActionMenuItem.h
@@ -3,7 +3,6 @@
 #pragma once
 
 #include "EdGraph/EdGraphSchema.h" // for FEdGraphSchemaAction
-#include "BlueprintEditor.h"       // for FNodeCreationAnalytic
 #include "SlateColor.h"
 #include "BlueprintActionMenuItem.generated.h"
 
@@ -28,7 +27,7 @@ public:
 	
 	/** Constructors */
 	FBlueprintActionMenuItem() : Action(nullptr), IconBrush(nullptr), IconTint(FLinearColor::White) {}
-	FBlueprintActionMenuItem(UBlueprintNodeSpawner* NodeSpawner, FSlateBrush const* MenuIcon, FSlateColor const& IconTint, int32 MenuGrouping = 0);
+	FBlueprintActionMenuItem(UBlueprintNodeSpawner const* NodeSpawner, FSlateBrush const* MenuIcon, FSlateColor const& IconTint, int32 MenuGrouping = 0);
 	
 	// FEdGraphSchemaAction interface
 	virtual FName         GetTypeId() const final { return StaticGetTypeId(); }
@@ -38,28 +37,19 @@ public:
 	// End FEdGraphSchemaAction interface
 
 	/**
-	 * Retrieves the icon brush for this menu entry.
+	 * Retrieves the icon brush for this menu entry (to be displayed alongside
+	 * in the menu).
+	 *
+	 * @param  ColorOut	An output parameter that's filled in with the color to tint the brush with.
 	 * @return An slate brush to be used for this menu item in the action menu.
 	 */
 	FSlateBrush const* GetMenuIcon(FSlateColor& ColorOut);
 
-	/**
-	 * Attempts to create a specific drag/drop action for this menu entry 
-	 * (certain menu entries require unique drag/drop handlers... like property
-	 * placeholders, where a drop action spawns a sub-menu for the user to pick
-	 * a node type from).
-	 * 
-	 * @param  AnalyticsDelegate	The analytics callback to assign the drag/drop operation.
-	 * @return An empty TSharedPtr if this menu item doesn't require a unique one, else a newly instantiated drag/drop op.
-	 */
-	TSharedPtr<FDragDropOperation> OnDragged(FNodeCreationAnalytic AnalyticsDelegate) const;
-
-	/** Specialized node-spawner, that comprises the action portion of this menu entry. */
-	UBlueprintNodeSpawner const* Action;
-
 private:
 	/** Brush that should be used for the icon on this menu item. */
 	FSlateBrush const* IconBrush;
 	/** Tint to return along with the icon brush. */
 	FSlateColor IconTint;
+	/** Specialized node-spawner, that comprises the action portion of this menu entry. */
+	UBlueprintNodeSpawner const* Action;
 };
diff --git a/Engine/Source/Editor/Kismet/Public/BlueprintDragDropMenuItem.h b/Engine/Source/Editor/Kismet/Public/BlueprintDragDropMenuItem.h
new file mode 100644
--- /dev/null
+++ b/Engine/Source/Editor/Kismet/Public/BlueprintDragDropMenuItem.h
@@ -0,0 +1,76 @@
+// Copyright 1998-2014 Epic Games, Inc. All Rights Reserved.
+
+#pragma once
+
+#include "EdGraph/EdGraphSchema.h" // for FEdGraphSchemaAction
+#include "BlueprintEditor.h"       // for FNodeCreationAnalytic
+#include "BlueprintDragDropMenuItem.generated.h"
+
+// Forward declarations
+class  UBlueprintNodeSpawner;
+class  FDragDropOperation;
+struct FSlateBrush;
+struct FSlateColor;
+struct FBlueprintActionContext;
+
+/**
+ * At certain times we want a single menu entry that represents a set of 
+ * UBlueprintNodeSpawners (generally when all those UBlueprintNodeSpawners wrap 
+ * the same UField). We do this to keep the menu less jumbled, and instead
+ * use the drag/drop action to present a sub-menu to the user (so they can pick 
+ * the the node type that they want). We do this with both delegates and 
+ * variable nodes (where the user can pick a getter vs. a setter, etc.).
+ *
+ * This class represents those "consolidated" actions, and essentially serves 
+ * as a FDragDropOperation spawner. It wraps a single UBlueprintNodeSpawner (any
+ * one of the set that it is supposed to represent), that it uses to determine
+ * the proper FDragDropOperation.
+ */
+USTRUCT()
+struct FBlueprintDragDropMenuItem : public FEdGraphSchemaAction
+{
+	GENERATED_USTRUCT_BODY();
+	
+public:
+	static FName StaticGetTypeId() { static FName const TypeId("FBlueprintDragDropMenuItem"); return TypeId; }
+	/** Default constructor for USTRUCT() purposes. */
+	FBlueprintDragDropMenuItem() : SampleAction(nullptr) {}
+
+	/**
+	 * Sole constructor, which ensures that the RepresentativeAction member is
+	 * set (and valid).
+	 * 
+	 * @param  SampleAction	One of the node-spawners that this "consolidated" menu item represents (shouldn't matter which one).
+	 * @param  MenuGrouping	Defines the order in which this entry is listed in the menu.
+	 */
+	FBlueprintDragDropMenuItem(FBlueprintActionContext const& Context, UBlueprintNodeSpawner* SampleAction, int32 MenuGrouping = 0);
+	
+	// FEdGraphSchemaAction interface
+	virtual FName GetTypeId() const final { return StaticGetTypeId(); }
+	virtual UEdGraphNode* PerformAction(UEdGraph* ParentGraph, UEdGraphPin* FromPin, FVector2D const Location, bool bSelectNewNode = true) override;
+	virtual UEdGraphNode* PerformAction(UEdGraph* ParentGraph, TArray<UEdGraphPin*>& FromPins, FVector2D const Location, bool bSelectNewNode = true) override;
+	virtual void  AddReferencedObjects(FReferenceCollector& Collector) override;
+	// End FEdGraphSchemaAction interface
+
+	/**
+	 * Retrieves the icon brush for this menu entry (to be displayed alongside
+	 * in the menu).
+	 *
+	 * @param  ColorOut	An output parameter that's filled in with the color to tint the brush with.
+	 * @return An slate brush to be used for this menu item in the action menu.
+	 */
+	FSlateBrush const* GetMenuIcon(FSlateColor& ColorOut);
+
+	/**
+	 * Attempts to create a certain drag/drop action with respect to the 
+	 * set UBlueprintNodeSpawner.
+	 * 
+	 * @param  AnalyticsDelegate	The analytics callback to assign the drag/drop operation.
+	 * @return An empty TSharedPtr if we failed to determine the FDragDropOperation type to use, else a newly instantiated drag/drop op.
+	 */
+	TSharedPtr<FDragDropOperation> OnDragged(FNodeCreationAnalytic AnalyticsDelegate) const;
+
+private:
+	/** An arbitrary member of the node-spawner subset that this menu entry represents. */
+	UBlueprintNodeSpawner const* SampleAction;
+};
diff --git a/Engine/Source/Editor/UMGEditor/Private/WidgetBlueprintCompiler.cpp b/Engine/Source/Editor/UMGEditor/Private/WidgetBlueprintCompiler.cpp
--- a/Engine/Source/Editor/UMGEditor/Private/WidgetBlueprintCompiler.cpp
+++ b/Engine/Source/Editor/UMGEditor/Private/WidgetBlueprintCompiler.cpp
@@ -7,8 +7,6 @@
 #include "Kismet2NameValidators.h"
 #include "KismetReinstanceUtilities.h"
 #include "MovieScene.h"
-
-#include "BlueprintPropertyNodeSpawner.h"
 #include "K2Node_FunctionEntry.h"
 
 
