Fix for build error regarding initialization order.

diff --git a/Engine/Source/Editor/Kismet/Private/BlueprintActionMenuItem.cpp b/Engine/Source/Editor/Kismet/Private/BlueprintActionMenuItem.cpp
--- a/Engine/Source/Editor/Kismet/Private/BlueprintActionMenuItem.cpp
+++ b/Engine/Source/Editor/Kismet/Private/BlueprintActionMenuItem.cpp
@@ -54,12 +54,12 @@ static void FBlueprintMenuActionItemImpl::DirtyBlueprintFromNewNode(UEdGraphNode
 
 //------------------------------------------------------------------------------
 FBlueprintActionMenuItem::FBlueprintActionMenuItem(UBlueprintNodeSpawner const* NodeSpawner, FSlateBrush const* MenuIcon, FSlateColor const& IconTintIn, int32 MenuGrouping/* = 0*/)
-	: Action(NodeSpawner)
+	: IconBrush(MenuIcon)
+	, IconTint(IconTintIn)
+	, Action(NodeSpawner)
 {
 	check(Action != nullptr);
 	Grouping  = MenuGrouping;
-	IconBrush = MenuIcon;
-	IconTint  = IconTintIn;
 }
 
 //------------------------------------------------------------------------------
diff --git a/Engine/Source/Editor/Kismet/Public/BlueprintActionMenuItem.h b/Engine/Source/Editor/Kismet/Public/BlueprintActionMenuItem.h
--- a/Engine/Source/Editor/Kismet/Public/BlueprintActionMenuItem.h
+++ b/Engine/Source/Editor/Kismet/Public/BlueprintActionMenuItem.h
@@ -26,7 +26,7 @@ public:
 	static FName StaticGetTypeId() { static FName const TypeId("FBlueprintActionMenuItem"); return TypeId; }
 	
 	/** Constructors */
-	FBlueprintActionMenuItem() : Action(nullptr), IconBrush(nullptr), IconTint(FLinearColor::White) {}
+	FBlueprintActionMenuItem() : IconBrush(nullptr), IconTint(FLinearColor::White), Action(nullptr) {}
 	FBlueprintActionMenuItem(UBlueprintNodeSpawner const* NodeSpawner, FSlateBrush const* MenuIcon, FSlateColor const& IconTint, int32 MenuGrouping = 0);
 	
 	// FEdGraphSchemaAction interface
