#UE4 CDO view reintroduced to merge tool, replaced set of strings with set of fnames

diff --git a/Engine/Source/Developer/Merge/Private/BlueprintMergeData.h b/Engine/Source/Developer/Merge/Private/BlueprintMergeData.h
--- a/Engine/Source/Developer/Merge/Private/BlueprintMergeData.h
+++ b/Engine/Source/Developer/Merge/Private/BlueprintMergeData.h
@@ -12,8 +12,8 @@ struct FBlueprintMergeData
 		, BlueprintLocal(NULL)
 		, BlueprintBase(NULL)
 		, RevisionBase()
-		, BlueprintNew(NULL)
-		, RevisionNew()
+		, BlueprintRemote(NULL)
+		, RevisionRemote()
 	{
 	}
 
@@ -22,15 +22,15 @@ struct FBlueprintMergeData
 		, const class UBlueprint*			InBlueprintLocal
 		, const class UBlueprint*			InBlueprintBase
 		, FRevisionInfo						InRevisionBase
-		, const class UBlueprint*			InBlueprintNew
-		, FRevisionInfo						InRevisionNew
+		, const class UBlueprint*			InBlueprintRemote
+		, FRevisionInfo						InRevisionRemote
 	)
-		: OwningEditor(		InOwningEditor	)
-		, BlueprintLocal(	InBlueprintLocal)
-		, BlueprintBase(	InBlueprintBase	)
-		, RevisionBase(		InRevisionBase	)
-		, BlueprintNew(		InBlueprintNew	)
-		, RevisionNew(		InRevisionNew	)
+		: OwningEditor(		InOwningEditor		)
+		, BlueprintLocal(	InBlueprintLocal	)
+		, BlueprintBase(	InBlueprintBase		)
+		, RevisionBase(		InRevisionBase		)
+		, BlueprintRemote(	InBlueprintRemote	)
+		, RevisionRemote(		InRevisionRemote)
 	{
 	}
 
@@ -40,6 +40,16 @@ struct FBlueprintMergeData
 	const class UBlueprint*				BlueprintBase;
 	FRevisionInfo						RevisionBase;
 
-	const class UBlueprint*				BlueprintNew;
-	FRevisionInfo						RevisionNew;
+	const class UBlueprint*				BlueprintRemote;
+	FRevisionInfo						RevisionRemote;
 };
+
+namespace EMergeParticipant
+{
+	enum Type
+	{
+		MERGE_PARTICIPANT_REMOTE,
+		MERGE_PARTICIPANT_BASE,
+		MERGE_PARTICIPANT_LOCAL,
+	};
+}
diff --git a/Engine/Source/Developer/Merge/Private/SBlueprintMerge.cpp b/Engine/Source/Developer/Merge/Private/SBlueprintMerge.cpp
--- a/Engine/Source/Developer/Merge/Private/SBlueprintMerge.cpp
+++ b/Engine/Source/Developer/Merge/Private/SBlueprintMerge.cpp
@@ -4,18 +4,18 @@
 
 #include "BlueprintEditor.h"
 #include "BlueprintEditorModes.h"
+#include "IDiffControl.h"
 #include "ISourceControlModule.h"
 #include "SBlueprintEditorToolbar.h"
 #include "SBlueprintMerge.h"
+#include "SMergeDetailsView.h"
 #include "SMergeGraphView.h"
 
 #define LOCTEXT_NAMESPACE "SBlueprintMerge"
 
 SBlueprintMerge::SBlueprintMerge()
 	: Data()
-	, GraphView()
-	, ComponentsView()
-	, DefaultsView()
+	, CurrentDiffControl( NULL )
 {
 }
 
@@ -25,9 +25,65 @@ void SBlueprintMerge::Construct(const FArguments InArgs, const FBlueprintMergeDa
 
 	Data = InData;
 
-	GraphView = SNew( SMergeGraphView, InData );
-	ComponentsView = SNew(SBorder);
-	DefaultsView = SNew(SBorder);
+	TSharedPtr<SMergeGraphView> GraphView = SNew( SMergeGraphView, InData );
+	GraphControl.Widget = GraphView;
+	GraphControl.DiffControl = GraphView.Get();
+
+	TreeControl.Widget = SNew(SBorder);
+
+	auto DetailsView = SNew( SMergeDetailsView, InData );
+	DetailsControl.Widget = DetailsView;
+	DetailsControl.DiffControl = &DetailsView.Get();
+
+	FToolBarBuilder ToolbarBuilder(TSharedPtr< FUICommandList >(), FMultiBoxCustomization::None);
+	ToolbarBuilder.AddToolBarButton( 
+		FUIAction( FExecuteAction::CreateSP(this, &SBlueprintMerge::PrevDiff), FCanExecuteAction::CreateSP(this, &SBlueprintMerge::CanCycleDiffs) )
+		, NAME_None
+		, LOCTEXT("PrevMergeLabel", "Prev")
+		, LOCTEXT("PrevMergeTooltip", "Go to previous difference")
+		, FSlateIcon(FEditorStyle::GetStyleSetName(), "BlueprintMerge.PrevDiff")
+	);
+	ToolbarBuilder.AddToolBarButton(
+		FUIAction(FExecuteAction::CreateSP(this, &SBlueprintMerge::NextDiff), FCanExecuteAction::CreateSP(this, &SBlueprintMerge::CanCycleDiffs))
+		, NAME_None
+		, LOCTEXT("NextMergeLabel", "Next")
+		, LOCTEXT("NextMergeTooltip", "Go to next difference")
+		, FSlateIcon(FEditorStyle::GetStyleSetName(), "BlueprintMerge.NextDiff") 
+	);
+	ToolbarBuilder.AddSeparator();
+	ToolbarBuilder.AddToolBarButton(
+		FUIAction(FExecuteAction::CreateSP(this, &SBlueprintMerge::OnAcceptResultClicked) )
+		, NAME_None
+		, LOCTEXT("FinishMergeLabel", "Finish Merge")
+		, LOCTEXT("FinishMergeTooltip", "Complete the merge operation - saves the blueprint and resolves the conflict with the SCC provider")
+		, FSlateIcon(FEditorStyle::GetStyleSetName(), "BlueprintMerge.Finish")
+	);
+	ToolbarBuilder.AddToolBarButton(
+		FUIAction(FExecuteAction::CreateSP(this, &SBlueprintMerge::OnCancelClicked))
+		, NAME_None
+		, LOCTEXT("CancelMergeLabel", "Cancel")
+		, LOCTEXT("CancelMergeTooltip", "Abort the merge operation")
+		, FSlateIcon(FEditorStyle::GetStyleSetName(), "BlueprintMerge.Cancel")
+	);
+
+	ChildSlot [
+		SNew( SVerticalBox )
+		+ SVerticalBox::Slot()
+		.AutoHeight()
+		[
+			SNew( SHorizontalBox )
+			+ SHorizontalBox::Slot()
+			.AutoWidth()
+			[
+				ToolbarBuilder.MakeWidget()
+			]
+		]
+		+ SVerticalBox::Slot()
+		[
+			SAssignNew(MainView, SBorder)
+			.BorderImage(FEditorStyle::GetBrush("NoBorder"))
+		]
+	];
 
 	InData.OwningEditor.Pin()->OnModeSet().AddSP( this, &SBlueprintMerge::OnModeChanged );
 	OnModeChanged( InData.OwningEditor.Pin()->GetCurrentMode() );
@@ -61,11 +117,26 @@ UBlueprint* SBlueprintMerge::GetTargetBlueprint()
 	return Data.OwningEditor.Pin()->GetBlueprintObj();
 }
 
-FReply SBlueprintMerge::OnAcceptResultClicked()
+void SBlueprintMerge::NextDiff()
+{
+	CurrentDiffControl->NextDiff();
+}
+
+void SBlueprintMerge::PrevDiff()
+{
+	CurrentDiffControl->PrevDiff();
+}
+
+bool SBlueprintMerge::CanCycleDiffs() const
+{
+	return CurrentDiffControl && CurrentDiffControl->HasDifferences();
+}
+
+void SBlueprintMerge::OnAcceptResultClicked()
 {
 	// Because merge operations are so destructive and can be confusing, lets write backups of the files involved:
 	const FString BackupSubDir = FPaths::GameSavedDir() / TEXT("Backup") / TEXT("Resolve_Backup[") + FDateTime::Now().ToString(TEXT("%Y-%m-%d-%H-%M-%S")) + TEXT("]");
-	WriteBackup(*Data.BlueprintNew->GetOutermost(), BackupSubDir, TEXT("RemoteAsset") + FPackageName::GetAssetPackageExtension() );
+	WriteBackup(*Data.BlueprintRemote->GetOutermost(), BackupSubDir, TEXT("RemoteAsset") + FPackageName::GetAssetPackageExtension() );
 	WriteBackup(*Data.BlueprintBase->GetOutermost(), BackupSubDir, TEXT("CommonBaseAsset") + FPackageName::GetAssetPackageExtension() );
 	WriteBackup(*Data.BlueprintLocal->GetOutermost(), BackupSubDir, TEXT("LocalAsset") + FPackageName::GetAssetPackageExtension());
 
@@ -87,14 +158,11 @@ FReply SBlueprintMerge::OnAcceptResultClicked()
 	}
 
 	Data.OwningEditor.Pin()->CloseMergeTool();
-
-	return FReply::Handled();
 }
 
-FReply SBlueprintMerge::OnCancelClicked()
+void SBlueprintMerge::OnCancelClicked()
 {
 	Data.OwningEditor.Pin()->CloseMergeTool();
-	return FReply::Handled();
 }
 
 void SBlueprintMerge::OnModeChanged(FName NewMode)
@@ -102,25 +170,19 @@ void SBlueprintMerge::OnModeChanged(FName NewMode)
 	if (NewMode == FBlueprintEditorApplicationModes::StandardBlueprintEditorMode ||
 		NewMode == FBlueprintEditorApplicationModes::BlueprintMacroMode)
 	{
-		this->ChildSlot
-			[
-				GraphView.ToSharedRef()
-			];
+		MainView->SetContent( GraphControl.Widget.ToSharedRef() );
+		CurrentDiffControl = GraphControl.DiffControl;
 	}
-	else if (NewMode == FBlueprintEditorApplicationModes::BlueprintDefaultsMode)
+	else if (NewMode == FBlueprintEditorApplicationModes::BlueprintComponentsMode)
 	{
-		this->ChildSlot
-			[
-				DefaultsView.ToSharedRef()
-			];
+		MainView->SetContent(TreeControl.Widget.ToSharedRef());
+		CurrentDiffControl = TreeControl.DiffControl;
 	}
-	else if (NewMode == FBlueprintEditorApplicationModes::BlueprintComponentsMode ||
-			 NewMode == FBlueprintEditorApplicationModes::BlueprintInterfaceMode)
+	else if (NewMode == FBlueprintEditorApplicationModes::BlueprintDefaultsMode ||
+			NewMode == FBlueprintEditorApplicationModes::BlueprintInterfaceMode)
 	{
-		this->ChildSlot
-			[
-				ComponentsView.ToSharedRef()
-			];
+		MainView->SetContent(DetailsControl.Widget.ToSharedRef());
+		CurrentDiffControl = DetailsControl.DiffControl;
 	}
 	else
 	{
diff --git a/Engine/Source/Developer/Merge/Private/SBlueprintMerge.h b/Engine/Source/Developer/Merge/Private/SBlueprintMerge.h
--- a/Engine/Source/Developer/Merge/Private/SBlueprintMerge.h
+++ b/Engine/Source/Developer/Merge/Private/SBlueprintMerge.h
@@ -21,15 +21,34 @@ private:
 	UBlueprint* GetTargetBlueprint();
 
 	/** Event handlers */
-	FReply OnAcceptResultClicked();
-	FReply OnCancelClicked();
+	void NextDiff();
+	void PrevDiff();
+	bool CanCycleDiffs() const;
+	void OnAcceptResultClicked();
+	void OnCancelClicked();
 	void OnModeChanged(FName NewMode);
 
 	FBlueprintMergeData Data;
 
-	TSharedPtr<SWidget>		GraphView;
-	TSharedPtr<SWidget>		ComponentsView;
-	TSharedPtr<SWidget>		DefaultsView;
+	TSharedPtr<SBorder>		MainView;
+
+	struct FDiffControl
+	{
+		FDiffControl()
+			: Widget()
+			, DiffControl(NULL)
+		{
+		}
+
+		TSharedPtr<SWidget> Widget;
+		class IDiffControl* DiffControl;
+	};
+
+	FDiffControl GraphControl;
+	FDiffControl TreeControl;
+	FDiffControl DetailsControl;
+
+	class IDiffControl* CurrentDiffControl;
 
 	// This has to be allocated here because SListView cannot own the list
 	// that it is displaying. It also seems like the display list *has*
diff --git a/Engine/Source/Developer/Merge/Private/SMergeDetailsView.cpp b/Engine/Source/Developer/Merge/Private/SMergeDetailsView.cpp
new file mode 100644
--- /dev/null
+++ b/Engine/Source/Developer/Merge/Private/SMergeDetailsView.cpp
@@ -0,0 +1,172 @@
+// Copyright 1998-2014 Epic Games, Inc. All Rights Reserved.
+
+#include "MergePrivatePCH.h"
+
+#include "DetailsDiff.h"
+#include "SBlueprintDiff.h"
+#include "SMergeDetailsView.h"
+
+void SMergeDetailsView::Construct(const FArguments InArgs, const FBlueprintMergeData& InData)
+{
+	Data = InData;
+
+	const UObject* RemoteCDO = DiffUtils::GetCDO(InData.BlueprintRemote);
+	TMap< FName, const UProperty* > RemotePropertyMap = DiffUtils::GetProperties(RemoteCDO);
+	const UObject* BaseCDO = DiffUtils::GetCDO(InData.BlueprintBase);
+	TMap< FName, const UProperty* > BasePropertyMap = DiffUtils::GetProperties(BaseCDO);
+	const UObject* LocalCDO = DiffUtils::GetCDO(InData.BlueprintLocal);
+	TMap< FName, const UProperty* > LocalPropertyMap = DiffUtils::GetProperties(LocalCDO);
+
+	TArray<FName> RemoteIdenticalProperties, RemoteDifferingProperties;
+	DiffUtils::CompareUnrelatedObjects(BaseCDO, BasePropertyMap, RemoteCDO, RemotePropertyMap, RemoteIdenticalProperties, RemoteDifferingProperties);
+	TArray<FName> LocalIdenticalProperties, LocalDifferingProperties;
+	DiffUtils::CompareUnrelatedObjects(BaseCDO, BasePropertyMap, LocalCDO, LocalPropertyMap, LocalIdenticalProperties, LocalDifferingProperties);
+
+	TSet<FName> RemoteIdenticalPropertiesSet(RemoteIdenticalProperties);
+	TSet<FName> LocalIdenticalPropertiesSet(LocalIdenticalProperties);
+	TSet<FName> BaseIdenticalPropertiesSet(RemoteIdenticalPropertiesSet.Intersect(LocalIdenticalPropertiesSet));
+	
+	/*	
+		DifferingProperties is an ordered list of all differing properties (properties added, removed or changed) in remote or local.
+		Strictly speaking it's impossible to guarantee that we'll traverse remote and local differences in the same order (for instance
+		because property layout could somehow change between revisions) but in practice the following works:
+
+			1. Iterate properties in base, add any properties that differ in remote or local to DifferingProperties.
+			2. Iterate properties in remote, add any new properties to DifferingProperties
+			3. Iterate properties in local, add any new properties to DifferingProperties, if they have not already been added
+	*/
+	const auto AddPropertiesOrdered = []( UProperty const* Property, const TSet<FName>& IdenticalProperties, TArray<FName>& ResultingProperties )
+	{
+		if ((!Property->HasAnyPropertyFlags(CPF_Parm) && Property->HasAnyPropertyFlags(CPF_Edit)))
+		{
+			// contains check here is O(n), so we're needlessly n^2:
+			if (!ResultingProperties.Contains(Property->GetFName()))
+			{
+				if (!IdenticalProperties.Contains(Property->GetFName()))
+				{
+					ResultingProperties.Add(Property->GetFName());
+				}
+			}
+		}
+	};
+
+	if( BaseCDO )
+	{
+		const UClass* Class = BaseCDO->GetClass();
+		for (TFieldIterator<UProperty> PropertyIt(Class, EFieldIteratorFlags::IncludeSuper); PropertyIt; ++PropertyIt)
+		{
+			AddPropertiesOrdered( *PropertyIt, BaseIdenticalPropertiesSet, DifferingProperties );
+		}
+	}
+	if( RemoteCDO )
+	{
+		const UClass* Class = RemoteCDO->GetClass();
+		for (TFieldIterator<UProperty> PropertyIt(Class, EFieldIteratorFlags::IncludeSuper); PropertyIt; ++PropertyIt)
+		{
+			AddPropertiesOrdered(*PropertyIt, RemoteIdenticalPropertiesSet, DifferingProperties);
+		}
+	}
+	if (LocalCDO)
+	{
+		const UClass* Class = LocalCDO->GetClass();
+		for (TFieldIterator<UProperty> PropertyIt(Class, EFieldIteratorFlags::IncludeSuper); PropertyIt; ++PropertyIt)
+		{
+			AddPropertiesOrdered(*PropertyIt, LocalIdenticalPropertiesSet, DifferingProperties);
+		}
+	}
+	CurrentDifference = -1;
+
+	// EMergeParticipant::MERGE_PARTICIPANT_REMOTE
+	{
+		DetailsViews.Add(
+			FDetailsDiff( RemoteCDO, RemotePropertyMap, RemoteIdenticalPropertiesSet)
+		);
+	}
+	// EMergeParticipant::MERGE_PARTICIPANT_BASE
+	{
+		DetailsViews.Add(
+			FDetailsDiff(BaseCDO, BasePropertyMap, BaseIdenticalPropertiesSet)
+		);
+	}
+	// EMergeParticipant::MERGE_PARTICIPANT_LOCAL
+	{
+		DetailsViews.Add(
+			FDetailsDiff(LocalCDO, LocalPropertyMap, LocalIdenticalPropertiesSet)
+		);
+	}
+
+	ChildSlot[
+		SNew(SSplitter)
+		+ SSplitter::Slot()
+		[
+			GetRemoteDetails().DetailsWidget()
+		] 
+		+ SSplitter::Slot()
+		[
+			GetBaseDetails().DetailsWidget()
+		]
+		+ SSplitter::Slot()
+		[
+			GetLocalDetails().DetailsWidget()
+		]
+	];
+}
+
+void SMergeDetailsView::NextDiff()
+{
+	if (DifferingProperties.Num() == 0)
+	{
+		return;
+	}
+
+	CurrentDifference = (CurrentDifference + 1) % DifferingProperties.Num();
+
+	HighlightCurrentDifference();
+}
+
+void SMergeDetailsView::PrevDiff()
+{
+	if (DifferingProperties.Num() == 0)
+	{
+		return;
+	}
+
+	--CurrentDifference;
+	if (CurrentDifference < 0)
+	{
+		CurrentDifference = DifferingProperties.Num() - 1;
+	}
+
+	HighlightCurrentDifference();
+}
+
+void SMergeDetailsView::HighlightCurrentDifference()
+{
+	const FName PropertyName = DifferingProperties[CurrentDifference];
+
+	for( auto& DetailDiff : DetailsViews )
+	{
+		DetailDiff.HighlightProperty(PropertyName);
+	}
+}
+
+bool SMergeDetailsView::HasDifferences() const
+{
+	return DifferingProperties.Num() != 0;
+}
+
+FDetailsDiff& SMergeDetailsView::GetRemoteDetails()
+{
+	return DetailsViews[EMergeParticipant::MERGE_PARTICIPANT_REMOTE];
+}
+
+FDetailsDiff& SMergeDetailsView::GetBaseDetails()
+{
+	return DetailsViews[EMergeParticipant::MERGE_PARTICIPANT_BASE];
+}
+
+FDetailsDiff& SMergeDetailsView::GetLocalDetails()
+{
+	return DetailsViews[EMergeParticipant::MERGE_PARTICIPANT_LOCAL];
+}
+
diff --git a/Engine/Source/Developer/Merge/Private/SMergeDetailsView.h b/Engine/Source/Developer/Merge/Private/SMergeDetailsView.h
new file mode 100644
--- /dev/null
+++ b/Engine/Source/Developer/Merge/Private/SMergeDetailsView.h
@@ -0,0 +1,36 @@
+// Copyright 1998-2014 Epic Games, Inc. All Rights Reserved.
+#pragma once
+
+#include "BlueprintMergeData.h"
+#include "IDiffControl.h"
+#include "DetailsDiff.h"
+
+class SMergeDetailsView : public SCompoundWidget
+						, public IDiffControl
+{
+public:
+	virtual ~SMergeDetailsView() {}
+
+	SLATE_BEGIN_ARGS(SMergeDetailsView)
+	{}
+	SLATE_END_ARGS()
+
+	void Construct(const FArguments InArgs, const FBlueprintMergeData& InData);
+private:
+	/** Implementation of IDiffControl: */
+	void NextDiff() override;
+	void PrevDiff() override;
+	bool HasDifferences() const override;
+
+	void HighlightCurrentDifference();
+	class FDetailsDiff& GetRemoteDetails();
+	class FDetailsDiff& GetBaseDetails();
+	class FDetailsDiff& GetLocalDetails();
+
+	FBlueprintMergeData Data;
+	TArray< FDetailsDiff > DetailsViews;
+
+	/** These have been duplicated from FCDODiffControl, opportunity to refactor exists: */
+	TArray< FName > DifferingProperties;
+	int CurrentDifference;
+};
diff --git a/Engine/Source/Developer/Merge/Private/SMergeGraphView.cpp b/Engine/Source/Developer/Merge/Private/SMergeGraphView.cpp
--- a/Engine/Source/Developer/Merge/Private/SMergeGraphView.cpp
+++ b/Engine/Source/Developer/Merge/Private/SMergeGraphView.cpp
@@ -208,10 +208,11 @@ void SMergeGraphView::Construct(const FArguments InArgs, const FBlueprintMergeDa
 {
 	Data = InData;
 	bViewsAreLocked = true;
+	DiffResultsListData = NULL;
 
 	TArray<FBlueprintRevPair> BlueprintsForDisplay;
 	// EMergeParticipant::MERGE_PARTICIPANT_REMOTE
-	BlueprintsForDisplay.Add(FBlueprintRevPair(InData.BlueprintNew, InData.RevisionNew));
+	BlueprintsForDisplay.Add(FBlueprintRevPair(InData.BlueprintRemote, InData.RevisionRemote));
 	// EMergeParticipant::MERGE_PARTICIPANT_BASE
 	BlueprintsForDisplay.Add(FBlueprintRevPair(InData.BlueprintBase, InData.RevisionBase));
 	// EMergeParticipant::MERGE_PARTICIPANT_LOCAL
@@ -356,6 +357,24 @@ void SMergeGraphView::Construct(const FArguments InArgs, const FBlueprintMergeDa
 		];
 }
 
+void SMergeGraphView::NextDiff()
+{
+	check(DiffResultList.Pin().IsValid() && DiffResultsListData);
+	DiffWidgetUtils::SelectNextRow( *DiffResultList.Pin(), *DiffResultsListData );
+}
+
+void SMergeGraphView::PrevDiff()
+{
+	check(DiffResultList.Pin().IsValid() && DiffResultsListData);
+	DiffWidgetUtils::SelectPrevRow( *DiffResultList.Pin(), *DiffResultsListData );
+}
+
+bool SMergeGraphView::HasDifferences() const
+{
+	auto DiffResultListPtr = DiffResultList.Pin();
+	return DiffResultListPtr.IsValid() && DiffResultListPtr->GetNumItemsBeingObserved() > 0;
+}
+
 void SMergeGraphView::OnGraphListSelectionChanged(TSharedPtr<FMergeGraphRowEntry> Item, ESelectInfo::Type SelectionType)
 {
 	if (SelectionType != ESelectInfo::OnMouseClick || !Item.IsValid())
@@ -375,29 +394,30 @@ void SMergeGraphView::OnGraphListSelectionChanged(TSharedPtr<FMergeGraphRowEntry
 
 	LockViews(DiffPanels, bViewsAreLocked);
 
+	DiffResultsListData = &(Item->LocalDifferences);
 	DiffResultsWidget->SetContent(
-		SNew(SListView< TSharedPtr<FDiffSingleResult> >)
+		SAssignNew(DiffResultList, SListView< TSharedPtr<FDiffSingleResult> >)
 		.ItemHeight(24)
 		.ListItemsSource(&Item->LocalDifferences)
 		.OnGenerateRow(SListView< TSharedPtr<FDiffSingleResult> >::FOnGenerateRow::CreateStatic(
-		[](TSharedPtr<FDiffSingleResult> ParamItem, const TSharedRef<STableViewBase>& OwnerTable) -> TSharedRef<ITableRow>
-	{
-		FDiffResultItem WidgetGenerator(*ParamItem);
-		return SNew(STableRow< TSharedPtr<FDiffSingleResult> >, OwnerTable)
-			.Content()
-			[
-				WidgetGenerator.GenerateWidget()
-			];
-	}
-	))
+			[](TSharedPtr<FDiffSingleResult> ParamItem, const TSharedRef<STableViewBase>& OwnerTable) -> TSharedRef<ITableRow>
+			{
+				FDiffResultItem WidgetGenerator(*ParamItem);
+				return SNew(STableRow< TSharedPtr<FDiffSingleResult> >, OwnerTable)
+					.Content()
+					[
+						WidgetGenerator.GenerateWidget()
+					];
+			}
+		))
 		.SelectionMode(ESelectionMode::Single)
 		.OnSelectionChanged(this, &SMergeGraphView::OnDiffListSelectionChanged)
-		);
+	);
 }
 
 void SMergeGraphView::OnDiffListSelectionChanged(TSharedPtr<FDiffSingleResult> Item, ESelectInfo::Type SelectionType)
 {
-	if (SelectionType != ESelectInfo::OnMouseClick || !Item.IsValid())
+	if (!Item.IsValid())
 	{
 		return;
 	}
diff --git a/Engine/Source/Developer/Merge/Private/SMergeGraphView.h b/Engine/Source/Developer/Merge/Private/SMergeGraphView.h
--- a/Engine/Source/Developer/Merge/Private/SMergeGraphView.h
+++ b/Engine/Source/Developer/Merge/Private/SMergeGraphView.h
@@ -2,19 +2,11 @@
 #pragma once
 
 #include "BlueprintMergeData.h"
+#include "IDiffControl.h"
 #include "SBlueprintDiff.h"
 
-namespace EMergeParticipant
-{
-	enum Type
-	{
-		MERGE_PARTICIPANT_REMOTE,
-		MERGE_PARTICIPANT_BASE,
-		MERGE_PARTICIPANT_LOCAL,
-	};
-}
-
 class SMergeGraphView : public SCompoundWidget
+						, public IDiffControl
 {
 public:
 	SLATE_BEGIN_ARGS(SMergeGraphView)
@@ -24,6 +16,12 @@ public:
 	void Construct(const FArguments InArgs, const FBlueprintMergeData& InData);
 
 private:
+	/** Implementation of IDiffControl: */
+	void NextDiff() override;
+	void PrevDiff() override;
+	bool HasDifferences() const override;
+
+	/** Helper functions and event handlers: */
 	void OnGraphListSelectionChanged(TSharedPtr<struct FMergeGraphRowEntry> Item, ESelectInfo::Type SelectionType);
 	void OnDiffListSelectionChanged(TSharedPtr<struct FDiffSingleResult> Item, ESelectInfo::Type SelectionType);
 
@@ -40,5 +38,23 @@ private:
 	TArray< TSharedPtr< struct FMergeGraphRowEntry> > DifferencesFromBase;
 	TSharedPtr<class SBorder> DiffResultsWidget;
 
+	TWeakPtr < SListView< TSharedPtr< struct FDiffSingleResult> > > DiffResultList;
+
+	/** 
+	 * Unfortunately this innocuous member requires some explanation:
+	 * 	This list is owned by the FMergeGraphRowEntry and displayed by DiffResultList, but we have
+	 * 	to keep track of it because SListView doesn't expose the list it's displaying, and we
+	 * 	don't have any need to interact with the ListView that displays the DifferencesFromBase
+	 * 	so we can't access the current FMergeGraphRowEntry. Basically we have three options:
+	 * 	
+	 * 	1. Cache the list view when a new FMergeGraphRowEntry is selected
+	 * 	2. Cache a reference to the ListView that displays DifferencesFromBase, solely so we can
+	 * 		access the currently selected item and dig out the diffs associated with that graph
+	 * 	3. Expose the list being displayed by the listview (requires changing slate)
+	 * 
+	 * 	I have chosen option 1.
+	 */
+	TArray< TSharedPtr< struct FDiffSingleResult> > const* DiffResultsListData;
+
 	bool bViewsAreLocked;
 };
diff --git a/Engine/Source/Editor/EditorStyle/Private/SlateEditorStyle.cpp b/Engine/Source/Editor/EditorStyle/Private/SlateEditorStyle.cpp
--- a/Engine/Source/Editor/EditorStyle/Private/SlateEditorStyle.cpp
+++ b/Engine/Source/Editor/EditorStyle/Private/SlateEditorStyle.cpp
@@ -4283,6 +4283,15 @@ void FSlateEditorStyle::FStyle::SetupPersonaStyle()
 		Set( "BlueprintDif.NextDiff", new IMAGE_BRUSH( "/Icons/diff_next_40x", Icon16x16 )  );
 		Set( "BlueprintDif.PrevDiff", new IMAGE_BRUSH( "/Icons/diff_prev_40x", Icon16x16 )  );
 	}
+
+	//Blueprint Merge
+	{
+		Set("BlueprintMerge.NextDiff", new IMAGE_BRUSH("/Icons/diff_next_40x", Icon16x16));
+		Set("BlueprintMerge.PrevDiff", new IMAGE_BRUSH("/Icons/diff_prev_40x", Icon16x16));
+		Set("BlueprintMerge.Finish", new IMAGE_BRUSH("/Icons/LV_Save", Icon16x16));
+		Set("BlueprintMerge.Cancel", new IMAGE_BRUSH("/Icons/LV_Remove", Icon16x16));
+	}
+
 	// Play in editor / play in world
 	{
 		Set( "PlayWorld.Simulate", new IMAGE_BRUSH( "Icons/icon_simulate_40x", Icon40x40 ) );
diff --git a/Engine/Source/Editor/Kismet/Private/DetailsDiff.cpp b/Engine/Source/Editor/Kismet/Private/DetailsDiff.cpp
new file mode 100644
--- /dev/null
+++ b/Engine/Source/Editor/Kismet/Private/DetailsDiff.cpp
@@ -0,0 +1,32 @@
+// Copyright 1998-2014 Epic Games, Inc. All Rights Reserved.
+
+#include "BlueprintEditorPrivatePCH.h"
+#include "DetailsDiff.h"
+#include "IDetailsView.h"
+#include "PropertyEditorModule.h"
+
+FDetailsDiff::FDetailsDiff(const UObject* InObject, const TMap< FName, const UProperty* >& InPropertyMap, const TSet< FName >& InIdenticalProperties)
+	: PropertyMap( InPropertyMap )
+	, DetailsView()
+{
+	FDetailsViewArgs DetailsViewArgs;
+	DetailsViewArgs.bShowDifferingPropertiesOption = true;
+
+	FPropertyEditorModule& EditModule = FModuleManager::Get().GetModuleChecked<FPropertyEditorModule>("PropertyEditor");
+	DetailsView = EditModule.CreateDetailView(DetailsViewArgs);
+	DetailsView->SetIsPropertyEditingEnabledDelegate(FIsPropertyEditingEnabled::CreateStatic([]{return false; }));
+	// This is a read only details view (see the property editing delegate), but it is not const correct:
+	DetailsView->SetObject(const_cast<UObject*>(InObject));
+	DetailsView->UpdateIdenticalProperties(InIdenticalProperties);
+}
+
+void FDetailsDiff::HighlightProperty(FName PropertyName)
+{
+	const UProperty** Property = PropertyMap.Find(PropertyName);
+	DetailsView->HighlightProperty( Property ? *Property : NULL );
+}
+
+TSharedRef< SWidget > FDetailsDiff::DetailsWidget()
+{
+	return DetailsView.ToSharedRef();
+}
diff --git a/Engine/Source/Editor/Kismet/Private/SBlueprintDiff.cpp b/Engine/Source/Editor/Kismet/Private/SBlueprintDiff.cpp
--- a/Engine/Source/Editor/Kismet/Private/SBlueprintDiff.cpp
+++ b/Engine/Source/Editor/Kismet/Private/SBlueprintDiff.cpp
@@ -6,6 +6,7 @@
 #include "BlueprintEditorCommands.h"
 #include "BlueprintEditorModes.h"
 #include "BlueprintUtilities.h"
+#include "DetailsDiff.h"
 #include "Editor/Kismet/Public/SBlueprintEditorToolbar.h"
 #include "Editor/Kismet/Public/SKismetInspector.h"
 #include "Editor/PropertyEditor/Public/IDetailsView.h"
@@ -13,6 +14,7 @@
 #include "Editor/UnrealEd/Public/EdGraphUtilities.h"
 #include "GraphDiffControl.h"
 #include "GraphEditor.h"
+#include "IDiffControl.h"
 #include "Kismet2/BlueprintEditorUtils.h"
 #include "SBlueprintDiff.h"
 #include "SGraphTitleBar.h"
@@ -21,55 +23,53 @@
 
 #define LOCTEXT_NAMESPACE "SBlueprintDif"
 
-class IDiffHighlighter
-{
-public:
-	virtual void NextDiff() = 0;
-	virtual void PrevDiff() = 0;
-	virtual bool HasDifferences() const = 0;
-};
+typedef TMap< FName, const UProperty* > FNamePropertyMap;
 
-typedef TMap< FString, const UProperty* > FNamePropertyMap;
-
-class CDODiffHighlighter : public TSharedFromThis<CDODiffHighlighter>
-						, public IDiffHighlighter
+class FCDODiffControl : public TSharedFromThis<FCDODiffControl>
+	, public IDiffControl
 {
 public:
-	CDODiffHighlighter( 
-			TWeakPtr<IDetailsView> InOldView
-			, FNamePropertyMap InOldProperties
-			, TWeakPtr<IDetailsView> InNewView
-			, FNamePropertyMap InNewProperties
-			, TArray< FString > InDifferingProperties )
-		: OldView( InOldView )
-		, OldProperties( InOldProperties )
-		, NewView(InNewView)
-		, NewProperties(InNewProperties)
-		, DifferingProperties( InDifferingProperties )
-		, CurrentDifference( -1 )
-	{
-	}
+	FCDODiffControl( const UObject* InOldCDO
+					, FNamePropertyMap InOldProperties
+					, const UObject* InNewCDO
+					, FNamePropertyMap InNewProperities
+					, TSet<FName> InIdenticalProperties
+					, TArray<FName> InDifferingProperties );
 
-	virtual ~CDODiffHighlighter() {}
+	TSharedRef<SWidget> OldDetailsWidget() { return OldDetails.DetailsWidget(); }
+	TSharedRef<SWidget> NewDetailsWidget() { return NewDetails.DetailsWidget(); }
 
+	virtual ~FCDODiffControl() {}
 private:
 	void NextDiff() override;
 	void PrevDiff() override;
 	bool HasDifferences() const override;
 
 	void HighlightCurrentDifference();
+	FDetailsDiff OldDetails;
+	FDetailsDiff NewDetails;
 
-	TWeakPtr<IDetailsView> OldView;
-	FNamePropertyMap OldProperties;
-	TWeakPtr<IDetailsView> NewView;
-	FNamePropertyMap NewProperties;
-	TArray< FString > DifferingProperties;
+	TArray< FName > DifferingProperties;
 	int CurrentDifference;
 };
 
-void CDODiffHighlighter::NextDiff()
+FCDODiffControl::FCDODiffControl( 
+		const UObject* InOldCDO
+		, FNamePropertyMap InOldProperties
+		, const UObject* InNewCDO
+		, FNamePropertyMap InNewProperities
+		, TSet<FName> InIdenticalProperties
+		, TArray<FName> InDifferingProperties )
+	: OldDetails( InOldCDO, InOldProperties, InIdenticalProperties )
+	, NewDetails( InNewCDO, InNewProperities, InIdenticalProperties )
+	, DifferingProperties( InDifferingProperties )
+	, CurrentDifference( -1 )
 {
-	if( DifferingProperties.Num() == 0 )
+}
+
+void FCDODiffControl::NextDiff()
+{
+	if (DifferingProperties.Num() == 0)
 	{
 		return;
 	}
@@ -79,7 +79,7 @@ void CDODiffHighlighter::NextDiff()
 	HighlightCurrentDifference();
 }
 
-void CDODiffHighlighter::PrevDiff()
+void FCDODiffControl::PrevDiff()
 {
 	if (DifferingProperties.Num() == 0)
 	{
@@ -87,7 +87,7 @@ void CDODiffHighlighter::PrevDiff()
 	}
 
 	--CurrentDifference;
-	if( CurrentDifference < 0 )
+	if (CurrentDifference < 0)
 	{
 		CurrentDifference = DifferingProperties.Num() - 1;
 	}
@@ -95,25 +95,22 @@ void CDODiffHighlighter::PrevDiff()
 	HighlightCurrentDifference();
 }
 
-bool CDODiffHighlighter::HasDifferences() const
+bool FCDODiffControl::HasDifferences() const
 {
 	return DifferingProperties.Num() != 0;
 }
 
-void CDODiffHighlighter::HighlightCurrentDifference()
+void FCDODiffControl::HighlightCurrentDifference()
 {
-	const FString& PropertyName = DifferingProperties[CurrentDifference];
+	const FName PropertyName = DifferingProperties[CurrentDifference];
 
-	const UProperty** OldProperty = OldProperties.Find(PropertyName);
-	const UProperty** NewProperty = NewProperties.Find(PropertyName);
-
-	OldView.Pin()->HighlightProperty(OldProperty ? *OldProperty : NULL);
-	NewView.Pin()->HighlightProperty(NewProperty ? *NewProperty : NULL);
+	OldDetails.HighlightProperty(PropertyName);
+	NewDetails.HighlightProperty(PropertyName);
 }
 
 /*List item that entry for a graph*/
 struct KISMET_API FListItemGraphToDiff	: public TSharedFromThis<FListItemGraphToDiff>
-										, public IDiffHighlighter
+										, public IDiffControl
 {
 	FListItemGraphToDiff(class SBlueprintDiff* Diff, class UEdGraph* GraphOld, class UEdGraph* GraphNew, const FRevisionInfo& RevisionOld, const FRevisionInfo& RevisionNew);
 	virtual ~FListItemGraphToDiff();
@@ -198,39 +195,24 @@ private:
 	TSharedPtr< FUICommandList > KeyCommands;
 };
 
-static FNamePropertyMap GetProperties( const UObject& ForObj )
+FNamePropertyMap DiffUtils::GetProperties( const UObject* ForObj )
 {
 	FNamePropertyMap Ret;
-	if (const UClass* Class = ForObj.GetClass())
+	if( ForObj )
 	{
+		const UClass* Class = ForObj->GetClass();
 		for (TFieldIterator<UProperty> PropertyIt(Class, EFieldIteratorFlags::IncludeSuper); PropertyIt; ++PropertyIt)
 		{
 			if ((!PropertyIt->HasAnyPropertyFlags(CPF_Parm) && PropertyIt->HasAnyPropertyFlags(CPF_Edit)))
 			{
-				Ret.Add(PropertyIt->GetName(), *PropertyIt);
+				Ret.Add(PropertyIt->GetFName(), *PropertyIt);
 			}
 		}
 	}
 	return Ret;
 }
 
-static FNamePropertyMap GetCDOProperties( const UBlueprint* ForBlueprint )
-{
-	if( ForBlueprint )
-	{
-		if( const UClass* GeneratedClass = ForBlueprint->GeneratedClass )
-		{
-			if( const UObject* CDO = GeneratedClass->ClassDefaultObject )
-			{
-				return GetProperties(*CDO);
-			}
-		}
-	}
-
-	return FNamePropertyMap();
-}
-
-static void GetPropertyNameSet( const UObject& ForObj, TSet< FString >& OutProperties )
+static void GetPropertyNameSet( const UObject& ForObj, TSet< FName >& OutProperties )
 {
 	if( const UClass* Class = ForObj.GetClass() )
 	{
@@ -238,13 +220,13 @@ static void GetPropertyNameSet( const UObject& ForObj, TSet< FString >& OutPrope
 		{
 			if ((!PropertyIt->HasAnyPropertyFlags(CPF_Parm) && PropertyIt->HasAnyPropertyFlags(CPF_Edit)))
 			{
-				OutProperties.Add(PropertyIt->GetName());
+				OutProperties.Add(PropertyIt->GetFName());
 			}
 		}
 	}
 }
 
-const UObject* GetCDO( const UBlueprint* ForBlueprint )
+const UObject* DiffUtils::GetCDO( const UBlueprint* ForBlueprint )
 {
 	if( !ForBlueprint 
 		|| !ForBlueprint->GeneratedClass )
@@ -550,16 +532,15 @@ void FListItemGraphToDiff::KeyWasPressed( const FKeyboardEvent& InKeyboardEvent
 	}
 }
 
-
 FDiffPanel::FDiffPanel()
 {
 	Blueprint = NULL;
 	LastFocusedPin = NULL;
 }
 
-static void CompareUnrelatedObjects( UObject const* A, const FNamePropertyMap& PropertyMapA, UObject const* B, const FNamePropertyMap& PropertyMapB, TArray<FString> &OutIdenticalProperties, TArray<FString> &OutDifferingProperties )
+void DiffUtils::CompareUnrelatedObjects( UObject const* A, const FNamePropertyMap& PropertyMapA, UObject const* B, const FNamePropertyMap& PropertyMapB, TArray<FName> &OutIdenticalProperties, TArray<FName> &OutDifferingProperties )
 {
-	TSet<FString> PropertiesInA, PropertiesInB;
+	TSet<FName> PropertiesInA, PropertiesInB;
 	if( A )
 	{
 		GetPropertyNameSet(*A, PropertiesInA);
@@ -578,7 +559,7 @@ static void CompareUnrelatedObjects( UObject const* A, const FNamePropertyMap& P
 	// for properties in common, dig out the uproperties and determine if they're identical:
 	if( A && B )
 	{
-		TSet<FString> Common = PropertiesInA.Intersect(PropertiesInB);
+		TSet<FName> Common = PropertiesInA.Intersect(PropertiesInB);
 		for (const auto& PropertyName : Common)
 		{
 			const UProperty* const* AProp = PropertyMapA.Find(PropertyName);
@@ -601,20 +582,43 @@ static void CompareUnrelatedObjects( UObject const* A, const FNamePropertyMap& P
 	}
 }
 
-// Generates a widget showing CDO values. NewBlueprint's CDO values will be displayed by the widget (Base just used for determining what to hide):
-static TSharedRef<IDetailsView> GenerateCDODiff( UObject const* CDO, TSet<FString> IdenticalPropertiesSet )
+int32 GetCurrentIndex( SListView< TSharedPtr< struct FDiffSingleResult> > const& ListView, const TArray< TSharedPtr< FDiffSingleResult > >& ListViewSource )
+{
+	const auto& Selected = ListView.GetSelectedItems();
+	if (Selected.Num() == 1)
+	{
+		int32 Index = 0;
+		for (auto It(ListViewSource.CreateConstIterator()); It; ++It, Index++)
+		{
+			if (*It == Selected[0])
+			{
+				return Index;
+			}
+		}
+	}
+	return -1;
+}
+
+void DiffWidgetUtils::SelectNextRow( SListView< TSharedPtr< FDiffSingleResult> >& ListView, const TArray< TSharedPtr< FDiffSingleResult > >& ListViewSource )
 {
-	FDetailsViewArgs DetailsViewArgs;
-	DetailsViewArgs.bShowDifferingPeropertiesOption = true;
+	int32 CurrentIndex = GetCurrentIndex(ListView, ListViewSource);
+	if( CurrentIndex == ListViewSource.Num() - 1 )
+	{
+		return;
+	}
+
+	ListView.SetSelection(ListViewSource[CurrentIndex + 1]);
+}
 
-	FPropertyEditorModule& EditModule = FModuleManager::Get().GetModuleChecked<FPropertyEditorModule>("PropertyEditor");
-	auto DetailsView = EditModule.CreateDetailView(DetailsViewArgs);
-	DetailsView->SetIsPropertyEditingEnabledDelegate(FIsPropertyEditingEnabled::CreateStatic([]{return false; }));
-	// This is a read only details view (see the property editing delegate), but it is not const correct:
-	DetailsView->SetObject(const_cast<UObject*>(CDO) );
-	DetailsView->UpdateIdenticalProperties(IdenticalPropertiesSet);
+void DiffWidgetUtils::SelectPrevRow(SListView< TSharedPtr< FDiffSingleResult> >& ListView, const TArray< TSharedPtr< FDiffSingleResult > >& ListViewSource )
+{
+	int32 CurrentIndex = GetCurrentIndex(ListView, ListViewSource);
+	if (CurrentIndex > 0)
+	{
+		return;
+	}
 
-	return DetailsView;
+	ListView.SetSelection(ListViewSource[CurrentIndex - 1]);
 }
 
 static TSharedRef<SWidget> GenerateComponentsDiff( UBlueprint const* BaseBlueprint, UBlueprint const* DisplayedBlueprint )
@@ -814,23 +818,23 @@ TSharedRef<SWidget> SBlueprintDiff::DefaultEmptyPanel()
 
 void SBlueprintDiff::NextDiff()
 {
-	if( DiffHighlighter.IsValid() )
+	if( DiffControl.IsValid() )
 	{
-		DiffHighlighter->NextDiff();
+		DiffControl->NextDiff();
 	}
 }
 
 void SBlueprintDiff::PrevDiff()
 {
-	if (DiffHighlighter.IsValid())
+	if (DiffControl.IsValid())
 	{
-		DiffHighlighter->PrevDiff();
+		DiffControl->PrevDiff();
 	}
 }
 
 bool SBlueprintDiff::CanCycleDiffs() const
 {
-	return DiffHighlighter.IsValid() && DiffHighlighter->HasDifferences();
+	return DiffControl.IsValid() && DiffControl->HasDifferences();
 }
 
 void SBlueprintDiff::FocusOnGraphRevisions( class UEdGraph* GraphOld, class UEdGraph* GraphNew, FListItemGraphToDiff* Diff )
@@ -839,7 +843,7 @@ void SBlueprintDiff::FocusOnGraphRevisions( class UEdGraph* GraphOld, class UEdG
 
 	ResetGraphEditors();
 
-	DiffHighlighter = Diff->AsShared();
+	DiffControl = Diff->AsShared();
 	//set new diff list
 	DiffListBorder->SetContent(Diff->GenerateDiffListWidget());
 }
@@ -1210,19 +1214,16 @@ TSharedRef<SWidget> SBlueprintDiff::GenerateGraphPanel()
 
 TSharedRef<SWidget> SBlueprintDiff::GenerateDefaultsPanel()
 {
-	const UObject* A = GetCDO(PanelOld.Blueprint);
-	FNamePropertyMap PropertyMapA = GetProperties(*A);
-	const UObject* B = GetCDO(PanelNew.Blueprint);
-	FNamePropertyMap PropertyMapB = GetProperties(*B);
-	TArray<FString> IdenticalProperties, DifferingProperties;
-	CompareUnrelatedObjects( A, PropertyMapA, B, PropertyMapB, IdenticalProperties, DifferingProperties);
-
-	TSet<FString> IdenticalPropertiesSet(IdenticalProperties);
-	auto LeftDetailsView = GenerateCDODiff(A, IdenticalPropertiesSet );
-	auto RightDetailsView = GenerateCDODiff(B, IdenticalPropertiesSet );
-
-	DiffHighlighter = TSharedPtr<CDODiffHighlighter>(new CDODiffHighlighter(LeftDetailsView, PropertyMapA, RightDetailsView, PropertyMapB, DifferingProperties ) );
+	const UObject* A = DiffUtils::GetCDO(PanelOld.Blueprint);
+	FNamePropertyMap PropertyMapA = DiffUtils::GetProperties(A);
+	const UObject* B = DiffUtils::GetCDO(PanelNew.Blueprint);
+	FNamePropertyMap PropertyMapB = DiffUtils::GetProperties(B);
+	TArray<FName> IdenticalProperties, DifferingProperties;
+	DiffUtils::CompareUnrelatedObjects( A, PropertyMapA, B, PropertyMapB, IdenticalProperties, DifferingProperties);
+
+	auto NewDiffControl = TSharedPtr<FCDODiffControl>(new FCDODiffControl(A, PropertyMapA, B, PropertyMapB, TSet<FName>(IdenticalProperties), DifferingProperties) );
 	//Splitter for left and right blueprint. Current convention is for the local (probably newer?) blueprint to be on the right:
+	DiffControl = NewDiffControl;
 	return SNew(SSplitter)
 		+ SSplitter::Slot()
 		.Value(0.5f)
@@ -1230,7 +1231,7 @@ TSharedRef<SWidget> SBlueprintDiff::GenerateDefaultsPanel()
 			SNew(SBorder)
 			.VAlign(VAlign_Fill)
 			[
-				LeftDetailsView
+				NewDiffControl->OldDetailsWidget()
 			]
 		]
 	+ SSplitter::Slot()
@@ -1239,7 +1240,7 @@ TSharedRef<SWidget> SBlueprintDiff::GenerateDefaultsPanel()
 			SNew(SBorder)
 			.VAlign(VAlign_Fill)
 			[
-				RightDetailsView
+				NewDiffControl->NewDetailsWidget()
 			]
 		];
 }
@@ -1272,7 +1273,7 @@ void SBlueprintDiff::SetCurrentMode(FName NewMode)
 {
 	CurrentMode = NewMode;
 
-	DiffHighlighter = TSharedPtr< IDiffHighlighter >();
+	DiffControl = TSharedPtr< IDiffControl >();
 	if( NewMode == FBlueprintEditorApplicationModes::StandardBlueprintEditorMode)
 	{
 		ModeContents->SetContent( GenerateGraphPanel() );
diff --git a/Engine/Source/Editor/Kismet/Public/DetailsDiff.h b/Engine/Source/Editor/Kismet/Public/DetailsDiff.h
new file mode 100644
--- /dev/null
+++ b/Engine/Source/Editor/Kismet/Public/DetailsDiff.h
@@ -0,0 +1,15 @@
+// Copyright 1998-2014 Epic Games, Inc. All Rights Reserved.
+#pragma once
+
+class KISMET_API FDetailsDiff
+{
+public:
+	FDetailsDiff( const UObject* InObject, const TMap< FName, const UProperty* >& InPropertyMap, const TSet< FName >& InIdenticalProperties );
+
+	void HighlightProperty( FName PropertyName );
+	TSharedRef< SWidget > DetailsWidget();
+
+private:
+	TMap< FName, const UProperty* > PropertyMap;
+	TSharedPtr< class IDetailsView > DetailsView;
+};
diff --git a/Engine/Source/Editor/Kismet/Public/IDiffControl.h b/Engine/Source/Editor/Kismet/Public/IDiffControl.h
new file mode 100644
--- /dev/null
+++ b/Engine/Source/Editor/Kismet/Public/IDiffControl.h
@@ -0,0 +1,11 @@
+// Copyright 1998-2014 Epic Games, Inc. All Rights Reserved.
+#pragma once
+
+class IDiffControl
+{
+public:
+	virtual void NextDiff() = 0;
+	virtual void PrevDiff() = 0;
+	virtual bool HasDifferences() const = 0;
+};
+
diff --git a/Engine/Source/Editor/Kismet/Public/SBlueprintDiff.h b/Engine/Source/Editor/Kismet/Public/SBlueprintDiff.h
--- a/Engine/Source/Editor/Kismet/Public/SBlueprintDiff.h
+++ b/Engine/Source/Editor/Kismet/Public/SBlueprintDiff.h
@@ -30,6 +30,19 @@ struct FDiffResultItem : public TSharedFromThis<FDiffResultItem>
 	TSharedRef<SWidget> KISMET_API GenerateWidget() const;
 };
 
+namespace DiffUtils
+{
+	KISMET_API TMap< FName, const UProperty* > GetProperties( const UObject* ForObj );
+	KISMET_API const UObject* GetCDO( const UBlueprint* ForBlueprint );
+	KISMET_API void CompareUnrelatedObjects( UObject const* A, const TMap< FName, const UProperty* >& PropertyMapA, UObject const* B, const TMap< FName, const UProperty* >& PropertyMapB, TArray<FName> &OutIdenticalProperties, TArray<FName> &OutDifferingProperties );
+}
+
+namespace DiffWidgetUtils
+{
+	KISMET_API void SelectNextRow(SListView< TSharedPtr< struct FDiffSingleResult> >& ListView, const TArray< TSharedPtr< struct FDiffSingleResult > >& ListViewSource );
+	KISMET_API void SelectPrevRow(SListView< TSharedPtr< struct FDiffSingleResult> >& ListView, const TArray< TSharedPtr< struct FDiffSingleResult > >& ListViewSource );
+}
+
 /*panel used to display the blueprint*/
 struct KISMET_API FDiffPanel
 {
@@ -182,7 +195,7 @@ protected:
 	TSharedPtr< FUICommandList > KeyCommands;
 
 	/** Helper class for highlighting diffs in different types of controls (graph view, details view, etc) */
-	TSharedPtr< class IDiffHighlighter > DiffHighlighter;
+	TSharedPtr< class IDiffControl > DiffControl;
 };
 
 
diff --git a/Engine/Source/Editor/PropertyEditor/Private/DetailItemNode.cpp b/Engine/Source/Editor/PropertyEditor/Private/DetailItemNode.cpp
--- a/Engine/Source/Editor/PropertyEditor/Private/DetailItemNode.cpp
+++ b/Engine/Source/Editor/PropertyEditor/Private/DetailItemNode.cpp
@@ -78,6 +78,8 @@ FDetailItemNode::~FDetailItemNode()
 
 void FDetailItemNode::InitPropertyEditor()
 {
+	Customization.GetPropertyNode()->SetTreeNode( this->AsShared() );
+
 	if( Customization.GetPropertyNode()->GetProperty() && Customization.GetPropertyNode()->GetProperty()->IsA<UArrayProperty>() )
 	{
 		const bool bUpdateFilteredNodes = false;
@@ -280,7 +282,7 @@ static bool PassesAllFilters( const FDetailLayoutCustomization& InCustomization,
 
 			const bool bPassesSearchFilter = bSearchFilterIsEmpty || ( bIsNotBeingFiltered || bIsSeenDueToFiltering || bIsParentSeenDueToFiltering );
 			const bool bPassesModifiedFilter = bPassesSearchFilter && ( InFilter.bShowOnlyModifiedProperties == false || PropertyNodePin->GetDiffersFromDefault() == true );
-			const bool bPassesDifferingFilter = InFilter.bShowOnlyDiffering == false || InFilter.NamesOfMatchingProperties.Find( PropertyNodePin->GetProperty()->GetName() ) == NULL;
+			const bool bPassesDifferingFilter = InFilter.bShowOnlyDiffering == false || InFilter.NamesOfMatchingProperties.Find( PropertyNodePin->GetProperty()->GetFName() ) == NULL;
 
 			// The property node is visible (note categories are never visible unless they have a child that is visible )
 			bPassesAllFilters = bPassesSearchFilter && bPassesModifiedFilter && bPassesDifferingFilter;
diff --git a/Engine/Source/Editor/PropertyEditor/Private/PropertyNode.h b/Engine/Source/Editor/PropertyEditor/Private/PropertyNode.h
--- a/Engine/Source/Editor/PropertyEditor/Private/PropertyNode.h
+++ b/Engine/Source/Editor/PropertyEditor/Private/PropertyNode.h
@@ -655,6 +655,12 @@ public:
 	/** Sets the highlighted flag to the requested value */
 	void SetIsHighlighted( bool bInIsHighlighted ) { bIsHighlighted = bInIsHighlighted; }
 
+	/** Returns the IDetailTreeNode associated with this property node */
+	TSharedPtr< class IDetailTreeNode > GetTreeNode() { return TreeNode.Pin(); }
+
+	/** Sets the IDetailTreeNode associated with this property node */
+	void SetTreeNode( TSharedPtr< class IDetailTreeNode > InTreeNode ) { TreeNode = InTreeNode; }
+
 protected:
 
 	TSharedRef<FEditPropertyChain> BuildPropertyChain( UProperty* PropertyAboutToChange );
@@ -774,6 +780,9 @@ protected:
 
 	/** An array of restrictions limiting this property's potential values in property editors.*/
 	TArray<TSharedRef<const class FPropertyRestriction>> Restrictions;
+
+	/** Optional reference to a tree node that is displaying this property */
+	TWeakPtr< class IDetailTreeNode > TreeNode;
 };
 
 class FComplexPropertyNode : public FPropertyNode
diff --git a/Engine/Source/Editor/PropertyEditor/Private/SDetailsView.cpp b/Engine/Source/Editor/PropertyEditor/Private/SDetailsView.cpp
--- a/Engine/Source/Editor/PropertyEditor/Private/SDetailsView.cpp
+++ b/Engine/Source/Editor/PropertyEditor/Private/SDetailsView.cpp
@@ -71,7 +71,7 @@ void SDetailsView::Construct(const FArguments& InArgs)
 			);
 		}
 
-		if( DetailsViewArgs.bShowDifferingPeropertiesOption )
+		if( DetailsViewArgs.bShowDifferingPropertiesOption )
 		{
 			DetailViewOptions.AddMenuEntry(
 				LOCTEXT("ShowOnlyDiffering", "Show Only Differing Properties"),
diff --git a/Engine/Source/Editor/PropertyEditor/Private/SDetailsViewBase.cpp b/Engine/Source/Editor/PropertyEditor/Private/SDetailsViewBase.cpp
--- a/Engine/Source/Editor/PropertyEditor/Private/SDetailsViewBase.cpp
+++ b/Engine/Source/Editor/PropertyEditor/Private/SDetailsViewBase.cpp
@@ -114,6 +114,11 @@ void SDetailsViewBase::HighlightProperty(const UProperty* Property)
 			{
 				(*PropertyNode)->SetIsHighlighted( true );
 				PrevHighlightedProperty = *PropertyNode;
+				auto TreeNodePtr = (*PropertyNode)->GetTreeNode();
+				if ( TreeNodePtr.IsValid() )
+				{
+					DetailTree->RequestScrollIntoView(TreeNodePtr.ToSharedRef());
+				}
 			}
 		}
 	}
diff --git a/Engine/Source/Editor/PropertyEditor/Private/SDetailsViewBase.h b/Engine/Source/Editor/PropertyEditor/Private/SDetailsViewBase.h
--- a/Engine/Source/Editor/PropertyEditor/Private/SDetailsViewBase.h
+++ b/Engine/Source/Editor/PropertyEditor/Private/SDetailsViewBase.h
@@ -66,7 +66,7 @@ struct FDetailFilter
 	bool bShowAllAdvanced;
 	/** If we should only show differing properties */
 	bool bShowOnlyDiffering;
-	TSet<FString> NamesOfMatchingProperties;
+	TSet<FName> NamesOfMatchingProperties;
 };
 
 struct FDetailColumnSizeData
@@ -353,7 +353,7 @@ protected:
 	void OnFilterTextChanged(const FText& InFilterText);
 
 	// @todo doc where does this belong?
-	virtual void UpdateIdenticalProperties(const TSet<FString> IdenticalProperties) override { CurrentFilter.NamesOfMatchingProperties = IdenticalProperties; }
+	virtual void UpdateIdenticalProperties(const TSet<FName> IdenticalProperties) override { CurrentFilter.NamesOfMatchingProperties = IdenticalProperties; }
 
 	/** 
 	 * Hides or shows properties based on the passed in filter text
diff --git a/Engine/Source/Editor/PropertyEditor/Public/IDetailsView.h b/Engine/Source/Editor/PropertyEditor/Public/IDetailsView.h
--- a/Engine/Source/Editor/PropertyEditor/Public/IDetailsView.h
+++ b/Engine/Source/Editor/PropertyEditor/Public/IDetailsView.h
@@ -34,7 +34,7 @@ struct FDetailsViewArgs
 	/** True if you want to show the actor label */
 	uint32 bShowActorLabel : 1;
 	/** Bind this delegate to hide differing properties */
-	uint32 bShowDifferingPeropertiesOption : 1;
+	uint32 bShowDifferingPropertiesOption : 1;
 
 	/** Default constructor */
 	FDetailsViewArgs( const bool InUpdateFromSelection = false
@@ -57,7 +57,7 @@ struct FDetailsViewArgs
 		, bShowOptions( true )
 		, bShowModifiedPropertiesOption(true)
 		, bShowActorLabel(true)
-		, bShowDifferingPeropertiesOption(false)
+		, bShowDifferingPropertiesOption(false)
 	{
 	}
 };
@@ -177,5 +177,5 @@ public:
 	/**
 	 * Sets the set of properties that are considered identical, used when filtering out identical properties (showing only differing properties)
 	 */
-	virtual void UpdateIdenticalProperties(const TSet<FString> IdenticalProperties) = 0;
+	virtual void UpdateIdenticalProperties(const TSet<FName> IdenticalProperties) = 0;
 };
