Blueprint menu refactor: pins show more than member functions/variables, and exclude variables from the palette.

diff --git a/Engine/Source/Editor/Kismet/Private/BlueprintActionMenuUtils.cpp b/Engine/Source/Editor/Kismet/Private/BlueprintActionMenuUtils.cpp
--- a/Engine/Source/Editor/Kismet/Private/BlueprintActionMenuUtils.cpp
+++ b/Engine/Source/Editor/Kismet/Private/BlueprintActionMenuUtils.cpp
@@ -9,6 +9,7 @@
 #include "K2Node_CallFunction.h"
 #include "K2Node_ComponentBoundEvent.h"
 #include "EdGraphSchema_K2.h"
+#include "K2Node_Variable.h"
 
 #define LOCTEXT_NAMESPACE "BlueprintActionMenuUtils"
 
@@ -33,8 +34,8 @@ static bool BlueprintActionMenuUtilsImpl::IsUnBoundSpawner(FBlueprintActionFilte
 }
 
 /*******************************************************************************
-* FBlueprintActionMenuUtils
-******************************************************************************/
+ * FBlueprintActionMenuUtils
+ ******************************************************************************/
 
 //------------------------------------------------------------------------------
 void FBlueprintActionMenuUtils::MakePaletteMenu(FBlueprintActionContext const& Context, UClass* FilterClass, FBlueprintActionMenuBuilder& MenuOut)
@@ -50,6 +51,11 @@ void FBlueprintActionMenuUtils::MakePaletteMenu(FBlueprintActionContext const& C
 	FBlueprintActionFilter MenuFilter(FilterFlags);
 	MenuFilter.Context = Context;
 	
+	// self member variables can be accessed through the MyBluprint panel (even
+	// inherited ones)... external variables can be accessed through the context
+	// menu (don't want to clutter the palette?)
+	MenuFilter.ExcludedNodeTypes.Add(UK2Node_Variable::StaticClass());
+	
 	if (FilterClass != nullptr)
 	{
 		MenuFilter.OwnerClasses.Add(FilterClass);
@@ -96,16 +102,25 @@ void FBlueprintActionMenuUtils::MakeContextMenu(FBlueprintActionContext const& C
 			MenuOut.AddMenuSection(ComponentsFilter, FText::GetEmpty(), ComponentsSectionGroup);
 		}
 		
-		for (UBlueprint* Blueprint : Context.Blueprints)
+		if (MenuFilter.Context.Pins.Num() == 0)
 		{
-			MenuFilter.OwnerClasses.Add(Blueprint->GeneratedClass);
+			for (UBlueprint* Blueprint : Context.Blueprints)
+			{
+				UClass* BlueprintClass = (Blueprint->SkeletonGeneratedClass != nullptr) ? Blueprint->SkeletonGeneratedClass : Blueprint->ParentClass;
+				// @TODO: are we sure we only want members in the unconnected context menu? 
+				MenuFilter.OwnerClasses.Add(BlueprintClass);
+			}
 		}
 	}
-
+	else
+	{
+		MenuFilter.Context.Pins.Empty();
+	}
 	MenuOut.AddMenuSection(MenuFilter);
 	MenuOut.RebuildActionList();
 
-	if (!GetDefault<UEdGraphSchema_K2>()->bUseLegacyActionMenus)
+	UEdGraphSchema_K2 const* const K2Schema = GetDefault<UEdGraphSchema_K2>();
+	if (!K2Schema->bUseLegacyActionMenus)
 	{
 		for (UEdGraph const* Graph : Context.Graphs)
 		{
