(WIP) Readdressing how we handle bound nodes in the new blueprint menu system (can't cache all permutations in the database, so we bind at menu construction time).

diff --git a/Engine/Source/Editor/BlueprintGraph/Classes/BlueprintBoundNodeSpawner.h b/Engine/Source/Editor/BlueprintGraph/Classes/BlueprintBoundNodeSpawner.h
--- a/Engine/Source/Editor/BlueprintGraph/Classes/BlueprintBoundNodeSpawner.h
+++ b/Engine/Source/Editor/BlueprintGraph/Classes/BlueprintBoundNodeSpawner.h
@@ -28,11 +28,11 @@ public:
 	 * object.
 	 *
 	 * @param  NodeClass	The event node type that you want this to spawn.
-	 * @param  BoundObject	The object you want spawned nodes bound to.
+	 * @param  BoundObjType	The type of object you want spawned nodes bound to.
 	 * @param  Outer		Optional outer for the new spawner (if left null, the transient package will be used).
 	 * @return A newly allocated instance of this class.
 	 */
-	static UBlueprintBoundNodeSpawner* Create(TSubclassOf<UEdGraphNode> NodeClass, UObject* BoundObject, UObject* Outer = nullptr);
+	static UBlueprintBoundNodeSpawner* Create(TSubclassOf<UEdGraphNode> NodeClass, UClass* BoundObjType, UObject* Outer = nullptr);
 
 	/**
 	 * Creates a new UBlueprintBoundNodeSpawner that is charged with using the 
@@ -42,23 +42,31 @@ public:
 	 * in charge of binding.
 	 *
 	 * @param  SubSpawner	The spawner that will instantiate the node.
-	 * @param  BoundObject	The object you want spawned nodes bound to.
+	 * @param  BoundObjType	The type of object you want spawned nodes bound to.
 	 * @param  Outer		Optional outer for the new spawner (if left null, the transient package will be used).
 	 * @return A newly allocated instance of this class.
 	 */
-	static UBlueprintBoundNodeSpawner* Create(UBlueprintNodeSpawner* SubSpawner, UObject* BoundObject, UObject* Outer = nullptr);
+	static UBlueprintBoundNodeSpawner* Create(UBlueprintNodeSpawner* SubSpawner, UClass* BoundObjType, UObject* Outer = nullptr);
 
 	// UBlueprintNodeSpawner interface
 	virtual UEdGraphNode* Invoke(UEdGraph* ParentGraph, FVector2D const Location) const override;
 	virtual FText GetDefaultMenuName() const override;
 	virtual FText GetDefaultMenuCategory() const override;
 	// End UBlueprintNodeSpawner interface
+
+
+	bool IsOfBindingType(UObject const* BindingCandidate) const;
 	
+
+	bool SetBoundObject(UObject const* Object);
+
 	/**
  	 * Retrieves the object that this binds to spawned nodes.
 	 * @return The object that this class was initialized with.
 	 */
 	UObject const* GetBoundObject() const;
+
+	
 	
 	/**
 	 * Retrieves the sub-spawner that this wraps (which is actually in-charge of
@@ -69,6 +77,10 @@ public:
 	UBlueprintNodeSpawner const* GetSubSpawner() const;
 
 private:
+	/** This denotes the type of object this is allowed to bind to. */
+	UPROPERTY()
+	UClass* BoundObjClass;
+
 	/** This object that this spawner attempts to bind new nodes with. */
 	UPROPERTY()
 	TWeakObjectPtr<UObject> BoundObjPtr;
diff --git a/Engine/Source/Editor/BlueprintGraph/Private/BlueprintActionDatabase.cpp b/Engine/Source/Editor/BlueprintGraph/Private/BlueprintActionDatabase.cpp
--- a/Engine/Source/Editor/BlueprintGraph/Private/BlueprintActionDatabase.cpp
+++ b/Engine/Source/Editor/BlueprintGraph/Private/BlueprintActionDatabase.cpp
@@ -117,6 +117,14 @@ namespace FBlueprintNodeSpawnerFactory
 	 * @return A new node-spawner, setup to spawn a UK2Node_AssignDelegate.
 	 */
 	static UBlueprintNodeSpawner* MakeAssignDelegateNodeSpawner(UMulticastDelegateProperty* DelegateProperty);
+
+	/**
+	 * 
+	 * 
+	 * @param  DelegateProperty	
+	 * @return 
+	 */
+	static UBlueprintNodeSpawner* MakeComponentBoundEventSpawner(UMulticastDelegateProperty* DelegateProperty);
 };
 
 //------------------------------------------------------------------------------
@@ -264,6 +272,20 @@ static UBlueprintNodeSpawner* FBlueprintNodeSpawnerFactory::MakeAssignDelegateNo
 	return UBlueprintDelegateNodeSpawner::Create(UK2Node_AssignDelegate::StaticClass(), DelegateProperty);
 }
 
+//------------------------------------------------------------------------------
+static UBlueprintNodeSpawner* FBlueprintNodeSpawnerFactory::MakeComponentBoundEventSpawner(UMulticastDelegateProperty* DelegateProperty)
+{
+	// BoundEventSpawner does the binding, we need a sub-spawner to lay the node
+	// ... we choose UBlueprintDelegateNodeSpawner because we need it to hold on
+	// to the UMulticastDelegateProperty
+	UBlueprintDelegateNodeSpawner* DelegateSpawner = UBlueprintDelegateNodeSpawner::Create(/*NodeClass =*/nullptr, DelegateProperty);
+	// UK2Node_ComponentBoundEvent is not a UK2Node_BaseMCDelegate so we have to
+	// set it this way
+	DelegateSpawner->NodeClass = UK2Node_ComponentBoundEvent::StaticClass(); 
+
+	return UBlueprintBoundNodeSpawner::Create(DelegateSpawner, UObjectProperty::StaticClass());
+}
+
 /*******************************************************************************
  * Static FBlueprintActionDatabase Helpers
  ******************************************************************************/
@@ -308,16 +330,6 @@ namespace BlueprintActionDatabaseImpl
 	 */
 	static void AddClassPropertyActions(UClass const* const Class, FActionList& ActionListOut);
 
-	/**
-	 * Loops over all of the property's class properties and creates bound 
-	 * actions for any that are viable for blueprint use. Evolved from certain
-	 * parts of FK2ActionMenuBuilder::AddBoundEventActionsForClass().
-	 *
-	 * @param  Property			The component property that you want node-spawners for.
-	 * @param  ActionListOut	The list you want populated with the new spawners.
-	 */
-	static void AddComponentBoundActions(UObjectProperty* Property, FActionList& ActionListOut);
-
 	/**
 	 * Evolved from FClassDynamicCastHelper::GetClassDynamicCastNodes(). If the
 	 * specified class is a viable blueprint variable type, then two cast nodes
@@ -462,6 +474,8 @@ static void BlueprintActionDatabaseImpl::AddClassFunctionActions(UClass const* c
 	using namespace FBlueprintNodeSpawnerFactory; // for MakeMessageNodeSpawner()
 	check(Class != nullptr);
 
+	bool const bIsComponent = Class->IsChildOf<UActorComponent>();
+
 	// loop over all the functions in the specified class; exclude-super because 
 	// we can always get the super functions by looking up that class separately 
 	for (TFieldIterator<UFunction> FunctionIt(Class, EFieldIteratorFlags::ExcludeSuper); FunctionIt; ++FunctionIt)
@@ -482,15 +496,27 @@ static void BlueprintActionDatabaseImpl::AddClassFunctionActions(UClass const* c
 			//        Blueprint "implemented interface", then we don't need to 
 			//        include it (the function is accounted for in from the 
 			//        interface class).
-			if (UBlueprintFunctionNodeSpawner* NodeSpawner = UBlueprintFunctionNodeSpawner::Create(Function))
-			{
-				ActionListOut.Add(NodeSpawner);
-			}
+			UBlueprintFunctionNodeSpawner* FuncSpawner = UBlueprintFunctionNodeSpawner::Create(Function);
+			ActionListOut.Add(FuncSpawner);
 
 			if (FKismetEditorUtilities::IsClassABlueprintInterface(Class))
 			{
 				ActionListOut.Add(MakeMessageNodeSpawner(Function));
 			}
+
+			// bindable function calls (for component properties)
+			if (bIsComponent)
+			{
+				bool const bIsPublic = !Function->HasMetaData(FBlueprintMetadata::MD_Protected) &&
+					!Function->HasMetaData(FBlueprintMetadata::MD_Private);
+				// may have to move the bIsPublic check from here into 
+				// BlueprintActionFilter, once we have blueprintable components?
+				if (bIsPublic)
+				{
+					// bound function for component property
+					ActionListOut.Add(UBlueprintBoundNodeSpawner::Create(FuncSpawner, UObjectProperty::StaticClass()));
+				}
+			}
 		}
 	}
 }
@@ -499,6 +525,8 @@ static void BlueprintActionDatabaseImpl::AddClassFunctionActions(UClass const* c
 static void BlueprintActionDatabaseImpl::AddClassPropertyActions(UClass const* const Class, FActionList& ActionListOut)
 {
 	using namespace FBlueprintNodeSpawnerFactory; // for MakeDelegateNodeSpawner()
+
+	bool const bIsComponent = Class->IsChildOf<UActorComponent>();
 	
 	// loop over all the properties in the specified class; exclude-super because 
 	// we can always get the super properties by looking up that class separately 
@@ -510,7 +538,7 @@ static void BlueprintActionDatabaseImpl::AddClassPropertyActions(UClass const* c
 			continue;
 		}
 
- 		bool bIsDelegate = Property->IsA(UMulticastDelegateProperty::StaticClass());
+ 		bool const bIsDelegate = Property->IsA(UMulticastDelegateProperty::StaticClass());
  		if (bIsDelegate)
  		{
 			UMulticastDelegateProperty* DelegateProperty = CastChecked<UMulticastDelegateProperty>(Property);
@@ -536,14 +564,14 @@ static void BlueprintActionDatabaseImpl::AddClassPropertyActions(UClass const* c
 
 			// @TODO: AddBoundEventActionsForClass()
 			//   UK2Node_ActorBoundEvent
+
+			if (bIsComponent)
+			{
+				ActionListOut.Add(MakeComponentBoundEventSpawner(DelegateProperty));
+			}
  		}
 		else
 		{
-// 			if (UObjectProperty* ObjProperty = Cast<UObjectProperty>(Property))
-// 			{
-// 				AddComponentBoundActions(ObjProperty, ActionListOut);
-// 			}
-
 			UBlueprintVariableNodeSpawner* GetterSpawner = UBlueprintVariableNodeSpawner::Create(UK2Node_VariableGet::StaticClass(), Property);
 			ActionListOut.Add(GetterSpawner);
 			UBlueprintVariableNodeSpawner* SetterSpawner = UBlueprintVariableNodeSpawner::Create(UK2Node_VariableSet::StaticClass(), Property);
@@ -552,66 +580,6 @@ static void BlueprintActionDatabaseImpl::AddClassPropertyActions(UClass const* c
 	}
 }
 
-//------------------------------------------------------------------------------
-static void BlueprintActionDatabaseImpl::AddComponentBoundActions(UObjectProperty* ComponentProperty, FActionList& ActionListOut)
-{
-	check(ComponentProperty != nullptr);
-
-	bool const bIsComponentProperty = ComponentProperty->PropertyClass->IsChildOf<UActorComponent>();
-	if (bIsComponentProperty)
-	{
-		// evolved from FK2ActionMenuBuilder::GetBoundEventsFromComponentSelection()
-		for (TFieldIterator<UMulticastDelegateProperty> PropertyIt(ComponentProperty->PropertyClass, EFieldIteratorFlags::IncludeSuper); PropertyIt; ++PropertyIt)
-		{
-			UMulticastDelegateProperty* Property = *PropertyIt;
-			if (!IsPropertyBlueprintVisible(Property))
-			{
-				continue;
-			}
-
-			UBlueprintDelegateNodeSpawner* DelegateSpawner = UBlueprintDelegateNodeSpawner::Create(/*NodeClass =*/nullptr, Property);
-			DelegateSpawner->NodeClass = UK2Node_ComponentBoundEvent::StaticClass();
-
-			auto BindComponentLambda = [](UEdGraphNode* NewNode, bool /*bIsTemplateNode*/, TWeakObjectPtr<UObjectProperty> ComponentProp, TWeakObjectPtr<UMulticastDelegateProperty> DelegateProp)
-			{
-				UK2Node_ComponentBoundEvent* BoundNode = CastChecked<UK2Node_ComponentBoundEvent>(NewNode);
-				if (ComponentProp.IsValid() && DelegateProp.IsValid())
-				{
-					BoundNode->InitializeComponentBoundEventParams(ComponentProp.Get(), DelegateProp.Get());
-				}
-			};
-
-			UBlueprintBoundNodeSpawner* BoundEventSpawner = UBlueprintBoundNodeSpawner::Create(DelegateSpawner, ComponentProperty);
-			TWeakObjectPtr<UObjectProperty> ComponentPropPtr = ComponentProperty;
-			TWeakObjectPtr<UMulticastDelegateProperty> DelegatePropPtr = Property;
-			BoundEventSpawner->CustomizeNodeDelegate = UBlueprintNodeSpawner::FCustomizeNodeDelegate::CreateStatic(BindComponentLambda, ComponentPropPtr, DelegatePropPtr);
-			
-			ActionListOut.Add(BoundEventSpawner);
-		}
-		
-		// evolved from FK2ActionMenuBuilder::GetFunctionCallsOnSelectedComponents()
-		for (TFieldIterator<UFunction> FuncIt(ComponentProperty->PropertyClass, EFieldIteratorFlags::IncludeSuper); FuncIt; ++FuncIt)
-		{
-			UFunction* ComponentFunc = *FuncIt;
-			if (!UEdGraphSchema_K2::CanUserKismetCallFunction(ComponentFunc))
-			{
-				continue;
-			}
-			
-			bool const bIsPure  = (ComponentFunc->HasAnyFunctionFlags(FUNC_BlueprintPure) != false);
-			bool const bIsConst = (ComponentFunc->HasAnyFunctionFlags(FUNC_Const) != false);
-			
-			// don't know why we only allow impure ("imperitive") and const functions,
-			// but this mimics legacy functionality in GetFunctionCallsOnSelectedComponents()
-			if (!bIsPure || bIsConst)
-			{
-				UBlueprintFunctionNodeSpawner* FuncSpawner = UBlueprintFunctionNodeSpawner::Create(ComponentFunc);
-				ActionListOut.Add(UBlueprintBoundNodeSpawner::Create(FuncSpawner, ComponentProperty));
-			}
-		}
-	}
-}
-
 //------------------------------------------------------------------------------
 static void BlueprintActionDatabaseImpl::AddClassCastActions(UClass* const Class, FActionList& ActionListOut)
 {
@@ -923,7 +891,8 @@ void FBlueprintActionDatabase::Tick(float DeltaTime)
 				{
 					FScopedDurationTimer ScopedTimer(TotalFrameDuration);
 
-					UEdGraphNode* TemplateNode = (*ClassActionList)[ActionListIndex]->GetTemplateNode();
+					UBlueprintNodeSpawner* Action = (*ClassActionList)[ActionListIndex];
+					UEdGraphNode* TemplateNode = Action->GetTemplateNode();
 					// since we're doing this incrementally, someone could have 
 					// already requested this template, and allocated its pins 
 					// (don't need to do redundant work)
diff --git a/Engine/Source/Editor/BlueprintGraph/Private/BlueprintActionFilter.cpp b/Engine/Source/Editor/BlueprintGraph/Private/BlueprintActionFilter.cpp
--- a/Engine/Source/Editor/BlueprintGraph/Private/BlueprintActionFilter.cpp
+++ b/Engine/Source/Editor/BlueprintGraph/Private/BlueprintActionFilter.cpp
@@ -3,12 +3,10 @@
 #include "BlueprintGraphPrivatePCH.h"
 #include "BlueprintActionFilter.h"
 #include "BlueprintNodeSpawner.h"
-#include "BlueprintFunctionNodeSpawner.h"
+#include "BlueprintNodeSpawnerUtils.h"
+#include "BlueprintVariableNodeSpawner.h"
 #include "BlueprintEventNodeSpawner.h"
-#include "BlueprintComponentNodeSpawner.h"
-#include "BlueprintDelegateNodeSpawner.h"
 #include "BlueprintBoundNodeSpawner.h"
-#include "BlueprintVariableNodeSpawner.h"
 #include "EdGraphSchema_K2.h"		// for FBlueprintMetadata
 #include "BlueprintEditorUtils.h"	// for FindBlueprintForGraph()
 #include "ObjectEditorUtils.h"		// for IsFunctionHiddenFromClass()/IsVariableCategoryHiddenFromClass()
@@ -33,44 +31,7 @@ namespace BlueprintActionFilterImpl
 	 * @param  Blueprint	
 	 * @return 
 	 */
-	static UClass* GetBlueprintClass(UBlueprint const* const Blueprint);
-
-	/**
-	 * All actions are associated with some class, whether it's the class of the
-	 * node it will spawn, or the outer of the associated field (property,
-	 * function, etc.).
-	 * 
-	 * @param  BlueprintAction	The node-spawner you want an associated class for.
-	 * @return The action's associated class (null if it doesn't have one).
-	 */
-	static UClass* GetActionClass(UBlueprintNodeSpawner const* BlueprintAction);
-	
-	/**
-	 * Certain node-spawners are associated with specific UFields (functions, 
-	 * properties, etc.). This attempts to retrieve the field from the spawner.
-	 * 
-	 * @param  BlueprintAction	The node-spawner you want an associated field for.
-	 * @return The action's associated field (null if it doesn't have one).
-	 */
-	static UField const* GetAssociatedField(UBlueprintNodeSpawner const* BlueprintAction);
-
-	/**
-	 * Certain node-spawners are associated with specific UFunction (call-
-	 * function, and event spawners). This attempts to retrieve the function 
-	 * from the spawner.
-	 * 
-	 * @param  BlueprintAction	The node-spawner you want an associated function for.
-	 * @return The action's associated function (null if it doesn't have one).
-	 */
-	static UFunction const* GetAssociatedFunction(UBlueprintNodeSpawner const* BlueprintAction);
-
-	/**
-	 * 
-	 * 
-	 * @param  BlueprintAction	
-	 * @return 
-	 */
-	static UProperty const* GetAssociatedProperty(UBlueprintNodeSpawner const* BlueprintAction);
+	static UClass* GetBlueprintClass(UBlueprint const* const Blueprint);	
 	
 	/**
 	 * Utility method to check and see if the specified node-spawner would 
@@ -298,100 +259,14 @@ UClass* BlueprintActionFilterImpl::GetBlueprintClass(UBlueprint const* const Blu
 	return (Blueprint->SkeletonGeneratedClass != nullptr) ? Blueprint->SkeletonGeneratedClass : Blueprint->ParentClass;
 }
 
-//------------------------------------------------------------------------------
-static UClass* BlueprintActionFilterImpl::GetActionClass(UBlueprintNodeSpawner const* BlueprintAction)
-{
-	UClass* ActionClass = BlueprintAction->NodeClass;
-	
-	if (UField const* AssociatedField = GetAssociatedField(BlueprintAction))
-	{
-		ActionClass = AssociatedField->GetOwnerClass();
-	}
-	else if (UBlueprintComponentNodeSpawner const* ComponentSpawner = Cast<UBlueprintComponentNodeSpawner>(BlueprintAction))
-	{
-		ActionClass = ComponentSpawner->GetComponentClass();
-	}
-	else if (UBlueprintBoundNodeSpawner const* BoundSpawner = Cast<UBlueprintBoundNodeSpawner>(BlueprintAction))
-	{
-		ActionClass = GetActionClass(BoundSpawner->GetSubSpawner());
-	}
-	else if (UBlueprintVariableNodeSpawner const* VarSpawner = Cast<UBlueprintVariableNodeSpawner>(BlueprintAction))
-	{
-		UObject const* VarOuter = VarSpawner->GetVarOuter();
-		if (VarSpawner->IsLocalVariable())
-		{
-			UBlueprint* Blueprint = FBlueprintEditorUtils::FindBlueprintForGraphChecked(CastChecked<UEdGraph>(VarOuter));
-			ActionClass = GetBlueprintClass(Blueprint);
-		}
-		// else, should have been caught in GetAssociatedField()
-	}
-	
-	return ActionClass;
-}
 
-//------------------------------------------------------------------------------
-static UField const* BlueprintActionFilterImpl::GetAssociatedField(UBlueprintNodeSpawner const* BlueprintAction)
-{
-	UField const* ClassField = nullptr;
-
-	if (UFunction const* Function = GetAssociatedFunction(BlueprintAction))
-	{
-		ClassField = Function;
-	}
-	else if (UProperty const* Property = GetAssociatedProperty(BlueprintAction))
-	{
-		ClassField = Property;
-	}
-	else if (UBlueprintBoundNodeSpawner const* BoundSpawner = Cast<UBlueprintBoundNodeSpawner>(BlueprintAction))
-	{
-		ClassField = GetAssociatedField(BoundSpawner->GetSubSpawner());
-	}
-	return ClassField;
-}
-
-//------------------------------------------------------------------------------
-static UFunction const* BlueprintActionFilterImpl::GetAssociatedFunction(UBlueprintNodeSpawner const* BlueprintAction)
-{
-	UFunction const* Function = nullptr;
-
-	if (UBlueprintFunctionNodeSpawner const* FuncNodeSpawner = Cast<UBlueprintFunctionNodeSpawner>(BlueprintAction))
-	{
-		Function = FuncNodeSpawner->GetFunction();
-	}
-	else if (UBlueprintEventNodeSpawner const* EventSpawner = Cast<UBlueprintEventNodeSpawner>(BlueprintAction))
-	{
-		Function = EventSpawner->GetEventFunction();
-	}
-	else if (UBlueprintBoundNodeSpawner const* BoundSpawner = Cast<UBlueprintBoundNodeSpawner>(BlueprintAction))
-	{
-		Function = GetAssociatedFunction(BoundSpawner->GetSubSpawner());
-	}
-
-	return Function;
-}
-
-//------------------------------------------------------------------------------
-static UProperty const* BlueprintActionFilterImpl::GetAssociatedProperty(UBlueprintNodeSpawner const* BlueprintAction)
-{
-	UProperty const* Property = nullptr;
-
-	if (UBlueprintDelegateNodeSpawner const* PropertySpawner = Cast<UBlueprintDelegateNodeSpawner>(BlueprintAction))
-	{
-		Property = PropertySpawner->GetProperty();
-	}
-	else if (UBlueprintVariableNodeSpawner const* VarSpawner = Cast<UBlueprintVariableNodeSpawner>(BlueprintAction))
-	{
-		Property = VarSpawner->GetVarProperty();
-	}
-	return Property;
-}
 
 //------------------------------------------------------------------------------
 static bool BlueprintActionFilterImpl::IsImpure(UBlueprintNodeSpawner const* NodeSpawner)
 {
 	bool bIsImpure = false;
 	
-	if (UFunction const* Function = GetAssociatedFunction(NodeSpawner))
+	if (UFunction const* Function = FBlueprintNodeSpawnerUtils::GetAssociatedFunction(NodeSpawner))
 	{
 		bIsImpure = (Function->HasAnyFunctionFlags(FUNC_BlueprintPure) == false);
 	}
@@ -421,7 +296,7 @@ static bool BlueprintActionFilterImpl::IsLatent(UBlueprintNodeSpawner const* Nod
 	{
 		bIsLatent = true;
 	}
-	else if (UFunction const* Function = GetAssociatedFunction(NodeSpawner))
+	else if (UFunction const* Function = FBlueprintNodeSpawnerUtils::GetAssociatedFunction(NodeSpawner))
 	{
 		bIsLatent = (Function->HasMetaData(FBlueprintMetadata::MD_Latent) != false);
 	}
@@ -497,7 +372,7 @@ static bool BlueprintActionFilterImpl::IsFieldInaccessible(FBlueprintActionFilte
 	bool bIsFilteredOut = false;
 	FBlueprintActionContext const& FilterContext = Filter.Context;
 	
-	if (UField const* ClassField = GetAssociatedField(BlueprintAction))
+	if (UField const* ClassField = FBlueprintNodeSpawnerUtils::GetAssociatedField(BlueprintAction))
 	{
 		bool bIsProtected = ClassField->HasMetaData(FBlueprintMetadata::MD_Protected);
 		bool bIsPrivate   = ClassField->HasMetaData(FBlueprintMetadata::MD_Private);
@@ -531,7 +406,7 @@ static bool BlueprintActionFilterImpl::IsRestrictedClassMember(FBlueprintActionF
 	bool bIsFilteredOut = false;
 	FBlueprintActionContext const& FilterContext = Filter.Context;
 	
-	if (UClass const* ActionClass = GetActionClass(BlueprintAction))
+	if (UClass const* ActionClass = FBlueprintNodeSpawnerUtils::GetAssociatedClass(BlueprintAction))
 	{
 		if (ActionClass->HasMetaData(FBlueprintMetadata::MD_RestrictedToClasses))
 		{
@@ -571,7 +446,7 @@ static bool BlueprintActionFilterImpl::IsPermissionNotGranted(FBlueprintActionFi
 	bool bIsFilteredOut = false;
 	FBlueprintActionContext const& FilterContext = Filter.Context;
 
-	if (UProperty const* const Property = GetAssociatedProperty(BlueprintAction))
+	if (UProperty const* const Property = FBlueprintNodeSpawnerUtils::GetAssociatedProperty(BlueprintAction))
 	{
 		UClass const* const NodeClass = BlueprintAction->NodeClass;
 		for (UBlueprint const* Blueprint : FilterContext.Blueprints)
@@ -595,7 +470,7 @@ static bool BlueprintActionFilterImpl::IsDeprecated(FBlueprintActionFilter const
 	{
 		bIsFilteredOut = true;
 	}
-	else if (UClass* ActionClass = GetActionClass(BlueprintAction))
+	else if (UClass* ActionClass = FBlueprintNodeSpawnerUtils::GetAssociatedClass(BlueprintAction))
 	{
 		bIsFilteredOut = ActionClass->HasAnyClassFlags(CLASS_Deprecated);
 	}
@@ -606,7 +481,7 @@ static bool BlueprintActionFilterImpl::IsDeprecated(FBlueprintActionFilter const
 static bool BlueprintActionFilterImpl::IsGlobal(FBlueprintActionFilter const& Filter, UBlueprintNodeSpawner const* BlueprintAction)
 {
 	bool bIsFilteredOut = false;
-	if (UField const* ClassField = GetAssociatedField(BlueprintAction))
+	if (UField const* ClassField = FBlueprintNodeSpawnerUtils::GetAssociatedField(BlueprintAction))
 	{
 		bIsFilteredOut = IsGlobal(ClassField);
 		
@@ -632,7 +507,7 @@ static bool BlueprintActionFilterImpl::IsGlobal(FBlueprintActionFilter const& Fi
 static bool BlueprintActionFilterImpl::IsExternalField(FBlueprintActionFilter const& Filter, UBlueprintNodeSpawner const* BlueprintAction)
 {
 	bool bIsFilteredOut = false;
-	if (UField const* ClassField = GetAssociatedField(BlueprintAction))
+	if (UField const* ClassField = FBlueprintNodeSpawnerUtils::GetAssociatedField(BlueprintAction))
 	{
 		UClass* FieldClass = Cast<UClass>(ClassField->GetOuter());
 		bool const bIsInterfaceMethod = FieldClass->IsChildOf<UInterface>();
@@ -662,7 +537,7 @@ static bool BlueprintActionFilterImpl::IsUneededMessageNode(FBlueprintActionFilt
 	bool bIsFilteredOut = false;
 	if (BlueprintAction->NodeClass->IsChildOf<UK2Node_Message>())
 	{
-		UFunction const* const InterfaceFunc = GetAssociatedFunction(BlueprintAction);
+		UFunction const* const InterfaceFunc = FBlueprintNodeSpawnerUtils::GetAssociatedFunction(BlueprintAction);
 		checkSlow(InterfaceFunc != nullptr);
 		checkSlow(InterfaceFunc->GetOuterUClass()->IsChildOf<UInterface>());
 
@@ -686,7 +561,7 @@ static bool BlueprintActionFilterImpl::IsFieldCategoryHidden(FBlueprintActionFil
 	bool bIsFilteredOut = false;
 	FBlueprintActionContext const& FilterContext = Filter.Context;
 
-	if (UField const* AssociatedField = GetAssociatedField(BlueprintAction))
+	if (UField const* AssociatedField = FBlueprintNodeSpawnerUtils::GetAssociatedField(BlueprintAction))
 	{
 		DECLARE_DELEGATE_RetVal_OneParam(bool, FIsFieldHiddenDelegate, UClass*);
 		FIsFieldHiddenDelegate IsFieldHiddenDelegate;
@@ -825,6 +700,21 @@ static bool BlueprintActionFilterImpl::IsBoundToUnselectedObject(FBlueprintActio
 	return bIsFilteredOut;
 }
 
+//------------------------------------------------------------------------------
+static bool BlueprintActionFilterImpl::IsBoundFunctionImpure(FBlueprintActionFilter const& Filter, UBlueprintNodeSpawner const* BlueprintAction)
+{
+	bool bIsFilteredOut = false;
+	if (UBlueprintBoundNodeSpawner const* BoundSpawner = Cast<UBlueprintBoundNodeSpawner>(BlueprintAction))
+	{
+		if (UFunction* BoundFunction = FBlueprintEditorUtils::GetAssociatedFunction(BoundSpawner))
+		{
+			bIsFilteredOut = !BoundFunction->HasAnyFunctionFlags(FUNC_BlueprintPure) && 
+				!Function->HasAnyFunctionFlags(FUNC_Const);
+		}
+	}
+	return bIsFilteredOut;
+}
+
 //------------------------------------------------------------------------------
 static bool BlueprintActionFilterImpl::IsOutOfScopeLocalVariable(FBlueprintActionFilter const& Filter, UBlueprintNodeSpawner const* BlueprintAction)
 {
@@ -894,7 +784,7 @@ static bool BlueprintActionFilterImpl::HasMatchingPin(UBlueprintNodeSpawner cons
 static bool BlueprintActionFilterImpl::IsPinCompatibleWithTargetSelf(UEdGraphPin const* Pin, UField const* MemberField, UBlueprintNodeSpawner const* MemberNodeSpawner)
 {
 	bool bIsCompatible = false;
-	checkSlow(GetAssociatedField(MemberNodeSpawner) == MemberField);
+	checkSlow(FBlueprintNodeSpawnerUtils::GetAssociatedField(MemberNodeSpawner) == MemberField);
 
 	UClass* OwnerClass = Cast<UClass>(MemberField->GetOuter());
 	if ((Pin->Direction == EGPD_Output) && (OwnerClass != nullptr))
@@ -952,7 +842,7 @@ static bool BlueprintActionFilterImpl::IsPinCompatibleWithTargetSelf(UEdGraphPin
 static bool BlueprintActionFilterImpl::IsFunctionMissingPinParam(FBlueprintActionFilter const& Filter, UBlueprintNodeSpawner const* BlueprintAction)
 {
 	bool bIsFilteredOut = false;
-	if (UFunction const* AssociatedFunc = GetAssociatedFunction(BlueprintAction))
+	if (UFunction const* AssociatedFunc = FBlueprintNodeSpawnerUtils::GetAssociatedFunction(BlueprintAction))
 	{
 		UEdGraphSchema_K2 const* K2Schema = GetDefault<UEdGraphSchema_K2>();
 		bool const bIsEventSpawner = BlueprintAction->NodeClass->IsChildOf<UK2Node_Event>();
@@ -997,7 +887,7 @@ static bool BlueprintActionFilterImpl::IsFunctionMissingPinParam(FBlueprintActio
 static bool BlueprintActionFilterImpl::IsMissmatchedPropertyType(FBlueprintActionFilter const& Filter, UBlueprintNodeSpawner const* BlueprintAction)
 {
 	bool bIsFilteredOut = false;
-	if (UProperty const* Property = GetAssociatedProperty(BlueprintAction))
+	if (UProperty const* Property = FBlueprintNodeSpawnerUtils::GetAssociatedProperty(BlueprintAction))
 	{
 		TArray<UEdGraphPin*> const& ContextPins = Filter.Context.Pins;
 		if (ContextPins.Num() > 0)
@@ -1063,7 +953,7 @@ static bool BlueprintActionFilterImpl::IsMissingMatchingPinParam(FBlueprintActio
 {
 	bool bIsFilteredOut = false;
 
-	UField const* AssociatedField = GetAssociatedField(BlueprintAction);
+	UField const* AssociatedField = FBlueprintNodeSpawnerUtils::GetAssociatedField(BlueprintAction);
 	// we have a separate tests for field nodes (IsFunctionMissingPinParam/IsMissmatchedPropertyType)
 	if (AssociatedField == nullptr)
 	{
@@ -1125,6 +1015,11 @@ FBlueprintActionFilter::FBlueprintActionFilter(uint32 Flags/*= 0x00*/)
 
 	AddIsFilteredTest(FIsFilteredDelegate::CreateStatic(IsFilteredNodeType, (Flags & BPFILTER_ExcludeChildNodeTypes) != 0));
 	AddIsFilteredTest(FIsFilteredDelegate::CreateStatic(IsExternalField));
+	
+	// don't know why we only allow impure ("imperitive") and const functions,
+	// but this mimics legacy functionality in GetFunctionCallsOnSelectedComponents()
+	AddIsFilteredTest(FIsFilteredDelegate::CreateStatic(IsBoundFunctionImpure)); 
+
 	AddIsFilteredTest(FIsFilteredDelegate::CreateStatic(IsBoundToUnselectedObject));
 	AddIsFilteredTest(FIsFilteredDelegate::CreateStatic(IsOutOfScopeLocalVariable));
 
diff --git a/Engine/Source/Editor/BlueprintGraph/Private/BlueprintBoundNodeSpawner.cpp b/Engine/Source/Editor/BlueprintGraph/Private/BlueprintBoundNodeSpawner.cpp
--- a/Engine/Source/Editor/BlueprintGraph/Private/BlueprintBoundNodeSpawner.cpp
+++ b/Engine/Source/Editor/BlueprintGraph/Private/BlueprintBoundNodeSpawner.cpp
@@ -2,6 +2,7 @@
 
 #include "BlueprintGraphPrivatePCH.h"
 #include "BlueprintBoundNodeSpawner.h"
+#include "BlueprintNodeSpawnerUtils.h"
 #include "BlueprintVariableNodeSpawner.h"
 #include "BlueprintDelegateNodeSpawner.h"
 #include "K2Node_AddComponent.h"
@@ -19,6 +20,17 @@
 
 namespace BlueprintBoundNodeSpawnerImpl
 {
+
+	/**
+	 * 
+	 * 
+	 * @param  NewNode	
+	 * @param  NodeSpawner	
+	 * @param  BindingObj	
+	 * @return 
+	 */
+	static bool BindComponentEventNode(UK2Node_ComponentBoundEvent* NewNode, UBlueprintNodeSpawner const* NodeSpawner, UObject* BindingObj);
+
 	/**
 	 * For UK2Node_AddComponent nodes. Spawns the component with a certain 
 	 * asset (mesh, child-actor, etc.).
@@ -26,7 +38,7 @@ namespace BlueprintBoundNodeSpawnerImpl
 	 * @param  NewNode	The newly spawned node that needs to be bound.
 	 * @param  AssetObj	The asset object you want bound to the node.
 	 */
-	static void BindAddComponentNodeWithAsset(UEdGraphNode* NewNode, UObject* AssetObj);
+	static bool BindAddComponentNodeWithAsset(UK2Node_AddComponent* NewNode, UObject* AssetObj);
 
 	/**
 	 * Binds function nodes to specific properties. Results in a call to a 
@@ -35,54 +47,84 @@ namespace BlueprintBoundNodeSpawnerImpl
 	 * @param  NewNode		The newly spawned node that needs to be bound.
 	 * @param  PropertyObj	The property object you want bound to the node.
 	 */
-	static void BindFunctionNodeWithVariable(UEdGraphNode* NewNode, UObject* VariableObj);
+	static bool BindFunctionNodeWithVariable(UK2Node_CallFunction* NewNode, UObject* PropertyObj);
 };
 
 //------------------------------------------------------------------------------
-static void BlueprintBoundNodeSpawnerImpl::BindAddComponentNodeWithAsset(UEdGraphNode* NewNode, UObject* AssetObj)
+static bool BlueprintBoundNodeSpawnerImpl::BindComponentEventNode(UK2Node_ComponentBoundEvent* NewNode, UBlueprintNodeSpawner const* NodeSpawner, UObject* BindingObj)
 {
-	// @TODO: assert the UObject is an asset
-	UK2Node_AddComponent* AddCompNode = CastChecked<UK2Node_AddComponent>(NewNode);
-	
-	if (UActorComponent* ComponentTemplate = AddCompNode->GetTemplateFromNode())
+	bool bSuccessfulBinding = false;
+	if (UObjectProperty* ComponentProperty = Cast<UObjectProperty>(BindingObj))
 	{
-		FComponentAssetBrokerage::AssignAssetToComponent(ComponentTemplate, AssetObj);
+		UProperty const* ActionProperty = FBlueprintNodeSpawnerUtils::GetAssociatedProperty(NodeSpawner);
+		if (UMulticastDelegateProperty const* DelegateProperty = Cast<UMulticastDelegateProperty>(ActionProperty))
+		{
+			NewNode->InitializeComponentBoundEventParams(ComponentProperty, DelegateProperty);
+			bSuccessfulBinding = true;
+		}
 	}
+
+	return bSuccessfulBinding;
 }
 
 //------------------------------------------------------------------------------
-static void BlueprintBoundNodeSpawnerImpl::BindFunctionNodeWithVariable(UEdGraphNode* NewNode, UObject* PropertyObj)
+static bool BlueprintBoundNodeSpawnerImpl::BindAddComponentNodeWithAsset(UK2Node_AddComponent* NewNode, UObject* AssetObj)
 {
-	UK2Node_CallFunction* FuncNode = CastChecked<UK2Node_CallFunction>(NewNode);
-	if (UProperty const* BindingProperty = Cast<UProperty>(PropertyObj))
+	bool bSuccessfulBinding = false;
+
+	bool const bIsTemplateNode = (NewNode->GetOutermost() == GetTransientPackage());
+	if (!bIsTemplateNode)
+	{
+		if (UActorComponent* ComponentTemplate = NewNode->GetTemplateFromNode())
+		{
+			// @TODO: ensure the UObject is an asset
+			bSuccessfulBinding = FComponentAssetBrokerage::AssignAssetToComponent(ComponentTemplate, AssetObj);
+		}
+	}	
+	return bSuccessfulBinding;
+}
+
+//------------------------------------------------------------------------------
+static bool BlueprintBoundNodeSpawnerImpl::BindFunctionNodeWithVariable(UK2Node_CallFunction* NewNode, UObject* PropertyObj)
+{
+	bool bSuccessfulBinding = false;
+	UProperty const* BindingProperty = Cast<UProperty>(PropertyObj);
+
+	bool const bIsTemplateNode = (NewNode->GetOutermost() == GetTransientPackage());
+	// @TODO: if we get it so pins can be matched with connected nodes, then we
+	//        should forgo the bIsTemplateNode check and let it bind.
+	if ((BindingProperty != nullptr) && bIsTemplateNode)
 	{
 		// @TODO: unnecessary allocation, could do this all by hand, or with a local
 		UBlueprintNodeSpawner* TempNodeSpawner = UBlueprintVariableNodeSpawner::Create(UK2Node_VariableGet::StaticClass(), BindingProperty);
 		
 		float const EstimatedVarNodeWidth = 224.0f;
 		FVector2D VarNodePos;
-		VarNodePos.X = FuncNode->NodePosX - EstimatedVarNodeWidth;
-		VarNodePos.Y = FuncNode->NodePosY;
+		VarNodePos.X = NewNode->NodePosX - EstimatedVarNodeWidth;
+		VarNodePos.Y = NewNode->NodePosY;
 		
 		float const EstimatedVarNodeHeight  = 48.0f;
-		float const EstimatedFuncNodeHeight = UEdGraphSchema_K2::EstimateNodeHeight(FuncNode);
-		float const FuncNodeMidYCoordinate  = FuncNode->NodePosY + (EstimatedFuncNodeHeight / 2.0f);
+		float const EstimatedFuncNodeHeight = UEdGraphSchema_K2::EstimateNodeHeight(NewNode);
+		float const FuncNodeMidYCoordinate = NewNode->NodePosY + (EstimatedFuncNodeHeight / 2.0f);
 		VarNodePos.Y = FuncNodeMidYCoordinate - (EstimatedVarNodeWidth / 2.0f);
 		
-		UEdGraph* ParentGraph = FuncNode->GetGraph();
+		UEdGraph* ParentGraph = NewNode->GetGraph();
 		UK2Node_VariableGet* GetVarNode = CastChecked<UK2Node_VariableGet>(TempNodeSpawner->Invoke(ParentGraph, VarNodePos));
 
 		ParentGraph->Modify();
 		ParentGraph->AddNode(GetVarNode, /*bFromUI =*/false, /*bSelectNewNode =*/false); 
 		
 		UEdGraphPin* LiteralOutput = GetVarNode->GetValuePin();
-		UEdGraphPin* CallSelfInput = FuncNode->FindPin(GetDefault<UEdGraphSchema_K2>()->PN_Self);
+		UEdGraphPin* CallSelfInput = NewNode->FindPin(GetDefault<UEdGraphSchema_K2>()->PN_Self);
 		// connect the new "get-var" node with the spawned function node
 		if ((LiteralOutput != nullptr) && (CallSelfInput != nullptr))
 		{
 			LiteralOutput->MakeLinkTo(CallSelfInput);
+			bSuccessfulBinding = true;
 		}
 	}
+
+	return bSuccessfulBinding;
 }
 
 /*******************************************************************************
@@ -90,26 +132,26 @@ static void BlueprintBoundNodeSpawnerImpl::BindFunctionNodeWithVariable(UEdGraph
  ******************************************************************************/
 
 //------------------------------------------------------------------------------
-UBlueprintBoundNodeSpawner* UBlueprintBoundNodeSpawner::Create(TSubclassOf<UEdGraphNode> NodeClass, UObject* BoundObject, UObject* Outer/*= nullptr*/)
+UBlueprintBoundNodeSpawner* UBlueprintBoundNodeSpawner::Create(TSubclassOf<UEdGraphNode> NodeClass, UClass* BoundObjType, UObject* Outer/*= nullptr*/)
 {
 	UBlueprintNodeSpawner* SubSpawner = UBlueprintNodeSpawner::Create(NodeClass);
-	return Create(SubSpawner, BoundObject, Outer);
+	return Create(SubSpawner, BoundObjType, Outer);
 }
 
 //------------------------------------------------------------------------------
-UBlueprintBoundNodeSpawner* UBlueprintBoundNodeSpawner::Create(UBlueprintNodeSpawner* SubSpawner, UObject* BoundObject, UObject* Outer/*= nullptr*/)
+UBlueprintBoundNodeSpawner* UBlueprintBoundNodeSpawner::Create(UBlueprintNodeSpawner* SubSpawner, UClass* BoundObjType, UObject* Outer/*= nullptr*/)
 {
 	if (Outer == nullptr)
 	{
 		Outer = GetTransientPackage();
 	}
 	check(SubSpawner != nullptr);
+	check(BoundObjType != nullptr);
 
 	UBlueprintBoundNodeSpawner* NodeSpawner = NewObject<UBlueprintBoundNodeSpawner>(Outer);
-	NodeSpawner->NodeClass   = SubSpawner->NodeClass;
-	NodeSpawner->SubSpawner  = DuplicateObject(SubSpawner, NodeSpawner);
-	NodeSpawner->SubSpawner->CustomizeNodeDelegate = SubSpawner->CustomizeNodeDelegate;
-	NodeSpawner->BoundObjPtr = BoundObject;
+	NodeSpawner->NodeClass     = SubSpawner->NodeClass;
+	NodeSpawner->SubSpawner    = SubSpawner;
+	NodeSpawner->BoundObjClass = BoundObjType;
 
 	return NodeSpawner;
 }
@@ -118,6 +160,7 @@ UBlueprintBoundNodeSpawner* UBlueprintBoundNodeSpawner::Create(UBlueprintNodeSpa
 UBlueprintBoundNodeSpawner::UBlueprintBoundNodeSpawner(class FPostConstructInitializeProperties const& PCIP)
 	: Super(PCIP)
 	, BoundObjPtr(nullptr)
+	, BoundObjClass(UObject::StaticClass())
 {
 }
 
@@ -127,25 +170,39 @@ UEdGraphNode* UBlueprintBoundNodeSpawner::Invoke(UEdGraph* ParentGraph, FVector2
 	using namespace BlueprintBoundNodeSpawnerImpl;
 	check(SubSpawner != nullptr);
 
-	UEdGraphNode* NewNode = SubSpawner->Invoke(ParentGraph, Location);
-	
-	UObject* BoundObj = nullptr;
-	if (BoundObjPtr.IsValid())
+	FCustomizeNodeDelegate UserNodeSetup = SubSpawner->CustomizeNodeDelegate;
+	// some binding is important to do before the node is fully formed, so we 
+	// have to construct this callback for our sub-spawner to use when spawning
+	// the node.
+	auto PostSpawnSetupLambda = [](UEdGraphNode* NewNode, bool bIsTemplateNode, UObject* BoundObject, UBlueprintNodeSpawner* NodeSpawner, FCustomizeNodeDelegate UserDelegate)
 	{
-		BoundObj = BoundObjPtr.Get();
-	}
-	
-	bool const bIsTemplateNode = (ParentGraph->GetOutermost() == GetTransientPackage());
-	if (!bIsTemplateNode)
-	{
-		if (NewNode->IsA<UK2Node_AddComponent>())
+		// user could have changed the node class (to something like
+		// UK2Node_BaseAsyncTask, which also wraps a function)
+		if (UK2Node_ComponentBoundEvent* ComponentEventNode = Cast<UK2Node_ComponentBoundEvent>(NewNode))
 		{
-			BindAddComponentNodeWithAsset(NewNode, BoundObj);
+			BindComponentEventNode(ComponentEventNode, NodeSpawner, BoundObject);
 		}
-		else if (NewNode->IsA<UK2Node_CallFunction>())
+		else if (UK2Node_AddComponent* AddComponentNode = Cast<UK2Node_AddComponent>(NewNode))
 		{
-			BindFunctionNodeWithVariable(NewNode, BoundObj);
+			BindAddComponentNodeWithAsset(AddComponentNode, BoundObject);
 		}
+
+		UserDelegate.ExecuteIfBound(NewNode, bIsTemplateNode);
+	};
+
+	UObject* BoundObj = BoundObjPtr.Get();
+	SubSpawner->CustomizeNodeDelegate = FCustomizeNodeDelegate::CreateStatic(PostSpawnSetupLambda, BoundObj, SubSpawner, UserNodeSetup);
+
+	UEdGraphNode* NewNode = SubSpawner->Invoke(ParentGraph, Location);
+	// restore the old custom delegate on the sub-spawner
+	SubSpawner->CustomizeNodeDelegate = UserNodeSetup;
+
+	// this simply spawns an additional node and attaches it, so it is ok if 
+	// we bind this after the sub-spawner's Invoke() (it doesn't affect the 
+	// initial node setup).
+	if (UK2Node_CallFunction* FuncNode = Cast<UK2Node_CallFunction>(NewNode))
+	{
+		BindFunctionNodeWithVariable(FuncNode, BoundObj);
 	}
 	return NewNode;
 }
@@ -195,14 +252,51 @@ FText UBlueprintBoundNodeSpawner::GetDefaultMenuCategory() const
 }
 
 //------------------------------------------------------------------------------
-UObject const* UBlueprintBoundNodeSpawner::GetBoundObject() const
+bool UBlueprintBoundNodeSpawner::IsOfBindingType(UObject const* BindingCandidate) const
 {
-	UObject* RawBoundObj = nullptr;
-	if (BoundObjPtr.IsValid())
+	bool bIsCompatible = false;
+	if (BindingCandidate == nullptr)
+	{
+		bIsCompatible = (BoundObjClass == nullptr);
+	}
+	else if (BoundObjClass != nullptr)
+	{
+		UClass* ObjectClass = BindingCandidate->GetClass();
+		bIsCompatible = ObjectClass->IsChildOf(BoundObjClass);
+
+		if (UObjectProperty const* ObjProperty = Cast<UObjectProperty>(BindingCandidate))
+		{
+			if (UField const* AssociatedField = FBlueprintNodeSpawnerUtils::GetAssociatedField(this))
+			{
+				AssociatedField->GetOwnerClass()->IsChildOf(ObjProperty->PropertyClass);
+			}
+		}
+	}
+
+	return bIsCompatible;
+}
+
+//------------------------------------------------------------------------------
+bool UBlueprintBoundNodeSpawner::SetBoundObject(UObject const* Object)
+{ 
+	bool bSuccess = false;
+	if ((Object == nullptr) || IsOfBindingType(Object))
+	{	
+		bSuccess = true;
+		BoundObjPtr = Object;
+	}
+	else
 	{
-		RawBoundObj = BoundObjPtr.Get();
+		ensureMsgf(false, TEXT("Cannot bind object (%s) with this node-spawner, it is not a '%s'"), *Object->GetName(), BoundObjClass->GetName());
 	}
-	return RawBoundObj;
+	return bSuccess;
+}
+
+//------------------------------------------------------------------------------
+UObject const* UBlueprintBoundNodeSpawner::GetBoundObject() const
+{
+	// Get() does IsValid() checks internally for us
+	return BoundObjPtr.Get();
 }
 
 //------------------------------------------------------------------------------
diff --git a/Engine/Source/Editor/BlueprintGraph/Private/BlueprintNodeSpawnerUtils.cpp b/Engine/Source/Editor/BlueprintGraph/Private/BlueprintNodeSpawnerUtils.cpp
new file mode 100644
--- /dev/null
+++ b/Engine/Source/Editor/BlueprintGraph/Private/BlueprintNodeSpawnerUtils.cpp
@@ -0,0 +1,98 @@
+// Copyright 1998-2014 Epic Games, Inc. All Rights Reserved.
+
+#include "BlueprintGraphPrivatePCH.h"
+#include "BlueprintNodeSpawnerUtils.h"
+#include "BlueprintNodeSpawner.h"
+#include "BlueprintFunctionNodeSpawner.h"
+#include "BlueprintVariableNodeSpawner.h"
+#include "BlueprintEventNodeSpawner.h"
+#include "BlueprintComponentNodeSpawner.h"
+#include "BlueprintBoundNodeSpawner.h"
+
+//------------------------------------------------------------------------------
+UClass* FBlueprintNodeSpawnerUtils::GetAssociatedClass(UBlueprintNodeSpawner const* BlueprintAction)
+{
+	UClass* ActionClass = BlueprintAction->NodeClass;
+	
+	if (UField const* AssociatedField = GetAssociatedField(BlueprintAction))
+	{
+		ActionClass = AssociatedField->GetOwnerClass();
+	}
+	else if (UBlueprintComponentNodeSpawner const* ComponentSpawner = Cast<UBlueprintComponentNodeSpawner>(BlueprintAction))
+	{
+		ActionClass = ComponentSpawner->GetComponentClass();
+	}
+	else if (UBlueprintBoundNodeSpawner const* BoundSpawner = Cast<UBlueprintBoundNodeSpawner>(BlueprintAction))
+	{
+		ActionClass = GetAssociatedClass(BoundSpawner->GetSubSpawner());
+	}
+	else if (UBlueprintVariableNodeSpawner const* VarSpawner = Cast<UBlueprintVariableNodeSpawner>(BlueprintAction))
+	{
+		UObject const* VarOuter = VarSpawner->GetVarOuter();
+		if (VarSpawner->IsLocalVariable())
+		{
+			UBlueprint* Blueprint = FBlueprintEditorUtils::FindBlueprintForGraphChecked(CastChecked<UEdGraph>(VarOuter));
+			ActionClass = (Blueprint->SkeletonGeneratedClass != nullptr) ? Blueprint->SkeletonGeneratedClass : Blueprint->ParentClass;
+		}
+		// else, should have been caught in GetAssociatedField()
+	}
+	
+	return ActionClass;
+}
+
+//------------------------------------------------------------------------------
+UField const* FBlueprintNodeSpawnerUtils::GetAssociatedField(UBlueprintNodeSpawner const* BlueprintAction)
+{
+	UField const* ClassField = nullptr;
+
+	if (UFunction const* Function = GetAssociatedFunction(BlueprintAction))
+	{
+		ClassField = Function;
+	}
+	else if (UProperty const* Property = GetAssociatedProperty(BlueprintAction))
+	{
+		ClassField = Property;
+	}
+	else if (UBlueprintBoundNodeSpawner const* BoundSpawner = Cast<UBlueprintBoundNodeSpawner>(BlueprintAction))
+	{
+		ClassField = GetAssociatedField(BoundSpawner->GetSubSpawner());
+	}
+	return ClassField;
+}
+
+//------------------------------------------------------------------------------
+UFunction const* FBlueprintNodeSpawnerUtils::GetAssociatedFunction(UBlueprintNodeSpawner const* BlueprintAction)
+{
+	UFunction const* Function = nullptr;
+
+	if (UBlueprintFunctionNodeSpawner const* FuncNodeSpawner = Cast<UBlueprintFunctionNodeSpawner>(BlueprintAction))
+	{
+		Function = FuncNodeSpawner->GetFunction();
+	}
+	else if (UBlueprintEventNodeSpawner const* EventSpawner = Cast<UBlueprintEventNodeSpawner>(BlueprintAction))
+	{
+		Function = EventSpawner->GetEventFunction();
+	}
+	else if (UBlueprintBoundNodeSpawner const* BoundSpawner = Cast<UBlueprintBoundNodeSpawner>(BlueprintAction))
+	{
+		Function = GetAssociatedFunction(BoundSpawner->GetSubSpawner());
+	}
+
+	return Function;
+}
+
+//------------------------------------------------------------------------------
+UProperty const* FBlueprintNodeSpawnerUtils::GetAssociatedProperty(UBlueprintNodeSpawner const* BlueprintAction)
+{
+	UProperty const* Property = nullptr;
+
+	if (UBlueprintDelegateNodeSpawner const* PropertySpawner = Cast<UBlueprintDelegateNodeSpawner>(BlueprintAction))
+	{
+		Property = PropertySpawner->GetProperty();
+	}
+	else if (UBlueprintVariableNodeSpawner const* VarSpawner = Cast<UBlueprintVariableNodeSpawner>(BlueprintAction))
+	{
+		Property = VarSpawner->GetVarProperty();
+	}
+	return Property;
+}
diff --git a/Engine/Source/Editor/BlueprintGraph/Private/BlueprintNodeSpawnerUtils.h b/Engine/Source/Editor/BlueprintGraph/Private/BlueprintNodeSpawnerUtils.h
new file mode 100644
--- /dev/null
+++ b/Engine/Source/Editor/BlueprintGraph/Private/BlueprintNodeSpawnerUtils.h
@@ -0,0 +1,48 @@
+// Copyright 1998-2014 Epic Games, Inc. All Rights Reserved.
+
+#pragma once
+
+// Forward declarations
+class UBlueprintNodeSpawner;
+
+
+struct FBlueprintNodeSpawnerUtils
+{
+	/**
+	 * All actions are associated with some class, whether it's the class of the
+	 * node it will spawn, or the outer of the associated field (property,
+	 * function, etc.).
+	 * 
+	 * @param  BlueprintAction	The node-spawner you want an associated class for.
+	 * @return The action's associated class (null if it doesn't have one).
+	 */
+	static UClass* GetAssociatedClass(UBlueprintNodeSpawner const* BlueprintAction);
+
+	/**
+	 * Certain node-spawners are associated with specific UFields (functions, 
+	 * properties, etc.). This attempts to retrieve the field from the spawner.
+	 * 
+	 * @param  BlueprintAction	The node-spawner you want an associated field for.
+	 * @return The action's associated field (null if it doesn't have one).
+	 */
+	static UField const* GetAssociatedField(UBlueprintNodeSpawner const* BlueprintAction);
+
+	/**
+	 * Certain node-spawners are associated with specific UFunction (call-
+	 * function, and event spawners). This attempts to retrieve the function 
+	 * from the spawner.
+	 * 
+	 * @param  BlueprintAction	The node-spawner you want an associated function for.
+	 * @return The action's associated function (null if it doesn't have one).
+	 */
+	static UFunction const* GetAssociatedFunction(UBlueprintNodeSpawner const* BlueprintAction);
+
+	/**
+	 * 
+	 * 
+	 * @param  BlueprintAction	
+	 * @return 
+	 */
+	static UProperty const* GetAssociatedProperty(UBlueprintNodeSpawner const* BlueprintAction);
+};
+
diff --git a/Engine/Source/Editor/BlueprintGraph/Private/K2Node_Variable.cpp b/Engine/Source/Editor/BlueprintGraph/Private/K2Node_Variable.cpp
--- a/Engine/Source/Editor/BlueprintGraph/Private/K2Node_Variable.cpp
+++ b/Engine/Source/Editor/BlueprintGraph/Private/K2Node_Variable.cpp
@@ -69,7 +69,10 @@ bool UK2Node_Variable::CreatePinForVariable(EEdGraphPinDirection Direction)
 	}
 	else
 	{
-		Message_Warn(*FString::Printf(TEXT("CreatePinForVariable: '%s' variable not found. Base class was probably changed."), *GetVarNameString()));
+		if (!VariableReference.IsLocalScope())
+		{
+			Message_Warn(*FString::Printf(TEXT("CreatePinForVariable: '%s' variable not found. Base class was probably changed."), *GetVarNameString()));
+		}
 		return false;
 	}
 
diff --git a/Engine/Source/Editor/Kismet/Private/BlueprintActionMenuBuilder.cpp b/Engine/Source/Editor/Kismet/Private/BlueprintActionMenuBuilder.cpp
--- a/Engine/Source/Editor/Kismet/Private/BlueprintActionMenuBuilder.cpp
+++ b/Engine/Source/Editor/Kismet/Private/BlueprintActionMenuBuilder.cpp
@@ -10,6 +10,7 @@
 #include "BlueprintFunctionNodeSpawner.h"
 #include "BlueprintDelegateNodeSpawner.h"
 #include "BlueprintVariableNodeSpawner.h"
+#include "BlueprintBoundNodeSpawner.h"
 #include "K2ActionMenuBuilder.h"		// for FBlueprintPaletteListBuilder/FBlueprintGraphActionListBuilder
 #include "EdGraphSchema_K2.h"			// for StaticClass(), bUseLegacyActionMenus, etc.
 #include "BlueprintEditor.h"			// for GetMyBlueprintWidget(), and GetIsContextSensitive()
@@ -38,6 +39,8 @@ public:
 	FText RootCategory;
 	/** The menu sort order to set every menu item with */
 	int32 MenuGrouping;
+	/** Cached context for the blueprint menu being built */
+	FBlueprintActionContext const& Context;
 	
 	/**
 	 * Spawns a new FBlueprintActionMenuItem with the node-spawner. Constructs
@@ -47,7 +50,7 @@ public:
 	 * @param  Action			The node-spawner that the new menu item should wrap.
 	 * @return A newly allocated FBlueprintActionMenuItem (which wraps the supplied action).
 	 */
-	TSharedPtr<FEdGraphSchemaAction> MakeActionMenuItem(TWeakPtr<FBlueprintEditor> EditorContext, UBlueprintNodeSpawner* Action);
+	TSharedPtr<FEdGraphSchemaAction> MakeActionMenuItem(TWeakPtr<FBlueprintEditor> EditorContext, UBlueprintNodeSpawner const* Action);
 
 	/**
 	 * Spawns a new FBlueprintDragDropMenuItem with the node-spawner. Constructs
@@ -56,7 +59,7 @@ public:
 	 * @param  SampleAction	One of the (possibly) many node-spawners that this menu item is set to represent.
 	 * @return A newly allocated FBlueprintActionMenuItem (which wraps the supplied action).
 	 */
-	TSharedPtr<FEdGraphSchemaAction> MakeDragDropMenuItem(UBlueprintNodeSpawner* SampleAction);
+	TSharedPtr<FEdGraphSchemaAction> MakeDragDropMenuItem(UBlueprintNodeSpawner const* SampleAction);
 	
 private:
 	/**
@@ -68,7 +71,7 @@ private:
 	 * @param  Action		The action you want to suss name information from.
 	 * @return A name for the menu item wrapping this action.
 	 */
-	FText GetMenuNameForAction(TWeakPtr<FBlueprintEditor> EditorContext, UBlueprintNodeSpawner* Action);
+	FText GetMenuNameForAction(TWeakPtr<FBlueprintEditor> EditorContext, UBlueprintNodeSpawner const* Action);
 
 	/**
 	 * Attempts to pull a menu category from the supplied spawner. If one isn't 
@@ -78,7 +81,7 @@ private:
 	 * @param  Action		The action you want to suss category information from.
 	 * @return A category for the menu item wrapping this action.
 	 */
-	FText GetCategoryForAction(TWeakPtr<FBlueprintEditor> EditorContext, UBlueprintNodeSpawner* Action);
+	FText GetCategoryForAction(TWeakPtr<FBlueprintEditor> EditorContext, UBlueprintNodeSpawner const* Action);
 
 	/**
 	 * Attempts to pull a menu tooltip from the supplied spawner. If one isn't 
@@ -88,7 +91,7 @@ private:
 	 * @param  Action		The action you want to suss tooltip information from.
 	 * @return A tooltip for the menu item wrapping this action.
 	 */
-	FText GetTooltipForAction(TWeakPtr<FBlueprintEditor> EditorContext, UBlueprintNodeSpawner* Action);
+	FText GetTooltipForAction(TWeakPtr<FBlueprintEditor> EditorContext, UBlueprintNodeSpawner const* Action);
 
 	/**
 	 * Attempts to pull a keywords from the supplied spawner. If one isn't 
@@ -98,7 +101,7 @@ private:
 	 * @param  Action		The action you want to suss keyword information from.
 	 * @return A keyword text string for the menu item wrapping this action.
 	 */
-	FText GetSearchKeywordsForAction(TWeakPtr<FBlueprintEditor> EditorContext, UBlueprintNodeSpawner* Action);
+	FText GetSearchKeywordsForAction(TWeakPtr<FBlueprintEditor> EditorContext, UBlueprintNodeSpawner const* Action);
 
 	/**
 	 * Attempts to pull a menu icon information from the supplied spawner. If
@@ -110,7 +113,7 @@ private:
 	 * @param  ColorOut		The color to tint the icon with.
 	 * @return Name of the brush to use (use FEditorStyle::GetBrush() to resolve).
 	 */
-	FName GetMenuIconForAction(TWeakPtr<FBlueprintEditor> EditorContext, UBlueprintNodeSpawner* Action, FLinearColor& ColorOut);
+	FName GetMenuIconForAction(TWeakPtr<FBlueprintEditor> EditorContext, UBlueprintNodeSpawner const* Action, FLinearColor& ColorOut);
 
 	/**
 	 * Utility getter function that retrieves the blueprint context for the menu
@@ -127,11 +130,7 @@ private:
 	 * @param  EditorContext	
 	 * @return 
 	 */
-	UEdGraphNode* GetTemplateNode(UBlueprintNodeSpawner* Action, TWeakPtr<FBlueprintEditor> EditorContext) const;
-
-private:
-	/** Cached context for the blueprint menu being built */
-	FBlueprintActionContext const& Context;
+	UEdGraphNode* GetTemplateNode(UBlueprintNodeSpawner const* Action, TWeakPtr<FBlueprintEditor> EditorContext) const;	
 };
 
 //------------------------------------------------------------------------------
@@ -143,7 +142,7 @@ FBlueprintActionMenuItemFactory::FBlueprintActionMenuItemFactory(FBlueprintActio
 }
 
 //------------------------------------------------------------------------------
-TSharedPtr<FEdGraphSchemaAction> FBlueprintActionMenuItemFactory::MakeActionMenuItem(TWeakPtr<FBlueprintEditor> EditorContext, UBlueprintNodeSpawner* Action)
+TSharedPtr<FEdGraphSchemaAction> FBlueprintActionMenuItemFactory::MakeActionMenuItem(TWeakPtr<FBlueprintEditor> EditorContext, UBlueprintNodeSpawner const* Action)
 {
 	FLinearColor IconTint = FLinearColor::White;
 	FName IconBrushName = GetMenuIconForAction(EditorContext, Action, IconTint);
@@ -159,7 +158,7 @@ TSharedPtr<FEdGraphSchemaAction> FBlueprintActionMenuItemFactory::MakeActionMenu
 }
 
 //------------------------------------------------------------------------------
-TSharedPtr<FEdGraphSchemaAction> FBlueprintActionMenuItemFactory::MakeDragDropMenuItem(UBlueprintNodeSpawner* SampleAction)
+TSharedPtr<FEdGraphSchemaAction> FBlueprintActionMenuItemFactory::MakeDragDropMenuItem(UBlueprintNodeSpawner const* SampleAction)
 {
 	// FBlueprintDragDropMenuItem takes care of its own menu MenuDescription, etc.
 	FBlueprintDragDropMenuItem* NewMenuItem = new FBlueprintDragDropMenuItem(Context, SampleAction, MenuGrouping);
@@ -169,7 +168,7 @@ TSharedPtr<FEdGraphSchemaAction> FBlueprintActionMenuItemFactory::MakeDragDropMe
 }
 
 //------------------------------------------------------------------------------
-FText FBlueprintActionMenuItemFactory::GetMenuNameForAction(TWeakPtr<FBlueprintEditor> EditorContext, UBlueprintNodeSpawner* Action)
+FText FBlueprintActionMenuItemFactory::GetMenuNameForAction(TWeakPtr<FBlueprintEditor> EditorContext, UBlueprintNodeSpawner const* Action)
 {
 	check(Action != nullptr);
 	// give the action the chance to save on performance (to keep from having to spawn a template node)
@@ -192,7 +191,7 @@ FText FBlueprintActionMenuItemFactory::GetMenuNameForAction(TWeakPtr<FBlueprintE
 }
 
 //------------------------------------------------------------------------------
-FText FBlueprintActionMenuItemFactory::GetCategoryForAction(TWeakPtr<FBlueprintEditor> EditorContext, UBlueprintNodeSpawner* Action)
+FText FBlueprintActionMenuItemFactory::GetCategoryForAction(TWeakPtr<FBlueprintEditor> EditorContext, UBlueprintNodeSpawner const* Action)
 {
 	check(Action != nullptr);
 	// give the action the chance to save on performance (to keep from having to spawn a template node)
@@ -212,7 +211,7 @@ FText FBlueprintActionMenuItemFactory::GetCategoryForAction(TWeakPtr<FBlueprintE
 		{
 			// put uncategorized function calls in a member function category
 			// (sorted by their respective classes)
-			if (UBlueprintFunctionNodeSpawner* FuncSpawner = Cast<UBlueprintFunctionNodeSpawner>(Action))
+			if (UBlueprintFunctionNodeSpawner const* FuncSpawner = Cast<UBlueprintFunctionNodeSpawner>(Action))
 			{
 				UFunction const* Function = FuncSpawner->GetFunction();
 				check(Function != nullptr);
@@ -237,7 +236,7 @@ FText FBlueprintActionMenuItemFactory::GetCategoryForAction(TWeakPtr<FBlueprintE
 }
 
 //------------------------------------------------------------------------------
-FText FBlueprintActionMenuItemFactory::GetTooltipForAction(TWeakPtr<FBlueprintEditor> EditorContext, UBlueprintNodeSpawner* Action)
+FText FBlueprintActionMenuItemFactory::GetTooltipForAction(TWeakPtr<FBlueprintEditor> EditorContext, UBlueprintNodeSpawner const* Action)
 {
 	check(Action != nullptr);
 	// give the action the chance to save on performance (to keep from having to spawn a template node)
@@ -255,7 +254,7 @@ FText FBlueprintActionMenuItemFactory::GetTooltipForAction(TWeakPtr<FBlueprintEd
 }
 
 //------------------------------------------------------------------------------
-FText FBlueprintActionMenuItemFactory::GetSearchKeywordsForAction(TWeakPtr<FBlueprintEditor> EditorContext, UBlueprintNodeSpawner* Action)
+FText FBlueprintActionMenuItemFactory::GetSearchKeywordsForAction(TWeakPtr<FBlueprintEditor> EditorContext, UBlueprintNodeSpawner const* Action)
 {
 	check(Action != nullptr);
 	// give the action the chance to save on performance (to keep from having to spawn a template node)
@@ -273,7 +272,7 @@ FText FBlueprintActionMenuItemFactory::GetSearchKeywordsForAction(TWeakPtr<FBlue
 }
 
 //------------------------------------------------------------------------------
-FName FBlueprintActionMenuItemFactory::GetMenuIconForAction(TWeakPtr<FBlueprintEditor> EditorContext, UBlueprintNodeSpawner* Action, FLinearColor& ColorOut)
+FName FBlueprintActionMenuItemFactory::GetMenuIconForAction(TWeakPtr<FBlueprintEditor> EditorContext, UBlueprintNodeSpawner const* Action, FLinearColor& ColorOut)
 {
 	FName BrushName = Action->GetDefaultMenuIcon(ColorOut);
 	if (BrushName.IsNone())
@@ -298,7 +297,7 @@ UBlueprint* FBlueprintActionMenuItemFactory::GetTargetBlueprint() const
 }
 
 //------------------------------------------------------------------------------
-UEdGraphNode* FBlueprintActionMenuItemFactory::GetTemplateNode(UBlueprintNodeSpawner* Action, TWeakPtr<FBlueprintEditor> EditorContext) const
+UEdGraphNode* FBlueprintActionMenuItemFactory::GetTemplateNode(UBlueprintNodeSpawner const* Action, TWeakPtr<FBlueprintEditor> EditorContext) const
 {
 	UEdGraph* TargetGraph = nullptr;
 	if (Context.Graphs.Num() > 0)
@@ -354,11 +353,11 @@ namespace FBlueprintActionMenuBuilderImpl
 		 * FBlueprintActionMenuItem for the specified menu (does not add the 
 		 * item to the menu-builder itself).
 		 *
-		 * @param  EditorContext
-		 * @param  Action			The node-spawner that the new menu item should wrap.
+		 * @param  EditorContext	
+		 * @param  DatabaseAction	The node-spawner that the new menu item should wrap.
 		 * @return An empty TSharedPtr if the action was filtered out, otherwise a newly allocated FBlueprintActionMenuItem.
 		 */
-		TSharedPtr<FEdGraphSchemaAction> MakeMenuItem(TWeakPtr<FBlueprintEditor> EditorContext, UBlueprintNodeSpawner* Action);
+		TSharedPtr<FEdGraphSchemaAction> MakeMenuItem(TWeakPtr<FBlueprintEditor> EditorContext, UBlueprintNodeSpawner const* const DatabaseAction);
 		
 		/**
 		 * Clears out any consolidated properties that this may have been 
@@ -374,7 +373,7 @@ namespace FBlueprintActionMenuBuilderImpl
 	};
 	
 	/**
-	 * To offer a fallback in case the action menu refactor is unstable, this
+	 * To offer a fallback in case this menu system is unstable on release, this
 	 * method implements the old way we used collect blueprint menu actions (for
 	 * both the palette and context menu).
 	 *
@@ -412,24 +411,52 @@ void FBlueprintActionMenuBuilderImpl::FMenuSectionDefinition::SetSectionSortOrde
 }
 
 //------------------------------------------------------------------------------
-TSharedPtr<FEdGraphSchemaAction> FBlueprintActionMenuBuilderImpl::FMenuSectionDefinition::MakeMenuItem(TWeakPtr<FBlueprintEditor> EditorContext, UBlueprintNodeSpawner* Action)
+TSharedPtr<FEdGraphSchemaAction> FBlueprintActionMenuBuilderImpl::FMenuSectionDefinition::MakeMenuItem(TWeakPtr<FBlueprintEditor> EditorContext, UBlueprintNodeSpawner const* const DatabaseAction)
 {
-	bool bPassedFilter = (Action != nullptr) && !Filter.IsFiltered(Action);
+	UBlueprintNodeSpawner const* MenuAction = DatabaseAction;
+
+	// we don't want the blueprint database growing out of control with an entry 
+	// for every object you could ever possibly bind to, so the 
+	// UBlueprintBoundNodeSpawner from the database serves as a template (note,
+	// it isn't bound to anything)... so, we do the binding here (if anything
+	// to bind to is selected), and use the temporary (bound) spawner in place 
+	// of the passed DatabaseAction
+	if (UBlueprintBoundNodeSpawner const* BoundNodeSpawner = Cast<UBlueprintBoundNodeSpawner>(DatabaseAction))
+	{
+		for (UObject const* SelectedObj : ItemFactory.Context.SelectedObjects)
+		{
+			// if the bound object is null, then this was most likely set up as a 
+			// template for us to copy and assign a binding object
+			if ((BoundNodeSpawner->GetBoundObject() == nullptr) && BoundNodeSpawner->IsOfBindingType(SelectedObj))
+			{
+				// don't worry about outer, if this makes it through the filter,
+				// then the menu item will add a reference to keep it around
+				UBlueprintBoundNodeSpawner* BoundAction = DuplicateObject(BoundNodeSpawner, GetTransientPackage());
+				BoundAction->SetBoundObject(SelectedObj);
+
+				// replace the menu action with one that is actually bound to 
+				// something
+				MenuAction = BoundAction;
+				break;
+			}
+		}
+	}
+	bool bPassedFilter = (MenuAction != nullptr) && !Filter.IsFiltered(MenuAction);
 
 	bool bCreateDragDropItem = false;
-	// if the user wants to consolidate all property actions, then we have to 
-	// check and see if this is one of those that needs consolidating (that 
-	// needs bCreateDragDropItem)
+	// if the caller wants to consolidate all property actions, then we have to 
+	// check and see if this is one of those that needs consolidating (needs 
+	// a FBlueprintDragDropMenuItem instead of a FBlueprintActionMenuItem)
 	if (bPassedFilter && (Flags & FBlueprintActionMenuBuilder::ConsolidatePropertyActions))
 	{
 		UProperty const* ActionProperty = nullptr;
-		if (UBlueprintVariableNodeSpawner* VariableSpawner = Cast<UBlueprintVariableNodeSpawner>(Action))
+		if (UBlueprintVariableNodeSpawner const* VariableSpawner = Cast<UBlueprintVariableNodeSpawner>(MenuAction))
 		{
 			bCreateDragDropItem  = true;
 			ActionProperty = VariableSpawner->GetVarProperty();
 			bPassedFilter = (ActionProperty != nullptr);
 		}
-		else if (UBlueprintDelegateNodeSpawner* DelegateSpawner = Cast<UBlueprintDelegateNodeSpawner>(Action))
+		else if (UBlueprintDelegateNodeSpawner const* DelegateSpawner = Cast<UBlueprintDelegateNodeSpawner>(MenuAction))
 		{
 			bCreateDragDropItem  = true;
 			ActionProperty = DelegateSpawner->GetProperty();
@@ -450,11 +477,11 @@ TSharedPtr<FEdGraphSchemaAction> FBlueprintActionMenuBuilderImpl::FMenuSectionDe
 	{
 		if (bCreateDragDropItem)
 		{
-			MenuEntry = ItemFactory.MakeDragDropMenuItem(Action);
+			MenuEntry = ItemFactory.MakeDragDropMenuItem(MenuAction);
 		}
 		else
 		{
-			MenuEntry = ItemFactory.MakeActionMenuItem(EditorContext, Action);
+			MenuEntry = ItemFactory.MakeActionMenuItem(EditorContext, MenuAction);
 		}
 	}
 	return MenuEntry;
diff --git a/Engine/Source/Editor/Kismet/Private/BlueprintDragDropMenuItem.cpp b/Engine/Source/Editor/Kismet/Private/BlueprintDragDropMenuItem.cpp
--- a/Engine/Source/Editor/Kismet/Private/BlueprintDragDropMenuItem.cpp
+++ b/Engine/Source/Editor/Kismet/Private/BlueprintDragDropMenuItem.cpp
@@ -21,7 +21,7 @@
 DEFINE_LOG_CATEGORY_STATIC(LogBlueprintDragDropMenuItem, Log, All);
 
 //------------------------------------------------------------------------------
-FBlueprintDragDropMenuItem::FBlueprintDragDropMenuItem(FBlueprintActionContext const& Context, UBlueprintNodeSpawner* SampleActionIn, int32 MenuGrouping/* = 0*/)
+FBlueprintDragDropMenuItem::FBlueprintDragDropMenuItem(FBlueprintActionContext const& Context, UBlueprintNodeSpawner const* SampleActionIn, int32 MenuGrouping/* = 0*/)
 	: SampleAction(SampleActionIn)
 {
 	check(SampleAction != nullptr);
diff --git a/Engine/Source/Editor/Kismet/Public/BlueprintDragDropMenuItem.h b/Engine/Source/Editor/Kismet/Public/BlueprintDragDropMenuItem.h
--- a/Engine/Source/Editor/Kismet/Public/BlueprintDragDropMenuItem.h
+++ b/Engine/Source/Editor/Kismet/Public/BlueprintDragDropMenuItem.h
@@ -43,7 +43,7 @@ public:
 	 * @param  SampleAction	One of the node-spawners that this "consolidated" menu item represents (shouldn't matter which one).
 	 * @param  MenuGrouping	Defines the order in which this entry is listed in the menu.
 	 */
-	FBlueprintDragDropMenuItem(FBlueprintActionContext const& Context, UBlueprintNodeSpawner* SampleAction, int32 MenuGrouping = 0);
+	FBlueprintDragDropMenuItem(FBlueprintActionContext const& Context, UBlueprintNodeSpawner const* SampleAction, int32 MenuGrouping = 0);
 	
 	// FEdGraphSchemaAction interface
 	virtual FName GetTypeId() const final { return StaticGetTypeId(); }
