Fixes for Doxygen warnings.

diff --git a/Engine/Source/Editor/Kismet/Private/SBlueprintDiff.cpp b/Engine/Source/Editor/Kismet/Private/SBlueprintDiff.cpp
--- a/Engine/Source/Editor/Kismet/Private/SBlueprintDiff.cpp
+++ b/Engine/Source/Editor/Kismet/Private/SBlueprintDiff.cpp
@@ -538,7 +538,7 @@ FDiffPanel::FDiffPanel()
 	LastFocusedPin = NULL;
 }
 
-void DiffUtils::CompareUnrelatedObjects( UObject const* A, const FNamePropertyMap& PropertyMapA, UObject const* B, const FNamePropertyMap& PropertyMapB, TArray<FName> &OutIdenticalProperties, TArray<FName> &OutDifferingProperties )
+void DiffUtils::CompareUnrelatedObjects( UObject const* A, const TMap< FName, const UProperty* >& PropertyMapA, UObject const* B, const TMap< FName, const UProperty* >& PropertyMapB, TArray<FName> &OutIdenticalProperties, TArray<FName> &OutDifferingProperties )
 {
 	TSet<FName> PropertiesInA, PropertiesInB;
 	if( A )
diff --git a/Engine/Source/Runtime/CoreUObject/Private/Misc/StringAssetReference.cpp b/Engine/Source/Runtime/CoreUObject/Private/Misc/StringAssetReference.cpp
--- a/Engine/Source/Runtime/CoreUObject/Private/Misc/StringAssetReference.cpp
+++ b/Engine/Source/Runtime/CoreUObject/Private/Misc/StringAssetReference.cpp
@@ -5,7 +5,7 @@
 #include "StringAssetReference.h"
 #include "PropertyTag.h"
 
-FStringAssetReference::FStringAssetReference(const class UObject* InObject)
+FStringAssetReference::FStringAssetReference(const UObject* InObject)
 {
 	if (InObject)
 	{
diff --git a/Engine/Source/Runtime/Slate/Private/Framework/Testing/SUserWidgetTest.cpp b/Engine/Source/Runtime/Slate/Private/Framework/Testing/SUserWidgetTest.cpp
--- a/Engine/Source/Runtime/Slate/Private/Framework/Testing/SUserWidgetTest.cpp
+++ b/Engine/Source/Runtime/Slate/Private/Framework/Testing/SUserWidgetTest.cpp
@@ -4,10 +4,7 @@
 #include "SUserWidgetTest.h"
 
 
-namespace Implementation
-{
-	
-class SUserWidgetExample : public ::SUserWidgetExample
+class SUserWidgetExampleImpl : public SUserWidgetExample
 {
 public:
 
@@ -27,10 +24,8 @@ public:
 	}
 };
 
-}
-
 
 TSharedRef<SUserWidgetExample> SUserWidgetExample::New()
 {
-	return MakeShareable(new Implementation::SUserWidgetExample()); 
+	return MakeShareable(new SUserWidgetExampleImpl()); 
 }
