#summary Added a "Go to Definition" menu entry when right clicking on call function nodes in Blueprints

#change UK2Node_CallFunction::GetFunctionGraph now returns the proper graph from whatever Blueprint the function belongs to

#change Cleaned up double clicking on function call nodes

#ttp 316688 - UE4: BLUEPRINTS: Go to definition behavior should be available in the right-click context menu for function calls

diff --git a/Engine/Source/Editor/BlueprintGraph/Classes/K2Node_CallFunction.h b/Engine/Source/Editor/BlueprintGraph/Classes/K2Node_CallFunction.h
--- a/Engine/Source/Editor/BlueprintGraph/Classes/K2Node_CallFunction.h
+++ b/Engine/Source/Editor/BlueprintGraph/Classes/K2Node_CallFunction.h
@@ -150,8 +150,12 @@ public:
 
 	static void CallForEachElementInArrayExpansion(UK2Node* Node, UEdGraphPin* MultiSelf, FKismetCompilerContext& CompilerContext, UEdGraph* SourceGraph);
 
-	/** Returns the graph for this function, if available */
-	UEdGraph* GetFunctionGraph() const;
+	/**
+	 * Returns the graph for this function, if available. In cases of calling an event, it will be the ubergraph for the event
+	 *
+	 * @param OutGraphNode		If this function calls an event, this param is the event node, otherwise it is NULL
+	 */
+	UEdGraph* GetFunctionGraph(const UEdGraphNode*& OutGraphNode) const;
 
 	/** Checks if the property is marked as "CustomStructureParam" */
 	static bool IsStructureWildcardProperty(const UFunction* InFunction, const FString& PropertyName);
diff --git a/Engine/Source/Editor/BlueprintGraph/Private/EdGraphSchema_K2.cpp b/Engine/Source/Editor/BlueprintGraph/Private/EdGraphSchema_K2.cpp
--- a/Engine/Source/Editor/BlueprintGraph/Private/EdGraphSchema_K2.cpp
+++ b/Engine/Source/Editor/BlueprintGraph/Private/EdGraphSchema_K2.cpp
@@ -992,7 +992,21 @@ void UEdGraphSchema_K2::GetContextMenuActions(const UEdGraph* CurrentGraph, cons
 					// Add the goto source code action for native functions
 					if (InGraphNode->IsA(UK2Node_CallFunction::StaticClass()))
 					{
-						MenuBuilder->AddMenuEntry(FGraphEditorCommands::Get().GotoNativeFunctionDefinition);
+						const UEdGraphNode* ResultEventNode= NULL;
+						if(Cast<UK2Node_CallFunction>(InGraphNode)->GetFunctionGraph(ResultEventNode))
+						{
+							MenuBuilder->AddMenuEntry(FGraphEditorCommands::Get().GoToDefinition);
+						}
+						else
+						{
+							MenuBuilder->AddMenuEntry(FGraphEditorCommands::Get().GotoNativeFunctionDefinition);
+						}
+					}
+
+					// Functions, macros, and composite nodes support going to a definition
+					if (InGraphNode->IsA(UK2Node_MacroInstance::StaticClass()) || InGraphNode->IsA(UK2Node_Composite::StaticClass()))
+					{
+						MenuBuilder->AddMenuEntry(FGraphEditorCommands::Get().GoToDefinition);
 					}
 
 					// show search for references for variable nodes and goto source code action
diff --git a/Engine/Source/Editor/BlueprintGraph/Private/K2Node_CallFunction.cpp b/Engine/Source/Editor/BlueprintGraph/Private/K2Node_CallFunction.cpp
--- a/Engine/Source/Editor/BlueprintGraph/Private/K2Node_CallFunction.cpp
+++ b/Engine/Source/Editor/BlueprintGraph/Private/K2Node_CallFunction.cpp
@@ -1520,9 +1520,50 @@ bool UK2Node_CallFunction::HasExternalBlueprintDependencies(TArray<class UStruct
 	return bResult || Super::HasExternalBlueprintDependencies(OptionalOutput);
 }
 
-UEdGraph* UK2Node_CallFunction::GetFunctionGraph() const
+UEdGraph* UK2Node_CallFunction::GetFunctionGraph(const UEdGraphNode*& OutGraphNode) const
 {
-	return FindObject<UEdGraph>(GetBlueprint(), *(FunctionReference.GetMemberName().ToString()));
+	OutGraphNode = NULL;
+
+	// Search for the Blueprint owner of the function graph, climbing up through the Blueprint hierarchy
+	UClass* MemberParentClass = FunctionReference.GetMemberParentClass(this);
+	if(MemberParentClass != NULL)
+	{
+		UBlueprintGeneratedClass* ParentClass = Cast<UBlueprintGeneratedClass>(MemberParentClass);
+		if(ParentClass != NULL && ParentClass->ClassGeneratedBy != NULL)
+		{
+			UBlueprint* Blueprint = Cast<UBlueprint>(ParentClass->ClassGeneratedBy);
+			while(Blueprint != NULL)
+			{
+				UEdGraph* TargetGraph = FindObject<UEdGraph>(Blueprint, *(FunctionReference.GetMemberName().ToString()));
+				if((TargetGraph != NULL) && !TargetGraph->HasAnyFlags(RF_Transient))
+				{
+					// Found the function graph in a Blueprint, return that graph
+					return TargetGraph;
+				}
+				else
+				{
+					// Did not find the function call as a graph, it may be a custom event
+					UK2Node_CustomEvent* CustomEventNode = NULL;
+
+					TArray<UK2Node_CustomEvent*> CustomEventNodes;
+					FBlueprintEditorUtils::GetAllNodesOfClass(Blueprint, CustomEventNodes);
+
+					for (UK2Node_CustomEvent* CustomEvent : CustomEventNodes)
+					{
+						if(CustomEvent->CustomFunctionName == FunctionReference.GetMemberName())
+						{
+							OutGraphNode = CustomEvent;
+							return CustomEvent->GetGraph();
+						}
+					}
+				}
+
+				ParentClass = Cast<UBlueprintGeneratedClass>(Blueprint->ParentClass);
+				Blueprint = ParentClass != NULL ? Cast<UBlueprint>(ParentClass->ClassGeneratedBy) : NULL;
+			}
+		}
+	}
+	return NULL;
 }
 
 bool UK2Node_CallFunction::IsStructureWildcardProperty(const UFunction* Function, const FString& PropertyName)
diff --git a/Engine/Source/Editor/GraphEditor/Private/GraphEditorActions.cpp b/Engine/Source/Editor/GraphEditor/Private/GraphEditorActions.cpp
--- a/Engine/Source/Editor/GraphEditor/Private/GraphEditorActions.cpp
+++ b/Engine/Source/Editor/GraphEditor/Private/GraphEditorActions.cpp
@@ -48,6 +48,8 @@ void FGraphEditorCommandsImpl::RegisterCommands()
 	UI_COMMAND( GotoNativeFunctionDefinition, "Goto Code Definition", "Goto the native code definition of this function", EUserInterfaceActionType::Button, FInputGesture() )
 	UI_COMMAND( GotoNativeVariableDefinition, "Goto Code Definition", "Goto the native code definition of this variable", EUserInterfaceActionType::Button, FInputGesture() )
 
+	UI_COMMAND( GoToDefinition, "Goto Definition", "Jumps to the graph this node is defined in if available.", EUserInterfaceActionType::Button, FInputGesture() )
+
 	UI_COMMAND( BreakPinLinks, "Break Link(s)", "Breaks pin links", EUserInterfaceActionType::Button, FInputGesture() )
 	UI_COMMAND( PromoteToVariable, "Promote to Variable", "Promotes something to a variable", EUserInterfaceActionType::Button, FInputGesture() )
 	UI_COMMAND( SplitStructPin, "Split Struct Pin", "Breaks a struct pin in to a separate pin per element", EUserInterfaceActionType::Button, FInputGesture() )
@@ -72,7 +74,6 @@ void FGraphEditorCommandsImpl::RegisterCommands()
 	UI_COMMAND( CreateComment, "Create Comment", "Create a comment box", EUserInterfaceActionType::Button, FInputGesture(EKeys::C))
 
 	UI_COMMAND( FindInstancesOfCustomEvent, "Find Instances of Event", "Find the instances of this custom event", EUserInterfaceActionType::Button, FInputGesture() )
-	
 }
 
 
diff --git a/Engine/Source/Editor/GraphEditor/Public/GraphEditorActions.h b/Engine/Source/Editor/GraphEditor/Public/GraphEditorActions.h
--- a/Engine/Source/Editor/GraphEditor/Public/GraphEditorActions.h
+++ b/Engine/Source/Editor/GraphEditor/Public/GraphEditorActions.h
@@ -70,6 +70,9 @@ public:
 	TSharedPtr< FUICommandInfo > GotoNativeFunctionDefinition;
 	TSharedPtr< FUICommandInfo > GotoNativeVariableDefinition;
 
+	// Jumps to a call function node's graph
+	TSharedPtr< FUICommandInfo > GoToDefinition;
+
 	// Pin-specific actions
 	TSharedPtr< FUICommandInfo > BreakPinLinks;
 	TSharedPtr< FUICommandInfo > PromoteToVariable;
diff --git a/Engine/Source/Editor/Kismet/Private/BlueprintEditor.cpp b/Engine/Source/Editor/Kismet/Private/BlueprintEditor.cpp
--- a/Engine/Source/Editor/Kismet/Private/BlueprintEditor.cpp
+++ b/Engine/Source/Editor/Kismet/Private/BlueprintEditor.cpp
@@ -731,6 +731,11 @@ TSharedRef<SGraphEditor> FBlueprintEditor::CreateGraphEditorWidget(TSharedRef<FT
 				FIsActionButtonVisible::CreateSP( this, &FBlueprintEditor::IsNativeCodeBrowsingAvailable )
 				);
 
+			GraphEditorCommands->MapAction( FGraphEditorCommands::Get().GoToDefinition,
+				FExecuteAction::CreateSP( this, &FBlueprintEditor::OnGoToDefinition ),
+				FCanExecuteAction::CreateSP( this, &FBlueprintEditor::CanGoToDefinition )
+				);
+
 			GraphEditorCommands->MapAction( FGraphEditorCommands::Get().GotoNativeVariableDefinition,
 				FExecuteAction::CreateSP( this, &FBlueprintEditor::GotoNativeVariableDefinition ),
 				FCanExecuteAction::CreateSP(this, &FBlueprintEditor::IsSelectionNativeVariable),
@@ -3983,7 +3988,12 @@ void FBlueprintEditor::OnExpandNodes()
 		}
 		else if (UK2Node_CallFunction* SelectedCallFunctionNode = Cast<UK2Node_CallFunction>(*NodeIt))
 		{
-			UEdGraph* FunctionGraph = SelectedCallFunctionNode->GetFunctionGraph();
+			const UEdGraphNode* ResultEventNode = NULL;
+			UEdGraph* FunctionGraph = SelectedCallFunctionNode->GetFunctionGraph(ResultEventNode);
+
+			// We should never get here when attempting to expand a call function that calls an event.
+			check(!ResultEventNode);
+
 			if(FunctionGraph)
 			{
 				DocumentManager->CleanInvalidTabs();
@@ -4025,7 +4035,9 @@ bool FBlueprintEditor::CanExpandNodes() const
 		}
 		else if (UK2Node_CallFunction* SelectedCallFunctionNode = Cast<UK2Node_CallFunction>(*NodeIt))
 		{
-			return SelectedCallFunctionNode->GetFunctionGraph() != NULL;
+			// If ResultEventNode is non-NULL, it means it is sourced by an event, we do not want to expand events
+			const UEdGraphNode* ResultEventNode = NULL;
+			return SelectedCallFunctionNode->GetFunctionGraph(ResultEventNode) != NULL && ResultEventNode == NULL;
 		}
 	}
 
@@ -4794,6 +4806,44 @@ bool FBlueprintEditor::CanStopWatchingPin() const
 	return false;
 }
 
+bool FBlueprintEditor::CanGoToDefinition() const
+{
+	const FGraphPanelSelectionSet SelectedNodes = GetSelectedNodes();
+	if(SelectedNodes.Num() == 1)
+	{
+		UObject* Node = *SelectedNodes.CreateConstIterator();
+		if(UK2Node_CallFunction* FunctionCall = Cast<UK2Node_CallFunction>(Node))
+		{
+			const UEdGraphNode* ResultEventNode = NULL;
+			if(FunctionCall->GetFunctionGraph(ResultEventNode))
+			{
+				return true;
+			}
+		}
+		else if (UK2Node_MacroInstance* MacroNode = Cast<UK2Node_MacroInstance>(Node))
+		{
+			UEdGraph* MacroGraph = MacroNode->GetMacroGraph();
+			if (MacroGraph)
+			{
+				return true;
+			}
+		}
+		else if(UK2Node_Composite* Composite = Cast<UK2Node_Composite>(Node))
+		{
+			return true;
+		}
+	}
+
+	return false;
+}
+
+void FBlueprintEditor::OnGoToDefinition()
+{
+	const FGraphPanelSelectionSet SelectedNodes = GetSelectedNodes();
+	check(SelectedNodes.Num() == 1);
+	OnNodeDoubleClicked(Cast<UEdGraphNode>(*SelectedNodes.CreateConstIterator()));
+}
+
 void FBlueprintEditor::ToggleSaveIntermediateBuildProducts()
 {
 	bSaveIntermediateBuildProducts = !bSaveIntermediateBuildProducts;
@@ -4809,54 +4859,18 @@ void FBlueprintEditor::OnNodeDoubleClicked(class UEdGraphNode* Node)
 	//@TODO: Pull these last few stragglers out into their respective nodes; requires asking the same question in a different way without knowledge of the Explorer
 	if (UK2Node_CallFunction* FunctionCall = Cast<UK2Node_CallFunction>(Node))
 	{
-		//Check whether the current graph node is a function call to a custom event.
-		UK2Node_CustomEvent* CustomEventNode = nullptr;
-		TArray<UK2Node_CustomEvent*> CustomNodes;
-		FBlueprintEditorUtils::GetAllNodesOfClass(GetBlueprintObj(), CustomNodes);
-		for (int32 i=0; i < CustomNodes.Num(); i++)
-		{
-			UK2Node_CustomEvent* CustomNodePtr = CustomNodes[ i ];
-			if ((CustomNodePtr != nullptr) && (CustomNodePtr->CustomFunctionName == FunctionCall->FunctionReference.GetMemberName()))
-			{
-				CustomEventNode = CustomNodePtr;
-			}
-		}
-		if (CustomEventNode != nullptr)
-		{
-			JumpToNode(CustomEventNode);
-		}
-		else
+		const UEdGraphNode* ResultEventNode = NULL;
+
+		if(UEdGraph* FunctionGraph = FunctionCall->GetFunctionGraph(ResultEventNode))
 		{
-			// Look for a local graph named the same thing as the function (e.g., a user created function or a harmless coincidence)
-			UEdGraph* TargetGraph = FindObject<UEdGraph>(GetBlueprintObj(), *(FunctionCall->FunctionReference.GetMemberName().ToString()));
-			if ((TargetGraph != nullptr) && !TargetGraph->HasAnyFlags(RF_Transient))
+			// If there is an event node, jump to it, otherwise jump to the function graph
+			if(ResultEventNode)
 			{
-				JumpToHyperlink(TargetGraph);
+				FKismetEditorUtilities::BringKismetToFocusAttentionOnObject(ResultEventNode, false);
 			}
 			else
 			{
-				// otherwise look for the function in another blueprint, possibly a parent
-				UClass* MemberParentClass = FunctionCall->FunctionReference.GetMemberParentClass(FunctionCall);
-				if(MemberParentClass != nullptr)
-				{
-					UBlueprintGeneratedClass* ParentClass = Cast<UBlueprintGeneratedClass>(MemberParentClass);
-					if(ParentClass != nullptr && ParentClass->ClassGeneratedBy != nullptr)
-					{
-						UBlueprint* Blueprint = Cast<UBlueprint>(ParentClass->ClassGeneratedBy);
-						while(Blueprint != nullptr)
-						{
-							TargetGraph = FindObject<UEdGraph>(Blueprint, *(FunctionCall->FunctionReference.GetMemberName().ToString()));
-							if((TargetGraph != nullptr) && !TargetGraph->HasAnyFlags(RF_Transient))
-							{
-								FKismetEditorUtilities::BringKismetToFocusAttentionOnObject(TargetGraph);
-								break;
-							}
-
-							ParentClass = Cast<UBlueprintGeneratedClass>(Blueprint->ParentClass);
-							Blueprint = ParentClass != nullptr ? Cast<UBlueprint>(ParentClass->ClassGeneratedBy) : nullptr;
-						}
-					}
-				}
+				FKismetEditorUtilities::BringKismetToFocusAttentionOnObject(FunctionGraph);
 			}
 		}
 	}
diff --git a/Engine/Source/Editor/Kismet/Public/BlueprintEditor.h b/Engine/Source/Editor/Kismet/Public/BlueprintEditor.h
--- a/Engine/Source/Editor/Kismet/Public/BlueprintEditor.h
+++ b/Engine/Source/Editor/Kismet/Public/BlueprintEditor.h
@@ -867,6 +867,12 @@ private:
 	/** Called to check if native code browsing is available */
 	bool IsNativeCodeBrowsingAvailable() const;
 
+	/** Called when the user wants to jump to a node's graph definition */
+	void OnGoToDefinition();
+
+	/** Checks to see if it is possible to jump to the selected node's graph definition. */
+	bool CanGoToDefinition() const;
+
 public://@TODO
 	TSharedPtr<FDocumentTracker> DocumentManager;
 
