(WIP) A more robust solution to binding blueprint actions with specific objects.

diff --git a/Engine/Source/Editor/BlueprintGraph/Classes/BlueprintBoundEventNodeSpawner.h b/Engine/Source/Editor/BlueprintGraph/Classes/BlueprintBoundEventNodeSpawner.h
new file mode 100644
--- /dev/null
+++ b/Engine/Source/Editor/BlueprintGraph/Classes/BlueprintBoundEventNodeSpawner.h
@@ -0,0 +1,60 @@
+// Copyright 1998-2014 Epic Games, Inc. All Rights Reserved.
+
+#pragma once
+
+#include "BlueprintEventNodeSpawner.h"
+#include "BlueprintBoundEventNodeSpawner.generated.h"
+
+class UK2Node_Event;
+
+/**
+ * Takes care of spawning UK2Node_Event nodes. Acts as the "action" portion of
+ * certain FBlueprintActionMenuItems. Will not spawn a new event node if one
+ * associated with the specified function already exits (instead, Invoke() will
+ * return the existing one). Evolved from FEdGraphSchemaAction_K2AddEvent and 
+ * FEdGraphSchemaAction_K2ViewNode.
+ */
+UCLASS(Transient)
+class BLUEPRINTGRAPH_API UBlueprintBoundEventNodeSpawner : public UBlueprintEventNodeSpawner
+{
+	GENERATED_UCLASS_BODY()
+
+public:
+	/**
+	 * Creates a new UBlueprintEventNodeSpawner for delegate bound events. Does
+	 * not come bound, instead it is left up to the ui to bind 
+	 *
+	 * @param  NodeClass		The event node type that you want this to spawn.
+	 * @param  EventDelegate	The delegate that you want to trigger the spawned event.
+	 * @param  Outer			Optional outer for the new spawner (if left null, the transient package will be used).
+	 * @return A newly allocated instance of this class.
+	 */
+	static UBlueprintBoundEventNodeSpawner* Create(TSubclassOf<UK2Node_Event> NodeClass, UMulticastDelegateProperty* EventDelegate, UObject* Outer = nullptr);
+
+	// UBlueprintNodeSpawner interface
+	virtual UEdGraphNode* Invoke(UEdGraph* ParentGraph, FVector2D const Location) const override;
+	virtual FText GetDefaultMenuName() const override;
+	virtual FText GetDefaultMenuCategory() const override;
+	// End UBlueprintNodeSpawner interface
+
+	// UBlueprintEventNodeSpawner interface
+	virtual UK2Node_Event const* FindPreExistingEvent(UBlueprint* Blueprint) const override;
+	// End UBlueprintEventNodeSpawner interface
+
+	// FBlueprintNodeBinder interface
+	virtual bool CanBind(UObject const* BindingCandidate) const override;
+	virtual bool BindToNode(UEdGraphNode* Node, UObject* Binding) const override;
+	// End FBlueprintNodeBinder interface
+
+	/**
+	 * 
+	 *
+	 * @return 
+	 */
+	UMulticastDelegateProperty const* GetEventDelegate() const;
+
+private:
+	/** */
+	UPROPERTY()
+	TWeakObjectPtr<UMulticastDelegateProperty> EventDelegate;
+};
diff --git a/Engine/Source/Editor/BlueprintGraph/Classes/BlueprintBoundNodeSpawner.h b/Engine/Source/Editor/BlueprintGraph/Classes/BlueprintBoundNodeSpawner.h
deleted file mode 100644
diff --git a/Engine/Source/Editor/BlueprintGraph/Classes/BlueprintEventNodeSpawner.h b/Engine/Source/Editor/BlueprintGraph/Classes/BlueprintEventNodeSpawner.h
--- a/Engine/Source/Editor/BlueprintGraph/Classes/BlueprintEventNodeSpawner.h
+++ b/Engine/Source/Editor/BlueprintGraph/Classes/BlueprintEventNodeSpawner.h
@@ -48,6 +48,13 @@ public:
 	virtual FText GetDefaultMenuName() const override;
 	virtual FText GetDefaultSearchKeywords() const override;
 	// End UBlueprintNodeSpawner interface
+
+	/**
+	 * 
+	 * 
+	 * @return 
+	 */
+	bool IsForCustomEvent() const;
 	
 	/**
 	 * Retrieves the function that this assigns to spawned nodes (defines the
@@ -57,6 +64,15 @@ public:
 	 */
 	UFunction const* GetEventFunction() const;
 
+	/**
+	 * 
+	 * 
+	 * @return 
+	 */
+	virtual UK2Node_Event const* FindPreExistingEvent(UBlueprint* Blueprint) const;
+
+	
+
 private:
 	/** The function to configure new nodes with. */
 	UPROPERTY()
diff --git a/Engine/Source/Editor/BlueprintGraph/Classes/BlueprintFunctionNodeSpawner.h b/Engine/Source/Editor/BlueprintGraph/Classes/BlueprintFunctionNodeSpawner.h
--- a/Engine/Source/Editor/BlueprintGraph/Classes/BlueprintFunctionNodeSpawner.h
+++ b/Engine/Source/Editor/BlueprintGraph/Classes/BlueprintFunctionNodeSpawner.h
@@ -48,6 +48,12 @@ public:
 	virtual FText GetDefaultMenuTooltip() const override;
 	virtual FText GetDefaultSearchKeywords() const override;
 	// End UBlueprintNodeSpawner interface
+
+	// FBlueprintNodeBinder interface
+	virtual bool CanBindMultipleObjects() const override;
+	virtual bool CanBind(UObject const* BindingCandidate) const override;
+	virtual bool BindToNode(UEdGraphNode* Node, UObject* Binding) const override;
+	// End FBlueprintNodeBinder interface
 	
 	/**
 	 * Retrieves the function that this assigns to spawned nodes (defines the 
diff --git a/Engine/Source/Editor/BlueprintGraph/Classes/BlueprintNodeBinder.h b/Engine/Source/Editor/BlueprintGraph/Classes/BlueprintNodeBinder.h
new file mode 100644
--- /dev/null
+++ b/Engine/Source/Editor/BlueprintGraph/Classes/BlueprintNodeBinder.h
@@ -0,0 +1,122 @@
+// Copyright 1998-2014 Epic Games, Inc. All Rights Reserved.
+
+#pragma once
+
+class UEdGraphNode;
+
+/**
+ * 
+ */
+class FBlueprintNodeBinder
+{
+public:
+	virtual ~FBlueprintNodeBinder() {}
+
+	/**
+	 * 
+	 * 
+	 * @param  BindingCandidate	
+	 * @return 
+	 */
+	virtual bool CanBind(UObject const* BindingCandidate) const = 0;
+
+	/**
+	 * 
+	 * 
+	 * @return 
+	 */
+	virtual bool CanBindMultipleObjects() const
+	{
+		return false;
+	}
+
+	/**
+	 * 
+	 * 
+	 * @return 
+	 */
+	bool IsBindingSet() const
+	{
+		return (BoundObjects.Num() > 0);
+	}
+
+	/**
+	 * 
+	 * 
+	 * @param  BindingCandidate	
+	 * @return 
+	 */
+	bool AddBinding(UObject const* BindingCandidate)
+	{
+		bool const bAddBinding = CanBind(BindingCandidate) && (!IsBindingSet() || CanBindMultipleObjects());
+		if (bAddBinding)
+		{
+			BoundObjects.Add(BindingCandidate);
+		}
+		return bAddBinding;
+	}
+
+	/**
+	 * 
+	 * 
+	 * @param  Node	
+	 * @return 
+	 */
+	bool Bind(UEdGraphNode* Node) const
+	{
+		bool bFullSuccess = true;
+		for (TWeakObjectPtr<UObject> Instance : BoundObjects)
+		{
+			if (Instance.IsValid())
+			{
+				bFullSuccess &= BindToNode(Node, Instance.Get());
+			}
+		}
+		return bFullSuccess;
+	}
+
+	/**
+	 * 
+	 * 
+	 * @param  Binding	
+	 * @return 
+	 */
+	void RemoveBinding(UObject const* Binding)
+	{
+		BoundObjects.Remove(Binding);
+	}
+
+	/**
+	 * 
+	 * 
+	 * @return 
+	 */
+	void ClearBindings()
+	{
+		BoundObjects.Empty();
+	}
+
+	/**
+	 * 
+	 * 
+	 * @return 
+	 */
+	TSet< TWeakObjectPtr<UObject> > const& GetBindings() const
+	{
+		return BoundObjects;
+	}
+
+protected:
+	/**
+	 * 
+	 * 
+	 * @param  Node	
+	 * @param  Binding	
+	 * @return 
+	 */
+	virtual bool BindToNode(UEdGraphNode* Node, UObject* Binding) const = 0;
+
+protected:
+	/** */
+	TSet< TWeakObjectPtr<UObject> > BoundObjects;
+};
diff --git a/Engine/Source/Editor/BlueprintGraph/Classes/K2Node_CallFunction.h b/Engine/Source/Editor/BlueprintGraph/Classes/K2Node_CallFunction.h
--- a/Engine/Source/Editor/BlueprintGraph/Classes/K2Node_CallFunction.h
+++ b/Engine/Source/Editor/BlueprintGraph/Classes/K2Node_CallFunction.h
@@ -160,6 +160,9 @@ public:
 	/** Checks if the property is marked as "CustomStructureParam" */
 	static bool IsStructureWildcardProperty(const UFunction* InFunction, const FString& PropertyName);
 
+	/** Used to determine the result of AllowMultipleSelfs() (without having a node instance) */
+	static bool CanFunctionSupportMultipleTargets(UFunction const* InFunction);
+
 private: 
 
 	/** Helper function to check if the SelfPin is wired correctly */
diff --git a/Engine/Source/Editor/BlueprintGraph/Classes/K2Node_ComponentBoundEvent.h b/Engine/Source/Editor/BlueprintGraph/Classes/K2Node_ComponentBoundEvent.h
--- a/Engine/Source/Editor/BlueprintGraph/Classes/K2Node_ComponentBoundEvent.h
+++ b/Engine/Source/Editor/BlueprintGraph/Classes/K2Node_ComponentBoundEvent.h
@@ -36,6 +36,6 @@ class UK2Node_ComponentBoundEvent : public UK2Node_Event
 	/** Return the delegate property that this event is bound to */
 	BLUEPRINTGRAPH_API UMulticastDelegateProperty* GetTargetDelegateProperty() const;
 
-	BLUEPRINTGRAPH_API void InitializeComponentBoundEventParams(UObjectProperty* InComponentProperty, const UMulticastDelegateProperty* InDelegateProperty);
+	BLUEPRINTGRAPH_API void InitializeComponentBoundEventParams(UObjectProperty const* InComponentProperty, const UMulticastDelegateProperty* InDelegateProperty);
 };
 
diff --git a/Engine/Source/Editor/BlueprintGraph/Private/BlueprintActionDatabase.cpp b/Engine/Source/Editor/BlueprintGraph/Private/BlueprintActionDatabase.cpp
--- a/Engine/Source/Editor/BlueprintGraph/Private/BlueprintActionDatabase.cpp
+++ b/Engine/Source/Editor/BlueprintGraph/Private/BlueprintActionDatabase.cpp
@@ -9,7 +9,7 @@
 #include "BlueprintDelegateNodeSpawner.h"
 #include "BlueprintEventNodeSpawner.h"
 #include "BlueprintComponentNodeSpawner.h"
-#include "BlueprintBoundNodeSpawner.h"
+#include "BlueprintBoundEventNodeSpawner.h"
 #include "BlueprintVariableNodeSpawner.h"
 #include "Stats.h"					// for RETURN_QUICK_DECLARE_CYCLE_STAT()
 #include "ScopedTimers.h"			// for FScopedDurationTimer
@@ -21,6 +21,9 @@
 
 // used below in FBlueprintNodeSpawnerFactory::MakeMacroNodeSpawner()
 #include "K2Node_MacroInstance.h"
+// used below in FBlueprintNodeSpawnerFactory::MakeComponentBoundEventSpawner()/MakeActorBoundEventSpawner()
+#include "K2Node_ComponentBoundEvent.h"
+#include "K2Node_ActorBoundEvent.h"
 // used below in FBlueprintNodeSpawnerFactory::MakeMessageNodeSpawner()
 #include "K2Node_Message.h"
 // used below in BlueprintActionDatabaseImpl::AddClassPropertyActions()
@@ -93,6 +96,14 @@ namespace FBlueprintNodeSpawnerFactory
 	 * @return 
 	 */
 	static UBlueprintNodeSpawner* MakeComponentBoundEventSpawner(UMulticastDelegateProperty* DelegateProperty);
+
+	/**
+	 * 
+	 * 
+	 * @param  DelegateProperty	
+	 * @return 
+	 */
+	static UBlueprintNodeSpawner* MakeActorBoundEventSpawner(UMulticastDelegateProperty* DelegateProperty);
 };
 
 //------------------------------------------------------------------------------
@@ -185,15 +196,13 @@ static UBlueprintNodeSpawner* FBlueprintNodeSpawnerFactory::MakeAssignDelegateNo
 //------------------------------------------------------------------------------
 static UBlueprintNodeSpawner* FBlueprintNodeSpawnerFactory::MakeComponentBoundEventSpawner(UMulticastDelegateProperty* DelegateProperty)
 {
-	// BoundEventSpawner does the binding, we need a sub-spawner to lay the node
-	// ... we choose UBlueprintDelegateNodeSpawner because we need it to hold on
-	// to the UMulticastDelegateProperty
-	UBlueprintDelegateNodeSpawner* DelegateSpawner = UBlueprintDelegateNodeSpawner::Create(/*NodeClass =*/nullptr, DelegateProperty);
-	// UK2Node_ComponentBoundEvent is not a UK2Node_BaseMCDelegate so we have to
-	// set it this way
-	DelegateSpawner->NodeClass = UK2Node_ComponentBoundEvent::StaticClass(); 
+	return UBlueprintBoundEventNodeSpawner::Create(UK2Node_ComponentBoundEvent::StaticClass(), DelegateProperty);
+}
 
-	return UBlueprintBoundNodeSpawner::Create(DelegateSpawner, UObjectProperty::StaticClass());
+//------------------------------------------------------------------------------
+static UBlueprintNodeSpawner* FBlueprintNodeSpawnerFactory::MakeActorBoundEventSpawner(UMulticastDelegateProperty* DelegateProperty)
+{
+	return UBlueprintBoundEventNodeSpawner::Create(UK2Node_ComponentBoundEvent::StaticClass(), DelegateProperty);
 }
 
 /*******************************************************************************
@@ -379,8 +388,6 @@ static void BlueprintActionDatabaseImpl::AddClassFunctionActions(UClass const* c
 	using namespace FBlueprintNodeSpawnerFactory; // for MakeMessageNodeSpawner()
 	check(Class != nullptr);
 
-	bool const bIsComponent = Class->IsChildOf<UActorComponent>();
-
 	// loop over all the functions in the specified class; exclude-super because 
 	// we can always get the super functions by looking up that class separately 
 	for (TFieldIterator<UFunction> FunctionIt(Class, EFieldIteratorFlags::ExcludeSuper); FunctionIt; ++FunctionIt)
@@ -408,20 +415,6 @@ static void BlueprintActionDatabaseImpl::AddClassFunctionActions(UClass const* c
 			{
 				ActionListOut.Add(MakeMessageNodeSpawner(Function));
 			}
-
-			// bindable function calls (for component properties)
-			if (bIsComponent)
-			{
-				bool const bIsPublic = !Function->HasMetaData(FBlueprintMetadata::MD_Protected) &&
-					!Function->HasMetaData(FBlueprintMetadata::MD_Private);
-				// may have to move the bIsPublic check from here into 
-				// BlueprintActionFilter, once we have blueprintable components?
-				if (bIsPublic)
-				{
-					// bound function for component property
-					ActionListOut.Add(UBlueprintBoundNodeSpawner::Create(FuncSpawner, UObjectProperty::StaticClass()));
-				}
-			}
 		}
 	}
 }
@@ -431,7 +424,8 @@ static void BlueprintActionDatabaseImpl::AddClassPropertyActions(UClass const* c
 {
 	using namespace FBlueprintNodeSpawnerFactory; // for MakeDelegateNodeSpawner()
 
-	bool const bIsComponent = Class->IsChildOf<UActorComponent>();
+	bool const bIsComponent  = Class->IsChildOf<UActorComponent>();
+	bool const bIsActorClass = Class->IsChildOf<AActor>();
 	
 	// loop over all the properties in the specified class; exclude-super because 
 	// we can always get the super properties by looking up that class separately 
@@ -467,13 +461,14 @@ static void BlueprintActionDatabaseImpl::AddClassPropertyActions(UClass const* c
 			UBlueprintNodeSpawner* ClearSpawner = UBlueprintDelegateNodeSpawner::Create(UK2Node_ClearDelegate::StaticClass(), DelegateProperty);
 			ActionListOut.Add(ClearSpawner);
 
-			// @TODO: AddBoundEventActionsForClass()
-			//   UK2Node_ActorBoundEvent
-
 			if (bIsComponent)
 			{
 				ActionListOut.Add(MakeComponentBoundEventSpawner(DelegateProperty));
 			}
+			else if (bIsActorClass)
+			{
+				ActionListOut.Add(MakeActorBoundEventSpawner(DelegateProperty));
+			}
  		}
 		else
 		{
diff --git a/Engine/Source/Editor/BlueprintGraph/Private/BlueprintActionFilter.cpp b/Engine/Source/Editor/BlueprintGraph/Private/BlueprintActionFilter.cpp
--- a/Engine/Source/Editor/BlueprintGraph/Private/BlueprintActionFilter.cpp
+++ b/Engine/Source/Editor/BlueprintGraph/Private/BlueprintActionFilter.cpp
@@ -6,7 +6,7 @@
 #include "BlueprintNodeSpawnerUtils.h"
 #include "BlueprintVariableNodeSpawner.h"
 #include "BlueprintEventNodeSpawner.h"
-#include "BlueprintBoundNodeSpawner.h"
+#include "BlueprintBoundEventNodeSpawner.h"
 #include "EdGraphSchema_K2.h"		// for FBlueprintMetadata
 #include "BlueprintEditorUtils.h"	// for FindBlueprintForGraph()
 #include "ObjectEditorUtils.h"		// for IsFunctionHiddenFromClass()/IsVariableCategoryHiddenFromClass()
@@ -188,15 +188,6 @@ namespace BlueprintActionFilterImpl
 	 */
 	static bool IsBoundToUnselectedObject(FBlueprintActionFilter const& Filter, UBlueprintNodeSpawner const* BlueprintAction);
 
-	/**
-	 * 
-	 * 
-	 * @param  Filter	
-	 * @param  BlueprintAction	
-	 * @return 
-	 */
-	static bool IsBoundFunctionImpure(FBlueprintActionFilter const& Filter, UBlueprintNodeSpawner const* BlueprintAction);
-
 	/**
 	 * 
 	 * 
@@ -525,7 +516,8 @@ static bool BlueprintActionFilterImpl::IsNonTargetMemeber(FBlueprintActionFilter
 
 		// global (and static library) fields can stay (unless explicitly
 		// excluded... save that for a separate test)
-		if ((FieldClass != nullptr) && (bPermitNonTargetGlobals || IsGloballyAccessible(ClassField)))
+		bool const bSkip = bPermitNonTargetGlobals && IsGloballyAccessible(ClassField);
+		if ((FieldClass != nullptr) && !bSkip)
 		{
 			for (UClass const* Class : Filter.TargetClasses)
 			{
@@ -668,40 +660,21 @@ static bool BlueprintActionFilterImpl::IsFilteredNodeType(FBlueprintActionFilter
 static bool BlueprintActionFilterImpl::IsBoundToUnselectedObject(FBlueprintActionFilter const& Filter, UBlueprintNodeSpawner const* BlueprintAction)
 {
 	bool bIsFilteredOut = false;
-	if (UBlueprintBoundNodeSpawner const* BoundSpawner = Cast<UBlueprintBoundNodeSpawner>(BlueprintAction))
+	if (UBlueprintBoundEventNodeSpawner const* BoundSpawner = Cast<UBlueprintBoundEventNodeSpawner>(BlueprintAction))
 	{
 		bIsFilteredOut = true;
-
-		if (UObject const* BoundObject = BoundSpawner->GetBoundObject())
+		for (UObject const* SelectedObj : Filter.Context.SelectedObjects)
 		{
-			for (UObject* SelectedObj : Filter.Context.SelectedObjects)
+			if (BoundSpawner->CanBind(SelectedObj))
 			{
-				if (SelectedObj == BoundObject)
-				{
-					bIsFilteredOut = false;
-					break;
-				}
+				bIsFilteredOut = false;
+				break;
 			}
 		}
 	}
 	return bIsFilteredOut;
 }
 
-//------------------------------------------------------------------------------
-static bool BlueprintActionFilterImpl::IsBoundFunctionImpure(FBlueprintActionFilter const& Filter, UBlueprintNodeSpawner const* BlueprintAction)
-{
-	bool bIsFilteredOut = false;
-	if (UBlueprintBoundNodeSpawner const* BoundSpawner = Cast<UBlueprintBoundNodeSpawner>(BlueprintAction))
-	{
-		if (UFunction const* BoundFunction = FBlueprintNodeSpawnerUtils::GetAssociatedFunction(BoundSpawner))
-		{
-			bIsFilteredOut = !BoundFunction->HasAnyFunctionFlags(FUNC_BlueprintPure) && 
-				!BoundFunction->HasAnyFunctionFlags(FUNC_Const);
-		}
-	}
-	return bIsFilteredOut;
-}
-
 //------------------------------------------------------------------------------
 static bool BlueprintActionFilterImpl::IsOutOfScopeLocalVariable(FBlueprintActionFilter const& Filter, UBlueprintNodeSpawner const* BlueprintAction)
 {
@@ -845,14 +818,7 @@ static bool BlueprintActionFilterImpl::IsPinCompatibleWithTargetSelf(UEdGraphPin
 				{
 					if (UFunction const* Function = Cast<UFunction>(MemberField))
 					{
-						// @TODO: duplicated from UK2Node_CallFunction::AllowMultipleSelfs()... 
-						//        move into a shared resource (don't want to call 
-						//        GetTemplateNode() to get at AllowMultipleSelfs() 
-						//        for perf reasons)
-						bool const bHasReturnParam = (Function->GetReturnProperty() != nullptr);
-						bool const bIsImpure = (Function->HasAnyFunctionFlags(FUNC_BlueprintPure) == false);
-						bool const bIsLatent = (Function->HasMetaData(FBlueprintMetadata::MD_Latent) != false);
-						bIsCompatible = !bHasReturnParam && bIsImpure && !bIsLatent;
+						bIsCompatible = UK2Node_CallFunction::CanFunctionSupportMultipleTargets(Function);
 					}
 					else
 					{
@@ -1047,12 +1013,7 @@ FBlueprintActionFilter::FBlueprintActionFilter(uint32 Flags/*= 0x00*/)
 	}
 
 	AddIsFilteredTest(FIsFilteredDelegate::CreateStatic(IsFilteredNodeType, (Flags & BPFILTER_RejectPermittedNodeSubClasses)   != 0));
-	AddIsFilteredTest(FIsFilteredDelegate::CreateStatic(IsNonTargetMemeber, (Flags & BPFILTER_RejectPersistentNonTargetFields) != 0));
-	
-	// don't know why we only allow impure ("imperitive") and const functions,
-	// but this mimics legacy functionality in GetFunctionCallsOnSelectedComponents()
-	AddIsFilteredTest(FIsFilteredDelegate::CreateStatic(IsBoundFunctionImpure)); 
-
+	AddIsFilteredTest(FIsFilteredDelegate::CreateStatic(IsNonTargetMemeber, (Flags & BPFILTER_RejectPersistentNonTargetFields) == 0));
 	AddIsFilteredTest(FIsFilteredDelegate::CreateStatic(IsBoundToUnselectedObject));
 	AddIsFilteredTest(FIsFilteredDelegate::CreateStatic(IsOutOfScopeLocalVariable));
 
diff --git a/Engine/Source/Editor/BlueprintGraph/Private/BlueprintBoundEventNodeSpawner.cpp b/Engine/Source/Editor/BlueprintGraph/Private/BlueprintBoundEventNodeSpawner.cpp
new file mode 100644
--- /dev/null
+++ b/Engine/Source/Editor/BlueprintGraph/Private/BlueprintBoundEventNodeSpawner.cpp
@@ -0,0 +1,149 @@
+// Copyright 1998-2014 Epic Games, Inc. All Rights Reserved.
+
+#include "BlueprintGraphPrivatePCH.h"
+#include "BlueprintBoundEventNodeSpawner.h"
+#include "ObjectEditorUtils.h"		// for GetCategory()
+#include "KismetEditorUtilities.h"	// for FindBoundEventForComponent()/FindBoundEventForActor()
+#include "EditorCategoryUtils.h"	// for GetCommonCategory()
+
+#define LOCTEXT_NAMESPACE "BlueprintBoundEventNodeSpawner"
+
+/*******************************************************************************
+ * UBlueprintBoundEventNodeSpawner
+ ******************************************************************************/
+
+//------------------------------------------------------------------------------
+UBlueprintBoundEventNodeSpawner* UBlueprintBoundEventNodeSpawner::Create(TSubclassOf<UK2Node_Event> NodeClass, UMulticastDelegateProperty* EventDelegate, UObject* Outer/* = nullptr*/)
+{
+	if (Outer == nullptr)
+	{
+		Outer = GetTransientPackage();
+	}
+
+	UBlueprintBoundEventNodeSpawner* NodeSpawner = NewObject<UBlueprintBoundEventNodeSpawner>(Outer);
+	NodeSpawner->NodeClass     = NodeClass;
+	NodeSpawner->EventDelegate = EventDelegate;
+
+	return NodeSpawner;
+}
+
+//------------------------------------------------------------------------------
+UBlueprintBoundEventNodeSpawner::UBlueprintBoundEventNodeSpawner(class FPostConstructInitializeProperties const& PCIP)
+	: Super(PCIP)
+	, EventDelegate(nullptr)
+{
+}
+
+//------------------------------------------------------------------------------
+UEdGraphNode* UBlueprintBoundEventNodeSpawner::Invoke(UEdGraph* ParentGraph, FVector2D const Location) const
+{
+	UK2Node_Event* EventNode = nullptr;
+	if (IsBindingSet())
+	{
+		EventNode = CastChecked<UK2Node_Event>(Super::Invoke(ParentGraph, Location));
+		Bind(EventNode);
+	}
+	return EventNode;
+}
+
+//------------------------------------------------------------------------------
+FText UBlueprintBoundEventNodeSpawner::GetDefaultMenuName() const
+{
+	FText const DelegateName = FText::FromName(EventDelegate->GetFName());
+	return FText::Format(LOCTEXT("ComponentEventName", "Add {0}"), DelegateName);
+}
+
+//------------------------------------------------------------------------------
+FText UBlueprintBoundEventNodeSpawner::GetDefaultMenuCategory() const
+{
+	FText DelegateCategory = FEditorCategoryUtils::GetCommonCategory(FCommonEditorCategory::Delegates);
+	if (UMulticastDelegateProperty const* Delegate = GetEventDelegate())
+	{
+		DelegateCategory = FText::FromString(FObjectEditorUtils::GetCategory(Delegate));
+	}
+	return DelegateCategory;
+}
+
+//------------------------------------------------------------------------------
+UK2Node_Event const* UBlueprintBoundEventNodeSpawner::FindPreExistingEvent(UBlueprint* Blueprint) const
+{
+	UK2Node_Event const* PreExistingEvent = nullptr;
+
+	UObject const* BoundObject = nullptr;
+	if (IsBindingSet())
+	{
+		BoundObject = BoundObjects.CreateConstIterator()->Get();
+	}
+
+	if (BoundObject != nullptr)
+	{
+		if (NodeClass->IsChildOf<UK2Node_ComponentBoundEvent>())
+		{
+			PreExistingEvent = FKismetEditorUtilities::FindBoundEventForComponent(Blueprint, EventDelegate->GetFName(), BoundObject->GetFName());
+		}
+		else if (NodeClass->IsChildOf<UK2Node_ActorBoundEvent>())
+		{
+			PreExistingEvent = FKismetEditorUtilities::FindBoundEventForActor(CastChecked<AActor>(BoundObject), EventDelegate->GetFName());
+		}
+	}
+	return PreExistingEvent;
+}
+
+//------------------------------------------------------------------------------
+bool UBlueprintBoundEventNodeSpawner::CanBind(UObject const* BindingCandidate) const
+{
+	bool bCanBind = false;
+
+	UClass* DelegateOwner = nullptr;
+	if (UMulticastDelegateProperty const* Delegate = GetEventDelegate())
+	{
+		DelegateOwner = Delegate->GetOwnerClass();
+	}
+
+	if (NodeClass->IsChildOf<UK2Node_ComponentBoundEvent>())
+	{
+		UObjectProperty const* BindingProperty = Cast<UObjectProperty>(BindingCandidate);
+
+		bCanBind = (BindingProperty != nullptr) && (DelegateOwner != nullptr) &&
+			BindingProperty->PropertyClass->IsChildOf(DelegateOwner);
+	}
+	else if (NodeClass->IsChildOf<UK2Node_ActorBoundEvent>())
+	{
+		bCanBind = BindingCandidate->IsA<AActor>() && (DelegateOwner != nullptr) &&
+			BindingCandidate->GetClass()->IsChildOf(DelegateOwner);
+	}
+	return bCanBind;
+}
+
+//------------------------------------------------------------------------------
+bool UBlueprintBoundEventNodeSpawner::BindToNode(UEdGraphNode* Node, UObject* Binding) const
+{
+	bool bWasBound = false;
+	if (UK2Node_ComponentBoundEvent* EventNode = Cast<UK2Node_ComponentBoundEvent>(Node))
+	{
+		if (UObjectProperty const* BoundProperty = Cast<UObjectProperty>(Binding))
+		{
+			EventNode->InitializeComponentBoundEventParams(BoundProperty, EventDelegate.Get());
+			bWasBound = true;
+			Node->ReconstructNode();
+		}
+	}
+	else if (UK2Node_ActorBoundEvent* EventNode = CastChecked<UK2Node_ActorBoundEvent>(Node))
+	{
+		if (AActor* BoundActor = Cast<AActor>(Binding))
+		{
+			EventNode->InitializeActorBoundEventParams(BoundActor, EventDelegate.Get());
+			bWasBound = true;
+			Node->ReconstructNode();
+		}
+	}
+	return bWasBound;
+}
+
+//------------------------------------------------------------------------------
+UMulticastDelegateProperty const* UBlueprintBoundEventNodeSpawner::GetEventDelegate() const
+{
+	return EventDelegate.Get();
+}
+
+#undef LOCTEXT_NAMESPACE
diff --git a/Engine/Source/Editor/BlueprintGraph/Private/BlueprintBoundNodeSpawner.cpp b/Engine/Source/Editor/BlueprintGraph/Private/BlueprintBoundNodeSpawner.cpp
deleted file mode 100644
diff --git a/Engine/Source/Editor/BlueprintGraph/Private/BlueprintEventNodeSpawner.cpp b/Engine/Source/Editor/BlueprintGraph/Private/BlueprintEventNodeSpawner.cpp
--- a/Engine/Source/Editor/BlueprintGraph/Private/BlueprintEventNodeSpawner.cpp
+++ b/Engine/Source/Editor/BlueprintGraph/Private/BlueprintEventNodeSpawner.cpp
@@ -93,27 +93,19 @@ UEdGraphNode* UBlueprintEventNodeSpawner::Invoke(UEdGraph* ParentGraph, FVector2
 {
 	check(ParentGraph != nullptr);
 	UBlueprint* Blueprint = FBlueprintEditorUtils::FindBlueprintForGraphChecked(ParentGraph);
-
-	UK2Node_Event* EventNode = nullptr;
-	
-	FName EventName = CustomEventName;
-	if (EventFunc != nullptr)
-	{
-		check(EventFunc != nullptr);
-		EventName = EventFunc->GetFName();
-	}
-	
-	UClass* ClassOwner = Blueprint->GeneratedClass;
 	// look to see if a node for this event already exists (only one node is
 	// allowed per event, per blueprint)
-	if (EventFunc != nullptr)
-	{
-		ClassOwner = CastChecked<UClass>(EventFunc->GetOuter());
-		EventNode  = FBlueprintEditorUtils::FindOverrideForFunction(Blueprint, ClassOwner, EventName);
-	}
-	else
+	UK2Node_Event const* PreExistingNode = FindPreExistingEvent(Blueprint);
+	// @TODO: casting away the const is bad form!
+	UK2Node_Event* EventNode = const_cast<UK2Node_Event*>(PreExistingNode);
+
+	bool const bIsCustomEvent = IsForCustomEvent();
+	check(bIsCustomEvent || (EventFunc != nullptr));
+
+	FName EventName = CustomEventName;
+	if (!bIsCustomEvent)
 	{
-		EventNode  = UBlueprintEventNodeSpawnerImpl::FindCustomEventNode(Blueprint, EventName);
+		EventName  = EventFunc->GetFName();	
 	}
 
 	// if there is no existing node, then we can happily spawn one into the graph
@@ -182,4 +174,31 @@ UFunction const* UBlueprintEventNodeSpawner::GetEventFunction() const
 	return EventFunc;
 }
 
+//------------------------------------------------------------------------------
+UK2Node_Event const* UBlueprintEventNodeSpawner::FindPreExistingEvent(UBlueprint* Blueprint) const
+{
+	UK2Node_Event* PreExistingNode = nullptr;
+
+	check(Blueprint != nullptr);
+	if (IsForCustomEvent())
+	{
+		PreExistingNode = UBlueprintEventNodeSpawnerImpl::FindCustomEventNode(Blueprint, CustomEventName);
+	}
+	else
+	{
+		check(EventFunc != nullptr);
+		UClass* ClassOwner = EventFunc->GetOwnerClass();
+
+		PreExistingNode = FBlueprintEditorUtils::FindOverrideForFunction(Blueprint, ClassOwner, EventFunc->GetFName());
+	}
+
+	return PreExistingNode;
+}
+
+//------------------------------------------------------------------------------
+bool UBlueprintEventNodeSpawner::IsForCustomEvent() const
+{
+	return (EventFunc == nullptr);
+}
+
 #undef LOCTEXT_NAMESPACE
diff --git a/Engine/Source/Editor/BlueprintGraph/Private/BlueprintFunctionNodeSpawner.cpp b/Engine/Source/Editor/BlueprintGraph/Private/BlueprintFunctionNodeSpawner.cpp
--- a/Engine/Source/Editor/BlueprintGraph/Private/BlueprintFunctionNodeSpawner.cpp
+++ b/Engine/Source/Editor/BlueprintGraph/Private/BlueprintFunctionNodeSpawner.cpp
@@ -3,9 +3,126 @@
 #include "BlueprintGraphPrivatePCH.h"
 #include "BlueprintFunctionNodeSpawner.h"
 #include "EdGraphSchema_K2.h"	// for FBlueprintMetadata
+#include "Gameframework/Actor.h"
+#include "BlueprintVariableNodeSpawner.h"
 
 #define LOCTEXT_NAMESPACE "BlueprintFunctionNodeSpawner"
 
+/*******************************************************************************
+ * Static UBlueprintFunctionNodeSpawner Helpers
+ ******************************************************************************/
+
+//------------------------------------------------------------------------------
+namespace BlueprintFunctionNodeSpawnerImpl
+{
+	FVector2D BindingOffset = FVector2D::ZeroVector;
+
+	/**
+	 * 
+	 * 
+	 * @param  NewNode	
+	 * @param  BoundObject	
+	 * @return 
+	 */
+	static bool BindFunctionNode(UK2Node_CallFunction* NewNode, UObject* BoundObject);
+
+	/**
+	 * 
+	 * 
+	 * @param  NewNode	
+	 * @param  BindingSpawner	
+	 * @return 
+	 */
+	template <class NodeType>
+	static bool BindFunctionNode(UK2Node_CallFunction* NewNode, UBlueprintNodeSpawner* BindingSpawner);
+
+	/**
+	 * 
+	 * 
+	 * @param  InputNode
+	 * @return 
+	 */
+	static FVector2D CalculateBindingPosition(UEdGraphNode* const InputNode);
+}
+
+//------------------------------------------------------------------------------
+static bool BlueprintFunctionNodeSpawnerImpl::BindFunctionNode(UK2Node_CallFunction* NewNode, UObject* BoundObject)
+{
+	bool bSuccessfulBinding = false;
+
+	bool const bIsTemplateNode = (NewNode->GetOutermost() == GetTransientPackage());
+	if (!bIsTemplateNode)
+	{
+		if (UProperty const* BoundProperty = Cast<UProperty>(BoundObject))
+		{
+			UBlueprintNodeSpawner* TempNodeSpawner = UBlueprintVariableNodeSpawner::Create(UK2Node_VariableGet::StaticClass(), BoundProperty);
+			bSuccessfulBinding = BindFunctionNode<UK2Node_VariableGet>(NewNode, TempNodeSpawner);
+		}
+		else if (AActor* BoundActor = Cast<AActor>(BoundObject))
+		{
+			auto PostSpawnSetupLambda = [](UEdGraphNode* NewNode, bool /*bIsTemplateNode*/, AActor* ActorInst)
+			{
+				UK2Node_Literal* ActorRefNode = CastChecked<UK2Node_Literal>(NewNode);
+				ActorRefNode->SetObjectRef(ActorInst);
+			};
+			UBlueprintNodeSpawner::FCustomizeNodeDelegate PostSpawnDelegate = UBlueprintNodeSpawner::FCustomizeNodeDelegate::CreateStatic(PostSpawnSetupLambda, BoundActor);
+			
+			UBlueprintNodeSpawner* TempNodeSpawner = UBlueprintNodeSpawner::Create<UK2Node_Literal>(/*Outer =*/GetTransientPackage(), PostSpawnDelegate);
+			bSuccessfulBinding = BindFunctionNode<UK2Node_Literal>(NewNode, TempNodeSpawner);
+		}		
+	}
+	return bSuccessfulBinding;
+}
+
+//------------------------------------------------------------------------------
+template <class NodeType>
+static bool BlueprintFunctionNodeSpawnerImpl::BindFunctionNode(UK2Node_CallFunction* NewNode, UBlueprintNodeSpawner* BindingSpawner)
+{
+	bool bSuccessfulBinding = false;
+
+	FVector2D BindingPos = CalculateBindingPosition(NewNode);
+	UEdGraph* ParentGraph = NewNode->GetGraph();
+	NodeType* BindingNode = CastChecked<NodeType>(BindingSpawner->Invoke(ParentGraph, BindingPos));
+
+	BindingOffset.Y += UEdGraphSchema_K2::EstimateNodeHeight(BindingNode);
+
+	// @TODO: Move this down into Invoke()
+	ParentGraph->Modify();
+	ParentGraph->AddNode(BindingNode, /*bFromUI =*/false, /*bSelectNewNode =*/false);
+
+	UEdGraphPin* LiteralOutput = BindingNode->GetValuePin();
+	UEdGraphPin* CallSelfInput = NewNode->FindPin(GetDefault<UEdGraphSchema_K2>()->PN_Self);
+	// connect the new "get-var" node with the spawned function node
+	if ((LiteralOutput != nullptr) && (CallSelfInput != nullptr))
+	{
+		LiteralOutput->MakeLinkTo(CallSelfInput);
+		bSuccessfulBinding = true;
+	}
+
+	return bSuccessfulBinding;
+}
+
+//------------------------------------------------------------------------------
+static FVector2D BlueprintFunctionNodeSpawnerImpl::CalculateBindingPosition(UEdGraphNode* const InputNode)
+{
+	float const EstimatedVarNodeWidth = 224.0f;
+	FVector2D AttachingNodePos;
+	AttachingNodePos.X = InputNode->NodePosX - EstimatedVarNodeWidth;
+	AttachingNodePos.Y = InputNode->NodePosY;
+
+	float const EstimatedVarNodeHeight = 48.0f;
+	float const EstimatedFuncNodeHeight = UEdGraphSchema_K2::EstimateNodeHeight(InputNode);
+	float const FuncNodeMidYCoordinate = InputNode->NodePosY + (EstimatedFuncNodeHeight / 2.0f);
+	AttachingNodePos.Y = FuncNodeMidYCoordinate - (EstimatedVarNodeWidth / 2.0f);
+
+	AttachingNodePos += BindingOffset;
+	return AttachingNodePos;
+}
+
+/*******************************************************************************
+ * UBlueprintFunctionNodeSpawner
+ ******************************************************************************/
+
 //------------------------------------------------------------------------------
 // Evolved from FK2ActionMenuBuilder::AddSpawnInfoForFunction()
 UBlueprintFunctionNodeSpawner* UBlueprintFunctionNodeSpawner::Create(UFunction const* const Function, UObject* Outer/* = nullptr*/)
@@ -91,7 +208,13 @@ UEdGraphNode* UBlueprintFunctionNodeSpawner::Invoke(UEdGraph* ParentGraph, FVect
 	};
 
 	FCustomizeNodeDelegate PostSpawnSetupDelegate = FCustomizeNodeDelegate::CreateStatic(PostSpawnSetupLambda, Function, CustomizeNodeDelegate);
-	return Super::Invoke(ParentGraph, Location, PostSpawnSetupDelegate);
+	UEdGraphNode* SpawnedNode = Super::Invoke(ParentGraph, Location, PostSpawnSetupDelegate);
+
+	// if this spawner was set up to spawn 
+	BlueprintFunctionNodeSpawnerImpl::BindingOffset = FVector2D::ZeroVector;
+	Bind(SpawnedNode);
+
+	return SpawnedNode;
 }
 
 //------------------------------------------------------------------------------
@@ -122,6 +245,49 @@ FText UBlueprintFunctionNodeSpawner::GetDefaultSearchKeywords() const
 	return FText::FromString(UK2Node_CallFunction::GetKeywordsForFunction(Function));
 }
 
+//------------------------------------------------------------------------------
+bool UBlueprintFunctionNodeSpawner::CanBindMultipleObjects() const
+{
+	check(Function != nullptr);
+	return UK2Node_CallFunction::CanFunctionSupportMultipleTargets(Function);
+}
+
+//------------------------------------------------------------------------------
+bool UBlueprintFunctionNodeSpawner::CanBind(UObject const* BindingCandidate) const
+{
+	bool bCanBind = false;
+	if (Function != nullptr)
+	{
+		// @TODO: don't know exactly why we can only bind non-pure or cost 
+		//        functions... this is mirrored after FK2ActionMenuBuilder::GetFunctionCallsOnSelectedActors()
+		//        and FK2ActionMenuBuilder::GetFunctionCallsOnSelectedComponents(),
+		//        where we make the same stipulation
+		bool const bIsImperative = !Function->HasAnyFunctionFlags(FUNC_BlueprintPure);
+		bool const bIsConstFunc  =  Function->HasAnyFunctionFlags(FUNC_Const);
+
+		if (bIsImperative || bIsConstFunc)
+		{
+			UClass* BindingClass = BindingCandidate->GetClass();
+			if (UObjectProperty const* ObjectProperty = Cast<UObjectProperty>(BindingCandidate))
+			{
+				BindingClass = ObjectProperty->PropertyClass;
+			}
+
+			if (UClass const* FuncOwner = Function->GetOwnerClass())
+			{
+				bCanBind = BindingClass->IsChildOf(FuncOwner);
+			}
+		}
+	}
+	return bCanBind;
+}
+
+//------------------------------------------------------------------------------
+bool UBlueprintFunctionNodeSpawner::BindToNode(UEdGraphNode* Node, UObject* Binding) const
+{
+	return BlueprintFunctionNodeSpawnerImpl::BindFunctionNode(CastChecked<UK2Node_CallFunction>(Node), Binding);
+}
+
 //------------------------------------------------------------------------------
 UFunction const* UBlueprintFunctionNodeSpawner::GetFunction() const
 {
diff --git a/Engine/Source/Editor/BlueprintGraph/Private/BlueprintNodeSpawner.cpp b/Engine/Source/Editor/BlueprintGraph/Private/BlueprintNodeSpawner.cpp
--- a/Engine/Source/Editor/BlueprintGraph/Private/BlueprintNodeSpawner.cpp
+++ b/Engine/Source/Editor/BlueprintGraph/Private/BlueprintNodeSpawner.cpp
@@ -120,6 +120,18 @@ FName UBlueprintNodeSpawner::GetDefaultMenuIcon(FLinearColor& ColorOut) const
 	return FName();
 }
 
+//------------------------------------------------------------------------------
+bool UBlueprintNodeSpawner::CanBind(UObject const* BindingCandidate) const 
+{
+	return false;
+}
+
+//------------------------------------------------------------------------------
+bool UBlueprintNodeSpawner::BindToNode(UEdGraphNode* Node, UObject* Binding) const
+{
+	return false;
+}
+
 //------------------------------------------------------------------------------
 UEdGraphNode* UBlueprintNodeSpawner::GetTemplateNode(UEdGraph* Outer) const
 {	
@@ -150,6 +162,11 @@ UEdGraphNode* UBlueprintNodeSpawner::Invoke(UEdGraph* ParentGraph, FVector2D con
 			NewNode->SetFlags(RF_Transactional);
 			NewNode->AllocateDefaultPins();
 			NewNode->PostPlacedNewNode();
+
+			ParentGraph->Modify();
+			// selections of this node happen up, out at the 
+			// FBlueprintActionMenuItem level
+			ParentGraph->AddNode(NewNode, /*bFromUI =*/true, /*bSelectNewNode =*/false);
 		}
 	}
 
diff --git a/Engine/Source/Editor/BlueprintGraph/Private/BlueprintNodeSpawnerUtils.cpp b/Engine/Source/Editor/BlueprintGraph/Private/BlueprintNodeSpawnerUtils.cpp
--- a/Engine/Source/Editor/BlueprintGraph/Private/BlueprintNodeSpawnerUtils.cpp
+++ b/Engine/Source/Editor/BlueprintGraph/Private/BlueprintNodeSpawnerUtils.cpp
@@ -7,9 +7,9 @@
 #include "BlueprintVariableNodeSpawner.h"
 #include "BlueprintEventNodeSpawner.h"
 #include "BlueprintComponentNodeSpawner.h"
-#include "BlueprintBoundNodeSpawner.h"
+#include "BlueprintBoundEventNodeSpawner.h"
 #include "BlueprintDelegateNodeSpawner.h"
-
+ 
 //------------------------------------------------------------------------------
 UClass* FBlueprintNodeSpawnerUtils::GetAssociatedClass(UBlueprintNodeSpawner const* BlueprintAction)
 {
@@ -23,10 +23,6 @@ UClass* FBlueprintNodeSpawnerUtils::GetAssociatedClass(UBlueprintNodeSpawner con
 	{
 		ActionClass = ComponentSpawner->GetComponentClass();
 	}
-	else if (UBlueprintBoundNodeSpawner const* BoundSpawner = Cast<UBlueprintBoundNodeSpawner>(BlueprintAction))
-	{
-		ActionClass = GetAssociatedClass(BoundSpawner->GetSubSpawner());
-	}
 	else if (UBlueprintVariableNodeSpawner const* VarSpawner = Cast<UBlueprintVariableNodeSpawner>(BlueprintAction))
 	{
 		UObject const* VarOuter = VarSpawner->GetVarOuter();
@@ -54,10 +50,6 @@ UField const* FBlueprintNodeSpawnerUtils::GetAssociatedField(UBlueprintNodeSpawn
 	{
 		ClassField = Property;
 	}
-	else if (UBlueprintBoundNodeSpawner const* BoundSpawner = Cast<UBlueprintBoundNodeSpawner>(BlueprintAction))
-	{
-		ClassField = GetAssociatedField(BoundSpawner->GetSubSpawner());
-	}
 	return ClassField;
 }
 
@@ -74,10 +66,6 @@ UFunction const* FBlueprintNodeSpawnerUtils::GetAssociatedFunction(UBlueprintNod
 	{
 		Function = EventSpawner->GetEventFunction();
 	}
-	else if (UBlueprintBoundNodeSpawner const* BoundSpawner = Cast<UBlueprintBoundNodeSpawner>(BlueprintAction))
-	{
-		Function = GetAssociatedFunction(BoundSpawner->GetSubSpawner());
-	}
 
 	return Function;
 }
@@ -95,5 +83,9 @@ UProperty const* FBlueprintNodeSpawnerUtils::GetAssociatedProperty(UBlueprintNod
 	{
 		Property = VarSpawner->GetVarProperty();
 	}
+	else if (UBlueprintBoundEventNodeSpawner const* BoundSpawner = Cast<UBlueprintBoundEventNodeSpawner>(BlueprintAction))
+	{
+		Property = BoundSpawner->GetEventDelegate();
+	}
 	return Property;
 }
diff --git a/Engine/Source/Editor/BlueprintGraph/Private/K2Node_CallFunction.cpp b/Engine/Source/Editor/BlueprintGraph/Private/K2Node_CallFunction.cpp
--- a/Engine/Source/Editor/BlueprintGraph/Private/K2Node_CallFunction.cpp
+++ b/Engine/Source/Editor/BlueprintGraph/Private/K2Node_CallFunction.cpp
@@ -728,13 +728,21 @@ bool UK2Node_CallFunction::AllowMultipleSelfs(bool bInputAsArray) const
 {
 	if (UFunction* Function = GetTargetFunction())
 	{
-		const bool bHasReturnParams = (Function->GetReturnProperty() != NULL);
-		return !bHasReturnParams && !IsNodePure() && !IsLatentFunction();
+		return CanFunctionSupportMultipleTargets(Function);
 	}
 
 	return Super::AllowMultipleSelfs(bInputAsArray);
 }
 
+bool UK2Node_CallFunction::CanFunctionSupportMultipleTargets(UFunction const* Function)
+{
+	bool const bIsImpure = !Function->HasAnyFunctionFlags(FUNC_BlueprintPure);
+	bool const bIsLatent = Function->HasMetaData(FBlueprintMetadata::MD_Latent);
+	bool const bHasReturnParam = (Function->GetReturnProperty() != nullptr);
+
+	return !bHasReturnParam && bIsImpure && !bIsLatent;
+}
+
 FLinearColor UK2Node_CallFunction::GetNodeTitleColor() const
 {
 	if (IsNodePure())
diff --git a/Engine/Source/Editor/BlueprintGraph/Private/K2Node_ComponentBoundEvent.cpp b/Engine/Source/Editor/BlueprintGraph/Private/K2Node_ComponentBoundEvent.cpp
--- a/Engine/Source/Editor/BlueprintGraph/Private/K2Node_ComponentBoundEvent.cpp
+++ b/Engine/Source/Editor/BlueprintGraph/Private/K2Node_ComponentBoundEvent.cpp
@@ -18,7 +18,7 @@ FText UK2Node_ComponentBoundEvent::GetNodeTitle(ENodeTitleType::Type TitleType)
 	return FText::Format(LOCTEXT("ComponentBoundEvent_Title", "{DelegatePropertyName} ({ComponentPropertyName})"), Args);
 }
 
-void UK2Node_ComponentBoundEvent::InitializeComponentBoundEventParams(UObjectProperty* InComponentProperty, const UMulticastDelegateProperty* InDelegateProperty)
+void UK2Node_ComponentBoundEvent::InitializeComponentBoundEventParams(UObjectProperty const* InComponentProperty, const UMulticastDelegateProperty* InDelegateProperty)
 {
 	if( InComponentProperty && InDelegateProperty )
 	{
diff --git a/Engine/Source/Editor/BlueprintGraph/Public/BlueprintActionFilter.h b/Engine/Source/Editor/BlueprintGraph/Public/BlueprintActionFilter.h
--- a/Engine/Source/Editor/BlueprintGraph/Public/BlueprintActionFilter.h
+++ b/Engine/Source/Editor/BlueprintGraph/Public/BlueprintActionFilter.h
@@ -62,7 +62,7 @@ public:
 		 * Rejects actions associated with global/static ("persistent") fields.
 		 * However, static members of TargetClasses are kept, and not outright 
 		 * rejected by this (if TargetClasses is left empty, then all 
-		 * "persistent" fields are rejected.
+		 * "persistent" fields are rejected).
 		 */
 		BPFILTER_RejectPersistentNonTargetFields= 0x00000002,
 
diff --git a/Engine/Source/Editor/BlueprintGraph/Public/BlueprintNodeSpawner.h b/Engine/Source/Editor/BlueprintGraph/Public/BlueprintNodeSpawner.h
--- a/Engine/Source/Editor/BlueprintGraph/Public/BlueprintNodeSpawner.h
+++ b/Engine/Source/Editor/BlueprintGraph/Public/BlueprintNodeSpawner.h
@@ -2,6 +2,7 @@
 
 #pragma once
 
+#include "BlueprintNodeBinder.h"
 #include "BlueprintNodeSpawner.generated.h"
 
 // Forward declarations
@@ -23,7 +24,7 @@ class UEdGraphNode;
  * these "actions".
  */
 UCLASS(Transient)
-class BLUEPRINTGRAPH_API UBlueprintNodeSpawner : public UObject
+class BLUEPRINTGRAPH_API UBlueprintNodeSpawner : public UObject, public FBlueprintNodeBinder
 {
 	GENERATED_UCLASS_BODY()
 	DECLARE_DELEGATE_TwoParams(FCustomizeNodeDelegate, UEdGraphNode*, bool);
@@ -157,6 +158,11 @@ public:
 	 */
 	UEdGraphNode* GetTemplateNode(UEdGraph* TargetGraph = nullptr) const;
 
+	// FBlueprintNodeBinder interface
+	virtual bool CanBind(UObject const* BindingCandidate) const override;
+	virtual bool BindToNode(UEdGraphNode* Node, UObject* Binding) const override;
+	// End FBlueprintNodeBinder interface
+
 protected:
 	/**
 	 * Protected Invoke() that let's sub-classes specify their own post-spawn
diff --git a/Engine/Source/Editor/Kismet/Private/BlueprintActionMenuBuilder.cpp b/Engine/Source/Editor/Kismet/Private/BlueprintActionMenuBuilder.cpp
--- a/Engine/Source/Editor/Kismet/Private/BlueprintActionMenuBuilder.cpp
+++ b/Engine/Source/Editor/Kismet/Private/BlueprintActionMenuBuilder.cpp
@@ -4,13 +4,13 @@
 #include "BlueprintActionMenuBuilder.h"
 #include "BlueprintActionMenuItem.h"
 #include "BlueprintDragDropMenuItem.h"
+#include "BlueprintBoundMenuItem.h"
 #include "BlueprintActionFilter.h"
 #include "BlueprintActionDatabase.h"
 #include "BlueprintNodeSpawner.h"
 #include "BlueprintFunctionNodeSpawner.h"
 #include "BlueprintDelegateNodeSpawner.h"
 #include "BlueprintVariableNodeSpawner.h"
-#include "BlueprintBoundNodeSpawner.h"
 #include "K2ActionMenuBuilder.h"		// for FBlueprintPaletteListBuilder/FBlueprintGraphActionListBuilder
 #include "EdGraphSchema_K2.h"			// for StaticClass(), bUseLegacyActionMenus, etc.
 #include "BlueprintEditor.h"			// for GetMyBlueprintWidget(), and GetIsContextSensitive()
@@ -60,6 +60,14 @@ public:
 	 * @return A newly allocated FBlueprintActionMenuItem (which wraps the supplied action).
 	 */
 	TSharedPtr<FEdGraphSchemaAction> MakeDragDropMenuItem(UBlueprintNodeSpawner const* SampleAction);
+
+	/**
+	 * 
+	 * 
+	 * @param  BoundAction	
+	 * @return 
+	 */
+	TSharedPtr<FBlueprintBoundMenuItem> MakeBoundMenuItem(TWeakPtr<FBlueprintEditor> EditorContext, UBlueprintNodeSpawner* BoundAction);
 	
 private:
 	/**
@@ -167,6 +175,20 @@ TSharedPtr<FEdGraphSchemaAction> FBlueprintActionMenuItemFactory::MakeDragDropMe
 	return MakeShareable(NewMenuItem);
 }
 
+//------------------------------------------------------------------------------
+TSharedPtr<FBlueprintBoundMenuItem> FBlueprintActionMenuItemFactory::MakeBoundMenuItem(TWeakPtr<FBlueprintEditor> EditorContext, UBlueprintNodeSpawner* BoundAction)
+{
+	FBlueprintBoundMenuItem* NewMenuItem = new FBlueprintBoundMenuItem(BoundAction, MenuGrouping);
+	// FBlueprintBoundMenuItem takes care of its own menu category
+	NewMenuItem->Category = FString::Printf(TEXT("%s|%s"), *RootCategory.ToString(), *NewMenuItem->Category);
+
+	NewMenuItem->MenuDescription    = GetMenuNameForAction(EditorContext, BoundAction);
+	NewMenuItem->TooltipDescription = GetTooltipForAction(EditorContext, BoundAction).ToString();
+	NewMenuItem->Keywords			= GetSearchKeywordsForAction(EditorContext, BoundAction).ToString();
+
+	return MakeShareable(NewMenuItem);
+}
+
 //------------------------------------------------------------------------------
 FText FBlueprintActionMenuItemFactory::GetMenuNameForAction(TWeakPtr<FBlueprintEditor> EditorContext, UBlueprintNodeSpawner const* Action)
 {
@@ -358,6 +380,16 @@ namespace FBlueprintActionMenuBuilderImpl
 		 * @return An empty TSharedPtr if the action was filtered out, otherwise a newly allocated FBlueprintActionMenuItem.
 		 */
 		TSharedPtr<FEdGraphSchemaAction> MakeMenuItem(TWeakPtr<FBlueprintEditor> EditorContext, UBlueprintNodeSpawner const* const DatabaseAction);
+
+		/**
+		 * 
+		 * 
+		 * @param  EditorContext	
+		 * @param  DatabaseAction	
+		 * @param  Bindings	
+		 * @return 
+		 */
+		TSharedPtr<FEdGraphSchemaAction> MakeBoundMenuItem(TWeakPtr<FBlueprintEditor> EditorContext, UBlueprintNodeSpawner const* const DatabaseAction, TArray<UObject*> const& Bindings);
 		
 		/**
 		 * Clears out any consolidated properties that this may have been 
@@ -369,7 +401,7 @@ namespace FBlueprintActionMenuBuilderImpl
 		/** In charge of spawning menu items for this section (holds category/ordering information)*/
 		FBlueprintActionMenuItemFactory ItemFactory;
 		/** Tracks the properties that we've already consolidated and passed (when using the ConsolidatePropertyActions flag)*/
-		TSet<UProperty const*> ConsolidateProperties;
+		TSet<UProperty const*> ConsolidatedProperties;
 	};
 	
 	/**
@@ -409,39 +441,77 @@ void FBlueprintActionMenuBuilderImpl::FMenuSectionDefinition::SetSectionSortOrde
 {
 	ItemFactory.MenuGrouping = MenuGrouping;
 }
-
+// 
 //------------------------------------------------------------------------------
-TSharedPtr<FEdGraphSchemaAction> FBlueprintActionMenuBuilderImpl::FMenuSectionDefinition::MakeMenuItem(TWeakPtr<FBlueprintEditor> EditorContext, UBlueprintNodeSpawner const* const DatabaseAction)
+TSharedPtr<FEdGraphSchemaAction> FBlueprintActionMenuBuilderImpl::FMenuSectionDefinition::MakeBoundMenuItem(TWeakPtr<FBlueprintEditor> EditorContext, UBlueprintNodeSpawner const* const DatabaseAction, TArray<UObject*> const& Bindings)
 {
-	UBlueprintNodeSpawner const* MenuAction = DatabaseAction;
+ 	TSharedPtr<FBlueprintBoundMenuItem> MenuItem;
 
+	UBlueprintNodeSpawner* BoundAction = nullptr;
 	// we don't want the blueprint database growing out of control with an entry 
-	// for every object you could ever possibly bind to, so the 
-	// UBlueprintBoundNodeSpawner from the database serves as a template (note,
-	// it isn't bound to anything)... so, we do the binding here (if anything
-	// to bind to is selected), and use the temporary (bound) spawner in place 
-	// of the passed DatabaseAction
-	if (UBlueprintBoundNodeSpawner const* BoundNodeSpawner = Cast<UBlueprintBoundNodeSpawner>(DatabaseAction))
+	// for every object you could ever possibly bind to, so each 
+	// UBlueprintNodeSpawner comes with an interface to bind through... however,
+	// we can't mutate the action database (by binding to those actions directly)
+	// so, we have to clone and test the action post binding
+	for (auto BindingIt = Bindings.CreateConstIterator(); BindingIt;)
 	{
-		for (UObject const* SelectedObj : ItemFactory.Context.SelectedObjects)
+		UObject const* BindingObj = *BindingIt;
+		++BindingIt;
+
+		// check to see if this object can be bound to this action
+		if (DatabaseAction->CanBind(BindingObj))
 		{
-			// if the bound object is null, then this was most likely set up as a 
-			// template for us to copy and assign a binding object
-			if ((BoundNodeSpawner->GetBoundObject() == nullptr) && BoundNodeSpawner->IsOfBindingType(SelectedObj))
+			// we don't want binding to mutate the database, so we clone the 
+			// action and pass that to FBlueprintBoundMenuItem for binding;
+			// if we haven't created a mutable BoundAction yet, we need that...
+			// we do this prior to filtering because there may be tests that 
+			// reject based off bindings
+			if (BoundAction == nullptr)
 			{
 				// don't worry about outer, if this makes it through the filter,
 				// then the menu item will add a reference to keep it around
-				UBlueprintBoundNodeSpawner* BoundAction = DuplicateObject(BoundNodeSpawner, GetTransientPackage());
-				BoundAction->SetBoundObject(SelectedObj);
+				BoundAction = DuplicateObject<UBlueprintNodeSpawner>(DatabaseAction, GetTransientPackage());
+				BoundAction->CustomizeNodeDelegate = DatabaseAction->CustomizeNodeDelegate; // delegate isn't duplicated
+			}
+
+			// add bindings before filtering (in case tests accept/reject based off of this)
+			BoundAction->AddBinding(BindingObj);
+
+			bool const bIsLastBinding = !BindingIt;
+			// if BoundAction is now "full" (meaning it can take any more 
+			// bindings), or if this is the last binding to test...
+			if (!BoundAction->CanBindMultipleObjects() || bIsLastBinding)
+			{
+				// have to check IsFiltered() for every "fully bound" action (in
+				// case there are tests that reject based off of this), we may 
+				// test this multiple times per action (we have to make sure 
+				// that every set of bound objects pass before folding them into
+				// MenuItem)
+				bool const bPassedFilter = !Filter.IsFiltered(BoundAction);
+				if (bPassedFilter)
+				{
+					if (!MenuItem.IsValid())
+					{
+						MenuItem = ItemFactory.MakeBoundMenuItem(EditorContext, BoundAction);
+					}
 
-				// replace the menu action with one that is actually bound to 
-				// something
-				MenuAction = BoundAction;
-				break;
+					for (TWeakObjectPtr<UObject> Object : BoundAction->GetBindings())
+					{
+						MenuItem->AddBinding(Object.Get());
+					}
+				}
+				BoundAction->ClearBindings();
 			}
 		}
 	}
-	bool bPassedFilter = (MenuAction != nullptr) && !Filter.IsFiltered(MenuAction);
+
+	return MenuItem;
+}
+
+//------------------------------------------------------------------------------
+TSharedPtr<FEdGraphSchemaAction> FBlueprintActionMenuBuilderImpl::FMenuSectionDefinition::MakeMenuItem(TWeakPtr<FBlueprintEditor> EditorContext, UBlueprintNodeSpawner const* const DatabaseAction)
+{	
+	bool bPassedFilter = (DatabaseAction != nullptr) && !Filter.IsFiltered(DatabaseAction);
 
 	bool bCreateDragDropItem = false;
 	// if the caller wants to consolidate all property actions, then we have to 
@@ -450,15 +520,15 @@ TSharedPtr<FEdGraphSchemaAction> FBlueprintActionMenuBuilderImpl::FMenuSectionDe
 	if (bPassedFilter && (Flags & FBlueprintActionMenuBuilder::ConsolidatePropertyActions))
 	{
 		UProperty const* ActionProperty = nullptr;
-		if (UBlueprintVariableNodeSpawner const* VariableSpawner = Cast<UBlueprintVariableNodeSpawner>(MenuAction))
+		if (UBlueprintVariableNodeSpawner const* VariableSpawner = Cast<UBlueprintVariableNodeSpawner>(DatabaseAction))
 		{
-			bCreateDragDropItem  = true;
+			bCreateDragDropItem = true;
 			ActionProperty = VariableSpawner->GetVarProperty();
 			bPassedFilter = (ActionProperty != nullptr);
 		}
-		else if (UBlueprintDelegateNodeSpawner const* DelegateSpawner = Cast<UBlueprintDelegateNodeSpawner>(MenuAction))
+		else if (UBlueprintDelegateNodeSpawner const* DelegateSpawner = Cast<UBlueprintDelegateNodeSpawner>(DatabaseAction))
 		{
-			bCreateDragDropItem  = true;
+			bCreateDragDropItem = true;
 			ActionProperty = DelegateSpawner->GetProperty();
 			bPassedFilter = (ActionProperty != nullptr);
 		}
@@ -466,31 +536,32 @@ TSharedPtr<FEdGraphSchemaAction> FBlueprintActionMenuBuilderImpl::FMenuSectionDe
 		if (ActionProperty != nullptr)
 		{
 			// @TODO: could be added as a filter test instead
-			bool const bAlreadyConsolidated = ConsolidateProperties.Contains(ActionProperty);
+			bool const bAlreadyConsolidated = ConsolidatedProperties.Contains(ActionProperty);
 			bPassedFilter = !bAlreadyConsolidated;
-			ConsolidateProperties.Add(ActionProperty);
+			ConsolidatedProperties.Add(ActionProperty);
 		}
 	}
 
-	TSharedPtr<FEdGraphSchemaAction> MenuEntry = nullptr;
+	TSharedPtr<FEdGraphSchemaAction> MenuEntry;
 	if (bPassedFilter)
 	{
 		if (bCreateDragDropItem)
 		{
-			MenuEntry = ItemFactory.MakeDragDropMenuItem(MenuAction);
+			MenuEntry = ItemFactory.MakeDragDropMenuItem(DatabaseAction);
 		}
 		else
 		{
-			MenuEntry = ItemFactory.MakeActionMenuItem(EditorContext, MenuAction);
+			MenuEntry = ItemFactory.MakeActionMenuItem(EditorContext, DatabaseAction);
 		}
 	}
+
 	return MenuEntry;
 }
 
 //------------------------------------------------------------------------------
 void FBlueprintActionMenuBuilderImpl::FMenuSectionDefinition::Empty()
 {
-	ConsolidateProperties.Empty();
+	ConsolidatedProperties.Empty();
 }
 
 //------------------------------------------------------------------------------
@@ -607,6 +678,16 @@ void FBlueprintActionMenuBuilder::RebuildActionList()
 					{
 						AddAction(MenuEntry);
 					}
+
+					TArray<UObject*> const& BindingObjects = MenuSection->Filter.Context.SelectedObjects;
+					// @TODO: this only makes one "BoundMenuEntry" for all 
+					//        bindings, we could have a non-consolidated mode if
+					//        we want (like we have for property actions)
+					TSharedPtr<FEdGraphSchemaAction> BoundMenuEntry = MenuSection->MakeBoundMenuItem(BlueprintEditorPtr, NodeSpawner, BindingObjects);
+					if (BoundMenuEntry.IsValid())
+					{
+						AddAction(BoundMenuEntry);
+					}
 				}
 			}
 		}
diff --git a/Engine/Source/Editor/Kismet/Private/BlueprintActionMenuUtils.cpp b/Engine/Source/Editor/Kismet/Private/BlueprintActionMenuUtils.cpp
--- a/Engine/Source/Editor/Kismet/Private/BlueprintActionMenuUtils.cpp
+++ b/Engine/Source/Editor/Kismet/Private/BlueprintActionMenuUtils.cpp
@@ -4,12 +4,13 @@
 #include "BlueprintActionMenuUtils.h"
 #include "BlueprintActionMenuBuilder.h"
 #include "BlueprintActionFilter.h"
-#include "BlueprintBoundNodeSpawner.h"
-#include "KismetEditorUtilities.h" // for CanPasteNodes()
+#include "KismetEditorUtilities.h"	// for CanPasteNodes()
+#include "K2Node_Variable.h"
 #include "K2Node_CallFunction.h"
+#include "K2Node_ActorBoundEvent.h"
 #include "K2Node_ComponentBoundEvent.h"
-#include "EdGraphSchema_K2.h"
-#include "K2Node_Variable.h"
+#include "EdGraphSchema_K2.h"		// for bUseLegacyActionMenus 
+#include "BlueprintEditorUtils.h"	// for DoesSupportComponents()
 
 #define LOCTEXT_NAMESPACE "BlueprintActionMenuUtils"
 
@@ -19,18 +20,46 @@
 
 namespace BlueprintActionMenuUtilsImpl
 {
+	/**
+	 * 
+	 * 
+	 * @param  Filter	
+	 * @param  BlueprintAction	
+	 * @return 
+	 */
 	static bool IsUnBoundSpawner(FBlueprintActionFilter const& Filter, UBlueprintNodeSpawner const* BlueprintAction);
+
+	/**
+	 * 
+	 * 
+	 * @param  ClassSet	
+	 * @return 
+	 */
+	static UClass* FinCommonBaseClass(TArray<UClass*> const& ClassSet);
 }
 
 //------------------------------------------------------------------------------
 static bool BlueprintActionMenuUtilsImpl::IsUnBoundSpawner(FBlueprintActionFilter const& /*Filter*/, UBlueprintNodeSpawner const* BlueprintAction)
 {
-	bool bIsFilteredOut = true;
-	if (UBlueprintBoundNodeSpawner const* BoundSpawner = Cast<UBlueprintBoundNodeSpawner>(BlueprintAction))
+	return !BlueprintAction->IsBindingSet();
+}
+
+//------------------------------------------------------------------------------
+static UClass* BlueprintActionMenuUtilsImpl::FinCommonBaseClass(TArray<UClass*> const& ClassSet)
+{
+	UClass* CommonClass = UObject::StaticClass();
+	if (ClassSet.Num() > 0)
 	{
-		bIsFilteredOut = (BoundSpawner->GetBoundObject() == nullptr);
-	}
-	return bIsFilteredOut;
+		CommonClass = ClassSet[0];
+		for (UClass const* Class : ClassSet)
+		{
+			while (!Class->IsChildOf(CommonClass))
+			{
+				CommonClass = CommonClass->GetSuperClass();
+			}
+		}
+	}	
+	return CommonClass;
 }
 
 /*******************************************************************************
@@ -71,52 +100,118 @@ void FBlueprintActionMenuUtils::MakeContextMenu(FBlueprintActionContext const& C
 {
 	using namespace BlueprintActionMenuUtilsImpl;
 	
-	static int32 const MainMenuSectionGroup   = 000;
-	static int32 const ComponentsSectionGroup = 100;
+	//--------------------------------------
+	// Composing Filters
+	//--------------------------------------
+
+	FBlueprintActionFilter MenuFilter;
+	MenuFilter.Context = Context;
+	MenuFilter.Context.SelectedObjects.Empty();
 
 	FBlueprintActionFilter ComponentsFilter;
 	ComponentsFilter.Context = Context;
-	ComponentsFilter.PermittedNodeTypes.Add(UK2Node_ComponentBoundEvent::StaticClass());
 	ComponentsFilter.PermittedNodeTypes.Add(UK2Node_CallFunction::StaticClass());
+	ComponentsFilter.PermittedNodeTypes.Add(UK2Node_ComponentBoundEvent::StaticClass());
+	// only want bound actions for this menu section
 	ComponentsFilter.AddIsFilteredTest(FBlueprintActionFilter::FIsFilteredDelegate::CreateStatic(IsUnBoundSpawner));
-	
-	bool bHasComponentsSection = false;
+
+	FBlueprintActionFilter LevelActorsFilter;
+	LevelActorsFilter.Context = Context;
+	LevelActorsFilter.PermittedNodeTypes.Add(UK2Node_CallFunction::StaticClass());
+	LevelActorsFilter.PermittedNodeTypes.Add(UK2Node_ActorBoundEvent::StaticClass());
+	// only want bound actions for this menu section
+	LevelActorsFilter.AddIsFilteredTest(FBlueprintActionFilter::FIsFilteredDelegate::CreateStatic(IsUnBoundSpawner));
+
+	// make sure the bound menu sections have the proper OwnerClasses specified
 	for (UObject* Selection : Context.SelectedObjects)
 	{
 		if (UObjectProperty* ObjProperty = Cast<UObjectProperty>(Selection))
 		{
 			ComponentsFilter.TargetClasses.Add(ObjProperty->PropertyClass);
-			bHasComponentsSection = true;
+			LevelActorsFilter.Context.SelectedObjects.Remove(Selection);
+		}
+		else if (AActor* LevelActor = Cast<AActor>(Selection))
+		{
+			ComponentsFilter.Context.SelectedObjects.Remove(Selection);
+			// the loop below (over the editor's selected actors) will add to 
+			// LevelActorsFilter's OwnerClasses
+		}
+		else
+		{
+			ComponentsFilter.Context.SelectedObjects.Remove(Selection);
+			LevelActorsFilter.Context.SelectedObjects.Remove(Selection);
 		}
 	}
 
-	MenuOut.Empty();
-
-	FBlueprintActionFilter MenuFilter;
-	MenuFilter.Context = Context;
-	MenuFilter.Context.SelectedObjects.Empty();
+	// make sure all selected level actors are accounted for (in case the caller
+	// did not include them in the context)
+	for (FSelectionIterator LvlActorIt(*GEditor->GetSelectedActors()); LvlActorIt; ++LvlActorIt)
+	{
+		AActor* LevelActor = Cast<AActor>(*LvlActorIt);
+		LevelActorsFilter.Context.SelectedObjects.AddUnique(LevelActor);
+		LevelActorsFilter.TargetClasses.Add(LevelActor->GetClass());
+	}
 
-	if (bIsContextSensitive)
+	bool const bOnlyBlueprintMembers = bIsContextSensitive && (MenuFilter.Context.Pins.Num() == 0);
+	bool bCanOperateOnLevelActors = bIsContextSensitive;
+	bool bCanHaveActorComponents  = bIsContextSensitive;
+	// determine if we can operate on certain object selections (level actors, 
+	// components, etc.)
+	for (UBlueprint* Blueprint : Context.Blueprints)
 	{
-		if (bHasComponentsSection)
-		{
-			MenuOut.AddMenuSection(ComponentsFilter, FText::GetEmpty(), ComponentsSectionGroup);
-		}
-		
-		if (MenuFilter.Context.Pins.Num() == 0)
+		UClass* BlueprintClass = Blueprint->SkeletonGeneratedClass;
+		if (BlueprintClass != nullptr)
 		{
-			for (UBlueprint* Blueprint : Context.Blueprints)
+			bCanOperateOnLevelActors &= BlueprintClass->IsChildOf<ALevelScriptActor>();
+			if (bOnlyBlueprintMembers)
 			{
-				UClass* BlueprintClass = (Blueprint->SkeletonGeneratedClass != nullptr) ? Blueprint->SkeletonGeneratedClass : Blueprint->ParentClass;
 				MenuFilter.TargetClasses.Add(BlueprintClass);
 			}
 		}
+		bCanHaveActorComponents &= FBlueprintEditorUtils::DoesSupportComponents(Blueprint);
+	}
+
+	//--------------------------------------
+	// Defining Menu Sections
+	//--------------------------------------
+
+	static int32 const MainMenuSectionGroup   = 000;
+	static int32 const ComponentsSectionGroup = 100;
+	static int32 const LevelActorSectionGroup = 101;
+
+	MenuOut.Empty();
+
+	bool const bAddComponentsSection = bIsContextSensitive && bCanHaveActorComponents && (ComponentsFilter.Context.SelectedObjects.Num() > 0);
+	// add the components section to the menu (if we don't have any components
+	// selected, then inform the user through a dummy menu entry)
+	if (bAddComponentsSection)
+	{
+		MenuOut.AddMenuSection(ComponentsFilter, FText::GetEmpty(), ComponentsSectionGroup);
+	}
+
+	bool const bAddLevelActorsSection = bIsContextSensitive && bCanOperateOnLevelActors && (LevelActorsFilter.Context.SelectedObjects.Num() > 0);
+	// add the level actor section to the menu
+	if (bAddLevelActorsSection)
+	{
+		// since we're consolidating all the bound actions, then we have to pick 
+		// one common base class to filter by
+		UClass* CommonClass = FinCommonBaseClass(LevelActorsFilter.TargetClasses);
+		LevelActorsFilter.TargetClasses.Empty();
+		LevelActorsFilter.TargetClasses.Add(CommonClass);
+
+		MenuOut.AddMenuSection(LevelActorsFilter, FText::GetEmpty(), LevelActorSectionGroup);
 	}
-	else
+
+	if (!bIsContextSensitive)
 	{
 		MenuFilter.Context.Pins.Empty();
 	}
 	MenuOut.AddMenuSection(MenuFilter);
+
+	//--------------------------------------
+	// Building the Menu
+	//--------------------------------------
+
 	MenuOut.RebuildActionList();
 
 	UEdGraphSchema_K2 const* const K2Schema = GetDefault<UEdGraphSchema_K2>();
@@ -133,13 +228,21 @@ void FBlueprintActionMenuUtils::MakeContextMenu(FBlueprintActionContext const& C
 			}
 		}
 
-		if (bIsContextSensitive && (Context.SelectedObjects.Num() == 0))
+		if (bIsContextSensitive && bCanHaveActorComponents && !bAddComponentsSection)
 		{
 			FText SelectComponentMsg = LOCTEXT("SelectComponentForEvents", "Select a Component to see available Events & Functions");
 			FText SelectComponentToolTip = LOCTEXT("SelectComponentForEventsTooltip", "Select a Component in the MyBlueprint tab to see available Events and Functions in this menu.");
 			TSharedPtr<FEdGraphSchemaAction> MsgAction = TSharedPtr<FEdGraphSchemaAction>(new FEdGraphSchemaAction_Dummy(TEXT(""), SelectComponentMsg, SelectComponentToolTip.ToString(), ComponentsSectionGroup));
 			MenuOut.AddAction(MsgAction);
 		}
+
+		if (bIsContextSensitive && bCanOperateOnLevelActors && !bAddLevelActorsSection)
+		{
+			FText SelectActorsMsg = LOCTEXT("SelectActorForEvents", "Select Actor(s) to see available Events & Functions");
+			FText SelectActorsToolTip = LOCTEXT("SelectActorForEventsTooltip", "Select Actor(s) in the level to see available Events and Functions in this menu.");
+			TSharedPtr<FEdGraphSchemaAction> MsgAction = TSharedPtr<FEdGraphSchemaAction>(new FEdGraphSchemaAction_Dummy(TEXT(""), SelectActorsMsg, SelectActorsToolTip.ToString(), LevelActorSectionGroup));
+			MenuOut.AddAction(MsgAction);
+		}
 	}	
 }
 
diff --git a/Engine/Source/Editor/Kismet/Private/BlueprintBoundMenuItem.cpp b/Engine/Source/Editor/Kismet/Private/BlueprintBoundMenuItem.cpp
new file mode 100644
--- /dev/null
+++ b/Engine/Source/Editor/Kismet/Private/BlueprintBoundMenuItem.cpp
@@ -0,0 +1,175 @@
+// Copyright 1998-2014 Epic Games, Inc. All Rights Reserved.
+
+#include "BlueprintEditorPrivatePCH.h"
+#include "BlueprintBoundMenuItem.h"
+#include "BlueprintActionMenuItem.h"
+#include "BlueprintNodeSpawner.h"
+#include "ScopedTransaction.h"
+
+#define LOCTEXT_NAMESPACE "BlueprintBoundMenuItem"
+
+/*******************************************************************************
+* Static FBlueprintBoundMenuItem Helpers
+******************************************************************************/
+
+namespace BlueprintBoundMenuItemImpl
+{
+	/**
+	 * 
+	 * 
+	 * @param  BoundSpawner	
+	 * @return 
+	 */
+	static FText GetMenuCategoryFromAction(UBlueprintNodeSpawner* BoundSpawner);
+
+	/**
+	 * 
+	 * 
+	 * @param  ObjectsList	
+	 * @return 
+	 */
+	static UClass* FindCommonBaseClass(TSet< TWeakObjectPtr<UObject> > const& ObjectsList);
+}
+
+//------------------------------------------------------------------------------
+static FText BlueprintBoundMenuItemImpl::GetMenuCategoryFromAction(UBlueprintNodeSpawner* BoundSpawner)
+{
+	TSet< TWeakObjectPtr<UObject> > const& BoundObjects = BoundSpawner->GetBindings();
+	check(BoundObjects.Num() > 0);
+
+	FText BoundObjectName;
+	if (BoundObjects.Num() > 1)
+	{
+		UClass* CommonBaseClass = FindCommonBaseClass(BoundObjects);
+		BoundObjectName = FText::Format(LOCTEXT("MultipleBindingsCategoryPostfix", "Selected {0}s"), CommonBaseClass->GetDisplayNameText());
+	}
+	else if (UObject const* BoundObject = BoundObjects.CreateConstIterator()->Get())
+	{
+		BoundObjectName = FText::FromName(BoundObject->GetFName());
+	}
+
+	FText MenuCategory = BoundSpawner->GetDefaultMenuCategory();
+	if (BoundSpawner->NodeClass->IsChildOf<UK2Node_Event>())
+	{
+		MenuCategory = FText::Format(LOCTEXT("ComponentEventCategory", "Add Event for {0}|{1}"), BoundObjectName, MenuCategory);
+	}
+	else if (BoundSpawner->NodeClass->IsChildOf<UK2Node_CallFunction>())
+	{
+		MenuCategory = FText::Format(LOCTEXT("ComponentEventCategory", "Call Function on {0}|{1}"), BoundObjectName, MenuCategory);
+	}
+
+	return MenuCategory;
+}
+
+//------------------------------------------------------------------------------
+static UClass* BlueprintBoundMenuItemImpl::FindCommonBaseClass(TSet< TWeakObjectPtr<UObject> > const& ObjectsList)
+{
+	UClass* CommonClass = nullptr;
+	for (TWeakObjectPtr<UObject> BoundObjPtr : ObjectsList)
+	{
+		if (UObject const* BoundObj = BoundObjPtr.Get())
+		{
+			if (CommonClass == nullptr)
+			{
+				CommonClass = BoundObj->GetClass();
+			}
+			else
+			{
+				while (!BoundObj->IsA(CommonClass))
+				{
+					CommonClass = CommonClass->GetSuperClass();
+				}
+			}
+		}
+	}
+
+	if (CommonClass == nullptr)
+	{
+		CommonClass = UObject::StaticClass();
+	}
+	return CommonClass;
+}
+
+/*******************************************************************************
+ * FBlueprintBoundMenuItem
+ ******************************************************************************/
+
+//------------------------------------------------------------------------------
+FBlueprintBoundMenuItem::FBlueprintBoundMenuItem(UBlueprintNodeSpawner* BoundSpawnerIn, int32 MenuGrouping/* = 0*/)
+	: BoundSpawner(BoundSpawnerIn) 
+{
+	Grouping = MenuGrouping;
+	Category = BlueprintBoundMenuItemImpl::GetMenuCategoryFromAction(BoundSpawnerIn).ToString();
+}
+
+//------------------------------------------------------------------------------
+UEdGraphNode* FBlueprintBoundMenuItem::PerformAction(UEdGraph* ParentGraph, UEdGraphPin* FromPin, FVector2D const Location, bool bSelectNewNode/* = true*/)
+{
+	FBlueprintActionMenuItem BlueprintActionItem(BoundSpawner, nullptr, FSlateColor(FLinearColor::White));
+
+	UEdGraphNode* NewNode = nullptr;
+	for (auto BoundObjIt = BoundObjects.CreateConstIterator(); BoundObjIt; )
+	{
+		do
+		{
+			if (BoundObjIt->IsValid())
+			{
+				BoundSpawner->AddBinding(BoundObjIt->Get());
+			}
+			++BoundObjIt;
+
+		} while (BoundObjIt && (BoundSpawner->CanBindMultipleObjects() || !BoundSpawner->IsBindingSet()));
+
+		NewNode = BlueprintActionItem.PerformAction(ParentGraph, FromPin, Location, bSelectNewNode);
+	}
+	
+	return NewNode;
+}
+
+//------------------------------------------------------------------------------
+UEdGraphNode* FBlueprintBoundMenuItem::PerformAction(UEdGraph* ParentGraph, TArray<UEdGraphPin*>& FromPins, FVector2D const Location, bool bSelectNewNode/* = true*/)
+{
+	UEdGraphPin* FromPin = nullptr;
+	if (FromPins.Num() > 0)
+	{
+		FromPin = FromPins[0];
+	}
+	
+	UEdGraphNode* SpawnedNode = PerformAction(ParentGraph, FromPin, Location, bSelectNewNode);
+	// try auto-wiring the rest of the pins (if there are any)
+	for (int32 PinIndex = 1; PinIndex < FromPins.Num(); ++PinIndex)
+	{
+		SpawnedNode->AutowireNewNode(FromPins[PinIndex]);
+	}
+	
+	return SpawnedNode;
+}
+
+//------------------------------------------------------------------------------
+void FBlueprintBoundMenuItem::AddReferencedObjects(FReferenceCollector& Collector)
+{
+	FEdGraphSchemaAction::AddReferencedObjects(Collector);
+	
+	// these don't get saved to disk, but we want to make sure the objects don't
+	// get GC'd while the action array is around
+	Collector.AddReferencedObject(BoundSpawner);
+}
+
+//------------------------------------------------------------------------------
+bool FBlueprintBoundMenuItem::AddBinding(UObject const* Binding)
+{
+	bool const bIsCompatible = (Binding != nullptr) && BoundSpawner->CanBind(Binding);
+	if (bIsCompatible)
+	{
+		BoundObjects.Add(Binding);
+	}
+	return bIsCompatible;
+}
+
+//------------------------------------------------------------------------------
+FSlateBrush const* FBlueprintBoundMenuItem::GetMenuIcon(FSlateColor& ColorOut)
+{
+	return nullptr;
+}
+
+#undef LOCTEXT_NAMESPACE
diff --git a/Engine/Source/Editor/Kismet/Public/BlueprintBoundMenuItem.h b/Engine/Source/Editor/Kismet/Public/BlueprintBoundMenuItem.h
new file mode 100644
--- /dev/null
+++ b/Engine/Source/Editor/Kismet/Public/BlueprintBoundMenuItem.h
@@ -0,0 +1,55 @@
+// Copyright 1998-2014 Epic Games, Inc. All Rights Reserved.
+
+#pragma once
+
+#include "EdGraph/EdGraphSchema.h" // for FEdGraphSchemaAction
+#include "SlateColor.h"
+#include "BlueprintBoundMenuItem.generated.h"
+
+// Forward declarations
+class UBlueprintNodeSpawner;
+class FDragDropOperation;
+struct FSlateBrush;
+
+/**
+ * Wrapper around a UBlueprintNodeSpawner, which takes care of specialized
+ * node spawning. This class should not be sub-classed, any special handling
+ * should be done inside a UBlueprintNodeSpawner subclass, which will be
+ * invoked from this class (separated to divide ui and node-spawning).
+ */
+USTRUCT()
+struct FBlueprintBoundMenuItem : public FEdGraphSchemaAction
+{
+	GENERATED_USTRUCT_BODY();
+	
+public:
+	static FName StaticGetTypeId() { static FName const TypeId("FBlueprintBoundMenuItem"); return TypeId; }
+	
+	/** Constructors */
+	FBlueprintBoundMenuItem() : BoundSpawner(nullptr) {}
+	FBlueprintBoundMenuItem(UBlueprintNodeSpawner* BoundSpawner, int32 MenuGrouping = 0);
+	
+	// FEdGraphSchemaAction interface
+	virtual FName         GetTypeId() const final { return StaticGetTypeId(); }
+	virtual UEdGraphNode* PerformAction(UEdGraph* ParentGraph, UEdGraphPin* FromPin, FVector2D const Location, bool bSelectNewNode = true) final;
+	virtual UEdGraphNode* PerformAction(UEdGraph* ParentGraph, TArray<UEdGraphPin*>& FromPins, FVector2D const Location, bool bSelectNewNode = true) final;
+	virtual void          AddReferencedObjects(FReferenceCollector& Collector) final;
+	// End FEdGraphSchemaAction interface
+
+	bool AddBinding(UObject const* Binding);
+
+	/**
+	 * Retrieves the icon brush for this menu entry (to be displayed alongside
+	 * in the menu).
+	 *
+	 * @param  ColorOut	An output parameter that's filled in with the color to tint the brush with.
+	 * @return An slate brush to be used for this menu item in the action menu.
+	 */
+	FSlateBrush const* GetMenuIcon(FSlateColor& ColorOut);	
+
+private:
+	/** Instanced node-spawner, that comprises the action portion of this menu entry. */
+	UBlueprintNodeSpawner* BoundSpawner;
+	/** */
+	TSet< TWeakObjectPtr<UObject> > BoundObjects;
+};
diff --git a/Engine/Source/Editor/UnrealEd/Private/Kismet2/Kismet2.cpp b/Engine/Source/Editor/UnrealEd/Private/Kismet2/Kismet2.cpp
--- a/Engine/Source/Editor/UnrealEd/Private/Kismet2/Kismet2.cpp
+++ b/Engine/Source/Editor/UnrealEd/Private/Kismet2/Kismet2.cpp
@@ -1251,7 +1251,7 @@ void FKismetEditorUtilities::CreateNewBoundEventForComponent(UActorComponent* Co
 	}
 }
 
-const UK2Node_ActorBoundEvent* FKismetEditorUtilities::FindBoundEventForActor(AActor* Actor, FName EventName)
+const UK2Node_ActorBoundEvent* FKismetEditorUtilities::FindBoundEventForActor(AActor const* Actor, FName EventName)
 {
 	const UK2Node_ActorBoundEvent* Node = NULL;
 	if(Actor != NULL && EventName != NAME_None)
diff --git a/Engine/Source/Editor/UnrealEd/Public/Kismet2/KismetEditorUtilities.h b/Engine/Source/Editor/UnrealEd/Public/Kismet2/KismetEditorUtilities.h
--- a/Engine/Source/Editor/UnrealEd/Public/Kismet2/KismetEditorUtilities.h
+++ b/Engine/Source/Editor/UnrealEd/Public/Kismet2/KismetEditorUtilities.h
@@ -114,7 +114,7 @@ public:
 	static int32 GetNumberOfSelectedNodes(const class UBlueprint* Blueprint);
 
 	/** Find the event node for this actor with the given event name */
-	static const class UK2Node_ActorBoundEvent* FindBoundEventForActor(AActor* Actor, FName EventName);
+	static const class UK2Node_ActorBoundEvent* FindBoundEventForActor(AActor const* Actor, FName EventName);
 
 	/** Find the event node for the component property with the given event name */
 	static const class UK2Node_ComponentBoundEvent* FindBoundEventForComponent(const UBlueprint* Blueprint, FName EventName, FName PropertyName);
