Fixing for non-unity build (missed include file).

diff --git a/Engine/Source/Editor/Kismet/Private/BlueprintActionMenuUtils.cpp b/Engine/Source/Editor/Kismet/Private/BlueprintActionMenuUtils.cpp
--- a/Engine/Source/Editor/Kismet/Private/BlueprintActionMenuUtils.cpp
+++ b/Engine/Source/Editor/Kismet/Private/BlueprintActionMenuUtils.cpp
@@ -4,6 +4,7 @@
 #include "BlueprintActionMenuUtils.h"
 #include "BlueprintActionMenuBuilder.h"
 #include "BlueprintActionFilter.h"
+#include "BlueprintNodeSpawner.h"
 #include "KismetEditorUtilities.h"	// for CanPasteNodes()
 #include "K2Node_Variable.h"
 #include "K2Node_CallFunction.h"
@@ -243,7 +244,7 @@ void FBlueprintActionMenuUtils::MakeContextMenu(FBlueprintActionContext const& C
 			TSharedPtr<FEdGraphSchemaAction> MsgAction = TSharedPtr<FEdGraphSchemaAction>(new FEdGraphSchemaAction_Dummy(TEXT(""), SelectActorsMsg, SelectActorsToolTip.ToString(), LevelActorSectionGroup));
 			MenuOut.AddAction(MsgAction);
 		}
-	}	
+	}
 }
 
 #undef LOCTEXT_NAMESPACE
