#ue4 - Perf: FBlueprintActionMenuBuilder::RebuildActionList: Iterating ActionDatabase doesn't create copies of TArrays. Iterating MenuSections doesn't copy TSharedPtrs.

diff --git a/Engine/Source/Editor/Kismet/Private/BlueprintActionMenuBuilder.cpp b/Engine/Source/Editor/Kismet/Private/BlueprintActionMenuBuilder.cpp
--- a/Engine/Source/Editor/Kismet/Private/BlueprintActionMenuBuilder.cpp
+++ b/Engine/Source/Editor/Kismet/Private/BlueprintActionMenuBuilder.cpp
@@ -665,13 +665,13 @@ void FBlueprintActionMenuBuilder::RebuildActionList()
 	if (!GetDefault<UEdGraphSchema_K2>()->bUseLegacyActionMenus)
 	{
 		FBlueprintActionDatabase::FActionRegistry const& ActionDatabase = FBlueprintActionDatabase::Get().GetAllActions();
-		for (auto ActionEntry : ActionDatabase)
+		for (auto const& ActionEntry : ActionDatabase)
 		{
 			for (UBlueprintNodeSpawner* NodeSpawner : ActionEntry.Value)
 			{
 				// @TODO: could probably have a super filter that spreads across 
 				//        all MenuSctions (to pair down on performance)
-				for (TSharedRef<FMenuSectionDefinition> MenuSection : MenuSections)
+				for (TSharedRef<FMenuSectionDefinition>& MenuSection : MenuSections)
 				{
 					TSharedPtr<FEdGraphSchemaAction> MenuEntry = MenuSection->MakeMenuItem(BlueprintEditorPtr, NodeSpawner);
 					if (MenuEntry.IsValid())
