#UE4 diff control can now enumerate differences in SCS, same for merge control. No longer looping around the end of the differences list

diff --git a/Engine/Source/Developer/Merge/Private/SBlueprintMerge.cpp b/Engine/Source/Developer/Merge/Private/SBlueprintMerge.cpp
--- a/Engine/Source/Developer/Merge/Private/SBlueprintMerge.cpp
+++ b/Engine/Source/Developer/Merge/Private/SBlueprintMerge.cpp
@@ -10,6 +10,7 @@
 #include "SBlueprintMerge.h"
 #include "SMergeDetailsView.h"
 #include "SMergeGraphView.h"
+#include "SMergeTreeView.h"
 
 #define LOCTEXT_NAMESPACE "SBlueprintMerge"
 
@@ -25,11 +26,13 @@ void SBlueprintMerge::Construct(const FArguments InArgs, const FBlueprintMergeDa
 
 	Data = InData;
 
-	TSharedPtr<SMergeGraphView> GraphView = SNew( SMergeGraphView, InData );
+	auto GraphView = SNew( SMergeGraphView, InData );
 	GraphControl.Widget = GraphView;
-	GraphControl.DiffControl = GraphView.Get();
+	GraphControl.DiffControl = &GraphView.Get();
 
-	TreeControl.Widget = SNew(SBorder);
+	auto TreeView = SNew(SMergeTreeView, InData);
+	TreeControl.Widget = TreeView;
+	TreeControl.DiffControl = &TreeView.Get();
 
 	auto DetailsView = SNew( SMergeDetailsView, InData );
 	DetailsControl.Widget = DetailsView;
@@ -37,14 +40,14 @@ void SBlueprintMerge::Construct(const FArguments InArgs, const FBlueprintMergeDa
 
 	FToolBarBuilder ToolbarBuilder(TSharedPtr< FUICommandList >(), FMultiBoxCustomization::None);
 	ToolbarBuilder.AddToolBarButton( 
-		FUIAction( FExecuteAction::CreateSP(this, &SBlueprintMerge::PrevDiff), FCanExecuteAction::CreateSP(this, &SBlueprintMerge::CanCycleDiffs) )
+		FUIAction( FExecuteAction::CreateSP(this, &SBlueprintMerge::PrevDiff), FCanExecuteAction::CreateSP(this, &SBlueprintMerge::HasPrevDiff) )
 		, NAME_None
 		, LOCTEXT("PrevMergeLabel", "Prev")
 		, LOCTEXT("PrevMergeTooltip", "Go to previous difference")
 		, FSlateIcon(FEditorStyle::GetStyleSetName(), "BlueprintMerge.PrevDiff")
 	);
 	ToolbarBuilder.AddToolBarButton(
-		FUIAction(FExecuteAction::CreateSP(this, &SBlueprintMerge::NextDiff), FCanExecuteAction::CreateSP(this, &SBlueprintMerge::CanCycleDiffs))
+		FUIAction(FExecuteAction::CreateSP(this, &SBlueprintMerge::NextDiff), FCanExecuteAction::CreateSP(this, &SBlueprintMerge::HasNextDiff))
 		, NAME_None
 		, LOCTEXT("NextMergeLabel", "Next")
 		, LOCTEXT("NextMergeTooltip", "Go to next difference")
@@ -127,9 +130,14 @@ void SBlueprintMerge::PrevDiff()
 	CurrentDiffControl->PrevDiff();
 }
 
-bool SBlueprintMerge::CanCycleDiffs() const
+bool SBlueprintMerge::HasNextDiff() const
 {
-	return CurrentDiffControl && CurrentDiffControl->HasDifferences();
+	return CurrentDiffControl && CurrentDiffControl->HasNextDifference();
+}
+
+bool SBlueprintMerge::HasPrevDiff() const
+{
+	return CurrentDiffControl && CurrentDiffControl->HasPrevDifference();
 }
 
 void SBlueprintMerge::OnAcceptResultClicked()
diff --git a/Engine/Source/Developer/Merge/Private/SBlueprintMerge.h b/Engine/Source/Developer/Merge/Private/SBlueprintMerge.h
--- a/Engine/Source/Developer/Merge/Private/SBlueprintMerge.h
+++ b/Engine/Source/Developer/Merge/Private/SBlueprintMerge.h
@@ -23,7 +23,8 @@ private:
 	/** Event handlers */
 	void NextDiff();
 	void PrevDiff();
-	bool CanCycleDiffs() const;
+	bool HasNextDiff() const;
+	bool HasPrevDiff() const;
 	void OnAcceptResultClicked();
 	void OnCancelClicked();
 	void OnModeChanged(FName NewMode);
diff --git a/Engine/Source/Developer/Merge/Private/SMergeDetailsView.cpp b/Engine/Source/Developer/Merge/Private/SMergeDetailsView.cpp
--- a/Engine/Source/Developer/Merge/Private/SMergeDetailsView.cpp
+++ b/Engine/Source/Developer/Merge/Private/SMergeDetailsView.cpp
@@ -150,9 +150,14 @@ void SMergeDetailsView::HighlightCurrentDifference()
 	}
 }
 
-bool SMergeDetailsView::HasDifferences() const
+bool SMergeDetailsView::HasNextDifference() const
 {
-	return DifferingProperties.Num() != 0;
+	return DifferingProperties.IsValidIndex( CurrentDifference + 1 );
+}
+
+bool SMergeDetailsView::HasPrevDifference() const
+{
+	return DifferingProperties.IsValidIndex(CurrentDifference - 1);
 }
 
 FDetailsDiff& SMergeDetailsView::GetRemoteDetails()
diff --git a/Engine/Source/Developer/Merge/Private/SMergeDetailsView.h b/Engine/Source/Developer/Merge/Private/SMergeDetailsView.h
--- a/Engine/Source/Developer/Merge/Private/SMergeDetailsView.h
+++ b/Engine/Source/Developer/Merge/Private/SMergeDetailsView.h
@@ -20,12 +20,13 @@ private:
 	/** Implementation of IDiffControl: */
 	void NextDiff() override;
 	void PrevDiff() override;
-	bool HasDifferences() const override;
+	bool HasNextDifference() const override;
+	bool HasPrevDifference() const override;
 
 	void HighlightCurrentDifference();
-	class FDetailsDiff& GetRemoteDetails();
-	class FDetailsDiff& GetBaseDetails();
-	class FDetailsDiff& GetLocalDetails();
+	FDetailsDiff& GetRemoteDetails();
+	FDetailsDiff& GetBaseDetails();
+	FDetailsDiff& GetLocalDetails();
 
 	FBlueprintMergeData Data;
 	TArray< FDetailsDiff > DetailsViews;
diff --git a/Engine/Source/Developer/Merge/Private/SMergeGraphView.cpp b/Engine/Source/Developer/Merge/Private/SMergeGraphView.cpp
--- a/Engine/Source/Developer/Merge/Private/SMergeGraphView.cpp
+++ b/Engine/Source/Developer/Merge/Private/SMergeGraphView.cpp
@@ -369,10 +369,32 @@ void SMergeGraphView::PrevDiff()
 	DiffWidgetUtils::SelectPrevRow( *DiffResultList.Pin(), *DiffResultsListData );
 }
 
-bool SMergeGraphView::HasDifferences() const
+bool SMergeGraphView::HasNextDifference() const
+{
+	if (HasNoDifferences())
+	{
+		return false;
+	}
+
+	check(DiffResultList.Pin().IsValid() && DiffResultsListData);
+	return DiffWidgetUtils::HasNextDifference(*DiffResultList.Pin(), *DiffResultsListData);
+}
+
+bool SMergeGraphView::HasPrevDifference() const
+{
+	if (HasNoDifferences())
+	{
+		return false;
+	}
+
+	check(DiffResultList.Pin().IsValid() && DiffResultsListData);
+	return DiffWidgetUtils::HasPrevDifference(*DiffResultList.Pin(), *DiffResultsListData);
+}
+
+bool SMergeGraphView::HasNoDifferences() const
 {
 	auto DiffResultListPtr = DiffResultList.Pin();
-	return DiffResultListPtr.IsValid() && DiffResultListPtr->GetNumItemsBeingObserved() > 0;
+	return !DiffResultListPtr.IsValid() || DiffResultListPtr->GetNumItemsBeingObserved() == 0;
 }
 
 void SMergeGraphView::OnGraphListSelectionChanged(TSharedPtr<FMergeGraphRowEntry> Item, ESelectInfo::Type SelectionType)
diff --git a/Engine/Source/Developer/Merge/Private/SMergeGraphView.h b/Engine/Source/Developer/Merge/Private/SMergeGraphView.h
--- a/Engine/Source/Developer/Merge/Private/SMergeGraphView.h
+++ b/Engine/Source/Developer/Merge/Private/SMergeGraphView.h
@@ -19,9 +19,11 @@ private:
 	/** Implementation of IDiffControl: */
 	void NextDiff() override;
 	void PrevDiff() override;
-	bool HasDifferences() const override;
+	bool HasNextDifference() const override;
+	bool HasPrevDifference() const override;
 
 	/** Helper functions and event handlers: */
+	bool HasNoDifferences() const;
 	void OnGraphListSelectionChanged(TSharedPtr<struct FMergeGraphRowEntry> Item, ESelectInfo::Type SelectionType);
 	void OnDiffListSelectionChanged(TSharedPtr<struct FDiffSingleResult> Item, ESelectInfo::Type SelectionType);
 
diff --git a/Engine/Source/Developer/Merge/Private/SMergeTreeView.cpp b/Engine/Source/Developer/Merge/Private/SMergeTreeView.cpp
new file mode 100644
--- /dev/null
+++ b/Engine/Source/Developer/Merge/Private/SMergeTreeView.cpp
@@ -0,0 +1,133 @@
+// Copyright 1998-2014 Epic Games, Inc. All Rights Reserved.
+
+#include "MergePrivatePCH.h"
+
+#include "SMergeTreeView.h"
+
+void SMergeTreeView::Construct(const FArguments InArgs, const FBlueprintMergeData& InData)
+{
+	Data = InData;
+	CurrentDifference = -1;
+
+	TArray< FSCSDiffEntry > RemoteDifferingProperties;
+	TArray< int > RemoteDiffsSortKeys;
+	TArray< FSCSDiffEntry > LocalDifferingProperties;
+	TArray< int > LocalDiffsSortKeys;
+
+	DiffUtils::CompareUnrelatedSCS(InData.BlueprintBase, InData.BlueprintRemote, RemoteDifferingProperties, &RemoteDiffsSortKeys);
+	DiffUtils::CompareUnrelatedSCS(InData.BlueprintBase, InData.BlueprintLocal, LocalDifferingProperties, &LocalDiffsSortKeys);
+
+	check(RemoteDiffsSortKeys.Num() == RemoteDifferingProperties.Num() && LocalDifferingProperties.Num() == LocalDiffsSortKeys.Num());
+
+	// Order the differing properties, and eliminate dupes. Eventually dupes should be treated as collisions when
+	// attempting to complete a merge automatically:
+	int RemoteIter = 0;
+	int LocalIter = 0;
+	while (RemoteDifferingProperties.IsValidIndex(RemoteIter) || LocalDifferingProperties.IsValidIndex(LocalIter))
+	{
+		const bool bStillHaveLocalDiffs = LocalDifferingProperties.IsValidIndex(LocalIter);
+		const bool bStillHaveRemoteDiffs = RemoteDifferingProperties.IsValidIndex(RemoteIter);
+		if (bStillHaveLocalDiffs && bStillHaveRemoteDiffs)
+		{
+			if (LocalDiffsSortKeys[LocalIter] <= RemoteDiffsSortKeys[RemoteIter])
+			{
+				DifferingProperties.Push(LocalDifferingProperties[LocalIter++]);
+			}
+			else
+			{
+				DifferingProperties.Push(RemoteDifferingProperties[RemoteIter++]);
+			}
+		}
+		else if (bStillHaveLocalDiffs)
+		{
+			DifferingProperties.Push(LocalDifferingProperties[LocalIter++]);
+		}
+		else
+		{
+			check(bStillHaveRemoteDiffs);
+			DifferingProperties.Push(RemoteDifferingProperties[RemoteIter++]);
+		}
+	}
+
+	// generate controls:
+	// EMergeParticipant::MERGE_PARTICIPANT_REMOTE
+	{
+		SCSViews.Push(
+			FSCSDiff(InData.BlueprintRemote)
+			);
+	}
+	// EMergeParticipant::MERGE_PARTICIPANT_BASE
+	{
+		SCSViews.Push(
+			FSCSDiff(InData.BlueprintBase)
+			);
+	}
+	// EMergeParticipant::MERGE_PARTICIPANT_LOCAL
+	{
+		SCSViews.Push(
+			FSCSDiff(InData.BlueprintLocal)
+			);
+	}
+
+	ChildSlot[
+		SNew(SSplitter)
+			+ SSplitter::Slot()
+			[
+				GetRemoteView().TreeWidget()
+			]
+		+ SSplitter::Slot()
+			[
+				GetBaseView().TreeWidget()
+			]
+		+ SSplitter::Slot()
+			[
+				GetLocalView().TreeWidget()
+			]
+	];
+}
+
+void SMergeTreeView::NextDiff()
+{
+	++CurrentDifference;
+	HighlightCurrentDifference();
+}
+
+void SMergeTreeView::PrevDiff()
+{
+	--CurrentDifference;
+	HighlightCurrentDifference();
+}
+
+bool SMergeTreeView::HasNextDifference() const
+{
+	return DifferingProperties.IsValidIndex(CurrentDifference + 1);
+}
+
+bool SMergeTreeView::HasPrevDifference() const
+{
+	return DifferingProperties.IsValidIndex(CurrentDifference - 1);
+}
+
+void SMergeTreeView::HighlightCurrentDifference()
+{
+	for (auto& View : SCSViews)
+	{
+		View.HighlightProperty(DifferingProperties[CurrentDifference]);
+	}
+}
+
+FSCSDiff& SMergeTreeView::GetRemoteView()
+{
+	return SCSViews[EMergeParticipant::MERGE_PARTICIPANT_REMOTE];
+}
+
+FSCSDiff& SMergeTreeView::GetBaseView()
+{
+	return SCSViews[EMergeParticipant::MERGE_PARTICIPANT_BASE];
+}
+
+FSCSDiff& SMergeTreeView::GetLocalView()
+{
+	return SCSViews[EMergeParticipant::MERGE_PARTICIPANT_LOCAL];
+}
+
diff --git a/Engine/Source/Developer/Merge/Private/SMergeTreeView.h b/Engine/Source/Developer/Merge/Private/SMergeTreeView.h
new file mode 100644
--- /dev/null
+++ b/Engine/Source/Developer/Merge/Private/SMergeTreeView.h
@@ -0,0 +1,36 @@
+// Copyright 1998-2014 Epic Games, Inc. All Rights Reserved.
+#pragma once
+
+#include "BlueprintMergeData.h"
+#include "IDiffControl.h"
+#include "SCSDiff.h"
+
+class SMergeTreeView	: public SCompoundWidget
+						, public IDiffControl
+{
+public:
+	virtual ~SMergeTreeView() {}
+
+	SLATE_BEGIN_ARGS(SMergeTreeView)
+	{}
+	SLATE_END_ARGS()
+
+	void Construct(const FArguments InArgs, const FBlueprintMergeData& InData);
+private:
+	/** Implementation of IDiffControl: */
+	void NextDiff() override;
+	void PrevDiff() override;
+	bool HasNextDifference() const override;
+	bool HasPrevDifference() const override;
+
+	void HighlightCurrentDifference();
+	FSCSDiff& GetRemoteView();
+	FSCSDiff& GetBaseView();
+	FSCSDiff& GetLocalView();
+
+	FBlueprintMergeData Data;
+	TArray< FSCSDiff > SCSViews;
+
+	TArray< FSCSDiffEntry > DifferingProperties;
+	int CurrentDifference;
+};
diff --git a/Engine/Source/Editor/Kismet/Private/SBlueprintDiff.cpp b/Engine/Source/Editor/Kismet/Private/SBlueprintDiff.cpp
--- a/Engine/Source/Editor/Kismet/Private/SBlueprintDiff.cpp
+++ b/Engine/Source/Editor/Kismet/Private/SBlueprintDiff.cpp
@@ -19,14 +19,87 @@
 #include "SBlueprintDiff.h"
 #include "SGraphTitleBar.h"
 #include "SSCSEditor.h"
+#include "SCSDiff.h"
 #include "WorkflowOrientedApp/SModeWidget.h"
 
 #define LOCTEXT_NAMESPACE "SBlueprintDif"
 
 typedef TMap< FName, const UProperty* > FNamePropertyMap;
 
-class FCDODiffControl : public TSharedFromThis<FCDODiffControl>
-	, public IDiffControl
+// Each difference in the tree will either be a tree node that is added in one Blueprint 
+// or a tree node and an FName of a property that has been added or edited in one Blueprint
+class FSCSDiffControl	: public TSharedFromThis< FSCSDiffControl >
+						, public IDiffControl
+{
+public:
+	FSCSDiffControl(
+					const UBlueprint* InOldBlueprint
+					, const UBlueprint* InNewBlueprint
+					, TArray< FSCSDiffEntry > InDifferingProperties
+					);
+
+
+	TSharedRef<SWidget> OldTreeWidget() { return OldSCS.TreeWidget(); }
+	TSharedRef<SWidget> NewTreeWidget() { return NewSCS.TreeWidget(); }
+
+	virtual ~FSCSDiffControl() {}
+private:
+	void NextDiff() override;
+	void PrevDiff() override;
+	bool HasNextDifference() const override;
+	bool HasPrevDifference() const override;
+
+	void HighlightCurrentDifference();
+
+	TArray< FSCSDiffEntry > DifferingProperties;
+	int CurrentDifference;
+
+	FSCSDiff OldSCS;
+	FSCSDiff NewSCS;
+};
+
+FSCSDiffControl::FSCSDiffControl(
+		const UBlueprint* InOldBlueprint
+		, const UBlueprint* InNewBlueprint
+		, TArray< FSCSDiffEntry > InDifferingProperties
+	)
+	: DifferingProperties(InDifferingProperties)
+	, CurrentDifference(-1)
+	, OldSCS( InOldBlueprint )
+	, NewSCS( InNewBlueprint )
+{
+}
+
+void FSCSDiffControl::NextDiff()
+{
+	++CurrentDifference;
+	HighlightCurrentDifference();
+}
+
+void FSCSDiffControl::PrevDiff()
+{
+	--CurrentDifference;
+	HighlightCurrentDifference();
+}
+
+bool FSCSDiffControl::HasNextDifference() const
+{
+	return DifferingProperties.IsValidIndex(CurrentDifference + 1);
+}
+
+bool FSCSDiffControl::HasPrevDifference() const
+{
+	return DifferingProperties.IsValidIndex(CurrentDifference - 1);
+}
+
+void FSCSDiffControl::HighlightCurrentDifference()
+{
+	OldSCS.HighlightProperty(DifferingProperties[CurrentDifference]);
+	NewSCS.HighlightProperty(DifferingProperties[CurrentDifference]);
+}
+
+class FCDODiffControl	: public TSharedFromThis<FCDODiffControl>
+						, public IDiffControl
 {
 public:
 	FCDODiffControl( const UObject* InOldCDO
@@ -42,7 +115,8 @@ public:
 private:
 	void NextDiff() override;
 	void PrevDiff() override;
-	bool HasDifferences() const override;
+	bool HasNextDifference() const override;
+	bool HasPrevDifference() const override;
 
 	void HighlightCurrentDifference();
 	FDetailsDiff OldDetails;
@@ -67,35 +141,24 @@ FCDODiffControl::FCDODiffControl(
 
 void FCDODiffControl::NextDiff()
 {
-	if (DifferingProperties.Num() == 0)
-	{
-		return;
-	}
-
-	CurrentDifference = (CurrentDifference + 1) % DifferingProperties.Num();
-
+	++CurrentDifference;
 	HighlightCurrentDifference();
 }
 
 void FCDODiffControl::PrevDiff()
 {
-	if (DifferingProperties.Num() == 0)
-	{
-		return;
-	}
-
 	--CurrentDifference;
-	if (CurrentDifference < 0)
-	{
-		CurrentDifference = DifferingProperties.Num() - 1;
-	}
-
 	HighlightCurrentDifference();
 }
 
-bool FCDODiffControl::HasDifferences() const
+bool FCDODiffControl::HasNextDifference() const
 {
-	return DifferingProperties.Num() != 0;
+	return DifferingProperties.IsValidIndex(CurrentDifference + 1);
+}
+
+bool FCDODiffControl::HasPrevDifference() const
+{
+	return DifferingProperties.IsValidIndex(CurrentDifference - 1);
 }
 
 void FCDODiffControl::HighlightCurrentDifference()
@@ -175,10 +238,11 @@ private:
 	void PrevDiff() override;
 
 	/** Called when determining whether the prev/nextbuttons should be enabled */
-	bool HasDifferences() const override;
+	bool HasNextDifference() const override;
+	bool HasPrevDifference() const override;
 
 	/** Get Index of the current diff that is selected */
-	int32 GetCurrentDiffIndex();
+	int32 GetCurrentDiffIndex() const;
 
 	/* Called when a new row is being generated */
 	TSharedRef<ITableRow> OnGenerateRow(FSharedDiffOnGraph ParamItem, const TSharedRef<STableViewBase>& OwnerTable);
@@ -492,12 +556,28 @@ void FListItemGraphToDiff::PrevDiff()
 	DiffList->SetSelection(DiffListSource[Index]);
 }
 
-bool FListItemGraphToDiff::HasDifferences() const
+bool FListItemGraphToDiff::HasNextDifference() const
 {
-	return DiffList.IsValid() && DiffListSource.Num() > 0;
+	if (!DiffList.IsValid())
+	{
+		return false;
+	}
+	int32 Index = GetCurrentDiffIndex();
+	return DiffListSource.IsValidIndex(Index + 1);
 }
 
-int32 FListItemGraphToDiff::GetCurrentDiffIndex() 
+bool FListItemGraphToDiff::HasPrevDifference() const
+{
+	if (!DiffList.IsValid())
+	{
+		return false;
+	}
+
+	int32 Index = GetCurrentDiffIndex();
+	return DiffListSource.IsValidIndex(Index - 1);
+}
+
+int32 FListItemGraphToDiff::GetCurrentDiffIndex() const
 {
 	if ( DiffList.IsValid() )
 	{
@@ -505,7 +585,7 @@ int32 FListItemGraphToDiff::GetCurrentDiffIndex()
 		if(Selected.Num() == 1)
 		{	
 			int32 Index = 0;
-			for(auto It(DiffListSource.CreateIterator());It;++It,Index++)
+			for(auto It(DiffListSource.CreateConstIterator());It;++It,Index++)
 			{
 				if(*It == Selected[0])
 				{
@@ -514,7 +594,7 @@ int32 FListItemGraphToDiff::GetCurrentDiffIndex()
 			}
 		}
 	}
-	return 0;
+	return -1;
 }
 
 void FListItemGraphToDiff::OnGraphChanged( const FEdGraphEditAction& Action )
@@ -536,7 +616,7 @@ FDiffPanel::FDiffPanel()
 	LastFocusedPin = NULL;
 }
 
-void DiffUtils::CompareUnrelatedObjects( UObject const* A, const TMap< FName, const UProperty* >& PropertyMapA, UObject const* B, const TMap< FName, const UProperty* >& PropertyMapB, TArray<FName> &OutIdenticalProperties, TArray<FName> &OutDifferingProperties )
+void DiffUtils::CompareUnrelatedObjects( const UObject* A, const TMap< FName, const UProperty* >& PropertyMapA, UObject const* B, const TMap< FName, const UProperty* >& PropertyMapB, TArray<FName> &OutIdenticalProperties, TArray<FName> &OutDifferingProperties )
 {
 	TSet<FName> PropertiesInA, PropertiesInB;
 	if( A )
@@ -562,8 +642,14 @@ void DiffUtils::CompareUnrelatedObjects( UObject const* A, const TMap< FName, co
 		{
 			const UProperty* const* AProp = PropertyMapA.Find(PropertyName);
 			const UProperty* const* BProp = PropertyMapB.Find(PropertyName);
-			check( AProp && BProp );
+			
+			if( !AProp && !BProp )
+			{
+				OutIdenticalProperties.Add(PropertyName);
+				continue;
+			}
 
+			check( AProp && BProp );
 			if( (*AProp)->SameType(*BProp) )
 			{
 				const void* AValue = (*AProp)->ContainerPtrToValuePtr<void>(A);
@@ -581,6 +667,81 @@ void DiffUtils::CompareUnrelatedObjects( UObject const* A, const TMap< FName, co
 	}
 }
 
+void DiffUtils::CompareUnrelatedSCS(const UBlueprint* A, const UBlueprint* B, TArray< FSCSDiffEntry >& OutDifferingEntries, TArray< int >* OptionalOutSortKeys )
+{
+	// we get bitten by const shallowness here, avoid mutation:
+	const TArray<USCS_Node*> NodesInA = A && A->SimpleConstructionScript ? A->SimpleConstructionScript->GetAllNodes() : TArray<USCS_Node*>();
+	const TArray<USCS_Node*> NodesInB = B && B->SimpleConstructionScript ? B->SimpleConstructionScript->GetAllNodes() : TArray<USCS_Node*>();
+
+	const auto FindByName = []( const TArray<USCS_Node*>& NodeList, FName Name ) -> USCS_Node*
+	{
+		for( auto Node : NodeList )
+		{
+			if( Node->VariableName == Name )
+			{
+				return Node;
+			}
+		}
+		return nullptr;
+	};
+
+	int32 SortKey = 0;
+	for( auto NodeInA : NodesInA )
+	{
+		auto NodeInB = FindByName( NodesInB, NodeInA->VariableName );
+		if( NodeInB )
+		{
+			// diff NodeA vs NodeB:
+			FNamePropertyMap PropertyMapA = DiffUtils::GetProperties(NodeInA->ComponentTemplate);
+			FNamePropertyMap PropertyMapB = DiffUtils::GetProperties(NodeInB->ComponentTemplate);
+			TArray<FName> IdenticalProperties, DifferingProperties;
+
+			DiffUtils::CompareUnrelatedObjects(A, PropertyMapA, B, PropertyMapB, IdenticalProperties, DifferingProperties);
+			for( auto Property : DifferingProperties )
+			{
+				FSCSDiffEntry Diff = { NodeInA->VariableName, Property };
+				OutDifferingEntries.Push( Diff );
+
+				if (OptionalOutSortKeys)
+				{
+					OptionalOutSortKeys->Push(SortKey);
+				}
+			}
+		}
+		else
+		{
+			// Node A was added:
+			FSCSDiffEntry Diff = { NodeInA->VariableName, FName() };
+			OutDifferingEntries.Push( Diff );
+
+			if (OptionalOutSortKeys)
+			{
+				OptionalOutSortKeys->Push(SortKey);
+			}
+		}
+
+		++SortKey;
+	}
+
+	// add nodes that were added in B:
+	for( auto NodeInB : NodesInB )
+	{
+		if( FindByName( NodesInA, NodeInB->VariableName ) == nullptr )
+		{
+			// Node B was added:
+			FSCSDiffEntry Diff = { NodeInB->VariableName, FName() };
+			OutDifferingEntries.Push(Diff);
+
+			if (OptionalOutSortKeys)
+			{
+				OptionalOutSortKeys->Push(SortKey);
+			}
+		}
+
+		++SortKey;
+	}
+}
+
 int32 GetCurrentIndex( SListView< TSharedPtr< struct FDiffSingleResult> > const& ListView, const TArray< TSharedPtr< FDiffSingleResult > >& ListViewSource )
 {
 	const auto& Selected = ListView.GetSelectedItems();
@@ -612,7 +773,7 @@ void DiffWidgetUtils::SelectNextRow( SListView< TSharedPtr< FDiffSingleResult> >
 void DiffWidgetUtils::SelectPrevRow(SListView< TSharedPtr< FDiffSingleResult> >& ListView, const TArray< TSharedPtr< FDiffSingleResult > >& ListViewSource )
 {
 	int32 CurrentIndex = GetCurrentIndex(ListView, ListViewSource);
-	if (CurrentIndex > 0)
+	if (CurrentIndex == 0)
 	{
 		return;
 	}
@@ -620,23 +781,16 @@ void DiffWidgetUtils::SelectPrevRow(SListView< TSharedPtr< FDiffSingleResult> >&
 	ListView.SetSelection(ListViewSource[CurrentIndex - 1]);
 }
 
-static TSharedRef<SWidget> GenerateComponentsDiff( UBlueprint const* BaseBlueprint, UBlueprint const* DisplayedBlueprint )
+bool DiffWidgetUtils::HasNextDifference(SListView< TSharedPtr< struct FDiffSingleResult> >& ListView, const TArray< TSharedPtr< struct FDiffSingleResult > >& ListViewSource)
 {
-	TSharedRef<SKismetInspector> Inspector = SNew(SKismetInspector)
-		.HideNameArea(true)
-		.ViewIdentifier(FName("BlueprintInspector"))
-		.IsPropertyEditingEnabledDelegate(FIsPropertyEditingEnabled::CreateStatic([] { return false; }));
+	int32 CurrentIndex = GetCurrentIndex(ListView, ListViewSource);
+	return ListViewSource.IsValidIndex(CurrentIndex+1);
+}
 
-	return SNew(SVerticalBox)
-		+ SVerticalBox::Slot()
-		.AutoHeight()
-		[
-			SNew(SSCSEditor, TSharedPtr<FBlueprintEditor>(), DisplayedBlueprint->SimpleConstructionScript, const_cast<UBlueprint*>(DisplayedBlueprint), Inspector )
-		]
-		+ SVerticalBox::Slot()
-		[
-			Inspector
-		];
+bool DiffWidgetUtils::HasPrevDifference(SListView< TSharedPtr< struct FDiffSingleResult> >& ListView, const TArray< TSharedPtr< struct FDiffSingleResult > >& ListViewSource)
+{
+	int32 CurrentIndex = GetCurrentIndex(ListView, ListViewSource);
+	return ListViewSource.IsValidIndex(CurrentIndex - 1);
 }
 
 //////////////////////////////////////////////////////////////////////////
@@ -712,8 +866,8 @@ void SBlueprintDiff::Construct( const FArguments& InArgs)
 	const FDiffListCommands& Commands = FDiffListCommands::Get();
 	KeyCommands = MakeShareable(new FUICommandList);
 
-	KeyCommands->MapAction(Commands.Previous, FExecuteAction::CreateSP(this, &SBlueprintDiff::PrevDiff), FCanExecuteAction::CreateSP( this, &SBlueprintDiff::CanCycleDiffs) );
-	KeyCommands->MapAction(Commands.Next, FExecuteAction::CreateSP(this, &SBlueprintDiff::NextDiff), FCanExecuteAction::CreateSP( this, &SBlueprintDiff::CanCycleDiffs) );
+	KeyCommands->MapAction(Commands.Previous, FExecuteAction::CreateSP(this, &SBlueprintDiff::PrevDiff), FCanExecuteAction::CreateSP( this, &SBlueprintDiff::HasPrevDiff) );
+	KeyCommands->MapAction(Commands.Next, FExecuteAction::CreateSP(this, &SBlueprintDiff::NextDiff), FCanExecuteAction::CreateSP( this, &SBlueprintDiff::HasNextDiff) );
 
 	FToolBarBuilder ToolbarBuilder(KeyCommands.ToSharedRef(), FMultiBoxCustomization::None);
 	ToolbarBuilder.AddToolBarButton(Commands.Previous, NAME_None, TAttribute<FText>(), TAttribute<FText>(), FSlateIcon(FEditorStyle::GetStyleSetName(), "BlueprintDif.PrevDiff"));
@@ -831,9 +985,14 @@ void SBlueprintDiff::PrevDiff()
 	}
 }
 
-bool SBlueprintDiff::CanCycleDiffs() const
+bool SBlueprintDiff::HasNextDiff() const
 {
-	return DiffControl.IsValid() && DiffControl->HasDifferences();
+	return DiffControl.IsValid() && DiffControl->HasNextDifference();
+}
+
+bool SBlueprintDiff::HasPrevDiff() const
+{
+	return DiffControl.IsValid() && DiffControl->HasPrevDifference();
 }
 
 void SBlueprintDiff::FocusOnGraphRevisions( class UEdGraph* GraphOld, class UEdGraph* GraphNew, FListItemGraphToDiff* Diff )
@@ -1247,6 +1406,11 @@ TSharedRef<SWidget> SBlueprintDiff::GenerateDefaultsPanel()
 TSharedRef<SWidget> SBlueprintDiff::GenerateComponentsPanel()
 {
 	//Splitter for left and right blueprint. Current convention is for the local (probably newer?) blueprint to be on the right:
+	TArray< FSCSDiffEntry > Differences;
+	DiffUtils::CompareUnrelatedSCS( PanelOld.Blueprint, PanelNew.Blueprint, Differences );
+
+	auto NewDiffControl = TSharedPtr<FSCSDiffControl>(new FSCSDiffControl(PanelOld.Blueprint, PanelNew.Blueprint, Differences));
+	DiffControl = NewDiffControl;
 	return SNew(SSplitter)
 		+ SSplitter::Slot()
 		.Value(0.5f)
@@ -1254,7 +1418,7 @@ TSharedRef<SWidget> SBlueprintDiff::GenerateComponentsPanel()
 			SNew(SBorder)
 			.VAlign(VAlign_Fill)
 			[
-				GenerateComponentsDiff(PanelNew.Blueprint, PanelOld.Blueprint)
+				NewDiffControl->OldTreeWidget()
 			]
 		]
 		+ SSplitter::Slot()
@@ -1263,7 +1427,7 @@ TSharedRef<SWidget> SBlueprintDiff::GenerateComponentsPanel()
 			SNew(SBorder)
 			.VAlign(VAlign_Fill)
 			[
-				GenerateComponentsDiff(PanelOld.Blueprint, PanelNew.Blueprint)
+				NewDiffControl->NewTreeWidget()
 			]
 		];
 }
diff --git a/Engine/Source/Editor/Kismet/Private/SCSDiff.cpp b/Engine/Source/Editor/Kismet/Private/SCSDiff.cpp
new file mode 100644
--- /dev/null
+++ b/Engine/Source/Editor/Kismet/Private/SCSDiff.cpp
@@ -0,0 +1,50 @@
+
+#include "BlueprintEditorPrivatePCH.h"
+#include "SCSDiff.h"
+#include "SKismetInspector.h"
+#include "SSCSEditor.h"
+
+#include <vector>
+
+FSCSDiff::FSCSDiff(const UBlueprint* InBlueprint)
+{
+	TSharedRef<SKismetInspector> Inspector = SNew(SKismetInspector)
+		.HideNameArea(true)
+		.ViewIdentifier(FName("BlueprintInspector"))
+		.IsPropertyEditingEnabledDelegate(FIsPropertyEditingEnabled::CreateStatic([] { return false; }));
+
+	ContainerWidget = SNew(SVerticalBox)
+		+ SVerticalBox::Slot()
+		.AutoHeight()
+		[
+			SAssignNew(SCSEditor, SSCSEditor, TSharedPtr<FBlueprintEditor>(), InBlueprint->SimpleConstructionScript, const_cast<UBlueprint*>(InBlueprint), Inspector)
+		]
+	+ SVerticalBox::Slot()
+		[
+			Inspector
+		];
+}
+
+void FSCSDiff::HighlightProperty(FSCSDiffEntry Property)
+{
+	check( Property.TreeNodeName != FName() );
+	if( const USimpleConstructionScript* SCS = SCSEditor->SCS )
+	{
+		TArray<USCS_Node*> Nodes = SCS->GetAllNodes();
+		for( auto Node : Nodes )
+		{
+			if( Node->VariableName == Property.TreeNodeName )
+			{
+				SCSEditor->HighlightTreeNode( Node, Property.PropertyName );
+				return;
+			}
+		}
+	}
+
+	SCSEditor->ClearSelection();
+}
+
+TSharedRef< SWidget > FSCSDiff::TreeWidget()
+{
+	return ContainerWidget.ToSharedRef();
+}
diff --git a/Engine/Source/Editor/Kismet/Private/SSCSEditor.cpp b/Engine/Source/Editor/Kismet/Private/SSCSEditor.cpp
--- a/Engine/Source/Editor/Kismet/Private/SSCSEditor.cpp
+++ b/Engine/Source/Editor/Kismet/Private/SSCSEditor.cpp
@@ -2,6 +2,7 @@
 
 
 #include "BlueprintEditorPrivatePCH.h"
+#include "Editor/PropertyEditor/Public/IDetailsView.h"
 #include "Editor/UnrealEd/Public/Kismet2/BlueprintEditorUtils.h"
 #include "BlueprintUtilities.h"
 #include "ComponentAssetBroker.h"
@@ -2471,6 +2472,23 @@ void SSCSEditor::SelectNode(FSCSEditorTreeNodePtrType InNodeToSelect, bool IsCnt
 	}
 }
 
+void SSCSEditor::HighlightTreeNode(const USCS_Node* Node, FName Property)
+{
+	check(Node);
+	auto TreeNode = FindTreeNode( Node );
+	check( TreeNode.IsValid() );
+	SelectNode( TreeNode, false );
+	if( Property != FName() )
+	{
+		auto KismetInspectorSPtr = KismetInspectorPtr.Pin();
+		check( KismetInspectorSPtr.IsValid() );
+		UActorComponent* Component = TreeNode->GetComponentTemplate();
+		UProperty* CurrentProp = FindField<UProperty>(Component->GetClass(), Property);
+		check( CurrentProp );
+		KismetInspectorSPtr->GetPropertyView()->HighlightProperty( CurrentProp );
+	}
+}
+
 void SSCSEditor::UpdateTree(bool bRegenerateTreeNodes)
 {
 	check(SCSTreeWidget.IsValid());
diff --git a/Engine/Source/Editor/Kismet/Private/SSCSEditor.h b/Engine/Source/Editor/Kismet/Private/SSCSEditor.h
--- a/Engine/Source/Editor/Kismet/Private/SSCSEditor.h
+++ b/Engine/Source/Editor/Kismet/Private/SSCSEditor.h
@@ -467,6 +467,14 @@ public:
 
 	/** Select the given tree node */
 	void SelectNode(FSCSEditorTreeNodePtrType InNodeToSelect, bool IsCntrlDown);
+	
+	/**
+	 * Highlight a tree node and, optionally, a property with in it
+	 *
+	 * @param Node		 A Reference to the Node SCS_Node to be highlighted
+	 * @param Property	The name of the property to be highlighted in the details view
+	 */
+	void HighlightTreeNode( const USCS_Node* Node, FName Property );
 
 	/**
 	 * Function to save current state of SimpleConstructionScript and nodes associated with it.
diff --git a/Engine/Source/Editor/Kismet/Public/IDiffControl.h b/Engine/Source/Editor/Kismet/Public/IDiffControl.h
--- a/Engine/Source/Editor/Kismet/Public/IDiffControl.h
+++ b/Engine/Source/Editor/Kismet/Public/IDiffControl.h
@@ -6,6 +6,7 @@ class IDiffControl
 public:
 	virtual void NextDiff() = 0;
 	virtual void PrevDiff() = 0;
-	virtual bool HasDifferences() const = 0;
+	virtual bool HasNextDifference() const = 0;
+	virtual bool HasPrevDifference() const = 0;
 };
 
diff --git a/Engine/Source/Editor/Kismet/Public/SBlueprintDiff.h b/Engine/Source/Editor/Kismet/Public/SBlueprintDiff.h
--- a/Engine/Source/Editor/Kismet/Public/SBlueprintDiff.h
+++ b/Engine/Source/Editor/Kismet/Public/SBlueprintDiff.h
@@ -30,17 +30,26 @@ struct FDiffResultItem : public TSharedFromThis<FDiffResultItem>
 	TSharedRef<SWidget> KISMET_API GenerateWidget() const;
 };
 
+struct FSCSDiffEntry
+{
+	FName TreeNodeName;
+	FName PropertyName;
+};
+
 namespace DiffUtils
 {
 	KISMET_API TMap< FName, const UProperty* > GetProperties( const UObject* ForObj );
 	KISMET_API const UObject* GetCDO( const UBlueprint* ForBlueprint );
-	KISMET_API void CompareUnrelatedObjects( UObject const* A, const TMap< FName, const UProperty* >& PropertyMapA, UObject const* B, const TMap< FName, const UProperty* >& PropertyMapB, TArray<FName> &OutIdenticalProperties, TArray<FName> &OutDifferingProperties );
+	KISMET_API void CompareUnrelatedObjects( const UObject* A, const TMap< FName, const UProperty* >& PropertyMapA, const UObject* B, const TMap< FName, const UProperty* >& PropertyMapB, TArray<FName> &OutIdenticalProperties, TArray<FName> &OutDifferingProperties );
+	KISMET_API void CompareUnrelatedSCS(const UBlueprint* Old, const UBlueprint* New, TArray< struct FSCSDiffEntry >& OutDifferingEntries, TArray< int >* OptionalOutSortKeys = nullptr );
 }
 
 namespace DiffWidgetUtils
 {
 	KISMET_API void SelectNextRow(SListView< TSharedPtr< struct FDiffSingleResult> >& ListView, const TArray< TSharedPtr< struct FDiffSingleResult > >& ListViewSource );
 	KISMET_API void SelectPrevRow(SListView< TSharedPtr< struct FDiffSingleResult> >& ListView, const TArray< TSharedPtr< struct FDiffSingleResult > >& ListViewSource );
+	KISMET_API bool HasNextDifference(SListView< TSharedPtr< struct FDiffSingleResult> >& ListView, const TArray< TSharedPtr< struct FDiffSingleResult > >& ListViewSource);
+	KISMET_API bool HasPrevDifference(SListView< TSharedPtr< struct FDiffSingleResult> >& ListView, const TArray< TSharedPtr< struct FDiffSingleResult > >& ListViewSource);
 }
 
 /*panel used to display the blueprint*/
@@ -120,7 +129,8 @@ protected:
 	void PrevDiff();
 
 	/** Called to determine whether we have a list of differences to cycle through */
-	bool CanCycleDiffs() const;
+	bool HasNextDiff() const;
+	bool HasPrevDiff() const;
 
 	/* Need to process keys for shortcuts to buttons */
 	virtual FReply OnKeyDown( const FGeometry& MyGeometry, const FKeyboardEvent& InKeyboardEvent ) override;
diff --git a/Engine/Source/Editor/Kismet/Public/SCSDiff.h b/Engine/Source/Editor/Kismet/Public/SCSDiff.h
new file mode 100644
--- /dev/null
+++ b/Engine/Source/Editor/Kismet/Public/SCSDiff.h
@@ -0,0 +1,16 @@
+// Copyright 1998-2014 Epic Games, Inc. All Rights Reserved.
+#pragma once
+
+#include "SBlueprintDiff.h"
+
+class KISMET_API FSCSDiff
+{
+public:
+	FSCSDiff(const class UBlueprint* InBlueprint);
+
+	void HighlightProperty(FSCSDiffEntry Property);
+	TSharedRef< SWidget > TreeWidget();
+private:
+	TSharedPtr< class SWidget > ContainerWidget;
+	TSharedPtr< class SSCSEditor > SCSEditor;
+};
