Missed files needed for new BP menu system optimizations.

diff --git a/Engine/Source/Editor/BlueprintGraph/Private/BlueprintNodeSpawnerUtils.cpp b/Engine/Source/Editor/BlueprintGraph/Private/BlueprintNodeSpawnerUtils.cpp
--- a/Engine/Source/Editor/BlueprintGraph/Private/BlueprintNodeSpawnerUtils.cpp
+++ b/Engine/Source/Editor/BlueprintGraph/Private/BlueprintNodeSpawnerUtils.cpp
@@ -9,33 +9,6 @@
 #include "BlueprintComponentNodeSpawner.h"
 #include "BlueprintBoundEventNodeSpawner.h"
 #include "BlueprintDelegateNodeSpawner.h"
- 
-//------------------------------------------------------------------------------
-UClass* FBlueprintNodeSpawnerUtils::GetAssociatedClass(UBlueprintNodeSpawner const* BlueprintAction)
-{
-	UClass* ActionClass = BlueprintAction->NodeClass;
-	
-	if (UField const* AssociatedField = GetAssociatedField(BlueprintAction))
-	{
-		ActionClass = AssociatedField->GetOwnerClass();
-	}
-	else if (UBlueprintComponentNodeSpawner const* ComponentSpawner = Cast<UBlueprintComponentNodeSpawner>(BlueprintAction))
-	{
-		ActionClass = ComponentSpawner->GetComponentClass();
-	}
-	else if (UBlueprintVariableNodeSpawner const* VarSpawner = Cast<UBlueprintVariableNodeSpawner>(BlueprintAction))
-	{
-		UObject const* VarOuter = VarSpawner->GetVarOuter();
-		if (VarSpawner->IsLocalVariable())
-		{
-			UBlueprint* Blueprint = FBlueprintEditorUtils::FindBlueprintForGraphChecked(CastChecked<UEdGraph>(VarOuter));
-			ActionClass = (Blueprint->SkeletonGeneratedClass != nullptr) ? Blueprint->SkeletonGeneratedClass : Blueprint->ParentClass;
-		}
-		// else, should have been caught in GetAssociatedField()
-	}
-	
-	return ActionClass;
-}
 
 //------------------------------------------------------------------------------
 UField const* FBlueprintNodeSpawnerUtils::GetAssociatedField(UBlueprintNodeSpawner const* BlueprintAction)
diff --git a/Engine/Source/Editor/BlueprintGraph/Private/BlueprintNodeSpawnerUtils.h b/Engine/Source/Editor/BlueprintGraph/Private/BlueprintNodeSpawnerUtils.h
--- a/Engine/Source/Editor/BlueprintGraph/Private/BlueprintNodeSpawnerUtils.h
+++ b/Engine/Source/Editor/BlueprintGraph/Private/BlueprintNodeSpawnerUtils.h
@@ -8,16 +8,6 @@ class UBlueprintNodeSpawner;
 
 struct FBlueprintNodeSpawnerUtils
 {
-	/**
-	 * All actions are associated with some class, whether it's the class of the
-	 * node it will spawn, or the outer of the associated field (property,
-	 * function, etc.).
-	 * 
-	 * @param  BlueprintAction	The node-spawner you want an associated class for.
-	 * @return The action's associated class (null if it doesn't have one).
-	 */
-	static UClass* GetAssociatedClass(UBlueprintNodeSpawner const* BlueprintAction);
-
 	/**
 	 * Certain node-spawners are associated with specific UFields (functions, 
 	 * properties, etc.). This attempts to retrieve the field from the spawner.
diff --git a/Engine/Source/Editor/Kismet/Private/BlueprintActionMenuUtils.cpp b/Engine/Source/Editor/Kismet/Private/BlueprintActionMenuUtils.cpp
--- a/Engine/Source/Editor/Kismet/Private/BlueprintActionMenuUtils.cpp
+++ b/Engine/Source/Editor/Kismet/Private/BlueprintActionMenuUtils.cpp
@@ -28,7 +28,7 @@ namespace BlueprintActionMenuUtilsImpl
 	 * @param  BlueprintAction	
 	 * @return 
 	 */
-	static bool IsUnBoundSpawner(FBlueprintActionFilter const& Filter, UBlueprintNodeSpawner const* BlueprintAction);
+	static bool IsUnBoundSpawner(FBlueprintActionFilter const& Filter, FBlueprintActionInfo& BlueprintAction);
 
 	/**
 	 * 
@@ -36,17 +36,17 @@ namespace BlueprintActionMenuUtilsImpl
 	 * @param  ClassSet	
 	 * @return 
 	 */
-	static UClass* FinCommonBaseClass(TArray<UClass*> const& ClassSet);
+	static UClass* FindCommonBaseClass(TArray<UClass*> const& ClassSet);
 }
 
 //------------------------------------------------------------------------------
-static bool BlueprintActionMenuUtilsImpl::IsUnBoundSpawner(FBlueprintActionFilter const& /*Filter*/, UBlueprintNodeSpawner const* BlueprintAction)
+static bool BlueprintActionMenuUtilsImpl::IsUnBoundSpawner(FBlueprintActionFilter const& /*Filter*/, FBlueprintActionInfo& BlueprintAction)
 {
-	return !BlueprintAction->IsBindingSet();
+	return !BlueprintAction.NodeSpawner->IsBindingSet();
 }
 
 //------------------------------------------------------------------------------
-static UClass* BlueprintActionMenuUtilsImpl::FinCommonBaseClass(TArray<UClass*> const& ClassSet)
+static UClass* BlueprintActionMenuUtilsImpl::FindCommonBaseClass(TArray<UClass*> const& ClassSet)
 {
 	UClass* CommonClass = UObject::StaticClass();
 	if (ClassSet.Num() > 0)
@@ -114,14 +114,14 @@ void FBlueprintActionMenuUtils::MakeContextMenu(FBlueprintActionContext const& C
 	ComponentsFilter.PermittedNodeTypes.Add(UK2Node_CallFunction::StaticClass());
 	ComponentsFilter.PermittedNodeTypes.Add(UK2Node_ComponentBoundEvent::StaticClass());
 	// only want bound actions for this menu section
-	ComponentsFilter.AddIsFilteredTest(FBlueprintActionFilter::FIsFilteredDelegate::CreateStatic(IsUnBoundSpawner));
+	ComponentsFilter.AddRejectionTest(FBlueprintActionFilter::FRejectionTestDelegate::CreateStatic(IsUnBoundSpawner));
 
 	FBlueprintActionFilter LevelActorsFilter;
 	LevelActorsFilter.Context = Context;
 	LevelActorsFilter.PermittedNodeTypes.Add(UK2Node_CallFunction::StaticClass());
 	LevelActorsFilter.PermittedNodeTypes.Add(UK2Node_ActorBoundEvent::StaticClass());
 	// only want bound actions for this menu section
-	LevelActorsFilter.AddIsFilteredTest(FBlueprintActionFilter::FIsFilteredDelegate::CreateStatic(IsUnBoundSpawner));
+	LevelActorsFilter.AddRejectionTest(FBlueprintActionFilter::FRejectionTestDelegate::CreateStatic(IsUnBoundSpawner));
 
 	// make sure the bound menu sections have the proper OwnerClasses specified
 	for (UObject* Selection : Context.SelectedObjects)
@@ -196,7 +196,7 @@ void FBlueprintActionMenuUtils::MakeContextMenu(FBlueprintActionContext const& C
 	{
 		// since we're consolidating all the bound actions, then we have to pick 
 		// one common base class to filter by
-		UClass* CommonClass = FinCommonBaseClass(LevelActorsFilter.TargetClasses);
+		UClass* CommonClass = FindCommonBaseClass(LevelActorsFilter.TargetClasses);
 		LevelActorsFilter.TargetClasses.Empty();
 		LevelActorsFilter.TargetClasses.Add(CommonClass);
 
