Deprecating EdGraphNode's GetTooltip() (in exchange for one that returns FText over FString).

diff --git a/Engine/Source/Editor/BlueprintGraph/Private/K2Node_SelectEnum.cpp b/Engine/Source/Editor/BlueprintGraph/Private/K2Node_SelectEnum.cpp
--- a/Engine/Source/Editor/BlueprintGraph/Private/K2Node_SelectEnum.cpp
+++ b/Engine/Source/Editor/BlueprintGraph/Private/K2Node_SelectEnum.cpp
@@ -41,9 +41,9 @@ void UDEPRECATED_K2Node_SelectEnum::AllocateDefaultPins()
 	UK2Node::AllocateDefaultPins();
 }
 
-FString UDEPRECATED_K2Node_SelectEnum::GetTooltip() const
+FText UDEPRECATED_K2Node_SelectEnum::GetTooltipText() const
 {
-	return TEXT("Return literal value set for each enum value");
+	return NSLOCTEXT("K2Node", "SelectEnumTooltip", "Return literal value set for each enum value");
 }
 
 FText  UDEPRECATED_K2Node_SelectEnum::GetNodeTitle(ENodeTitleType::Type TitleType) const
diff --git a/Engine/Source/Editor/BlueprintGraph/Private/K2Node_Self.cpp b/Engine/Source/Editor/BlueprintGraph/Private/K2Node_Self.cpp
--- a/Engine/Source/Editor/BlueprintGraph/Private/K2Node_Self.cpp
+++ b/Engine/Source/Editor/BlueprintGraph/Private/K2Node_Self.cpp
@@ -45,9 +45,9 @@ void UK2Node_Self::AllocateDefaultPins()
 	Super::AllocateDefaultPins();
 }
 
-FString UK2Node_Self::GetTooltip() const
+FText UK2Node_Self::GetTooltipText() const
 {
-	return FString::Printf(*NSLOCTEXT("K2Node", "GetSelfReference", "Gets a reference to this instance of the blueprint").ToString());
+	return NSLOCTEXT("K2Node", "GetSelfReference", "Gets a reference to this instance of the blueprint");
 }
 
 FString UK2Node_Self::GetKeywords() const
diff --git a/Engine/Source/Editor/BlueprintGraph/Private/K2Node_SetFieldsInStruct.cpp b/Engine/Source/Editor/BlueprintGraph/Private/K2Node_SetFieldsInStruct.cpp
--- a/Engine/Source/Editor/BlueprintGraph/Private/K2Node_SetFieldsInStruct.cpp
+++ b/Engine/Source/Editor/BlueprintGraph/Private/K2Node_SetFieldsInStruct.cpp
@@ -65,12 +65,12 @@ FText UK2Node_SetFieldsInStruct::GetNodeTitle(ENodeTitleType::Type TitleType) co
 	return CachedNodeTitle;
 }
 
-FString UK2Node_SetFieldsInStruct::GetTooltip() const
+FText UK2Node_SetFieldsInStruct::GetTooltipText() const
 {
-	return FString::Printf(
-		*LOCTEXT("SetFieldsInStruct_Tooltip", "Adds a node that modifies a '%s'").ToString(),
-		*(StructType ? StructType->GetName() : FString())
-		);
+	return FText::Format(
+		LOCTEXT("SetFieldsInStruct_Tooltip", "Adds a node that modifies a '{0}'"),
+		StructType ? FText::FromName(StructType->GetFName()) : FText::GetEmpty()
+	);
 }
 
 FName UK2Node_SetFieldsInStruct::GetPaletteIcon(FLinearColor& OutColor) const
diff --git a/Engine/Source/Editor/BlueprintGraph/Private/K2Node_SpawnActor.cpp b/Engine/Source/Editor/BlueprintGraph/Private/K2Node_SpawnActor.cpp
--- a/Engine/Source/Editor/BlueprintGraph/Private/K2Node_SpawnActor.cpp
+++ b/Engine/Source/Editor/BlueprintGraph/Private/K2Node_SpawnActor.cpp
@@ -14,7 +14,7 @@ static FString NoCollisionFailPinName(TEXT("SpawnEvenIfColliding"));
 UK2Node_SpawnActor::UK2Node_SpawnActor(const class FPostConstructInitializeProperties& PCIP)
 	: Super(PCIP)
 {
-	NodeTooltip = LOCTEXT("NodeTooltip", "Attempts to spawn a new Actor with the specified transform").ToString();
+	NodeTooltip = LOCTEXT("NodeTooltip", "Attempts to spawn a new Actor with the specified transform");
 }
 
 void UK2Node_SpawnActor::AllocateDefaultPins()
@@ -157,7 +157,7 @@ void UK2Node_SpawnActor::PinDefaultValueChanged(UEdGraphPin* ChangedPin)
 	}
 }
 
-FString UK2Node_SpawnActor::GetTooltip() const
+FText UK2Node_SpawnActor::GetTooltipText() const
 {
 	return NodeTooltip;
 }
diff --git a/Engine/Source/Editor/BlueprintGraph/Private/K2Node_SpawnActorFromClass.cpp b/Engine/Source/Editor/BlueprintGraph/Private/K2Node_SpawnActorFromClass.cpp
--- a/Engine/Source/Editor/BlueprintGraph/Private/K2Node_SpawnActorFromClass.cpp
+++ b/Engine/Source/Editor/BlueprintGraph/Private/K2Node_SpawnActorFromClass.cpp
@@ -25,7 +25,7 @@ FString FK2Node_SpawnActorFromClassHelper::NoCollisionFailPinName(TEXT("SpawnEve
 UK2Node_SpawnActorFromClass::UK2Node_SpawnActorFromClass(const class FPostConstructInitializeProperties& PCIP)
 	: Super(PCIP)
 {
-	NodeTooltip = LOCTEXT("NodeTooltip", "Attempts to spawn a new Actor with the specified transform").ToString();
+	NodeTooltip = LOCTEXT("NodeTooltip", "Attempts to spawn a new Actor with the specified transform");
 }
 
 void UK2Node_SpawnActorFromClass::AllocateDefaultPins()
@@ -202,7 +202,7 @@ void UK2Node_SpawnActorFromClass::PinDefaultValueChanged(UEdGraphPin* ChangedPin
 	}
 }
 
-FString UK2Node_SpawnActorFromClass::GetTooltip() const
+FText UK2Node_SpawnActorFromClass::GetTooltipText() const
 {
 	return NodeTooltip;
 }
diff --git a/Engine/Source/Editor/BlueprintGraph/Private/K2Node_StructMemberGet.cpp b/Engine/Source/Editor/BlueprintGraph/Private/K2Node_StructMemberGet.cpp
--- a/Engine/Source/Editor/BlueprintGraph/Private/K2Node_StructMemberGet.cpp
+++ b/Engine/Source/Editor/BlueprintGraph/Private/K2Node_StructMemberGet.cpp
@@ -69,11 +69,11 @@ void UK2Node_StructMemberGet::AllocatePinsForSingleMemberGet(FName MemberName)
 	}
 }
 
-FString UK2Node_StructMemberGet::GetTooltip() const
+FText UK2Node_StructMemberGet::GetTooltipText() const
 {
 	FFormatNamedArguments Args;
 	Args.Add(TEXT("VariableName"), FText::FromString(GetVarNameString()));
-	return FText::Format(LOCTEXT("K2Node_StructMemberGet_Tooltip", "Get member variables of {VariableName}"), Args).ToString();
+	return FText::Format(LOCTEXT("K2Node_StructMemberGet_Tooltip", "Get member variables of {VariableName}"), Args);
 }
 
 FText UK2Node_StructMemberGet::GetNodeTitle(ENodeTitleType::Type TitleType) const
diff --git a/Engine/Source/Editor/BlueprintGraph/Private/K2Node_StructMemberSet.cpp b/Engine/Source/Editor/BlueprintGraph/Private/K2Node_StructMemberSet.cpp
--- a/Engine/Source/Editor/BlueprintGraph/Private/K2Node_StructMemberSet.cpp
+++ b/Engine/Source/Editor/BlueprintGraph/Private/K2Node_StructMemberSet.cpp
@@ -42,11 +42,11 @@ void UK2Node_StructMemberSet::AllocateDefaultPins()
 	}
 }
 
-FString UK2Node_StructMemberSet::GetTooltip() const
+FText UK2Node_StructMemberSet::GetTooltipText() const
 {
 	FFormatNamedArguments Args;
 	Args.Add(TEXT("VariableName"), FText::FromName(VariableReference.GetMemberName()));
-	return FText::Format(LOCTEXT("K2Node_StructMemberSet_Tooltip", "Set member variables of {VariableName}"), Args).ToString();
+	return FText::Format(LOCTEXT("K2Node_StructMemberSet_Tooltip", "Set member variables of {VariableName}"), Args);
 }
 
 FText UK2Node_StructMemberSet::GetNodeTitle(ENodeTitleType::Type TitleType) const
diff --git a/Engine/Source/Editor/BlueprintGraph/Private/K2Node_SwitchEnum.cpp b/Engine/Source/Editor/BlueprintGraph/Private/K2Node_SwitchEnum.cpp
--- a/Engine/Source/Editor/BlueprintGraph/Private/K2Node_SwitchEnum.cpp
+++ b/Engine/Source/Editor/BlueprintGraph/Private/K2Node_SwitchEnum.cpp
@@ -65,9 +65,9 @@ FText UK2Node_SwitchEnum::GetNodeTitle(ENodeTitleType::Type TitleType) const
 	return CachedNodeTitle;
 }
 
-FString UK2Node_SwitchEnum::GetTooltip() const
+FText UK2Node_SwitchEnum::GetTooltipText() const
 {
-	return NSLOCTEXT("K2Node", "SwitchEnum_ToolTip", "Selects an output that matches the input value").ToString();
+	return NSLOCTEXT("K2Node", "SwitchEnum_ToolTip", "Selects an output that matches the input value");
 }
 
 void UK2Node_SwitchEnum::GetMenuActions(FBlueprintActionDatabaseRegistrar& ActionRegistrar) const
diff --git a/Engine/Source/Editor/BlueprintGraph/Private/K2Node_SwitchInteger.cpp b/Engine/Source/Editor/BlueprintGraph/Private/K2Node_SwitchInteger.cpp
--- a/Engine/Source/Editor/BlueprintGraph/Private/K2Node_SwitchInteger.cpp
+++ b/Engine/Source/Editor/BlueprintGraph/Private/K2Node_SwitchInteger.cpp
@@ -19,9 +19,9 @@ FText UK2Node_SwitchInteger::GetNodeTitle(ENodeTitleType::Type TitleType) const
 	return NSLOCTEXT("K2Node", "Switch_Interger", "Switch on Int");
 }
 
-FString UK2Node_SwitchInteger::GetTooltip() const
+FText UK2Node_SwitchInteger::GetTooltipText() const
 {
-	return NSLOCTEXT("K2Node", "SwitchInt_ToolTip", "Selects an output that matches the input value").ToString();
+	return NSLOCTEXT("K2Node", "SwitchInt_ToolTip", "Selects an output that matches the input value");
 }
 
 void UK2Node_SwitchInteger::GetMenuActions(FBlueprintActionDatabaseRegistrar& ActionRegistrar) const
diff --git a/Engine/Source/Editor/BlueprintGraph/Private/K2Node_SwitchName.cpp b/Engine/Source/Editor/BlueprintGraph/Private/K2Node_SwitchName.cpp
--- a/Engine/Source/Editor/BlueprintGraph/Private/K2Node_SwitchName.cpp
+++ b/Engine/Source/Editor/BlueprintGraph/Private/K2Node_SwitchName.cpp
@@ -34,9 +34,9 @@ FText UK2Node_SwitchName::GetNodeTitle(ENodeTitleType::Type TitleType) const
 	return NSLOCTEXT("K2Node", "Switch_Name", "Switch on Name");
 }
 
-FString UK2Node_SwitchName::GetTooltip() const
+FText UK2Node_SwitchName::GetTooltipText() const
 {
-	return NSLOCTEXT("K2Node", "SwitchName_ToolTip", "Selects an output that matches the input value").ToString();
+	return NSLOCTEXT("K2Node", "SwitchName_ToolTip", "Selects an output that matches the input value");
 }
 
 void UK2Node_SwitchName::GetMenuActions(FBlueprintActionDatabaseRegistrar& ActionRegistrar) const
diff --git a/Engine/Source/Editor/BlueprintGraph/Private/K2Node_SwitchString.cpp b/Engine/Source/Editor/BlueprintGraph/Private/K2Node_SwitchString.cpp
--- a/Engine/Source/Editor/BlueprintGraph/Private/K2Node_SwitchString.cpp
+++ b/Engine/Source/Editor/BlueprintGraph/Private/K2Node_SwitchString.cpp
@@ -46,9 +46,9 @@ FText UK2Node_SwitchString::GetNodeTitle(ENodeTitleType::Type TitleType) const
 	return NSLOCTEXT("K2Node", "Switch_String", "Switch on String");
 }
 
-FString UK2Node_SwitchString::GetTooltip() const
+FText UK2Node_SwitchString::GetTooltipText() const
 {
-	return NSLOCTEXT("K2Node", "SwitchString_ToolTip", "Selects an output that matches the input value").ToString();
+	return NSLOCTEXT("K2Node", "SwitchString_ToolTip", "Selects an output that matches the input value");
 }
 
 void UK2Node_SwitchString::GetMenuActions(FBlueprintActionDatabaseRegistrar& ActionRegistrar) const
diff --git a/Engine/Source/Editor/BlueprintGraph/Private/K2Node_TemporaryVariable.cpp b/Engine/Source/Editor/BlueprintGraph/Private/K2Node_TemporaryVariable.cpp
--- a/Engine/Source/Editor/BlueprintGraph/Private/K2Node_TemporaryVariable.cpp
+++ b/Engine/Source/Editor/BlueprintGraph/Private/K2Node_TemporaryVariable.cpp
@@ -47,11 +47,11 @@ void UK2Node_TemporaryVariable::AllocateDefaultPins()
 	Super::AllocateDefaultPins();
 }
 
-FString UK2Node_TemporaryVariable::GetTooltip() const
+FText UK2Node_TemporaryVariable::GetTooltipText() const
 {
 	FFormatNamedArguments Args;
 	Args.Add(TEXT("VariableType"), UEdGraphSchema_K2::TypeToText(VariableType));
-	return FText::Format(NSLOCTEXT("K2Node", "LocalTemporaryVariable", "Local temporary {VariableType} variable"), Args).ToString();
+	return FText::Format(NSLOCTEXT("K2Node", "LocalTemporaryVariable", "Local temporary {VariableType} variable"), Args);
 }
 
 FText UK2Node_TemporaryVariable::GetNodeTitle(ENodeTitleType::Type TitleType) const
diff --git a/Engine/Source/Editor/BlueprintGraph/Private/K2Node_Timeline.cpp b/Engine/Source/Editor/BlueprintGraph/Private/K2Node_Timeline.cpp
--- a/Engine/Source/Editor/BlueprintGraph/Private/K2Node_Timeline.cpp
+++ b/Engine/Source/Editor/BlueprintGraph/Private/K2Node_Timeline.cpp
@@ -550,9 +550,9 @@ void UK2Node_Timeline::ExpandNode(FKismetCompilerContext& CompilerContext, UEdGr
 	}
 }
 
-FString UK2Node_Timeline::GetTooltip() const
+FText UK2Node_Timeline::GetTooltipText() const
 {
-	return LOCTEXT("TimelineTooltip", "Timeline node allows values to be keyframed over time.\nDouble click to open timeline editor.").ToString();
+	return LOCTEXT("TimelineTooltip", "Timeline node allows values to be keyframed over time.\nDouble click to open timeline editor.");
 }
 
 FString UK2Node_Timeline::GetDocumentationExcerptName() const
diff --git a/Engine/Source/Editor/BlueprintGraph/Private/K2Node_Tunnel.cpp b/Engine/Source/Editor/BlueprintGraph/Private/K2Node_Tunnel.cpp
--- a/Engine/Source/Editor/BlueprintGraph/Private/K2Node_Tunnel.cpp
+++ b/Engine/Source/Editor/BlueprintGraph/Private/K2Node_Tunnel.cpp
@@ -34,19 +34,19 @@ void UK2Node_Tunnel::PostPasteNode()
 	//@TODO: The gateway<->gateway node pairing should be unique, so we need to disallow this paste somehow (or just let it be an error later on)!
 }
 
-FString UK2Node_Tunnel::GetTooltip() const
+FText UK2Node_Tunnel::GetTooltipText() const
 {
 	if (bCanHaveInputs && !bCanHaveOutputs)
 	{
-		return NSLOCTEXT("K2Node", "OutputTunnelTooltip", "Outputs from this graph").ToString();
+		return NSLOCTEXT("K2Node", "OutputTunnelTooltip", "Outputs from this graph");
 	}
 	else if (!bCanHaveInputs && bCanHaveOutputs)
 	{
-		return NSLOCTEXT("K2Node", "InputTunnelTooltip", "Inputs into this graph").ToString();
+		return NSLOCTEXT("K2Node", "InputTunnelTooltip", "Inputs into this graph");
 	}
 	else
 	{
-		return NSLOCTEXT("K2Node", "TunnelConnectionTooltip", "Tunnel Connection").ToString();
+		return NSLOCTEXT("K2Node", "TunnelConnectionTooltip", "Tunnel Connection");
 	}
 }
 
diff --git a/Engine/Source/Editor/BlueprintGraph/Private/K2Node_VariableGet.cpp b/Engine/Source/Editor/BlueprintGraph/Private/K2Node_VariableGet.cpp
--- a/Engine/Source/Editor/BlueprintGraph/Private/K2Node_VariableGet.cpp
+++ b/Engine/Source/Editor/BlueprintGraph/Private/K2Node_VariableGet.cpp
@@ -85,7 +85,7 @@ void UK2Node_VariableGet::ReallocatePinsDuringReconstruction(TArray<UEdGraphPin*
 	}
 }
 
-FString UK2Node_VariableGet::GetTooltip() const
+FText UK2Node_VariableGet::GetTooltipText() const
 {
 	FFormatNamedArguments Args;
 	Args.Add( TEXT( "VarName" ), FText::FromString( GetVarNameString() ));
@@ -123,7 +123,7 @@ FString UK2Node_VariableGet::GetTooltip() const
 			}
 		}
 	}
-	return FText::Format( NSLOCTEXT( "K2Node", "GetVariable_ToolTip", "Read the value of variable {VarName}{TextPartition}{MetaData}"), Args ).ToString();
+	return FText::Format( NSLOCTEXT( "K2Node", "GetVariable_ToolTip", "Read the value of variable {VarName}{TextPartition}{MetaData}"), Args );
 }
 
 FText UK2Node_VariableGet::GetNodeTitle(ENodeTitleType::Type TitleType) const
diff --git a/Engine/Source/Editor/BlueprintGraph/Private/K2Node_VariableSet.cpp b/Engine/Source/Editor/BlueprintGraph/Private/K2Node_VariableSet.cpp
--- a/Engine/Source/Editor/BlueprintGraph/Private/K2Node_VariableSet.cpp
+++ b/Engine/Source/Editor/BlueprintGraph/Private/K2Node_VariableSet.cpp
@@ -48,7 +48,7 @@ void UK2Node_VariableSet::ReallocatePinsDuringReconstruction(TArray<UEdGraphPin*
 	}
 }
 
-FString UK2Node_VariableSet::GetTooltip() const
+FText UK2Node_VariableSet::GetTooltipText() const
 {
 	FFormatNamedArguments Args;
 	Args.Add( TEXT( "VarName" ), FText::FromString( GetVarNameString() ));
@@ -81,7 +81,7 @@ FString UK2Node_VariableSet::GetTooltip() const
 			}
 		}
 	}
-	return FText::Format( NSLOCTEXT( "K2Node", "SetValueOfVariable", "Set the value of variable {VarName}{ReplicationCall}{ReplicationNotifyName}{TextPartition}{MetaData}"), Args ).ToString();
+	return FText::Format( NSLOCTEXT( "K2Node", "SetValueOfVariable", "Set the value of variable {VarName}{ReplicationCall}{ReplicationNotifyName}{TextPartition}{MetaData}"), Args );
 }
 
 FText UK2Node_VariableSet::GetNodeTitle(ENodeTitleType::Type TitleType) const
diff --git a/Engine/Source/Editor/BlueprintGraph/Private/K2Node_VariableSetRef.cpp b/Engine/Source/Editor/BlueprintGraph/Private/K2Node_VariableSetRef.cpp
--- a/Engine/Source/Editor/BlueprintGraph/Private/K2Node_VariableSetRef.cpp
+++ b/Engine/Source/Editor/BlueprintGraph/Private/K2Node_VariableSetRef.cpp
@@ -121,9 +121,9 @@ void UK2Node_VariableSetRef::ReallocatePinsDuringReconstruction(TArray<UEdGraphP
   	}
 }
 
-FString UK2Node_VariableSetRef::GetTooltip() const
+FText UK2Node_VariableSetRef::GetTooltipText() const
 {
-	return FString::Printf(*NSLOCTEXT("K2Node", "SetValueOfRefVariable", "Set the value of the connected pass-by-ref variable").ToString());
+	return NSLOCTEXT("K2Node", "SetValueOfRefVariable", "Set the value of the connected pass-by-ref variable");
 }
 
 FText UK2Node_VariableSetRef::GetNodeTitle(ENodeTitleType::Type TitleType) const
diff --git a/Engine/Source/Editor/GameplayAbilitiesEditor/Private/K2Node_GameplayCueEvent.cpp b/Engine/Source/Editor/GameplayAbilitiesEditor/Private/K2Node_GameplayCueEvent.cpp
--- a/Engine/Source/Editor/GameplayAbilitiesEditor/Private/K2Node_GameplayCueEvent.cpp
+++ b/Engine/Source/Editor/GameplayAbilitiesEditor/Private/K2Node_GameplayCueEvent.cpp
@@ -20,9 +20,9 @@ UK2Node_GameplayCueEvent::UK2Node_GameplayCueEvent(const class FPostConstructIni
 	EventSignatureClass = UGameplayCueInterface::StaticClass();
 }
 
-FString UK2Node_GameplayCueEvent::GetTooltip() const
+FText UK2Node_GameplayCueEvent::GetTooltipText() const
 {
-	return TEXT("Handle GameplayCue Event");
+	return LOCTEXT("GameplayCueEvent_Tooltip", "Handle GameplayCue Event");
 }
 
 FText UK2Node_GameplayCueEvent::GetNodeTitle(ENodeTitleType::Type TitleType) const
@@ -76,7 +76,7 @@ void UK2Node_GameplayCueEvent::GetMenuEntries(FGraphContextMenuBuilder& Context)
 
 		const FString Category = FunctionCategory;
 		const FText MenuDesc = NodeTemplate->GetNodeTitle(ENodeTitleType::ListView);
-		const FString Tooltip = NodeTemplate->GetTooltip();
+		const FString Tooltip = NodeTemplate->GetTooltipText().ToString();
 		const FString Keywords = NodeTemplate->GetKeywords();
 
 		TSharedPtr<FEdGraphSchemaAction_K2NewNode> NodeAction = FK2ActionMenuBuilder::AddNewNodeAction(Context, Category, MenuDesc, Tooltip, 0, Keywords);
diff --git a/Engine/Source/Editor/GameplayTagsEditor/Private/GameplayTagsK2Node_LiteralGameplayTag.cpp b/Engine/Source/Editor/GameplayTagsEditor/Private/GameplayTagsK2Node_LiteralGameplayTag.cpp
--- a/Engine/Source/Editor/GameplayTagsEditor/Private/GameplayTagsK2Node_LiteralGameplayTag.cpp
+++ b/Engine/Source/Editor/GameplayTagsEditor/Private/GameplayTagsK2Node_LiteralGameplayTag.cpp
@@ -47,7 +47,7 @@ void UGameplayTagsK2Node_LiteralGameplayTag::GetMenuEntries( FGraphContextMenuBu
 
 	const FString Category = TEXT("Gameplay Tags|Tag Container");
 	const FText MenuDesc = EnumNodeTemplate->GetNodeTitle(ENodeTitleType::ListView);
-	const FString Tooltip = EnumNodeTemplate->GetTooltip();
+	const FString Tooltip = EnumNodeTemplate->GetTooltipText().ToString();
 	const FString Keywords = EnumNodeTemplate->GetKeywords();
 
 	TSharedPtr<FEdGraphSchemaAction_K2NewNode> NodeAction = FK2ActionMenuBuilder::AddNewNodeAction(ContextMenuBuilder, Category, MenuDesc, Tooltip, 0, Keywords);
diff --git a/Engine/Source/Editor/GraphEditor/Private/KismetNodes/SGraphNodeK2Base.cpp b/Engine/Source/Editor/GraphEditor/Private/KismetNodes/SGraphNodeK2Base.cpp
--- a/Engine/Source/Editor/GraphEditor/Private/KismetNodes/SGraphNodeK2Base.cpp
+++ b/Engine/Source/Editor/GraphEditor/Private/KismetNodes/SGraphNodeK2Base.cpp
@@ -40,7 +40,7 @@ void SGraphNodeK2Base::UpdateCompactNode()
 	LeftNodeBox.Reset();
 
  	TSharedPtr< SToolTip > NodeToolTip = SNew( SToolTip );
-	if ( !GraphNode->GetTooltip().IsEmpty() )
+	if (!GraphNode->GetTooltipText().IsEmpty())
 	{
 		NodeToolTip = IDocumentation::Get()->CreateToolTip( TAttribute< FText >( this, &SGraphNode::GetNodeTooltip ), NULL, GraphNode->GetDocumentationLink(), GraphNode->GetDocumentationExcerptName() );
 	}
diff --git a/Engine/Source/Editor/GraphEditor/Private/KismetNodes/SGraphNodeK2Composite.cpp b/Engine/Source/Editor/GraphEditor/Private/KismetNodes/SGraphNodeK2Composite.cpp
--- a/Engine/Source/Editor/GraphEditor/Private/KismetNodes/SGraphNodeK2Composite.cpp
+++ b/Engine/Source/Editor/GraphEditor/Private/KismetNodes/SGraphNodeK2Composite.cpp
@@ -196,7 +196,7 @@ FString SGraphNodeK2Composite::GetPreviewCornerText() const
 
 FString SGraphNodeK2Composite::GetTooltipTextForNode() const
 {
-	return GraphNode->GetTooltip();
+	return GraphNode->GetTooltipText().ToString();
 }
 
 TSharedRef<SWidget> SGraphNodeK2Composite::CreateNodeBody()
diff --git a/Engine/Source/Editor/GraphEditor/Private/SGraphNode.cpp b/Engine/Source/Editor/GraphEditor/Private/SGraphNode.cpp
--- a/Engine/Source/Editor/GraphEditor/Private/SGraphNode.cpp
+++ b/Engine/Source/Editor/GraphEditor/Private/SGraphNode.cpp
@@ -488,7 +488,7 @@ FText SGraphNode::GetNodeTooltip() const
 			return FText::FromString(GraphNode->GetNodeTitle(ENodeTitleType::ListView).BuildSourceString());
 		}
 
-		FText TooltipText = FText::FromString(GraphNode->GetTooltip());
+		FText TooltipText = GraphNode->GetTooltipText();
 
 		if (UEdGraph* Graph = GraphNode->GetGraph())
 		{
diff --git a/Engine/Source/Editor/Kismet/Private/BlueprintActionMenuBuilder.cpp b/Engine/Source/Editor/Kismet/Private/BlueprintActionMenuBuilder.cpp
--- a/Engine/Source/Editor/Kismet/Private/BlueprintActionMenuBuilder.cpp
+++ b/Engine/Source/Editor/Kismet/Private/BlueprintActionMenuBuilder.cpp
@@ -268,7 +268,7 @@ FText FBlueprintActionMenuItemFactory::GetTooltipForAction(TWeakPtr<FBlueprintEd
 	{
 		if (UEdGraphNode* NodeTemplate = GetTemplateNode(Action, EditorContext))
 		{
-			Tooltip = FText::FromString(NodeTemplate->GetTooltip());
+			Tooltip = NodeTemplate->GetTooltipText();
 		}
 	}
 	
diff --git a/Engine/Source/Editor/Kismet/Private/SBlueprintPalette.cpp b/Engine/Source/Editor/Kismet/Private/SBlueprintPalette.cpp
--- a/Engine/Source/Editor/Kismet/Private/SBlueprintPalette.cpp
+++ b/Engine/Source/Editor/Kismet/Private/SBlueprintPalette.cpp
@@ -296,7 +296,7 @@ static void GetPaletteItemIcon(TSharedPtr<FEdGraphSchemaAction> ActionIn, UBluep
 	else if (UK2Node const* const NodeTemplate = FK2SchemaActionUtils::ExtractNodeTemplateFromAction(ActionIn))
 	{
 		// If the node wants to create tooltip text, use that instead, because its probably more detailed
-		FString NodeToolTipText = NodeTemplate->GetTooltip();
+		FString NodeToolTipText = NodeTemplate->GetTooltipText().ToString();
 		if (NodeToolTipText.Len() > 0)
 		{
 			ToolTipOut = NodeToolTipText;
@@ -1305,7 +1305,7 @@ FText SBlueprintPaletteItem::GetToolTipText() const
 			}
 
 			// If the node wants to create tooltip text, use that instead, because its probably more detailed
-			FString NodeToolTipText = NodeTemplate->GetTooltip();
+			FString NodeToolTipText = NodeTemplate->GetTooltipText().ToString();
 			if (NodeToolTipText.Len() > 0)
 			{
 				ToolTipText = NodeToolTipText;
@@ -1384,7 +1384,7 @@ FText SBlueprintPaletteItem::GetToolTipText() const
 			FEdGraphSchemaAction_K2TargetNode* TargetNodeAction = (FEdGraphSchemaAction_K2TargetNode*)PaletteAction.Get();
 			if (TargetNodeAction->NodeTemplate)
 			{
-				ToolTipText = TargetNodeAction->NodeTemplate->GetTooltip();
+				ToolTipText = TargetNodeAction->NodeTemplate->GetTooltipText().ToString();
 			}
 		}
 	}
diff --git a/Engine/Source/Editor/Kismet/Private/SMyBlueprint.cpp b/Engine/Source/Editor/Kismet/Private/SMyBlueprint.cpp
--- a/Engine/Source/Editor/Kismet/Private/SMyBlueprint.cpp
+++ b/Engine/Source/Editor/Kismet/Private/SMyBlueprint.cpp
@@ -605,7 +605,7 @@ void SMyBlueprint::GetChildEvents(UEdGraph const* EdGraph, int32 const SectionId
 	{
 		UK2Node_Event* const EventNode = (*It);
 
-		FString const Tooltip = EventNode->GetTooltip();
+		FString const Tooltip = EventNode->GetTooltipText().ToString();
 		FText const Description = EventNode->GetNodeTitle(ENodeTitleType::EditableTitle);
 
 		TSharedPtr<FEdGraphSchemaAction_K2Event> EventNodeAction = MakeShareable(new FEdGraphSchemaAction_K2Event(ActionCategory, Description, Tooltip, 0));
diff --git a/Engine/Source/Editor/OnlineBlueprintSupport/Private/K2Node_LeaderboardFlush.cpp b/Engine/Source/Editor/OnlineBlueprintSupport/Private/K2Node_LeaderboardFlush.cpp
--- a/Engine/Source/Editor/OnlineBlueprintSupport/Private/K2Node_LeaderboardFlush.cpp
+++ b/Engine/Source/Editor/OnlineBlueprintSupport/Private/K2Node_LeaderboardFlush.cpp
@@ -15,9 +15,9 @@ UK2Node_LeaderboardFlush::UK2Node_LeaderboardFlush(const FPostConstructInitializ
 	ProxyClass = ULeaderboardFlushCallbackProxy::StaticClass();
 }
 
-FString UK2Node_LeaderboardFlush::GetTooltip() const
+FText UK2Node_LeaderboardFlush::GetTooltipText() const
 {
-	return LOCTEXT("K2Node_LeaderboardFlush_Tooltip", "Flushes leaderboards for a session").ToString();
+	return LOCTEXT("K2Node_LeaderboardFlush_Tooltip", "Flushes leaderboards for a session");
 }
 
 FText UK2Node_LeaderboardFlush::GetNodeTitle(ENodeTitleType::Type TitleType) const
diff --git a/Engine/Source/Editor/UMGEditor/Private/Nodes/K2Node_CreateDragDropOperation.cpp b/Engine/Source/Editor/UMGEditor/Private/Nodes/K2Node_CreateDragDropOperation.cpp
--- a/Engine/Source/Editor/UMGEditor/Private/Nodes/K2Node_CreateDragDropOperation.cpp
+++ b/Engine/Source/Editor/UMGEditor/Private/Nodes/K2Node_CreateDragDropOperation.cpp
@@ -14,7 +14,7 @@
 UK2Node_CreateDragDropOperation::UK2Node_CreateDragDropOperation(const class FPostConstructInitializeProperties& PCIP)
 	: Super(PCIP)
 {
-	NodeTooltip = LOCTEXT("NodeTooltip", "Creates a new drag drop operation").ToString();
+	NodeTooltip = LOCTEXT("NodeTooltip", "Creates a new drag drop operation");
 }
 
 FText UK2Node_CreateDragDropOperation::GetBaseNodeTitle() const
diff --git a/Engine/Source/Editor/UMGEditor/Private/Nodes/K2Node_CreateWidget.cpp b/Engine/Source/Editor/UMGEditor/Private/Nodes/K2Node_CreateWidget.cpp
--- a/Engine/Source/Editor/UMGEditor/Private/Nodes/K2Node_CreateWidget.cpp
+++ b/Engine/Source/Editor/UMGEditor/Private/Nodes/K2Node_CreateWidget.cpp
@@ -21,7 +21,7 @@ FString FK2Node_CreateWidgetHelper::OwningPlayerPinName(TEXT("OwningPlayer"));
 UK2Node_CreateWidget::UK2Node_CreateWidget(const class FPostConstructInitializeProperties& PCIP)
 	: Super(PCIP)
 {
-	NodeTooltip = LOCTEXT("NodeTooltip", "Creates a new widget").ToString();
+	NodeTooltip = LOCTEXT("NodeTooltip", "Creates a new widget");
 }
 
 void UK2Node_CreateWidget::AllocateDefaultPins()
diff --git a/Engine/Source/Editor/UnrealEd/Private/MaterialGraphNode.cpp b/Engine/Source/Editor/UnrealEd/Private/MaterialGraphNode.cpp
--- a/Engine/Source/Editor/UnrealEd/Private/MaterialGraphNode.cpp
+++ b/Engine/Source/Editor/UnrealEd/Private/MaterialGraphNode.cpp
@@ -255,7 +255,7 @@ FLinearColor UMaterialGraphNode::GetNodeTitleColor() const
 	return Settings->PureFunctionCallNodeTitleColor;
 }
 
-FString UMaterialGraphNode::GetTooltip() const
+FText UMaterialGraphNode::GetTooltipText() const
 {
 	if (MaterialExpression)
 	{
@@ -272,10 +272,10 @@ FString UMaterialGraphNode::GetTooltip() const
 				ToolTip += ToolTips[Index];
 			}
 
-			return ToolTip;
+			return FText::FromString(ToolTip);
 		}
 	}
-	return TEXT("");
+	return FText::GetEmpty();
 }
 
 void UMaterialGraphNode::PrepareForCopying()
diff --git a/Engine/Source/Editor/UnrealEd/Private/MaterialGraphNode_Root.cpp b/Engine/Source/Editor/UnrealEd/Private/MaterialGraphNode_Root.cpp
--- a/Engine/Source/Editor/UnrealEd/Private/MaterialGraphNode_Root.cpp
+++ b/Engine/Source/Editor/UnrealEd/Private/MaterialGraphNode_Root.cpp
@@ -29,9 +29,9 @@ FLinearColor UMaterialGraphNode_Root::GetNodeTitleColor() const
 	return GetDefault<UGraphEditorSettings>()->ResultNodeTitleColor;
 }
 
-FString UMaterialGraphNode_Root::GetTooltip() const
+FText UMaterialGraphNode_Root::GetTooltipText() const
 {
-	return LOCTEXT("RootToolTip", "Description of final material inputs").ToString();
+	return LOCTEXT("RootToolTip", "Description of final material inputs");
 }
 
 void UMaterialGraphNode_Root::PostPlacedNewNode()
diff --git a/Engine/Source/Editor/UnrealEd/Private/SoundCueGraphNode.cpp b/Engine/Source/Editor/UnrealEd/Private/SoundCueGraphNode.cpp
--- a/Engine/Source/Editor/UnrealEd/Private/SoundCueGraphNode.cpp
+++ b/Engine/Source/Editor/UnrealEd/Private/SoundCueGraphNode.cpp
@@ -240,16 +240,16 @@ void USoundCueGraphNode::GetContextMenuActions(const FGraphNodeContextMenuBuilde
 	}
 }
 
-FString USoundCueGraphNode::GetTooltip() const
+FText USoundCueGraphNode::GetTooltipText() const
 {
-	FString Tooltip;
+	FText Tooltip;
 	if (SoundNode)
 	{
-		Tooltip = SoundNode->GetClass()->GetToolTipText().ToString();
+		Tooltip = SoundNode->GetClass()->GetToolTipText();
 	}
-	if (Tooltip.Len() == 0)
+	if (Tooltip.IsEmpty())
 	{
-		Tooltip = GetNodeTitle(ENodeTitleType::ListView).ToString();
+		Tooltip = GetNodeTitle(ENodeTitleType::ListView);
 	}
 	return Tooltip;
 }
diff --git a/Engine/Source/Editor/UnrealEd/Private/SoundCueGraphNode_Root.cpp b/Engine/Source/Editor/UnrealEd/Private/SoundCueGraphNode_Root.cpp
--- a/Engine/Source/Editor/UnrealEd/Private/SoundCueGraphNode_Root.cpp
+++ b/Engine/Source/Editor/UnrealEd/Private/SoundCueGraphNode_Root.cpp
@@ -28,9 +28,9 @@ FText USoundCueGraphNode_Root::GetNodeTitle(ENodeTitleType::Type TitleType) cons
 	return LOCTEXT("RootTitle", "Output");
 }
 
-FString USoundCueGraphNode_Root::GetTooltip() const
+FText USoundCueGraphNode_Root::GetTooltipText() const
 {
-	return LOCTEXT("RootToolTip", "Wire the final Sound Node into this node").ToString();
+	return LOCTEXT("RootToolTip", "Wire the final Sound Node into this node");
 }
 
 void USoundCueGraphNode_Root::CreateInputPins()
diff --git a/Engine/Source/Runtime/Engine/Classes/EdGraph/EdGraphNode.h b/Engine/Source/Runtime/Engine/Classes/EdGraph/EdGraphNode.h
--- a/Engine/Source/Runtime/Engine/Classes/EdGraph/EdGraphNode.h
+++ b/Engine/Source/Runtime/Engine/Classes/EdGraph/EdGraphNode.h
@@ -276,13 +276,16 @@ public:
 	/**
 	 * Gets the tooltip to display when over the node
 	 */
-	virtual FString GetTooltip() const;
+	virtual FText GetTooltipText() const;
+
+	DEPRECATED(4.5, "UEdGraphNode::GetTooltip() is deprecated, instead use GetTooltipText(), which returns localized text.")
+	virtual FString GetTooltip() const { return GetTooltipText().ToString(); }
 
 	/**
 	 * Returns the keywords that should be used when searching for this node
 	 */
 	virtual FString GetKeywords() const { return TEXT(""); }
-
+	 
 	/**
 	 * Returns the link used for external documentation for the graph node
 	 */
diff --git a/Engine/Source/Runtime/Engine/Classes/EdGraph/EdGraphNode_Comment.h b/Engine/Source/Runtime/Engine/Classes/EdGraph/EdGraphNode_Comment.h
--- a/Engine/Source/Runtime/Engine/Classes/EdGraph/EdGraphNode_Comment.h
+++ b/Engine/Source/Runtime/Engine/Classes/EdGraph/EdGraphNode_Comment.h
@@ -45,7 +45,7 @@ public:
 #if WITH_EDITOR
 	// Begin UEdGraphNode interface
 	virtual void AllocateDefaultPins() override {}
-	ENGINE_API virtual FString GetTooltip() const override;
+	ENGINE_API virtual FText GetTooltipText() const override;
 	ENGINE_API virtual FLinearColor GetNodeCommentColor() const override;
 	ENGINE_API virtual FText GetNodeTitle(ENodeTitleType::Type TitleType) const override;
 	virtual bool ShouldOverridePinNames() const override { return true; }
diff --git a/Engine/Source/Runtime/Engine/Classes/EdGraph/EdGraphNode_Documentation.h b/Engine/Source/Runtime/Engine/Classes/EdGraph/EdGraphNode_Documentation.h
--- a/Engine/Source/Runtime/Engine/Classes/EdGraph/EdGraphNode_Documentation.h
+++ b/Engine/Source/Runtime/Engine/Classes/EdGraph/EdGraphNode_Documentation.h
@@ -23,7 +23,7 @@ public:
 	virtual void AllocateDefaultPins() override {}
 	virtual bool ShouldOverridePinNames() const override { return true; }
 	virtual FLinearColor GetNodeTitleColor() const override { return FLinearColor( 0.2f, 1.f, 0.2f ); }
-	ENGINE_API virtual FString GetTooltip() const override;
+	ENGINE_API virtual FText GetTooltipText() const override;
 	ENGINE_API virtual FText GetNodeTitle(ENodeTitleType::Type TitleType) const override;
 	ENGINE_API virtual FString GetPinNameOverride(const UEdGraphPin& Pin) const override;
 	ENGINE_API virtual void ResizeNode(const FVector2D& NewSize) override;
diff --git a/Engine/Source/Runtime/Engine/Private/EdGraph/EdGraphNode.cpp b/Engine/Source/Runtime/Engine/Private/EdGraph/EdGraphNode.cpp
--- a/Engine/Source/Runtime/Engine/Private/EdGraph/EdGraphNode.cpp
+++ b/Engine/Source/Runtime/Engine/Private/EdGraph/EdGraphNode.cpp
@@ -178,9 +178,9 @@ FLinearColor UEdGraphNode::GetNodeCommentColor() const
 	return FLinearColor::White;
 }
 
-FString UEdGraphNode::GetTooltip() const
+FText UEdGraphNode::GetTooltipText() const
 {
-	return GetClass()->GetToolTipText().ToString();
+	return GetClass()->GetToolTipText();
 }
 
 FString UEdGraphNode::GetDocumentationExcerptName() const
diff --git a/Engine/Source/Runtime/Engine/Private/EdGraph/EdGraphNode_Comment.cpp b/Engine/Source/Runtime/Engine/Private/EdGraph/EdGraphNode_Comment.cpp
--- a/Engine/Source/Runtime/Engine/Private/EdGraph/EdGraphNode_Comment.cpp
+++ b/Engine/Source/Runtime/Engine/Private/EdGraph/EdGraphNode_Comment.cpp
@@ -55,9 +55,9 @@ void UEdGraphNode_Comment::PostPlacedNewNode()
 	NodeComment = NSLOCTEXT("K2Node", "CommentBlock_NewEmptyComment", "Comment").ToString();
 }
 
-FString UEdGraphNode_Comment::GetTooltip() const
+FText UEdGraphNode_Comment::GetTooltipText() const
 {
-	return FString::Printf(*NSLOCTEXT("K2Node", "CommentBlock_Tooltip", "Comment:\n%s").ToString(), *NodeComment);
+	return FText::Format(NSLOCTEXT("K2Node", "CommentBlock_Tooltip", "Comment:\n{0}"), FText::FromString(NodeComment));
 }
 
 FString UEdGraphNode_Comment::GetDocumentationLink() const
diff --git a/Engine/Source/Runtime/Engine/Private/EdGraph/EdGraphNode_Documentation.cpp b/Engine/Source/Runtime/Engine/Private/EdGraph/EdGraphNode_Documentation.cpp
--- a/Engine/Source/Runtime/Engine/Private/EdGraph/EdGraphNode_Documentation.cpp
+++ b/Engine/Source/Runtime/Engine/Private/EdGraph/EdGraphNode_Documentation.cpp
@@ -30,9 +30,9 @@ void UEdGraphNode_Documentation::PostPlacedNewNode()
 	NodeComment = TEXT( "" );
 }
 
-FString UEdGraphNode_Documentation::GetTooltip() const
+FText UEdGraphNode_Documentation::GetTooltipText() const
 {
-	return FString::Printf( *NSLOCTEXT( "K2Node", "DocumentationBlock_Tooltip", "Documentation:\n%s" ).ToString(), *NodeComment );
+	return FText::Format(NSLOCTEXT("K2Node", "DocumentationBlock_Tooltip", "Documentation:\n{0}"), FText::FromString(NodeComment));
 }
 
 FText UEdGraphNode_Documentation::GetNodeTitle( ENodeTitleType::Type TitleType ) const
