Removing back'n forth string/text conversions for keywords in the new BP menu system.

diff --git a/Engine/Source/Editor/BlueprintGraph/Classes/BlueprintEventNodeSpawner.h b/Engine/Source/Editor/BlueprintGraph/Classes/BlueprintEventNodeSpawner.h
--- a/Engine/Source/Editor/BlueprintGraph/Classes/BlueprintEventNodeSpawner.h
+++ b/Engine/Source/Editor/BlueprintGraph/Classes/BlueprintEventNodeSpawner.h
@@ -48,7 +48,7 @@ public:
 	virtual UEdGraphNode* Invoke(UEdGraph* ParentGraph, FVector2D const Location) const override;
 	virtual FText GetDefaultMenuName() const override;
 	virtual FText GetDefaultMenuCategory() const override;
-	virtual FText GetDefaultSearchKeywords() const override;
+	virtual FString GetDefaultSearchKeywords() const override;
 	// End UBlueprintNodeSpawner interface
 
 	/**
diff --git a/Engine/Source/Editor/BlueprintGraph/Classes/BlueprintFunctionNodeSpawner.h b/Engine/Source/Editor/BlueprintGraph/Classes/BlueprintFunctionNodeSpawner.h
--- a/Engine/Source/Editor/BlueprintGraph/Classes/BlueprintFunctionNodeSpawner.h
+++ b/Engine/Source/Editor/BlueprintGraph/Classes/BlueprintFunctionNodeSpawner.h
@@ -46,7 +46,7 @@ public:
 	virtual FText GetDefaultMenuName() const override;
 	virtual FText GetDefaultMenuCategory() const override;
 	virtual FText GetDefaultMenuTooltip() const override;
-	virtual FText GetDefaultSearchKeywords() const override;
+	virtual FString GetDefaultSearchKeywords() const override;
 	// End UBlueprintNodeSpawner interface
 
 	// FBlueprintNodeBinder interface
diff --git a/Engine/Source/Editor/BlueprintGraph/Private/BlueprintEventNodeSpawner.cpp b/Engine/Source/Editor/BlueprintGraph/Private/BlueprintEventNodeSpawner.cpp
--- a/Engine/Source/Editor/BlueprintGraph/Private/BlueprintEventNodeSpawner.cpp
+++ b/Engine/Source/Editor/BlueprintGraph/Private/BlueprintEventNodeSpawner.cpp
@@ -167,12 +167,12 @@ FText UBlueprintEventNodeSpawner::GetDefaultMenuCategory() const
 }
 
 //------------------------------------------------------------------------------
-FText UBlueprintEventNodeSpawner::GetDefaultSearchKeywords() const
+FString UBlueprintEventNodeSpawner::GetDefaultSearchKeywords() const
 {
-	FText Keywords = FText::GetEmpty();
+	FString Keywords;
 	if (EventFunc != nullptr)
 	{
-		Keywords = FText::FromString(UK2Node_CallFunction::GetKeywordsForFunction(EventFunc));
+		Keywords = UK2Node_CallFunction::GetKeywordsForFunction(EventFunc);
 	}
 	return Keywords;
 }
diff --git a/Engine/Source/Editor/BlueprintGraph/Private/BlueprintFunctionNodeSpawner.cpp b/Engine/Source/Editor/BlueprintGraph/Private/BlueprintFunctionNodeSpawner.cpp
--- a/Engine/Source/Editor/BlueprintGraph/Private/BlueprintFunctionNodeSpawner.cpp
+++ b/Engine/Source/Editor/BlueprintGraph/Private/BlueprintFunctionNodeSpawner.cpp
@@ -239,10 +239,10 @@ FText UBlueprintFunctionNodeSpawner::GetDefaultMenuTooltip() const
 }
 
 //------------------------------------------------------------------------------
-FText UBlueprintFunctionNodeSpawner::GetDefaultSearchKeywords() const
+FString UBlueprintFunctionNodeSpawner::GetDefaultSearchKeywords() const
 {
 	check(Function != nullptr);
-	return FText::FromString(UK2Node_CallFunction::GetKeywordsForFunction(Function));
+	return UK2Node_CallFunction::GetKeywordsForFunction(Function);
 }
 
 //------------------------------------------------------------------------------
diff --git a/Engine/Source/Editor/BlueprintGraph/Private/BlueprintNodeSpawner.cpp b/Engine/Source/Editor/BlueprintGraph/Private/BlueprintNodeSpawner.cpp
--- a/Engine/Source/Editor/BlueprintGraph/Private/BlueprintNodeSpawner.cpp
+++ b/Engine/Source/Editor/BlueprintGraph/Private/BlueprintNodeSpawner.cpp
@@ -102,12 +102,12 @@ FText UBlueprintNodeSpawner::GetDefaultMenuTooltip() const
 }
 
 //------------------------------------------------------------------------------
-FText UBlueprintNodeSpawner::GetDefaultSearchKeywords() const
+FString UBlueprintNodeSpawner::GetDefaultSearchKeywords() const
 {
 	// DO NOT make a template node and query it here (the separate ui building
 	// code can do that if it likes)... this is meant to be an overridable
 	// alternative to that (for perf reasons)
-	return FText::GetEmpty();
+	return FString();
 }
 
 //------------------------------------------------------------------------------
diff --git a/Engine/Source/Editor/BlueprintGraph/Public/BlueprintNodeSpawner.h b/Engine/Source/Editor/BlueprintGraph/Public/BlueprintNodeSpawner.h
--- a/Engine/Source/Editor/BlueprintGraph/Public/BlueprintNodeSpawner.h
+++ b/Engine/Source/Editor/BlueprintGraph/Public/BlueprintNodeSpawner.h
@@ -131,9 +131,11 @@ public:
 	 * them from. Although, this can be thought of as a suggestion that the ui
 	 * building code is free to do with as it pleases (doesn't have to use it).
 	 *
+	 * @TODO: Should search keywords be localized? Probably.
+	 *
 	 * @return An empty text string (could be a localized keywords for sub-classes).
 	 */
-	virtual FText GetDefaultSearchKeywords() const;
+	virtual FString GetDefaultSearchKeywords() const;
 
 	/**
 	 * To save on performance, certain sub-classes can pull icons info from
diff --git a/Engine/Source/Editor/Kismet/Private/BlueprintActionMenuBuilder.cpp b/Engine/Source/Editor/Kismet/Private/BlueprintActionMenuBuilder.cpp
--- a/Engine/Source/Editor/Kismet/Private/BlueprintActionMenuBuilder.cpp
+++ b/Engine/Source/Editor/Kismet/Private/BlueprintActionMenuBuilder.cpp
@@ -105,11 +105,13 @@ private:
 	 * Attempts to pull a keywords from the supplied spawner. If one isn't 
 	 * provided, then it spawns a temporary node and pulls them from that.
 	 * 
+	 * @TODO: Should search keywords be localized? Probably. 
+	 *
 	 * @param  EditorContext
 	 * @param  Action		The action you want to suss keyword information from.
 	 * @return A keyword text string for the menu item wrapping this action.
 	 */
-	FText GetSearchKeywordsForAction(TWeakPtr<FBlueprintEditor> EditorContext, UBlueprintNodeSpawner const* Action);
+	FString GetSearchKeywordsForAction(TWeakPtr<FBlueprintEditor> EditorContext, UBlueprintNodeSpawner const* Action);
 
 	/**
 	 * Attempts to pull a menu icon information from the supplied spawner. If
@@ -159,7 +161,7 @@ TSharedPtr<FEdGraphSchemaAction> FBlueprintActionMenuItemFactory::MakeActionMenu
 	NewMenuItem->MenuDescription    = GetMenuNameForAction(EditorContext, Action);
 	NewMenuItem->TooltipDescription = GetTooltipForAction(EditorContext, Action).ToString();
 	NewMenuItem->Category           = GetCategoryForAction(EditorContext, Action).ToString();
-	NewMenuItem->Keywords           = GetSearchKeywordsForAction(EditorContext, Action).ToString();
+	NewMenuItem->Keywords           = GetSearchKeywordsForAction(EditorContext, Action);
 
 	NewMenuItem->Category = FString::Printf(TEXT("%s|%s"), *RootCategory.ToString(), *NewMenuItem->Category);	
 	return MakeShareable(NewMenuItem);
@@ -184,7 +186,7 @@ TSharedPtr<FBlueprintBoundMenuItem> FBlueprintActionMenuItemFactory::MakeBoundMe
 
 	NewMenuItem->MenuDescription    = GetMenuNameForAction(EditorContext, BoundAction);
 	NewMenuItem->TooltipDescription = GetTooltipForAction(EditorContext, BoundAction).ToString();
-	NewMenuItem->Keywords			= GetSearchKeywordsForAction(EditorContext, BoundAction).ToString();
+	NewMenuItem->Keywords			= GetSearchKeywordsForAction(EditorContext, BoundAction);
 
 	return MakeShareable(NewMenuItem);
 }
@@ -276,17 +278,20 @@ FText FBlueprintActionMenuItemFactory::GetTooltipForAction(TWeakPtr<FBlueprintEd
 }
 
 //------------------------------------------------------------------------------
-FText FBlueprintActionMenuItemFactory::GetSearchKeywordsForAction(TWeakPtr<FBlueprintEditor> EditorContext, UBlueprintNodeSpawner const* Action)
+FString FBlueprintActionMenuItemFactory::GetSearchKeywordsForAction(TWeakPtr<FBlueprintEditor> EditorContext, UBlueprintNodeSpawner const* Action)
 {
 	check(Action != nullptr);
-	// give the action the chance to save on performance (to keep from having to spawn a template node)
-	FText SearchKeywords = Action->GetDefaultSearchKeywords();
+	// give the action the chance to save on performance (to keep from having to 
+	// spawn a template node)
+	//
+	// @TODO: Should search keywords be localized? Probably.
+	FString SearchKeywords = Action->GetDefaultSearchKeywords();
 	
 	if (SearchKeywords.IsEmpty())
 	{
 		if (UEdGraphNode* NodeTemplate = GetTemplateNode(Action, EditorContext))
 		{
-			SearchKeywords = FText::FromString(NodeTemplate->GetKeywords());
+			SearchKeywords = NodeTemplate->GetKeywords();
 		}
 	}
 	
diff --git a/Engine/Source/Runtime/Engine/Classes/EdGraph/EdGraphNode.h b/Engine/Source/Runtime/Engine/Classes/EdGraph/EdGraphNode.h
--- a/Engine/Source/Runtime/Engine/Classes/EdGraph/EdGraphNode.h
+++ b/Engine/Source/Runtime/Engine/Classes/EdGraph/EdGraphNode.h
@@ -283,6 +283,8 @@ public:
 
 	/**
 	 * Returns the keywords that should be used when searching for this node
+	 *
+	 * @TODO: Should search keywords be localized? Probably.
 	 */
 	virtual FString GetKeywords() const { return TEXT(""); }
 	 
