- Changed UUserWidget outer to be game instance if it exists
- BP debugger now locates blueprints that exist within the game instance
- NULL to nullptr

diff --git a/Engine/Source/Editor/Kismet/Private/SBlueprintEditorSelectedDebugObjectWidget.cpp b/Engine/Source/Editor/Kismet/Private/SBlueprintEditorSelectedDebugObjectWidget.cpp
--- a/Engine/Source/Editor/Kismet/Private/SBlueprintEditorSelectedDebugObjectWidget.cpp
+++ b/Engine/Source/Editor/Kismet/Private/SBlueprintEditorSelectedDebugObjectWidget.cpp
@@ -19,6 +19,7 @@
 #include "IDocumentation.h"
 #include "SLevelOfDetailBranchNode.h"
 #include "SBlueprintEditorSelectedDebugObjectWidget.h"
+#include "Engine/GameInstance.h"
 
 #define LOCTEXT_NAMESPACE "KismetToolbar"
 
@@ -40,7 +41,7 @@ void SBlueprintEditorSelectedDebugObjectWidget::Construct(const FArguments& InAr
 	DebugWorldsComboBox = SNew(STextComboBox)
 		.ToolTip(IDocumentation::Get()->CreateToolTip(
 		LOCTEXT("BlueprintDebugWorldTooltip", "Select a world to debug"),
-		NULL,
+		nullptr,
 		TEXT("Shared/Editors/BlueprintEditor/BlueprintDebugger"),
 		TEXT("DebugWorld")))
 		.OptionsSource(&DebugWorldNames)
@@ -52,7 +53,7 @@ void SBlueprintEditorSelectedDebugObjectWidget::Construct(const FArguments& InAr
 	DebugObjectsComboBox = SNew(STextComboBox)
 		.ToolTip(IDocumentation::Get()->CreateToolTip(
 		LOCTEXT("BlueprintDebugObjectTooltip", "Select an object to debug"),
-		NULL,
+		nullptr,
 		TEXT("Shared/Editors/BlueprintEditor/BlueprintDebugger"),
 		TEXT("DebugObject")))
 		.OptionsSource(&DebugObjectNames)
@@ -172,7 +173,7 @@ void SBlueprintEditorSelectedDebugObjectWidget::Tick(const FGeometry& AllottedGe
 		}
 		else
 		{
-			LastObjectObserved = NULL;
+			LastObjectObserved = nullptr;
 
 			// If the object name is a name (rather than the 'No debug selected' string then regenerate the names (which will reset the combo box) as the object is invalid.
 			TSharedPtr<FString> CurrentString = DebugObjectsComboBox->GetSelectedItem();
@@ -213,7 +214,7 @@ void SBlueprintEditorSelectedDebugObjectWidget::OnRefresh()
 		}
 		else
 		{
-			// If we didnt find and object its probably a good idea to regenerate the names now - this will also ensure the combo box has a valid selection.
+			// If we didn't find and object its probably a good idea to regenerate the names now - this will also ensure the combo box has a valid selection.
 			GenerateDebugObjectNames(false);
 		}
 	}
@@ -232,10 +233,10 @@ void SBlueprintEditorSelectedDebugObjectWidget::GenerateDebugWorldNames(bool bRe
 	DebugWorldNames.Empty();
 	DebugWorlds.Empty();
 
-	DebugWorlds.Add(NULL);
+	DebugWorlds.Add(nullptr);
 	DebugWorldNames.Add(MakeShareable(new FString(GetDebugAllWorldsString())));
 
-	UWorld* PreviewWorld = NULL;
+	UWorld* PreviewWorld = nullptr;
 	TSharedPtr<SSCSEditorViewport> PreviewViewportPtr = BlueprintEditor.Pin()->GetSCSViewport();
 	if (PreviewViewportPtr.IsValid())
 	{
@@ -331,7 +332,7 @@ void SBlueprintEditorSelectedDebugObjectWidget::GenerateDebugObjectNames(bool bR
 	// Empty the lists of actors and regenerate them
 	DebugObjects.Empty();
 	DebugObjectNames.Empty();
-	DebugObjects.Add(NULL);
+	DebugObjects.Add(nullptr);
 	DebugObjectNames.Add(MakeShareable(new FString(GetNoDebugString())));
 
 	// Grab custom objects that should always be visible, regardless of the world
@@ -350,8 +351,8 @@ void SBlueprintEditorSelectedDebugObjectWidget::GenerateDebugObjectNames(bool bR
 		}
 	}
 
-	// Check for a specific debug world. If DebugWorld=NULL we take that as "any PIE world"
-	UWorld* DebugWorld = NULL;
+	// Check for a specific debug world. If DebugWorld=nullptr we take that as "any PIE world"
+	UWorld* DebugWorld = nullptr;
 	if (DebugWorldsComboBox.IsValid())
 	{
 		TSharedPtr<FString> CurrentWorldSelection = DebugWorldsComboBox->GetSelectedItem();
@@ -362,7 +363,7 @@ void SBlueprintEditorSelectedDebugObjectWidget::GenerateDebugObjectNames(bool bR
 		}
 	}
 
-	UWorld* PreviewWorld = NULL;
+	UWorld* PreviewWorld = nullptr;
 	TSharedPtr<SSCSEditorViewport> PreviewViewportPtr = BlueprintEditor.Pin()->GetSCSViewport();
 	if (PreviewViewportPtr.IsValid())
 	{
@@ -374,7 +375,7 @@ void SBlueprintEditorSelectedDebugObjectWidget::GenerateDebugObjectNames(bool bR
 		UObject* TestObject = *It;
 
 		// Skip Blueprint preview objects (don't allow them to be selected for debugging)
-		if (PreviewWorld != NULL && TestObject->IsIn(PreviewWorld))
+		if (PreviewWorld != nullptr && TestObject->IsIn(PreviewWorld))
 		{
 			continue;
 		}
@@ -386,39 +387,38 @@ void SBlueprintEditorSelectedDebugObjectWidget::GenerateDebugObjectNames(bool bR
 		if (bPassesFlags && bGeneratedByThisBlueprint)
 		{
 			UObject *ObjOuter = TestObject;
-			UWorld *ObjWorld = NULL;
-			while (ObjWorld == NULL && ObjOuter != NULL)
+			UWorld *ObjWorld = nullptr;
+			do		// Run through at least once in case the TestObject is a UGameInstance
 			{
-				ObjOuter = ObjOuter->GetOuter();
-				ObjWorld = Cast<UWorld>(ObjOuter);
-			}
+				UGameInstance *ObjGameInstance = Cast<UGameInstance>(ObjOuter);
 
-			// Object not in any world
-			if (!ObjWorld)
-			{
-				continue;
-			}
+				ObjOuter = ObjOuter->GetOuter();
+				ObjWorld = ObjGameInstance ? ObjGameInstance->GetWorld() : Cast<UWorld>(ObjOuter);
+			} while (ObjWorld == nullptr && ObjOuter != nullptr);
 
-			// Make check on owning level (not streaming level)
-			if (ObjWorld->PersistentLevel && ObjWorld->PersistentLevel->OwningWorld)
+			if (ObjWorld)
 			{
-				ObjWorld = ObjWorld->PersistentLevel->OwningWorld;
-			}
+				// Make check on owning level (not streaming level)
+				if (ObjWorld->PersistentLevel && ObjWorld->PersistentLevel->OwningWorld)
+				{
+					ObjWorld = ObjWorld->PersistentLevel->OwningWorld;
+				}
 
-			// We have a specific debug world and the object isnt in it
-			if (DebugWorld && ObjWorld != DebugWorld)
-			{
-				continue;
-			}
+				// We have a specific debug world and the object isn't in it
+				if (DebugWorld && ObjWorld != DebugWorld)
+				{
+					continue;
+				}
 
-			if ((ObjWorld->WorldType == EWorldType::Editor) && (GUnrealEd->GetPIEViewport() == nullptr))
-			{
-				AddDebugObject(TestObject);
+				if ((ObjWorld->WorldType == EWorldType::Editor) && (GUnrealEd->GetPIEViewport() == nullptr))
+				{
+					AddDebugObject(TestObject);
+				}
+				else if (ObjWorld->WorldType == EWorldType::PIE)
+				{
+					AddDebugObject(TestObject);
+				}
 			}
-			else if (ObjWorld->WorldType == EWorldType::PIE)
-			{
-				AddDebugObject(TestObject);
-			}			
 		}
 	}
 
@@ -530,7 +530,7 @@ void SBlueprintEditorSelectedDebugObjectWidget::DebugObjectSelectionChanged(TSha
 	{
 		if (*DebugObjectNames[ObjectIndex] == *NewSelection)
 		{
-			UObject* DebugObj = DebugObjects[ObjectIndex].IsValid() ? DebugObjects[ObjectIndex].Get() : NULL;
+			UObject* DebugObj = DebugObjects[ObjectIndex].IsValid() ? DebugObjects[ObjectIndex].Get() : nullptr;
 			GetBlueprintObj()->SetObjectBeingDebugged(DebugObj);
 		}
 	}
@@ -564,12 +564,12 @@ void SBlueprintEditorSelectedDebugObjectWidget::SelectedDebugObject_OnClicked()
 
 EVisibility SBlueprintEditorSelectedDebugObjectWidget::IsDebugWorldComboVisible() const
 {
-	if (GEditor->PlayWorld != NULL)
+	if (GEditor->PlayWorld != nullptr)
 	{
 		int32 LocalWorldCount = 0;
 		for (const FWorldContext& Context : GEngine->GetWorldContexts())
 		{
-			if (Context.WorldType == EWorldType::PIE && Context.World() != NULL)
+			if (Context.WorldType == EWorldType::PIE && Context.World() != nullptr)
 			{
 				++LocalWorldCount;
 			}
diff --git a/Engine/Source/Runtime/UMG/Public/Blueprint/UserWidget.h b/Engine/Source/Runtime/UMG/Public/Blueprint/UserWidget.h
--- a/Engine/Source/Runtime/UMG/Public/Blueprint/UserWidget.h
+++ b/Engine/Source/Runtime/UMG/Public/Blueprint/UserWidget.h
@@ -3,6 +3,7 @@
 #pragma once
 
 #include "Geometry.h"
+#include "Engine/GameInstance.h"
 #include "UserWidget.generated.h"
 
 static FGeometry NullGeometry;
@@ -364,10 +365,11 @@ T* CreateWidget(UWorld* World, UClass* UserWidgetClass)
 		return nullptr;
 	}
 
+	// Assign the outer to the game instance if it exists, otherwise use the world
+	UObject* Outer = World->GetGameInstance() ? StaticCast<UObject*>(World->GetGameInstance()) : StaticCast<UObject*>(World);
 	ULocalPlayer* Player = World->GetFirstLocalPlayerFromController();
-
-	UObject* Outer = ( Player == nullptr ) ? StaticCast<UObject*>(World) : StaticCast<UObject*>(Player);
 	UUserWidget* NewWidget = ConstructObject<UUserWidget>(UserWidgetClass, Outer);
+
 	NewWidget->SetPlayerContext(FLocalPlayerContext(Player));
 	NewWidget->Initialize();
 
@@ -383,7 +385,11 @@ T* CreateWidget(APlayerController* OwningPlayer, UClass* UserWidgetClass)
 		return nullptr;
 	}
 
-	UUserWidget* NewWidget = ConstructObject<UUserWidget>(UserWidgetClass, OwningPlayer);
+	// Assign the outer to the game instance if it exists, otherwise use the player controller's world
+	UWorld* World = OwningPlayer->GetWorld();
+	UObject* Outer = World->GetGameInstance() ? StaticCast<UObject*>(World->GetGameInstance()) : StaticCast<UObject*>(World);
+	UUserWidget* NewWidget = ConstructObject<UUserWidget>(UserWidgetClass, Outer);
+	
 	NewWidget->SetPlayerContext(FLocalPlayerContext(OwningPlayer));
 	NewWidget->Initialize();
 
