BP, UDS: Member variables in UDS can be reordered.
ttp335008 bp: custom structs: it would nice to be able to move the elements of the structs up and down in the list
ttp335569 Can't move items in a Struct up or down.

diff --git a/Engine/Source/Editor/Kismet/Private/UserDefinedStructureEditor.cpp b/Engine/Source/Editor/Kismet/Private/UserDefinedStructureEditor.cpp
--- a/Engine/Source/Editor/Kismet/Private/UserDefinedStructureEditor.cpp
+++ b/Engine/Source/Editor/Kismet/Private/UserDefinedStructureEditor.cpp
@@ -422,6 +422,12 @@ private:
 	FSimpleDelegate OnRegenerateChildren;
 };
 
+enum EMemberFieldPosition
+{
+	MFP_First	=	0x1,
+	MFP_Last	=	0x2,
+};
+
 ///////////////////////////////////////////////////////////////////////////////////////
 // FUserDefinedStructureFieldLayout
 
@@ -429,10 +435,11 @@ private:
 class FUserDefinedStructureFieldLayout : public IDetailCustomNodeBuilder, public TSharedFromThis<FUserDefinedStructureFieldLayout>
 {
 public:
-	FUserDefinedStructureFieldLayout(TWeakPtr<class FUserDefinedStructureDetails> InStructureDetails, TWeakPtr<class FUserDefinedStructureLayout> InStructureLayout, FGuid InFieldGuid)
+	FUserDefinedStructureFieldLayout(TWeakPtr<class FUserDefinedStructureDetails> InStructureDetails, TWeakPtr<class FUserDefinedStructureLayout> InStructureLayout, FGuid InFieldGuid, uint32 InPositionFlags)
 		: StructureDetails(InStructureDetails)
 		, StructureLayout(InStructureLayout)
-		, FieldGuid(InFieldGuid) {}
+		, FieldGuid(InFieldGuid)
+		, PositionFlags(InPositionFlags) {}
 
 	void OnChanged()
 	{
@@ -642,6 +649,26 @@ public:
 		}
 	}
 
+	FReply OnMoveUp()
+	{
+		auto StructureDetailsSP = StructureDetails.Pin();
+		if (StructureDetailsSP.IsValid() && !(PositionFlags & EMemberFieldPosition::MFP_First))
+		{
+			FStructureEditorUtils::MoveVariable(StructureDetailsSP->GetUserDefinedStruct(), FieldGuid, FStructureEditorUtils::MD_Up);
+		}
+		return FReply::Handled();
+	}
+
+	FReply OnMoveDown()
+	{
+		auto StructureDetailsSP = StructureDetails.Pin();
+		if (StructureDetailsSP.IsValid() && !(PositionFlags & EMemberFieldPosition::MFP_Last))
+		{
+			FStructureEditorUtils::MoveVariable(StructureDetailsSP->GetUserDefinedStruct(), FieldGuid, FStructureEditorUtils::MD_Down);
+		}
+		return FReply::Handled();
+	}
+
 	virtual void GenerateHeaderRowContent( FDetailWidgetRow& NodeRow ) override 
 	{
 		auto K2Schema = GetDefault<UEdGraphSchema_K2>();
@@ -692,6 +719,34 @@ public:
 				.bAllowWildcard(false)
 				.Font( IDetailLayoutBuilder::GetDetailFont() )
 			]
+			+ SHorizontalBox::Slot()
+			.AutoWidth()
+			.HAlign(HAlign_Right)
+			.VAlign(VAlign_Center)
+			[
+				SNew(SButton)
+				.ContentPadding(0)
+				.OnClicked(this, &FUserDefinedStructureFieldLayout::OnMoveUp)
+				.IsEnabled(!(EMemberFieldPosition::MFP_First & PositionFlags))
+				[
+					SNew(SImage)
+					.Image(FEditorStyle::GetBrush("BlueprintEditor.Details.ArgUpButton"))
+				]
+			]
+			+ SHorizontalBox::Slot()
+			.AutoWidth()
+			.HAlign(HAlign_Right)
+			.VAlign(VAlign_Center)
+			[
+				SNew(SButton)
+				.ContentPadding(0)
+				.OnClicked(this, &FUserDefinedStructureFieldLayout::OnMoveDown)
+				.IsEnabled(!(EMemberFieldPosition::MFP_Last & PositionFlags))
+				[
+					SNew(SImage)
+					.Image(FEditorStyle::GetBrush("BlueprintEditor.Details.ArgDownButton"))
+				]
+			]
 			+SHorizontalBox::Slot()
 			.AutoWidth()
 			.HAlign(HAlign_Right)
@@ -777,6 +832,8 @@ private:
 	FGuid FieldGuid;
 
 	FSimpleDelegate OnRegenerateChildren;
+
+	uint32 PositionFlags;
 };
 
 ///////////////////////////////////////////////////////////////////////////////////////
@@ -844,9 +901,14 @@ void FUserDefinedStructureLayout::GenerateChildContent( IDetailChildrenBuilder&
 	{
 		if(auto Struct = StructureDetailsSP->GetUserDefinedStruct())
 		{
-			for (auto& VarDesc : FStructureEditorUtils::GetVarDesc(Struct))
+			auto& VarDescArrayRef = FStructureEditorUtils::GetVarDesc(Struct);
+			for (int32 Index = 0; Index < VarDescArrayRef.Num(); ++Index)
 			{
-				TSharedRef<class FUserDefinedStructureFieldLayout> VarLayout = MakeShareable(new FUserDefinedStructureFieldLayout(StructureDetails,  SharedThis(this), VarDesc.VarGuid));
+				auto& VarDesc = VarDescArrayRef[Index];
+				uint32 PositionFlag = 0;
+				PositionFlag |= (0 == Index) ? EMemberFieldPosition::MFP_First : 0;
+				PositionFlag |= ((VarDescArrayRef.Num() - 1) == Index) ? EMemberFieldPosition::MFP_Last : 0;
+				TSharedRef<class FUserDefinedStructureFieldLayout> VarLayout = MakeShareable(new FUserDefinedStructureFieldLayout(StructureDetails,  SharedThis(this), VarDesc.VarGuid, PositionFlag));
 				ChildrenBuilder.AddChildCustomBuilder(VarLayout);
 			}
 		}
diff --git a/Engine/Source/Editor/UnrealEd/Private/Kismet2/StructureEditorUtils.cpp b/Engine/Source/Editor/UnrealEd/Private/Kismet2/StructureEditorUtils.cpp
--- a/Engine/Source/Editor/UnrealEd/Private/Kismet2/StructureEditorUtils.cpp
+++ b/Engine/Source/Editor/UnrealEd/Private/Kismet2/StructureEditorUtils.cpp
@@ -611,7 +611,7 @@ bool FStructureEditorUtils::ChangeEditableOnBPInstance(UUserDefinedStruct* Struc
 	const bool bNewDontEditoOnInstance = !bInIsEditable;
 	if (VarDesc && (bNewDontEditoOnInstance != VarDesc->bDontEditoOnInstance))
 	{
-		const FScopedTransaction Transaction(LOCTEXT("ChangeVariableDefaultValue", "Change Variable Default Value"));
+		const FScopedTransaction Transaction(LOCTEXT("ChangeVariableOnBPInstance", "Change variable editable on BP instance"));
 		ModifyStructData(Struct);
 
 		VarDesc->bDontEditoOnInstance = bNewDontEditoOnInstance;
@@ -621,6 +621,30 @@ bool FStructureEditorUtils::ChangeEditableOnBPInstance(UUserDefinedStruct* Struc
 	return false;
 }
 
+bool FStructureEditorUtils::MoveVariable(UUserDefinedStruct* Struct, FGuid VarGuid, EMoveDirection MoveDirection)
+{
+	if (Struct)
+	{
+		const bool bMoveUp = (EMoveDirection::MD_Up == MoveDirection);
+		auto& DescArray = GetVarDesc(Struct);
+		const int32 InitialIndex = bMoveUp ? 1 : 0;
+		const int32 IndexLimit = DescArray.Num() - (bMoveUp ? 0 : 1);
+		for (int32 Index = InitialIndex; Index < IndexLimit; ++Index)
+		{
+			if (DescArray[Index].VarGuid == VarGuid)
+			{
+				const FScopedTransaction Transaction(LOCTEXT("ChangeVariableDefaultValue", "Varaibles reordered"));
+				ModifyStructData(Struct);
+
+				DescArray.Swap(Index, Index + (bMoveUp ? -1 : 1));
+				OnStructureChanged(Struct);
+				return true;
+			}
+		}
+	}
+	return false;
+}
+
 void FStructureEditorUtils::ModifyStructData(UUserDefinedStruct* Struct)
 {
 	UUserDefinedStructEditorData* EditorData = Struct ? Cast<UUserDefinedStructEditorData>(Struct->EditorData) : NULL;
diff --git a/Engine/Source/Editor/UnrealEd/Public/Kismet2/StructureEditorUtils.h b/Engine/Source/Editor/UnrealEd/Public/Kismet2/StructureEditorUtils.h
--- a/Engine/Source/Editor/UnrealEd/Public/Kismet2/StructureEditorUtils.h
+++ b/Engine/Source/Editor/UnrealEd/Public/Kismet2/StructureEditorUtils.h
@@ -73,6 +73,13 @@ public:
 
 	static bool ChangeEditableOnBPInstance(UUserDefinedStruct* Struct, FGuid VarGuid, bool bInIsEditable);
 
+	enum EMoveDirection
+	{
+		MD_Up,
+		MD_Down,
+	};
+	static bool MoveVariable(UUserDefinedStruct* Struct, FGuid VarGuid, EMoveDirection MoveDirection);
+
 	//3D Widget
 	static bool CanEnable3dWidget(const UUserDefinedStruct* Struct, FGuid VarGuid);
 
