Shrinking the mem footprint of the new BP menu system (moving bindings out from node spawners, passing into Invoke() instead).

diff --git a/Engine/Source/Editor/BlueprintGraph/Classes/BlueprintBoundEventNodeSpawner.h b/Engine/Source/Editor/BlueprintGraph/Classes/BlueprintBoundEventNodeSpawner.h
--- a/Engine/Source/Editor/BlueprintGraph/Classes/BlueprintBoundEventNodeSpawner.h
+++ b/Engine/Source/Editor/BlueprintGraph/Classes/BlueprintBoundEventNodeSpawner.h
@@ -32,19 +32,19 @@ public:
 	static UBlueprintBoundEventNodeSpawner* Create(TSubclassOf<UK2Node_Event> NodeClass, UMulticastDelegateProperty* EventDelegate, UObject* Outer = nullptr);
 
 	// UBlueprintNodeSpawner interface
-	virtual UEdGraphNode* Invoke(UEdGraph* ParentGraph, FVector2D const Location) const override;
+	virtual UEdGraphNode* Invoke(UEdGraph* ParentGraph, FBindingSet const& Bindings, FVector2D const Location) const override;
 	virtual FText GetDefaultMenuName() const override;
 	virtual FText GetDefaultMenuCategory() const override;
 	// End UBlueprintNodeSpawner interface
 
 	// UBlueprintEventNodeSpawner interface
-	virtual UK2Node_Event const* FindPreExistingEvent(UBlueprint* Blueprint) const override;
+	virtual UK2Node_Event const* FindPreExistingEvent(UBlueprint* Blueprint, FBindingSet const& Bindings) const override;
 	// End UBlueprintEventNodeSpawner interface
 
-	// FBlueprintNodeBinder interface
-	virtual bool CanBind(UObject const* BindingCandidate) const override;
+	// IBlueprintNodeBinder interface
+	virtual bool IsBindingCompatible(UObject const* BindingCandidate) const override;
 	virtual bool BindToNode(UEdGraphNode* Node, UObject* Binding) const override;
-	// End FBlueprintNodeBinder interface
+	// End IBlueprintNodeBinder interface
 
 	/**
 	 * 
diff --git a/Engine/Source/Editor/BlueprintGraph/Classes/BlueprintComponentNodeSpawner.h b/Engine/Source/Editor/BlueprintGraph/Classes/BlueprintComponentNodeSpawner.h
--- a/Engine/Source/Editor/BlueprintGraph/Classes/BlueprintComponentNodeSpawner.h
+++ b/Engine/Source/Editor/BlueprintGraph/Classes/BlueprintComponentNodeSpawner.h
@@ -32,7 +32,7 @@ public:
 	static UBlueprintComponentNodeSpawner* Create(TSubclassOf<UActorComponent> const ComponentClass, UObject* Outer = nullptr);
 
 	// UBlueprintNodeSpawner interface
-	virtual UEdGraphNode* Invoke(UEdGraph* ParentGraph, FVector2D const Location) const override;
+	virtual UEdGraphNode* Invoke(UEdGraph* ParentGraph, FBindingSet const& Bindings, FVector2D const Location) const override;
 	virtual FText GetDefaultMenuName() const override;
 	virtual FText GetDefaultMenuCategory() const override;
 	virtual FName GetDefaultMenuIcon(FLinearColor& ColorOut) const override;
diff --git a/Engine/Source/Editor/BlueprintGraph/Classes/BlueprintDelegateNodeSpawner.h b/Engine/Source/Editor/BlueprintGraph/Classes/BlueprintDelegateNodeSpawner.h
--- a/Engine/Source/Editor/BlueprintGraph/Classes/BlueprintDelegateNodeSpawner.h
+++ b/Engine/Source/Editor/BlueprintGraph/Classes/BlueprintDelegateNodeSpawner.h
@@ -32,7 +32,7 @@ public:
 	static UBlueprintDelegateNodeSpawner* Create(TSubclassOf<UK2Node_BaseMCDelegate> NodeClass, UMulticastDelegateProperty const* const Property, UObject* Outer = nullptr);
 
 	// UBlueprintNodeSpawner interface
-	virtual UEdGraphNode* Invoke(UEdGraph* ParentGraph, FVector2D const Location) const override;
+	virtual UEdGraphNode* Invoke(UEdGraph* ParentGraph, FBindingSet const& Bindings, FVector2D const Location) const override;
 	virtual FText GetDefaultMenuName() const override;
 	virtual FText GetDefaultMenuCategory() const override;
 	virtual FName GetDefaultMenuIcon(FLinearColor& ColorOut) const override;
diff --git a/Engine/Source/Editor/BlueprintGraph/Classes/BlueprintEventNodeSpawner.h b/Engine/Source/Editor/BlueprintGraph/Classes/BlueprintEventNodeSpawner.h
--- a/Engine/Source/Editor/BlueprintGraph/Classes/BlueprintEventNodeSpawner.h
+++ b/Engine/Source/Editor/BlueprintGraph/Classes/BlueprintEventNodeSpawner.h
@@ -45,7 +45,7 @@ public:
 	static UBlueprintEventNodeSpawner* Create(TSubclassOf<UK2Node_Event> NodeClass, FName CustomEventName, UObject* Outer = nullptr);
 
 	// UBlueprintNodeSpawner interface
-	virtual UEdGraphNode* Invoke(UEdGraph* ParentGraph, FVector2D const Location) const override;
+	virtual UEdGraphNode* Invoke(UEdGraph* ParentGraph, FBindingSet const& Bindings, FVector2D const Location) const override;
 	virtual FText GetDefaultMenuName() const override;
 	virtual FText GetDefaultMenuCategory() const override;
 	virtual FString GetDefaultSearchKeywords() const override;
@@ -71,7 +71,7 @@ public:
 	 * 
 	 * @return 
 	 */
-	virtual UK2Node_Event const* FindPreExistingEvent(UBlueprint* Blueprint) const;
+	virtual UK2Node_Event const* FindPreExistingEvent(UBlueprint* Blueprint, FBindingSet const& Bindings) const;
 
 protected:
 	/** Constructing FText strings can be costly, so we cache the default menu name */
diff --git a/Engine/Source/Editor/BlueprintGraph/Classes/BlueprintFunctionNodeSpawner.h b/Engine/Source/Editor/BlueprintGraph/Classes/BlueprintFunctionNodeSpawner.h
--- a/Engine/Source/Editor/BlueprintGraph/Classes/BlueprintFunctionNodeSpawner.h
+++ b/Engine/Source/Editor/BlueprintGraph/Classes/BlueprintFunctionNodeSpawner.h
@@ -42,19 +42,19 @@ public:
 	static UBlueprintFunctionNodeSpawner* Create(TSubclassOf<UK2Node_CallFunction> NodeClass, UFunction const* const Function, UObject* Outer = nullptr);
 
 	// UBlueprintNodeSpawner interface
-	virtual UEdGraphNode* Invoke(UEdGraph* ParentGraph, FVector2D const Location) const override;
+	virtual UEdGraphNode* Invoke(UEdGraph* ParentGraph, FBindingSet const& Bindings, FVector2D const Location) const override;
 	virtual FText GetDefaultMenuName() const override;
 	virtual FText GetDefaultMenuCategory() const override;
 	virtual FText GetDefaultMenuTooltip() const override;
 	virtual FString GetDefaultSearchKeywords() const override;
 	// End UBlueprintNodeSpawner interface
 
-	// FBlueprintNodeBinder interface
+	// IBlueprintNodeBinder interface
+	virtual bool IsBindingCompatible(UObject const* BindingCandidate) const override;
 	virtual bool CanBindMultipleObjects() const override;
-	virtual bool CanBind(UObject const* BindingCandidate) const override;
 	virtual bool BindToNode(UEdGraphNode* Node, UObject* Binding) const override;
-	// End FBlueprintNodeBinder interface
-	
+	// End IBlueprintNodeBinder interface
+
 	/**
 	 * Retrieves the function that this assigns to spawned nodes (defines the 
 	 * node's signature).
diff --git a/Engine/Source/Editor/BlueprintGraph/Classes/BlueprintNodeBinder.h b/Engine/Source/Editor/BlueprintGraph/Classes/BlueprintNodeBinder.h
--- a/Engine/Source/Editor/BlueprintGraph/Classes/BlueprintNodeBinder.h
+++ b/Engine/Source/Editor/BlueprintGraph/Classes/BlueprintNodeBinder.h
@@ -2,121 +2,63 @@
 
 #pragma once
 
+// Forward declarations
 class UEdGraphNode;
 
-/**
- * 
- */
-class FBlueprintNodeBinder
+
+class IBlueprintNodeBinder
 {
 public:
-	virtual ~FBlueprintNodeBinder() {}
-
-	/**
-	 * 
-	 * 
-	 * @param  BindingCandidate	
-	 * @return 
-	 */
-	virtual bool CanBind(UObject const* BindingCandidate) const = 0;
-
-	/**
-	 * 
-	 * 
-	 * @return 
-	 */
-	virtual bool CanBindMultipleObjects() const
-	{
-		return false;
-	}
+	/** */
+	typedef TSet< TWeakObjectPtr<UObject> > FBindingSet;
 
+public:
 	/**
+	 * Checks to see if the specified object can be bound by this.
 	 * 
-	 * 
-	 * @return 
+	 * @param  BindingCandidate	The object you want to check for.
+	 * @return True if BindingCandidate can be bound by this controller, false if not.
 	 */
-	bool IsBindingSet() const
-	{
-		return (BoundObjects.Num() > 0);
-	}
+	virtual bool IsBindingCompatible(UObject const* BindingCandidate) const = 0;
 
 	/**
+	 * Determines if this will accept more than one binding (used to block multiple 
+	 * bindings from being applied to nodes that can only have one).
 	 * 
-	 * 
-	 * @param  BindingCandidate	
-	 * @return 
+	 * @return True if this will accept multiple bindings, otherwise false.
 	 */
-	bool AddBinding(UObject const* BindingCandidate)
-	{
-		bool const bAddBinding = CanBind(BindingCandidate) && (!IsBindingSet() || CanBindMultipleObjects());
-		if (bAddBinding)
-		{
-			BoundObjects.Add(BindingCandidate);
-		}
-		return bAddBinding;
-	}
+	virtual bool CanBindMultipleObjects() const = 0;
 
 	/**
+	 * Attempts to bind all bindings to the supplied node.
 	 * 
-	 * 
-	 * @param  Node	
-	 * @return 
+	 * @param  Node	 The node you want bound to.
+	 * @return True if all bindings were successfully applied, false if any failed.
 	 */
-	bool Bind(UEdGraphNode* Node) const
+	bool ApplyBindings(UEdGraphNode* Node, FBindingSet const& Bindings) const
 	{
-		bool bFullSuccess = true;
-		for (TWeakObjectPtr<UObject> Instance : BoundObjects)
+		uint32 BindingCount = 0;
+		for (TWeakObjectPtr<UObject> const& Binding : Bindings)
 		{
-			if (Instance.IsValid())
+			if (Binding.IsValid() && BindToNode(Node, Binding.Get()))
 			{
-				bFullSuccess &= BindToNode(Node, Instance.Get());
+				++BindingCount;
+				if (!CanBindMultipleObjects())
+				{
+					break;
+				}
 			}
 		}
-		return bFullSuccess;
-	}
-
-	/**
-	 * 
-	 * 
-	 * @param  Binding	
-	 * @return 
-	 */
-	void RemoveBinding(UObject const* Binding)
-	{
-		BoundObjects.Remove(Binding);
-	}
-
-	/**
-	 * 
-	 * 
-	 * @return 
-	 */
-	void ClearBindings()
-	{
-		BoundObjects.Empty();
-	}
-
-	/**
-	 * 
-	 * 
-	 * @return 
-	 */
-	TSet< TWeakObjectPtr<UObject> > const& GetBindings() const
-	{
-		return BoundObjects;
+		return (BindingCount == Bindings.Num());
 	}
 
 protected:
 	/**
+	 * Attempts to apply the specified binding to the supplied node.
 	 * 
-	 * 
-	 * @param  Node	
-	 * @param  Binding	
-	 * @return 
+	 * @param  Node		The node you want bound.
+	 * @param  Binding	The binding you want applied to Node.
+	 * @return True if the binding was successful, false if not.
 	 */
 	virtual bool BindToNode(UEdGraphNode* Node, UObject* Binding) const = 0;
-
-protected:
-	/** */
-	TSet< TWeakObjectPtr<UObject> > BoundObjects;
 };
diff --git a/Engine/Source/Editor/BlueprintGraph/Classes/BlueprintVariableNodeSpawner.h b/Engine/Source/Editor/BlueprintGraph/Classes/BlueprintVariableNodeSpawner.h
--- a/Engine/Source/Editor/BlueprintGraph/Classes/BlueprintVariableNodeSpawner.h
+++ b/Engine/Source/Editor/BlueprintGraph/Classes/BlueprintVariableNodeSpawner.h
@@ -46,7 +46,7 @@ public:
 	static UBlueprintVariableNodeSpawner* Create(TSubclassOf<UK2Node_Variable> NodeClass, UEdGraph* VarContext, FBPVariableDescription const& VarDesc, UObject* Outer = nullptr);
 
 	// UBlueprintNodeSpawner interface
-	virtual UEdGraphNode* Invoke(UEdGraph* ParentGraph, FVector2D const Location) const override;
+	virtual UEdGraphNode* Invoke(UEdGraph* ParentGraph, FBindingSet const& Bindings, FVector2D const Location) const override;
 	virtual FText GetDefaultMenuName() const override;
 	virtual FText GetDefaultMenuCategory() const override;
 	virtual FName GetDefaultMenuIcon(FLinearColor& ColorOut) const override;
diff --git a/Engine/Source/Editor/BlueprintGraph/Private/BlueprintActionDatabase.cpp b/Engine/Source/Editor/BlueprintGraph/Private/BlueprintActionDatabase.cpp
--- a/Engine/Source/Editor/BlueprintGraph/Private/BlueprintActionDatabase.cpp
+++ b/Engine/Source/Editor/BlueprintGraph/Private/BlueprintActionDatabase.cpp
@@ -202,7 +202,7 @@ static UBlueprintNodeSpawner* FBlueprintNodeSpawnerFactory::MakeComponentBoundEv
 //------------------------------------------------------------------------------
 static UBlueprintNodeSpawner* FBlueprintNodeSpawnerFactory::MakeActorBoundEventSpawner(UMulticastDelegateProperty* DelegateProperty)
 {
-	return UBlueprintBoundEventNodeSpawner::Create(UK2Node_ComponentBoundEvent::StaticClass(), DelegateProperty);
+	return UBlueprintBoundEventNodeSpawner::Create(UK2Node_ActorBoundEvent::StaticClass(), DelegateProperty);
 }
 
 /*******************************************************************************
diff --git a/Engine/Source/Editor/BlueprintGraph/Private/BlueprintActionFilter.cpp b/Engine/Source/Editor/BlueprintGraph/Private/BlueprintActionFilter.cpp
--- a/Engine/Source/Editor/BlueprintGraph/Private/BlueprintActionFilter.cpp
+++ b/Engine/Source/Editor/BlueprintGraph/Private/BlueprintActionFilter.cpp
@@ -187,7 +187,7 @@ namespace BlueprintActionFilterImpl
 	 * @param  BlueprintAction	The action you wish to query.
 	 * @return True if the action would produce a bound node, tied to an object that isn't selected.
 	 */
-	static bool IsBoundToUnselectedObject(FBlueprintActionFilter const& Filter, FBlueprintActionInfo& BlueprintAction);
+	static bool IsUnBoundBindingSpawner(FBlueprintActionFilter const& Filter, FBlueprintActionInfo& BlueprintAction);
 
 	/**
 	 * 
@@ -670,20 +670,14 @@ static bool BlueprintActionFilterImpl::IsFilteredNodeType(FBlueprintActionFilter
 }
 
 //------------------------------------------------------------------------------
-static bool BlueprintActionFilterImpl::IsBoundToUnselectedObject(FBlueprintActionFilter const& Filter, FBlueprintActionInfo& BlueprintAction)
+static bool BlueprintActionFilterImpl::IsUnBoundBindingSpawner(FBlueprintActionFilter const& Filter, FBlueprintActionInfo& BlueprintAction)
 {
+	bool const bIsBindingSpecificSpawner = (Cast<UBlueprintBoundEventNodeSpawner>(BlueprintAction.NodeSpawner) != nullptr);
+
 	bool bIsFilteredOut = false;
-	if (UBlueprintBoundEventNodeSpawner const* BoundSpawner = Cast<UBlueprintBoundEventNodeSpawner>(BlueprintAction.NodeSpawner))
+	if (bIsBindingSpecificSpawner)
 	{
-		bIsFilteredOut = true;
-		for (UObject const* SelectedObj : Filter.Context.SelectedObjects)
-		{
-			if (BoundSpawner->CanBind(SelectedObj))
-			{
-				bIsFilteredOut = false;
-				break;
-			}
-		}
+		bIsFilteredOut = (BlueprintAction.GetBindings().Num() <= 0);
 	}
 	return bIsFilteredOut;
 }
@@ -1044,12 +1038,32 @@ FBlueprintActionInfo::FBlueprintActionInfo(UObject const* ActionOwnerIn, UBluepr
 	checkSlow(NodeSpawner != nullptr);
 }
 
+//------------------------------------------------------------------------------
+FBlueprintActionInfo::FBlueprintActionInfo(FBlueprintActionInfo const& Rhs, IBlueprintNodeBinder::FBindingSet const& InBindings)
+	: NodeSpawner(Rhs.NodeSpawner)
+	, Bindings(InBindings)
+	, ActionOwner(Rhs.ActionOwner)
+	, CacheFlags(Rhs.CacheFlags)
+	, CachedOwnerClass(Rhs.CachedOwnerClass)
+	, CachedActionField(Rhs.CachedActionField)
+	, CachedActionProperty(Rhs.CachedActionProperty)
+	, CachedActionFunction(Rhs.CachedActionFunction)
+{
+	checkSlow(NodeSpawner != nullptr);
+}
+
 //------------------------------------------------------------------------------
 UObject const* FBlueprintActionInfo::GetActionOwner()
 {
 	return ActionOwner;
 }
 
+//------------------------------------------------------------------------------
+IBlueprintNodeBinder::FBindingSet const& FBlueprintActionInfo::GetBindings() const
+{
+	return Bindings;
+}
+
 //------------------------------------------------------------------------------
 UClass const* FBlueprintActionInfo::GetOwnerClass()
 {
@@ -1174,7 +1188,7 @@ FBlueprintActionFilter::FBlueprintActionFilter(uint32 Flags/*= 0x00*/)
 
 	AddRejectionTest(FRejectionTestDelegate::CreateStatic(IsFilteredNodeType, (Flags & BPFILTER_RejectPermittedNodeSubClasses) != 0));
 	AddRejectionTest(FRejectionTestDelegate::CreateStatic(IsNonTargetMemeber, (Flags & BPFILTER_RejectPersistentNonTargetFields) == 0));
-	AddRejectionTest(FRejectionTestDelegate::CreateStatic(IsBoundToUnselectedObject));
+	AddRejectionTest(FRejectionTestDelegate::CreateStatic(IsUnBoundBindingSpawner));
 	AddRejectionTest(FRejectionTestDelegate::CreateStatic(IsOutOfScopeLocalVariable));
 
 	// @TODO: account for all K2ActionMenuBuilder checks...
diff --git a/Engine/Source/Editor/BlueprintGraph/Private/BlueprintBoundEventNodeSpawner.cpp b/Engine/Source/Editor/BlueprintGraph/Private/BlueprintBoundEventNodeSpawner.cpp
--- a/Engine/Source/Editor/BlueprintGraph/Private/BlueprintBoundEventNodeSpawner.cpp
+++ b/Engine/Source/Editor/BlueprintGraph/Private/BlueprintBoundEventNodeSpawner.cpp
@@ -35,13 +35,12 @@ UBlueprintBoundEventNodeSpawner::UBlueprintBoundEventNodeSpawner(class FPostCons
 }
 
 //------------------------------------------------------------------------------
-UEdGraphNode* UBlueprintBoundEventNodeSpawner::Invoke(UEdGraph* ParentGraph, FVector2D const Location) const
+UEdGraphNode* UBlueprintBoundEventNodeSpawner::Invoke(UEdGraph* ParentGraph, FBindingSet const& Bindings, FVector2D const Location) const
 {
 	UK2Node_Event* EventNode = nullptr;
-	if (IsBindingSet())
+	if (Bindings.Num() > 0)
 	{
-		EventNode = CastChecked<UK2Node_Event>(Super::Invoke(ParentGraph, Location));
-		Bind(EventNode);
+		EventNode = CastChecked<UK2Node_Event>(Super::Invoke(ParentGraph, Bindings, Location));
 	}
 	return EventNode;
 }
@@ -70,14 +69,14 @@ FText UBlueprintBoundEventNodeSpawner::GetDefaultMenuCategory() const
 }
 
 //------------------------------------------------------------------------------
-UK2Node_Event const* UBlueprintBoundEventNodeSpawner::FindPreExistingEvent(UBlueprint* Blueprint) const
+UK2Node_Event const* UBlueprintBoundEventNodeSpawner::FindPreExistingEvent(UBlueprint* Blueprint, FBindingSet const& Bindings) const
 {
 	UK2Node_Event const* PreExistingEvent = nullptr;
 
 	UObject const* BoundObject = nullptr;
-	if (IsBindingSet())
+	if (Bindings.Num() > 0)
 	{
-		BoundObject = BoundObjects.CreateConstIterator()->Get();
+		BoundObject = Bindings.CreateConstIterator()->Get();
 	}
 
 	if (BoundObject != nullptr)
@@ -95,7 +94,7 @@ UK2Node_Event const* UBlueprintBoundEventNodeSpawner::FindPreExistingEvent(UBlue
 }
 
 //------------------------------------------------------------------------------
-bool UBlueprintBoundEventNodeSpawner::CanBind(UObject const* BindingCandidate) const
+bool UBlueprintBoundEventNodeSpawner::IsBindingCompatible(UObject const* BindingCandidate) const
 {
 	bool bCanBind = false;
 
diff --git a/Engine/Source/Editor/BlueprintGraph/Private/BlueprintComponentNodeSpawner.cpp b/Engine/Source/Editor/BlueprintGraph/Private/BlueprintComponentNodeSpawner.cpp
--- a/Engine/Source/Editor/BlueprintGraph/Private/BlueprintComponentNodeSpawner.cpp
+++ b/Engine/Source/Editor/BlueprintGraph/Private/BlueprintComponentNodeSpawner.cpp
@@ -37,7 +37,7 @@ UBlueprintComponentNodeSpawner::UBlueprintComponentNodeSpawner(class FPostConstr
 //------------------------------------------------------------------------------
 // Evolved from a combination of FK2ActionMenuBuilder::CreateAddComponentAction()
 // and FEdGraphSchemaAction_K2AddComponent::PerformAction().
-UEdGraphNode* UBlueprintComponentNodeSpawner::Invoke(UEdGraph* ParentGraph, FVector2D const Location) const
+UEdGraphNode* UBlueprintComponentNodeSpawner::Invoke(UEdGraph* ParentGraph, FBindingSet const& Bindings, FVector2D const Location) const
 {
 	check(ComponentClass != nullptr);
 	
@@ -53,7 +53,7 @@ UEdGraphNode* UBlueprintComponentNodeSpawner::Invoke(UEdGraph* ParentGraph, FVec
 	};
 	FCustomizeNodeDelegate PostSpawnDelegate = FCustomizeNodeDelegate::CreateStatic(PostSpawnLambda, CustomizeNodeDelegate);
 	
-	UK2Node_AddComponent* NewNode = CastChecked<UK2Node_AddComponent>(Super::Invoke(ParentGraph, Location, PostSpawnDelegate));
+	UK2Node_AddComponent* NewNode = CastChecked<UK2Node_AddComponent>(Super::Invoke(ParentGraph, Bindings, Location, PostSpawnDelegate));
 
 	bool const bIsTemplateNode = (ParentGraph->GetOutermost() == GetTransientPackage());
 	if (!bIsTemplateNode)
diff --git a/Engine/Source/Editor/BlueprintGraph/Private/BlueprintDelegateNodeSpawner.cpp b/Engine/Source/Editor/BlueprintGraph/Private/BlueprintDelegateNodeSpawner.cpp
--- a/Engine/Source/Editor/BlueprintGraph/Private/BlueprintDelegateNodeSpawner.cpp
+++ b/Engine/Source/Editor/BlueprintGraph/Private/BlueprintDelegateNodeSpawner.cpp
@@ -36,7 +36,7 @@ UBlueprintDelegateNodeSpawner::UBlueprintDelegateNodeSpawner(class FPostConstruc
 }
 
 //------------------------------------------------------------------------------
-UEdGraphNode* UBlueprintDelegateNodeSpawner::Invoke(UEdGraph* ParentGraph, FVector2D const Location) const
+UEdGraphNode* UBlueprintDelegateNodeSpawner::Invoke(UEdGraph* ParentGraph, FBindingSet const& Bindings, FVector2D const Location) const
 {
 	auto PostSpawnSetupLambda = [](UEdGraphNode* NewNode, bool bIsTemplateNode, UMulticastDelegateProperty const* Property, FCustomizeNodeDelegate UserDelegate)
 	{
@@ -53,7 +53,7 @@ UEdGraphNode* UBlueprintDelegateNodeSpawner::Invoke(UEdGraph* ParentGraph, FVect
 	};
 
 	FCustomizeNodeDelegate PostSpawnSetupDelegate = FCustomizeNodeDelegate::CreateStatic(PostSpawnSetupLambda, GetProperty(), CustomizeNodeDelegate);
-	return Super::Invoke(ParentGraph, Location, PostSpawnSetupDelegate);
+	return Super::Invoke(ParentGraph, Bindings, Location, PostSpawnSetupDelegate);
 }
 
 //------------------------------------------------------------------------------
diff --git a/Engine/Source/Editor/BlueprintGraph/Private/BlueprintEventNodeSpawner.cpp b/Engine/Source/Editor/BlueprintGraph/Private/BlueprintEventNodeSpawner.cpp
--- a/Engine/Source/Editor/BlueprintGraph/Private/BlueprintEventNodeSpawner.cpp
+++ b/Engine/Source/Editor/BlueprintGraph/Private/BlueprintEventNodeSpawner.cpp
@@ -89,13 +89,13 @@ UBlueprintEventNodeSpawner::UBlueprintEventNodeSpawner(class FPostConstructIniti
 }
 
 //------------------------------------------------------------------------------
-UEdGraphNode* UBlueprintEventNodeSpawner::Invoke(UEdGraph* ParentGraph, FVector2D const Location) const
+UEdGraphNode* UBlueprintEventNodeSpawner::Invoke(UEdGraph* ParentGraph, FBindingSet const& Bindings, FVector2D const Location) const
 {
 	check(ParentGraph != nullptr);
 	UBlueprint* Blueprint = FBlueprintEditorUtils::FindBlueprintForGraphChecked(ParentGraph);
 	// look to see if a node for this event already exists (only one node is
 	// allowed per event, per blueprint)
-	UK2Node_Event const* PreExistingNode = FindPreExistingEvent(Blueprint);
+	UK2Node_Event const* PreExistingNode = FindPreExistingEvent(Blueprint, Bindings);
 	// @TODO: casting away the const is bad form!
 	UK2Node_Event* EventNode = const_cast<UK2Node_Event*>(PreExistingNode);
 
@@ -129,7 +129,7 @@ UEdGraphNode* UBlueprintEventNodeSpawner::Invoke(UEdGraph* ParentGraph, FVector2
 		};
 
 		FCustomizeNodeDelegate PostSpawnDelegate = FCustomizeNodeDelegate::CreateStatic(PostSpawnLambda, EventFunc, EventName, CustomizeNodeDelegate);
-		EventNode = Cast<UK2Node_Event>(Super::Invoke(ParentGraph, Location, PostSpawnDelegate));
+		EventNode = Cast<UK2Node_Event>(Super::Invoke(ParentGraph, Bindings, Location, PostSpawnDelegate));
 	}
 	// else, a node for this event already exists, and we should return that 
 	// (the FBlueprintActionMenuItem should detect this and focus in on it).
@@ -184,7 +184,7 @@ UFunction const* UBlueprintEventNodeSpawner::GetEventFunction() const
 }
 
 //------------------------------------------------------------------------------
-UK2Node_Event const* UBlueprintEventNodeSpawner::FindPreExistingEvent(UBlueprint* Blueprint) const
+UK2Node_Event const* UBlueprintEventNodeSpawner::FindPreExistingEvent(UBlueprint* Blueprint, FBindingSet const& /*Bindings*/) const
 {
 	UK2Node_Event* PreExistingNode = nullptr;
 
diff --git a/Engine/Source/Editor/BlueprintGraph/Private/BlueprintFunctionNodeSpawner.cpp b/Engine/Source/Editor/BlueprintGraph/Private/BlueprintFunctionNodeSpawner.cpp
--- a/Engine/Source/Editor/BlueprintGraph/Private/BlueprintFunctionNodeSpawner.cpp
+++ b/Engine/Source/Editor/BlueprintGraph/Private/BlueprintFunctionNodeSpawner.cpp
@@ -82,7 +82,7 @@ static bool BlueprintFunctionNodeSpawnerImpl::BindFunctionNode(UK2Node_CallFunct
 
 	FVector2D BindingPos = CalculateBindingPosition(NewNode);
 	UEdGraph* ParentGraph = NewNode->GetGraph();
-	NodeType* BindingNode = CastChecked<NodeType>(BindingSpawner->Invoke(ParentGraph, BindingPos));
+	NodeType* BindingNode = CastChecked<NodeType>(BindingSpawner->Invoke(ParentGraph, IBlueprintNodeBinder::FBindingSet(), BindingPos));
 
 	BindingOffset.Y += UEdGraphSchema_K2::EstimateNodeHeight(BindingNode);
 
@@ -193,7 +193,7 @@ UBlueprintFunctionNodeSpawner::UBlueprintFunctionNodeSpawner(class FPostConstruc
 }
 
 //------------------------------------------------------------------------------
-UEdGraphNode* UBlueprintFunctionNodeSpawner::Invoke(UEdGraph* ParentGraph, FVector2D const Location) const
+UEdGraphNode* UBlueprintFunctionNodeSpawner::Invoke(UEdGraph* ParentGraph, FBindingSet const& Bindings, FVector2D const Location) const
 {
 	auto PostSpawnSetupLambda = [](UEdGraphNode* NewNode, bool bIsTemplateNode, UFunction const* Function, FCustomizeNodeDelegate UserDelegate)
 	{
@@ -207,12 +207,12 @@ UEdGraphNode* UBlueprintFunctionNodeSpawner::Invoke(UEdGraph* ParentGraph, FVect
 		UserDelegate.ExecuteIfBound(NewNode, bIsTemplateNode);
 	};
 
-	FCustomizeNodeDelegate PostSpawnSetupDelegate = FCustomizeNodeDelegate::CreateStatic(PostSpawnSetupLambda, Function, CustomizeNodeDelegate);
-	UEdGraphNode* SpawnedNode = Super::Invoke(ParentGraph, Location, PostSpawnSetupDelegate);
-
-	// if this spawner was set up to spawn 
+	// if this spawner was set up to spawn a bound node, reset this so the 
+	// bound nodes get positioned properly
 	BlueprintFunctionNodeSpawnerImpl::BindingOffset = FVector2D::ZeroVector;
-	Bind(SpawnedNode);
+
+	FCustomizeNodeDelegate PostSpawnSetupDelegate = FCustomizeNodeDelegate::CreateStatic(PostSpawnSetupLambda, Function, CustomizeNodeDelegate);
+	UEdGraphNode* SpawnedNode = Super::Invoke(ParentGraph, Bindings, Location, PostSpawnSetupDelegate);
 
 	return SpawnedNode;
 }
@@ -235,14 +235,24 @@ FText UBlueprintFunctionNodeSpawner::GetDefaultMenuCategory() const
 FText UBlueprintFunctionNodeSpawner::GetDefaultMenuTooltip() const
 {
 	check(Function != nullptr);
-	return FText::FromString(UK2Node_CallFunction::GetDefaultTooltipForFunction(Function));
+
+	FText Tooltip = FText::FromString(UK2Node_CallFunction::GetDefaultTooltipForFunction(Function));
+	if (Tooltip.IsEmpty())
+	{
+		Tooltip = GetDefaultMenuName();
+	}
+	return Tooltip;
 }
 
 //------------------------------------------------------------------------------
 FString UBlueprintFunctionNodeSpawner::GetDefaultSearchKeywords() const
 {
 	check(Function != nullptr);
-	return UK2Node_CallFunction::GetKeywordsForFunction(Function);
+	
+	FString SearchKeywords = UK2Node_CallFunction::GetKeywordsForFunction(Function);
+	// add at least one character, so that the menu item doesn't attempt to
+	// ping a template node
+	return SearchKeywords.AppendChar(TEXT(' '));
 }
 
 //------------------------------------------------------------------------------
@@ -253,12 +263,12 @@ bool UBlueprintFunctionNodeSpawner::CanBindMultipleObjects() const
 }
 
 //------------------------------------------------------------------------------
-bool UBlueprintFunctionNodeSpawner::CanBind(UObject const* BindingCandidate) const
+bool UBlueprintFunctionNodeSpawner::IsBindingCompatible(UObject const* BindingCandidate) const
 {
 	bool bCanBind = false;
 	if (Function != nullptr)
 	{
-		// @TODO: don't know exactly why we can only bind non-pure or cost 
+		// @TODO: don't know exactly why we can only bind non-pure/const 
 		//        functions... this is mirrored after FK2ActionMenuBuilder::GetFunctionCallsOnSelectedActors()
 		//        and FK2ActionMenuBuilder::GetFunctionCallsOnSelectedComponents(),
 		//        where we make the same stipulation
diff --git a/Engine/Source/Editor/BlueprintGraph/Private/BlueprintNodeSpawner.cpp b/Engine/Source/Editor/BlueprintGraph/Private/BlueprintNodeSpawner.cpp
--- a/Engine/Source/Editor/BlueprintGraph/Private/BlueprintNodeSpawner.cpp
+++ b/Engine/Source/Editor/BlueprintGraph/Private/BlueprintNodeSpawner.cpp
@@ -69,9 +69,9 @@ UBlueprintNodeSpawner::~UBlueprintNodeSpawner()
 }
 
 //------------------------------------------------------------------------------
-UEdGraphNode* UBlueprintNodeSpawner::Invoke(UEdGraph* ParentGraph, FVector2D const Location) const
+UEdGraphNode* UBlueprintNodeSpawner::Invoke(UEdGraph* ParentGraph, FBindingSet const& Bindings, FVector2D const Location) const
 {
-	return Invoke(ParentGraph, Location, CustomizeNodeDelegate);
+	return Invoke(ParentGraph, Bindings, Location, CustomizeNodeDelegate);
 }
 
 //------------------------------------------------------------------------------
@@ -120,18 +120,6 @@ FName UBlueprintNodeSpawner::GetDefaultMenuIcon(FLinearColor& ColorOut) const
 	return FName();
 }
 
-//------------------------------------------------------------------------------
-bool UBlueprintNodeSpawner::CanBind(UObject const* BindingCandidate) const 
-{
-	return false;
-}
-
-//------------------------------------------------------------------------------
-bool UBlueprintNodeSpawner::BindToNode(UEdGraphNode* Node, UObject* Binding) const
-{
-	return false;
-}
-
 //------------------------------------------------------------------------------
 UEdGraphNode* UBlueprintNodeSpawner::GetTemplateNode(UEdGraph* Outer) const
 {	
@@ -140,7 +128,7 @@ UEdGraphNode* UBlueprintNodeSpawner::GetTemplateNode(UEdGraph* Outer) const
 }
 
 //------------------------------------------------------------------------------
-UEdGraphNode* UBlueprintNodeSpawner::Invoke(UEdGraph* ParentGraph, FVector2D const Location, FCustomizeNodeDelegate PostSpawnDelegate) const
+UEdGraphNode* UBlueprintNodeSpawner::Invoke(UEdGraph* ParentGraph, FBindingSet const& Bindings, FVector2D const Location, FCustomizeNodeDelegate PostSpawnDelegate) const
 {
 	UEdGraphNode* NewNode = nullptr;
 	if (NodeClass != nullptr)
@@ -163,6 +151,8 @@ UEdGraphNode* UBlueprintNodeSpawner::Invoke(UEdGraph* ParentGraph, FVector2D con
 			NewNode->AllocateDefaultPins();
 			NewNode->PostPlacedNewNode();
 		}
+
+		ApplyBindings(NewNode, Bindings);
 	}
 
 	return NewNode;
diff --git a/Engine/Source/Editor/BlueprintGraph/Private/BlueprintNodeTemplateCache.cpp b/Engine/Source/Editor/BlueprintGraph/Private/BlueprintNodeTemplateCache.cpp
--- a/Engine/Source/Editor/BlueprintGraph/Private/BlueprintNodeTemplateCache.cpp
+++ b/Engine/Source/Editor/BlueprintGraph/Private/BlueprintNodeTemplateCache.cpp
@@ -204,7 +204,7 @@ UEdGraphNode* FBlueprintNodeTemplateCache::GetNodeTemplate(UBlueprintNodeSpawner
 
 		if (CompatibleOuter != nullptr)
 		{
-			TemplateNode = NodeSpawner->Invoke(CompatibleOuter, FVector2D::ZeroVector);
+			TemplateNode = NodeSpawner->Invoke(CompatibleOuter, IBlueprintNodeBinder::FBindingSet(), FVector2D::ZeroVector);
 			NodeTemplateCache.Add(NodeSpawner, TemplateNode);
 		}
 	}
diff --git a/Engine/Source/Editor/BlueprintGraph/Private/BlueprintVariableNodeSpawner.cpp b/Engine/Source/Editor/BlueprintGraph/Private/BlueprintVariableNodeSpawner.cpp
--- a/Engine/Source/Editor/BlueprintGraph/Private/BlueprintVariableNodeSpawner.cpp
+++ b/Engine/Source/Editor/BlueprintGraph/Private/BlueprintVariableNodeSpawner.cpp
@@ -53,7 +53,7 @@ UBlueprintVariableNodeSpawner::UBlueprintVariableNodeSpawner(class FPostConstruc
 }
 
 //------------------------------------------------------------------------------
-UEdGraphNode* UBlueprintVariableNodeSpawner::Invoke(UEdGraph* ParentGraph, FVector2D const Location) const
+UEdGraphNode* UBlueprintVariableNodeSpawner::Invoke(UEdGraph* ParentGraph, FBindingSet const& Bindings, FVector2D const Location) const
 {
 	FCustomizeNodeDelegate PostSpawnDelegate = CustomizeNodeDelegate;
 
@@ -88,7 +88,7 @@ UEdGraphNode* UBlueprintVariableNodeSpawner::Invoke(UEdGraph* ParentGraph, FVect
 		PostSpawnDelegate = FCustomizeNodeDelegate::CreateStatic(MemberVarSetupLambda, Property, CustomizeNodeDelegate);
 	}
 
-	return Super::Invoke(ParentGraph, Location, PostSpawnDelegate);
+	return Super::Invoke(ParentGraph, Bindings, Location, PostSpawnDelegate);
 }
 
 //------------------------------------------------------------------------------
diff --git a/Engine/Source/Editor/BlueprintGraph/Private/K2Node_ActorBoundEvent.cpp b/Engine/Source/Editor/BlueprintGraph/Private/K2Node_ActorBoundEvent.cpp
--- a/Engine/Source/Editor/BlueprintGraph/Private/K2Node_ActorBoundEvent.cpp
+++ b/Engine/Source/Editor/BlueprintGraph/Private/K2Node_ActorBoundEvent.cpp
@@ -142,6 +142,7 @@ void UK2Node_ActorBoundEvent::InitializeActorBoundEventParams(AActor* InEventOwn
 		CustomFunctionName = FName( *FString::Printf(TEXT("BndEvt__%s_%s_%s"), *InEventOwner->GetName(), *GetName(), *EventSignatureName.ToString()) );
 		bOverrideFunction = false;
 		bInternalEvent = true;
+		CachedNodeTitle.MarkDirty();
 	}
 }
 
@@ -164,6 +165,7 @@ UMulticastDelegateProperty* UK2Node_ActorBoundEvent::GetTargetDelegateProperty()
 			TargetDelegateProp = NewProperty;
 			DelegatePropertyName = NewProperty->GetFName();
 			EventSignatureClass = Cast<UClass>(NewProperty->GetOuter());
+			CachedNodeTitle.MarkDirty();
 		}
 	}
 
diff --git a/Engine/Source/Editor/BlueprintGraph/Private/K2Node_ComponentBoundEvent.cpp b/Engine/Source/Editor/BlueprintGraph/Private/K2Node_ComponentBoundEvent.cpp
--- a/Engine/Source/Editor/BlueprintGraph/Private/K2Node_ComponentBoundEvent.cpp
+++ b/Engine/Source/Editor/BlueprintGraph/Private/K2Node_ComponentBoundEvent.cpp
@@ -35,6 +35,7 @@ void UK2Node_ComponentBoundEvent::InitializeComponentBoundEventParams(UObjectPro
 		CustomFunctionName = FName( *FString::Printf(TEXT("BndEvt__%s_%s_%s"), *InComponentProperty->GetName(), *GetName(), *EventSignatureName.ToString()) );
 		bOverrideFunction = false;
 		bInternalEvent = true;
+		CachedNodeTitle.MarkDirty();
 	}
 }
 
@@ -52,6 +53,7 @@ void UK2Node_ComponentBoundEvent::RegisterDynamicBinding(UDynamicBlueprintBindin
 	Binding.DelegatePropertyName = DelegatePropertyName;
 	Binding.FunctionNameToBind = CustomFunctionName;
 
+	CachedNodeTitle.MarkDirty();
 	ComponentBindingObject->ComponentDelegateBindings.Add(Binding);
 }
 
diff --git a/Engine/Source/Editor/BlueprintGraph/Public/BlueprintActionFilter.h b/Engine/Source/Editor/BlueprintGraph/Public/BlueprintActionFilter.h
--- a/Engine/Source/Editor/BlueprintGraph/Public/BlueprintActionFilter.h
+++ b/Engine/Source/Editor/BlueprintGraph/Public/BlueprintActionFilter.h
@@ -2,6 +2,8 @@
 
 #pragma once
 
+#include "BlueprintNodeBinder.h" // for IBlueprintNodeBinder::FBindingSet
+
 // Forward declarations
 class UBlueprintNodeSpawner;
 class UBlueprint;
@@ -55,6 +57,9 @@ struct BLUEPRINTGRAPH_API FBlueprintActionInfo
 	/** */
 	FBlueprintActionInfo(UObject const* ActionOwner, UBlueprintNodeSpawner const* Action);
 
+	/** */
+	FBlueprintActionInfo(FBlueprintActionInfo const& Rhs, IBlueprintNodeBinder::FBindingSet const& Bindings);
+
 	/**
 	 * Retrieves the key that the wrapped action is associated with in the 
 	 * FBlueprintActionDatabase (either a UClass, or asset object).
@@ -63,6 +68,13 @@ struct BLUEPRINTGRAPH_API FBlueprintActionInfo
 	 */
 	UObject const* GetActionOwner();
 
+	/**
+	 * 
+	 * 
+	 * @return 
+	 */
+	IBlueprintNodeBinder::FBindingSet const& GetBindings() const;
+
 	/**
 	 * Retrieves a class associated with the wrapped action. Intended to be the 
 	 * action's class "owner". Could be null if the action is keyed to an asset.
@@ -127,6 +139,9 @@ private:
 	UProperty const* CachedActionProperty;
 	/** */
 	UFunction const* CachedActionFunction;
+
+	/** */
+	IBlueprintNodeBinder::FBindingSet Bindings;
 };
 
 
diff --git a/Engine/Source/Editor/BlueprintGraph/Public/BlueprintNodeSpawner.h b/Engine/Source/Editor/BlueprintGraph/Public/BlueprintNodeSpawner.h
--- a/Engine/Source/Editor/BlueprintGraph/Public/BlueprintNodeSpawner.h
+++ b/Engine/Source/Editor/BlueprintGraph/Public/BlueprintNodeSpawner.h
@@ -24,7 +24,7 @@ class UEdGraphNode;
  * these "actions".
  */
 UCLASS(Transient)
-class BLUEPRINTGRAPH_API UBlueprintNodeSpawner : public UObject, public FBlueprintNodeBinder
+class BLUEPRINTGRAPH_API UBlueprintNodeSpawner : public UObject, public IBlueprintNodeBinder
 {
 	GENERATED_UCLASS_BODY()
 	DECLARE_DELEGATE_TwoParams(FCustomizeNodeDelegate, UEdGraphNode*, bool);
@@ -82,10 +82,11 @@ public:
 	 * this case upon use.
 	 * 
 	 * @param  ParentGraph	The graph you want the node spawned into.
+	 * @param  Bindings		
 	 * @param  Location     Where you want the new node positioned in the graph.
 	 * @return Null if it failed to spawn a node, otherwise a newly spawned node or possibly one that already existed.
 	 */
-	virtual UEdGraphNode* Invoke(UEdGraph* ParentGraph, FVector2D const Location) const;
+	virtual UEdGraphNode* Invoke(UEdGraph* ParentGraph, FBindingSet const& Bindings, FVector2D const Location) const;
 	
 	// @TODO: for favorites, generated for the spawner type + data (function, property, node, etc.)
 	//virtual FGuid GetActionHash() const;
@@ -160,22 +161,24 @@ public:
 	 */
 	UEdGraphNode* GetTemplateNode(UEdGraph* TargetGraph = nullptr) const;
 
-	// FBlueprintNodeBinder interface
-	virtual bool CanBind(UObject const* BindingCandidate) const override;
-	virtual bool BindToNode(UEdGraphNode* Node, UObject* Binding) const override;
-	// End FBlueprintNodeBinder interface
-
+	// IBlueprintNodeBinder interface
+	virtual bool   IsBindingCompatible(UObject const* BindingCandidate) const override { return false; }
+	virtual bool   CanBindMultipleObjects() const override { return false; }
 protected:
+	virtual bool   BindToNode(UEdGraphNode* Node, UObject* Binding) const override { return false; }
+	// End IBlueprintNodeBinder interface
+
 	/**
 	 * Protected Invoke() that let's sub-classes specify their own post-spawn
 	 * delegate. Creates a new node based off of the set NodeClass.
 	 * 
 	 * @param  ParentGraph			The graph you want the node spawned into.
+	 * @param  Bindings				
 	 * @param  Location				Where you want the new node positioned in the graph.
 	 * @param  PostSpawnDelegate	A delegate to run after spawning the node, but prior to allocating the node's pins.
 	 * @return Null if it failed to spawn a node (if NodeClass is null), otherwise a newly spawned node.
 	 */
-	UEdGraphNode* Invoke(UEdGraph* ParentGraph, FVector2D const Location, FCustomizeNodeDelegate PostSpawnDelegate) const;
+	UEdGraphNode* Invoke(UEdGraph* ParentGraph, FBindingSet const& Bindings, FVector2D const Location, FCustomizeNodeDelegate PostSpawnDelegate) const;
 };
 
 /*******************************************************************************
diff --git a/Engine/Source/Editor/Kismet/Private/BlueprintActionMenuBuilder.cpp b/Engine/Source/Editor/Kismet/Private/BlueprintActionMenuBuilder.cpp
--- a/Engine/Source/Editor/Kismet/Private/BlueprintActionMenuBuilder.cpp
+++ b/Engine/Source/Editor/Kismet/Private/BlueprintActionMenuBuilder.cpp
@@ -67,7 +67,7 @@ public:
 	 * @param  BoundAction	
 	 * @return 
 	 */
-	TSharedPtr<FBlueprintBoundMenuItem> MakeBoundMenuItem(TWeakPtr<FBlueprintEditor> EditorContext, UBlueprintNodeSpawner* BoundAction);
+	TSharedPtr<FBlueprintBoundMenuItem> MakeBoundMenuItem(TWeakPtr<FBlueprintEditor> EditorContext, UBlueprintNodeSpawner const* Action, IBlueprintNodeBinder::FBindingSet const& Bindings);
 	
 private:
 	/**
@@ -178,15 +178,17 @@ TSharedPtr<FEdGraphSchemaAction> FBlueprintActionMenuItemFactory::MakeDragDropMe
 }
 
 //------------------------------------------------------------------------------
-TSharedPtr<FBlueprintBoundMenuItem> FBlueprintActionMenuItemFactory::MakeBoundMenuItem(TWeakPtr<FBlueprintEditor> EditorContext, UBlueprintNodeSpawner* BoundAction)
+TSharedPtr<FBlueprintBoundMenuItem> FBlueprintActionMenuItemFactory::MakeBoundMenuItem(TWeakPtr<FBlueprintEditor> EditorContext, UBlueprintNodeSpawner const* Action, IBlueprintNodeBinder::FBindingSet const& Bindings)
 {
-	FBlueprintBoundMenuItem* NewMenuItem = new FBlueprintBoundMenuItem(BoundAction, MenuGrouping);
+	FBlueprintBoundMenuItem* NewMenuItem = new FBlueprintBoundMenuItem(Action, MenuGrouping);
+	NewMenuItem->AddBindings(Bindings);
+
 	// FBlueprintBoundMenuItem takes care of its own menu category
 	NewMenuItem->Category = FString::Printf(TEXT("%s|%s"), *RootCategory.ToString(), *NewMenuItem->Category);
 
-	NewMenuItem->MenuDescription    = GetMenuNameForAction(EditorContext, BoundAction);
-	NewMenuItem->TooltipDescription = GetTooltipForAction(EditorContext, BoundAction).ToString();
-	NewMenuItem->Keywords			= GetSearchKeywordsForAction(EditorContext, BoundAction);
+	NewMenuItem->MenuDescription = GetMenuNameForAction(EditorContext, Action);
+	NewMenuItem->TooltipDescription = GetTooltipForAction(EditorContext, Action).ToString();
+	NewMenuItem->Keywords = GetSearchKeywordsForAction(EditorContext, Action);
 
 	return MakeShareable(NewMenuItem);
 }
@@ -223,35 +225,32 @@ FText FBlueprintActionMenuItemFactory::GetCategoryForAction(TWeakPtr<FBlueprintE
 	
 	if (MenuCategory.IsEmpty())
 	{
-		// @TODO: consider moving GetMenuCategory() up into UEdGraphNode
-		if (UK2Node* NodeTemplate = Cast<UK2Node>(GetTemplateNode(Action, EditorContext)))
+		// put uncategorized function calls in a member function category
+		// (sorted by their respective classes)
+		if (UBlueprintFunctionNodeSpawner const* FuncSpawner = Cast<UBlueprintFunctionNodeSpawner>(Action))
 		{
-			MenuCategory = NodeTemplate->GetMenuCategory();
+			UFunction const* Function = FuncSpawner->GetFunction();
+			check(Function != nullptr);
+			UClass* FuncOwner = Function->GetOuterUClass();
+
+			UBlueprint* Blueprint = GetTargetBlueprint();
+			check(Blueprint != nullptr);
+			UClass* BlueprintClass = (Blueprint->SkeletonGeneratedClass != nullptr) ? Blueprint->SkeletonGeneratedClass : Blueprint->ParentClass;
+
+			// if this is NOT a self function call (self function calls
+			// don't get nested any deeper)
+			if (!BlueprintClass->IsChildOf(FuncOwner))
+			{
+				MenuCategory = FuncOwner->GetDisplayNameText();
+			}
+			MenuCategory = FText::Format(LOCTEXT("MemberFunctionsCategory", "Call Function|{0}"), MenuCategory);
 		}
-		
-		// if the category is still empty, then try and construct the category
-		// from the spawner's type
-		if (MenuCategory.IsEmpty())
+		else 
 		{
-			// put uncategorized function calls in a member function category
-			// (sorted by their respective classes)
-			if (UBlueprintFunctionNodeSpawner const* FuncSpawner = Cast<UBlueprintFunctionNodeSpawner>(Action))
+			// @TODO: consider moving GetMenuCategory() up into UEdGraphNode
+			if (UK2Node* NodeTemplate = Cast<UK2Node>(GetTemplateNode(Action, EditorContext)))
 			{
-				UFunction const* Function = FuncSpawner->GetFunction();
-				check(Function != nullptr);
-				UClass* FuncOwner = Function->GetOuterUClass();
-				
-				UBlueprint* Blueprint = GetTargetBlueprint();
-				check(Blueprint != nullptr);
-				UClass* BlueprintClass = (Blueprint->SkeletonGeneratedClass != nullptr) ? Blueprint->SkeletonGeneratedClass : Blueprint->ParentClass;
-				
-				// if this is NOT a self function call (self function calls
-				// don't get nested any deeper)
-				if (!BlueprintClass->IsChildOf(FuncOwner))
-				{
-					MenuCategory = FuncOwner->GetDisplayNameText();
-				}
-				MenuCategory = FText::Format(LOCTEXT("MemberFunctionsCategory", "Call Function|{0}"), MenuCategory);
+				MenuCategory = NodeTemplate->GetMenuCategory();
 			}
 		}
 	}
@@ -448,75 +447,60 @@ void FBlueprintActionMenuBuilderImpl::FMenuSectionDefinition::SetSectionSortOrde
 }
 // 
 //------------------------------------------------------------------------------
-TSharedPtr<FEdGraphSchemaAction> FBlueprintActionMenuBuilderImpl::FMenuSectionDefinition::MakeBoundMenuItem(TWeakPtr<FBlueprintEditor> EditorContext, FBlueprintActionInfo& DatabaseAction, TArray<UObject*> const& Bindings)
+TSharedPtr<FEdGraphSchemaAction> FBlueprintActionMenuBuilderImpl::FMenuSectionDefinition::MakeBoundMenuItem(TWeakPtr<FBlueprintEditor> EditorContext, FBlueprintActionInfo& DatabaseActionInfo, TArray<UObject*> const& PerspectiveBindings)
 {
- 	TSharedPtr<FBlueprintBoundMenuItem> MenuItem;
+	TSharedPtr<FBlueprintBoundMenuItem> MenuItem;
+	UBlueprintNodeSpawner const* DatabaseAction = DatabaseActionInfo.NodeSpawner;
 	
+	IBlueprintNodeBinder::FBindingSet CompatibleBindingsToFilter;
 	// we don't want the blueprint database growing out of control with an entry 
 	// for every object you could ever possibly bind to, so each 
-	// UBlueprintNodeSpawner comes with an interface to bind through... however,
-	// we can't mutate the action database (by binding to those actions directly)
-	// so, we have to clone and test the action post binding
-	for (auto BindingIt = Bindings.CreateConstIterator(); BindingIt;)
+	// UBlueprintNodeSpawner comes with an interface to test/bind through... 
+	for (auto BindingIt = PerspectiveBindings.CreateConstIterator(); BindingIt;)
 	{
 		UObject const* BindingObj = *BindingIt;
+		++BindingIt;
+		bool const bIsLastBinding = !BindingIt;
+
 		// check to see if this object can be bound to this action
-		if (DatabaseAction.NodeSpawner->CanBind(BindingObj))
+		if (DatabaseAction->IsBindingCompatible(BindingObj))
 		{
-			// we don't want binding to mutate the database, so we clone the 
-			// action and pass that to FBlueprintBoundMenuItem for binding;
-			// if we haven't created a mutable BoundAction yet, we need that...
-			// we do this prior to filtering because there may be tests that 
-			// reject based off bindings
-			//
-			// don't worry about outer, if this makes it through the filter,
-			// then the menu item will add a reference to keep it around
-			UBlueprintNodeSpawner* BoundAction = DuplicateObject<UBlueprintNodeSpawner>(DatabaseAction.NodeSpawner, GetTransientPackage());
-			BoundAction->CustomizeNodeDelegate = DatabaseAction.NodeSpawner->CustomizeNodeDelegate; // delegate isn't duplicated
-
-			FBlueprintActionInfo BoundActionInfo(DatabaseAction.GetActionOwner(), BoundAction);
-			do
+			// add bindings before filtering (in case tests accept/reject based off of this)
+			CompatibleBindingsToFilter.Add(BindingObj);
+
+			// if BoundAction is now "full" (meaning it can take any more 
+			// bindings), or if this is the last binding to test...
+			if (!DatabaseAction->CanBindMultipleObjects() || bIsLastBinding)
 			{
-				// add bindings before filtering (in case tests accept/reject based off of this)
-				BoundAction->AddBinding(BindingObj);
-				
-				++BindingIt;
-				bool const bIsLastBinding = !BindingIt;
-
-				// if BoundAction is now "full" (meaning it can take any more 
-				// bindings), or if this is the last binding to test...
-				if (!BoundAction->CanBindMultipleObjects() || bIsLastBinding)
+				// we don't want binding to mutate DatabaseActionInfo, so we clone  
+				// the action info, and tack on some binding data
+				FBlueprintActionInfo BoundActionInfo(DatabaseActionInfo, CompatibleBindingsToFilter);
+
+				// have to check IsFiltered() for every "fully bound" action (in
+				// case there are tests that reject based off of this), we may 
+				// test this multiple times per action (we have to make sure 
+				// that every set of bound objects pass before folding them into
+				// MenuItem)
+				bool const bPassedFilter = !Filter.IsFiltered(BoundActionInfo);
+				if (bPassedFilter)
 				{
-					// have to check IsFiltered() for every "fully bound" action (in
-					// case there are tests that reject based off of this), we may 
-					// test this multiple times per action (we have to make sure 
-					// that every set of bound objects pass before folding them into
-					// MenuItem)
-					bool const bPassedFilter = !Filter.IsFiltered(BoundActionInfo);
-					if (bPassedFilter)
+					if (!MenuItem.IsValid())
+					{
+						MenuItem = ItemFactory.MakeBoundMenuItem(EditorContext, DatabaseAction, CompatibleBindingsToFilter);
+					}
+					else
 					{
-						if (!MenuItem.IsValid())
-						{
-							MenuItem = ItemFactory.MakeBoundMenuItem(EditorContext, BoundAction);
-						}
-
-						for (TWeakObjectPtr<UObject> Object : BoundAction->GetBindings())
-						{
-							MenuItem->AddBinding(Object.Get());
-						}
+						// move these bindings over to the menu item (so we can 
+						// test the next set)
+						MenuItem->AddBindings(CompatibleBindingsToFilter);
 					}
-					BoundAction->ClearBindings();
 				}
+				CompatibleBindingsToFilter.Empty(); // do before we copy back over cached fields for DatabaseActionInfo
 
-				if (!bIsLastBinding)
-				{
-					BindingObj = *BindingIt;
-				}
-			} while (BindingIt);
-		}
-		else
-		{
-			++BindingIt;
+				// copy over any fields that got cached for filtering (with
+				// an empty binding set)
+				/*DatabaseActionInfo = FBlueprintActionInfo(BoundActionInfo, CompatibleBindingsToFilter);*/
+			}
 		}
 	}
 
diff --git a/Engine/Source/Editor/Kismet/Private/BlueprintActionMenuItem.cpp b/Engine/Source/Editor/Kismet/Private/BlueprintActionMenuItem.cpp
--- a/Engine/Source/Editor/Kismet/Private/BlueprintActionMenuItem.cpp
+++ b/Engine/Source/Editor/Kismet/Private/BlueprintActionMenuItem.cpp
@@ -62,6 +62,16 @@ FBlueprintActionMenuItem::FBlueprintActionMenuItem(UBlueprintNodeSpawner const*
 	Grouping  = MenuGrouping;
 }
 
+//------------------------------------------------------------------------------
+FBlueprintActionMenuItem::FBlueprintActionMenuItem(UBlueprintNodeSpawner const* NodeSpawner, IBlueprintNodeBinder::FBindingSet const& InBindings)
+	: IconBrush(nullptr)
+	, IconTint(FLinearColor::White)
+	, Action(NodeSpawner)
+	, Bindings(InBindings)
+{
+	check(Action != nullptr);
+}
+
 //------------------------------------------------------------------------------
 UEdGraphNode* FBlueprintActionMenuItem::PerformAction(UEdGraph* ParentGraph, UEdGraphPin* FromPin, FVector2D const Location, bool bSelectNewNode/* = true*/)
 {
@@ -88,7 +98,7 @@ UEdGraphNode* FBlueprintActionMenuItem::PerformAction(UEdGraph* ParentGraph, UEd
 	}
 
 	// this could return an existing node
-	UEdGraphNode* SpawnedNode = Action->Invoke(ParentGraph, Location);
+	UEdGraphNode* SpawnedNode = Action->Invoke(ParentGraph, Bindings, Location);
 	
 	// if a returned node hasn't been added to the graph yet (it must have been freshly spawned)
 	if (ParentGraph->Nodes.Find(SpawnedNode) == INDEX_NONE)
diff --git a/Engine/Source/Editor/Kismet/Private/BlueprintActionMenuUtils.cpp b/Engine/Source/Editor/Kismet/Private/BlueprintActionMenuUtils.cpp
--- a/Engine/Source/Editor/Kismet/Private/BlueprintActionMenuUtils.cpp
+++ b/Engine/Source/Editor/Kismet/Private/BlueprintActionMenuUtils.cpp
@@ -42,7 +42,7 @@ namespace BlueprintActionMenuUtilsImpl
 //------------------------------------------------------------------------------
 static bool BlueprintActionMenuUtilsImpl::IsUnBoundSpawner(FBlueprintActionFilter const& /*Filter*/, FBlueprintActionInfo& BlueprintAction)
 {
-	return !BlueprintAction.NodeSpawner->IsBindingSet();
+	return (BlueprintAction.GetBindings().Num() <= 0);
 }
 
 //------------------------------------------------------------------------------
@@ -100,7 +100,7 @@ void FBlueprintActionMenuUtils::MakePaletteMenu(FBlueprintActionContext const& C
 void FBlueprintActionMenuUtils::MakeContextMenu(FBlueprintActionContext const& Context, bool bIsContextSensitive, FBlueprintActionMenuBuilder& MenuOut)
 {
 	using namespace BlueprintActionMenuUtilsImpl;
-	
+
 	//--------------------------------------
 	// Composing Filters
 	//--------------------------------------
diff --git a/Engine/Source/Editor/Kismet/Private/BlueprintBoundMenuItem.cpp b/Engine/Source/Editor/Kismet/Private/BlueprintBoundMenuItem.cpp
--- a/Engine/Source/Editor/Kismet/Private/BlueprintBoundMenuItem.cpp
+++ b/Engine/Source/Editor/Kismet/Private/BlueprintBoundMenuItem.cpp
@@ -20,7 +20,7 @@ namespace BlueprintBoundMenuItemImpl
 	 * @param  BoundSpawner	
 	 * @return 
 	 */
-	static FText GetMenuCategoryFromAction(UBlueprintNodeSpawner* BoundSpawner);
+	static FText GetMenuCategoryFromAction(UBlueprintNodeSpawner const* BoundSpawner, IBlueprintNodeBinder::FBindingSet const& Bindings);
 
 	/**
 	 * 
@@ -28,22 +28,21 @@ namespace BlueprintBoundMenuItemImpl
 	 * @param  ObjectsList	
 	 * @return 
 	 */
-	static UClass* FindCommonBaseClass(TSet< TWeakObjectPtr<UObject> > const& ObjectsList);
+	static UClass* FindCommonBaseClass(IBlueprintNodeBinder::FBindingSet const& ObjectsList);
 }
 
 //------------------------------------------------------------------------------
-static FText BlueprintBoundMenuItemImpl::GetMenuCategoryFromAction(UBlueprintNodeSpawner* BoundSpawner)
+static FText BlueprintBoundMenuItemImpl::GetMenuCategoryFromAction(UBlueprintNodeSpawner const* BoundSpawner, IBlueprintNodeBinder::FBindingSet const& Bindings)
 {
-	TSet< TWeakObjectPtr<UObject> > const& BoundObjects = BoundSpawner->GetBindings();
-	check(BoundObjects.Num() > 0);
+	check(Bindings.Num() > 0);
 
 	FText BoundObjectName;
-	if (BoundObjects.Num() > 1)
+	if (Bindings.Num() > 1)
 	{
-		UClass* CommonBaseClass = FindCommonBaseClass(BoundObjects);
+		UClass* CommonBaseClass = FindCommonBaseClass(Bindings);
 		BoundObjectName = FText::Format(LOCTEXT("MultipleBindingsCategoryPostfix", "Selected {0}s"), CommonBaseClass->GetDisplayNameText());
 	}
-	else if (UObject const* BoundObject = BoundObjects.CreateConstIterator()->Get())
+	else if (UObject const* BoundObject = Bindings.CreateConstIterator()->Get())
 	{
 		BoundObjectName = FText::FromName(BoundObject->GetFName());
 	}
@@ -62,7 +61,7 @@ static FText BlueprintBoundMenuItemImpl::GetMenuCategoryFromAction(UBlueprintNod
 }
 
 //------------------------------------------------------------------------------
-static UClass* BlueprintBoundMenuItemImpl::FindCommonBaseClass(TSet< TWeakObjectPtr<UObject> > const& ObjectsList)
+static UClass* BlueprintBoundMenuItemImpl::FindCommonBaseClass(IBlueprintNodeBinder::FBindingSet const& ObjectsList)
 {
 	UClass* CommonClass = nullptr;
 	for (TWeakObjectPtr<UObject> BoundObjPtr : ObjectsList)
@@ -95,32 +94,34 @@ static UClass* BlueprintBoundMenuItemImpl::FindCommonBaseClass(TSet< TWeakObject
  ******************************************************************************/
 
 //------------------------------------------------------------------------------
-FBlueprintBoundMenuItem::FBlueprintBoundMenuItem(UBlueprintNodeSpawner* BoundSpawnerIn, int32 MenuGrouping/* = 0*/)
+FBlueprintBoundMenuItem::FBlueprintBoundMenuItem(UBlueprintNodeSpawner const* BoundSpawnerIn, int32 MenuGrouping/* = 0*/)
 	: BoundSpawner(BoundSpawnerIn) 
 {
 	Grouping = MenuGrouping;
-	Category = BlueprintBoundMenuItemImpl::GetMenuCategoryFromAction(BoundSpawnerIn).ToString();
+	// category will be generated as we add bindings
 }
 
 //------------------------------------------------------------------------------
 UEdGraphNode* FBlueprintBoundMenuItem::PerformAction(UEdGraph* ParentGraph, UEdGraphPin* FromPin, FVector2D const Location, bool bSelectNewNode/* = true*/)
 {
-	FBlueprintActionMenuItem BlueprintActionItem(BoundSpawner, nullptr, FSlateColor(FLinearColor::White));
-
 	UEdGraphNode* NewNode = nullptr;
+
+	IBlueprintNodeBinder::FBindingSet Bindings;
 	for (auto BoundObjIt = BoundObjects.CreateConstIterator(); BoundObjIt; )
 	{
 		do
 		{
 			if (BoundObjIt->IsValid())
 			{
-				BoundSpawner->AddBinding(BoundObjIt->Get());
+				Bindings.Add(BoundObjIt->Get());
 			}
 			++BoundObjIt;
 
-		} while (BoundObjIt && (BoundSpawner->CanBindMultipleObjects() || !BoundSpawner->IsBindingSet()));
+		} while ( BoundObjIt && (BoundSpawner->CanBindMultipleObjects() || (Bindings.Num() == 0)) );
 
+		FBlueprintActionMenuItem BlueprintActionItem(BoundSpawner, Bindings);
 		NewNode = BlueprintActionItem.PerformAction(ParentGraph, FromPin, Location, bSelectNewNode);
+		// @TODO: select ALL spawned nodes
 	}
 	
 	return NewNode;
@@ -156,19 +157,16 @@ void FBlueprintBoundMenuItem::AddReferencedObjects(FReferenceCollector& Collecto
 }
 
 //------------------------------------------------------------------------------
-bool FBlueprintBoundMenuItem::AddBinding(UObject const* Binding)
+void FBlueprintBoundMenuItem::AddBindings(IBlueprintNodeBinder::FBindingSet const& BindingSet)
 {
-	bool const bIsCompatible = (Binding != nullptr) && BoundSpawner->CanBind(Binding);
-	if (bIsCompatible)
-	{
-		BoundObjects.Add(Binding);
-	}
-	return bIsCompatible;
+	BoundObjects.Append(BindingSet);
+	Category = BlueprintBoundMenuItemImpl::GetMenuCategoryFromAction(BoundSpawner, BoundObjects).ToString();
 }
 
 //------------------------------------------------------------------------------
 FSlateBrush const* FBlueprintBoundMenuItem::GetMenuIcon(FSlateColor& ColorOut)
 {
+	// @TODO: 
 	return nullptr;
 }
 
diff --git a/Engine/Source/Editor/Kismet/Public/BlueprintActionMenuItem.h b/Engine/Source/Editor/Kismet/Public/BlueprintActionMenuItem.h
--- a/Engine/Source/Editor/Kismet/Public/BlueprintActionMenuItem.h
+++ b/Engine/Source/Editor/Kismet/Public/BlueprintActionMenuItem.h
@@ -4,6 +4,7 @@
 
 #include "EdGraph/EdGraphSchema.h" // for FEdGraphSchemaAction
 #include "SlateColor.h"
+#include "BlueprintNodeBinder.h" // for IBlueprintNodeBinder::FBindingSet
 #include "BlueprintActionMenuItem.generated.h"
 
 // Forward declarations
@@ -28,6 +29,7 @@ public:
 	/** Constructors */
 	FBlueprintActionMenuItem() : IconBrush(nullptr), IconTint(FLinearColor::White), Action(nullptr) {}
 	FBlueprintActionMenuItem(UBlueprintNodeSpawner const* NodeSpawner, FSlateBrush const* MenuIcon, FSlateColor const& IconTint, int32 MenuGrouping = 0);
+	FBlueprintActionMenuItem(UBlueprintNodeSpawner const* NodeSpawner, IBlueprintNodeBinder::FBindingSet const& Bindings);
 	
 	// FEdGraphSchemaAction interface
 	virtual FName         GetTypeId() const final { return StaticGetTypeId(); }
@@ -52,4 +54,6 @@ private:
 	FSlateColor IconTint;
 	/** Specialized node-spawner, that comprises the action portion of this menu entry. */
 	UBlueprintNodeSpawner const* Action;
+	/** */
+	IBlueprintNodeBinder::FBindingSet Bindings;
 };
diff --git a/Engine/Source/Editor/Kismet/Public/BlueprintBoundMenuItem.h b/Engine/Source/Editor/Kismet/Public/BlueprintBoundMenuItem.h
--- a/Engine/Source/Editor/Kismet/Public/BlueprintBoundMenuItem.h
+++ b/Engine/Source/Editor/Kismet/Public/BlueprintBoundMenuItem.h
@@ -4,6 +4,7 @@
 
 #include "EdGraph/EdGraphSchema.h" // for FEdGraphSchemaAction
 #include "SlateColor.h"
+#include "BlueprintNodeBinder.h" // for IBlueprintNodeBinder::FBindingSet
 #include "BlueprintBoundMenuItem.generated.h"
 
 // Forward declarations
@@ -27,7 +28,7 @@ public:
 	
 	/** Constructors */
 	FBlueprintBoundMenuItem() : BoundSpawner(nullptr) {}
-	FBlueprintBoundMenuItem(UBlueprintNodeSpawner* BoundSpawner, int32 MenuGrouping = 0);
+	FBlueprintBoundMenuItem(UBlueprintNodeSpawner const* BoundSpawner, int32 MenuGrouping = 0);
 	
 	// FEdGraphSchemaAction interface
 	virtual FName         GetTypeId() const final { return StaticGetTypeId(); }
@@ -36,7 +37,12 @@ public:
 	virtual void          AddReferencedObjects(FReferenceCollector& Collector) final;
 	// End FEdGraphSchemaAction interface
 
-	bool AddBinding(UObject const* Binding);
+	/**
+	 * 
+	 * 
+	 * @param  BindingSet	
+	 */
+	void AddBindings(IBlueprintNodeBinder::FBindingSet const& BindingSet);
 
 	/**
 	 * Retrieves the icon brush for this menu entry (to be displayed alongside
@@ -49,7 +55,7 @@ public:
 
 private:
 	/** Instanced node-spawner, that comprises the action portion of this menu entry. */
-	UBlueprintNodeSpawner* BoundSpawner;
+	UBlueprintNodeSpawner const* BoundSpawner;
 	/** */
-	TSet< TWeakObjectPtr<UObject> > BoundObjects;
+	IBlueprintNodeBinder::FBindingSet BoundObjects;
 };
