ttp345417 LIVE: BP: CRASH: Renaming a Custom Event to the same name as a Function will cause a crash

diff --git a/Engine/Source/Editor/Kismet/Private/BlueprintDetailsCustomization.cpp b/Engine/Source/Editor/Kismet/Private/BlueprintDetailsCustomization.cpp
--- a/Engine/Source/Editor/Kismet/Private/BlueprintDetailsCustomization.cpp
+++ b/Engine/Source/Editor/Kismet/Private/BlueprintDetailsCustomization.cpp
@@ -4778,21 +4778,26 @@ FText FBlueprintGraphNodeDetails::OnGetName() const
 	return Name;
 }
 
-void FBlueprintGraphNodeDetails::OnNameChanged(const FText& InNewText)
+struct FGraphNodeNameValidatorHelper
 {
-	bIsNodeNameInvalid = true;
-
-	if( GraphNodePtr.IsValid() && BlueprintEditorPtr.IsValid() )
+	static EValidatorResult Validate(TWeakObjectPtr<UEdGraphNode> GraphNodePtr, TWeakPtr<FBlueprintEditor> BlueprintEditorPtr, const FString& NewName)
 	{
-		FName NodeName( *GraphNodePtr->GetNodeTitle(ENodeTitleType::EditableTitle).ToString() );
+		check(GraphNodePtr.IsValid() && BlueprintEditorPtr.IsValid());
 		TSharedPtr<INameValidatorInterface> NameValidator = GraphNodePtr->MakeNameValidator();
-
-		if( !NameValidator.IsValid() )
+		if (!NameValidator.IsValid())
 		{
+			const FName NodeName(*GraphNodePtr->GetNodeTitle(ENodeTitleType::EditableTitle).ToString());
 			NameValidator = MakeShareable(new FKismetNameValidator(BlueprintEditorPtr.Pin()->GetBlueprintObj(), NodeName));
 		}
+		return NameValidator->IsValid(NewName);
+	}
+};
 
-		EValidatorResult ValidatorResult = NameValidator->IsValid(InNewText.ToString());
+void FBlueprintGraphNodeDetails::OnNameChanged(const FText& InNewText)
+{
+	if( GraphNodePtr.IsValid() && BlueprintEditorPtr.IsValid() )
+	{
+		const EValidatorResult ValidatorResult = FGraphNodeNameValidatorHelper::Validate(GraphNodePtr, BlueprintEditorPtr, InNewText.ToString());
 		if(ValidatorResult == EValidatorResult::AlreadyInUse)
 		{
 			NameEditableTextBox->SetError(FText::Format(LOCTEXT("RenameFailed_InUse", "{0} is in use by another variable or function!"), InNewText));
@@ -4807,7 +4812,6 @@ void FBlueprintGraphNodeDetails::OnNameChanged(const FText& InNewText)
 		}
 		else
 		{
-			bIsNodeNameInvalid = false;
 			NameEditableTextBox->SetError(FText::GetEmpty());
 		}
 	}
@@ -4815,9 +4819,12 @@ void FBlueprintGraphNodeDetails::OnNameChanged(const FText& InNewText)
 
 void FBlueprintGraphNodeDetails::OnNameCommitted(const FText& InNewText, ETextCommit::Type InTextCommit)
 {
-	if( BlueprintEditorPtr.IsValid() && GraphNodePtr.IsValid() )
+	if (BlueprintEditorPtr.IsValid() && GraphNodePtr.IsValid())
 	{
-		BlueprintEditorPtr.Pin()->OnNodeTitleCommitted( InNewText, InTextCommit, GraphNodePtr.Get() );
+		if (FGraphNodeNameValidatorHelper::Validate(GraphNodePtr, BlueprintEditorPtr, InNewText.ToString()) == EValidatorResult::Ok)
+		{
+			BlueprintEditorPtr.Pin()->OnNodeTitleCommitted(InNewText, InTextCommit, GraphNodePtr.Get());
+		}
 	}
 }
 
diff --git a/Engine/Source/Editor/Kismet/Private/BlueprintDetailsCustomization.h b/Engine/Source/Editor/Kismet/Private/BlueprintDetailsCustomization.h
--- a/Engine/Source/Editor/Kismet/Private/BlueprintDetailsCustomization.h
+++ b/Engine/Source/Editor/Kismet/Private/BlueprintDetailsCustomization.h
@@ -716,8 +716,6 @@ private:
 
 	/** The widget used when editing the name */ 
 	TSharedPtr<SEditableTextBox> NameEditableTextBox;
-	/** Flag to indicate whether or not the variable name is invalid */
-	bool bIsNodeNameInvalid;
 	/** The target GraphNode */
 	TWeakObjectPtr<UEdGraphNode> GraphNodePtr;
 	/** Weak reference to the Blueprint editor */
