FSCSEditorViewportClient::GetWidgetCoordSystem more robust code ttp345600

diff --git a/Engine/Source/Editor/Kismet/Private/SCSEditorViewportClient.cpp b/Engine/Source/Editor/Kismet/Private/SCSEditorViewportClient.cpp
--- a/Engine/Source/Editor/Kismet/Private/SCSEditorViewportClient.cpp
+++ b/Engine/Source/Editor/Kismet/Private/SCSEditorViewportClient.cpp
@@ -637,15 +637,17 @@ FMatrix FSCSEditorViewportClient::GetWidgetCoordSystem() const
 	if( GetWidgetCoordSystemSpace() == COORD_Local )
 	{
 		AActor* PreviewActor = GetPreviewActor();
-		if(PreviewActor)
+		auto BlueprintEditor = BlueprintEditorPtr.Pin();
+		if (PreviewActor && BlueprintEditor.IsValid())
 		{
-			TArray<FSCSEditorTreeNodePtrType> SelectedNodes = BlueprintEditorPtr.Pin()->GetSelectedSCSEditorTreeNodes();
+			TArray<FSCSEditorTreeNodePtrType> SelectedNodes = BlueprintEditor->GetSelectedSCSEditorTreeNodes();
 			if(SelectedNodes.Num() > 0)
 			{
-				USceneComponent* SceneComp = Cast<USceneComponent>(SelectedNodes.Last().Get()->FindComponentInstanceInActor(PreviewActor, true));
+				const auto SelectedNode = SelectedNodes.Last();
+				USceneComponent* SceneComp = SelectedNode.IsValid() ? Cast<USceneComponent>(SelectedNode->FindComponentInstanceInActor(PreviewActor, true)) : NULL;
 				if( SceneComp )
 				{
-					TSharedPtr<ISCSEditorCustomization> Customization = BlueprintEditorPtr.Pin()->CustomizeSCSEditor(SceneComp);
+					TSharedPtr<ISCSEditorCustomization> Customization = BlueprintEditor->CustomizeSCSEditor(SceneComp);
 					FMatrix CustomTransform;
 					if(Customization.IsValid() && Customization->HandleGetWidgetTransform(SceneComp, CustomTransform))
 					{
