More Tutorial metatags - mostly on various editors (Persona, cascade etc)

diff --git a/Engine/Source/Editor/Cascade/Private/Cascade.cpp b/Engine/Source/Editor/Cascade/Private/Cascade.cpp
--- a/Engine/Source/Editor/Cascade/Private/Cascade.cpp
+++ b/Engine/Source/Editor/Cascade/Private/Cascade.cpp
@@ -2030,6 +2030,7 @@ void FCascade::ExtendToolbar()
 	TSharedRef<SWidget> CurrentLOD = SNew(SBox)
 			[
 				SNew(SHorizontalBox)
+				.AddMetaData<FTagMetaData>(TEXT("Cascade.LODBOx"))
 				+SHorizontalBox::Slot()
 				.AutoWidth()
 				.VAlign(VAlign_Center)
diff --git a/Engine/Source/Editor/Cascade/Private/SCascadePreviewToolbar.cpp b/Engine/Source/Editor/Cascade/Private/SCascadePreviewToolbar.cpp
--- a/Engine/Source/Editor/Cascade/Private/SCascadePreviewToolbar.cpp
+++ b/Engine/Source/Editor/Cascade/Private/SCascadePreviewToolbar.cpp
@@ -33,6 +33,7 @@ void SCascadePreviewViewportToolBar::Construct(const FArguments& InArgs)
 				.Cursor(EMouseCursor::Default)
 				.Label(NSLOCTEXT("Cascade", "ViewMenuTitle_Default", "View"))
 				.OnGetMenuContent(this, &SCascadePreviewViewportToolBar::GenerateViewMenu) 
+				.AddMetaData<FTagMetaData>(TEXT("CascadeViewButton.View"))
 			]
 			+SHorizontalBox::Slot()
 			.AutoWidth()
@@ -43,6 +44,7 @@ void SCascadePreviewViewportToolBar::Construct(const FArguments& InArgs)
 				.Cursor(EMouseCursor::Default)
 				.Label(NSLOCTEXT("Cascade", "TimeMenuTitle_Default", "Time"))
 				.OnGetMenuContent(this, &SCascadePreviewViewportToolBar::GenerateTimeMenu) 
+				.AddMetaData<FTagMetaData>(TEXT("CascadeViewButton.Time"))
 			]
 		]
 	];
diff --git a/Engine/Source/Editor/DestructibleMeshEditor/Private/DestructibleMeshEditor.cpp b/Engine/Source/Editor/DestructibleMeshEditor/Private/DestructibleMeshEditor.cpp
--- a/Engine/Source/Editor/DestructibleMeshEditor/Private/DestructibleMeshEditor.cpp
+++ b/Engine/Source/Editor/DestructibleMeshEditor/Private/DestructibleMeshEditor.cpp
@@ -278,6 +278,7 @@ void FDestructibleMeshEditor::ExtendToolbar()
 			.OnSelectionChanged(this, &FDestructibleMeshEditor::PreviewDepthSelectionChanged)
 			.InitiallySelectedItem(PreviewDepths.Num() > 0 ? PreviewDepths[0] : NULL)
 			.IsEnabled( FSlateApplication::Get().GetNormalExecutionAttribute() )
+			.AddMetaData<FTagMetaData>(TEXT("Destructible.PreviewDepth"))
 			.Content()
 			[
 				SNew( STextBlock )
@@ -292,7 +293,8 @@ void FDestructibleMeshEditor::ExtendToolbar()
 		.HAlign( HAlign_Left )
 		[
 			SNew( SVerticalBox )
-			+SVerticalBox::Slot()
+			.AddMetaData<FTagMetaData>(TEXT("Destructible.ExplodeAmount")) 
+			+ SVerticalBox::Slot()
 			.AutoHeight()
 			.Padding( FMargin(40.0f, 0.0f) )
 			.HAlign(HAlign_Center)
diff --git a/Engine/Source/Editor/Kismet/Private/SBlueprintEditorSelectedDebugObjectWidget.cpp b/Engine/Source/Editor/Kismet/Private/SBlueprintEditorSelectedDebugObjectWidget.cpp
--- a/Engine/Source/Editor/Kismet/Private/SBlueprintEditorSelectedDebugObjectWidget.cpp
+++ b/Engine/Source/Editor/Kismet/Private/SBlueprintEditorSelectedDebugObjectWidget.cpp
@@ -59,7 +59,8 @@ void SBlueprintEditorSelectedDebugObjectWidget::Construct(const FArguments& InAr
 		.OptionsSource(&DebugObjectNames)
 		.InitiallySelectedItem(GetDebugObjectName())
 		.OnComboBoxOpening(this, &SBlueprintEditorSelectedDebugObjectWidget::GenerateDebugObjectNames, true)
-		.OnSelectionChanged(this, &SBlueprintEditorSelectedDebugObjectWidget::DebugObjectSelectionChanged);
+		.OnSelectionChanged(this, &SBlueprintEditorSelectedDebugObjectWidget::DebugObjectSelectionChanged)
+		.AddMetaData<FTagMetaData>(TEXT("SelectDebugObjectCobmo"));
 
 	TSharedRef<SWidget> DebugObjectSelectionWidget =
 		SNew(SHorizontalBox)
diff --git a/Engine/Source/Editor/Kismet/Private/STimelineEditor.cpp b/Engine/Source/Editor/Kismet/Private/STimelineEditor.cpp
--- a/Engine/Source/Editor/Kismet/Private/STimelineEditor.cpp
+++ b/Engine/Source/Editor/Kismet/Private/STimelineEditor.cpp
@@ -676,6 +676,7 @@ void STimelineEditor::Construct(const FArguments& InArgs, TSharedPtr<FBlueprintE
 			SNew(SBorder)
 			. BorderImage( FEditorStyle::GetBrush( TEXT("Graph.TitleBackground") ) )
 			. HAlign(HAlign_Center)
+			.AddMetaData<FTagMetaData>(TEXT("TimelineEditor.Title"))
 			[
 				SNew(SHorizontalBox)
 				+SHorizontalBox::Slot()
@@ -715,6 +716,7 @@ void STimelineEditor::Construct(const FArguments& InArgs, TSharedPtr<FBlueprintE
 				]
 				.ToolTipText( LOCTEXT( "AddFloatTrack", "Add Float Track" ) )
 				.OnClicked( this, &STimelineEditor::CreateNewTrack, FTimelineEdTrack::TT_FloatInterp )
+				.AddMetaData<FTagMetaData>(TEXT("TimelineEditor.AddFloatTrack"))
 			]
 			+SHorizontalBox::Slot()
 			.AutoWidth()
@@ -729,6 +731,7 @@ void STimelineEditor::Construct(const FArguments& InArgs, TSharedPtr<FBlueprintE
 				]
 				.ToolTipText( LOCTEXT( "AddVectorTrack", "Add Vector Track" ) )
 				.OnClicked( this, &STimelineEditor::CreateNewTrack, FTimelineEdTrack::TT_VectorInterp )
+				.AddMetaData<FTagMetaData>(TEXT("TimelineEditor.AddVectorTrack"))
 			]
 			+SHorizontalBox::Slot()
 			.AutoWidth()
@@ -743,6 +746,7 @@ void STimelineEditor::Construct(const FArguments& InArgs, TSharedPtr<FBlueprintE
 				]
 				.ToolTipText( LOCTEXT( "AddEventTrack", "Add Event Track" ) )
 				.OnClicked( this, &STimelineEditor::CreateNewTrack, FTimelineEdTrack::TT_Event )
+				.AddMetaData<FTagMetaData>(TEXT("TimelineEditor.AddEventTrack"))
 			]
 			+SHorizontalBox::Slot()
 			.AutoWidth()
@@ -757,6 +761,7 @@ void STimelineEditor::Construct(const FArguments& InArgs, TSharedPtr<FBlueprintE
 				]
 				.ToolTipText( LOCTEXT( "AddColorTrack", "Add Color Track" ) )
 				.OnClicked( this, &STimelineEditor::CreateNewTrack, FTimelineEdTrack::TT_LinearColorInterp )
+				.AddMetaData<FTagMetaData>(TEXT("TimelineEditor.AddColorTrack"))
 			]
 			+SHorizontalBox::Slot()
 			.AutoWidth()
@@ -772,6 +777,7 @@ void STimelineEditor::Construct(const FArguments& InArgs, TSharedPtr<FBlueprintE
 				.ToolTipText( LOCTEXT( "AddExternalAsset", "Add Selected Curve Asset" ) )
 				.IsEnabled( this, &STimelineEditor::IsCurveAssetSelected )
 				.OnClicked( this, &STimelineEditor::CreateNewTrackFromAsset )
+				.AddMetaData<FTagMetaData>(TEXT("TimelineEditor.AddCurveAssetTrack"))
 			]
 			+SHorizontalBox::Slot()
 			.AutoWidth()
@@ -792,6 +798,7 @@ void STimelineEditor::Construct(const FArguments& InArgs, TSharedPtr<FBlueprintE
 				.OnTextCommitted( this, &STimelineEditor::OnLengthStringChanged )
 				.SelectAllTextWhenFocused(true)
 				.MinDesiredWidth(64)
+				.AddMetaData<FTagMetaData>(TEXT("TimelineEditor.Length"))
 			]
 			+SHorizontalBox::Slot()
 			.AutoWidth()
@@ -803,6 +810,7 @@ void STimelineEditor::Construct(const FArguments& InArgs, TSharedPtr<FBlueprintE
 				.OnCheckStateChanged( this, &STimelineEditor::OnUseLastKeyframeChanged )
 				[
 					SNew(STextBlock) .Text( LOCTEXT( "UseLastKeyframe", "Use Last Keyframe?" ) )
+					.AddMetaData<FTagMetaData>(TEXT("TimelineEditor.UseLastKeyframe"))
 				]
 			]
 			+SHorizontalBox::Slot()
@@ -815,6 +823,7 @@ void STimelineEditor::Construct(const FArguments& InArgs, TSharedPtr<FBlueprintE
 				.OnCheckStateChanged( this, &STimelineEditor::OnAutoPlayChanged )
 				[
 					SNew(STextBlock) .Text( LOCTEXT( "AutoPlay", "AutoPlay" ) )
+					.AddMetaData<FTagMetaData>(TEXT("TimelineEditor.AutoPlay"))
 				]
 			]
 			+SHorizontalBox::Slot()
@@ -827,6 +836,7 @@ void STimelineEditor::Construct(const FArguments& InArgs, TSharedPtr<FBlueprintE
 				.OnCheckStateChanged( this, &STimelineEditor::OnLoopChanged )
 				[
 					SNew(STextBlock) .Text( LOCTEXT( "Loop", "Loop" ) )
+					.AddMetaData<FTagMetaData>(TEXT("TimelineEditor.Loop"))
 				]
 			]
 			+SHorizontalBox::Slot()
@@ -839,6 +849,7 @@ void STimelineEditor::Construct(const FArguments& InArgs, TSharedPtr<FBlueprintE
 				.OnCheckStateChanged( this, &STimelineEditor::OnReplicatedChanged )
 				[
 					SNew(STextBlock) .Text( LOCTEXT( "Replicated", "Replicated" ) )
+					.AddMetaData<FTagMetaData>(TEXT("TimelineEditor.Replicated"))
 				]
 			]
 		]
diff --git a/Engine/Source/Editor/Persona/Private/SAnimCurvePanel.cpp b/Engine/Source/Editor/Persona/Private/SAnimCurvePanel.cpp
--- a/Engine/Source/Editor/Persona/Private/SAnimCurvePanel.cpp
+++ b/Engine/Source/Editor/Persona/Private/SAnimCurvePanel.cpp
@@ -664,6 +664,7 @@ void SAnimCurvePanel::Construct(const FArguments& InArgs)
 						.Text( LOCTEXT("AddFloatTrack", "Add...") )
 						.ToolTipText( LOCTEXT("AddTrackTooltip", "Add float track above here") )
 						.OnClicked( this, &SAnimCurvePanel::AddButtonClicked )
+						.AddMetaData<FTagMetaData>(TEXT("AnimCurve.AddFloat"))
 					]
 
 					+SHorizontalBox::Slot()
diff --git a/Engine/Source/Editor/Persona/Private/SAnimNotifyPanel.cpp b/Engine/Source/Editor/Persona/Private/SAnimNotifyPanel.cpp
--- a/Engine/Source/Editor/Persona/Private/SAnimNotifyPanel.cpp
+++ b/Engine/Source/Editor/Persona/Private/SAnimNotifyPanel.cpp
@@ -2932,6 +2932,7 @@ void SAnimNotifyPanel::Construct(const FArguments& InArgs)
 		[
 			SNew( SExpandableArea )
 			.AreaTitle( LOCTEXT("Notifies", "Notifies") )
+			.AddMetaData<FTagMetaData>(TEXT("AnimNotify.Notify"))
 			.BodyContent()
 			[
 				SAssignNew( PanelArea, SBorder )
diff --git a/Engine/Source/Editor/Persona/Private/SAnimationScrubPanel.cpp b/Engine/Source/Editor/Persona/Private/SAnimationScrubPanel.cpp
--- a/Engine/Source/Editor/Persona/Private/SAnimationScrubPanel.cpp
+++ b/Engine/Source/Editor/Persona/Private/SAnimationScrubPanel.cpp
@@ -24,6 +24,7 @@ void SAnimationScrubPanel::Construct( const SAnimationScrubPanel::FArguments& In
 	this->ChildSlot
 	[
 		SNew(SHorizontalBox)
+		.AddMetaData<FTagMetaData>(TEXT("AnimScrub.Scrub"))
 		+SHorizontalBox::Slot()
 		.HAlign(HAlign_Fill) 
 		.VAlign(VAlign_Center)
diff --git a/Engine/Source/Editor/Persona/Private/SPersonaToolbar.cpp b/Engine/Source/Editor/Persona/Private/SPersonaToolbar.cpp
--- a/Engine/Source/Editor/Persona/Private/SPersonaToolbar.cpp
+++ b/Engine/Source/Editor/Persona/Private/SPersonaToolbar.cpp
@@ -95,6 +95,7 @@ void FPersonaToolbar::FillPersonaModeToolbar(FToolBarBuilder& ParentToolbarBuild
 				NULL,
 				TEXT("Shared/Editors/Persona"),
 				TEXT("SkeletonMode")))
+			.AddMetaData<FTagMetaData>(TEXT("Persona.Skeleton"))
 			.ShortContents()
 			[
 				SNew(SContentReference)
@@ -118,6 +119,7 @@ void FPersonaToolbar::FillPersonaModeToolbar(FToolBarBuilder& ParentToolbarBuild
 				NULL,
 				TEXT("Shared/Editors/Persona"),
 				TEXT("MeshMode")))
+			.AddMetaData<FTagMetaData>(TEXT("Persona.Mesh"))
 			.ShortContents()
 			[
 				SNew(SContentReference)
@@ -164,6 +166,7 @@ void FPersonaToolbar::FillPersonaModeToolbar(FToolBarBuilder& ParentToolbarBuild
 				NULL,
 				TEXT("Shared/Editors/Persona"),
 				TEXT("AnimationMode")))
+			.AddMetaData<FTagMetaData>(TEXT("Persona.Anim"))
 			.ShortContents()
 			[
 				//@TODO: How to show this guy?
@@ -197,6 +200,7 @@ void FPersonaToolbar::FillPersonaModeToolbar(FToolBarBuilder& ParentToolbarBuild
 					TEXT("GraphMode")))
 				.IconImage(FEditorStyle::GetBrush("FullBlueprintEditor.SwitchToScriptingMode"))
 				.SmallIconImage(FEditorStyle::GetBrush("FullBlueprintEditor.SwitchToScriptingMode.Small"))
+				.AddMetaData<FTagMetaData>(TEXT("Persona.Graph"))
 				.ShortContents()
 				[
 					SNew(SContentReference)
diff --git a/Engine/Source/Editor/Persona/Private/SSkeletonTree.cpp b/Engine/Source/Editor/Persona/Private/SSkeletonTree.cpp
--- a/Engine/Source/Editor/Persona/Private/SSkeletonTree.cpp
+++ b/Engine/Source/Editor/Persona/Private/SSkeletonTree.cpp
@@ -829,6 +829,7 @@ void SSkeletonTree::Construct(const FArguments& InArgs)
 			.SelectAllTextWhenFocused( true )
 			.OnTextChanged( this, &SSkeletonTree::OnFilterTextChanged )
 			.HintText( LOCTEXT( "SearchBoxHint", "Search Skeleton Tree...") )
+			.AddMetaData<FTagMetaData>(TEXT("SkelTree.Search"))
 		]
 
 		+ SVerticalBox::Slot()
@@ -845,6 +846,7 @@ void SSkeletonTree::Construct(const FArguments& InArgs)
 				.ContentPadding(3)
 				.OnGetMenuContent( this, &SSkeletonTree::CreateBoneFilterMenu )
 				.ToolTipText( LOCTEXT( "BoneFilterToolTip", "Change which types of bones are shown" ) )
+				.AddMetaData<FTagMetaData>(TEXT("SkelTree.Bones"))
 				.ButtonContent()
 				[
 					SNew( STextBlock )
@@ -861,6 +863,7 @@ void SSkeletonTree::Construct(const FArguments& InArgs)
 				.ContentPadding(3)
 				.OnGetMenuContent( this, &SSkeletonTree::CreateSocketFilterMenu )
 				.ToolTipText( LOCTEXT( "SocketFilterToolTip", "Change which types of sockets are shown" ) )
+				.AddMetaData<FTagMetaData>(TEXT("SkelTree.Sockets"))
 				.ButtonContent()
 				[
 					SNew( STextBlock )
diff --git a/Engine/Source/Editor/PropertyEditor/Private/SDetailsView.cpp b/Engine/Source/Editor/PropertyEditor/Private/SDetailsView.cpp
--- a/Engine/Source/Editor/PropertyEditor/Private/SDetailsView.cpp
+++ b/Engine/Source/Editor/PropertyEditor/Private/SDetailsView.cpp
@@ -125,6 +125,7 @@ void SDetailsView::Construct(const FArguments& InArgs)
 			// Create the search box
 			SAssignNew( SearchBox, SSearchBox )
 			.OnTextChanged( this, &SDetailsView::OnFilterTextChanged  )
+			.AddMetaData<FTagMetaData>(TEXT("Details.Search"))
 		]
 		+SHorizontalBox::Slot()
 		.Padding( 4.0f, 0.0f, 0.0f, 0.0f )
diff --git a/Engine/Source/Runtime/Slate/Private/Framework/Docking/SDockingTabStack.cpp b/Engine/Source/Runtime/Slate/Private/Framework/Docking/SDockingTabStack.cpp
--- a/Engine/Source/Runtime/Slate/Private/Framework/Docking/SDockingTabStack.cpp
+++ b/Engine/Source/Runtime/Slate/Private/Framework/Docking/SDockingTabStack.cpp
@@ -114,6 +114,11 @@ void SDockingTabStack::Construct( const FArguments& InArgs, const TSharedRef<FTa
 		.Padding(5.0f, 0.0f, 0.0f, 0.0f)
 		.VAlign((VAlign_Center));
 
+	FTagMetaData MetaData(TEXT("Undefined"));
+	if (Tabs.Num() > 0)
+	{
+		MetaData.Tag = *Tabs[0].TabId.ToString();
+	}
 	ChildSlot
 	[
 		SNew(SVerticalBox)
@@ -159,6 +164,7 @@ void SDockingTabStack::Construct( const FArguments& InArgs, const TSharedRef<FTa
 						SAssignNew(ContentSlot, SBorder)
 							.BorderImage(this, &SDockingTabStack::GetContentAreaBrush)
 							.Padding(this, &SDockingTabStack::GetContentPadding)
+							.AddMetaData<FTagMetaData>(MetaData)
 							[
 								SNew(STextBlock)
 									.Text(LOCTEXT("EmptyTabMessage", "Empty Tab!"))
