Messaging users with empty graphs, how to start adding nodes.

diff --git a/Engine/Source/Editor/EditorSettingsViewer/EditorSettingsViewer.Build.cs b/Engine/Source/Editor/EditorSettingsViewer/EditorSettingsViewer.Build.cs
--- a/Engine/Source/Editor/EditorSettingsViewer/EditorSettingsViewer.Build.cs
+++ b/Engine/Source/Editor/EditorSettingsViewer/EditorSettingsViewer.Build.cs
@@ -23,7 +23,8 @@ namespace UnrealBuildTool.Rules
 					"Slate",
 					"SlateCore",
 					"UnrealEd",
-                    "InternationalizationSettings"
+                    "InternationalizationSettings",
+					"Kismet"
 				}
 			);
 
diff --git a/Engine/Source/Editor/EditorSettingsViewer/Private/EditorSettingsViewerModule.cpp b/Engine/Source/Editor/EditorSettingsViewer/Private/EditorSettingsViewerModule.cpp
--- a/Engine/Source/Editor/EditorSettingsViewer/Private/EditorSettingsViewerModule.cpp
+++ b/Engine/Source/Editor/EditorSettingsViewer/Private/EditorSettingsViewerModule.cpp
@@ -2,7 +2,7 @@
 
 #include "EditorSettingsViewerPrivatePCH.h"
 #include "Tests/AutomationTestSettings.h"
-
+#include "BlueprintEditorSettings.h"
 
 #define LOCTEXT_NAMESPACE "FEditorSettingsViewerModule"
 
@@ -188,6 +188,13 @@ protected:
 			LOCTEXT("ContentEditorsGraphEditorSettingsDescription", "Customize Anim, Blueprint and Material Editor."),
 			GetMutableDefault<UGraphEditorSettings>()
 		);
+
+		// graph editors
+		SettingsModule.RegisterSettings("Editor", "ContentEditors", "BlueprintEditor",
+			LOCTEXT("ContentEditorsBlueprintEditorSettingsName", "Blueprint Editor"),
+			LOCTEXT("ContentEditorsGraphBlueprintSettingsDescription", "Customize Blueprint Editors."),
+			GetMutableDefault<UBlueprintEditorSettings>()
+		);
 	}
 
 	/**
diff --git a/Engine/Source/Editor/EditorStyle/Private/SlateEditorStyle.cpp b/Engine/Source/Editor/EditorStyle/Private/SlateEditorStyle.cpp
--- a/Engine/Source/Editor/EditorStyle/Private/SlateEditorStyle.cpp
+++ b/Engine/Source/Editor/EditorStyle/Private/SlateEditorStyle.cpp
@@ -3468,6 +3468,14 @@ void FSlateEditorStyle::FStyle::SetupGraphEditorStyles()
 				.SetColorAndOpacity( FLinearColor(0.8, 0.8f, 0.8f, 0.2f) )
 				.SetShadowOffset( FVector2D::ZeroVector )
 			);
+
+			Set( "Graph.InstructionText", FTextBlockStyle(NormalText)
+				.SetFont( TTF_FONT( "Fonts/Roboto-BoldCondensed", 24 ) )
+				.SetColorAndOpacity( FLinearColor(1.f, 1.f, 1.f, 0.6f) )
+				.SetShadowOffset( FVector2D::ZeroVector )
+			);
+
+			Set( "Graph.InstructionBackground", new BOX_BRUSH("Common/GroupBorder", FMargin(4.0f / 16.0f), FLinearColor(0.1f, 0.1f, 0.1f, 0.7f)) );
 		}
 
 		{
diff --git a/Engine/Source/Editor/GraphEditor/Private/GraphEditorModule.cpp b/Engine/Source/Editor/GraphEditor/Private/GraphEditorModule.cpp
--- a/Engine/Source/Editor/GraphEditor/Private/GraphEditorModule.cpp
+++ b/Engine/Source/Editor/GraphEditor/Private/GraphEditorModule.cpp
@@ -54,6 +54,7 @@ void FGraphEditorModule::ShutdownModule()
 TSharedRef<SGraphEditor> FGraphEditorModule::PRIVATE_MakeGraphEditor( 
 	const TSharedPtr<FUICommandList>& InAdditionalCommands, 
 	const TAttribute<bool>& InIsEditable, 
+	const TAttribute<bool>& InIsEmpty,
 	TAttribute<FGraphAppearanceInfo> Appearance,
 	TSharedPtr<SWidget> InTitleBar,
 	const TAttribute<bool>& InTitleBarEnabledOnly,
diff --git a/Engine/Source/Editor/GraphEditor/Private/SGraphEditorImpl.cpp b/Engine/Source/Editor/GraphEditor/Private/SGraphEditorImpl.cpp
--- a/Engine/Source/Editor/GraphEditor/Private/SGraphEditorImpl.cpp
+++ b/Engine/Source/Editor/GraphEditor/Private/SGraphEditorImpl.cpp
@@ -260,7 +260,30 @@ void SGraphEditorImpl::Construct( const FArguments& InArgs )
 		+SOverlay::Slot()
 		.VAlign(VAlign_Top)
 		[
-			InArgs._TitleBar.IsValid() ? InArgs._TitleBar.ToSharedRef() : (TSharedRef<SWidget>)SNullWidget::NullWidget
+			SNew(SVerticalBox)
+			+ SVerticalBox::Slot()
+			[
+				InArgs._TitleBar.IsValid() ? InArgs._TitleBar.ToSharedRef() : (TSharedRef<SWidget>)SNullWidget::NullWidget
+			]
+			+ SVerticalBox::Slot()
+			.Padding(20.f, 20.f, 20.f, 0.f)
+			.VAlign(VAlign_Top)
+			.HAlign(HAlign_Center)
+			.AutoHeight()
+			[
+				SNew(SBorder)
+				.Padding(FMargin(10.f, 4.f))
+				.BorderImage(FEditorStyle::GetBrush(TEXT("Graph.InstructionBackground")))
+				.BorderBackgroundColor(this, &SGraphEditorImpl::InstructionBorderColor)
+				.HAlign(HAlign_Center)
+				.ColorAndOpacity(this, &SGraphEditorImpl::InstructionTextTint)
+				.Visibility(this, &SGraphEditorImpl::InstructionTextVisibility)
+				[
+					SNew(STextBlock)
+					.TextStyle(FEditorStyle::Get(), "Graph.InstructionText")
+					.Text(this, &SGraphEditorImpl::GetInstructionText)
+				]
+			]			
 		]
 
 		// Bottom-right corner text indicating the type of tool
@@ -299,6 +322,23 @@ void SGraphEditorImpl::Construct( const FArguments& InArgs )
 			.Text(ReadOnlyText)
 		]
 
+// 		+ SOverlay::Slot()
+// 		.Padding(20)
+// 		.VAlign(VAlign_Fill)
+// 		.HAlign(HAlign_Fill)
+// 		[
+// 			SNew(SVerticalBox)
+// 			+ SVerticalBox::Slot()
+// 			.FillHeight(0.5)
+// 			.VAlign(VAlign_Bottom)
+// 			.HAlign(HAlign_Center)
+// 			[
+// 				
+// 			]
+// 			+ SVerticalBox::Slot()
+// 			.FillHeight(0.5)
+// 		]
+
 		// Bottom-right corner text for notification list position
 		+SOverlay::Slot()
 		.Padding(15.f)
@@ -616,7 +656,45 @@ EVisibility SGraphEditorImpl::ReadOnlyVisibility() const
 	return EVisibility::Hidden;
 }
 
+FString SGraphEditorImpl::GetInstructionText() const
+{
+	if (Appearance.IsBound())
+	{
+		return Appearance.Get().InstructionText.ToString();
+	}
+	return TEXT("");
+}
+
+EVisibility SGraphEditorImpl::InstructionTextVisibility() const
+{
+	if (GetInstructionTextFade() > 0.0f)
+	{
+		return EVisibility::Visible;
+	}
+	return EVisibility::Hidden;
+}
+
+float SGraphEditorImpl::GetInstructionTextFade() const
+{
+	float InstructionOpacity = 1.0f;
+	if (Appearance.IsBound())
+	{
+		InstructionOpacity = Appearance.Get().InstructionFade.Get();
+	}
+	return InstructionOpacity;
+}
 
+FLinearColor SGraphEditorImpl::InstructionTextTint() const
+{
+	return FLinearColor(1.f, 1.f, 1.f, GetInstructionTextFade());
+}
+
+FSlateColor SGraphEditorImpl::InstructionBorderColor() const
+{
+	FLinearColor BorderColor(0.1f, 0.1f, 0.1f, 0.7f);
+	BorderColor.A *= GetInstructionTextFade();
+	return BorderColor;
+}
 
 /////////////////////////////////////////////////////
 
diff --git a/Engine/Source/Editor/GraphEditor/Private/SGraphEditorImpl.h b/Engine/Source/Editor/GraphEditor/Private/SGraphEditorImpl.h
--- a/Engine/Source/Editor/GraphEditor/Private/SGraphEditorImpl.h
+++ b/Engine/Source/Editor/GraphEditor/Private/SGraphEditorImpl.h
@@ -61,7 +61,6 @@ private:
 	SGraphEditor::FOnCreateActionMenu OnCreateActionMenu;
 
 	TAttribute<bool> IsEditable;
-
 	TAttribute<bool> TitleBarEnabledOnly;
 
 	bool bAutoExpandActionMenu;
@@ -78,6 +77,17 @@ private:
 	/** Function to check whether we should show read-only text in the panel */
 	EVisibility ReadOnlyVisibility() const;
 
+	/** Returns dynamic text, meant to passively instruct the user on what to do in the graph */
+	FString GetInstructionText() const;
+	/** Function to check whether we should show instruction text to the user */
+	EVisibility InstructionTextVisibility() const;
+	/** Returns a 0.0 to 1.0 value, denoting the instruction text's fade percent */
+	float GetInstructionTextFade() const;
+	/** A dynamic tint for the instruction text (allows us to nicely fade it in/out) */
+	FLinearColor InstructionTextTint() const;
+	/** Determines the color of the box containing the instruction text */
+	FSlateColor InstructionBorderColor() const;
+
 	/** Notification list to pass messages to editor users  */
 	TSharedPtr<SNotificationList> NotificationListPtr;
 
diff --git a/Engine/Source/Editor/GraphEditor/Public/GraphEditorModule.h b/Engine/Source/Editor/GraphEditor/Public/GraphEditorModule.h
--- a/Engine/Source/Editor/GraphEditor/Public/GraphEditorModule.h
+++ b/Engine/Source/Editor/GraphEditor/Public/GraphEditorModule.h
@@ -32,6 +32,7 @@ private:
 	virtual TSharedRef<SGraphEditor> PRIVATE_MakeGraphEditor(
 		const TSharedPtr<FUICommandList>& InAdditionalCommands,
 		const TAttribute<bool>& InIsEditable,
+		const TAttribute<bool>& InIsEmpty,
 		TAttribute<FGraphAppearanceInfo> Appearance,
 		TSharedPtr<SWidget> InTitleBar,
 		const TAttribute<bool>& InTitleBarEnabledOnly,
diff --git a/Engine/Source/Editor/Kismet/Private/BlueprintDetailsCustomization.cpp b/Engine/Source/Editor/Kismet/Private/BlueprintDetailsCustomization.cpp
--- a/Engine/Source/Editor/Kismet/Private/BlueprintDetailsCustomization.cpp
+++ b/Engine/Source/Editor/Kismet/Private/BlueprintDetailsCustomization.cpp
@@ -2462,7 +2462,7 @@ bool FBaseBlueprintGraphActionDetails::AttemptToCreateResultNode()
 		const UEdGraphSchema_K2* Schema = Cast<const UEdGraphSchema_K2>(FunctionResult->GetSchema());
 		FunctionResult->NodePosX = FunctionEntryNode->NodePosX + FunctionEntryNode->NodeWidth + 256;
 		FunctionResult->NodePosY = FunctionEntryNode->NodePosY;
-			
+		UEdGraphSchema_K2::SetNodeMetaData(FunctionResult, FNodeMetadata::DefaultGraphNode);
 		ResultNodeCreator.Finalize();
 		
 		// Connect the function entry to the result node, if applicable
diff --git a/Engine/Source/Editor/Kismet/Private/BlueprintEditor.cpp b/Engine/Source/Editor/Kismet/Private/BlueprintEditor.cpp
--- a/Engine/Source/Editor/Kismet/Private/BlueprintEditor.cpp
+++ b/Engine/Source/Editor/Kismet/Private/BlueprintEditor.cpp
@@ -68,6 +68,7 @@
 #include "AnimationStateMachineSchema.h"
 #include "AnimationTransitionGraph.h"
 #include "BlueprintEditorModes.h"
+#include "BlueprintEditorSettings.h"
 
 #include "EngineAnalytics.h"
 #include "IAnalyticsProvider.h"
@@ -168,6 +169,59 @@ private:
 /////////////////////////////////////////////////////
 // FBlueprintEditor
 
+namespace BlueprintEditorImpl
+{
+	static const float InstructionFadeDuration = 0.5f;
+
+	/**
+	*
+	*
+	* @param  InGraph
+	* @return
+	*/
+	static bool GraphHasUserPlacedNodes(UEdGraph const* InGraph);
+
+	/**
+	*
+	*
+	* @param  InGraph
+	* @return
+	*/
+	static bool GraphHasDefaultNode(UEdGraph const* InGraph);
+}
+
+static bool BlueprintEditorImpl::GraphHasUserPlacedNodes(UEdGraph const* InGraph)
+{
+	bool bHasUserPlacedNodes = false;
+
+	for (UEdGraphNode const* Node : InGraph->Nodes)
+	{
+		if (!Node->GetOutermost()->GetMetaData()->HasValue(Node, FNodeMetadata::DefaultGraphNode))
+		{
+			bHasUserPlacedNodes = true;
+			break;
+		}
+	}
+
+	return bHasUserPlacedNodes;
+}
+
+static bool BlueprintEditorImpl::GraphHasDefaultNode(UEdGraph const* InGraph)
+{
+	bool bHasDefaultNodes = false;
+
+	for (UEdGraphNode const* Node : InGraph->Nodes)
+	{
+		if (Node->GetOutermost()->GetMetaData()->HasValue(Node, FNodeMetadata::DefaultGraphNode))
+		{
+			bHasDefaultNodes = true;
+			break;
+		}
+	}
+
+	return bHasDefaultNodes;
+}
+
 bool FBlueprintEditor::IsASubGraph( const UEdGraph* GraphPtr )
 {
 	if( GraphPtr && GraphPtr->GetOuter() )
@@ -458,10 +512,6 @@ TSharedRef<SGraphEditor> FBlueprintEditor::CreateGraphEditorWidget(TSharedRef<FT
 {
 	check((InGraph != NULL) && IsEditingSingleBlueprint());
 
-	// Cache off whether or not this is an interface, since it is used to govern multiple widget's behavior
-	const bool bIsInterface = (GetBlueprintObj()->BlueprintType == BPTYPE_Interface);
-	const bool bIsDelegate = FBlueprintEditorUtils::IsDelegateSignatureGraph(InGraph);
-
 	// No need to regenerate the commands.
 	if(!GraphEditorCommands.IsValid())
 	{
@@ -770,11 +820,13 @@ TSharedRef<SGraphEditor> FBlueprintEditor::CreateGraphEditorWidget(TSharedRef<FT
 		InEvents.OnCreateActionMenu = SGraphEditor::FOnCreateActionMenu::CreateSP(this, &FBlueprintEditor::OnCreateGraphActionMenu);
 	}
 
+
+	bool const bEnableTitleBarOnly = !IsGraphPanelEnabled(InGraph);
 	TSharedRef<SGraphEditor> Editor = SNew(SGraphEditor)
 		.AdditionalCommands(GraphEditorCommands)
 		.IsEditable(this, &FBlueprintEditor::IsEditable, InGraph)
 		.TitleBar(TitleBarWidget)
-		.TitleBarEnabledOnly(bIsInterface || bIsDelegate)
+		.TitleBarEnabledOnly(bEnableTitleBarOnly)
 		.Appearance(this, &FBlueprintEditor::GetGraphAppearance)
 		.GraphToEdit(InGraph)
 		.GraphEvents(InEvents)
@@ -825,6 +877,17 @@ FGraphAppearanceInfo FBlueprintEditor::GetGraphAppearance() const
 		break;
 	}
 
+	UEdGraph const* EditingGraph = GetFocusedGraph();
+	if (EditingGraph && BlueprintEditorImpl::GraphHasDefaultNode(EditingGraph))
+	{
+		AppearanceInfo.InstructionText = LOCTEXT("AppearanceInstructionText_DefaultGraph", "Drag Off Pins to Create/Connect New Nodes.");
+	}
+	else // if the graph is empty...
+	{
+		AppearanceInfo.InstructionText = LOCTEXT("AppearanceInstructionText_EmptyGraph", "Right-Click to Create New Nodes.");
+	}
+	AppearanceInfo.InstructionFade.Bind(this, &FBlueprintEditor::GetInstructionTextOpacity);
+
 	AppearanceInfo.PIENotifyText = GetPIEStatus();
 
 	return AppearanceInfo;
@@ -854,6 +917,8 @@ FBlueprintEditor::FBlueprintEditor()
 	, CurrentUISelection(FBlueprintEditor::NoSelection)
 	, bEditorMarkedAsClosed(false)
 	, bCodeBasedProject(false)
+	, bActionMenuIsOpen(false)
+	, InstructionsFadeCountdown(0.f)
 {
 	AnalyticsStats.GraphActionMenusNonCtxtSensitiveExecCount = 0;
 	AnalyticsStats.GraphActionMenusCtxtSensitiveExecCount = 0;
@@ -2267,6 +2332,12 @@ void FBlueprintEditor::OnGraphEditorDropStreamingLevel(const TArray< TWeakObject
 
 FActionMenuContent FBlueprintEditor::OnCreateGraphActionMenu(UEdGraph* InGraph, const FVector2D& InNodePosition, const TArray<UEdGraphPin*>& InDraggedPins, bool bAutoExpand, SGraphEditor::FActionMenuClosed InOnMenuClosed)
 {
+	bActionMenuIsOpen = true;
+	if (!BlueprintEditorImpl::GraphHasUserPlacedNodes(InGraph))
+	{
+		InstructionsFadeCountdown = BlueprintEditorImpl::InstructionFadeDuration;
+	}
+
 	TSharedRef<SBlueprintActionMenu> ActionMenu = 
 		SNew(SBlueprintActionMenu, SharedThis(this))
 		.GraphObj(InGraph)
@@ -2290,6 +2361,16 @@ void FBlueprintEditor::OnGraphActionMenuClosed(bool bActionExecuted, bool bConte
 	{
 		AnalyticsStats.GraphActionMenusCancelledCount++;
 	}
+
+	if (UEdGraph* EditingGraph = GetFocusedGraph())
+	{
+		// if the user didn't place any nodes...
+		if (!BlueprintEditorImpl::GraphHasUserPlacedNodes(EditingGraph))
+		{
+			InstructionsFadeCountdown = 0.0f;
+		}
+	}
+	bActionMenuIsOpen = false;
 }
 
 void FBlueprintEditor::OnSelectedNodesChanged(const FGraphPanelSelectionSet& NewSelection)
@@ -5499,6 +5580,11 @@ void FBlueprintEditor::Tick(float DeltaTime)
 
 		SaveEditedObjectState();
 	}
+
+	if (InstructionsFadeCountdown > 0.f)
+	{
+		InstructionsFadeCountdown -= DeltaTime;
+	}
 }
 TStatId FBlueprintEditor::GetStatId() const
 {
@@ -6539,6 +6625,34 @@ bool FBlueprintEditor::IsEditable(UEdGraph* InGraph) const
 	return InEditingMode() && InGraph && InGraph->bEditable;
 }
 
+bool FBlueprintEditor::IsGraphPanelEnabled(UEdGraph* InGraph) const
+{
+	bool const bIsInterface = (FBlueprintEditorUtils::FindBlueprintForGraph(InGraph)->BlueprintType == BPTYPE_Interface);
+	bool const bIsDelegate  = FBlueprintEditorUtils::IsDelegateSignatureGraph(InGraph);
+
+	return !bIsInterface && !bIsDelegate;
+}
+
+float FBlueprintEditor::GetInstructionTextOpacity() const
+{
+	UEdGraph* EditingGraph = GetFocusedGraph();
+	UBlueprintEditorSettings const* Settings = GetDefault<UBlueprintEditorSettings>();
+
+	if ((EditingGraph == nullptr) || !IsEditable(EditingGraph) || !IsGraphPanelEnabled(EditingGraph) || !Settings->bShowGraphInstructionText)
+	{
+		return 0.0f;
+	}
+	else if ((InstructionsFadeCountdown > 0.0f) || bActionMenuIsOpen)
+	{
+		return InstructionsFadeCountdown / BlueprintEditorImpl::InstructionFadeDuration;
+	}
+	else if (BlueprintEditorImpl::GraphHasUserPlacedNodes(EditingGraph))
+	{
+		return 0.0f;
+	}
+	return 1.0f;
+}
+
 FString FBlueprintEditor::GetGraphDecorationString(UEdGraph* InGraph) const
 {
 	return TEXT("");
diff --git a/Engine/Source/Editor/Kismet/Private/BlueprintEditorSettings.cpp b/Engine/Source/Editor/Kismet/Private/BlueprintEditorSettings.cpp
new file mode 100644
--- /dev/null
+++ b/Engine/Source/Editor/Kismet/Private/BlueprintEditorSettings.cpp
@@ -0,0 +1,10 @@
+// Copyright 1998-2014 Epic Games, Inc. All Rights Reserved.
+
+#include "BlueprintEditorPrivatePCH.h"
+#include "BlueprintEditorSettings.h"
+
+UBlueprintEditorSettings::UBlueprintEditorSettings(const class FPostConstructInitializeProperties& PCIP)
+	: Super(PCIP)
+{
+	bShowGraphInstructionText = true;
+}
diff --git a/Engine/Source/Editor/Kismet/Public/BlueprintEditor.h b/Engine/Source/Editor/Kismet/Public/BlueprintEditor.h
--- a/Engine/Source/Editor/Kismet/Public/BlueprintEditor.h
+++ b/Engine/Source/Editor/Kismet/Public/BlueprintEditor.h
@@ -236,6 +236,12 @@ public:
 	/** Returns whether a graph is editable or not */
 	virtual bool IsEditable(UEdGraph* InGraph) const;
 
+	/** */
+	bool IsGraphPanelEnabled(UEdGraph* InGraph) const;
+
+	/** */
+	virtual float GetInstructionTextOpacity() const;
+
 	/** Returns true if in editing mode */
 	bool InEditingMode() const;
 
@@ -1035,6 +1041,10 @@ private:
 	/** Delegates that are fired when the blueprint editor changes modes */
 	FOnModeSet OnModeSetData;
 
+	/** */
+	bool bActionMenuIsOpen;
+	/** */
+	float InstructionsFadeCountdown;
 };
 
 #undef LOCTEXT_NAMESPACE
diff --git a/Engine/Source/Editor/Kismet/Public/BlueprintEditorSettings.h b/Engine/Source/Editor/Kismet/Public/BlueprintEditorSettings.h
new file mode 100644
--- /dev/null
+++ b/Engine/Source/Editor/Kismet/Public/BlueprintEditorSettings.h
@@ -0,0 +1,19 @@
+// Copyright 1998-2014 Epic Games, Inc. All Rights Reserved.
+
+#pragma once
+
+#include "BlueprintEditorSettings.generated.h"
+
+UCLASS(config=EditorUserSettings)
+class KISMET_API UBlueprintEditorSettings
+	:	public UObject
+{
+	GENERATED_UCLASS_BODY()
+
+// UX Settings
+public:
+	/** Determines if lightweight tutorial text shows up at the top of empty blueprint graphs */
+	UPROPERTY(EditAnywhere, config, Category="User Experience")
+	bool bShowGraphInstructionText;
+};
+ 
\ No newline at end of file
diff --git a/Engine/Source/Editor/UnrealEd/Private/GraphEditor.cpp b/Engine/Source/Editor/UnrealEd/Private/GraphEditor.cpp
--- a/Engine/Source/Editor/UnrealEd/Private/GraphEditor.cpp
+++ b/Engine/Source/Editor/UnrealEd/Private/GraphEditor.cpp
@@ -18,6 +18,7 @@ void SGraphEditor::ConstructImplementation( const FArguments& InArgs )
 	// Construct the implementation and make it the contents of this widget.
 	Implementation = GraphEdModule.PRIVATE_MakeGraphEditor( InArgs._AdditionalCommands, 
 		InArgs._IsEditable, 
+		InArgs._IsEmpty,
 		InArgs._Appearance,
 		InArgs._TitleBar,
 		InArgs._TitleBarEnabledOnly,
diff --git a/Engine/Source/Editor/UnrealEd/Public/GraphEditor.h b/Engine/Source/Editor/UnrealEd/Public/GraphEditor.h
--- a/Engine/Source/Editor/UnrealEd/Public/GraphEditor.h
+++ b/Engine/Source/Editor/UnrealEd/Public/GraphEditor.h
@@ -18,6 +18,7 @@ struct FGraphAppearanceInfo
 {
 	FGraphAppearanceInfo()
 		: CornerImage(NULL)
+		, InstructionFade(1.f)
 	{
 	}
 
@@ -29,6 +30,10 @@ struct FGraphAppearanceInfo
 	FString PIENotifyText;
 	/** If set, will be used as override for read only text */
 	FString ReadOnlyText;
+	/** Text to display if the graph is empty (to guide the user on what to do) */
+	FText InstructionText;
+	/** Allows graphs to nicely fade instruction text (or completely hide it). */
+	TAttribute<float> InstructionFade;
 };
 
 /** Struct used to return info about action menu */
@@ -105,6 +110,7 @@ public:
 	SLATE_BEGIN_ARGS(SGraphEditor)
 		: _AdditionalCommands( static_cast<FUICommandList*>(NULL) )
 		, _IsEditable(true)
+		, _IsEmpty(false)
 		, _TitleBarEnabledOnly(false)
 		, _GraphToEdit(NULL)
 		, _GraphToDiff(NULL)
@@ -114,6 +120,7 @@ public:
 
 		SLATE_ARGUMENT( TSharedPtr<FUICommandList>, AdditionalCommands )
 		SLATE_ATTRIBUTE( bool, IsEditable )		
+		SLATE_ATTRIBUTE( bool, IsEmpty )	
 		SLATE_ARGUMENT( TSharedPtr<SWidget>, TitleBar )
 		SLATE_ATTRIBUTE( FGraphAppearanceInfo, Appearance )
 		SLATE_ATTRIBUTE( bool, TitleBarEnabledOnly )
