Fixing it so the new BP instruction text shows for different graph panels split out (not just for whichever has focus).

diff --git a/Engine/Source/Editor/Kismet/Private/BlueprintEditor.cpp b/Engine/Source/Editor/Kismet/Private/BlueprintEditor.cpp
--- a/Engine/Source/Editor/Kismet/Private/BlueprintEditor.cpp
+++ b/Engine/Source/Editor/Kismet/Private/BlueprintEditor.cpp
@@ -827,7 +827,7 @@ TSharedRef<SGraphEditor> FBlueprintEditor::CreateGraphEditorWidget(TSharedRef<FT
 		.IsEditable(this, &FBlueprintEditor::IsEditable, InGraph)
 		.TitleBar(TitleBarWidget)
 		.TitleBarEnabledOnly(bEnableTitleBarOnly)
-		.Appearance(this, &FBlueprintEditor::GetGraphAppearance)
+		.Appearance(this, &FBlueprintEditor::GetGraphAppearance, InGraph)
 		.GraphToEdit(InGraph)
 		.GraphEvents(InEvents)
 		.OnNavigateHistoryBack(FSimpleDelegate::CreateSP(this, &FBlueprintEditor::NavigateTab, FDocumentTracker::NavigateBackwards))
@@ -859,9 +859,16 @@ TSharedRef<SGraphEditor> FBlueprintEditor::CreateGraphEditorWidget(TSharedRef<FT
 
 FGraphAppearanceInfo FBlueprintEditor::GetGraphAppearance() const
 {
+	return GetGraphAppearance(GetFocusedGraph());
+}
+
+FGraphAppearanceInfo FBlueprintEditor::GetGraphAppearance(UEdGraph* InGraph) const
+{
+	UBlueprint* Blueprint = (InGraph != nullptr) ? FBlueprintEditorUtils::FindBlueprintForGraph(InGraph) : GetBlueprintObj();
+
 	// Create the appearance info
 	FGraphAppearanceInfo AppearanceInfo;
-	switch ( GetBlueprintObj()->BlueprintType )
+	switch (Blueprint->BlueprintType)
 	{
 	case BPTYPE_LevelScript:
 		AppearanceInfo.CornerText = LOCTEXT("AppearanceCornerText_LevelScript", "LEVEL BLUEPRINT").ToString();
@@ -878,7 +885,7 @@ FGraphAppearanceInfo FBlueprintEditor::GetGraphAppearance() const
 	}
 
 	UEdGraph const* EditingGraph = GetFocusedGraph();
-	if (EditingGraph && BlueprintEditorImpl::GraphHasDefaultNode(EditingGraph))
+	if (InGraph && BlueprintEditorImpl::GraphHasDefaultNode(InGraph))
 	{
 		AppearanceInfo.InstructionText = LOCTEXT("AppearanceInstructionText_DefaultGraph", "Drag Off Pins to Create/Connect New Nodes.");
 	}
@@ -886,7 +893,8 @@ FGraphAppearanceInfo FBlueprintEditor::GetGraphAppearance() const
 	{
 		AppearanceInfo.InstructionText = LOCTEXT("AppearanceInstructionText_EmptyGraph", "Right-Click to Create New Nodes.");
 	}
-	AppearanceInfo.InstructionFade.Bind(this, &FBlueprintEditor::GetInstructionTextOpacity);
+	auto InstructionOpacityDelegate = TAttribute<float>::FGetter::CreateSP(this, &FBlueprintEditor::GetInstructionTextOpacity, InGraph);
+	AppearanceInfo.InstructionFade.Bind(InstructionOpacityDelegate);
 
 	AppearanceInfo.PIENotifyText = GetPIEStatus();
 
@@ -917,7 +925,7 @@ FBlueprintEditor::FBlueprintEditor()
 	, CurrentUISelection(FBlueprintEditor::NoSelection)
 	, bEditorMarkedAsClosed(false)
 	, bCodeBasedProject(false)
-	, bActionMenuIsOpen(false)
+	, HasOpenActionMenu(nullptr)
 	, InstructionsFadeCountdown(0.f)
 {
 	AnalyticsStats.GraphActionMenusNonCtxtSensitiveExecCount = 0;
@@ -2332,7 +2340,7 @@ void FBlueprintEditor::OnGraphEditorDropStreamingLevel(const TArray< TWeakObject
 
 FActionMenuContent FBlueprintEditor::OnCreateGraphActionMenu(UEdGraph* InGraph, const FVector2D& InNodePosition, const TArray<UEdGraphPin*>& InDraggedPins, bool bAutoExpand, SGraphEditor::FActionMenuClosed InOnMenuClosed)
 {
-	bActionMenuIsOpen = true;
+	HasOpenActionMenu = InGraph;
 	if (!BlueprintEditorImpl::GraphHasUserPlacedNodes(InGraph))
 	{
 		InstructionsFadeCountdown = BlueprintEditorImpl::InstructionFadeDuration;
@@ -2370,7 +2378,7 @@ void FBlueprintEditor::OnGraphActionMenuClosed(bool bActionExecuted, bool bConte
 			InstructionsFadeCountdown = 0.0f;
 		}
 	}
-	bActionMenuIsOpen = false;
+	HasOpenActionMenu = nullptr;
 }
 
 void FBlueprintEditor::OnSelectedNodesChanged(const FGraphPanelSelectionSet& NewSelection)
@@ -6633,20 +6641,18 @@ bool FBlueprintEditor::IsGraphPanelEnabled(UEdGraph* InGraph) const
 	return !bIsInterface && !bIsDelegate;
 }
 
-float FBlueprintEditor::GetInstructionTextOpacity() const
+float FBlueprintEditor::GetInstructionTextOpacity(UEdGraph* InGraph) const
 {
-	UEdGraph* EditingGraph = GetFocusedGraph();
 	UBlueprintEditorSettings const* Settings = GetDefault<UBlueprintEditorSettings>();
-
-	if ((EditingGraph == nullptr) || !IsEditable(EditingGraph) || !IsGraphPanelEnabled(EditingGraph) || !Settings->bShowGraphInstructionText)
+	if ((InGraph == nullptr) || !IsEditable(InGraph) || !IsGraphPanelEnabled(InGraph) || !Settings->bShowGraphInstructionText)
 	{
 		return 0.0f;
 	}
-	else if ((InstructionsFadeCountdown > 0.0f) || bActionMenuIsOpen)
+	else if ((InstructionsFadeCountdown > 0.0f) || (HasOpenActionMenu == InGraph))
 	{
 		return InstructionsFadeCountdown / BlueprintEditorImpl::InstructionFadeDuration;
 	}
-	else if (BlueprintEditorImpl::GraphHasUserPlacedNodes(EditingGraph))
+	else if (BlueprintEditorImpl::GraphHasUserPlacedNodes(InGraph))
 	{
 		return 0.0f;
 	}
diff --git a/Engine/Source/Editor/Kismet/Public/BlueprintEditor.h b/Engine/Source/Editor/Kismet/Public/BlueprintEditor.h
--- a/Engine/Source/Editor/Kismet/Public/BlueprintEditor.h
+++ b/Engine/Source/Editor/Kismet/Public/BlueprintEditor.h
@@ -235,12 +235,11 @@ public:
 
 	/** Returns whether a graph is editable or not */
 	virtual bool IsEditable(UEdGraph* InGraph) const;
-
-	/** */
+	/** Determines if the graph's title bar should be the only interactable widget .*/
 	bool IsGraphPanelEnabled(UEdGraph* InGraph) const;
 
-	/** */
-	virtual float GetInstructionTextOpacity() const;
+	/** Used to determine the visibility of the graph's instruction text. */
+	float GetInstructionTextOpacity(UEdGraph* InGraph) const;
 
 	/** Returns true if in editing mode */
 	bool InEditingMode() const;
@@ -827,6 +826,8 @@ protected:
 
 	/** Get graph appearance */
 	virtual FGraphAppearanceInfo GetGraphAppearance() const;
+	/** Used to get the apperance of a specific graph, GetGraphAppearance() uses the currently focused graph. */
+	FGraphAppearanceInfo GetGraphAppearance(class UEdGraph* InGraph) const;
 
 private:
 
@@ -1041,9 +1042,9 @@ private:
 	/** Delegates that are fired when the blueprint editor changes modes */
 	FOnModeSet OnModeSetData;
 
-	/** */
-	bool bActionMenuIsOpen;
-	/** */
+	/** When set, flags which graph has a action menu currently open (if null, no graphs do). */
+	UEdGraph* HasOpenActionMenu;
+	/** Used to nicely fade instruction text, when the context menu is opened. */
 	float InstructionsFadeCountdown;
 };
 
