Removing the blueprint palette by default.

diff --git a/Engine/Source/Editor/Kismet/Private/BlueprintEditorModes.cpp b/Engine/Source/Editor/Kismet/Private/BlueprintEditorModes.cpp
--- a/Engine/Source/Editor/Kismet/Private/BlueprintEditorModes.cpp
+++ b/Engine/Source/Editor/Kismet/Private/BlueprintEditorModes.cpp
@@ -117,7 +117,7 @@ FBlueprintEditorApplicationMode::FBlueprintEditorApplicationMode(TSharedPtr<clas
 					->Split
 					(
 						FTabManager::NewStack()
-						->AddTab( FBlueprintEditorTabs::PaletteID, ETabState::OpenedTab )
+						->AddTab( FBlueprintEditorTabs::PaletteID, ETabState::ClosedTab )
 					)
 				)
 			)
@@ -473,7 +473,7 @@ FBlueprintMacroApplicationMode::FBlueprintMacroApplicationMode(TSharedPtr<class
 					->Split
 					(
 						FTabManager::NewStack()
-						->AddTab( FBlueprintEditorTabs::PaletteID, ETabState::OpenedTab )
+						->AddTab( FBlueprintEditorTabs::PaletteID, ETabState::ClosedTab )
 					)
 				)
 			)
diff --git a/Engine/Source/Editor/UMGEditor/Private/BlueprintModes/WidgetGraphApplicationMode.cpp b/Engine/Source/Editor/UMGEditor/Private/BlueprintModes/WidgetGraphApplicationMode.cpp
--- a/Engine/Source/Editor/UMGEditor/Private/BlueprintModes/WidgetGraphApplicationMode.cpp
+++ b/Engine/Source/Editor/UMGEditor/Private/BlueprintModes/WidgetGraphApplicationMode.cpp
@@ -71,7 +71,7 @@ FWidgetGraphApplicationMode::FWidgetGraphApplicationMode(TSharedPtr<FWidgetBluep
 					->Split
 					(
 						FTabManager::NewStack()
-						->AddTab( FBlueprintEditorTabs::PaletteID, ETabState::OpenedTab )
+						->AddTab( FBlueprintEditorTabs::PaletteID, ETabState::ClosedTab )
 					)
 				)
 			)
