[343351] Blueprints - Do not allow variable drag-and-drop onto a graph when it would create an invalid self reference to an external object.

- Modified FKismetVariableDragDropAction::CanVariableBeDropped() to always check for an external Blueprint reference (even for non-local variables).
- Modified FKismetVariableDragDropAction::HoverTargetChanged() to display a proper warning message for a blocked non-local variable drag-and-drop attempt.

diff --git a/Engine/Source/Editor/Kismet/Private/BPVariableDragDropAction.cpp b/Engine/Source/Editor/Kismet/Private/BPVariableDragDropAction.cpp
--- a/Engine/Source/Editor/Kismet/Private/BPVariableDragDropAction.cpp
+++ b/Engine/Source/Editor/Kismet/Private/BPVariableDragDropAction.cpp
@@ -83,7 +83,15 @@ void FKismetVariableDragDropAction::HoverTargetChanged()
 		Args.Add(TEXT("Scope"), FText::FromString(HoveredGraph->GetName()));
 
 		StatusSymbol = FEditorStyle::GetBrush(TEXT("Graph.ConnectorFeedback.Error"));
-		Message = FText::Format( LOCTEXT("IncorrectGraph_Error", "Cannot place local variable '{VariableName}' in external scope '{Scope}'"), Args);
+
+		if(IsFromBlueprint(FBlueprintEditorUtils::FindBlueprintForGraph(HoveredGraph)) && VariableProperty->GetOuter()->IsA(UFunction::StaticClass()))
+		{
+			Message = FText::Format( LOCTEXT("IncorrectGraphForLocalVariable_Error", "Cannot place local variable '{VariableName}' in external scope '{Scope}'"), Args);
+		}
+		else
+		{
+			Message = FText::Format( LOCTEXT("IncorrectGraphForVariable_Error", "Cannot place variable '{VariableName}' in external scope '{Scope}'"), Args);
+		}
 	}
 	else if (PinUnderCursor != NULL)
 	{
@@ -553,12 +561,14 @@ bool FKismetVariableDragDropAction::CanVariableBeDropped(const UProperty* InVari
 {
 	UObject* Outer = InVariableProperty->GetOuter();
 
+	// Only allow variables to be placed within the same blueprint (otherwise the self context on the dropped node will be invalid)
+	bool bCanVariableBeDropped = IsFromBlueprint(FBlueprintEditorUtils::FindBlueprintForGraph(&InGraph));
+
 	// Local variables have some special conditions for being allowed to be placed
-	bool bCanVariableBeDropped = true;
-	if(Outer->IsA(UFunction::StaticClass()))
+	if(bCanVariableBeDropped && Outer->IsA(UFunction::StaticClass()))
 	{
 		// Check if the top level graph has the same name as the function, if they do not then the variable cannot be placed in the graph
-		if(FBlueprintEditorUtils::GetTopLevelGraph(&InGraph)->GetFName() != Outer->GetFName() || !IsFromBlueprint(FBlueprintEditorUtils::FindBlueprintForGraph(&InGraph)) )
+		if(FBlueprintEditorUtils::GetTopLevelGraph(&InGraph)->GetFName() != Outer->GetFName())
 		{
 			bCanVariableBeDropped = false;
 		}
