Making it so the blueprint favoriting system works with the new menu system.

diff --git a/Engine/Source/Editor/BlueprintGraph/Classes/BlueprintBoundEventNodeSpawner.h b/Engine/Source/Editor/BlueprintGraph/Classes/BlueprintBoundEventNodeSpawner.h
--- a/Engine/Source/Editor/BlueprintGraph/Classes/BlueprintBoundEventNodeSpawner.h
+++ b/Engine/Source/Editor/BlueprintGraph/Classes/BlueprintBoundEventNodeSpawner.h
@@ -32,6 +32,7 @@ public:
 	static UBlueprintBoundEventNodeSpawner* Create(TSubclassOf<UK2Node_Event> NodeClass, UMulticastDelegateProperty* EventDelegate, UObject* Outer = nullptr);
 
 	// UBlueprintNodeSpawner interface
+	virtual FBlueprintNodeSpawnerSignature GetSpawnerSignature() const override;
 	virtual UEdGraphNode* Invoke(UEdGraph* ParentGraph, FBindingSet const& Bindings, FVector2D const Location) const override;
 	virtual FText GetDefaultMenuName() const override;
 	virtual FText GetDefaultMenuCategory() const override;
@@ -56,5 +57,5 @@ public:
 private:
 	/** */
 	UPROPERTY()
-	TWeakObjectPtr<UMulticastDelegateProperty> EventDelegate;
+	UMulticastDelegateProperty* EventDelegate;
 };
diff --git a/Engine/Source/Editor/BlueprintGraph/Classes/BlueprintBoundNodeSpawner.h b/Engine/Source/Editor/BlueprintGraph/Classes/BlueprintBoundNodeSpawner.h
--- a/Engine/Source/Editor/BlueprintGraph/Classes/BlueprintBoundNodeSpawner.h
+++ b/Engine/Source/Editor/BlueprintGraph/Classes/BlueprintBoundNodeSpawner.h
@@ -22,6 +22,10 @@ public:
 	 */
 	static UBlueprintBoundNodeSpawner* Create(TSubclassOf<UEdGraphNode> NodeClass, UObject* Outer = nullptr);
 
+	// UBlueprintNodeSpawner interface
+	virtual FBlueprintNodeSpawnerSignature GetSpawnerSignature() const override;
+	// End UBlueprintNodeSpawner interface
+
 	// FBlueprintNodeBinder interface
 	virtual bool IsBindingCompatible(UObject const* BindingCandidate) const override;
 	virtual bool BindToNode(UEdGraphNode* Node, UObject* Binding) const override;
diff --git a/Engine/Source/Editor/BlueprintGraph/Classes/BlueprintComponentNodeSpawner.h b/Engine/Source/Editor/BlueprintGraph/Classes/BlueprintComponentNodeSpawner.h
--- a/Engine/Source/Editor/BlueprintGraph/Classes/BlueprintComponentNodeSpawner.h
+++ b/Engine/Source/Editor/BlueprintGraph/Classes/BlueprintComponentNodeSpawner.h
@@ -32,6 +32,7 @@ public:
 	static UBlueprintComponentNodeSpawner* Create(TSubclassOf<UActorComponent> const ComponentClass, UObject* Outer = nullptr);
 
 	// UBlueprintNodeSpawner interface
+	virtual FBlueprintNodeSpawnerSignature GetSpawnerSignature() const override;
 	virtual UEdGraphNode* Invoke(UEdGraph* ParentGraph, FBindingSet const& Bindings, FVector2D const Location) const override;
 	virtual FText GetDefaultMenuName() const override;
 	virtual FText GetDefaultMenuCategory() const override;
diff --git a/Engine/Source/Editor/BlueprintGraph/Classes/BlueprintDelegateNodeSpawner.h b/Engine/Source/Editor/BlueprintGraph/Classes/BlueprintDelegateNodeSpawner.h
--- a/Engine/Source/Editor/BlueprintGraph/Classes/BlueprintDelegateNodeSpawner.h
+++ b/Engine/Source/Editor/BlueprintGraph/Classes/BlueprintDelegateNodeSpawner.h
@@ -2,7 +2,7 @@
 
 #pragma once
 
-#include "BlueprintNodeSpawner.h"
+#include "BlueprintFieldNodeSpawner.h"
 #include "BlueprintDelegateNodeSpawner.generated.h"
 
 // Forward declarations
@@ -14,7 +14,7 @@ class UK2Node_BaseMCDelegate;
  * FEdGraphSchemaAction_K2Delegate, FEdGraphSchemaAction_K2AssignDelegate, etc.
  */
 UCLASS(Transient)
-class BLUEPRINTGRAPH_API UBlueprintDelegateNodeSpawner : public UBlueprintNodeSpawner
+class BLUEPRINTGRAPH_API UBlueprintDelegateNodeSpawner : public UBlueprintFieldNodeSpawner
 {
 	GENERATED_UCLASS_BODY()
 
@@ -32,7 +32,6 @@ public:
 	static UBlueprintDelegateNodeSpawner* Create(TSubclassOf<UK2Node_BaseMCDelegate> NodeClass, UMulticastDelegateProperty const* const Property, UObject* Outer = nullptr);
 
 	// UBlueprintNodeSpawner interface
-	virtual UEdGraphNode* Invoke(UEdGraph* ParentGraph, FBindingSet const& Bindings, FVector2D const Location) const override;
 	virtual FText GetDefaultMenuName() const override;
 	virtual FText GetDefaultMenuCategory() const override;
 	virtual FName GetDefaultMenuIcon(FLinearColor& ColorOut) const override;
@@ -45,9 +44,4 @@ public:
 	 * @return The delegate property that this was initialized with.
 	 */
 	UMulticastDelegateProperty const* GetProperty() const;
-	
-private:
-	/** */
-	UPROPERTY()
-	TWeakObjectPtr<UMulticastDelegateProperty> Property;
 };
diff --git a/Engine/Source/Editor/BlueprintGraph/Classes/BlueprintEventNodeSpawner.h b/Engine/Source/Editor/BlueprintGraph/Classes/BlueprintEventNodeSpawner.h
--- a/Engine/Source/Editor/BlueprintGraph/Classes/BlueprintEventNodeSpawner.h
+++ b/Engine/Source/Editor/BlueprintGraph/Classes/BlueprintEventNodeSpawner.h
@@ -45,6 +45,7 @@ public:
 	static UBlueprintEventNodeSpawner* Create(TSubclassOf<UK2Node_Event> NodeClass, FName CustomEventName, UObject* Outer = nullptr);
 
 	// UBlueprintNodeSpawner interface
+	virtual FBlueprintNodeSpawnerSignature GetSpawnerSignature() const override;
 	virtual UEdGraphNode* Invoke(UEdGraph* ParentGraph, FBindingSet const& Bindings, FVector2D const Location) const override;
 	virtual FText GetDefaultMenuName() const override;
 	virtual FText GetDefaultMenuCategory() const override;
diff --git a/Engine/Source/Editor/BlueprintGraph/Classes/BlueprintFieldNodeSpawner.h b/Engine/Source/Editor/BlueprintGraph/Classes/BlueprintFieldNodeSpawner.h
new file mode 100644
--- /dev/null
+++ b/Engine/Source/Editor/BlueprintGraph/Classes/BlueprintFieldNodeSpawner.h
@@ -0,0 +1,55 @@
+// Copyright 1998-2014 Epic Games, Inc. All Rights Reserved.
+
+#pragma once
+
+#include "BlueprintNodeSpawner.h"
+#include "BlueprintFieldNodeSpawner.generated.h"
+
+// Forward declarations
+class UK2Node;
+
+/**
+ * Takes care of spawning various field related nodes (nodes associated with 
+ * functions, enums, structs, properties, etc.). Acts as the "action" portion
+ * for certain FBlueprintActionMenuItems. 
+ */
+UCLASS(Transient)
+class BLUEPRINTGRAPH_API UBlueprintFieldNodeSpawner : public UBlueprintNodeSpawner
+{
+	GENERATED_UCLASS_BODY()
+	DECLARE_DELEGATE_TwoParams(FSetNodeFieldDelegate, UEdGraphNode*, UField const*);
+
+public:
+	/**
+	 * Creates a new UBlueprintFieldNodeSpawner for the supplied field.
+	 * Does not do any compatibility checking to ensure that the field is
+	 * viable for blueprint use.
+	 * 
+	 * @param  NodeClass	The type of node you want the spawner to create.
+	 * @param  Field		The field you want assigned to new nodes.
+	 * @param  Outer		Optional outer for the new spawner (if left null, the transient package will be used).
+	 * @return A newly allocated instance of this class.
+	 */
+	static UBlueprintFieldNodeSpawner* Create(TSubclassOf<UK2Node> NodeClass, UField const* const Field, UObject* Outer = nullptr);
+
+	// UBlueprintNodeSpawner interface
+	virtual FBlueprintNodeSpawnerSignature GetSpawnerSignature() const override;
+	virtual UEdGraphNode* Invoke(UEdGraph* ParentGraph, FBindingSet const& Bindings, FVector2D const Location) const override;
+	// End UBlueprintNodeSpawner interface
+
+	/** */
+	FSetNodeFieldDelegate SetNodeFieldDelegate;
+
+	/**
+	 * Retrieves the field that this assigns to spawned nodes (defines the 
+	 * node's signature).
+	 *
+	 * @return The field that this class was initialized with.
+	 */
+	UField const* GetField() const;
+
+protected:
+	/** The field to configure new nodes with. */
+	UPROPERTY()
+	UField const* Field;
+};
diff --git a/Engine/Source/Editor/BlueprintGraph/Classes/BlueprintFunctionNodeSpawner.h b/Engine/Source/Editor/BlueprintGraph/Classes/BlueprintFunctionNodeSpawner.h
--- a/Engine/Source/Editor/BlueprintGraph/Classes/BlueprintFunctionNodeSpawner.h
+++ b/Engine/Source/Editor/BlueprintGraph/Classes/BlueprintFunctionNodeSpawner.h
@@ -2,7 +2,7 @@
 
 #pragma once
 
-#include "BlueprintNodeSpawner.h"
+#include "BlueprintFieldNodeSpawner.h"
 #include "BlueprintFunctionNodeSpawner.generated.h"
 
 // Forward declarations
@@ -13,7 +13,7 @@ class UK2Node_CallFunction;
  * "action" portion of certain FBlueprintActionMenuItems. 
  */
 UCLASS(Transient)
-class BLUEPRINTGRAPH_API UBlueprintFunctionNodeSpawner : public UBlueprintNodeSpawner
+class BLUEPRINTGRAPH_API UBlueprintFunctionNodeSpawner : public UBlueprintFieldNodeSpawner
 {
 	GENERATED_UCLASS_BODY()
 
@@ -64,9 +64,4 @@ public:
 	 * @return The function that this class was initialized with.
 	 */
 	UFunction const* GetFunction() const;
-
-private:
-	/** The function to configure new nodes with. */
-	UPROPERTY()
-	UFunction const* Function;
 };
diff --git a/Engine/Source/Editor/BlueprintGraph/Classes/BlueprintVariableNodeSpawner.h b/Engine/Source/Editor/BlueprintGraph/Classes/BlueprintVariableNodeSpawner.h
--- a/Engine/Source/Editor/BlueprintGraph/Classes/BlueprintVariableNodeSpawner.h
+++ b/Engine/Source/Editor/BlueprintGraph/Classes/BlueprintVariableNodeSpawner.h
@@ -2,7 +2,7 @@
 
 #pragma once
 
-#include "BlueprintNodeSpawner.h"
+#include "BlueprintFieldNodeSpawner.h"
 #include "EdGraph/EdGraphNodeUtils.h" // for FNodeTextCache
 #include "BlueprintVariableNodeSpawner.generated.h"
 
@@ -17,7 +17,7 @@ class UK2Node_Variable;
  * local function variables.
  */
 UCLASS(Transient)
-class BLUEPRINTGRAPH_API UBlueprintVariableNodeSpawner : public UBlueprintNodeSpawner
+class BLUEPRINTGRAPH_API UBlueprintVariableNodeSpawner : public UBlueprintFieldNodeSpawner
 {
 	GENERATED_UCLASS_BODY()
 
@@ -47,6 +47,7 @@ public:
 
 	// UBlueprintNodeSpawner interface
 	virtual void Prime() override;
+	virtual FBlueprintNodeSpawnerSignature GetSpawnerSignature() const override;
 	virtual UEdGraphNode* Invoke(UEdGraph* ParentGraph, FBindingSet const& Bindings, FVector2D const Location) const override;
 	virtual FText GetDefaultMenuName() const override;
 	virtual FText GetDefaultMenuCategory() const override;
@@ -99,17 +100,13 @@ private:
 	 */
 	FText GetVariableName() const;
 
-	/** */
-	UPROPERTY()
-	TWeakObjectPtr<UProperty> MemberVariable;
-
 	/** The graph that the local variable belongs to (if this is a local variable spawner). */
 	UPROPERTY()
-	TWeakObjectPtr<UEdGraph>  LocalVarOuter;
+	UEdGraph* LocalVarOuter;
 
 	/** */
 	UPROPERTY()
-	FBPVariableDescription    LocalVarDesc;
+	FBPVariableDescription LocalVarDesc;
 
 	/** Constructing FText strings can be costly, so we cache the default menu name/category/tooltip */
 	FNodeTextCache CachedMenuName;
diff --git a/Engine/Source/Editor/BlueprintGraph/Classes/K2Node_VariableGet.h b/Engine/Source/Editor/BlueprintGraph/Classes/K2Node_VariableGet.h
--- a/Engine/Source/Editor/BlueprintGraph/Classes/K2Node_VariableGet.h
+++ b/Engine/Source/Editor/BlueprintGraph/Classes/K2Node_VariableGet.h
@@ -24,7 +24,7 @@ class BLUEPRINTGRAPH_API UK2Node_VariableGet : public UK2Node_Variable
 	virtual class FNodeHandlingFunctor* CreateNodeHandler(class FKismetCompilerContext& CompilerContext) const override;
 	// End K2Node interface
 
-	static FText GetPropertyTooltip(UProperty* VariableProperty);
+	static FText GetPropertyTooltip(UProperty const* VariableProperty);
 	static FText GetBlueprintVarTooltip(FBPVariableDescription const& VarDesc);
 
 private:
diff --git a/Engine/Source/Editor/BlueprintGraph/Classes/K2Node_VariableSet.h b/Engine/Source/Editor/BlueprintGraph/Classes/K2Node_VariableSet.h
--- a/Engine/Source/Editor/BlueprintGraph/Classes/K2Node_VariableSet.h
+++ b/Engine/Source/Editor/BlueprintGraph/Classes/K2Node_VariableSet.h
@@ -28,7 +28,7 @@ class UK2Node_VariableSet : public UK2Node_Variable
 	BLUEPRINTGRAPH_API FName GetRepNotifyName() const;
 	BLUEPRINTGRAPH_API bool ShouldFlushDormancyOnSet() const;
 
-	static FText GetPropertyTooltip(UProperty* VariableProperty);
+	static FText GetPropertyTooltip(UProperty const* VariableProperty);
 	static FText GetBlueprintVarTooltip(FBPVariableDescription const& VarDesc);
 
 private:
diff --git a/Engine/Source/Editor/BlueprintGraph/Private/BlueprintActionFilter.cpp b/Engine/Source/Editor/BlueprintGraph/Private/BlueprintActionFilter.cpp
--- a/Engine/Source/Editor/BlueprintGraph/Private/BlueprintActionFilter.cpp
+++ b/Engine/Source/Editor/BlueprintGraph/Private/BlueprintActionFilter.cpp
@@ -1047,6 +1047,7 @@ namespace EBlueprintActionInfoFlags
 		CachedField    = (1 << 1),
 		CachedProperty = (1 << 2),
 		CachedFunction = (1 << 3),
+		CachedGuid     = (1 << 4),
 	};
 }
 
@@ -1117,6 +1118,17 @@ UClass const* FBlueprintActionInfo::GetNodeClass()
 	return NodeClass;
 }
 
+//------------------------------------------------------------------------------
+FGuid FBlueprintActionInfo::GetActionGuid()
+{
+	if ((CacheFlags & EBlueprintActionInfoFlags::CachedGuid) == 0)
+	{
+		CachedActionId = NodeSpawner->GetSpawnerSignature().AsGuid();
+		CacheFlags |= EBlueprintActionInfoFlags::CachedGuid;
+	}
+	return CachedActionId;
+}
+
 //------------------------------------------------------------------------------
 UField const* FBlueprintActionInfo::GetAssociatedMemberField()
 {
diff --git a/Engine/Source/Editor/BlueprintGraph/Private/BlueprintBoundEventNodeSpawner.cpp b/Engine/Source/Editor/BlueprintGraph/Private/BlueprintBoundEventNodeSpawner.cpp
--- a/Engine/Source/Editor/BlueprintGraph/Private/BlueprintBoundEventNodeSpawner.cpp
+++ b/Engine/Source/Editor/BlueprintGraph/Private/BlueprintBoundEventNodeSpawner.cpp
@@ -34,6 +34,15 @@ UBlueprintBoundEventNodeSpawner::UBlueprintBoundEventNodeSpawner(class FPostCons
 {
 }
 
+//------------------------------------------------------------------------------
+FBlueprintNodeSpawnerSignature UBlueprintBoundEventNodeSpawner::GetSpawnerSignature() const
+{
+	// explicit actions for binding (like this) cannot be reconstructed form a 
+	// signature (since this spawner does not own whatever it will be binding 
+	// to), therefore we return an empty (invalid) signature
+	return FBlueprintNodeSpawnerSignature();
+}
+
 //------------------------------------------------------------------------------
 UEdGraphNode* UBlueprintBoundEventNodeSpawner::Invoke(UEdGraph* ParentGraph, FBindingSet const& Bindings, FVector2D const Location) const
 {
@@ -127,7 +136,7 @@ bool UBlueprintBoundEventNodeSpawner::BindToNode(UEdGraphNode* Node, UObject* Bi
 	{
 		if (UObjectProperty const* BoundProperty = Cast<UObjectProperty>(Binding))
 		{
-			EventNode->InitializeComponentBoundEventParams(BoundProperty, EventDelegate.Get());
+			EventNode->InitializeComponentBoundEventParams(BoundProperty, EventDelegate);
 			bWasBound = true;
 			Node->ReconstructNode();
 		}
@@ -136,7 +145,7 @@ bool UBlueprintBoundEventNodeSpawner::BindToNode(UEdGraphNode* Node, UObject* Bi
 	{
 		if (AActor* BoundActor = Cast<AActor>(Binding))
 		{
-			EventNode->InitializeActorBoundEventParams(BoundActor, EventDelegate.Get());
+			EventNode->InitializeActorBoundEventParams(BoundActor, EventDelegate);
 			bWasBound = true;
 			Node->ReconstructNode();
 		}
@@ -147,7 +156,7 @@ bool UBlueprintBoundEventNodeSpawner::BindToNode(UEdGraphNode* Node, UObject* Bi
 //------------------------------------------------------------------------------
 UMulticastDelegateProperty const* UBlueprintBoundEventNodeSpawner::GetEventDelegate() const
 {
-	return EventDelegate.Get();
+	return EventDelegate;
 }
 
 #undef LOCTEXT_NAMESPACE
diff --git a/Engine/Source/Editor/BlueprintGraph/Private/BlueprintBoundNodeSpawner.cpp b/Engine/Source/Editor/BlueprintGraph/Private/BlueprintBoundNodeSpawner.cpp
--- a/Engine/Source/Editor/BlueprintGraph/Private/BlueprintBoundNodeSpawner.cpp
+++ b/Engine/Source/Editor/BlueprintGraph/Private/BlueprintBoundNodeSpawner.cpp
@@ -29,6 +29,15 @@ UBlueprintBoundNodeSpawner::UBlueprintBoundNodeSpawner(class FPostConstructIniti
 {
 }
 
+//------------------------------------------------------------------------------
+FBlueprintNodeSpawnerSignature UBlueprintBoundNodeSpawner::GetSpawnerSignature() const
+{
+	// explicit actions for binding (like this) cannot be reconstructed form a 
+	// signature (since this spawner does not own whatever it will be binding 
+	// to), therefore we return an empty (invalid) signature
+	return FBlueprintNodeSpawnerSignature();
+}
+
 //------------------------------------------------------------------------------
 bool UBlueprintBoundNodeSpawner::IsBindingCompatible(UObject const* BindingCandidate) const
 {
diff --git a/Engine/Source/Editor/BlueprintGraph/Private/BlueprintComponentNodeSpawner.cpp b/Engine/Source/Editor/BlueprintGraph/Private/BlueprintComponentNodeSpawner.cpp
--- a/Engine/Source/Editor/BlueprintGraph/Private/BlueprintComponentNodeSpawner.cpp
+++ b/Engine/Source/Editor/BlueprintGraph/Private/BlueprintComponentNodeSpawner.cpp
@@ -35,6 +35,14 @@ UBlueprintComponentNodeSpawner::UBlueprintComponentNodeSpawner(class FPostConstr
 {
 }
 
+//------------------------------------------------------------------------------
+FBlueprintNodeSpawnerSignature UBlueprintComponentNodeSpawner::GetSpawnerSignature() const
+{
+	FBlueprintNodeSpawnerSignature SpawnerSignature = Super::GetSpawnerSignature();
+	SpawnerSignature.AddSubObject(ComponentClass);
+	return SpawnerSignature;
+}
+
 //------------------------------------------------------------------------------
 // Evolved from a combination of FK2ActionMenuBuilder::CreateAddComponentAction()
 // and FEdGraphSchemaAction_K2AddComponent::PerformAction().
diff --git a/Engine/Source/Editor/BlueprintGraph/Private/BlueprintDelegateNodeSpawner.cpp b/Engine/Source/Editor/BlueprintGraph/Private/BlueprintDelegateNodeSpawner.cpp
--- a/Engine/Source/Editor/BlueprintGraph/Private/BlueprintDelegateNodeSpawner.cpp
+++ b/Engine/Source/Editor/BlueprintGraph/Private/BlueprintDelegateNodeSpawner.cpp
@@ -24,22 +24,13 @@ UBlueprintDelegateNodeSpawner* UBlueprintDelegateNodeSpawner::Create(TSubclassOf
 	}
 
 	UBlueprintDelegateNodeSpawner* NodeSpawner = NewObject<UBlueprintDelegateNodeSpawner>(Outer);
-	NodeSpawner->Property  = Property;
+	NodeSpawner->Field     = Property;
 	NodeSpawner->NodeClass = NodeClass;
-	return NodeSpawner;
-}
 
-//------------------------------------------------------------------------------
-UBlueprintDelegateNodeSpawner::UBlueprintDelegateNodeSpawner(class FPostConstructInitializeProperties const& PCIP)
-	: Super(PCIP)
-{
-}
-
-//------------------------------------------------------------------------------
-UEdGraphNode* UBlueprintDelegateNodeSpawner::Invoke(UEdGraph* ParentGraph, FBindingSet const& Bindings, FVector2D const Location) const
-{
-	auto PostSpawnSetupLambda = [](UEdGraphNode* NewNode, bool bIsTemplateNode, UMulticastDelegateProperty const* Property, FCustomizeNodeDelegate UserDelegate)
+	auto SetDelegateLambda = [](UEdGraphNode* NewNode, UField const* Field)
 	{
+		UMulticastDelegateProperty const* Property = Cast<UMulticastDelegateProperty>(Field);
+
 		UK2Node_BaseMCDelegate* DelegateNode = Cast<UK2Node_BaseMCDelegate>(NewNode);
 		if ((DelegateNode != nullptr) && (Property != nullptr))
 		{
@@ -49,11 +40,16 @@ UEdGraphNode* UBlueprintDelegateNodeSpawner::Invoke(UEdGraph* ParentGraph, FBind
 
 			DelegateNode->SetFromProperty(Property, bIsSelfContext);
 		}
-		UserDelegate.ExecuteIfBound(NewNode, bIsTemplateNode);
 	};
+	NodeSpawner->SetNodeFieldDelegate = FSetNodeFieldDelegate::CreateStatic(SetDelegateLambda);
+
+	return NodeSpawner;
+}
 
-	FCustomizeNodeDelegate PostSpawnSetupDelegate = FCustomizeNodeDelegate::CreateStatic(PostSpawnSetupLambda, GetProperty(), CustomizeNodeDelegate);
-	return Super::Invoke(ParentGraph, Bindings, Location, PostSpawnSetupDelegate);
+//------------------------------------------------------------------------------
+UBlueprintDelegateNodeSpawner::UBlueprintDelegateNodeSpawner(class FPostConstructInitializeProperties const& PCIP)
+	: Super(PCIP)
+{
 }
 
 //------------------------------------------------------------------------------
@@ -108,8 +104,7 @@ FName UBlueprintDelegateNodeSpawner::GetDefaultMenuIcon(FLinearColor& ColorOut)
 //------------------------------------------------------------------------------
 UMulticastDelegateProperty const* UBlueprintDelegateNodeSpawner::GetProperty() const
 {
-	// Get() does IsValid() checks for us
-	return Property.Get();
+	return Cast<UMulticastDelegateProperty>(GetField());
 }
 
 #undef LOCTEXT_NAMESPACE
diff --git a/Engine/Source/Editor/BlueprintGraph/Private/BlueprintEditorSettings.cpp b/Engine/Source/Editor/BlueprintGraph/Private/BlueprintEditorSettings.cpp
--- a/Engine/Source/Editor/BlueprintGraph/Private/BlueprintEditorSettings.cpp
+++ b/Engine/Source/Editor/BlueprintGraph/Private/BlueprintEditorSettings.cpp
@@ -3,6 +3,7 @@
 #include "BlueprintGraphPrivatePCH.h"
 #include "BlueprintEditorSettings.h"
 #include "Editor/UnrealEd/Classes/Settings/EditorExperimentalSettings.h"
+#include "Editor/UnrealEd/Classes/Editor/EditorUserSettings.h"
 
 UBlueprintEditorSettings::UBlueprintEditorSettings(const class FPostConstructInitializeProperties& PCIP)
 	: Super(PCIP)
@@ -13,4 +14,8 @@ UBlueprintEditorSettings::UBlueprintEditorSettings(const class FPostConstructIni
 	// settings that were moved out of experimental...
 	UEditorExperimentalSettings const* ExperimentalSettings = GetDefault<UEditorExperimentalSettings>();
 	bDrawMidpointArrowsInBlueprints = ExperimentalSettings->bDrawMidpointArrowsInBlueprints;
+
+	// settings that were moved out of editor-user settings...
+	UEditorUserSettings const* UserSettings = GetDefault<UEditorUserSettings>();
+	bShowActionMenuItemSignatures = UserSettings->bDisplayActionListItemRefIds;
 }
diff --git a/Engine/Source/Editor/BlueprintGraph/Private/BlueprintEventNodeSpawner.cpp b/Engine/Source/Editor/BlueprintGraph/Private/BlueprintEventNodeSpawner.cpp
--- a/Engine/Source/Editor/BlueprintGraph/Private/BlueprintEventNodeSpawner.cpp
+++ b/Engine/Source/Editor/BlueprintGraph/Private/BlueprintEventNodeSpawner.cpp
@@ -88,6 +88,22 @@ UBlueprintEventNodeSpawner::UBlueprintEventNodeSpawner(class FPostConstructIniti
 {
 }
 
+//------------------------------------------------------------------------------
+FBlueprintNodeSpawnerSignature UBlueprintEventNodeSpawner::GetSpawnerSignature() const
+{
+	FBlueprintNodeSpawnerSignature SpawnerSignature = Super::GetSpawnerSignature();
+	if (IsForCustomEvent() && !CustomEventName.IsNone())
+	{
+		static const FName CustomSignatureKey(TEXT("CustomEvent"));
+		SpawnerSignature.AddKeyValue(CustomSignatureKey, CustomEventName.ToString());
+	}
+	else
+	{
+		SpawnerSignature.AddSubObject(EventFunc);
+	}
+	return SpawnerSignature;
+}
+
 //------------------------------------------------------------------------------
 UEdGraphNode* UBlueprintEventNodeSpawner::Invoke(UEdGraph* ParentGraph, FBindingSet const& Bindings, FVector2D const Location) const
 {
diff --git a/Engine/Source/Editor/BlueprintGraph/Private/BlueprintFieldNodeSpawner.cpp b/Engine/Source/Editor/BlueprintGraph/Private/BlueprintFieldNodeSpawner.cpp
new file mode 100644
--- /dev/null
+++ b/Engine/Source/Editor/BlueprintGraph/Private/BlueprintFieldNodeSpawner.cpp
@@ -0,0 +1,59 @@
+// Copyright 1998-2014 Epic Games, Inc. All Rights Reserved.
+
+#include "BlueprintGraphPrivatePCH.h"
+#include "BlueprintFieldNodeSpawner.h"
+
+#define LOCTEXT_NAMESPACE "BlueprintFieldNodeSpawner"
+
+//------------------------------------------------------------------------------
+UBlueprintFieldNodeSpawner* UBlueprintFieldNodeSpawner::Create(TSubclassOf<UK2Node> NodeClass, UField const* const Field, UObject* Outer/* = nullptr*/)
+{
+	if (Outer == nullptr)
+	{
+		Outer = GetTransientPackage();
+	}
+	UBlueprintFieldNodeSpawner* NodeSpawner = NewObject<UBlueprintFieldNodeSpawner>(Outer);
+	NodeSpawner->Field     = Field;
+	NodeSpawner->NodeClass = NodeClass;
+	
+	return NodeSpawner;
+}
+
+//------------------------------------------------------------------------------
+UBlueprintFieldNodeSpawner::UBlueprintFieldNodeSpawner(class FPostConstructInitializeProperties const& PCIP)
+	: Super(PCIP)
+	, Field(nullptr)
+{
+}
+
+//------------------------------------------------------------------------------
+FBlueprintNodeSpawnerSignature UBlueprintFieldNodeSpawner::GetSpawnerSignature() const
+{
+	FBlueprintNodeSpawnerSignature SpawnerSignature = Super::GetSpawnerSignature();
+	SpawnerSignature.AddSubObject(Field);
+
+	return SpawnerSignature;
+}
+
+//------------------------------------------------------------------------------
+UEdGraphNode* UBlueprintFieldNodeSpawner::Invoke(UEdGraph* ParentGraph, FBindingSet const& Bindings, FVector2D const Location) const
+{
+	auto PostSpawnSetupLambda = [](UEdGraphNode* NewNode, bool bIsTemplateNode, UField const* Field, FSetNodeFieldDelegate SetFieldDelegate, FCustomizeNodeDelegate UserDelegate)
+	{
+		SetFieldDelegate.ExecuteIfBound(NewNode, Field);
+		UserDelegate.ExecuteIfBound(NewNode, bIsTemplateNode);
+	};
+
+	FCustomizeNodeDelegate PostSpawnSetupDelegate = FCustomizeNodeDelegate::CreateStatic(PostSpawnSetupLambda, GetField(), SetNodeFieldDelegate, CustomizeNodeDelegate);
+	UEdGraphNode* SpawnedNode = Super::Invoke(ParentGraph, Bindings, Location, PostSpawnSetupDelegate);
+
+	return SpawnedNode;
+}
+
+//------------------------------------------------------------------------------
+UField const* UBlueprintFieldNodeSpawner::GetField() const
+{
+	return Field;
+}
+
+#undef LOCTEXT_NAMESPACE
diff --git a/Engine/Source/Editor/BlueprintGraph/Private/BlueprintFunctionNodeSpawner.cpp b/Engine/Source/Editor/BlueprintGraph/Private/BlueprintFunctionNodeSpawner.cpp
--- a/Engine/Source/Editor/BlueprintGraph/Private/BlueprintFunctionNodeSpawner.cpp
+++ b/Engine/Source/Editor/BlueprintGraph/Private/BlueprintFunctionNodeSpawner.cpp
@@ -172,7 +172,7 @@ UBlueprintFunctionNodeSpawner* UBlueprintFunctionNodeSpawner::Create(TSubclassOf
 		Outer = GetTransientPackage();
 	}
 	UBlueprintFunctionNodeSpawner* NodeSpawner = NewObject<UBlueprintFunctionNodeSpawner>(Outer);
-	NodeSpawner->Function = Function;
+	NodeSpawner->Field = Function;
 
 	if (NodeClass == nullptr)
 	{
@@ -183,13 +183,23 @@ UBlueprintFunctionNodeSpawner* UBlueprintFunctionNodeSpawner::Create(TSubclassOf
 		NodeSpawner->NodeClass = NodeClass;
 	}
 
+	auto SetNodeFunctionLambda = [](UEdGraphNode* NewNode, UField const* Field)
+	{
+		// user could have changed the node class (to something like
+		// UK2Node_BaseAsyncTask, which also wraps a function)
+		if (UK2Node_CallFunction* FuncNode = Cast<UK2Node_CallFunction>(NewNode))
+		{
+			FuncNode->SetFromFunction(Cast<UFunction>(Field));
+		}
+	};
+	NodeSpawner->SetNodeFieldDelegate = FSetNodeFieldDelegate::CreateStatic(SetNodeFunctionLambda);
+
 	return NodeSpawner;
 }
 
 //------------------------------------------------------------------------------
 UBlueprintFunctionNodeSpawner::UBlueprintFunctionNodeSpawner(class FPostConstructInitializeProperties const& PCIP)
 	: Super(PCIP)
-	, Function(nullptr)
 {
 }
 
@@ -206,31 +216,18 @@ void UBlueprintFunctionNodeSpawner::Prime()
 //------------------------------------------------------------------------------
 UEdGraphNode* UBlueprintFunctionNodeSpawner::Invoke(UEdGraph* ParentGraph, FBindingSet const& Bindings, FVector2D const Location) const
 {
-	auto PostSpawnSetupLambda = [](UEdGraphNode* NewNode, bool bIsTemplateNode, UFunction const* Function, FCustomizeNodeDelegate UserDelegate)
-	{
-		// user could have changed the node class (to something like
-		// UK2Node_BaseAsyncTask, which also wraps a function)
-		if (UK2Node_CallFunction* FuncNode = Cast<UK2Node_CallFunction>(NewNode))
-		{
-			FuncNode->SetFromFunction(Function);
-		}
-
-		UserDelegate.ExecuteIfBound(NewNode, bIsTemplateNode);
-	};
-
 	// if this spawner was set up to spawn a bound node, reset this so the 
 	// bound nodes get positioned properly
 	BlueprintFunctionNodeSpawnerImpl::BindingOffset = FVector2D::ZeroVector;
 
-	FCustomizeNodeDelegate PostSpawnSetupDelegate = FCustomizeNodeDelegate::CreateStatic(PostSpawnSetupLambda, Function, CustomizeNodeDelegate);
-	UEdGraphNode* SpawnedNode = Super::Invoke(ParentGraph, Bindings, Location, PostSpawnSetupDelegate);
-
+	UEdGraphNode* SpawnedNode = Super::Invoke(ParentGraph, Bindings, Location);
 	return SpawnedNode;
 }
 
 //------------------------------------------------------------------------------
 FText UBlueprintFunctionNodeSpawner::GetDefaultMenuName() const
 {
+	UFunction const* Function = GetFunction();
 	check(Function != nullptr);
 	return FText::FromString(UK2Node_CallFunction::GetUserFacingFunctionName(Function));
 }
@@ -238,6 +235,7 @@ FText UBlueprintFunctionNodeSpawner::GetDefaultMenuName() const
 //------------------------------------------------------------------------------
 FText UBlueprintFunctionNodeSpawner::GetDefaultMenuCategory() const
 {
+	UFunction const* Function = GetFunction();
 	check(Function != nullptr);
 	return FText::FromString(UK2Node_CallFunction::GetDefaultCategoryForFunction(Function, TEXT("")));
 }
@@ -245,6 +243,7 @@ FText UBlueprintFunctionNodeSpawner::GetDefaultMenuCategory() const
 //------------------------------------------------------------------------------
 FText UBlueprintFunctionNodeSpawner::GetDefaultMenuTooltip() const
 {
+	UFunction const* Function = GetFunction();
 	check(Function != nullptr);
 
 	FText Tooltip = FText::FromString(UK2Node_CallFunction::GetDefaultTooltipForFunction(Function));
@@ -258,6 +257,7 @@ FText UBlueprintFunctionNodeSpawner::GetDefaultMenuTooltip() const
 //------------------------------------------------------------------------------
 FString UBlueprintFunctionNodeSpawner::GetDefaultSearchKeywords() const
 {
+	UFunction const* Function = GetFunction();
 	check(Function != nullptr);
 	
 	FString SearchKeywords = UK2Node_CallFunction::GetKeywordsForFunction(Function);
@@ -269,12 +269,13 @@ FString UBlueprintFunctionNodeSpawner::GetDefaultSearchKeywords() const
 //------------------------------------------------------------------------------
 FName UBlueprintFunctionNodeSpawner::GetDefaultMenuIcon(FLinearColor& ColorOut) const
 {
-	return UK2Node_CallFunction::GetPaletteIconForFunction(Function, ColorOut);
+	return UK2Node_CallFunction::GetPaletteIconForFunction(GetFunction(), ColorOut);
 }
 
 //------------------------------------------------------------------------------
 bool UBlueprintFunctionNodeSpawner::CanBindMultipleObjects() const
 {
+	UFunction const* Function = GetFunction();
 	check(Function != nullptr);
 	return UK2Node_CallFunction::CanFunctionSupportMultipleTargets(Function);
 }
@@ -283,7 +284,7 @@ bool UBlueprintFunctionNodeSpawner::CanBindMultipleObjects() const
 bool UBlueprintFunctionNodeSpawner::IsBindingCompatible(UObject const* BindingCandidate) const
 {
 	bool bCanBind = false;
-	if (Function != nullptr)
+	if (UFunction const* Function = GetFunction())
 	{
 		// @TODO: don't know exactly why we can only bind non-pure/const 
 		//        functions... this is mirrored after FK2ActionMenuBuilder::GetFunctionCallsOnSelectedActors()
@@ -318,7 +319,7 @@ bool UBlueprintFunctionNodeSpawner::BindToNode(UEdGraphNode* Node, UObject* Bind
 //------------------------------------------------------------------------------
 UFunction const* UBlueprintFunctionNodeSpawner::GetFunction() const
 {
-	return Function;
+	return Cast<UFunction>(GetField());
 }
 
 #undef LOCTEXT_NAMESPACE
diff --git a/Engine/Source/Editor/BlueprintGraph/Private/BlueprintNodeSpawner.cpp b/Engine/Source/Editor/BlueprintGraph/Private/BlueprintNodeSpawner.cpp
--- a/Engine/Source/Editor/BlueprintGraph/Private/BlueprintNodeSpawner.cpp
+++ b/Engine/Source/Editor/BlueprintGraph/Private/BlueprintNodeSpawner.cpp
@@ -99,6 +99,12 @@ void UBlueprintNodeSpawner::Prime()
 	GetDefaultMenuTooltip();
 }
 
+//------------------------------------------------------------------------------
+FBlueprintNodeSpawnerSignature UBlueprintNodeSpawner::GetSpawnerSignature() const
+{
+	return FBlueprintNodeSpawnerSignature(NodeClass);
+}
+
 //------------------------------------------------------------------------------
 UEdGraphNode* UBlueprintNodeSpawner::Invoke(UEdGraph* ParentGraph, FBindingSet const& Bindings, FVector2D const Location) const
 {
diff --git a/Engine/Source/Editor/BlueprintGraph/Private/BlueprintNodeSpawnerSignature.cpp b/Engine/Source/Editor/BlueprintGraph/Private/BlueprintNodeSpawnerSignature.cpp
new file mode 100644
--- /dev/null
+++ b/Engine/Source/Editor/BlueprintGraph/Private/BlueprintNodeSpawnerSignature.cpp
@@ -0,0 +1,170 @@
+// Copyright 1998-2014 Epic Games, Inc. All Rights Reserved.
+
+#include "BlueprintGraphPrivatePCH.h"
+#include "BlueprintNodeSpawnerSignature.h"
+
+/*******************************************************************************
+ * Static FBlueprintNodeSpawnerSignature Helpers
+ ******************************************************************************/
+
+namespace BlueprintNodeSpawnerSignatureImpl
+{
+	static FString const SignatureOpeningStr("(");
+	static FString const SignatureElementDelim(",");
+	static FString const SignatureClosingStr(")");
+	static FString const SignatureKeyDelim("=");
+}
+
+/*******************************************************************************
+ * FBlueprintNodeSpawnerSignature
+ ******************************************************************************/
+
+//------------------------------------------------------------------------------
+FBlueprintNodeSpawnerSignature::FBlueprintNodeSpawnerSignature(FString const& UserString)
+{
+	using namespace BlueprintNodeSpawnerSignatureImpl;
+
+	FString SanitizedSignature = UserString;
+	SanitizedSignature.RemoveFromStart(SignatureOpeningStr);
+	SanitizedSignature.RemoveFromEnd(SignatureClosingStr);
+
+	TArray<FString> SignatureElements;
+	SanitizedSignature.ParseIntoArray(&SignatureElements, *SignatureElementDelim, 1);
+
+	for (FString& SignatureElement : SignatureElements)
+	{
+		SignatureElement.Trim();
+
+		FString SignatureKey, SignatureValue;
+		SignatureElement.Split(SignatureKeyDelim, &SignatureKey, &SignatureValue);
+		// @TODO: look for UObject redirects with SignatureValue
+
+		SignatureValue.RemoveFromStart(TEXT("\""));
+		SignatureValue.RemoveFromEnd(TEXT("\""));
+		AddKeyValue(FName(*SignatureKey), SignatureValue);
+	}
+}
+
+//------------------------------------------------------------------------------
+FBlueprintNodeSpawnerSignature::FBlueprintNodeSpawnerSignature(TSubclassOf<UEdGraphNode> NodeClass)
+{
+	SetNodeClass(NodeClass);
+}
+
+//------------------------------------------------------------------------------
+void FBlueprintNodeSpawnerSignature::SetNodeClass(TSubclassOf<UEdGraphNode> NodeClass)
+{
+	static const FName NodeClassSignatureKey(TEXT("NodeName"));
+
+	if (NodeClass != nullptr)
+	{
+		AddKeyValue(NodeClassSignatureKey, NodeClass->GetPathName());
+		MarkDirty();
+	}
+	else
+	{
+		SignatureSet.Remove(NodeClassSignatureKey);
+	}
+}
+
+//------------------------------------------------------------------------------
+void FBlueprintNodeSpawnerSignature::AddSubObject(UObject const* SignatureObj)
+{
+	// not ideal for generic "objects", but we have to keep in sync with the 
+	// old favorites system (for backwards compatibility)
+	static FName const BaseSubObjectSignatureKey(TEXT("FieldName"));
+	FName SubObjectSignatureKey = BaseSubObjectSignatureKey;
+
+	int32 FNameIndex = 0;
+	while (SignatureSet.Find(SubObjectSignatureKey) != nullptr)
+	{
+		SubObjectSignatureKey = FName(BaseSubObjectSignatureKey, ++FNameIndex);
+	}
+
+	AddKeyValue(SubObjectSignatureKey, SignatureObj->GetPathName());
+	MarkDirty();
+}
+
+//------------------------------------------------------------------------------
+void FBlueprintNodeSpawnerSignature::AddKeyValue(FName SignatureKey, FString const& Value)
+{
+	SignatureSet.Add(SignatureKey, Value);
+	MarkDirty();
+}
+
+//------------------------------------------------------------------------------
+bool FBlueprintNodeSpawnerSignature::IsValid() const
+{
+	return (SignatureSet.Num() > 0);
+}
+
+//------------------------------------------------------------------------------
+FString const& FBlueprintNodeSpawnerSignature::ToString() const
+{
+	using namespace BlueprintNodeSpawnerSignatureImpl;
+
+	if (CachedSignatureString.IsEmpty() && IsValid())
+	{
+		TArray<FString> SignatureElements;
+		for (auto& SignatureElement : SignatureSet)
+		{
+			FString ObjSignature = SignatureElement.Key.ToString() + SignatureKeyDelim + TEXT("\"") + SignatureElement.Value + TEXT("\"");
+			SignatureElements.Add(ObjSignature);
+		}
+		SignatureElements.Sort();
+
+		CachedSignatureString = SignatureOpeningStr;
+		for (FString& SignatureElement : SignatureElements)
+		{
+			CachedSignatureString += SignatureElement + SignatureElementDelim;
+		}
+		CachedSignatureString.RemoveFromEnd(SignatureElementDelim);
+		CachedSignatureString += SignatureClosingStr;
+	}
+	
+	return CachedSignatureString;
+}
+
+//------------------------------------------------------------------------------
+FGuid const& FBlueprintNodeSpawnerSignature::AsGuid() const
+{
+	static const int32 BytesPerMd5Hash = 16;
+	static const int32 BytesPerGuidVal =  4;
+	static const int32 MembersPerGuid  =  4;
+	static const int32 BitsPerByte     =  8;
+
+	if (!CachedSignatureGuid.IsValid() && IsValid())
+	{
+		FString const& SignatureString = ToString();
+
+		FMD5  Md5Gen;
+		uint8 HashedBytes[BytesPerMd5Hash];
+		Md5Gen.Update((unsigned char*)TCHAR_TO_ANSI(*SignatureString), SignatureString.Len());
+		Md5Gen.Final(HashedBytes);
+		//FString HashedStr = FMD5::HashAnsiString(*SignatureString); // to check against
+
+		for (int32 GuidIndex = 0; GuidIndex < MembersPerGuid; ++GuidIndex)
+		{
+			int32 const Md5UpperIndex = (GuidIndex + 1) * BytesPerGuidVal - 1;
+			int32 const Md5LowerIndex = Md5UpperIndex - BytesPerGuidVal;
+
+			int32 GuidValue = 0;
+
+			int32 BitOffset = 1;
+			for (int32 Md5ByteIndex = Md5UpperIndex; Md5ByteIndex >= Md5LowerIndex; --Md5ByteIndex)
+			{
+				GuidValue += HashedBytes[Md5ByteIndex] * BitOffset;
+				BitOffset <<= BitsPerByte;
+			}
+			CachedSignatureGuid[GuidIndex] = GuidValue;
+		}
+	}
+	return CachedSignatureGuid;
+}
+
+//------------------------------------------------------------------------------
+void FBlueprintNodeSpawnerSignature::MarkDirty()
+{
+	CachedSignatureGuid.Invalidate();
+	CachedSignatureString.Empty();
+}
diff --git a/Engine/Source/Editor/BlueprintGraph/Private/BlueprintNodeSpawnerUtils.h b/Engine/Source/Editor/BlueprintGraph/Private/BlueprintNodeSpawnerUtils.h
--- a/Engine/Source/Editor/BlueprintGraph/Private/BlueprintNodeSpawnerUtils.h
+++ b/Engine/Source/Editor/BlueprintGraph/Private/BlueprintNodeSpawnerUtils.h
@@ -5,7 +5,6 @@
 // Forward declarations
 class UBlueprintNodeSpawner;
 
-
 struct FBlueprintNodeSpawnerUtils
 {
 	/**
@@ -18,7 +17,7 @@ struct FBlueprintNodeSpawnerUtils
 	static UField const* GetAssociatedField(UBlueprintNodeSpawner const* BlueprintAction);
 
 	/**
-	 * Certain node-spawners are associated with specific UFunction (call-
+	 * Certain node-spawners are associated with specific UFunctions (call-
 	 * function, and event spawners). This attempts to retrieve the function 
 	 * from the spawner.
 	 * 
@@ -28,10 +27,12 @@ struct FBlueprintNodeSpawnerUtils
 	static UFunction const* GetAssociatedFunction(UBlueprintNodeSpawner const* BlueprintAction);
 
 	/**
+	 * Certain node-spawners are associated with specific UProperties (get/set 
+	 * nodes, delegates, etc.). This attempts to retrieve a property from the
+	 * the spawner.
 	 * 
-	 * 
-	 * @param  BlueprintAction	
-	 * @return 
+	 * @param  BlueprintAction	The node-spawner you want an associated property for.
+	 * @return The action's associated property (null if it doesn't have one).
 	 */
 	static UProperty const* GetAssociatedProperty(UBlueprintNodeSpawner const* BlueprintAction);
 };
diff --git a/Engine/Source/Editor/BlueprintGraph/Private/BlueprintVariableNodeSpawner.cpp b/Engine/Source/Editor/BlueprintGraph/Private/BlueprintVariableNodeSpawner.cpp
--- a/Engine/Source/Editor/BlueprintGraph/Private/BlueprintVariableNodeSpawner.cpp
+++ b/Engine/Source/Editor/BlueprintGraph/Private/BlueprintVariableNodeSpawner.cpp
@@ -25,8 +25,23 @@ UBlueprintVariableNodeSpawner* UBlueprintVariableNodeSpawner::Create(TSubclassOf
 	}
 
 	UBlueprintVariableNodeSpawner* NodeSpawner = NewObject<UBlueprintVariableNodeSpawner>(Outer);
-	NodeSpawner->NodeClass      = NodeClass;
-	NodeSpawner->MemberVariable = VarProperty;
+	NodeSpawner->NodeClass = NodeClass;
+	NodeSpawner->Field     = VarProperty;
+
+	auto MemberVarSetupLambda = [](UEdGraphNode* NewNode, UField const* Field)
+	{
+		if (UProperty const* Property = Cast<UProperty>(Field))
+		{
+			UBlueprint* Blueprint = FBlueprintEditorUtils::FindBlueprintForNodeChecked(NewNode);
+			UClass* OwnerClass = Property->GetOwnerClass();
+			bool const bIsSelfContext = Blueprint->SkeletonGeneratedClass->IsChildOf(OwnerClass);
+
+			UK2Node_Variable* VarNode = CastChecked<UK2Node_Variable>(NewNode);
+			VarNode->SetFromProperty(Property, bIsSelfContext);
+		}
+	};
+	NodeSpawner->SetNodeFieldDelegate = FSetNodeFieldDelegate::CreateStatic(MemberVarSetupLambda);
+
 	return NodeSpawner;
 }
 
@@ -39,10 +54,13 @@ UBlueprintVariableNodeSpawner* UBlueprintVariableNodeSpawner::Create(TSubclassOf
 		Outer = GetTransientPackage();
 	}
 
+	// @TODO: consider splitting out local variable spawners (since they don't 
+	//        conform to UBlueprintFieldNodeSpawner
 	UBlueprintVariableNodeSpawner* NodeSpawner = NewObject<UBlueprintVariableNodeSpawner>(Outer);
 	NodeSpawner->NodeClass     = NodeClass;
 	NodeSpawner->LocalVarOuter = VarContext;
 	NodeSpawner->LocalVarDesc  = VarDesc;
+
 	return NodeSpawner;
 }
 
@@ -66,10 +84,25 @@ void UBlueprintVariableNodeSpawner::Prime()
 }
 
 //------------------------------------------------------------------------------
-UEdGraphNode* UBlueprintVariableNodeSpawner::Invoke(UEdGraph* ParentGraph, FBindingSet const& Bindings, FVector2D const Location) const
+FBlueprintNodeSpawnerSignature UBlueprintVariableNodeSpawner::GetSpawnerSignature() const
 {
-	FCustomizeNodeDelegate PostSpawnDelegate = CustomizeNodeDelegate;
+	FBlueprintNodeSpawnerSignature SpawnerSignature = Super::GetSpawnerSignature();
+	if (IsLocalVariable())
+	{
+		SpawnerSignature.AddSubObject(LocalVarOuter);
+		static const FName LocalVarSignatureKey(TEXT("LocalVarName"));
+		SpawnerSignature.AddKeyValue(LocalVarSignatureKey, LocalVarDesc.VarName.ToString());
+	}
+	
+	return SpawnerSignature;
+}
 
+//------------------------------------------------------------------------------
+UEdGraphNode* UBlueprintVariableNodeSpawner::Invoke(UEdGraph* ParentGraph, FBindingSet const& Bindings, FVector2D const Location) const
+{
+	UEdGraphNode* NewNode = nullptr;
+	// @TODO: consider splitting out local variable spawners (since they don't 
+	//        conform to UBlueprintFieldNodeSpawner
 	if (IsLocalVariable())
 	{
 		auto LocalVarSetupLambda = [](UEdGraphNode* NewNode, bool bIsTemplateNode, FName VarName, UObject const* VarOuter, FGuid VarGuid, FCustomizeNodeDelegate UserDelegate)
@@ -79,29 +112,20 @@ UEdGraphNode* UBlueprintVariableNodeSpawner::Invoke(UEdGraph* ParentGraph, FBind
 			UserDelegate.ExecuteIfBound(NewNode, bIsTemplateNode);
 		};
 
+		FCustomizeNodeDelegate PostSpawnDelegate = CustomizeNodeDelegate;
 		if (UObject const* LocalVarOuter = GetVarOuter())
 		{
 			PostSpawnDelegate = FCustomizeNodeDelegate::CreateStatic(LocalVarSetupLambda, LocalVarDesc.VarName, LocalVarOuter, LocalVarDesc.VarGuid, CustomizeNodeDelegate);
 		}
+
+		NewNode = UBlueprintNodeSpawner::Invoke(ParentGraph, Bindings, Location, PostSpawnDelegate);
 	}
-	else if (MemberVariable.IsValid())
+	else
 	{
-		auto MemberVarSetupLambda = [](UEdGraphNode* NewNode, bool bIsTemplateNode, UProperty const* Property, FCustomizeNodeDelegate UserDelegate)
-		{
-			UBlueprint* Blueprint = FBlueprintEditorUtils::FindBlueprintForNodeChecked(NewNode);
-			UClass* OwnerClass = Property->GetOwnerClass();
-			bool const bIsSelfContext = Blueprint->SkeletonGeneratedClass->IsChildOf(OwnerClass);
-
-			UK2Node_Variable* VarNode = CastChecked<UK2Node_Variable>(NewNode);
-			VarNode->SetFromProperty(Property, bIsSelfContext);
-			UserDelegate.ExecuteIfBound(NewNode, bIsTemplateNode);
-		};
-
-		UProperty const* Property = MemberVariable.Get();
-		PostSpawnDelegate = FCustomizeNodeDelegate::CreateStatic(MemberVarSetupLambda, Property, CustomizeNodeDelegate);
+		NewNode = Super::Invoke(ParentGraph, Bindings, Location);
 	}
 
-	return Super::Invoke(ParentGraph, Bindings, Location, PostSpawnDelegate);
+	return NewNode;
 }
 
 //------------------------------------------------------------------------------
@@ -135,9 +159,9 @@ FText UBlueprintVariableNodeSpawner::GetDefaultMenuCategory() const
 		{
 			VarSubCategory = FText::FromName(LocalVarDesc.Category);
 		}
-		else if (MemberVariable.IsValid())
+		else if (UProperty const* MemberVariable = GetVarProperty())
 		{
-			VarSubCategory = FText::FromString(FObjectEditorUtils::GetCategory(MemberVariable.Get()));
+			VarSubCategory = FText::FromString(FObjectEditorUtils::GetCategory(MemberVariable));
 		}
 		CachedCategory = FEditorCategoryUtils::BuildCategoryString(FCommonEditorCategory::Variables, VarSubCategory);
 	}
@@ -155,9 +179,9 @@ FText UBlueprintVariableNodeSpawner::GetDefaultMenuTooltip() const
 			{
 				CachedTooltip = UK2Node_VariableSet::GetBlueprintVarTooltip(LocalVarDesc);
 			}
-			else
+			else if (UProperty const* MemberVariable = GetVarProperty())
 			{
-				CachedTooltip = UK2Node_VariableSet::GetPropertyTooltip(MemberVariable.Get());
+				CachedTooltip = UK2Node_VariableSet::GetPropertyTooltip(MemberVariable);
 			}
 		}
 		else if (NodeClass->IsChildOf<UK2Node_VariableGet>())
@@ -166,9 +190,9 @@ FText UBlueprintVariableNodeSpawner::GetDefaultMenuTooltip() const
 			{
 				CachedTooltip = UK2Node_VariableGet::GetBlueprintVarTooltip(LocalVarDesc);
 			}
-			else
+			else if (UProperty const* MemberVariable = GetVarProperty())
 			{
-				CachedTooltip = UK2Node_VariableGet::GetPropertyTooltip(MemberVariable.Get());
+				CachedTooltip = UK2Node_VariableGet::GetPropertyTooltip(MemberVariable);
 			}
 		}
 	}	
@@ -203,12 +227,9 @@ UObject const* UBlueprintVariableNodeSpawner::GetVarOuter() const
 	UObject const* VarOuter = nullptr;
 	if (IsLocalVariable())
 	{
-		if (LocalVarOuter.IsValid())
-		{
-			VarOuter = LocalVarOuter.Get();
-		}
+		VarOuter = LocalVarOuter;
 	}
-	else if (MemberVariable.IsValid())
+	else if (UProperty const* MemberVariable = GetVarProperty())
 	{
 		VarOuter = MemberVariable->GetOuter();
 	}	
@@ -219,7 +240,7 @@ UObject const* UBlueprintVariableNodeSpawner::GetVarOuter() const
 UProperty const* UBlueprintVariableNodeSpawner::GetVarProperty() const
 {
 	// Get() does IsValid() checks for us
-	return MemberVariable.Get();
+	return Cast<UProperty>(GetField());
 }
 
 //------------------------------------------------------------------------------
@@ -248,9 +269,9 @@ FText UBlueprintVariableNodeSpawner::GetVariableName() const
 	{
 		VarName = bShowFriendlyNames ? FText::FromString(LocalVarDesc.FriendlyName) : FText::FromName(LocalVarDesc.VarName);
 	}
-	else if (MemberVariable.IsValid())
+	else if (UProperty const* MemberVariable = GetVarProperty())
 	{
-		VarName = bShowFriendlyNames ? FText::FromString(UEditorEngine::GetFriendlyName(MemberVariable.Get())) : FText::FromName(MemberVariable->GetFName());
+		VarName = bShowFriendlyNames ? FText::FromString(UEditorEngine::GetFriendlyName(MemberVariable)) : FText::FromName(MemberVariable->GetFName());
 	}
 	return VarName;
 }
diff --git a/Engine/Source/Editor/BlueprintGraph/Private/K2Node_VariableGet.cpp b/Engine/Source/Editor/BlueprintGraph/Private/K2Node_VariableGet.cpp
--- a/Engine/Source/Editor/BlueprintGraph/Private/K2Node_VariableGet.cpp
+++ b/Engine/Source/Editor/BlueprintGraph/Private/K2Node_VariableGet.cpp
@@ -105,7 +105,7 @@ void UK2Node_VariableGet::ReallocatePinsDuringReconstruction(TArray<UEdGraphPin*
 	}
 }
 
-FText UK2Node_VariableGet::GetPropertyTooltip(UProperty* VariableProperty)
+FText UK2Node_VariableGet::GetPropertyTooltip(UProperty const* VariableProperty)
 {
 	FName VarName = NAME_None;
 	if (VariableProperty != nullptr)
diff --git a/Engine/Source/Editor/BlueprintGraph/Private/K2Node_VariableSet.cpp b/Engine/Source/Editor/BlueprintGraph/Private/K2Node_VariableSet.cpp
--- a/Engine/Source/Editor/BlueprintGraph/Private/K2Node_VariableSet.cpp
+++ b/Engine/Source/Editor/BlueprintGraph/Private/K2Node_VariableSet.cpp
@@ -28,7 +28,7 @@ namespace K2Node_VariableSetImpl
 	 * @param  VariableProperty	The variable property you wish to check.
 	 * @return True if the specified variable RepNotify AND is defined in a blueprint.
 	 */
-	static bool PropertyHasLocalRepNotify(UProperty* VariableProperty);
+	static bool PropertyHasLocalRepNotify(UProperty const* VariableProperty);
 }
 
 static FText K2Node_VariableSetImpl::GetBaseTooltip(FName VarName)
@@ -40,7 +40,7 @@ static FText K2Node_VariableSetImpl::GetBaseTooltip(FName VarName)
 
 }
 
-static bool K2Node_VariableSetImpl::PropertyHasLocalRepNotify(UProperty* VariableProperty)
+static bool K2Node_VariableSetImpl::PropertyHasLocalRepNotify(UProperty const* VariableProperty)
 {
 	if (VariableProperty != nullptr)
 	{
@@ -113,7 +113,7 @@ void UK2Node_VariableSet::ReallocatePinsDuringReconstruction(TArray<UEdGraphPin*
 
 
 
-FText UK2Node_VariableSet::GetPropertyTooltip(UProperty* VariableProperty)
+FText UK2Node_VariableSet::GetPropertyTooltip(UProperty const* VariableProperty)
 {
 	FText TextFormat;
 	FFormatNamedArguments Args;
diff --git a/Engine/Source/Editor/BlueprintGraph/Public/BlueprintActionFilter.h b/Engine/Source/Editor/BlueprintGraph/Public/BlueprintActionFilter.h
--- a/Engine/Source/Editor/BlueprintGraph/Public/BlueprintActionFilter.h
+++ b/Engine/Source/Editor/BlueprintGraph/Public/BlueprintActionFilter.h
@@ -69,9 +69,9 @@ struct BLUEPRINTGRAPH_API FBlueprintActionInfo
 	UObject const* GetActionOwner();
 
 	/**
+	 * Retrieves any bindings that the action will apply to spawned nodes.
 	 * 
-	 * 
-	 * @return 
+	 * @return The set of bindings tied to this action.
 	 */
 	IBlueprintNodeBinder::FBindingSet const& GetBindings() const;
 
@@ -93,6 +93,13 @@ struct BLUEPRINTGRAPH_API FBlueprintActionInfo
 	 */
 	UClass const* GetNodeClass();
 
+	/**
+	 * 
+	 * 
+	 * @return 
+	 */
+	FGuid GetActionGuid();
+
 	/**
 	 * Certain actions are associated with specific member fields (a member 
 	 * function call, a variable get/set, etc.) This retrieves that member field 
@@ -133,13 +140,13 @@ private:
 	
 	/** */
 	UClass const*    CachedOwnerClass;
-	/** */
 	UField const*    CachedActionField;
-	/** */
 	UProperty const* CachedActionProperty;
-	/** */
 	UFunction const* CachedActionFunction;
 
+	/** */
+	FGuid CachedActionId;
+
 	/** */
 	IBlueprintNodeBinder::FBindingSet Bindings;
 };
diff --git a/Engine/Source/Editor/BlueprintGraph/Public/BlueprintEditorSettings.h b/Engine/Source/Editor/BlueprintGraph/Public/BlueprintEditorSettings.h
--- a/Engine/Source/Editor/BlueprintGraph/Public/BlueprintEditorSettings.h
+++ b/Engine/Source/Editor/BlueprintGraph/Public/BlueprintEditorSettings.h
@@ -15,19 +15,24 @@ class BLUEPRINTGRAPH_API UBlueprintEditorSettings
 // Style Settings
 public:
 	/** Should arrows indicating data/execution flow be drawn halfway along wires? */
-	UPROPERTY(EditAnywhere, config, Category="Visual Style", meta=(DisplayName="Draw midpoint arrows in Blueprints"))
+	UPROPERTY(EditAnywhere, config, Category=VisualStyle, meta=(DisplayName="Draw midpoint arrows in Blueprints"))
 	bool bDrawMidpointArrowsInBlueprints;
 
 // UX Settings
 public:
 	/** Determines if lightweight tutorial text shows up at the top of empty blueprint graphs */
-	UPROPERTY(EditAnywhere, config, Category="User Experience")
+	UPROPERTY(EditAnywhere, config, Category=UserExperience)
 	bool bShowGraphInstructionText;
 
+// Developer Settings
+public:
+	/** If enabled, tooltips on action menu items will show the associated action's signature id (can be used to setup custom favorites menus).*/
+	UPROPERTY(EditAnywhere, config, Category=DeveloperTools)
+	bool bShowActionMenuItemSignatures;
+
 // Perf Settings
 public:
 	/** The node template cache is used to speed up blueprint menuing. This determines the peak data size for that cache. */
-	UPROPERTY(EditAnywhere, AdvancedDisplay, config, Category="Performance", DisplayName="Node-Template Cache Cap (MB)", meta=(ClampMin="0", UIMin="0"))
+	UPROPERTY(EditAnywhere, AdvancedDisplay, config, Category=Performance, DisplayName="Node-Template Cache Cap (MB)", meta=(ClampMin="0", UIMin="0"))
 	float NodeTemplateCacheCapMB;
 };
- 
\ No newline at end of file
diff --git a/Engine/Source/Editor/BlueprintGraph/Public/BlueprintNodeSpawner.h b/Engine/Source/Editor/BlueprintGraph/Public/BlueprintNodeSpawner.h
--- a/Engine/Source/Editor/BlueprintGraph/Public/BlueprintNodeSpawner.h
+++ b/Engine/Source/Editor/BlueprintGraph/Public/BlueprintNodeSpawner.h
@@ -3,6 +3,7 @@
 #pragma once
 
 #include "BlueprintNodeBinder.h"
+#include "BlueprintNodeSpawnerSignature.h"
 #include "BlueprintNodeSpawner.generated.h"
 
 // Forward declarations
@@ -79,6 +80,16 @@ public:
 	 */
 	virtual void Prime();
 
+	/**
+	 * We want to be able to compare spawners, and have a signature that is 
+	 * rebuildable on subsequent runs. So, what makes each spawner unique is the 
+	 * type of node that it spawns, and any fields the node would be initialized 
+	 * with; that is what this returns.
+	 * 
+	 * @return A set of object-paths/names that distinguish this spawner from others.
+	 */
+	virtual FBlueprintNodeSpawnerSignature GetSpawnerSignature() const;
+
 	/**
 	 * Takes care of spawning a node for the specified graph. Looks to see if 
 	 * the supplied graph is transient, and if so, spawns a NOT fully formed
@@ -94,10 +105,7 @@ public:
 	 * @param  Location     Where you want the new node positioned in the graph.
 	 * @return Null if it failed to spawn a node, otherwise a newly spawned node or possibly one that already existed.
 	 */
-	virtual UEdGraphNode* Invoke(UEdGraph* ParentGraph, FBindingSet const& Bindings, FVector2D const Location) const;
-	
-	// @TODO: for favorites, generated for the spawner type + data (function, property, node, etc.)
-	//virtual FGuid GetActionHash() const;
+	virtual UEdGraphNode* Invoke(UEdGraph* ParentGraph, FBindingSet const& Bindings, FVector2D const Location) const;	
 	
 	/**
 	 * To save on performance, certain sub-classes can concoct menu names from
diff --git a/Engine/Source/Editor/BlueprintGraph/Public/BlueprintNodeSpawnerSignature.h b/Engine/Source/Editor/BlueprintGraph/Public/BlueprintNodeSpawnerSignature.h
new file mode 100644
--- /dev/null
+++ b/Engine/Source/Editor/BlueprintGraph/Public/BlueprintNodeSpawnerSignature.h
@@ -0,0 +1,73 @@
+// Copyright 1998-2014 Epic Games, Inc. All Rights Reserved.
+
+#pragma once
+
+struct BLUEPRINTGRAPH_API FBlueprintNodeSpawnerSignature
+{
+public:
+	FBlueprintNodeSpawnerSignature() {}
+	FBlueprintNodeSpawnerSignature(FString const& UserString);
+	FBlueprintNodeSpawnerSignature(TSubclassOf<UEdGraphNode> NodeClass);
+
+	/**
+	 * 
+	 * 
+	 * @param  NodeClass	
+	 * @return 
+	 */
+	void SetNodeClass(TSubclassOf<UEdGraphNode> NodeClass);
+
+	/**
+	 * 
+	 * 
+	 * @param  SignatureObj	
+	 * @return 
+	 */
+	void AddSubObject(UObject const* SignatureObj);
+
+	/**
+	 * 
+	 * 
+	 * @param  SignatureKey	
+	 * @param  Value	
+	 * @return 
+	 */
+	void AddKeyValue(FName SignatureKey, FString const& Value);
+
+	/**
+	 * 
+	 * 
+	 * @return 
+	 */
+	bool IsValid() const;
+
+	/**
+	 * 
+	 * 
+	 * @return 
+	 */
+	FString const& ToString() const;
+
+	/**
+	*
+	*
+	* @return
+	*/
+	FGuid const& AsGuid() const;
+
+private:
+	/**
+	 * 
+	 * 
+	 * @return 
+	 */
+	void MarkDirty();
+
+private:
+	typedef TMap<FName, FString> FSignatureSet;
+	FSignatureSet SignatureSet;
+
+	mutable FGuid   CachedSignatureGuid;
+	mutable FString CachedSignatureString;
+};
+
diff --git a/Engine/Source/Editor/Kismet/Classes/BlueprintPaletteFavorites.h b/Engine/Source/Editor/Kismet/Classes/BlueprintPaletteFavorites.h
--- a/Engine/Source/Editor/Kismet/Classes/BlueprintPaletteFavorites.h
+++ b/Engine/Source/Editor/Kismet/Classes/BlueprintPaletteFavorites.h
@@ -2,8 +2,13 @@
 
 #pragma once
 
+#include "BlueprintNodeSpawnerSignature.h"
 #include "BlueprintPaletteFavorites.generated.h"
 
+// Forward Declarations
+class  UBlueprintNodeSpawner;
+struct FBlueprintActionInfo;
+
 /*******************************************************************************
 * FFavoritedPaletteItem
 *******************************************************************************/
@@ -29,7 +34,15 @@ public:
 	 * 
 	 * @param  PaletteAction	The action you wish to favorite.
 	 */
-	FFavoritedBlueprintPaletteItem(TSharedPtr<FEdGraphSchemaAction> PaletteAction = NULL);
+	FFavoritedBlueprintPaletteItem(TSharedPtr<FEdGraphSchemaAction> PaletteAction = nullptr);
+
+	/**
+	 * 
+	 * 
+	 * @param  BlueprintAction	
+	 * @return 
+	 */
+	FFavoritedBlueprintPaletteItem(UBlueprintNodeSpawner const* BlueprintAction);
 
 	/**
 	 * Sometimes we're not able to construct favorites from specified actions, 
@@ -57,44 +70,31 @@ public:
 	bool operator==(TSharedPtr<FEdGraphSchemaAction> PaletteAction) const;
 
 	/**
-	 * We want to be able to specify some of these in .ini files, so let we have
-	 * to have a readable string representation for them.
 	 * 
-	 * @return A string representation of this item.
+	 * 
+	 * @param  BlueprintAction	
+	 * @return 
 	 */
-	FString ToString() const;
-
-private:
+	bool operator==(FBlueprintActionInfo& BlueprintAction) const;
 
 	/**
-	 * Blueprint palette actions are used to place new nodes. So to uniquely 
-	 * identify the action, we specify the node that they generate (this is the 
-	 * node's type name).
+	 * 
+	 * 
+	 * @param  ActionGuid	
+	 * @return 
 	 */
-	UPROPERTY()
-	FString  NodeClassName;
+	bool operator==(FGuid const& ActionGuid) const;
 
 	/**
-	 * Multiple node types could share the same name, so to ensure we get the 
-	 * right one we track its outer as well.
-	 */
-	UPROPERTY()
-	UObject* NodeClassOuter;
-	
-	/** 
-	 * Sometimes a node has multiple permutations (like UK2Node_CallFunction), so
-	 * you need a more granular way of identifying it (like the function itself). 
-	 * This could be empty, or filled in if we need that additional granularity.
+	 * We want to be able to specify some of these in .ini files, so let we have
+	 * to have a readable string representation for them.
+	 * 
+	 * @return A string representation of this item.
 	 */
-	UPROPERTY()
-	FString  FieldName;
+	FString const& ToString() const;
 
-	/** 
-	 * Like FieldName, this is used to more granularly identify UK2Nodes (could
-	 * be null if it's not needed).
-	 */
-	UPROPERTY()
-	UObject* FieldOuter;
+private:
+	FBlueprintNodeSpawnerSignature ActionSignature;
 };
 
 /*******************************************************************************
@@ -130,6 +130,22 @@ public:
 	 */
 	bool IsFavorited(TSharedPtr<FEdGraphSchemaAction> PaletteAction) const;
 
+	/**
+	 * 
+	 * 
+	 * @param  BlueprintAction	
+	 * @return 
+	 */
+	bool IsFavorited(FBlueprintActionInfo& BlueprintAction) const;
+
+	/**
+	 * 
+	 * 
+	 * @param  BlueprintAction	
+	 * @return 
+	 */
+	bool IsFavorited(UBlueprintNodeSpawner const* BlueprintAction) const;
+
 	/**
 	 * Adds the specified action to the current favorites list (fails if the action
 	 * action can't be favorited, or if the favorite already exists). Will also 
@@ -234,7 +250,7 @@ public:
 	/** 
 	 * A list of strings that are used to identify specific palette actions. 
 	 * This is what gets saved out when the user has customized their own set, 
-	 * and is not updated untile PreSave().
+	 * and is not updated until PreSave().
 	 */
 	UPROPERTY(config)
 	TArray<FString> CustomFavorites;
diff --git a/Engine/Source/Editor/Kismet/Private/BlueprintActionMenuBuilder.cpp b/Engine/Source/Editor/Kismet/Private/BlueprintActionMenuBuilder.cpp
--- a/Engine/Source/Editor/Kismet/Private/BlueprintActionMenuBuilder.cpp
+++ b/Engine/Source/Editor/Kismet/Private/BlueprintActionMenuBuilder.cpp
@@ -59,7 +59,7 @@ public:
 	 * @param  SampleAction	One of the (possibly) many node-spawners that this menu item is set to represent.
 	 * @return A newly allocated FBlueprintActionMenuItem (which wraps the supplied action).
 	 */
-	TSharedPtr<FEdGraphSchemaAction> MakeDragDropMenuItem(UBlueprintNodeSpawner const* SampleAction);
+	TSharedPtr<FBlueprintDragDropMenuItem> MakeDragDropMenuItem(UBlueprintNodeSpawner const* SampleAction);
 
 	/**
 	 * 
@@ -168,7 +168,7 @@ TSharedPtr<FEdGraphSchemaAction> FBlueprintActionMenuItemFactory::MakeActionMenu
 }
 
 //------------------------------------------------------------------------------
-TSharedPtr<FEdGraphSchemaAction> FBlueprintActionMenuItemFactory::MakeDragDropMenuItem(UBlueprintNodeSpawner const* SampleAction)
+TSharedPtr<FBlueprintDragDropMenuItem> FBlueprintActionMenuItemFactory::MakeDragDropMenuItem(UBlueprintNodeSpawner const* SampleAction)
 {
 	// FBlueprintDragDropMenuItem takes care of its own menu MenuDescription, etc.
 	FBlueprintDragDropMenuItem* NewMenuItem = new FBlueprintDragDropMenuItem(Context, SampleAction, MenuGrouping);
@@ -405,7 +405,7 @@ namespace FBlueprintActionMenuBuilderImpl
 		/** In charge of spawning menu items for this section (holds category/ordering information)*/
 		FBlueprintActionMenuItemFactory ItemFactory;
 		/** Tracks the properties that we've already consolidated and passed (when using the ConsolidatePropertyActions flag)*/
-		TSet<UProperty const*> ConsolidatedProperties;
+		TMap<UProperty const*, TSharedPtr<FBlueprintDragDropMenuItem>> ConsolidatedProperties;
 	};
 	
 	/**
@@ -510,9 +510,9 @@ TSharedPtr<FEdGraphSchemaAction> FBlueprintActionMenuBuilderImpl::FMenuSectionDe
 //------------------------------------------------------------------------------
 TSharedPtr<FEdGraphSchemaAction> FBlueprintActionMenuBuilderImpl::FMenuSectionDefinition::MakeMenuItem(TWeakPtr<FBlueprintEditor> EditorContext, FBlueprintActionInfo& DatabaseAction)
 {	
+	TSharedPtr<FEdGraphSchemaAction> MenuEntry;
 	bool bPassedFilter = !Filter.IsFiltered(DatabaseAction);
 
-	bool bCreateDragDropItem = false;
 	// if the caller wants to consolidate all property actions, then we have to 
 	// check and see if this is one of those that needs consolidating (needs 
 	// a FBlueprintDragDropMenuItem instead of a FBlueprintActionMenuItem)
@@ -521,37 +521,37 @@ TSharedPtr<FEdGraphSchemaAction> FBlueprintActionMenuBuilderImpl::FMenuSectionDe
 		UProperty const* ActionProperty = nullptr;
 		if (UBlueprintVariableNodeSpawner const* VariableSpawner = Cast<UBlueprintVariableNodeSpawner>(DatabaseAction.NodeSpawner))
 		{
-			bCreateDragDropItem = true;
 			ActionProperty = VariableSpawner->GetVarProperty();
 			bPassedFilter = (ActionProperty != nullptr);
 		}
 		else if (UBlueprintDelegateNodeSpawner const* DelegateSpawner = Cast<UBlueprintDelegateNodeSpawner>(DatabaseAction.NodeSpawner))
 		{
-			bCreateDragDropItem = true;
 			ActionProperty = DelegateSpawner->GetProperty();
 			bPassedFilter = (ActionProperty != nullptr);
 		}
 
 		if (ActionProperty != nullptr)
 		{
-			// @TODO: could be added as a filter test instead
-			bool const bAlreadyConsolidated = ConsolidatedProperties.Contains(ActionProperty);
-			bPassedFilter = !bAlreadyConsolidated;
-			ConsolidatedProperties.Add(ActionProperty);
+			if (TSharedPtr<FBlueprintDragDropMenuItem>* ConsolidatedMenuItem = ConsolidatedProperties.Find(ActionProperty))
+			{
+				(*ConsolidatedMenuItem)->AppendAction(DatabaseAction.NodeSpawner);
+				// this menu entry has already been returned, don't need to 
+				// create/insert a new one
+				bPassedFilter = false;
+			}
+			else
+			{
+				TSharedPtr<FBlueprintDragDropMenuItem> NewMenuItem = ItemFactory.MakeDragDropMenuItem(DatabaseAction.NodeSpawner);
+				ConsolidatedProperties.Add(ActionProperty, NewMenuItem);
+				MenuEntry = NewMenuItem;
+			}
 		}
 	}
 
-	TSharedPtr<FEdGraphSchemaAction> MenuEntry;
-	if (bPassedFilter)
+	
+	if (!MenuEntry.IsValid() && bPassedFilter)
 	{
-		if (bCreateDragDropItem)
-		{
-			MenuEntry = ItemFactory.MakeDragDropMenuItem(DatabaseAction.NodeSpawner);
-		}
-		else
-		{
-			MenuEntry = ItemFactory.MakeActionMenuItem(EditorContext, DatabaseAction.NodeSpawner);
-		}
+		MenuEntry = ItemFactory.MakeActionMenuItem(EditorContext, DatabaseAction.NodeSpawner);
 	}
 
 	return MenuEntry;
diff --git a/Engine/Source/Editor/Kismet/Private/BlueprintActionMenuItem.cpp b/Engine/Source/Editor/Kismet/Private/BlueprintActionMenuItem.cpp
--- a/Engine/Source/Editor/Kismet/Private/BlueprintActionMenuItem.cpp
+++ b/Engine/Source/Editor/Kismet/Private/BlueprintActionMenuItem.cpp
@@ -174,4 +174,10 @@ FSlateBrush const* FBlueprintActionMenuItem::GetMenuIcon(FSlateColor& ColorOut)
 	return IconBrush;
 }
 
+//------------------------------------------------------------------------------
+UBlueprintNodeSpawner const* FBlueprintActionMenuItem::GetRawAction() const
+{
+	return Action;
+}
+
 #undef LOCTEXT_NAMESPACE
diff --git a/Engine/Source/Editor/Kismet/Private/BlueprintActionMenuUtils.cpp b/Engine/Source/Editor/Kismet/Private/BlueprintActionMenuUtils.cpp
--- a/Engine/Source/Editor/Kismet/Private/BlueprintActionMenuUtils.cpp
+++ b/Engine/Source/Editor/Kismet/Private/BlueprintActionMenuUtils.cpp
@@ -12,6 +12,8 @@
 #include "K2Node_ComponentBoundEvent.h"
 #include "EdGraphSchema_K2.h"		// for bUseLegacyActionMenus 
 #include "BlueprintEditorUtils.h"	// for DoesSupportComponents()
+#include "BlueprintPaletteFavorites.h"
+#include "K2ActionMenuBuilder.h"
 
 #define LOCTEXT_NAMESPACE "BlueprintActionMenuUtils"
 
@@ -22,19 +24,30 @@
 namespace BlueprintActionMenuUtilsImpl
 {
 	/**
+	 * Additional filter rejection test, for menu sections that only contain 
+	 * bound actions. Rejects any action that is not bound.
 	 * 
-	 * 
-	 * @param  Filter	
-	 * @param  BlueprintAction	
-	 * @return 
+	 * @param  Filter			The filter querying this rejection test.
+	 * @param  BlueprintAction	The action to test.
+	 * @return True if the spawner is unbound (and should be rejected), otherwise false.
 	 */
 	static bool IsUnBoundSpawner(FBlueprintActionFilter const& Filter, FBlueprintActionInfo& BlueprintAction);
 
 	/**
+	 * Filter rejection test, for favorite menus. Rejects any actions that are 
+	 * not favorited by the user.
 	 * 
+	 * @param  Filter			The filter querying this rejection test.
+	 * @param  BlueprintAction	The action to test.
+	 * @return True if the spawner is not favorited (and should be rejected), otherwise false.
+	 */
+	static bool IsNonFavoritedAction(FBlueprintActionFilter const& Filter, FBlueprintActionInfo& BlueprintAction);
+
+	/**
+	 * Utility function to find a common base class from a set of given classes.
 	 * 
-	 * @param  ClassSet	
-	 * @return 
+	 * @param  ClassSet	The set of classes that you want a single base class for.
+	 * @return A common base class for all the specified classes (fallsback to UObject's class).
 	 */
 	static UClass* FindCommonBaseClass(TArray<UClass*> const& ClassSet);
 }
@@ -45,6 +58,17 @@ static bool BlueprintActionMenuUtilsImpl::IsUnBoundSpawner(FBlueprintActionFilte
 	return (BlueprintAction.GetBindings().Num() <= 0);
 }
 
+//------------------------------------------------------------------------------
+static bool BlueprintActionMenuUtilsImpl::IsNonFavoritedAction(FBlueprintActionFilter const& Filter, FBlueprintActionInfo& BlueprintAction)
+{
+	UEditorUserSettings& EditorUserSettings = GEditor->AccessEditorUserSettings();
+	// grab the user's favorites
+	UBlueprintPaletteFavorites const* BlueprintFavorites = EditorUserSettings.BlueprintFavorites;
+	checkSlow(BlueprintFavorites != nullptr);
+
+	return !BlueprintFavorites->IsFavorited(BlueprintAction);
+}
+
 //------------------------------------------------------------------------------
 static UClass* BlueprintActionMenuUtilsImpl::FindCommonBaseClass(TArray<UClass*> const& ClassSet)
 {
@@ -248,4 +272,53 @@ void FBlueprintActionMenuUtils::MakeContextMenu(FBlueprintActionContext const& C
 	}
 }
 
+//------------------------------------------------------------------------------
+void FBlueprintActionMenuUtils::MakeFavoritesMenu(FBlueprintActionContext const& Context, FBlueprintActionMenuBuilder& MenuOut)
+{
+	MenuOut.Empty();
+
+	UEditorExperimentalSettings const* ExperimentalSettings = GetDefault<UEditorExperimentalSettings>();
+	if (ExperimentalSettings->bUseRefactoredBlueprintMenuingSystem)
+	{
+		FBlueprintActionFilter MenuFilter;
+		MenuFilter.Context = Context;
+		MenuFilter.AddRejectionTest(FBlueprintActionFilter::FRejectionTestDelegate::CreateStatic(BlueprintActionMenuUtilsImpl::IsNonFavoritedAction));
+
+		MenuOut.AddMenuSection(MenuFilter);
+		MenuOut.RebuildActionList();
+	}
+	else
+	{
+		check(Context.Blueprints.Num() > 0);
+		FBlueprintPaletteListBuilder LegacyMenuBuilder(Context.Blueprints[0]);
+		UEdGraphSchema_K2 const* K2Schema = GetDefault<UEdGraphSchema_K2>();
+		FK2ActionMenuBuilder(K2Schema).GetPaletteActions(LegacyMenuBuilder, /*ClassFilter =*/nullptr);
+
+		UEditorUserSettings& EditorUserSettings = GEditor->AccessEditorUserSettings();
+		// grab the user's favorites
+		UBlueprintPaletteFavorites const* BlueprintFavorites = EditorUserSettings.BlueprintFavorites;
+		check(BlueprintFavorites != nullptr);
+
+		for (int32 ActionIndex = 0; ActionIndex < LegacyMenuBuilder.GetNumActions(); ++ActionIndex)
+		{
+			FGraphActionListBuilderBase::ActionGroup& ActionGroup = LegacyMenuBuilder.GetAction(ActionIndex);
+			if (!ensure(ActionGroup.Actions.Num() == 1))
+			{
+				continue;
+			}
+
+			TSharedPtr<FEdGraphSchemaAction> Action = ActionGroup.Actions[0];
+			if (!Action.IsValid())
+			{
+				continue;
+			}
+
+			if (BlueprintFavorites->IsFavorited(Action))
+			{
+				MenuOut.AddAction(Action);
+			}
+		}
+	} 	
+}
+
 #undef LOCTEXT_NAMESPACE
diff --git a/Engine/Source/Editor/Kismet/Private/BlueprintBoundMenuItem.cpp b/Engine/Source/Editor/Kismet/Private/BlueprintBoundMenuItem.cpp
--- a/Engine/Source/Editor/Kismet/Private/BlueprintBoundMenuItem.cpp
+++ b/Engine/Source/Editor/Kismet/Private/BlueprintBoundMenuItem.cpp
@@ -178,4 +178,10 @@ FSlateBrush const* FBlueprintBoundMenuItem::GetMenuIcon(FSlateColor& ColorOut)
 	return nullptr;
 }
 
+//------------------------------------------------------------------------------
+UBlueprintNodeSpawner const* FBlueprintBoundMenuItem::GetBoundAction() const
+{
+	return BoundSpawner;
+}
+
 #undef LOCTEXT_NAMESPACE
diff --git a/Engine/Source/Editor/Kismet/Private/BlueprintDragDropMenuItem.cpp b/Engine/Source/Editor/Kismet/Private/BlueprintDragDropMenuItem.cpp
--- a/Engine/Source/Editor/Kismet/Private/BlueprintDragDropMenuItem.cpp
+++ b/Engine/Source/Editor/Kismet/Private/BlueprintDragDropMenuItem.cpp
@@ -21,9 +21,9 @@
 DEFINE_LOG_CATEGORY_STATIC(LogBlueprintDragDropMenuItem, Log, All);
 
 //------------------------------------------------------------------------------
-FBlueprintDragDropMenuItem::FBlueprintDragDropMenuItem(FBlueprintActionContext const& Context, UBlueprintNodeSpawner const* SampleActionIn, int32 MenuGrouping/* = 0*/)
-	: SampleAction(SampleActionIn)
+FBlueprintDragDropMenuItem::FBlueprintDragDropMenuItem(FBlueprintActionContext const& Context, UBlueprintNodeSpawner const* SampleAction, int32 MenuGrouping/* = 0*/)
 {
+	AppendAction(SampleAction);
 	check(SampleAction != nullptr);
 	Grouping  = MenuGrouping;
 
@@ -87,7 +87,7 @@ FBlueprintDragDropMenuItem::FBlueprintDragDropMenuItem(FBlueprintActionContext c
 UEdGraphNode* FBlueprintDragDropMenuItem::PerformAction(UEdGraph* ParentGraph, UEdGraphPin* FromPin, FVector2D const Location, bool bSelectNewNode/* = true*/)
 {
 	// we shouldn't get here (this should just be used for drag/drop ops), but just in case we do...
-	FBlueprintActionMenuItem BlueprintActionItem(SampleAction, nullptr, FSlateColor(FLinearColor::White));
+	FBlueprintActionMenuItem BlueprintActionItem(GetSampleAction(), nullptr, FSlateColor(FLinearColor::White));
 	return BlueprintActionItem.PerformAction(ParentGraph, FromPin, Location, bSelectNewNode);
 }
 
@@ -117,7 +117,16 @@ void FBlueprintDragDropMenuItem::AddReferencedObjects(FReferenceCollector& Colle
 	
 	// these don't get saved to disk, but we want to make sure the objects don't
 	// get GC'd while the action array is around
-	Collector.AddReferencedObject(SampleAction);
+	for (UBlueprintNodeSpawner const* Action : ActionSet)
+	{
+		Collector.AddReferencedObject(Action);
+	}
+}
+
+//------------------------------------------------------------------------------
+void FBlueprintDragDropMenuItem::AppendAction(UBlueprintNodeSpawner const* Action)
+{
+	ActionSet.Add(Action);
 }
 
 //------------------------------------------------------------------------------
@@ -126,6 +135,7 @@ FSlateBrush const* FBlueprintDragDropMenuItem::GetMenuIcon(FSlateColor& ColorOut
 	FSlateBrush const* IconBrush = nullptr;
 	ColorOut = FSlateColor(FLinearColor::White);
 
+	UBlueprintNodeSpawner const* SampleAction = GetSampleAction();
 	if (UBlueprintDelegateNodeSpawner const* DelegateSpawner = Cast<UBlueprintDelegateNodeSpawner const>(SampleAction))
 	{
 		IconBrush = FEditorStyle::GetBrush(TEXT("GraphEditor.Delegate_16x"));
@@ -141,11 +151,25 @@ FSlateBrush const* FBlueprintDragDropMenuItem::GetMenuIcon(FSlateColor& ColorOut
 	return IconBrush;
 }
 
+//------------------------------------------------------------------------------
+UBlueprintNodeSpawner const* FBlueprintDragDropMenuItem::GetSampleAction() const
+{
+	check(ActionSet.Num() > 0);
+	return *ActionSet.CreateConstIterator();
+}
+
+//------------------------------------------------------------------------------
+TSet<UBlueprintNodeSpawner const*> const& FBlueprintDragDropMenuItem::GetActionSet() const
+{
+	return ActionSet;
+}
+
 //------------------------------------------------------------------------------
 TSharedPtr<FDragDropOperation> FBlueprintDragDropMenuItem::OnDragged(FNodeCreationAnalytic AnalyticsDelegate) const
 {
 	TSharedPtr<FDragDropOperation> DragDropAction = nullptr;
 
+	UBlueprintNodeSpawner const* SampleAction = GetSampleAction();
 	if (UBlueprintDelegateNodeSpawner const* DelegateSpawner = Cast<UBlueprintDelegateNodeSpawner const>(SampleAction))
 	{
 		if (UMulticastDelegateProperty const* Property = DelegateSpawner->GetProperty())
diff --git a/Engine/Source/Editor/Kismet/Private/BlueprintPaletteFavorites.cpp b/Engine/Source/Editor/Kismet/Private/BlueprintPaletteFavorites.cpp
--- a/Engine/Source/Editor/Kismet/Private/BlueprintPaletteFavorites.cpp
+++ b/Engine/Source/Editor/Kismet/Private/BlueprintPaletteFavorites.cpp
@@ -2,215 +2,203 @@
 
 #include "BlueprintEditorPrivatePCH.h"
 #include "BlueprintPaletteFavorites.h"
+#include "BlueprintNodeSpawner.h"
+#include "BlueprintActionMenuItem.h"
+#include "BlueprintActionFilter.h"  // for FBlueprintActionInfo
+#include "BlueprintDragDropMenuItem.h"
+#include "BlueprintBoundMenuItem.h"
 
 /*******************************************************************************
-* Static File Helpers
-*******************************************************************************/
+ * Static UBlueprintPaletteFavorites Helpers
+ ******************************************************************************/
 
 /** namespace'd to avoid collisions with united builds */
-namespace BlueprintFavorites
+namespace BlueprintPaletteFavoritesImpl
 {
 	static FString const ConfigSection("BlueprintEditor.Favorites");
-	static FString const NodeNameConfigField("NodeName=");
-	static FString const FieldNameConfigField("FieldName=");
 	static FString const CustomProfileId("CustomProfile");
 	static FString const DefaultProfileConfigKey("DefaultProfile");
-}
-
-/**
- * Walks the specified object's outer hierarchy and concatenates all names along
- * the way (composing a path).
- * 
- * @param  Object	The object you want a path for.
- * @return A string representing the object's outer path.
- */
-static FString ConstructObjectPath(UObject const* Object)
-{
-	FString Path;
-	if (Object != NULL)
-	{
-		Path = Object->GetName();
-
-		UObject const* Outer = Object->GetOuter();
-		while (Outer != NULL)
-		{
-			Path = Outer->GetName() + "." + Path;
-			Outer = Outer->GetOuter();
-		}
 
-	}
-	
-	return Path;
+	/**
+	 * Before we refactored the blueprint menu system, signatures were manually
+	 * constructed based off node type, by combining a series of objects and 
+	 * names. Here we construct a new FBlueprintNodeSpawnerSignature from the 
+	 * old code (so as to mirror functionality).
+	 * 
+	 * @param  PaletteAction	The action you want a signature for.
+	 * @return A signature object, distinguishing the palette action from others (could also be invalid).
+	 */
+	static FBlueprintNodeSpawnerSignature ConstructLegacySignature(TSharedPtr<FEdGraphSchemaAction> PaletteAction);
 }
 
-/*******************************************************************************
-* FFavoritedBlueprintPaletteItem
-*******************************************************************************/
-
 //------------------------------------------------------------------------------
-FFavoritedBlueprintPaletteItem::FFavoritedBlueprintPaletteItem(FString const& SerializedAction)
-	: NodeClassOuter(NULL)
-	, FieldOuter(NULL)
+static FBlueprintNodeSpawnerSignature BlueprintPaletteFavoritesImpl::ConstructLegacySignature(TSharedPtr<FEdGraphSchemaAction> InPaletteAction)
 {
-	FString ParsedNodeName;
-	if (FParse::Value(*SerializedAction, *BlueprintFavorites::NodeNameConfigField, ParsedNodeName))
+	TSubclassOf<UEdGraphNode> SignatureNodeClass;
+	UObject const* SignatureSubObject = nullptr;
+	FName SignatureSubObjName;
+
+	FName const ActionId = InPaletteAction->GetTypeId();
+	if (ActionId == FEdGraphSchemaAction_K2AddComponent::StaticGetTypeId())
 	{
-		// @TODO look for redirects first
+		FEdGraphSchemaAction_K2AddComponent* AddComponentAction = (FEdGraphSchemaAction_K2AddComponent*)InPaletteAction.Get();
+		checkSlow(AddComponentAction->NodeTemplate != nullptr);
 
-		if (ResolveName(NodeClassOuter, ParsedNodeName, false, false))
-		{
-			NodeClassName = ParsedNodeName;
-		}
+		SignatureNodeClass = AddComponentAction->NodeTemplate->GetClass();
+		SignatureSubObject = AddComponentAction->ComponentClass;
 	}
+	else if (ActionId == FEdGraphSchemaAction_K2AddComment::StaticGetTypeId())
+	{
+		SignatureNodeClass = UEdGraphNode_Comment::StaticClass();
+	}
+	else if (ActionId == FEdGraphSchemaAction_K2Delegate::StaticGetTypeId())
+	{
+		FEdGraphSchemaAction_K2Delegate* DelegateAction = (FEdGraphSchemaAction_K2Delegate*)InPaletteAction.Get();
 
-	FString ParsedFieldName;
-	if (FParse::Value(*SerializedAction, *BlueprintFavorites::FieldNameConfigField, ParsedFieldName))
+		SignatureNodeClass = UK2Node_BaseMCDelegate::StaticClass();
+		SignatureSubObject = DelegateAction->GetDelegatePoperty();
+	}
+	// if we can pull out a node associated with this action
+	else if (UK2Node const* NodeTemplate = FK2SchemaActionUtils::ExtractNodeTemplateFromAction(InPaletteAction))
 	{
-		// @TODO look for redirects first
+		bool bIsSupported = false;
+		// now, if we need more info to help identify the node, let's fill 
+		// out FieldName/FieldOuter
 
-		if ((NodeClassName == UK2Node_MacroInstance::StaticClass()->GetName()) ||
-			(NodeClassName == UK2Node_InputKey::StaticClass()->GetName())      ||
-			(NodeClassName == UK2Node_InputAction::StaticClass()->GetName())   || 
-			(NodeClassName == UK2Node_InputAxisEvent::StaticClass()->GetName()) )
+		// with UK2Node_CallFunction node, we know we can use the function 
+		// to discern between them
+		if (UK2Node_CallFunction const* CallFuncNode = Cast<UK2Node_CallFunction const>(NodeTemplate))
 		{
-			FieldName = ParsedFieldName;
+			SignatureSubObject = CallFuncNode->FunctionReference.ResolveMember<UFunction>(CallFuncNode);
+			bIsSupported = (SignatureSubObject != nullptr);
 		}
-		else if (ResolveName(FieldOuter, ParsedFieldName, false, false))
+		else if (UK2Node_InputAxisEvent const* InputAxisEventNode = Cast<UK2Node_InputAxisEvent const>(NodeTemplate))
 		{
-			FieldName = ParsedFieldName;
+			SignatureSubObjName = InputAxisEventNode->InputAxisName;
+			bIsSupported = (SignatureSubObjName != NAME_None);
 		}
-	}	
-}
-
-//------------------------------------------------------------------------------
-FFavoritedBlueprintPaletteItem::FFavoritedBlueprintPaletteItem(TSharedPtr<FEdGraphSchemaAction> PaletteActionIn)
-	: NodeClassOuter(NULL)
-	, FieldOuter(NULL)
-{
-	if (PaletteActionIn.IsValid())
-	{
-		FName const ActionId = PaletteActionIn->GetTypeId();
-		if (ActionId == FEdGraphSchemaAction_K2AddComponent::StaticGetTypeId())
+		else if (UK2Node_Event const* EventNode = Cast<UK2Node_Event const>(NodeTemplate))
 		{
-			FEdGraphSchemaAction_K2AddComponent* AddComponentAction = (FEdGraphSchemaAction_K2AddComponent*)PaletteActionIn.Get();
-			check(AddComponentAction->NodeTemplate != NULL)
-
-			// grab information on the node's class
-			UClass* NodeClass = AddComponentAction->NodeTemplate->GetClass();
-			NodeClassOuter = NodeClass->GetOuter();
-			NodeClassName  = NodeClass->GetName();
-
-			FieldOuter = AddComponentAction->ComponentClass->GetOuter();
-			FieldName  = AddComponentAction->ComponentClass->GetName();
+			if (EventNode->EventSignatureClass != nullptr)
+			{
+				SignatureSubObject = FindField<UFunction>(EventNode->EventSignatureClass, EventNode->EventSignatureName);
+				bIsSupported = (SignatureSubObject != nullptr);
+			}
 		}
-		else if (ActionId == FEdGraphSchemaAction_K2AddComment::StaticGetTypeId())
+		else if (UK2Node_MacroInstance const* MacroNode = Cast<UK2Node_MacroInstance const>(NodeTemplate))
 		{
-			UClass* NodeClass = UEdGraphNode_Comment::StaticClass();
-			NodeClassOuter = NodeClass->GetOuter();
-			NodeClassName  = NodeClass->GetName();
+			SignatureSubObject = MacroNode->GetMacroGraph();
+			bIsSupported = (SignatureSubObject != nullptr);
 		}
-		else if (ActionId == FEdGraphSchemaAction_K2Delegate::StaticGetTypeId())
+		else if (UK2Node_InputKey const* InputKeyNode = Cast<UK2Node_InputKey const>(NodeTemplate))
 		{
-			FEdGraphSchemaAction_K2Delegate* DelegateAction = (FEdGraphSchemaAction_K2Delegate*)PaletteActionIn.Get();
+			SignatureSubObjName = InputKeyNode->InputKey.GetFName();
+			bIsSupported = (SignatureSubObjName != NAME_None);
+		}
+		else if (UK2Node_InputAction const* InputActionNode = Cast<UK2Node_InputAction const>(NodeTemplate))
+		{
+			SignatureSubObjName = InputActionNode->InputActionName;
+			bIsSupported = (SignatureSubObjName != NAME_None);
+		}
+		else if (Cast<UK2Node_IfThenElse const>(NodeTemplate) ||
+			Cast<UK2Node_MakeArray const>(NodeTemplate) ||
+			Cast<UK2Node_SpawnActorFromClass const>(NodeTemplate) ||
+			Cast<UK2Node_SpawnActor const>(NodeTemplate) ||
+			Cast<UK2Node_Timeline const>(NodeTemplate) ||
+			Cast<UK2Node_InputTouch const>(NodeTemplate))
+		{
+			bIsSupported = true;
+		}
 
-			UClass* NodeClass = UK2Node_BaseMCDelegate::StaticClass();
-			NodeClassOuter = NodeClass->GetOuter();
-			NodeClassName  = NodeClass->GetName();
+		if (bIsSupported)
+		{
+			SignatureNodeClass = NodeTemplate->GetClass();
+		}
+	}
 
-			FieldOuter = DelegateAction->GetDelegateClass();
-			FieldName  = DelegateAction->GetDelegateName().ToString();
+	FBlueprintNodeSpawnerSignature LegacySignatureSet(nullptr);
+	if (SignatureNodeClass != nullptr)
+	{
+		LegacySignatureSet.SetNodeClass(SignatureNodeClass);
+		if (SignatureSubObject != nullptr)
+		{
+			LegacySignatureSet.AddSubObject(SignatureSubObject);
 		}
-		// if we can pull out a node associated with this action
-		else if (UK2Node const* NodeTemplate = FK2SchemaActionUtils::ExtractNodeTemplateFromAction(PaletteActionIn))
+		else if (SignatureSubObjName != NAME_None)
 		{
-			bool bIsSupported = false;
-			// now, if we need more info to help identify the node, let's fill 
-			// out FieldName/FieldOuter
-
-			// with UK2Node_CallFunction node, we know we can use the function 
-			// to discern between them
-			if (UK2Node_CallFunction const* CallFuncNode = Cast<UK2Node_CallFunction const>(NodeTemplate))
-			{
-				FieldName  = CallFuncNode->FunctionReference.GetMemberName().ToString();
-				UFunction const* InnerFunction = CallFuncNode->FunctionReference.ResolveMember<UFunction>(CallFuncNode);
-				if (InnerFunction)
-				{
-					FieldOuter = InnerFunction->GetOuter();
-	
-					bIsSupported = true;
-				}
-			}
-			else if (UK2Node_InputAxisEvent const* InputAxisEventNode = Cast<UK2Node_InputAxisEvent const>(NodeTemplate))
-			{
-				FieldName  = InputAxisEventNode->InputAxisName.ToString();
-				FieldOuter = NULL;
+			static FName const SubObjectSignatureKey("FieldName");
+			LegacySignatureSet.AddKeyValue(SubObjectSignatureKey, SignatureSubObjName.ToString());
+		}
+	}
 
-				bIsSupported = true;
-			}
-			else if (UK2Node_Event const* EventNode = Cast<UK2Node_Event const>(NodeTemplate))
-			{
-				FieldName  = EventNode->EventSignatureName.ToString();
-				FieldOuter = EventNode->EventSignatureClass;
+	return LegacySignatureSet;
+}
 
-				bIsSupported = true;
-			}
-			else if (UK2Node_MacroInstance const* MacroNode = Cast<UK2Node_MacroInstance const>(NodeTemplate))
-			{
-				FieldName  = ConstructObjectPath(MacroNode->GetMacroGraph());
-				FieldOuter = NULL;
+/*******************************************************************************
+ * FFavoritedBlueprintPaletteItem
+ ******************************************************************************/
 
-				bIsSupported = true;
-			}
-			else if (UK2Node_InputKey const* InputKeyNode = Cast<UK2Node_InputKey const>(NodeTemplate))
-			{
-				FieldName  = InputKeyNode->InputKey.ToString();
-				FieldOuter = NULL;
+//------------------------------------------------------------------------------
+FFavoritedBlueprintPaletteItem::FFavoritedBlueprintPaletteItem(FString const& SerializedAction)
+	: ActionSignature(SerializedAction)
+{	
+}
 
-				bIsSupported = true;
-			}
-			else if (UK2Node_InputAction const* InputActionNode = Cast<UK2Node_InputAction const>(NodeTemplate))
-			{
-				FieldName  = InputActionNode->InputActionName.ToString();
-				FieldOuter = NULL;
+//------------------------------------------------------------------------------
+FFavoritedBlueprintPaletteItem::FFavoritedBlueprintPaletteItem(TSharedPtr<FEdGraphSchemaAction> InPaletteAction)
+{
+	if (InPaletteAction.IsValid())
+	{
+		if (InPaletteAction->GetTypeId() == FBlueprintActionMenuItem::StaticGetTypeId())
+		{
+			FBlueprintActionMenuItem* ActionMenuItem = (FBlueprintActionMenuItem*)InPaletteAction.Get();
+			ActionSignature = ActionMenuItem->GetRawAction()->GetSpawnerSignature();
+		}
+		else if (InPaletteAction->GetTypeId() == FBlueprintDragDropMenuItem::StaticGetTypeId())
+		{
+			FBlueprintDragDropMenuItem* CollectionMenuItem = (FBlueprintDragDropMenuItem*)InPaletteAction.Get();
+			ActionSignature = CollectionMenuItem->GetSampleAction()->GetSpawnerSignature();
 
-				bIsSupported = true;
-			}
-			else if (Cast<UK2Node_IfThenElse const>(NodeTemplate) || 
-			         Cast<UK2Node_MakeArray const>(NodeTemplate)  || 
-					 Cast<UK2Node_SpawnActorFromClass const>(NodeTemplate) || 
-					 Cast<UK2Node_SpawnActor const>(NodeTemplate) ||
-					 Cast<UK2Node_Timeline const>(NodeTemplate)   ||
-					 Cast<UK2Node_InputTouch const>(NodeTemplate))
-			{
-				// @TODO eventually we should switch this from being inclusive to exclusive 
-				//       (once we figure out what smaller subset is not explicitly supported)
-				bIsSupported = true;
-			}
+			// drag-n-drop menu items represent a collection of actions on the 
+			// same field (they spawn a submenu for the user to pick from)... so
+			// they don't have a single node class
+			ActionSignature.SetNodeClass(nullptr);
 
-			if (bIsSupported)
-			{
-				// grab information on the node's class
-				UClass* NodeClass = NodeTemplate->GetClass();
-				NodeClassOuter = NodeClass->GetOuter();
-				NodeClassName  = NodeClass->GetName();
-			}
+			static const FName CollectionSignatureKey(TEXT("ActionCollection"));
+			ActionSignature.AddKeyValue(CollectionSignatureKey, TEXT("true"));
+		}
+		else if (InPaletteAction->GetTypeId() == FBlueprintBoundMenuItem::StaticGetTypeId())
+		{
+			FBlueprintBoundMenuItem* BoundMenuItem = (FBlueprintBoundMenuItem*)InPaletteAction.Get();
+			// this is the bare action with out any bindings, you cannot 
+			// favorite actions with bindings (mainly because the database 
+			// actions are alway unbound, binding is done externally at menu 
+			// build time)
+			ActionSignature = BoundMenuItem->GetBoundAction()->GetSpawnerSignature();
+		}
+		else
+		{
+			ActionSignature = BlueprintPaletteFavoritesImpl::ConstructLegacySignature(InPaletteAction);
 		}
 	}
 }
 
+//------------------------------------------------------------------------------
+FFavoritedBlueprintPaletteItem::FFavoritedBlueprintPaletteItem(UBlueprintNodeSpawner const* BlueprintAction)
+	: ActionSignature(BlueprintAction->GetSpawnerSignature())
+{	
+}
+
 //------------------------------------------------------------------------------
 bool FFavoritedBlueprintPaletteItem::IsValid() const 
 {
-	return (NodeClassOuter != NULL);
+	return ActionSignature.IsValid();
 }
 
 //------------------------------------------------------------------------------
 bool FFavoritedBlueprintPaletteItem::operator==(FFavoritedBlueprintPaletteItem const& Rhs) const
 {
-	return ((NodeClassOuter == Rhs.NodeClassOuter) && 
-	        (FieldOuter == Rhs.FieldOuter)         && 
-	        (NodeClassName == Rhs.NodeClassName)   && 
-	        (FieldName == Rhs.FieldName));
+	return (ActionSignature.AsGuid() == Rhs.ActionSignature.AsGuid());
 }
 
 //------------------------------------------------------------------------------
@@ -220,39 +208,26 @@ bool FFavoritedBlueprintPaletteItem::operator==(TSharedPtr<FEdGraphSchemaAction>
 }
 
 //------------------------------------------------------------------------------
-FString FFavoritedBlueprintPaletteItem::ToString() const
+bool FFavoritedBlueprintPaletteItem::operator==(FBlueprintActionInfo& BlueprintAction) const
 {
-	FString SerializedFavorite;
-	if (IsValid())
-	{
-		// match the formatting we expect from the .ini files...
-		SerializedFavorite += "(";
-		SerializedFavorite += BlueprintFavorites::NodeNameConfigField;
-		SerializedFavorite += "\"";
-		SerializedFavorite += ConstructObjectPath(NodeClassOuter);
-		SerializedFavorite += ".";
-		SerializedFavorite += NodeClassName;
-		if (!FieldName.IsEmpty())
-		{
-			SerializedFavorite += "\", ";
-			SerializedFavorite += BlueprintFavorites::FieldNameConfigField;
-			SerializedFavorite += "\"";
-			if (FieldOuter != NULL)
-			{
-				SerializedFavorite += ConstructObjectPath(FieldOuter);
-				SerializedFavorite += ":";
-			}
-			SerializedFavorite += FieldName;
-		}
-		SerializedFavorite += "\")";
-	}
+	return (BlueprintAction.GetActionGuid() == ActionSignature.AsGuid());
+}
 
-	return SerializedFavorite;
+//------------------------------------------------------------------------------
+bool FFavoritedBlueprintPaletteItem::operator==(FGuid const& ActionGuid) const
+{
+	return (ActionGuid == ActionSignature.AsGuid());
+}
+
+//------------------------------------------------------------------------------
+FString const& FFavoritedBlueprintPaletteItem::ToString() const
+{
+	return ActionSignature.ToString();
 }
 
 /*******************************************************************************
-* UBlueprintPaletteFavorites Public Interface
-*******************************************************************************/
+ * UBlueprintPaletteFavorites Public Interface
+ ******************************************************************************/
 
 //------------------------------------------------------------------------------
 UBlueprintPaletteFavorites::UBlueprintPaletteFavorites(class FPostConstructInitializeProperties const& PCIP)
@@ -265,7 +240,7 @@ void UBlueprintPaletteFavorites::PostInitProperties()
 {
 	Super::PostInitProperties();
 
-	if (CurrentProfile == BlueprintFavorites::CustomProfileId)
+	if (CurrentProfile == BlueprintPaletteFavoritesImpl::CustomProfileId)
 	{
 		LoadCustomFavorites();
 	}
@@ -281,7 +256,7 @@ void UBlueprintPaletteFavorites::PostEditChangeProperty(struct FPropertyChangedE
 	Super::PostEditChangeProperty(PropertyChangedEvent);
 
 	CustomFavorites.Empty();
-	if (CurrentProfile == BlueprintFavorites::CustomProfileId)
+	if (CurrentProfile == BlueprintPaletteFavoritesImpl::CustomProfileId)
 	{
 		for (FFavoritedBlueprintPaletteItem const& Favorite : CurrentFavorites) 
 		{
@@ -304,57 +279,127 @@ bool UBlueprintPaletteFavorites::IsFavorited(TSharedPtr<FEdGraphSchemaAction> Pa
 {
 	bool bIsFavorited = false;
 
-	FFavoritedBlueprintPaletteItem ActionAsFavorite(PaletteAction);
-	if (ActionAsFavorite.IsValid())
+	if (PaletteAction->GetTypeId() == FBlueprintDragDropMenuItem::StaticGetTypeId())
 	{
-		for (FFavoritedBlueprintPaletteItem const& Favorite : CurrentFavorites)
+		FBlueprintDragDropMenuItem* CollectionMenuItem = (FBlueprintDragDropMenuItem*)PaletteAction.Get();
+
+		bIsFavorited = true;
+		for (UBlueprintNodeSpawner const* Action : CollectionMenuItem->GetActionSet())
 		{
-			if (Favorite == ActionAsFavorite)
+			if (!IsFavorited(Action))
 			{
-				bIsFavorited = true;
+				bIsFavorited = false;
 				break;
 			}
 		}
 	}
+	else
+	{
+		FFavoritedBlueprintPaletteItem ActionAsFavorite(PaletteAction);
+		if (ActionAsFavorite.IsValid())
+		{
+			for (FFavoritedBlueprintPaletteItem const& Favorite : CurrentFavorites)
+			{
+				if (Favorite == ActionAsFavorite)
+				{
+					bIsFavorited = true;
+					break;
+				}
+			}
+		}
+	}
 	return bIsFavorited;
 }
 
 //------------------------------------------------------------------------------
-void UBlueprintPaletteFavorites::AddFavorite(TSharedPtr<FEdGraphSchemaAction> PaletteAction)
+bool UBlueprintPaletteFavorites::IsFavorited(FBlueprintActionInfo& BlueprintAction) const
 {
-	if (!IsFavorited(PaletteAction) && CanBeFavorited(PaletteAction))
+	bool bIsFavorited = false;
+
+	for (FFavoritedBlueprintPaletteItem const& Favorite : CurrentFavorites)
 	{
-		CurrentFavorites.Add(FFavoritedBlueprintPaletteItem(PaletteAction));
-		SetProfile(BlueprintFavorites::CustomProfileId);
+		if (Favorite == BlueprintAction)
+		{
+			bIsFavorited = true;
+			break;
+		}
 	}
+	return bIsFavorited;
 }
 
 //------------------------------------------------------------------------------
-void UBlueprintPaletteFavorites::AddFavorites(TArray< TSharedPtr<FEdGraphSchemaAction> > PaletteActions)
+bool UBlueprintPaletteFavorites::IsFavorited(UBlueprintNodeSpawner const* BlueprintAction) const
 {
-	bool bAnyAdded = false;
-	for (TSharedPtr<FEdGraphSchemaAction>& NewFave : PaletteActions)
+	bool bIsFavorited = false;
+
+	FGuid ActionGuid = BlueprintAction->GetSpawnerSignature().AsGuid();
+	for (FFavoritedBlueprintPaletteItem const& Favorite : CurrentFavorites)
+	{
+		if (Favorite == ActionGuid)
+		{
+			bIsFavorited = true;
+			break;
+		}
+	}
+	return bIsFavorited;
+}
+
+//------------------------------------------------------------------------------
+void UBlueprintPaletteFavorites::AddFavorite(TSharedPtr<FEdGraphSchemaAction> PaletteAction)
+{
+	if (!IsFavorited(PaletteAction) && CanBeFavorited(PaletteAction))
 	{
-		if (!IsFavorited(NewFave) && CanBeFavorited(NewFave))
+		if (PaletteAction->GetTypeId() == FBlueprintDragDropMenuItem::StaticGetTypeId())
 		{
-			CurrentFavorites.Add(FFavoritedBlueprintPaletteItem(NewFave));
-			bAnyAdded = true;
+			FBlueprintDragDropMenuItem* CollectionMenuItem = (FBlueprintDragDropMenuItem*)PaletteAction.Get();
+			for (UBlueprintNodeSpawner const* Action : CollectionMenuItem->GetActionSet())
+			{
+				CurrentFavorites.Add(FFavoritedBlueprintPaletteItem(Action));
+			}
 		}
+		else
+		{
+			CurrentFavorites.Add(FFavoritedBlueprintPaletteItem(PaletteAction));
+		}
+		SetProfile(BlueprintPaletteFavoritesImpl::CustomProfileId);
 	}
+}
 
-	if (bAnyAdded)
-	{ 
-		SetProfile(BlueprintFavorites::CustomProfileId);
+//------------------------------------------------------------------------------
+void UBlueprintPaletteFavorites::AddFavorites(TArray< TSharedPtr<FEdGraphSchemaAction> > PaletteActions)
+{
+	for (TSharedPtr<FEdGraphSchemaAction>& NewFave : PaletteActions)
+	{
+		AddFavorite(NewFave);
 	}
 }
 
 //------------------------------------------------------------------------------
 void UBlueprintPaletteFavorites::RemoveFavorite(TSharedPtr<FEdGraphSchemaAction> PaletteAction)
 {
-	if (IsFavorited(PaletteAction))
+	if (PaletteAction->GetTypeId() == FBlueprintDragDropMenuItem::StaticGetTypeId())
+	{
+		bool bAnyRemoved = false;
+
+		FBlueprintDragDropMenuItem* CollectionMenuItem = (FBlueprintDragDropMenuItem*)PaletteAction.Get();
+		for (UBlueprintNodeSpawner const* Action : CollectionMenuItem->GetActionSet())
+		{
+			if (IsFavorited(Action))
+			{
+				CurrentFavorites.Remove(Action);
+				bAnyRemoved = true;
+			}
+		}
+
+		if (bAnyRemoved)
+		{
+			SetProfile(BlueprintPaletteFavoritesImpl::CustomProfileId);
+		}
+	}
+	else if (IsFavorited(PaletteAction))
 	{
 		CurrentFavorites.Remove(PaletteAction);
-		SetProfile(BlueprintFavorites::CustomProfileId);
+		SetProfile(BlueprintPaletteFavoritesImpl::CustomProfileId);
 	}
 }
 
@@ -364,16 +409,7 @@ void UBlueprintPaletteFavorites::RemoveFavorites(TArray< TSharedPtr<FEdGraphSche
 	bool bAnyRemoved = false;
 	for (TSharedPtr<FEdGraphSchemaAction>& OldFave : PaletteActions)
 	{
-		if (IsFavorited(OldFave))
-		{
-			CurrentFavorites.Remove(OldFave);
-			bAnyRemoved = true;
-		}
-	}
-
-	if (bAnyRemoved)
-	{ 
-		SetProfile(BlueprintFavorites::CustomProfileId);
+		RemoveFavorite(OldFave);
 	}
 }
 
@@ -391,7 +427,7 @@ void UBlueprintPaletteFavorites::LoadProfile(FString const& ProfileName)
 //------------------------------------------------------------------------------
 bool UBlueprintPaletteFavorites::IsUsingCustomProfile() const
 {
-	return (GetCurrentProfile() == BlueprintFavorites::CustomProfileId);
+	return (GetCurrentProfile() == BlueprintPaletteFavoritesImpl::CustomProfileId);
 }
 
 //------------------------------------------------------------------------------
@@ -410,19 +446,19 @@ void UBlueprintPaletteFavorites::ClearAllFavorites()
 	if (CurrentFavorites.Num() > 0) 
 	{
 		CurrentFavorites.Empty();
-		SetProfile(BlueprintFavorites::CustomProfileId);		
+		SetProfile(BlueprintPaletteFavoritesImpl::CustomProfileId);
 	}
 }
 
 /*******************************************************************************
-* UBlueprintPaletteFavorites Private Methods
-*******************************************************************************/
+ * UBlueprintPaletteFavorites Private Methods
+ ******************************************************************************/
 
 //------------------------------------------------------------------------------
 FString const& UBlueprintPaletteFavorites::GetDefaultProfileId() const
 {
 	static FString DefaultProfileId("DefaultFavorites");
-	GConfig->GetString(*BlueprintFavorites::ConfigSection, *BlueprintFavorites::DefaultProfileConfigKey, DefaultProfileId, GEditorIni);
+	GConfig->GetString(*BlueprintPaletteFavoritesImpl::ConfigSection, *BlueprintPaletteFavoritesImpl::DefaultProfileConfigKey, DefaultProfileId, GEditorIni);
 	return DefaultProfileId;
 }
 
@@ -434,10 +470,10 @@ void UBlueprintPaletteFavorites::LoadSetProfile()
 
 	TArray<FString> ProfileFavorites;
 	// if this profile doesn't exist anymore
-	if (CurrentProfile.IsEmpty() || (GConfig->GetArray(*BlueprintFavorites::ConfigSection, *CurrentProfile, ProfileFavorites, GEditorIni) == 0))
+	if (CurrentProfile.IsEmpty() || (GConfig->GetArray(*BlueprintPaletteFavoritesImpl::ConfigSection, *CurrentProfile, ProfileFavorites, GEditorIni) == 0))
 	{
-		// @TODO log warning
-		GConfig->GetArray(*BlueprintFavorites::ConfigSection, *GetDefaultProfileId(), ProfileFavorites, GEditorIni);
+		// @TODO: log warning
+		GConfig->GetArray(*BlueprintPaletteFavoritesImpl::ConfigSection, *GetDefaultProfileId(), ProfileFavorites, GEditorIni);
 	}
 
 	for (FString const& Favorite : ProfileFavorites)
@@ -454,7 +490,7 @@ void UBlueprintPaletteFavorites::LoadSetProfile()
 void UBlueprintPaletteFavorites::LoadCustomFavorites()
 {
 	CurrentFavorites.Empty();
-	check(CurrentProfile == BlueprintFavorites::CustomProfileId);
+	check(CurrentProfile == BlueprintPaletteFavoritesImpl::CustomProfileId);
 
 	for (FString const& Favorite : CustomFavorites)
 	{
diff --git a/Engine/Source/Editor/Kismet/Private/SBlueprintFavoritesPalette.cpp b/Engine/Source/Editor/Kismet/Private/SBlueprintFavoritesPalette.cpp
--- a/Engine/Source/Editor/Kismet/Private/SBlueprintFavoritesPalette.cpp
+++ b/Engine/Source/Editor/Kismet/Private/SBlueprintFavoritesPalette.cpp
@@ -5,6 +5,8 @@
 #include "Engine.h"
 #include "KismetClasses.h"
 #include "K2ActionMenuBuilder.h"
+#include "BlueprintActionMenuUtils.h"
+#include "BlueprintActionMenuBuilder.h"
 
 #define LOCTEXT_NAMESPACE "BlueprintFavoritesPalette"
 
@@ -301,42 +303,56 @@ void SBlueprintFavoritesPalette::Construct(FArguments const& InArgs, TWeakPtr<FB
 //------------------------------------------------------------------------------
 void SBlueprintFavoritesPalette::CollectAllActions(FGraphActionListBuilderBase& OutAllActions)
 {
-	UBlueprint const* const Blueprint = GetBlueprint();
-	FBlueprintPaletteListBuilder AllBlueprintActions(Blueprint);
-	// get all actions available to this blueprint
-	UEdGraphSchema_K2::GetPaletteActions(AllBlueprintActions);	
-
-	UEditorUserSettings& EditorUserSettings = GEditor->AccessEditorUserSettings();
-	// grab the user's favorites
-	UBlueprintPaletteFavorites const* BlueprintFavorites = EditorUserSettings.BlueprintFavorites;
-	check(BlueprintFavorites != NULL);
-
-	FString RootCategory = SBlueprintFavoritesPaletteUtils::FavoritesCategoryName;
-	if (!bShowFrequentlyUsed)
-	{
-		RootCategory = TEXT("");
-	}
-
-	// @TODO there has to be a better way than going through EVERY single action (using favorites as a filter?)
-	for (int32 ActionIndex = 0; ActionIndex < AllBlueprintActions.GetNumActions(); ++ActionIndex)
-	{
-		FGraphActionListBuilderBase::ActionGroup& ActionGroup = AllBlueprintActions.GetAction(ActionIndex);
-		if (!ensureMsg(ActionGroup.Actions.Num() == 1, TEXT("SBlueprintFavoritesPalette::CollectAllActions() - Currently not handling action-groups with more than one action.")))
-		{
-			continue;
-		}
-
-		TSharedPtr<FEdGraphSchemaAction> Action = ActionGroup.Actions[0];
-		if (!Action.IsValid())
-		{
-			continue;
-		}
-
-		if (BlueprintFavorites->IsFavorited(Action))
-		{
-			OutAllActions.AddAction(Action, RootCategory);
-		}
-	}
+// 	UBlueprint const* const Blueprint = GetBlueprint();
+// 	FBlueprintPaletteListBuilder AllBlueprintActions(Blueprint);
+// 	// get all actions available to this blueprint
+// 	UEdGraphSchema_K2::GetPaletteActions(AllBlueprintActions);	
+// 
+// 	UEditorUserSettings& EditorUserSettings = GEditor->AccessEditorUserSettings();
+// 	// grab the user's favorites
+// 	UBlueprintPaletteFavorites const* BlueprintFavorites = EditorUserSettings.BlueprintFavorites;
+// 	check(BlueprintFavorites != NULL);
+// 
+// 	FString RootCategory = SBlueprintFavoritesPaletteUtils::FavoritesCategoryName;
+// 	if (!bShowFrequentlyUsed)
+// 	{
+// 		RootCategory = TEXT("");
+// 	}
+// 
+// 	// @TODO there has to be a better way than going through EVERY single action (using favorites as a filter?)
+// 	for (int32 ActionIndex = 0; ActionIndex < AllBlueprintActions.GetNumActions(); ++ActionIndex)
+// 	{
+// 		FGraphActionListBuilderBase::ActionGroup& ActionGroup = AllBlueprintActions.GetAction(ActionIndex);
+// 		if (!ensureMsg(ActionGroup.Actions.Num() == 1, TEXT("SBlueprintFavoritesPalette::CollectAllActions() - Currently not handling action-groups with more than one action.")))
+// 		{
+// 			continue;
+// 		}
+// 
+// 		TSharedPtr<FEdGraphSchemaAction> Action = ActionGroup.Actions[0];
+// 		if (!Action.IsValid())
+// 		{
+// 			continue;
+// 		}
+// 
+// 		if (BlueprintFavorites->IsFavorited(Action))
+// 		{
+// 			OutAllActions.AddAction(Action, RootCategory);
+// 		}
+// 	}
+
+// 	FString RootCategory = SBlueprintLibraryPaletteUtils::LibraryCategoryName;
+// 	if (bUseLegacyLayout)
+// 	{
+// 		RootCategory = TEXT("");
+// 	}
+
+	FBlueprintActionContext FilterContext;
+	FilterContext.Blueprints.Add(GetBlueprint());
+
+	FBlueprintActionMenuBuilder FavoritesBuilder(BlueprintEditorPtr);
+	FBlueprintActionMenuUtils::MakeFavoritesMenu(FilterContext, FavoritesBuilder);
+
+	OutAllActions.Append(FavoritesBuilder);
 }
 
 //------------------------------------------------------------------------------
diff --git a/Engine/Source/Editor/Kismet/Private/SBlueprintLibraryPalette.cpp b/Engine/Source/Editor/Kismet/Private/SBlueprintLibraryPalette.cpp
--- a/Engine/Source/Editor/Kismet/Private/SBlueprintLibraryPalette.cpp
+++ b/Engine/Source/Editor/Kismet/Private/SBlueprintLibraryPalette.cpp
@@ -220,8 +220,6 @@ void SBlueprintLibraryPalette::CollectAllActions(FGraphActionListBuilderBase& Ou
 	
 	FBlueprintActionMenuBuilder PaletteBuilder(BlueprintEditorPtr);
 	FBlueprintActionMenuUtils::MakePaletteMenu(FilterContext, ClassFilter, PaletteBuilder);
-	
-	PaletteBuilder.RebuildActionList();
 	OutAllActions.Append(PaletteBuilder);
 }
 
diff --git a/Engine/Source/Editor/Kismet/Private/SBlueprintPalette.cpp b/Engine/Source/Editor/Kismet/Private/SBlueprintPalette.cpp
--- a/Engine/Source/Editor/Kismet/Private/SBlueprintPalette.cpp
+++ b/Engine/Source/Editor/Kismet/Private/SBlueprintPalette.cpp
@@ -30,6 +30,7 @@
 #include "BlueprintDragDropMenuItem.h"
 #include "BlueprintNodeSpawner.h"
 #include "TutorialMetaData.h"
+#include "BlueprintEditorSettings.h" // for bShowActionMenuItemSignatures
 
 #define LOCTEXT_NAMESPACE "BlueprintPalette"
 
@@ -995,7 +996,7 @@ TSharedRef<SWidget> SBlueprintPaletteItem::CreateTextSlotWidget(const FSlateFont
 		InCreateData->OnRenameRequest->BindSP(InlineRenameWidget.Get(), &SInlineEditableTextBlock::EnterEditingMode);
 	}
 
-	if (GEditor->EditorUserSettings->bDisplayActionListItemRefIds && ActionPtr.IsValid())
+	if (GetDefault<UBlueprintEditorSettings>()->bShowActionMenuItemSignatures && ActionPtr.IsValid())
 	{
 		check(InlineRenameWidget.IsValid());
 		TSharedPtr<IToolTip> ExistingToolTip = InlineRenameWidget->GetToolTip();
diff --git a/Engine/Source/Editor/Kismet/Public/BlueprintActionMenuItem.h b/Engine/Source/Editor/Kismet/Public/BlueprintActionMenuItem.h
--- a/Engine/Source/Editor/Kismet/Public/BlueprintActionMenuItem.h
+++ b/Engine/Source/Editor/Kismet/Public/BlueprintActionMenuItem.h
@@ -47,6 +47,9 @@ public:
 	 */
 	FSlateBrush const* GetMenuIcon(FSlateColor& ColorOut);
 
+	/** @return */
+	UBlueprintNodeSpawner const* GetRawAction() const;
+
 private:
 	/** Brush that should be used for the icon on this menu item. */
 	FSlateBrush const* IconBrush;
diff --git a/Engine/Source/Editor/Kismet/Public/BlueprintActionMenuUtils.h b/Engine/Source/Editor/Kismet/Public/BlueprintActionMenuUtils.h
--- a/Engine/Source/Editor/Kismet/Public/BlueprintActionMenuUtils.h
+++ b/Engine/Source/Editor/Kismet/Public/BlueprintActionMenuUtils.h
@@ -29,5 +29,15 @@ struct FBlueprintActionMenuUtils
 	 * @param  MenuOut				The structure that will be populated with context menu items.
 	 */
 	KISMET_API static void MakeContextMenu(FBlueprintActionContext const& Context, bool bIsContextSensitive, FBlueprintActionMenuBuilder& MenuOut);
+
+	/**
+	 * A centralized utility function for constructing the blueprint favorites
+	 * menu. Takes the palette menu building code and folds in an additional 
+	 * "IsFavorited" filter check.
+	 *
+	 * @param  Context	Contains the blueprint that the palette is for.
+	 * @param  MenuOut	The structure that will be populated with favorite menu items.
+	 */
+	KISMET_API static void MakeFavoritesMenu(FBlueprintActionContext const& Context, FBlueprintActionMenuBuilder& MenuOut);
 };
 
diff --git a/Engine/Source/Editor/Kismet/Public/BlueprintBoundMenuItem.h b/Engine/Source/Editor/Kismet/Public/BlueprintBoundMenuItem.h
--- a/Engine/Source/Editor/Kismet/Public/BlueprintBoundMenuItem.h
+++ b/Engine/Source/Editor/Kismet/Public/BlueprintBoundMenuItem.h
@@ -51,7 +51,10 @@ public:
 	 * @param  ColorOut	An output parameter that's filled in with the color to tint the brush with.
 	 * @return An slate brush to be used for this menu item in the action menu.
 	 */
-	FSlateBrush const* GetMenuIcon(FSlateColor& ColorOut);	
+	FSlateBrush const* GetMenuIcon(FSlateColor& ColorOut);
+
+	/** @return   */
+	UBlueprintNodeSpawner const* GetBoundAction() const;
 
 private:
 	/** Instanced node-spawner, that comprises the action portion of this menu entry. */
diff --git a/Engine/Source/Editor/Kismet/Public/BlueprintDragDropMenuItem.h b/Engine/Source/Editor/Kismet/Public/BlueprintDragDropMenuItem.h
--- a/Engine/Source/Editor/Kismet/Public/BlueprintDragDropMenuItem.h
+++ b/Engine/Source/Editor/Kismet/Public/BlueprintDragDropMenuItem.h
@@ -34,7 +34,7 @@ struct FBlueprintDragDropMenuItem : public FEdGraphSchemaAction
 public:
 	static FName StaticGetTypeId() { static FName const TypeId("FBlueprintDragDropMenuItem"); return TypeId; }
 	/** Default constructor for USTRUCT() purposes. */
-	FBlueprintDragDropMenuItem() : SampleAction(nullptr) {}
+	FBlueprintDragDropMenuItem() {}
 
 	/**
 	 * Sole constructor, which ensures that the RepresentativeAction member is
@@ -61,6 +61,15 @@ public:
 	 */
 	FSlateBrush const* GetMenuIcon(FSlateColor& ColorOut);
 
+	/** */
+	void AppendAction(UBlueprintNodeSpawner const* Action);
+
+	/** @return   */
+	UBlueprintNodeSpawner const* GetSampleAction() const;
+
+	/** @return   */
+	TSet<UBlueprintNodeSpawner const*> const& GetActionSet() const;
+
 	/**
 	 * Attempts to create a certain drag/drop action with respect to the 
 	 * set UBlueprintNodeSpawner.
@@ -72,5 +81,5 @@ public:
 
 private:
 	/** An arbitrary member of the node-spawner subset that this menu entry represents. */
-	UBlueprintNodeSpawner const* SampleAction;
+	TSet<UBlueprintNodeSpawner const*> ActionSet;
 };
diff --git a/Engine/Source/Editor/UnrealEd/Classes/Editor/EditorUserSettings.h b/Engine/Source/Editor/UnrealEd/Classes/Editor/EditorUserSettings.h
--- a/Engine/Source/Editor/UnrealEd/Classes/Editor/EditorUserSettings.h
+++ b/Engine/Source/Editor/UnrealEd/Classes/Editor/EditorUserSettings.h
@@ -24,7 +24,7 @@ class UEditorUserSettings : public UObject
 	uint32 bDisplayDocumentationLink:1;
 
 	/** If enabled, tooltips on SGraphPaletteItems will show the associated action's string id */
-	UPROPERTY(EditAnywhere, config, Category=DeveloperTools)
+	UPROPERTY(/*EditAnywhere - deprecated (moved into UBlueprintEditorSettings), */config/*, Category=DeveloperTools*/)
 	uint32 bDisplayActionListItemRefIds:1;
 	
 	/** When enabled, the application frame rate, memory and Unreal object count will be displayed in the main editor UI */
diff --git a/Engine/Source/Editor/UnrealEd/Classes/Settings/EditorExperimentalSettings.h b/Engine/Source/Editor/UnrealEd/Classes/Settings/EditorExperimentalSettings.h
--- a/Engine/Source/Editor/UnrealEd/Classes/Settings/EditorExperimentalSettings.h
+++ b/Engine/Source/Editor/UnrealEd/Classes/Settings/EditorExperimentalSettings.h
@@ -57,7 +57,7 @@ public:
 	bool bDrawMidpointArrowsInBlueprints;
 
 	/** Determines if the Blueprint editor should use the new (experimental) menu system, or the old (legacy) system */
-	UPROPERTY(EditAnywhere, config, Category=Blueprints, meta=(DisplayName="Use new Blueprint menuing system"))
+	UPROPERTY(EditAnywhere, config, Category=Blueprints, meta=(DisplayName="Use New Blueprint Menuing System"))
 	bool bUseRefactoredBlueprintMenuingSystem;
 
 	/** Whether to show Audio Streaming options for SoundWaves (disabling will not stop all audio streaming) */
diff --git a/Engine/Source/Runtime/InputCore/Classes/InputCoreTypes.h b/Engine/Source/Runtime/InputCore/Classes/InputCoreTypes.h
--- a/Engine/Source/Runtime/InputCore/Classes/InputCoreTypes.h
+++ b/Engine/Source/Runtime/InputCore/Classes/InputCoreTypes.h
@@ -37,6 +37,7 @@ struct INPUTCORE_API FKey
 	bool IsBindableInBlueprints() const;
 	FText GetDisplayName() const;
 	FString ToString() const;
+	FName GetFName() const;
 
 	bool SerializeFromMismatchedTag(struct FPropertyTag const& Tag, FArchive& Ar);
 	bool ExportTextItem(FString& ValueStr, FKey const& DefaultValue, UObject* Parent, int32 PortFlags, UObject* ExportRootScope) const;
diff --git a/Engine/Source/Runtime/InputCore/Private/InputCoreTypes.cpp b/Engine/Source/Runtime/InputCore/Private/InputCoreTypes.cpp
--- a/Engine/Source/Runtime/InputCore/Private/InputCoreTypes.cpp
+++ b/Engine/Source/Runtime/InputCore/Private/InputCoreTypes.cpp
@@ -643,6 +643,11 @@ FString FKey::ToString() const
 	return KeyName.ToString(); 
 }
 
+FName FKey::GetFName() const
+{
+	return KeyName;
+}
+
 bool FKey::IsModifierKey() const 
 {
 	ConditionalLookupKeyDetails();
