Fixing mac build errors for the Blueprint favoriting system.

diff --git a/Engine/Source/Editor/Kismet/Private/BlueprintPaletteFavorites.cpp b/Engine/Source/Editor/Kismet/Private/BlueprintPaletteFavorites.cpp
--- a/Engine/Source/Editor/Kismet/Private/BlueprintPaletteFavorites.cpp
+++ b/Engine/Source/Editor/Kismet/Private/BlueprintPaletteFavorites.cpp
@@ -116,7 +116,7 @@ static FBlueprintNodeSpawnerSignature BlueprintPaletteFavoritesImpl::ConstructLe
 		}
 	}
 
-	FBlueprintNodeSpawnerSignature LegacySignatureSet(nullptr);
+	FBlueprintNodeSpawnerSignature LegacySignatureSet;
 	if (SignatureNodeClass != nullptr)
 	{
 		LegacySignatureSet.SetNodeClass(SignatureNodeClass);
