Fixing non-unity build error (missing include).

diff --git a/Engine/Source/Editor/Kismet/Private/SBlueprintFavoritesPalette.cpp b/Engine/Source/Editor/Kismet/Private/SBlueprintFavoritesPalette.cpp
--- a/Engine/Source/Editor/Kismet/Private/SBlueprintFavoritesPalette.cpp
+++ b/Engine/Source/Editor/Kismet/Private/SBlueprintFavoritesPalette.cpp
@@ -7,6 +7,7 @@
 #include "K2ActionMenuBuilder.h"
 #include "BlueprintActionMenuUtils.h"
 #include "BlueprintActionMenuBuilder.h"
+#include "BlueprintActionFilter.h" // for FBlueprintActionContext
 
 #define LOCTEXT_NAMESPACE "BlueprintFavoritesPalette"
 
@@ -303,49 +304,6 @@ void SBlueprintFavoritesPalette::Construct(FArguments const& InArgs, TWeakPtr<FB
 //------------------------------------------------------------------------------
 void SBlueprintFavoritesPalette::CollectAllActions(FGraphActionListBuilderBase& OutAllActions)
 {
-// 	UBlueprint const* const Blueprint = GetBlueprint();
-// 	FBlueprintPaletteListBuilder AllBlueprintActions(Blueprint);
-// 	// get all actions available to this blueprint
-// 	UEdGraphSchema_K2::GetPaletteActions(AllBlueprintActions);	
-// 
-// 	UEditorUserSettings& EditorUserSettings = GEditor->AccessEditorUserSettings();
-// 	// grab the user's favorites
-// 	UBlueprintPaletteFavorites const* BlueprintFavorites = EditorUserSettings.BlueprintFavorites;
-// 	check(BlueprintFavorites != NULL);
-// 
-// 	FString RootCategory = SBlueprintFavoritesPaletteUtils::FavoritesCategoryName;
-// 	if (!bShowFrequentlyUsed)
-// 	{
-// 		RootCategory = TEXT("");
-// 	}
-// 
-// 	// @TODO there has to be a better way than going through EVERY single action (using favorites as a filter?)
-// 	for (int32 ActionIndex = 0; ActionIndex < AllBlueprintActions.GetNumActions(); ++ActionIndex)
-// 	{
-// 		FGraphActionListBuilderBase::ActionGroup& ActionGroup = AllBlueprintActions.GetAction(ActionIndex);
-// 		if (!ensureMsg(ActionGroup.Actions.Num() == 1, TEXT("SBlueprintFavoritesPalette::CollectAllActions() - Currently not handling action-groups with more than one action.")))
-// 		{
-// 			continue;
-// 		}
-// 
-// 		TSharedPtr<FEdGraphSchemaAction> Action = ActionGroup.Actions[0];
-// 		if (!Action.IsValid())
-// 		{
-// 			continue;
-// 		}
-// 
-// 		if (BlueprintFavorites->IsFavorited(Action))
-// 		{
-// 			OutAllActions.AddAction(Action, RootCategory);
-// 		}
-// 	}
-
-// 	FString RootCategory = SBlueprintLibraryPaletteUtils::LibraryCategoryName;
-// 	if (bUseLegacyLayout)
-// 	{
-// 		RootCategory = TEXT("");
-// 	}
-
 	FBlueprintActionContext FilterContext;
 	FilterContext.Blueprints.Add(GetBlueprint());
 
