Destroy preview actor before setting the world as having Begun Play
Check that there is a PreviewBlueprint before seeking object for it otherwise it will find the first non-blueprint actor as the preview actor

diff --git a/Engine/Source/Editor/Kismet/Private/SCSEditorViewportClient.cpp b/Engine/Source/Editor/Kismet/Private/SCSEditorViewportClient.cpp
--- a/Engine/Source/Editor/Kismet/Private/SCSEditorViewportClient.cpp
+++ b/Engine/Source/Editor/Kismet/Private/SCSEditorViewportClient.cpp
@@ -768,6 +768,9 @@ bool FSCSEditorViewportClient::GetIsSimulateEnabled()
 
 void FSCSEditorViewportClient::ToggleIsSimulateEnabled() 
 {
+	// Must destroy existing actors before we toggle the world state
+	DestroyPreview();
+
 	bIsSimulateEnabled = !bIsSimulateEnabled;
 	PreviewScene->GetWorld()->bBegunPlay = bIsSimulateEnabled;
 	PreviewScene->GetWorld()->bShouldSimulatePhysics = bIsSimulateEnabled;
@@ -877,7 +880,7 @@ void FSCSEditorViewportClient::EndTransaction()
 AActor* FSCSEditorViewportClient::GetPreviewActor() const
 {
 	// Note: The weak ptr can become stale if the actor is reinstanced due to a Blueprint change, etc. In that case we look to see if we can find the new instance in the preview world and then update the weak ptr.
-	if(PreviewActorPtr.IsStale(true))
+	if(PreviewActorPtr.IsStale(true) && PreviewBlueprint)
 	{
 		UWorld* PreviewWorld = PreviewScene->GetWorld();
 		for(TActorIterator<AActor> It(PreviewWorld); It; ++It)
