<saved by Perforce>

diff --git a/Engine/Source/Editor/BlueprintGraph/Classes/BlueprintBoundEventNodeSpawner.h b/Engine/Source/Editor/BlueprintGraph/Classes/BlueprintBoundEventNodeSpawner.h
--- a/Engine/Source/Editor/BlueprintGraph/Classes/BlueprintBoundEventNodeSpawner.h
+++ b/Engine/Source/Editor/BlueprintGraph/Classes/BlueprintBoundEventNodeSpawner.h
@@ -34,7 +34,7 @@ public:
 	// UBlueprintNodeSpawner interface
 	virtual FBlueprintNodeSignature GetSpawnerSignature() const override;
 	virtual UEdGraphNode* Invoke(UEdGraph* ParentGraph, FBindingSet const& Bindings, FVector2D const Location) const override;
-	virtual FText GetDefaultMenuName() const override;
+	virtual FText GetDefaultMenuName(FBindingSet const& Bindings) const override;
 	virtual FText GetDefaultMenuCategory() const override;
 	// End UBlueprintNodeSpawner interface
 
diff --git a/Engine/Source/Editor/BlueprintGraph/Classes/BlueprintBoundNodeSpawner.h b/Engine/Source/Editor/BlueprintGraph/Classes/BlueprintBoundNodeSpawner.h
--- a/Engine/Source/Editor/BlueprintGraph/Classes/BlueprintBoundNodeSpawner.h
+++ b/Engine/Source/Editor/BlueprintGraph/Classes/BlueprintBoundNodeSpawner.h
@@ -15,6 +15,7 @@ class BLUEPRINTGRAPH_API UBlueprintBoundNodeSpawner : public UBlueprintNodeSpawn
 	GENERATED_UCLASS_BODY()
 	DECLARE_DELEGATE_RetVal_OneParam(bool, FCanBindObjectDelegate, UObject const*);
 	DECLARE_DELEGATE_RetVal_TwoParams(bool, FOnBindObjectDelegate, UEdGraphNode*, UObject*);
+	DECLARE_DELEGATE_RetVal_OneParam(FText, FOnGenerateMenuDescriptionDelegate, const IBlueprintNodeBinder::FBindingSet& );
 
 public:
 	/**
@@ -24,6 +25,7 @@ public:
 
 	// UBlueprintNodeSpawner interface
 	virtual FBlueprintNodeSignature GetSpawnerSignature() const override;
+	virtual FText GetDefaultMenuName(FBindingSet const& Bindings) const override;
 	// End UBlueprintNodeSpawner interface
 
 	// FBlueprintNodeBinder interface
@@ -41,4 +43,9 @@ public:
 	 * A delegate to perform specialized node setup during binding
 	 */
 	FOnBindObjectDelegate OnBindObjectDelegate;
+
+	/**
+	 * A delegate to generate a description for the menu item. Executed any time bindings change
+	 */
+	FOnGenerateMenuDescriptionDelegate OnGenerateMenuDescriptionDelegate;
 };
diff --git a/Engine/Source/Editor/BlueprintGraph/Classes/BlueprintComponentNodeSpawner.h b/Engine/Source/Editor/BlueprintGraph/Classes/BlueprintComponentNodeSpawner.h
--- a/Engine/Source/Editor/BlueprintGraph/Classes/BlueprintComponentNodeSpawner.h
+++ b/Engine/Source/Editor/BlueprintGraph/Classes/BlueprintComponentNodeSpawner.h
@@ -34,7 +34,7 @@ public:
 	// UBlueprintNodeSpawner interface
 	virtual FBlueprintNodeSignature GetSpawnerSignature() const override;
 	virtual UEdGraphNode* Invoke(UEdGraph* ParentGraph, FBindingSet const& Bindings, FVector2D const Location) const override;
-	virtual FText GetDefaultMenuName() const override;
+	virtual FText GetDefaultMenuName(FBindingSet const& Bindings) const override;
 	virtual FText GetDefaultMenuCategory() const override;
 	virtual FName GetDefaultMenuIcon(FLinearColor& ColorOut) const override;
 	// End UBlueprintNodeSpawner interface
diff --git a/Engine/Source/Editor/BlueprintGraph/Classes/BlueprintDelegateNodeSpawner.h b/Engine/Source/Editor/BlueprintGraph/Classes/BlueprintDelegateNodeSpawner.h
--- a/Engine/Source/Editor/BlueprintGraph/Classes/BlueprintDelegateNodeSpawner.h
+++ b/Engine/Source/Editor/BlueprintGraph/Classes/BlueprintDelegateNodeSpawner.h
@@ -32,7 +32,7 @@ public:
 	static UBlueprintDelegateNodeSpawner* Create(TSubclassOf<UK2Node_BaseMCDelegate> NodeClass, UMulticastDelegateProperty const* const Property, UObject* Outer = nullptr);
 
 	// UBlueprintNodeSpawner interface
-	virtual FText GetDefaultMenuName() const override;
+	virtual FText GetDefaultMenuName(FBindingSet const& Bindings) const override;
 	virtual FText GetDefaultMenuCategory() const override;
 	virtual FName GetDefaultMenuIcon(FLinearColor& ColorOut) const override;
 	// End UBlueprintNodeSpawner interface
diff --git a/Engine/Source/Editor/BlueprintGraph/Classes/BlueprintEventNodeSpawner.h b/Engine/Source/Editor/BlueprintGraph/Classes/BlueprintEventNodeSpawner.h
--- a/Engine/Source/Editor/BlueprintGraph/Classes/BlueprintEventNodeSpawner.h
+++ b/Engine/Source/Editor/BlueprintGraph/Classes/BlueprintEventNodeSpawner.h
@@ -47,7 +47,7 @@ public:
 	// UBlueprintNodeSpawner interface
 	virtual FBlueprintNodeSignature GetSpawnerSignature() const override;
 	virtual UEdGraphNode* Invoke(UEdGraph* ParentGraph, FBindingSet const& Bindings, FVector2D const Location) const override;
-	virtual FText GetDefaultMenuName() const override;
+	virtual FText GetDefaultMenuName(FBindingSet const& Bindings) const override;
 	virtual FText GetDefaultMenuCategory() const override;
 	virtual FString GetDefaultSearchKeywords() const override;
 	// End UBlueprintNodeSpawner interface
diff --git a/Engine/Source/Editor/BlueprintGraph/Classes/BlueprintFunctionNodeSpawner.h b/Engine/Source/Editor/BlueprintGraph/Classes/BlueprintFunctionNodeSpawner.h
--- a/Engine/Source/Editor/BlueprintGraph/Classes/BlueprintFunctionNodeSpawner.h
+++ b/Engine/Source/Editor/BlueprintGraph/Classes/BlueprintFunctionNodeSpawner.h
@@ -44,7 +44,7 @@ public:
 	// UBlueprintNodeSpawner interface
 	virtual void Prime() override;
 	virtual UEdGraphNode* Invoke(UEdGraph* ParentGraph, FBindingSet const& Bindings, FVector2D const Location) const override;
-	virtual FText GetDefaultMenuName() const override;
+	virtual FText GetDefaultMenuName(FBindingSet const& Bindings) const override;
 	virtual FText GetDefaultMenuCategory() const override;
 	virtual FText GetDefaultMenuTooltip() const override;
 	virtual FString GetDefaultSearchKeywords() const override;
diff --git a/Engine/Source/Editor/BlueprintGraph/Classes/BlueprintVariableNodeSpawner.h b/Engine/Source/Editor/BlueprintGraph/Classes/BlueprintVariableNodeSpawner.h
--- a/Engine/Source/Editor/BlueprintGraph/Classes/BlueprintVariableNodeSpawner.h
+++ b/Engine/Source/Editor/BlueprintGraph/Classes/BlueprintVariableNodeSpawner.h
@@ -49,7 +49,7 @@ public:
 	virtual void Prime() override;
 	virtual FBlueprintNodeSignature GetSpawnerSignature() const override;
 	virtual UEdGraphNode* Invoke(UEdGraph* ParentGraph, FBindingSet const& Bindings, FVector2D const Location) const override;
-	virtual FText GetDefaultMenuName() const override;
+	virtual FText GetDefaultMenuName(FBindingSet const& Bindings) const override;
 	virtual FText GetDefaultMenuCategory() const override;
 	virtual FText GetDefaultMenuTooltip() const override;
 	virtual FString GetDefaultSearchKeywords() const override;
diff --git a/Engine/Source/Editor/BlueprintGraph/Classes/K2Node_MatineeController.h b/Engine/Source/Editor/BlueprintGraph/Classes/K2Node_MatineeController.h
--- a/Engine/Source/Editor/BlueprintGraph/Classes/K2Node_MatineeController.h
+++ b/Engine/Source/Editor/BlueprintGraph/Classes/K2Node_MatineeController.h
@@ -26,6 +26,7 @@ class UK2Node_MatineeController : public UK2Node
 	virtual bool IsNodeSafeToIgnore() const override { return true; }
 	virtual AActor* GetReferencedLevelActor() const override;
 	virtual void ExpandNode(class FKismetCompilerContext& CompilerContext, UEdGraph* SourceGraph) override;
+	virtual void GetMenuActions(FBlueprintActionDatabaseRegistrar& ActionRegistrar) const override;
 	// End UK2Node interface
 
 	// Begin UObject interface
diff --git a/Engine/Source/Editor/BlueprintGraph/Private/BlueprintBoundEventNodeSpawner.cpp b/Engine/Source/Editor/BlueprintGraph/Private/BlueprintBoundEventNodeSpawner.cpp
--- a/Engine/Source/Editor/BlueprintGraph/Private/BlueprintBoundEventNodeSpawner.cpp
+++ b/Engine/Source/Editor/BlueprintGraph/Private/BlueprintBoundEventNodeSpawner.cpp
@@ -55,7 +55,7 @@ UEdGraphNode* UBlueprintBoundEventNodeSpawner::Invoke(UEdGraph* ParentGraph, FBi
 }
 
 //------------------------------------------------------------------------------
-FText UBlueprintBoundEventNodeSpawner::GetDefaultMenuName() const
+FText UBlueprintBoundEventNodeSpawner::GetDefaultMenuName(FBindingSet const& Bindings) const
 {
 	if (CachedMenuName.IsOutOfDate())
 	{
diff --git a/Engine/Source/Editor/BlueprintGraph/Private/BlueprintBoundNodeSpawner.cpp b/Engine/Source/Editor/BlueprintGraph/Private/BlueprintBoundNodeSpawner.cpp
--- a/Engine/Source/Editor/BlueprintGraph/Private/BlueprintBoundNodeSpawner.cpp
+++ b/Engine/Source/Editor/BlueprintGraph/Private/BlueprintBoundNodeSpawner.cpp
@@ -38,6 +38,16 @@ FBlueprintNodeSignature UBlueprintBoundNodeSpawner::GetSpawnerSignature() const
 	return FBlueprintNodeSignature();
 }
 
+//------------------------------------------------------------------------------
+FText UBlueprintBoundNodeSpawner::GetDefaultMenuName(FBindingSet const& Bindings) const
+{
+	if( OnGenerateMenuDescriptionDelegate.IsBound() )
+	{
+		return OnGenerateMenuDescriptionDelegate.Execute(Bindings);
+	}
+	return UBlueprintNodeSpawner::GetDefaultMenuName( Bindings );
+}
+
 //------------------------------------------------------------------------------
 bool UBlueprintBoundNodeSpawner::IsBindingCompatible(UObject const* BindingCandidate) const
 {
diff --git a/Engine/Source/Editor/BlueprintGraph/Private/BlueprintComponentNodeSpawner.cpp b/Engine/Source/Editor/BlueprintGraph/Private/BlueprintComponentNodeSpawner.cpp
--- a/Engine/Source/Editor/BlueprintGraph/Private/BlueprintComponentNodeSpawner.cpp
+++ b/Engine/Source/Editor/BlueprintGraph/Private/BlueprintComponentNodeSpawner.cpp
@@ -100,7 +100,7 @@ UEdGraphNode* UBlueprintComponentNodeSpawner::Invoke(UEdGraph* ParentGraph, FBin
 }
 
 //------------------------------------------------------------------------------
-FText UBlueprintComponentNodeSpawner::GetDefaultMenuName() const
+FText UBlueprintComponentNodeSpawner::GetDefaultMenuName(FBindingSet const& Bindings) const
 {
 	check(ComponentClass != nullptr);
 	if (CachedMenuName.IsOutOfDate())
diff --git a/Engine/Source/Editor/BlueprintGraph/Private/BlueprintDelegateNodeSpawner.cpp b/Engine/Source/Editor/BlueprintGraph/Private/BlueprintDelegateNodeSpawner.cpp
--- a/Engine/Source/Editor/BlueprintGraph/Private/BlueprintDelegateNodeSpawner.cpp
+++ b/Engine/Source/Editor/BlueprintGraph/Private/BlueprintDelegateNodeSpawner.cpp
@@ -53,9 +53,9 @@ UBlueprintDelegateNodeSpawner::UBlueprintDelegateNodeSpawner(class FPostConstruc
 }
 
 //------------------------------------------------------------------------------
-FText UBlueprintDelegateNodeSpawner::GetDefaultMenuName() const
+FText UBlueprintDelegateNodeSpawner::GetDefaultMenuName(FBindingSet const& Bindings) const
 {	
-	FText MenuName = Super::GetDefaultMenuName();
+	FText MenuName = Super::GetDefaultMenuName(Bindings);
 
 	UMulticastDelegateProperty const* Property = GetProperty();
 	if ((NodeClass == nullptr) && (Property != nullptr))
diff --git a/Engine/Source/Editor/BlueprintGraph/Private/BlueprintEventNodeSpawner.cpp b/Engine/Source/Editor/BlueprintGraph/Private/BlueprintEventNodeSpawner.cpp
--- a/Engine/Source/Editor/BlueprintGraph/Private/BlueprintEventNodeSpawner.cpp
+++ b/Engine/Source/Editor/BlueprintGraph/Private/BlueprintEventNodeSpawner.cpp
@@ -154,7 +154,7 @@ UEdGraphNode* UBlueprintEventNodeSpawner::Invoke(UEdGraph* ParentGraph, FBinding
 }
 
 //------------------------------------------------------------------------------
-FText UBlueprintEventNodeSpawner::GetDefaultMenuName() const
+FText UBlueprintEventNodeSpawner::GetDefaultMenuName(FBindingSet const& Bindings) const
 {
 	if (CachedMenuName.IsOutOfDate())
 	{
diff --git a/Engine/Source/Editor/BlueprintGraph/Private/BlueprintFunctionNodeSpawner.cpp b/Engine/Source/Editor/BlueprintGraph/Private/BlueprintFunctionNodeSpawner.cpp
--- a/Engine/Source/Editor/BlueprintGraph/Private/BlueprintFunctionNodeSpawner.cpp
+++ b/Engine/Source/Editor/BlueprintGraph/Private/BlueprintFunctionNodeSpawner.cpp
@@ -225,7 +225,7 @@ UEdGraphNode* UBlueprintFunctionNodeSpawner::Invoke(UEdGraph* ParentGraph, FBind
 }
 
 //------------------------------------------------------------------------------
-FText UBlueprintFunctionNodeSpawner::GetDefaultMenuName() const
+FText UBlueprintFunctionNodeSpawner::GetDefaultMenuName(FBindingSet const& Bindings) const
 {
 	UFunction const* Function = GetFunction();
 	check(Function != nullptr);
@@ -249,7 +249,8 @@ FText UBlueprintFunctionNodeSpawner::GetDefaultMenuTooltip() const
 	FText Tooltip = FText::FromString(UK2Node_CallFunction::GetDefaultTooltipForFunction(Function));
 	if (Tooltip.IsEmpty())
 	{
-		Tooltip = GetDefaultMenuName();
+		FBindingSet EmptyContext;
+		Tooltip = GetDefaultMenuName(EmptyContext);
 	}
 	return Tooltip;
 }
diff --git a/Engine/Source/Editor/BlueprintGraph/Private/BlueprintNodeSpawner.cpp b/Engine/Source/Editor/BlueprintGraph/Private/BlueprintNodeSpawner.cpp
--- a/Engine/Source/Editor/BlueprintGraph/Private/BlueprintNodeSpawner.cpp
+++ b/Engine/Source/Editor/BlueprintGraph/Private/BlueprintNodeSpawner.cpp
@@ -94,7 +94,8 @@ void UBlueprintNodeSpawner::Prime()
 	}
 
 	// in case any of these cache FText::Format() operations
-	GetDefaultMenuName();
+	FBindingSet EmptyContext;
+	GetDefaultMenuName(EmptyContext);
 	GetDefaultMenuCategory();
 	GetDefaultMenuTooltip();
 }
@@ -122,7 +123,7 @@ UEdGraphNode* UBlueprintNodeSpawner::Invoke(UEdGraph* ParentGraph, FBindingSet c
 }
 
 //------------------------------------------------------------------------------
-FText UBlueprintNodeSpawner::GetDefaultMenuName() const
+FText UBlueprintNodeSpawner::GetDefaultMenuName(FBindingSet const& Bindings) const
 {
 	// DO NOT make a template node and query it here (the separate ui building
 	// code can do that if it likes)... this is meant to be an overridable
diff --git a/Engine/Source/Editor/BlueprintGraph/Private/BlueprintVariableNodeSpawner.cpp b/Engine/Source/Editor/BlueprintGraph/Private/BlueprintVariableNodeSpawner.cpp
--- a/Engine/Source/Editor/BlueprintGraph/Private/BlueprintVariableNodeSpawner.cpp
+++ b/Engine/Source/Editor/BlueprintGraph/Private/BlueprintVariableNodeSpawner.cpp
@@ -78,7 +78,8 @@ void UBlueprintVariableNodeSpawner::Prime()
 
 	// all of these perform expensive FText::Format() operations and cache the 
 	// results...
-	GetDefaultMenuName();
+	FBindingSet EmptyContext;
+	GetDefaultMenuName(EmptyContext);
 	GetDefaultMenuCategory();
 	GetDefaultMenuTooltip();
 }
@@ -129,7 +130,7 @@ UEdGraphNode* UBlueprintVariableNodeSpawner::Invoke(UEdGraph* ParentGraph, FBind
 }
 
 //------------------------------------------------------------------------------
-FText UBlueprintVariableNodeSpawner::GetDefaultMenuName() const
+FText UBlueprintVariableNodeSpawner::GetDefaultMenuName(const FBindingSet& Bindings) const
 {	
 	if (CachedMenuName.IsOutOfDate())
 	{
diff --git a/Engine/Source/Editor/BlueprintGraph/Private/K2Node_MatineeController.cpp b/Engine/Source/Editor/BlueprintGraph/Private/K2Node_MatineeController.cpp
--- a/Engine/Source/Editor/BlueprintGraph/Private/K2Node_MatineeController.cpp
+++ b/Engine/Source/Editor/BlueprintGraph/Private/K2Node_MatineeController.cpp
@@ -1,6 +1,9 @@
 // Copyright 1998-2014 Epic Games, Inc. All Rights Reserved.
 
 #include "BlueprintGraphPrivatePCH.h"
+
+#include "BlueprintActionDatabaseRegistrar.h"
+#include "BlueprintBoundNodeSpawner.h"
 #include "KismetCompiler.h"
 #include "Runtime/Engine/Public/MatineeDelegates.h"
 #include "Matinee/MatineeActor.h"
@@ -124,6 +127,51 @@ void UK2Node_MatineeController::ExpandNode(FKismetCompilerContext& CompilerConte
 	}
 }
 
+void UK2Node_MatineeController::GetMenuActions(FBlueprintActionDatabaseRegistrar& ActionRegistrar) const
+{
+	auto CanBindObjectLambda = [](UObject const* BindingObject)
+	{
+		if (AMatineeActor const* Actor = Cast<AMatineeActor>(BindingObject))
+		{
+			// Make sure the Actor has a world
+			if (Actor->GetWorld())
+			{
+				return true;
+			}
+		}
+		return false;
+	};
+
+	auto MenuCategoryDescription = []( const IBlueprintNodeBinder::FBindingSet& BindingContext ) -> FText
+	{
+		if( BindingContext.Num() == 1 )
+		{
+			return FText::Format(NSLOCTEXT("K2Node", "MatineeeControllerTitle", "Create a Matinee Controller for {0}"), FText::FromString( (*(BindingContext.CreateConstIterator()))->GetName()) );
+		}
+		else if( BindingContext.Num() > 1 )
+		{
+			return FText::Format(NSLOCTEXT("K2Node", "MultipleMatineeeControllerTitle", "Create Matinee Controllers for {0} selected MatineeActors"), FText::AsNumber(BindingContext.Num()) );
+		}
+		else
+		{
+			return NSLOCTEXT("K2Node", "FallbackMatineeeControllerTitle", "Error: No MatineeActors in Context");
+		}
+	};
+
+	auto PostBindSetupLambda = [](UEdGraphNode* NewNode, UObject* BindObject)
+	{
+		UK2Node_MatineeController* MatineeNode = CastChecked<UK2Node_MatineeController>(NewNode);
+		MatineeNode->MatineeActor = CastChecked<AMatineeActor>(BindObject);
+		return true;
+	};
+
+	UBlueprintBoundNodeSpawner* NodeSpawner = UBlueprintBoundNodeSpawner::Create(GetClass());
+	NodeSpawner->CanBindObjectDelegate = UBlueprintBoundNodeSpawner::FCanBindObjectDelegate::CreateStatic(CanBindObjectLambda);
+	NodeSpawner->OnBindObjectDelegate = UBlueprintBoundNodeSpawner::FOnBindObjectDelegate::CreateStatic(PostBindSetupLambda);
+	NodeSpawner->OnGenerateMenuDescriptionDelegate = UBlueprintBoundNodeSpawner::FOnGenerateMenuDescriptionDelegate::CreateStatic(MenuCategoryDescription);
+	ActionRegistrar.AddBlueprintAction(NodeSpawner);
+}
+
 void UK2Node_MatineeController::OnEventKeyframeAdded(const AMatineeActor* InMatineeActor, const FName& InPinName, int32 InIndex)
 {
 	if ( MatineeActor == InMatineeActor )
diff --git a/Engine/Source/Editor/BlueprintGraph/Public/BlueprintNodeSpawner.h b/Engine/Source/Editor/BlueprintGraph/Public/BlueprintNodeSpawner.h
--- a/Engine/Source/Editor/BlueprintGraph/Public/BlueprintNodeSpawner.h
+++ b/Engine/Source/Editor/BlueprintGraph/Public/BlueprintNodeSpawner.h
@@ -116,7 +116,7 @@ public:
 	 *
 	 * @return An empty text string (could be a localized name for sub-classes).
 	 */
-	virtual FText GetDefaultMenuName() const;
+	virtual FText GetDefaultMenuName(FBindingSet const& Bindings) const;
 	
 	/**
 	 * To save on performance, certain sub-classes can concoct menu categories
@@ -189,6 +189,7 @@ public:
 	// IBlueprintNodeBinder interface
 	virtual bool   IsBindingCompatible(UObject const* BindingCandidate) const override { return false; }
 	virtual bool   CanBindMultipleObjects() const override { return false; }
+
 protected:
 	virtual bool   BindToNode(UEdGraphNode* Node, UObject* Binding) const override { return false; }
 	// End IBlueprintNodeBinder interface
diff --git a/Engine/Source/Editor/Kismet/Private/BlueprintActionMenuBuilder.cpp b/Engine/Source/Editor/Kismet/Private/BlueprintActionMenuBuilder.cpp
--- a/Engine/Source/Editor/Kismet/Private/BlueprintActionMenuBuilder.cpp
+++ b/Engine/Source/Editor/Kismet/Private/BlueprintActionMenuBuilder.cpp
@@ -50,7 +50,7 @@ public:
 	 * @param  Action			The node-spawner that the new menu item should wrap.
 	 * @return A newly allocated FBlueprintActionMenuItem (which wraps the supplied action).
 	 */
-	TSharedPtr<FEdGraphSchemaAction> MakeActionMenuItem(TWeakPtr<FBlueprintEditor> EditorContext, UBlueprintNodeSpawner const* Action);
+	TSharedPtr<FEdGraphSchemaAction> MakeActionMenuItem(TWeakPtr<FBlueprintEditor> EditorContext, FBlueprintActionInfo const& ActionInfo);
 
 	/**
 	 * Spawns a new FBlueprintDragDropMenuItem with the node-spawner. Constructs
@@ -67,7 +67,7 @@ public:
 	 * @param  BoundAction	
 	 * @return 
 	 */
-	TSharedPtr<FBlueprintBoundMenuItem> MakeBoundMenuItem(TWeakPtr<FBlueprintEditor> EditorContext, UBlueprintNodeSpawner const* Action, IBlueprintNodeBinder::FBindingSet const& Bindings);
+	TSharedPtr<FBlueprintBoundMenuItem> MakeBoundMenuItem(TWeakPtr<FBlueprintEditor> EditorContext, FBlueprintActionInfo const& ActionInfo);
 	
 private:
 	/**
@@ -79,7 +79,7 @@ private:
 	 * @param  Action		The action you want to suss name information from.
 	 * @return A name for the menu item wrapping this action.
 	 */
-	FText GetMenuNameForAction(TWeakPtr<FBlueprintEditor> EditorContext, UBlueprintNodeSpawner const* Action);
+	FText GetMenuNameForAction(TWeakPtr<FBlueprintEditor> EditorContext, FBlueprintActionInfo const& ActionInfo);
 
 	/**
 	 * Attempts to pull a menu category from the supplied spawner. If one isn't 
@@ -152,13 +152,15 @@ FBlueprintActionMenuItemFactory::FBlueprintActionMenuItemFactory(FBlueprintActio
 }
 
 //------------------------------------------------------------------------------
-TSharedPtr<FEdGraphSchemaAction> FBlueprintActionMenuItemFactory::MakeActionMenuItem(TWeakPtr<FBlueprintEditor> EditorContext, UBlueprintNodeSpawner const* Action)
+TSharedPtr<FEdGraphSchemaAction> FBlueprintActionMenuItemFactory::MakeActionMenuItem(TWeakPtr<FBlueprintEditor> EditorContext, FBlueprintActionInfo const& ActionInfo)
 {
+	UBlueprintNodeSpawner const* Action = ActionInfo.NodeSpawner;
+
 	FLinearColor IconTint = FLinearColor::White;
 	FName IconBrushName = GetMenuIconForAction(EditorContext, Action, IconTint);
 
 	FBlueprintActionMenuItem * NewMenuItem = new FBlueprintActionMenuItem(Action, FEditorStyle::GetBrush(IconBrushName), IconTint, MenuGrouping);
-	NewMenuItem->MenuDescription    = GetMenuNameForAction(EditorContext, Action);
+	NewMenuItem->MenuDescription    = GetMenuNameForAction(EditorContext, ActionInfo);
 	NewMenuItem->TooltipDescription = GetTooltipForAction(EditorContext, Action).ToString();
 	NewMenuItem->Category           = GetCategoryForAction(EditorContext, Action).ToString();
 	NewMenuItem->Keywords           = GetSearchKeywordsForAction(EditorContext, Action);
@@ -178,15 +180,17 @@ TSharedPtr<FBlueprintDragDropMenuItem> FBlueprintActionMenuItemFactory::MakeDrag
 }
 
 //------------------------------------------------------------------------------
-TSharedPtr<FBlueprintBoundMenuItem> FBlueprintActionMenuItemFactory::MakeBoundMenuItem(TWeakPtr<FBlueprintEditor> EditorContext, UBlueprintNodeSpawner const* Action, IBlueprintNodeBinder::FBindingSet const& Bindings)
+TSharedPtr<FBlueprintBoundMenuItem> FBlueprintActionMenuItemFactory::MakeBoundMenuItem(TWeakPtr<FBlueprintEditor> EditorContext, FBlueprintActionInfo const& ActionInfo)
 {
+	IBlueprintNodeBinder::FBindingSet const& Bindings = ActionInfo.GetBindings();
+	UBlueprintNodeSpawner const* Action = ActionInfo.NodeSpawner;
 	FBlueprintBoundMenuItem* NewMenuItem = new FBlueprintBoundMenuItem(Action, MenuGrouping);
+	NewMenuItem->MenuDescription = GetMenuNameForAction(EditorContext, ActionInfo);
+
 	NewMenuItem->AddBindings(Bindings);
 
 	// FBlueprintBoundMenuItem takes care of its own menu category
 	NewMenuItem->Category = FString::Printf(TEXT("%s|%s"), *RootCategory.ToString(), *NewMenuItem->Category);
-
-	NewMenuItem->MenuDescription = GetMenuNameForAction(EditorContext, Action);
 	NewMenuItem->TooltipDescription = GetTooltipForAction(EditorContext, Action).ToString();
 	NewMenuItem->Keywords = GetSearchKeywordsForAction(EditorContext, Action);
 
@@ -194,11 +198,12 @@ TSharedPtr<FBlueprintBoundMenuItem> FBlueprintActionMenuItemFactory::MakeBoundMe
 }
 
 //------------------------------------------------------------------------------
-FText FBlueprintActionMenuItemFactory::GetMenuNameForAction(TWeakPtr<FBlueprintEditor> EditorContext, UBlueprintNodeSpawner const* Action)
+FText FBlueprintActionMenuItemFactory::GetMenuNameForAction(TWeakPtr<FBlueprintEditor> EditorContext, FBlueprintActionInfo const& ActionInfo)
 {
+	UBlueprintNodeSpawner const* Action = ActionInfo.NodeSpawner;
 	check(Action != nullptr);
 	// give the action the chance to save on performance (to keep from having to spawn a template node)
-	FText MenuName = Action->GetDefaultMenuName();
+	FText MenuName = Action->GetDefaultMenuName(ActionInfo.GetBindings());
 	
 	if (MenuName.IsEmpty())
 	{
@@ -486,7 +491,7 @@ TSharedPtr<FEdGraphSchemaAction> FBlueprintActionMenuBuilderImpl::FMenuSectionDe
 				{
 					if (!MenuItem.IsValid())
 					{
-						MenuItem = ItemFactory.MakeBoundMenuItem(EditorContext, DatabaseAction, CompatibleBindingsToFilter);
+						MenuItem = ItemFactory.MakeBoundMenuItem(EditorContext, BoundActionInfo);
 					}
 					else
 					{
@@ -551,7 +556,7 @@ TSharedPtr<FEdGraphSchemaAction> FBlueprintActionMenuBuilderImpl::FMenuSectionDe
 	
 	if (!MenuEntry.IsValid() && bPassedFilter)
 	{
-		MenuEntry = ItemFactory.MakeActionMenuItem(EditorContext, DatabaseAction.NodeSpawner);
+		MenuEntry = ItemFactory.MakeActionMenuItem(EditorContext, DatabaseAction);
 	}
 
 	return MenuEntry;
diff --git a/Engine/Source/Editor/Kismet/Private/BlueprintActionMenuUtils.cpp b/Engine/Source/Editor/Kismet/Private/BlueprintActionMenuUtils.cpp
--- a/Engine/Source/Editor/Kismet/Private/BlueprintActionMenuUtils.cpp
+++ b/Engine/Source/Editor/Kismet/Private/BlueprintActionMenuUtils.cpp
@@ -5,15 +5,16 @@
 #include "BlueprintActionMenuBuilder.h"
 #include "BlueprintActionFilter.h"
 #include "BlueprintNodeSpawner.h"
-#include "KismetEditorUtilities.h"	// for CanPasteNodes()
-#include "K2Node_VariableGet.h"
-#include "K2Node_VariableSet.h"
-#include "K2Node_CallFunction.h"
-#include "K2Node_ActorBoundEvent.h"
-#include "K2Node_ComponentBoundEvent.h"
-#include "EdGraphSchema_K2.h"		// for bUseLegacyActionMenus 
 #include "BlueprintEditorUtils.h"	// for DoesSupportComponents()
 #include "BlueprintPaletteFavorites.h"
+#include "EdGraphSchema_K2.h"		// for bUseLegacyActionMenus 
+#include "K2Node_ActorBoundEvent.h"
+#include "K2Node_CallFunction.h"
+#include "K2Node_ComponentBoundEvent.h"
+#include "K2Node_MatineeController.h"
+#include "K2Node_VariableGet.h"
+#include "K2Node_VariableSet.h"
+#include "KismetEditorUtilities.h"	// for CanPasteNodes()
 #include "K2ActionMenuBuilder.h"
 
 #define LOCTEXT_NAMESPACE "BlueprintActionMenuUtils"
@@ -147,6 +148,7 @@ void FBlueprintActionMenuUtils::MakeContextMenu(FBlueprintActionContext const& C
 	LevelActorsFilter.PermittedNodeTypes.Add(UK2Node_CallFunction::StaticClass());
 	LevelActorsFilter.PermittedNodeTypes.Add(UK2Node_ActorBoundEvent::StaticClass());
 	LevelActorsFilter.PermittedNodeTypes.Add(UK2Node_Literal::StaticClass());
+	LevelActorsFilter.PermittedNodeTypes.Add(UK2Node_MatineeController::StaticClass());
 	// only want bound actions for this menu section
 	LevelActorsFilter.AddRejectionTest(FBlueprintActionFilter::FRejectionTestDelegate::CreateStatic(IsUnBoundSpawner));
 
diff --git a/Engine/Source/Editor/Kismet/Private/BlueprintBoundMenuItem.cpp b/Engine/Source/Editor/Kismet/Private/BlueprintBoundMenuItem.cpp
--- a/Engine/Source/Editor/Kismet/Private/BlueprintBoundMenuItem.cpp
+++ b/Engine/Source/Editor/Kismet/Private/BlueprintBoundMenuItem.cpp
@@ -169,6 +169,13 @@ void FBlueprintBoundMenuItem::AddBindings(IBlueprintNodeBinder::FBindingSet cons
 {
 	BoundObjects.Append(BindingSet);
 	Category = BlueprintBoundMenuItemImpl::GetMenuCategoryFromAction(BoundSpawner, BoundObjects).ToString();
+
+	// Allow the binding to generate a menu description:
+	FText NewMenuDescription = BoundSpawner->GetDefaultMenuName(BoundObjects);
+	if(!NewMenuDescription.IsEmpty())
+	{
+		MenuDescription = NewMenuDescription;
+	}
 }
 
 //------------------------------------------------------------------------------
