Adding support for CallOnMember nodes (so you can directly operate on components that your target class has).

diff --git a/Engine/Source/Editor/BlueprintGraph/Private/BlueprintComponentNodeSpawner.cpp b/Engine/Source/Editor/BlueprintGraph/Private/BlueprintComponentNodeSpawner.cpp
--- a/Engine/Source/Editor/BlueprintGraph/Private/BlueprintComponentNodeSpawner.cpp
+++ b/Engine/Source/Editor/BlueprintGraph/Private/BlueprintComponentNodeSpawner.cpp
@@ -62,7 +62,7 @@ UEdGraphNode* UBlueprintComponentNodeSpawner::Invoke(UEdGraph* ParentGraph, FBin
 	};
 	FCustomizeNodeDelegate PostSpawnDelegate = FCustomizeNodeDelegate::CreateStatic(PostSpawnLambda, CustomizeNodeDelegate);
 	
-	UK2Node_AddComponent* NewNode = CastChecked<UK2Node_AddComponent>(Super::Invoke(ParentGraph, Bindings, Location, PostSpawnDelegate));
+	UK2Node_AddComponent* NewNode = Super::SpawnNode<UK2Node_AddComponent>(NodeClass, ParentGraph, Bindings, Location, PostSpawnDelegate);
 
 	bool const bIsTemplateNode = FBlueprintNodeTemplateCache::IsTemplateOuter(ParentGraph);
 	if (!bIsTemplateNode)
diff --git a/Engine/Source/Editor/BlueprintGraph/Private/BlueprintEventNodeSpawner.cpp b/Engine/Source/Editor/BlueprintGraph/Private/BlueprintEventNodeSpawner.cpp
--- a/Engine/Source/Editor/BlueprintGraph/Private/BlueprintEventNodeSpawner.cpp
+++ b/Engine/Source/Editor/BlueprintGraph/Private/BlueprintEventNodeSpawner.cpp
@@ -145,7 +145,7 @@ UEdGraphNode* UBlueprintEventNodeSpawner::Invoke(UEdGraph* ParentGraph, FBinding
 		};
 
 		FCustomizeNodeDelegate PostSpawnDelegate = FCustomizeNodeDelegate::CreateStatic(PostSpawnLambda, EventFunc, EventName, CustomizeNodeDelegate);
-		EventNode = Cast<UK2Node_Event>(Super::Invoke(ParentGraph, Bindings, Location, PostSpawnDelegate));
+		EventNode = Super::SpawnNode<UK2Node_Event>(NodeClass, ParentGraph, Bindings, Location, PostSpawnDelegate);
 	}
 	// else, a node for this event already exists, and we should return that 
 	// (the FBlueprintActionMenuItem should detect this and focus in on it).
diff --git a/Engine/Source/Editor/BlueprintGraph/Private/BlueprintFieldNodeSpawner.cpp b/Engine/Source/Editor/BlueprintGraph/Private/BlueprintFieldNodeSpawner.cpp
--- a/Engine/Source/Editor/BlueprintGraph/Private/BlueprintFieldNodeSpawner.cpp
+++ b/Engine/Source/Editor/BlueprintGraph/Private/BlueprintFieldNodeSpawner.cpp
@@ -45,9 +45,7 @@ UEdGraphNode* UBlueprintFieldNodeSpawner::Invoke(UEdGraph* ParentGraph, FBinding
 	};
 
 	FCustomizeNodeDelegate PostSpawnSetupDelegate = FCustomizeNodeDelegate::CreateStatic(PostSpawnSetupLambda, GetField(), SetNodeFieldDelegate, CustomizeNodeDelegate);
-	UEdGraphNode* SpawnedNode = Super::Invoke(ParentGraph, Bindings, Location, PostSpawnSetupDelegate);
-
-	return SpawnedNode;
+	return Super::SpawnNode(NodeClass, ParentGraph, Bindings, Location, PostSpawnSetupDelegate);
 }
 
 //------------------------------------------------------------------------------
diff --git a/Engine/Source/Editor/BlueprintGraph/Private/BlueprintFunctionNodeSpawner.cpp b/Engine/Source/Editor/BlueprintGraph/Private/BlueprintFunctionNodeSpawner.cpp
--- a/Engine/Source/Editor/BlueprintGraph/Private/BlueprintFunctionNodeSpawner.cpp
+++ b/Engine/Source/Editor/BlueprintGraph/Private/BlueprintFunctionNodeSpawner.cpp
@@ -56,8 +56,17 @@ static bool BlueprintFunctionNodeSpawnerImpl::BindFunctionNode(UK2Node_CallFunct
 	{
 		if (UProperty const* BoundProperty = Cast<UProperty>(BoundObject))
 		{
-			UBlueprintNodeSpawner* TempNodeSpawner = UBlueprintVariableNodeSpawner::Create(UK2Node_VariableGet::StaticClass(), BoundProperty);
-			bSuccessfulBinding = BindFunctionNode<UK2Node_VariableGet>(NewNode, TempNodeSpawner);
+			if (UK2Node_CallFunctionOnMember* CallOnMemberNode = Cast<UK2Node_CallFunctionOnMember>(NewNode))
+			{
+				CallOnMemberNode->MemberVariableToCallOn.SetFromField<UProperty>(BoundProperty, NewNode);
+				bSuccessfulBinding = true;
+				CallOnMemberNode->ReconstructNode();
+			}
+			else
+			{
+				UBlueprintNodeSpawner* TempNodeSpawner = UBlueprintVariableNodeSpawner::Create(UK2Node_VariableGet::StaticClass(), BoundProperty);
+				bSuccessfulBinding = BindFunctionNode<UK2Node_VariableGet>(NewNode, TempNodeSpawner);
+			}
 		}
 		else if (AActor* BoundActor = Cast<AActor>(BoundObject))
 		{
@@ -149,9 +158,7 @@ UBlueprintFunctionNodeSpawner* UBlueprintFunctionNodeSpawner::Create(UFunction c
 	{
 		NodeClass = UK2Node_CallDataTableFunction::StaticClass();
 	}
-	// @TODO:
-	//   else if CallOnMember => UK2Node_CallFunctionOnMember
-	//   else if bIsParentContext => UK2Node_CallParentFunction
+	// @TODO: else if bIsParentContext => UK2Node_CallParentFunction
 	else if (bHasArrayPointerParms)
 	{
 		NodeClass = UK2Node_CallArrayFunction::StaticClass();
@@ -216,12 +223,23 @@ void UBlueprintFunctionNodeSpawner::Prime()
 //------------------------------------------------------------------------------
 UEdGraphNode* UBlueprintFunctionNodeSpawner::Invoke(UEdGraph* ParentGraph, FBindingSet const& Bindings, FVector2D const Location) const
 {
+	auto PostSpawnSetupLambda = [](UEdGraphNode* NewNode, bool bIsTemplateNode, UFunction const* Function, FSetNodeFieldDelegate SetFieldDelegate, FCustomizeNodeDelegate UserDelegate)
+	{
+		SetFieldDelegate.ExecuteIfBound(NewNode, Function);
+		UserDelegate.ExecuteIfBound(NewNode, bIsTemplateNode);
+	};
+	FCustomizeNodeDelegate PostSpawnSetupDelegate = FCustomizeNodeDelegate::CreateStatic(PostSpawnSetupLambda, GetFunction(), SetNodeFieldDelegate, CustomizeNodeDelegate);
+
+	UClass* SpawnClass = NodeClass;
+	if ((Bindings.Num() == 1) && (*Bindings.CreateConstIterator())->IsA<UObjectProperty>())
+	{
+		SpawnClass = UK2Node_CallFunctionOnMember::StaticClass();
+	}
 	// if this spawner was set up to spawn a bound node, reset this so the 
 	// bound nodes get positioned properly
 	BlueprintFunctionNodeSpawnerImpl::BindingOffset = FVector2D::ZeroVector;
 
-	UEdGraphNode* SpawnedNode = Super::Invoke(ParentGraph, Bindings, Location);
-	return SpawnedNode;
+	return Super::SpawnNode(SpawnClass, ParentGraph, Bindings, Location, PostSpawnSetupDelegate);
 }
 
 //------------------------------------------------------------------------------
@@ -229,7 +247,19 @@ FText UBlueprintFunctionNodeSpawner::GetDefaultMenuName(FBindingSet const& Bindi
 {
 	UFunction const* Function = GetFunction();
 	check(Function != nullptr);
-	return FText::FromString(UK2Node_CallFunction::GetUserFacingFunctionName(Function));
+	FString FunctionName = UK2Node_CallFunction::GetUserFacingFunctionName(Function);
+
+	if (Bindings.Num() == 1)
+	{
+		UObjectProperty const* ObjectProperty = Cast<UObjectProperty>(Bindings.CreateConstIterator()->Get());
+		if (ObjectProperty != nullptr)
+		{
+			// @TODO: this should probably be an FText::Format()
+			FunctionName = FString::Printf(TEXT("%s (%s)"), *FunctionName, *ObjectProperty->GetName());
+		}
+	}
+
+	return FText::FromString(FunctionName);
 }
 
 //------------------------------------------------------------------------------
@@ -285,29 +315,22 @@ bool UBlueprintFunctionNodeSpawner::CanBindMultipleObjects() const
 bool UBlueprintFunctionNodeSpawner::IsBindingCompatible(UObject const* BindingCandidate) const
 {
 	bool bCanBind = false;
-	if (UFunction const* Function = GetFunction())
+
+	UFunction const* Function = GetFunction();
+	if ((Function != nullptr) && (NodeClass == UK2Node_CallFunction::StaticClass()))
 	{
-		// @TODO: don't know exactly why we can only bind non-pure/const 
-		//        functions... this is mirrored after FK2ActionMenuBuilder::GetFunctionCallsOnSelectedActors()
-		//        and FK2ActionMenuBuilder::GetFunctionCallsOnSelectedComponents(),
-		//        where we make the same stipulation
-		bool const bIsImperative = !Function->HasAnyFunctionFlags(FUNC_BlueprintPure);
-		bool const bIsConstFunc  =  Function->HasAnyFunctionFlags(FUNC_Const);
-
-		if (bIsImperative || bIsConstFunc)
+		UClass* BindingClass = BindingCandidate->GetClass();
+		if (UObjectProperty const* ObjectProperty = Cast<UObjectProperty>(BindingCandidate))
 		{
-			UClass* BindingClass = BindingCandidate->GetClass();
-			if (UObjectProperty const* ObjectProperty = Cast<UObjectProperty>(BindingCandidate))
-			{
-				BindingClass = ObjectProperty->PropertyClass;
-			}
+			BindingClass = ObjectProperty->PropertyClass;
+		}
 
-			if (UClass const* FuncOwner = Function->GetOwnerClass())
-			{
-				bCanBind = BindingClass->IsChildOf(FuncOwner);
-			}
+		if (UClass const* FuncOwner = Function->GetOwnerClass())
+		{
+			bCanBind = BindingClass->IsChildOf(FuncOwner);
 		}
 	}
+
 	return bCanBind;
 }
 
diff --git a/Engine/Source/Editor/BlueprintGraph/Private/BlueprintNodeSpawner.cpp b/Engine/Source/Editor/BlueprintGraph/Private/BlueprintNodeSpawner.cpp
--- a/Engine/Source/Editor/BlueprintGraph/Private/BlueprintNodeSpawner.cpp
+++ b/Engine/Source/Editor/BlueprintGraph/Private/BlueprintNodeSpawner.cpp
@@ -119,7 +119,7 @@ FBlueprintNodeSignature UBlueprintNodeSpawner::GetSpawnerSignature() const
 //------------------------------------------------------------------------------
 UEdGraphNode* UBlueprintNodeSpawner::Invoke(UEdGraph* ParentGraph, FBindingSet const& Bindings, FVector2D const Location) const
 {
-	return Invoke(ParentGraph, Bindings, Location, CustomizeNodeDelegate);
+	return SpawnNode(NodeClass, ParentGraph, Bindings, Location, CustomizeNodeDelegate);
 }
 
 //------------------------------------------------------------------------------
@@ -187,14 +187,20 @@ UEdGraphNode* UBlueprintNodeSpawner::GetTemplateNode(ENoInit) const
 {
 	return BlueprintNodeSpawnerImpl::GetSharedTemplateCache()->GetNodeTemplate(this, NoInit);
 }
+// 
+// //------------------------------------------------------------------------------
+// UEdGraphNode* UBlueprintNodeSpawner::Invoke(UEdGraph* ParentGraph, FBindingSet const& Bindings, FVector2D const Location, FCustomizeNodeDelegate PostSpawnDelegate) const
+// {
+// 	return SpawnNode(NodeClass, ParentGraph, Bindings, Location, PostSpawnDelegate);
+// }
 
 //------------------------------------------------------------------------------
-UEdGraphNode* UBlueprintNodeSpawner::Invoke(UEdGraph* ParentGraph, FBindingSet const& Bindings, FVector2D const Location, FCustomizeNodeDelegate PostSpawnDelegate) const
+UEdGraphNode* UBlueprintNodeSpawner::SpawnEdGraphNode(TSubclassOf<UEdGraphNode> InNodeClass, UEdGraph* ParentGraph, FBindingSet const& Bindings, FVector2D const Location, FCustomizeNodeDelegate PostSpawnDelegate) const
 {
 	UEdGraphNode* NewNode = nullptr;
-	if (NodeClass != nullptr)
+	if (InNodeClass != nullptr)
 	{
-		NewNode = NewObject<UEdGraphNode>(ParentGraph, NodeClass);
+		NewNode = NewObject<UEdGraphNode>(ParentGraph, InNodeClass);
 		check(NewNode != nullptr);
 		NewNode->CreateNewGuid();
 
diff --git a/Engine/Source/Editor/BlueprintGraph/Private/BlueprintVariableNodeSpawner.cpp b/Engine/Source/Editor/BlueprintGraph/Private/BlueprintVariableNodeSpawner.cpp
--- a/Engine/Source/Editor/BlueprintGraph/Private/BlueprintVariableNodeSpawner.cpp
+++ b/Engine/Source/Editor/BlueprintGraph/Private/BlueprintVariableNodeSpawner.cpp
@@ -119,7 +119,7 @@ UEdGraphNode* UBlueprintVariableNodeSpawner::Invoke(UEdGraph* ParentGraph, FBind
 			PostSpawnDelegate = FCustomizeNodeDelegate::CreateStatic(LocalVarSetupLambda, LocalVarDesc.VarName, LocalVarOuter, LocalVarDesc.VarGuid, CustomizeNodeDelegate);
 		}
 
-		NewNode = UBlueprintNodeSpawner::Invoke(ParentGraph, Bindings, Location, PostSpawnDelegate);
+		NewNode = UBlueprintNodeSpawner::SpawnNode(NodeClass, ParentGraph, Bindings, Location, PostSpawnDelegate);
 	}
 	else
 	{
diff --git a/Engine/Source/Editor/BlueprintGraph/Private/K2Node_CallFunctionOnMember.cpp b/Engine/Source/Editor/BlueprintGraph/Private/K2Node_CallFunctionOnMember.cpp
--- a/Engine/Source/Editor/BlueprintGraph/Private/K2Node_CallFunctionOnMember.cpp
+++ b/Engine/Source/Editor/BlueprintGraph/Private/K2Node_CallFunctionOnMember.cpp
@@ -65,31 +65,43 @@ void UK2Node_CallFunctionOnMember::ExpandNode(class FKismetCompilerContext& Comp
 		UEdGraphPin* SelfPin = Schema->FindSelfPin(*this, EGPD_Input);
 		if(SelfPin != NULL)
 		{
-			for(int32 TargetIdx = 0; TargetIdx < SelfPin->LinkedTo.Num(); TargetIdx++)
+			if (SelfPin->LinkedTo.Num() == 0)
 			{
-				UEdGraphPin* SourcePin = SelfPin->LinkedTo[TargetIdx];
-				if(SourcePin != NULL)
-				{
-					// Create 'get var' node to get the member
-					UK2Node_VariableGet* GetVarNode = CompilerContext.SpawnIntermediateNode<UK2Node_VariableGet>(this, SourceGraph);
-					GetVarNode->VariableReference = MemberVariableToCallOn;
-					GetVarNode->AllocateDefaultPins();
-
-					UEdGraphPin* VarNodeSelfPin = Schema->FindSelfPin(*GetVarNode, EGPD_Input);
-					if( VarNodeSelfPin != NULL )
-					{
-						VarNodeSelfPin->MakeLinkTo(SourcePin);
+				UK2Node_VariableGet* GetVarNode = CompilerContext.SpawnIntermediateNode<UK2Node_VariableGet>(this, SourceGraph);
+				GetVarNode->VariableReference.SetSelfMember(MemberVariableToCallOn.GetMemberName());
+				GetVarNode->AllocateDefaultPins();
 
-						UEdGraphPin* ValuePin = GetVarNode->GetValuePin();
-						ValuePin->MakeLinkTo(CallFuncSelfPin);
-					}
-					else
+				UEdGraphPin* ValuePin = GetVarNode->GetValuePin();
+				ValuePin->MakeLinkTo(CallFuncSelfPin);
+			}
+			else
+			{
+				for (int32 TargetIdx = 0; TargetIdx < SelfPin->LinkedTo.Num(); TargetIdx++)
+				{
+					UEdGraphPin* SourcePin = SelfPin->LinkedTo[TargetIdx];
+					if (SourcePin != NULL)
 					{
-						// Failed to find the member to call on for this expansion, so warn about it
-						CompilerContext.MessageLog.Warning(*LOCTEXT("CallFunctionOnInvalidMember_Warning", "Function node @@ called on invalid target member.").ToString(), this);
+						// Create 'get var' node to get the member
+						UK2Node_VariableGet* GetVarNode = CompilerContext.SpawnIntermediateNode<UK2Node_VariableGet>(this, SourceGraph);
+						GetVarNode->VariableReference = MemberVariableToCallOn;
+						GetVarNode->AllocateDefaultPins();
+
+						UEdGraphPin* VarNodeSelfPin = Schema->FindSelfPin(*GetVarNode, EGPD_Input);
+						if (VarNodeSelfPin != NULL)
+						{
+							VarNodeSelfPin->MakeLinkTo(SourcePin);
+
+							UEdGraphPin* ValuePin = GetVarNode->GetValuePin();
+							ValuePin->MakeLinkTo(CallFuncSelfPin);
+						}
+						else
+						{
+							// Failed to find the member to call on for this expansion, so warn about it
+							CompilerContext.MessageLog.Warning(*LOCTEXT("CallFunctionOnInvalidMember_Warning", "Function node @@ called on invalid target member.").ToString(), this);
+						}
 					}
 				}
-			}
+			}			
 		}
 
 		// Now move the rest of the connections (including exec connections...)
diff --git a/Engine/Source/Editor/BlueprintGraph/Public/BlueprintNodeSpawner.h b/Engine/Source/Editor/BlueprintGraph/Public/BlueprintNodeSpawner.h
--- a/Engine/Source/Editor/BlueprintGraph/Public/BlueprintNodeSpawner.h
+++ b/Engine/Source/Editor/BlueprintGraph/Public/BlueprintNodeSpawner.h
@@ -187,24 +187,54 @@ public:
 	UEdGraphNode* GetTemplateNode(ENoInit) const;
 
 	// IBlueprintNodeBinder interface
-	virtual bool   IsBindingCompatible(UObject const* BindingCandidate) const override { return false; }
-	virtual bool   CanBindMultipleObjects() const override { return false; }
-
+	virtual bool IsBindingCompatible(UObject const* BindingCandidate) const override { return false; }
+	virtual bool CanBindMultipleObjects() const override { return false; }
 protected:
-	virtual bool   BindToNode(UEdGraphNode* Node, UObject* Binding) const override { return false; }
+	virtual bool BindToNode(UEdGraphNode* Node, UObject* Binding) const override { return false; }
 	// End IBlueprintNodeBinder interface
 
 	/**
-	 * Protected Invoke() that let's sub-classes specify their own post-spawn
-	 * delegate. Creates a new node based off of the set NodeClass.
+	 * Protected SpawnNode() that let's sub-classes specify their own post-spawn
+	 * delegate and node class.
 	 * 
+	 * @param  NodeClass			The type of node you want spawned.
 	 * @param  ParentGraph			The graph you want the node spawned into.
-	 * @param  Bindings				
+	 * @param  Bindings				The bindings to apply to the node (post spawn).
 	 * @param  Location				Where you want the new node positioned in the graph.
 	 * @param  PostSpawnDelegate	A delegate to run after spawning the node, but prior to allocating the node's pins.
 	 * @return Null if it failed to spawn a node (if NodeClass is null), otherwise a newly spawned node.
 	 */
-	UEdGraphNode* Invoke(UEdGraph* ParentGraph, FBindingSet const& Bindings, FVector2D const Location, FCustomizeNodeDelegate PostSpawnDelegate) const;
+	template <class NodeType = UEdGraphNode>
+	NodeType* SpawnNode(TSubclassOf<UEdGraphNode> NodeClass, UEdGraph* ParentGraph, FBindingSet const& Bindings, FVector2D const Location, FCustomizeNodeDelegate PostSpawnDelegate) const;
+
+	/**
+	 * Simplified version of the other SpawnNode(), that just let's sub-classes 
+	 * specify their own post-spawn delegate (the node class is comes directly 
+	 * from the template parameter).
+	 * 
+	 * @param  ParentGraph			The graph you want the node spawned into.
+	 * @param  Bindings				The bindings to apply to the node (post spawn).
+	 * @param  Location				Where you want the new node positioned in the graph.
+	 * @param  PostSpawnDelegate	A delegate to run after spawning the node, but prior to allocating the node's pins.
+	 * @return Null if it failed to spawn a node (if NodeClass is null), otherwise a newly spawned node.
+	 */
+	template <class NodeType>
+	NodeType* SpawnNode(UEdGraph* ParentGraph, FBindingSet const& Bindings, FVector2D const Location, FCustomizeNodeDelegate PostSpawnDelegate) const;
+
+private:
+	/**
+	 * Does the actual node spawning. Creates a new node (of the specified type),
+	 * sets the node's position, calls PostSpawnDelegate on the new node, and
+	 * finally applies any bindings that were passed to it.
+	 * 
+	 * @param  NodeClass			The type of node you want spawned.
+	 * @param  ParentGraph			The graph you want the node spawned into.
+	 * @param  Bindings				The bindings to apply to the node (post spawn).
+	 * @param  Location				Where you want the new node positioned in the graph.
+	 * @param  PostSpawnDelegate	A delegate to run after spawning the node, but prior to allocating the node's pins.
+	 * @return Null if it failed to spawn a node (if NodeClass is null), otherwise a newly spawned node.
+	 */
+	UEdGraphNode* SpawnEdGraphNode(TSubclassOf<UEdGraphNode> NodeClass, UEdGraph* ParentGraph, FBindingSet const& Bindings, FVector2D const Location, FCustomizeNodeDelegate PostSpawnDelegate) const;
 };
 
 /*******************************************************************************
@@ -217,3 +247,17 @@ UBlueprintNodeSpawner* UBlueprintNodeSpawner::Create(UObject* Outer, FCustomizeN
 {
 	return Create(NodeType::StaticClass(), Outer, PostSpawnDelegate);
 }
+
+//------------------------------------------------------------------------------
+template <class NodeType>
+NodeType* UBlueprintNodeSpawner::SpawnNode(TSubclassOf<UEdGraphNode> NodeClass, UEdGraph* ParentGraph, FBindingSet const& Bindings, FVector2D const Location, FCustomizeNodeDelegate PostSpawnDelegate) const
+{
+	return CastChecked<NodeType>(SpawnEdGraphNode(NodeClass, ParentGraph, Bindings, Location, PostSpawnDelegate));
+}
+
+//------------------------------------------------------------------------------
+template <class NodeType>
+NodeType* UBlueprintNodeSpawner::SpawnNode(UEdGraph* ParentGraph, FBindingSet const& Bindings, FVector2D const Location, FCustomizeNodeDelegate PostSpawnDelegate) const
+{
+	return SpawnNode<NodeType>(NodeType::StaticClass(), ParentGraph, Bindings, Location, PostSpawnDelegate);
+}
diff --git a/Engine/Source/Editor/Kismet/Private/BlueprintActionMenuBuilder.cpp b/Engine/Source/Editor/Kismet/Private/BlueprintActionMenuBuilder.cpp
--- a/Engine/Source/Editor/Kismet/Private/BlueprintActionMenuBuilder.cpp
+++ b/Engine/Source/Editor/Kismet/Private/BlueprintActionMenuBuilder.cpp
@@ -184,16 +184,19 @@ TSharedPtr<FBlueprintBoundMenuItem> FBlueprintActionMenuItemFactory::MakeBoundMe
 {
 	IBlueprintNodeBinder::FBindingSet const& Bindings = ActionInfo.GetBindings();
 	UBlueprintNodeSpawner const* Action = ActionInfo.NodeSpawner;
+
 	FBlueprintBoundMenuItem* NewMenuItem = new FBlueprintBoundMenuItem(Action, MenuGrouping);
-	NewMenuItem->MenuDescription = GetMenuNameForAction(EditorContext, ActionInfo);
 
+	// AddBindings() updates the MenuDescription, everytime a binding is added, 
+	// so set the default menu name before it
+	NewMenuItem->MenuDescription = GetMenuNameForAction(EditorContext, ActionInfo);
 	NewMenuItem->AddBindings(Bindings);
 
-	// FBlueprintBoundMenuItem takes care of its own menu category
-	NewMenuItem->Category = FString::Printf(TEXT("%s|%s"), *RootCategory.ToString(), *NewMenuItem->Category);
 	NewMenuItem->TooltipDescription = GetTooltipForAction(EditorContext, Action).ToString();
-	NewMenuItem->Keywords = GetSearchKeywordsForAction(EditorContext, Action);
+	NewMenuItem->Category           = GetCategoryForAction(EditorContext, Action).ToString();
+	NewMenuItem->Keywords           = GetSearchKeywordsForAction(EditorContext, Action);
 
+	NewMenuItem->Category = FString::Printf(TEXT("%s|%s"), *RootCategory.ToString(), *NewMenuItem->Category);
 	return MakeShareable(NewMenuItem);
 }
 
@@ -360,6 +363,8 @@ UEdGraphNode* FBlueprintActionMenuItemFactory::GetTemplateNode(UBlueprintNodeSpa
 
 namespace FBlueprintActionMenuBuilderImpl
 {
+	typedef TArray< TSharedPtr<FEdGraphSchemaAction> > MenuItemList;
+
 	/** Defines a sub-section of the overall blueprint menu (filter, heading, etc.) */
 	struct FMenuSectionDefinition
 	{
@@ -388,7 +393,7 @@ namespace FBlueprintActionMenuBuilderImpl
 		 * @param  DatabaseAction	The node-spawner that the new menu item should wrap.
 		 * @return An empty TSharedPtr if the action was filtered out, otherwise a newly allocated FBlueprintActionMenuItem.
 		 */
-		TSharedPtr<FEdGraphSchemaAction> MakeMenuItem(TWeakPtr<FBlueprintEditor> EditorContext, FBlueprintActionInfo& DatabaseAction);
+		MenuItemList MakeMenuItems(TWeakPtr<FBlueprintEditor> EditorContext, FBlueprintActionInfo& DatabaseAction);
 
 		/**
 		 * 
@@ -398,7 +403,7 @@ namespace FBlueprintActionMenuBuilderImpl
 		 * @param  Bindings	
 		 * @return 
 		 */
-		TSharedPtr<FEdGraphSchemaAction> MakeBoundMenuItem(TWeakPtr<FBlueprintEditor> EditorContext, FBlueprintActionInfo& DatabaseAction, TArray<UObject*> const& Bindings);
+		void AddBoundMenuItems(TWeakPtr<FBlueprintEditor> EditorContext, FBlueprintActionInfo& DatabaseAction, TArray<UObject*> const& Bindings, MenuItemList& MenuItemsOut);
 		
 		/**
 		 * Clears out any consolidated properties that this may have been 
@@ -423,6 +428,14 @@ namespace FBlueprintActionMenuBuilderImpl
 	 * @param  MenuOut  		The menu builder we want all the legacy actions appended to.
 	 */
 	static void AppendLegacyItems(FMenuSectionDefinition const& MenuDef, TWeakPtr<FBlueprintEditor> BlueprintEditor, FBlueprintActionMenuBuilder& MenuOut);
+
+	/**
+	 * 
+	 * 
+	 * @param  Context	
+	 * @return 
+	 */
+	static TArray<UObject*> GetBindingCandidates(FBlueprintActionContext const& Context);
 }
 
 //------------------------------------------------------------------------------
@@ -452,12 +465,14 @@ void FBlueprintActionMenuBuilderImpl::FMenuSectionDefinition::SetSectionSortOrde
 }
 // 
 //------------------------------------------------------------------------------
-TSharedPtr<FEdGraphSchemaAction> FBlueprintActionMenuBuilderImpl::FMenuSectionDefinition::MakeBoundMenuItem(TWeakPtr<FBlueprintEditor> EditorContext, FBlueprintActionInfo& DatabaseActionInfo, TArray<UObject*> const& PerspectiveBindings)
+void FBlueprintActionMenuBuilderImpl::FMenuSectionDefinition::AddBoundMenuItems(TWeakPtr<FBlueprintEditor> EditorContext, FBlueprintActionInfo& DatabaseActionInfo, TArray<UObject*> const& PerspectiveBindings, MenuItemList& MenuItemsOut)
 {
-	TSharedPtr<FBlueprintBoundMenuItem> MenuItem;
 	UBlueprintNodeSpawner const* DatabaseAction = DatabaseActionInfo.NodeSpawner;
+
+	TSharedPtr<FBlueprintBoundMenuItem> LastMadeMenuItem;
+	bool const bConsolidate = (Flags & FBlueprintActionMenuBuilder::ConsolidateBoundActions) != 0;
 	
-	IBlueprintNodeBinder::FBindingSet CompatibleBindingsToFilter;
+	IBlueprintNodeBinder::FBindingSet CompatibleBindings;
 	// we don't want the blueprint database growing out of control with an entry 
 	// for every object you could ever possibly bind to, so each 
 	// UBlueprintNodeSpawner comes with an interface to test/bind through... 
@@ -471,51 +486,56 @@ TSharedPtr<FEdGraphSchemaAction> FBlueprintActionMenuBuilderImpl::FMenuSectionDe
 		if (DatabaseAction->IsBindingCompatible(BindingObj))
 		{
 			// add bindings before filtering (in case tests accept/reject based off of this)
-			CompatibleBindingsToFilter.Add(BindingObj);
-
-			// if BoundAction is now "full" (meaning it can take any more 
-			// bindings), or if this is the last binding to test...
-			if (!DatabaseAction->CanBindMultipleObjects() || bIsLastBinding)
+			CompatibleBindings.Add(BindingObj);
+		}
+// 		else if (bConsolidate)
+// 		{
+// 			MadeMenuItems.Empty();
+// 			LastMadeMenuItem.Reset();
+// 			break;
+// 		}
+
+		// if BoundAction is now "full" (meaning it can take any more 
+		// bindings), or if this is the last binding to test...
+		if ((CompatibleBindings.Num() > 0) && (!DatabaseAction->CanBindMultipleObjects() || bIsLastBinding || !bConsolidate))
+		{
+			// we don't want binding to mutate DatabaseActionInfo, so we clone  
+			// the action info, and tack on some binding data
+			FBlueprintActionInfo BoundActionInfo(DatabaseActionInfo, CompatibleBindings);
+
+			// have to check IsFiltered() for every "fully bound" action (in
+			// case there are tests that reject based off of this), we may 
+			// test this multiple times per action (we have to make sure 
+			// that every set of bound objects pass before folding them into
+			// MenuItem)
+			bool const bPassedFilter = !Filter.IsFiltered(BoundActionInfo);
+			if (bPassedFilter)
 			{
-				// we don't want binding to mutate DatabaseActionInfo, so we clone  
-				// the action info, and tack on some binding data
-				FBlueprintActionInfo BoundActionInfo(DatabaseActionInfo, CompatibleBindingsToFilter);
-
-				// have to check IsFiltered() for every "fully bound" action (in
-				// case there are tests that reject based off of this), we may 
-				// test this multiple times per action (we have to make sure 
-				// that every set of bound objects pass before folding them into
-				// MenuItem)
-				bool const bPassedFilter = !Filter.IsFiltered(BoundActionInfo);
-				if (bPassedFilter)
+				if (!bConsolidate || !LastMadeMenuItem.IsValid())
 				{
-					if (!MenuItem.IsValid())
-					{
-						MenuItem = ItemFactory.MakeBoundMenuItem(EditorContext, BoundActionInfo);
-					}
-					else
-					{
-						// move these bindings over to the menu item (so we can 
-						// test the next set)
-						MenuItem->AddBindings(CompatibleBindingsToFilter);
-					}
+					LastMadeMenuItem = ItemFactory.MakeBoundMenuItem(EditorContext, BoundActionInfo);
+					MenuItemsOut.Add(LastMadeMenuItem);
+				}
+				else
+				{
+					// move these bindings over to the menu item (so we can 
+					// test the next set)
+					LastMadeMenuItem->AddBindings(CompatibleBindings);
 				}
-				CompatibleBindingsToFilter.Empty(); // do before we copy back over cached fields for DatabaseActionInfo
-
-				// copy over any fields that got cached for filtering (with
-				// an empty binding set)
-				/*DatabaseActionInfo = FBlueprintActionInfo(BoundActionInfo, CompatibleBindingsToFilter);*/
 			}
+			CompatibleBindings.Empty(); // do before we copy back over cached fields for DatabaseActionInfo
+
+			// copy over any fields that got cached for filtering (with
+			// an empty binding set)
+			/*DatabaseActionInfo = FBlueprintActionInfo(BoundActionInfo, CompatibleBindings);*/
 		}
 	}
-
-	return MenuItem;
 }
 
 //------------------------------------------------------------------------------
-TSharedPtr<FEdGraphSchemaAction> FBlueprintActionMenuBuilderImpl::FMenuSectionDefinition::MakeMenuItem(TWeakPtr<FBlueprintEditor> EditorContext, FBlueprintActionInfo& DatabaseAction)
+FBlueprintActionMenuBuilderImpl::MenuItemList FBlueprintActionMenuBuilderImpl::FMenuSectionDefinition::MakeMenuItems(TWeakPtr<FBlueprintEditor> EditorContext, FBlueprintActionInfo& DatabaseAction)
 {	
-	TSharedPtr<FEdGraphSchemaAction> MenuEntry;
+	TSharedPtr<FEdGraphSchemaAction> UnBoundMenuEntry;
 	bool bPassedFilter = !Filter.IsFiltered(DatabaseAction);
 
 	// if the caller wants to consolidate all property actions, then we have to 
@@ -548,18 +568,24 @@ TSharedPtr<FEdGraphSchemaAction> FBlueprintActionMenuBuilderImpl::FMenuSectionDe
 			{
 				TSharedPtr<FBlueprintDragDropMenuItem> NewMenuItem = ItemFactory.MakeDragDropMenuItem(DatabaseAction.NodeSpawner);
 				ConsolidatedProperties.Add(ActionProperty, NewMenuItem);
-				MenuEntry = NewMenuItem;
+				UnBoundMenuEntry = NewMenuItem;
 			}
 		}
 	}
 
-	
-	if (!MenuEntry.IsValid() && bPassedFilter)
+	if (!UnBoundMenuEntry.IsValid() && bPassedFilter)
 	{
-		MenuEntry = ItemFactory.MakeActionMenuItem(EditorContext, DatabaseAction);
+		UnBoundMenuEntry = ItemFactory.MakeActionMenuItem(EditorContext, DatabaseAction);
 	}
 
-	return MenuEntry;
+	FBlueprintActionMenuBuilderImpl::MenuItemList MenuItems;
+	if (UnBoundMenuEntry.IsValid())
+	{
+		MenuItems.Add(UnBoundMenuEntry);
+	}
+	AddBoundMenuItems(EditorContext, DatabaseAction, GetBindingCandidates(Filter.Context), MenuItems);
+
+	return MenuItems;
 }
 
 //------------------------------------------------------------------------------
@@ -624,6 +650,12 @@ void FBlueprintActionMenuBuilderImpl::AppendLegacyItems(FMenuSectionDefinition c
 	}
 }
 
+//------------------------------------------------------------------------------
+static TArray<UObject*> FBlueprintActionMenuBuilderImpl::GetBindingCandidates(FBlueprintActionContext const& Context)
+{
+	return Context.SelectedObjects;
+}
+
 /*******************************************************************************
  * FBlueprintActionMenuBuilder
  ******************************************************************************/
@@ -681,21 +713,10 @@ void FBlueprintActionMenuBuilder::RebuildActionList()
 				//        all MenuSctions (to pair down on performance?)
 				for (TSharedRef<FMenuSectionDefinition> const& MenuSection : MenuSections)
 				{
-					TSharedPtr<FEdGraphSchemaAction> MenuEntry = MenuSection->MakeMenuItem(BlueprintEditorPtr, BlueprintAction);
-					if (MenuEntry.IsValid())
+					for (TSharedPtr<FEdGraphSchemaAction> MenuEntry : MenuSection->MakeMenuItems(BlueprintEditorPtr, BlueprintAction))
 					{
 						AddAction(MenuEntry);
 					}
-
-					TArray<UObject*> const& BindingObjects = MenuSection->Filter.Context.SelectedObjects;
-					// @TODO: this only makes one "BoundMenuEntry" for all 
-					//        bindings, we could have a non-consolidated mode if
-					//        we want (like we have for property actions)
-					TSharedPtr<FEdGraphSchemaAction> BoundMenuEntry = MenuSection->MakeBoundMenuItem(BlueprintEditorPtr, BlueprintAction, BindingObjects);
-					if (BoundMenuEntry.IsValid())
-					{
-						AddAction(BoundMenuEntry);
-					}
 				}
 			}
 		}	
diff --git a/Engine/Source/Editor/Kismet/Private/BlueprintActionMenuUtils.cpp b/Engine/Source/Editor/Kismet/Private/BlueprintActionMenuUtils.cpp
--- a/Engine/Source/Editor/Kismet/Private/BlueprintActionMenuUtils.cpp
+++ b/Engine/Source/Editor/Kismet/Private/BlueprintActionMenuUtils.cpp
@@ -25,6 +25,10 @@
 
 namespace BlueprintActionMenuUtilsImpl
 {
+	static int32 const MainMenuSectionGroup   = 000;
+	static int32 const ComponentsSectionGroup = 100;
+	static int32 const LevelActorSectionGroup = 101;
+
 	/**
 	 * Additional filter rejection test, for menu sections that only contain 
 	 * bound actions. Rejects any action that is not bound.
@@ -45,13 +49,43 @@ namespace BlueprintActionMenuUtilsImpl
 	 */
 	static bool IsNonFavoritedAction(FBlueprintActionFilter const& Filter, FBlueprintActionInfo& BlueprintAction);
 
+	static bool IsPureNonConstAction(FBlueprintActionFilter const& Filter, FBlueprintActionInfo& BlueprintAction);
+
+	static bool IsUnexposedMemberAction(FBlueprintActionFilter const& Filter, FBlueprintActionInfo& BlueprintAction);
+
 	/**
 	 * Utility function to find a common base class from a set of given classes.
 	 * 
-	 * @param  ClassSet	The set of classes that you want a single base class for.
-	 * @return A common base class for all the specified classes (fallsback to UObject's class).
+	 * @param  ObjectSet	The set of objects that you want a single base class for.
+	 * @return A common base class for all the specified classes (falls back to UObject's class).
+	 */
+	static UClass* FindCommonBaseClass(TArray<UObject*> const& ObjectSet);
+
+	/**
+	 * 
+	 * 
+	 * @param  Pin	
+	 * @return 
+	 */
+	static UClass* GetPinClassType(UEdGraphPin const* Pin);
+
+	/**
+	 * 
+	 * 
+	 * @param  ComponentsFilter	
+	 * @param  MenuOut	
+	 * @return 
 	 */
-	static UClass* FindCommonBaseClass(TArray<UClass*> const& ClassSet);
+	static void AddComponentSections(FBlueprintActionFilter const& ComponentsFilter, FBlueprintActionMenuBuilder& MenuOut);
+
+	/**
+	 * 
+	 * 
+	 * @param  LevelActorsFilter	
+	 * @param  MenuOut	
+	 * @return 
+	 */
+	static void AddLevelActorSections(FBlueprintActionFilter const& LevelActorsFilter, FBlueprintActionMenuBuilder& MenuOut);
 }
 
 //------------------------------------------------------------------------------
@@ -72,14 +106,59 @@ static bool BlueprintActionMenuUtilsImpl::IsNonFavoritedAction(FBlueprintActionF
 }
 
 //------------------------------------------------------------------------------
-static UClass* BlueprintActionMenuUtilsImpl::FindCommonBaseClass(TArray<UClass*> const& ClassSet)
+static bool BlueprintActionMenuUtilsImpl::IsPureNonConstAction(FBlueprintActionFilter const& Filter, FBlueprintActionInfo& BlueprintAction)
+{
+	bool bIsFliteredOut = false;
+
+	if (UFunction const* Function = BlueprintAction.GetAssociatedFunction())
+	{
+		bool const bIsImperative = !Function->HasAnyFunctionFlags(FUNC_BlueprintPure);
+		bool const bIsConstFunc  =  Function->HasAnyFunctionFlags(FUNC_Const);
+		bIsFliteredOut = !bIsImperative && !bIsConstFunc;
+	}
+	return bIsFliteredOut;
+}
+
+//------------------------------------------------------------------------------
+static bool BlueprintActionMenuUtilsImpl::IsUnexposedMemberAction(FBlueprintActionFilter const& Filter, FBlueprintActionInfo& BlueprintAction)
+{
+	bool bIsFliteredOut = false;
+
+	if (UFunction const* Function = BlueprintAction.GetAssociatedFunction())
+	{
+		TArray<FString> AllExposedCategories;
+		for (TWeakObjectPtr<UObject> Binding : BlueprintAction.GetBindings())
+		{
+			if (UObjectProperty* ObjectProperty = Cast<UObjectProperty>(Binding.Get()))
+			{
+				FString const ExposedCategoryMetadata = ObjectProperty->GetMetaData(FBlueprintMetadata::MD_ExposeFunctionCategories);
+				if (ExposedCategoryMetadata.IsEmpty())
+				{
+					continue;
+				}
+
+				TArray<FString> PropertyExposedCategories;
+				ExposedCategoryMetadata.ParseIntoArray(&PropertyExposedCategories, TEXT(","), true);
+				AllExposedCategories.Append(PropertyExposedCategories);
+			}
+		}
+
+		FString FunctionCategory = Function->GetMetaData(FBlueprintMetadata::MD_FunctionCategory);
+		bIsFliteredOut = !AllExposedCategories.Contains(FunctionCategory);
+	}
+	return bIsFliteredOut;
+}
+
+//------------------------------------------------------------------------------
+static UClass* BlueprintActionMenuUtilsImpl::FindCommonBaseClass(TArray<UObject*> const& ObjectSet)
 {
 	UClass* CommonClass = UObject::StaticClass();
-	if (ClassSet.Num() > 0)
+	if (ObjectSet.Num() > 0)
 	{
-		CommonClass = ClassSet[0];
-		for (UClass const* Class : ClassSet)
+		CommonClass = ObjectSet[0]->GetClass();
+		for (UObject const* Object : ObjectSet)
 		{
+			UClass* Class = Object->GetClass();
 			while (!Class->IsChildOf(CommonClass))
 			{
 				CommonClass = CommonClass->GetSuperClass();
@@ -89,6 +168,78 @@ static UClass* BlueprintActionMenuUtilsImpl::FindCommonBaseClass(TArray<UClass*>
 	return CommonClass;
 }
 
+//------------------------------------------------------------------------------
+static UClass* BlueprintActionMenuUtilsImpl::GetPinClassType(UEdGraphPin const* Pin)
+{
+	UClass* PinObjClass = nullptr;
+
+	FEdGraphPinType const& PinType = Pin->PinType;
+	if ((PinType.PinCategory == UEdGraphSchema_K2::PC_Object) || 
+		(PinType.PinCategory == UEdGraphSchema_K2::PC_Interface))
+	{
+		bool const bIsSelfPin = !PinType.PinSubCategoryObject.IsValid();
+		if (bIsSelfPin)
+		{
+			PinObjClass = CastChecked<UK2Node>(Pin->GetOwningNode())->GetBlueprint()->SkeletonGeneratedClass;
+		}
+		else
+		{
+			PinObjClass = Cast<UClass>(PinType.PinSubCategoryObject.Get());
+		}
+	}
+
+	return PinObjClass;
+}
+
+//------------------------------------------------------------------------------
+static void BlueprintActionMenuUtilsImpl::AddComponentSections(FBlueprintActionFilter const& ComponentsFilter, FBlueprintActionMenuBuilder& MenuOut)
+{
+	FText EventSectionHeading = LOCTEXT("ComponentsEventCategory", "Add Event for Selected Components");
+	FText FuncSectionHeading  = LOCTEXT("ComponentsFuncCategory",  "Call Function on Selected Components");
+
+	if (ComponentsFilter.Context.SelectedObjects.Num() == 1)
+	{
+		FText const ComponentName = FText::FromName(ComponentsFilter.Context.SelectedObjects.Last()->GetFName());
+		FuncSectionHeading  = FText::Format(LOCTEXT("SingleComponentFuncCategory", "Call Function on {0}"), ComponentName);
+		EventSectionHeading = FText::Format(LOCTEXT("SingleComponentFuncCategory", "Add Event for {0}"), ComponentName);
+	}
+
+	FBlueprintActionFilter ComponentFunctionsFilter = ComponentsFilter;
+	ComponentFunctionsFilter.PermittedNodeTypes.Add(UK2Node_CallFunction::StaticClass());
+	MenuOut.AddMenuSection(ComponentFunctionsFilter, FuncSectionHeading, ComponentsSectionGroup, FBlueprintActionMenuBuilder::ConsolidateBoundActions);
+
+	FBlueprintActionFilter ComponentEventsFilter = ComponentsFilter;
+	ComponentEventsFilter.PermittedNodeTypes.Add(UK2Node_ComponentBoundEvent::StaticClass());
+	MenuOut.AddMenuSection(ComponentEventsFilter, EventSectionHeading, ComponentsSectionGroup, FBlueprintActionMenuBuilder::ConsolidateBoundActions);
+}
+
+//------------------------------------------------------------------------------
+static void BlueprintActionMenuUtilsImpl::AddLevelActorSections(FBlueprintActionFilter const& LevelActorsFilter, FBlueprintActionMenuBuilder& MenuOut)
+{
+	FText EventSectionHeading = LOCTEXT("ActorsEventCategory", "Add Event for Selected Actors");
+	FText FuncSectionHeading  = LOCTEXT("ActorsFuncCategory",  "Call Function on Selected Actors");
+
+	if (LevelActorsFilter.Context.SelectedObjects.Num() == 1)
+	{
+		FText const ActorName = FText::FromName(LevelActorsFilter.Context.SelectedObjects.Last()->GetFName());
+		FuncSectionHeading  = FText::Format(LOCTEXT("SingleActorFuncCategory", "Call Function on {0}"), ActorName);
+		EventSectionHeading = FText::Format(LOCTEXT("SingleActorFuncCategory", "Add Event for {0}"), ActorName);
+	}
+
+	FBlueprintActionFilter ActorFunctionsFilter = LevelActorsFilter;
+	ActorFunctionsFilter.PermittedNodeTypes.Add(UK2Node_CallFunction::StaticClass());
+	MenuOut.AddMenuSection(ActorFunctionsFilter, FuncSectionHeading, LevelActorSectionGroup, FBlueprintActionMenuBuilder::ConsolidateBoundActions);
+
+	FBlueprintActionFilter ActorEventsFilter = LevelActorsFilter;
+	ActorEventsFilter.PermittedNodeTypes.Add(UK2Node_ActorBoundEvent::StaticClass());
+	MenuOut.AddMenuSection(ActorEventsFilter, EventSectionHeading, LevelActorSectionGroup, FBlueprintActionMenuBuilder::ConsolidateBoundActions);
+
+	FBlueprintActionFilter ActorReferencesFilter = LevelActorsFilter;
+	ActorReferencesFilter.RejectedNodeTypes.Append(ActorFunctionsFilter.PermittedNodeTypes);
+	ActorReferencesFilter.RejectedNodeTypes.Append(ActorEventsFilter.PermittedNodeTypes);
+	MenuOut.AddMenuSection(ActorReferencesFilter, FText::GetEmpty(), LevelActorSectionGroup, FBlueprintActionMenuBuilder::ConsolidateBoundActions);
+}
+
 /*******************************************************************************
  * FBlueprintActionMenuUtils
  ******************************************************************************/
@@ -132,23 +283,23 @@ void FBlueprintActionMenuUtils::MakeContextMenu(FBlueprintActionContext const& C
 	// Composing Filters
 	//--------------------------------------
 
-	FBlueprintActionFilter MenuFilter;
-	MenuFilter.Context = Context;
-	MenuFilter.Context.SelectedObjects.Empty();
+	FBlueprintActionFilter MainMenuFilter;
+	MainMenuFilter.Context = Context;
+	MainMenuFilter.Context.SelectedObjects.Empty();
 
 	FBlueprintActionFilter ComponentsFilter;
 	ComponentsFilter.Context = Context;
-	ComponentsFilter.PermittedNodeTypes.Add(UK2Node_CallFunction::StaticClass());
-	ComponentsFilter.PermittedNodeTypes.Add(UK2Node_ComponentBoundEvent::StaticClass());
 	// only want bound actions for this menu section
 	ComponentsFilter.AddRejectionTest(FBlueprintActionFilter::FRejectionTestDelegate::CreateStatic(IsUnBoundSpawner));
+	// @TODO: don't know exactly why we can only bind non-pure/const functions;
+	//        this is mirrored after FK2ActionMenuBuilder::GetFunctionCallsOnSelectedActors()
+	//        and FK2ActionMenuBuilder::GetFunctionCallsOnSelectedComponents(),
+	//        where we make the same stipulation
+	ComponentsFilter.AddRejectionTest(FBlueprintActionFilter::FRejectionTestDelegate::CreateStatic(IsPureNonConstAction));
+	
 
 	FBlueprintActionFilter LevelActorsFilter;
 	LevelActorsFilter.Context = Context;
-	LevelActorsFilter.PermittedNodeTypes.Add(UK2Node_CallFunction::StaticClass());
-	LevelActorsFilter.PermittedNodeTypes.Add(UK2Node_ActorBoundEvent::StaticClass());
-	LevelActorsFilter.PermittedNodeTypes.Add(UK2Node_Literal::StaticClass());
-	LevelActorsFilter.PermittedNodeTypes.Add(UK2Node_MatineeController::StaticClass());
 	// only want bound actions for this menu section
 	LevelActorsFilter.AddRejectionTest(FBlueprintActionFilter::FRejectionTestDelegate::CreateStatic(IsUnBoundSpawner));
 
@@ -157,14 +308,11 @@ void FBlueprintActionMenuUtils::MakeContextMenu(FBlueprintActionContext const& C
 	{
 		if (UObjectProperty* ObjProperty = Cast<UObjectProperty>(Selection))
 		{
-			ComponentsFilter.TargetClasses.Add(ObjProperty->PropertyClass);
 			LevelActorsFilter.Context.SelectedObjects.Remove(Selection);
 		}
 		else if (AActor* LevelActor = Cast<AActor>(Selection))
 		{
 			ComponentsFilter.Context.SelectedObjects.Remove(Selection);
-			// the loop below (over the editor's selected actors) will add to 
-			// LevelActorsFilter's OwnerClasses
 		}
 		else
 		{
@@ -179,10 +327,9 @@ void FBlueprintActionMenuUtils::MakeContextMenu(FBlueprintActionContext const& C
 	{
 		AActor* LevelActor = Cast<AActor>(*LvlActorIt);
 		LevelActorsFilter.Context.SelectedObjects.AddUnique(LevelActor);
-		LevelActorsFilter.TargetClasses.Add(LevelActor->GetClass());
 	}
 
-	bool const bOnlyBlueprintMembers = bIsContextSensitive && (MenuFilter.Context.Pins.Num() == 0);
+	bool const bAddTargetContext  = bIsContextSensitive;
 	bool bCanOperateOnLevelActors = bIsContextSensitive;
 	bool bCanHaveActorComponents  = bIsContextSensitive;
 	// determine if we can operate on certain object selections (level actors, 
@@ -193,21 +340,75 @@ void FBlueprintActionMenuUtils::MakeContextMenu(FBlueprintActionContext const& C
 		if (BlueprintClass != nullptr)
 		{
 			bCanOperateOnLevelActors &= BlueprintClass->IsChildOf<ALevelScriptActor>();
-			if (bOnlyBlueprintMembers)
+			if (bAddTargetContext)
 			{
-				MenuFilter.TargetClasses.Add(BlueprintClass);
+				MainMenuFilter.TargetClasses.Add(BlueprintClass);
 			}
 		}
 		bCanHaveActorComponents &= FBlueprintEditorUtils::DoesSupportComponents(Blueprint);
 	}
 
+	if (bAddTargetContext)
+	{
+		bool bContextPinIsObj = false;
+
+		// if we're dragging from a pin, we further extend the context to cover
+		// that pin and any other pins it sits beside 
+		for (UEdGraphPin* ContextPin : Context.Pins)
+		{
+			if (UClass* PinObjClass = GetPinClassType(ContextPin))
+			{
+				if (!bContextPinIsObj)
+				{
+					MainMenuFilter.TargetClasses.Empty();
+				}
+				MainMenuFilter.TargetClasses.Add(PinObjClass);
+				bContextPinIsObj = true;
+			}
+
+			for (UEdGraphPin* NodePin : ContextPin->GetOwningNode()->Pins)
+			{
+				if ((NodePin->Direction == ContextPin->Direction) && !bContextPinIsObj)
+				{
+					if (UClass* PinClass = GetPinClassType(NodePin))
+					{
+						MainMenuFilter.TargetClasses.Add(PinClass);
+					}
+				}
+			}
+		}
+	}
+
+	FBlueprintActionFilter CallOnMemberFilter;
+	CallOnMemberFilter.Context = MainMenuFilter.Context;
+	CallOnMemberFilter.PermittedNodeTypes.Add(UK2Node_CallFunction::StaticClass());
+	CallOnMemberFilter.AddRejectionTest(FBlueprintActionFilter::FRejectionTestDelegate::CreateStatic(IsUnBoundSpawner));
+	// instead of looking for "ExposeFunctionCategories" on object properties,
+	// we just expose functions for all components, this is problematic if we 
+	// were using the "ExposeFunctionCategories" on any non-component object 
+	// properties... if we re-enable this, then make sure to remove the 
+	// IsChildOf<UActorComponent>() check below
+	//CallOnMemberFilter.AddRejectionTest(FBlueprintActionFilter::FRejectionTestDelegate::CreateStatic(IsUnexposedMemberAction));
+
+	for (UClass const* TargetClass : MainMenuFilter.TargetClasses)
+	{
+		for (TFieldIterator<UObjectProperty> PropertyIt(TargetClass, EFieldIteratorFlags::IncludeSuper); PropertyIt; ++PropertyIt)
+		{
+			UObjectProperty* ComponentProperty = *PropertyIt;
+			// @TODO: consider a more relaxed scenario where we allow other
+			//        (non-component) properties as options here too
+			if (!ComponentProperty->PropertyClass->IsChildOf<UActorComponent>() || 
+				!ComponentProperty->HasAnyPropertyFlags(CPF_BlueprintVisible))
+			{
+				continue;
+			}
+			CallOnMemberFilter.Context.SelectedObjects.Add(ComponentProperty);
+		}
+	}
+
 	//--------------------------------------
 	// Defining Menu Sections
-	//--------------------------------------
-
-	static int32 const MainMenuSectionGroup   = 000;
-	static int32 const ComponentsSectionGroup = 100;
-	static int32 const LevelActorSectionGroup = 101;
+	//--------------------------------------	
 
 	MenuOut.Empty();
 
@@ -216,27 +417,25 @@ void FBlueprintActionMenuUtils::MakeContextMenu(FBlueprintActionContext const& C
 	// selected, then inform the user through a dummy menu entry)
 	if (bAddComponentsSection)
 	{
-		MenuOut.AddMenuSection(ComponentsFilter, FText::GetEmpty(), ComponentsSectionGroup);
+		AddComponentSections(ComponentsFilter, MenuOut);
 	}
 
 	bool const bAddLevelActorsSection = bIsContextSensitive && bCanOperateOnLevelActors && (LevelActorsFilter.Context.SelectedObjects.Num() > 0);
 	// add the level actor section to the menu
 	if (bAddLevelActorsSection)
 	{
-		// since we're consolidating all the bound actions, then we have to pick 
-		// one common base class to filter by
-		UClass* CommonClass = FindCommonBaseClass(LevelActorsFilter.TargetClasses);
-		LevelActorsFilter.TargetClasses.Empty();
-		LevelActorsFilter.TargetClasses.Add(CommonClass);
-
-		MenuOut.AddMenuSection(LevelActorsFilter, FText::GetEmpty(), LevelActorSectionGroup);
+		AddLevelActorSections(LevelActorsFilter, MenuOut);
 	}
 
-	if (!bIsContextSensitive)
+	if (bIsContextSensitive)
+	{
+		MenuOut.AddMenuSection(CallOnMemberFilter);
+	}
+	else
 	{
-		MenuFilter.Context.Pins.Empty();
+		MainMenuFilter.Context.Pins.Empty();
 	}
-	MenuOut.AddMenuSection(MenuFilter);
+	MenuOut.AddMenuSection(MainMenuFilter);
 
 	//--------------------------------------
 	// Building the Menu
diff --git a/Engine/Source/Editor/Kismet/Private/BlueprintBoundMenuItem.cpp b/Engine/Source/Editor/Kismet/Private/BlueprintBoundMenuItem.cpp
--- a/Engine/Source/Editor/Kismet/Private/BlueprintBoundMenuItem.cpp
+++ b/Engine/Source/Editor/Kismet/Private/BlueprintBoundMenuItem.cpp
@@ -14,14 +14,6 @@
 
 namespace BlueprintBoundMenuItemImpl
 {
-	/**
-	 * 
-	 * 
-	 * @param  BoundSpawner	
-	 * @return 
-	 */
-	static FText GetMenuCategoryFromAction(UBlueprintNodeSpawner const* BoundSpawner, IBlueprintNodeBinder::FBindingSet const& Bindings);
-
 	/**
 	 * 
 	 * 
@@ -31,35 +23,6 @@ namespace BlueprintBoundMenuItemImpl
 	static UClass* FindCommonBaseClass(IBlueprintNodeBinder::FBindingSet const& ObjectsList);
 }
 
-//------------------------------------------------------------------------------
-static FText BlueprintBoundMenuItemImpl::GetMenuCategoryFromAction(UBlueprintNodeSpawner const* BoundSpawner, IBlueprintNodeBinder::FBindingSet const& Bindings)
-{
-	check(Bindings.Num() > 0);
-
-	FText BoundObjectName;
-	if (Bindings.Num() > 1)
-	{
-		UClass* CommonBaseClass = FindCommonBaseClass(Bindings);
-		BoundObjectName = FText::Format(LOCTEXT("MultipleBindingsCategoryPostfix", "Selected {0}s"), CommonBaseClass->GetDisplayNameText());
-	}
-	else if (UObject const* BoundObject = Bindings.CreateConstIterator()->Get())
-	{
-		BoundObjectName = FText::FromName(BoundObject->GetFName());
-	}
-
-	FText MenuCategory = BoundSpawner->GetDefaultMenuCategory();
-	if (BoundSpawner->NodeClass->IsChildOf<UK2Node_Event>())
-	{
-		MenuCategory = FText::Format(LOCTEXT("ComponentEventCategory", "Add Event for {0}|{1}"), BoundObjectName, MenuCategory);
-	}
-	else if (BoundSpawner->NodeClass->IsChildOf<UK2Node_CallFunction>())
-	{
-		MenuCategory = FText::Format(LOCTEXT("ComponentEventCategory", "Call Function on {0}|{1}"), BoundObjectName, MenuCategory);
-	}
-
-	return MenuCategory;
-}
-
 //------------------------------------------------------------------------------
 static UClass* BlueprintBoundMenuItemImpl::FindCommonBaseClass(IBlueprintNodeBinder::FBindingSet const& ObjectsList)
 {
@@ -98,7 +61,6 @@ FBlueprintBoundMenuItem::FBlueprintBoundMenuItem(UBlueprintNodeSpawner const* Bo
 	: BoundSpawner(BoundSpawnerIn) 
 {
 	Grouping = MenuGrouping;
-	// category will be generated as we add bindings
 }
 
 //------------------------------------------------------------------------------
@@ -168,7 +130,6 @@ void FBlueprintBoundMenuItem::AddReferencedObjects(FReferenceCollector& Collecto
 void FBlueprintBoundMenuItem::AddBindings(IBlueprintNodeBinder::FBindingSet const& BindingSet)
 {
 	BoundObjects.Append(BindingSet);
-	Category = BlueprintBoundMenuItemImpl::GetMenuCategoryFromAction(BoundSpawner, BoundObjects).ToString();
 
 	// Allow the binding to generate a menu description:
 	FText NewMenuDescription = BoundSpawner->GetDefaultMenuName(BoundObjects);
diff --git a/Engine/Source/Editor/Kismet/Public/BlueprintActionMenuBuilder.h b/Engine/Source/Editor/Kismet/Public/BlueprintActionMenuBuilder.h
--- a/Engine/Source/Editor/Kismet/Public/BlueprintActionMenuBuilder.h
+++ b/Engine/Source/Editor/Kismet/Public/BlueprintActionMenuBuilder.h
@@ -29,6 +29,10 @@ public:
 		// Rolls properties into a single menu item that will spawn a drag-drop
 		// menu for users to pick a node type from.
 		ConsolidatePropertyActions = (1<<0),
+
+		// Rolls bound node spawners into a single menu entry that will spawn 
+		// multiple nodes, each bound to a single binding.
+		ConsolidateBoundActions    = (1<<1),
 	};
 	
 public:
