Constructing icons for bound menu items (under the new blueprint menu system).

diff --git a/Engine/Source/Editor/Kismet/Private/BlueprintBoundMenuItem.cpp b/Engine/Source/Editor/Kismet/Private/BlueprintBoundMenuItem.cpp
--- a/Engine/Source/Editor/Kismet/Private/BlueprintBoundMenuItem.cpp
+++ b/Engine/Source/Editor/Kismet/Private/BlueprintBoundMenuItem.cpp
@@ -142,7 +142,22 @@ void FBlueprintBoundMenuItem::AddBindings(IBlueprintNodeBinder::FBindingSet cons
 //------------------------------------------------------------------------------
 FSlateBrush const* FBlueprintBoundMenuItem::GetMenuIcon(FSlateColor& ColorOut)
 {
-	// @TODO: 
+	FLinearColor IconTint = FLinearColor::White;
+
+	FName BrushName = BoundSpawner->GetDefaultMenuIcon(IconTint);
+	if (BrushName == NAME_None)
+	{
+		if (UEdGraphNode* TemplateNode = BoundSpawner->GetTemplateNode())
+		{
+			BrushName = TemplateNode->GetPaletteIcon(IconTint);
+		}
+	}
+	
+	if (BrushName != NAME_None)
+	{
+		ColorOut = FSlateColor(IconTint);
+		return FEditorStyle::GetBrush(BrushName);
+	}
 	return nullptr;
 }
 
diff --git a/Engine/Source/Editor/Kismet/Private/SBlueprintPalette.cpp b/Engine/Source/Editor/Kismet/Private/SBlueprintPalette.cpp
--- a/Engine/Source/Editor/Kismet/Private/SBlueprintPalette.cpp
+++ b/Engine/Source/Editor/Kismet/Private/SBlueprintPalette.cpp
@@ -28,6 +28,7 @@
 #include "AnimationTransitionGraph.h"
 #include "BlueprintActionMenuItem.h"
 #include "BlueprintDragDropMenuItem.h"
+#include "BlueprintBoundMenuItem.h"
 #include "BlueprintNodeSpawner.h"
 #include "TutorialMetaData.h"
 #include "BlueprintEditorSettings.h" // for bShowActionMenuItemSignatures
@@ -369,6 +370,11 @@ static void GetPaletteItemIcon(TSharedPtr<FEdGraphSchemaAction> ActionIn, UBluep
 		FBlueprintDragDropMenuItem* DragDropAction = (FBlueprintDragDropMenuItem*)ActionIn.Get();
 		BrushOut = DragDropAction->GetMenuIcon(ColorOut);
 	}
+	else if (ActionIn->GetTypeId() == FBlueprintBoundMenuItem::StaticGetTypeId())
+	{
+		FBlueprintBoundMenuItem* BoundAction = (FBlueprintBoundMenuItem*)ActionIn.Get();
+		BrushOut = BoundAction->GetMenuIcon(ColorOut);
+	}
 	else if (ActionIn->GetTypeId() == FEdGraphSchemaAction_K2AddDocumentation::StaticGetTypeId())
 	{
 		ToolTipOut = LOCTEXT("Documentation_Tooltip", "Create a documentation node.").ToString();
