External functions, macros, and variables, will be put into a category of their Blueprint if they have no specified category (and always for external variables).

diff --git a/Engine/Source/Editor/Kismet/Private/BlueprintActionMenuBuilder.cpp b/Engine/Source/Editor/Kismet/Private/BlueprintActionMenuBuilder.cpp
--- a/Engine/Source/Editor/Kismet/Private/BlueprintActionMenuBuilder.cpp
+++ b/Engine/Source/Editor/Kismet/Private/BlueprintActionMenuBuilder.cpp
@@ -251,9 +251,40 @@ FText FBlueprintActionMenuItemFactory::GetCategoryForAction(TWeakPtr<FBlueprintE
 			{
 				MenuCategory = FuncOwner->GetDisplayNameText();
 			}
-			MenuCategory = FText::Format(LOCTEXT("MemberFunctionsCategory", "Call Function|{0}"), MenuCategory);
+			MenuCategory = FText::Format(LOCTEXT("MemberFunctionsCategory", "{0}|Call Function"), MenuCategory);
 		}
-		else 
+		else if (UBlueprintNodeSpawner const* NodeSpawner = Cast<UBlueprintNodeSpawner>(Action))
+		{
+			// Only for macro instances do we want to mess with the category to put it into a category
+			// named after the Blueprint if no other category is specified
+			if(NodeSpawner->NodeClass == UK2Node_MacroInstance::StaticClass())
+			{
+				UK2Node_MacroInstance* MacroInstance = Cast<UK2Node_MacroInstance>(NodeSpawner->GetTemplateNode());
+				if(UEdGraph* MacroGraph = MacroInstance->GetMacroGraph())
+				{
+					// Check if the MacroGraph has a category
+					FKismetUserDeclaredFunctionMetadata* MacroGraphMetadata = UK2Node_MacroInstance::GetAssociatedGraphMetadata(MacroGraph);
+					if ((MacroGraphMetadata != nullptr) && MacroGraphMetadata->Category.IsEmpty())
+					{
+						MenuCategory = MacroInstance->GetMenuCategory();
+
+						UBlueprint* MacroBlueprint = FBlueprintEditorUtils::FindBlueprintForGraph(MacroGraph);
+						if(MacroBlueprint != GetTargetBlueprint())
+						{
+							FText BlueprintDisplayName = MacroBlueprint->GeneratedClass->GetDisplayNameText();
+							FFormatNamedArguments Args;
+							Args.Add(TEXT("BlueprintDisplayName"), BlueprintDisplayName);
+							Args.Add(TEXT("MacroCategory"), MenuCategory);
+
+							MenuCategory = FText::Format(LOCTEXT("MemberMacroCategory", "{BlueprintDisplayName}|{MacroCategory}"), Args);
+						}
+					}
+				}
+			}
+		}
+
+		// If the menu category is still empty, fill it in with the node template's defined category
+		if(MenuCategory.IsEmpty())
 		{
 			// @TODO: consider moving GetMenuCategory() up into UEdGraphNode
 			if (UK2Node* NodeTemplate = Cast<UK2Node>(GetTemplateNode(Action, EditorContext)))
@@ -262,6 +293,32 @@ FText FBlueprintActionMenuItemFactory::GetCategoryForAction(TWeakPtr<FBlueprintE
 			}
 		}
 	}
+	else
+	{
+		if (UBlueprintVariableNodeSpawner const* VarSpawner = Cast<UBlueprintVariableNodeSpawner>(Action))
+		{
+			UProperty const* Property = VarSpawner->GetVarProperty();
+			if(Property)
+			{
+				check(Property != nullptr);
+				UClass* PropertyOwner = Property->GetTypedOuter<UClass>();
+
+				UBlueprint* Blueprint = GetTargetBlueprint();
+				check(Blueprint != nullptr);
+				UClass* BlueprintClass = (Blueprint->SkeletonGeneratedClass != nullptr) ? Blueprint->SkeletonGeneratedClass : Blueprint->ParentClass;
+
+				// if this is NOT a self function call (self function calls
+				// don't get nested any deeper)
+				if (!BlueprintClass->IsChildOf(PropertyOwner))
+				{
+					FFormatNamedArguments Args;
+					Args.Add(TEXT("PropertyDisplayName"), PropertyOwner->GetDisplayNameText());
+					Args.Add(TEXT("VariableCategory"), MenuCategory);
+					MenuCategory = FText::Format(LOCTEXT("MemberVariablesCategory", "{PropertyDisplayName}|{VariableCategory}"), Args);
+				}
+			}
+		}
+	}
 	
 	return MenuCategory;
 }
diff --git a/Engine/Source/Runtime/CoreUObject/Private/UObject/Class.cpp b/Engine/Source/Runtime/CoreUObject/Private/UObject/Class.cpp
--- a/Engine/Source/Runtime/CoreUObject/Private/UObject/Class.cpp
+++ b/Engine/Source/Runtime/CoreUObject/Private/UObject/Class.cpp
@@ -190,6 +190,7 @@ struct FDisplayNameHelper
 		if (Class && !Class->HasAnyClassFlags(CLASS_Native))
 		{
 			Name.RemoveFromEnd(TEXT("_C"));
+			Name.RemoveFromStart(TEXT("SKEL_"));
 		}
 		return Name;
 	}
