#docs Added/improved comments for ENetMode, EVariableReplication and linked to online doc pages.

diff --git a/Engine/Source/Editor/Kismet/Private/BlueprintDetailsCustomization.h b/Engine/Source/Editor/Kismet/Private/BlueprintDetailsCustomization.h
--- a/Engine/Source/Editor/Kismet/Private/BlueprintDetailsCustomization.h
+++ b/Engine/Source/Editor/Kismet/Private/BlueprintDetailsCustomization.h
@@ -4,14 +4,31 @@
 
 #include "EdGraph/EdGraphNode_Documentation.h"
 
-/** Variable replication states */
+/**
+ * Variable network replication options.
+ * @see https://docs.unrealengine.com/latest/INT/Gameplay/Networking/Replication/
+ */
 namespace EVariableReplication
 {
 	enum Type
 	{
+		/**
+		 * Not replicated.
+		 */
 		None,
+
+		/**
+		 * Replicated from server to client.
+		 * As values change on the server, client automatically receives new values, if Actor is set to replicate.
+		 */
 		Replicated,
+
+		/**
+		 * Replicated from server to client, with a notification function called on clients when a new value arrives.
+		 * An event with the name "On Rep <VariableName>" is created.
+		 */
 		RepNotify,
+
 		MAX
 	};
 }
diff --git a/Engine/Source/Runtime/Engine/Classes/Engine/EngineBaseTypes.h b/Engine/Source/Runtime/Engine/Classes/Engine/EngineBaseTypes.h
--- a/Engine/Source/Runtime/Engine/Classes/Engine/EngineBaseTypes.h
+++ b/Engine/Source/Runtime/Engine/Classes/Engine/EngineBaseTypes.h
@@ -556,6 +556,7 @@ struct ENGINE_API FURL
 
 /**
  * The network mode the game is currently running.
+ * @see https://docs.unrealengine.com/latest/INT/Gameplay/Networking/Replication/
  */
 enum ENetMode
 {
