[302977] Add orthographic view options to the Blueprint editor preview viewport in Components mode.

- Also switched the view options drop-down from "View" text to an arrow icon (to be consistent w/ the main level editor viewport).

diff --git a/Engine/Source/Editor/Kismet/Private/SSCSEditorViewport.cpp b/Engine/Source/Editor/Kismet/Private/SSCSEditorViewport.cpp
--- a/Engine/Source/Editor/Kismet/Private/SSCSEditorViewport.cpp
+++ b/Engine/Source/Editor/Kismet/Private/SSCSEditorViewport.cpp
@@ -38,15 +38,26 @@ public:
 				SNew(SHorizontalBox)
 				+SHorizontalBox::Slot()
 				.AutoWidth()
-				.Padding(5.0f, 2.0f)
+				.Padding(2.0f, 2.0f)
 				[
 					SNew(SEditorViewportToolbarMenu)
 					.ParentToolBar(SharedThis(this))
 					.Cursor(EMouseCursor::Default)
-					.Label(NSLOCTEXT("BlueprintEditor", "ViewMenuTitle_Default", "View"))
+					.Image("EditorViewportToolBar.MenuDropdown")
 					.OnGetMenuContent(this, &SSCSEditorViewportToolBar::GeneratePreviewMenu)
 				]
 				+ SHorizontalBox::Slot()
+				.AutoWidth()
+				.Padding(2.0f, 2.0f)
+				[
+					SNew( SEditorViewportToolbarMenu )
+					.ParentToolBar( SharedThis( this ) )
+					.Cursor( EMouseCursor::Default )
+					.Label(this, &SSCSEditorViewportToolBar::GetCameraMenuLabel)
+					.LabelIcon(this, &SSCSEditorViewportToolBar::GetCameraMenuLabelIcon)
+					.OnGetMenuContent(this, &SSCSEditorViewportToolBar::GenerateCameraMenu)
+				]
+				+ SHorizontalBox::Slot()
 				.Padding( 3.0f, 1.0f )
 				.HAlign( HAlign_Right )
 				[
@@ -82,6 +93,86 @@ public:
 		return PreviewOptionsMenuBuilder.MakeWidget();
 	}
 
+	FText GetCameraMenuLabel() const
+	{
+		FText Label = NSLOCTEXT("BlueprintEditor", "CameraMenuTitle_Default", "Camera");
+
+		if(EditorViewport.IsValid())
+		{
+			switch(EditorViewport.Pin()->GetViewportClient()->GetViewportType())
+			{
+			case LVT_Perspective:
+				Label = NSLOCTEXT("BlueprintEditor", "CameraMenuTitle_Perspective", "Perspective");
+				break;
+
+			case LVT_OrthoXY:
+				Label = NSLOCTEXT("BlueprintEditor", "CameraMenuTitle_Top", "Top");
+				break;
+
+			case LVT_OrthoYZ:
+				Label = NSLOCTEXT("BlueprintEditor", "CameraMenuTitle_Side", "Side");
+				break;
+
+			case LVT_OrthoXZ:
+				Label = NSLOCTEXT("BlueprintEditor", "CameraMenuTitle_Front", "Front");
+				break;
+
+			case LVT_OrthoFreelook:
+				Label = NSLOCTEXT("BlueprintEditor", "CameraMenuTitle_OrthoFreelook", "Ortho");
+				break;
+			}
+		}
+
+		return Label;
+	}
+
+	const FSlateBrush* GetCameraMenuLabelIcon() const
+	{
+		FName Icon = NAME_None;
+
+		if(EditorViewport.IsValid())
+		{
+			switch(EditorViewport.Pin()->GetViewportClient()->GetViewportType())
+			{
+			case LVT_Perspective:
+				Icon = FName("EditorViewport.Perspective");
+				break;
+
+			case LVT_OrthoXY:
+				Icon = FName("EditorViewport.Top");
+				break;
+
+			case LVT_OrthoYZ:
+				Icon = FName( "EditorViewport.Side");
+				break;
+
+			case LVT_OrthoXZ:
+				Icon = FName("EditorViewport.Front");
+				break;
+			}
+		}
+
+		return FEditorStyle::GetBrush(Icon);
+	}
+
+	TSharedRef<SWidget> GenerateCameraMenu() const
+	{
+		TSharedPtr<const FUICommandList> CommandList = EditorViewport.IsValid()? EditorViewport.Pin()->GetCommandList(): nullptr;
+
+		const bool bInShouldCloseWindowAfterMenuSelection = true;
+		FMenuBuilder CameraMenuBuilder(bInShouldCloseWindowAfterMenuSelection, CommandList);
+
+		CameraMenuBuilder.AddMenuEntry(FEditorViewportCommands::Get().Perspective);
+
+		CameraMenuBuilder.BeginSection("LevelViewportCameraType_Ortho", NSLOCTEXT("BlueprintEditor", "CameraTypeHeader_Ortho", "Orthographic"));
+			CameraMenuBuilder.AddMenuEntry(FEditorViewportCommands::Get().Top);
+			CameraMenuBuilder.AddMenuEntry(FEditorViewportCommands::Get().Side);
+			CameraMenuBuilder.AddMenuEntry(FEditorViewportCommands::Get().Front);
+		CameraMenuBuilder.EndSection();
+
+		return CameraMenuBuilder.MakeWidget();
+	}
+
 private:
 	/** Reference to the parent viewport */
 	TWeakPtr<SSCSEditorViewport> EditorViewport;
