[346630] Blueprint editor - Prevent copy/duplicate of component nodes not marked as BlueprintSpawnable.

- Modified SSCSEditor::CanCopyNodes() to include additional checks for relevant class & class metadata flags.
- Also minor updates to default root scene component tooltip text.

diff --git a/Engine/Source/Editor/Kismet/Private/SSCSEditor.cpp b/Engine/Source/Editor/Kismet/Private/SSCSEditor.cpp
--- a/Engine/Source/Editor/Kismet/Private/SSCSEditor.cpp
+++ b/Engine/Source/Editor/Kismet/Private/SSCSEditor.cpp
@@ -1778,12 +1778,12 @@ FText SSCS_RowWidget::GetTooltipText() const
 	{
 		if(NodePtr->IsInherited())
 		{
-			return LOCTEXT("InheritedDefaultSceneRootToolTip", "This is the default scene root component. It has been inherited from the parent class, so it cannot be renamed, replaced or deleted. New scene components will automatically be attached to it.");
+			return LOCTEXT("InheritedDefaultSceneRootToolTip", "This is the default scene root component. It cannot be copied, renamed or deleted. It has been inherited from the parent class, so its properties cannot be edited here. New scene components will automatically be attached to it.");
 		}
 		else
 		{
-		return LOCTEXT("DefaultSceneRootToolTip", "This is the default scene root component. It cannot be renamed or deleted. Adding a new scene component will automatically replace it as the new root.");
-	}
+			return LOCTEXT("DefaultSceneRootToolTip", "This is the default scene root component. It cannot be copied, renamed or deleted. Adding a new scene component will automatically replace it as the new root.");
+		}
 	}
 	else
 	{
@@ -2755,15 +2755,27 @@ void SSCSEditor::CutSelectedNodes()
 bool SSCSEditor::CanCopyNodes() const
 {
 	TArray<FSCSEditorTreeNodePtrType> SelectedNodes = SCSTreeWidget->GetSelectedItems();
-	for (int32 i = 0; i < SelectedNodes.Num(); ++i)
+	bool bCanCopy = SelectedNodes.Num() > 0;
+	if(bCanCopy)
 	{
-		if (SelectedNodes[i]->GetComponentTemplate() == NULL || SelectedNodes[i]->IsDefaultSceneRoot())
+		for (int32 i = 0; i < SelectedNodes.Num() && bCanCopy; ++i)
 		{
-			return false;
+			// Check for the default scene root; that cannot be copied/duplicated
+			UActorComponent* ComponentTemplate = SelectedNodes[i]->GetComponentTemplate();
+			bCanCopy = ComponentTemplate != nullptr && !SelectedNodes[i]->IsDefaultSceneRoot();
+			if (bCanCopy)
+			{
+				UClass* ComponentTemplateClass = ComponentTemplate->GetClass();
+				check(ComponentTemplateClass != nullptr);
+
+				// Component class cannot be abstract and must also be tagged as BlueprintSpawnable
+				bCanCopy = !ComponentTemplateClass->HasAnyClassFlags(CLASS_Abstract)
+					&& ComponentTemplateClass->HasMetaData(FBlueprintMetadata::MD_BlueprintSpawnableComponent);
+			}
 		}
 	}
 
-	return SelectedNodes.Num() > 0;
+	return bCanCopy;
 }
 
 void SSCSEditor::CopySelectedNodes()
