Fix for consturcting bad legacy menu context (from the new blueprint menu system).

diff --git a/Engine/Source/Editor/Kismet/Private/BlueprintActionMenuBuilder.cpp b/Engine/Source/Editor/Kismet/Private/BlueprintActionMenuBuilder.cpp
--- a/Engine/Source/Editor/Kismet/Private/BlueprintActionMenuBuilder.cpp
+++ b/Engine/Source/Editor/Kismet/Private/BlueprintActionMenuBuilder.cpp
@@ -668,12 +668,19 @@ void FBlueprintActionMenuBuilderImpl::AppendLegacyItems(FMenuSectionDefinition c
 		{
 			LegacyBuilder.FromPin = MenuContext.Pins[0];
 		}
-		LegacyBuilder.SelectedObjects.Append(MenuContext.SelectedObjects);
 		
 		bool bIsContextSensitive = true;
 		if (BlueprintEditor.IsValid())
 		{
 			bIsContextSensitive = BlueprintEditor.Pin()->GetIsContextSensitive();
+			if (bIsContextSensitive)
+			{
+				FEdGraphSchemaAction_K2Var* SelectedVar = BlueprintEditor.Pin()->GetMyBlueprintWidget()->SelectionAsVar();
+				if ((SelectedVar != nullptr) && (SelectedVar->GetProperty() != nullptr))
+				{
+					LegacyBuilder.SelectedObjects.Add(SelectedVar->GetProperty());
+				}
+			}
 		}
 		
 		if (bIsContextSensitive)
diff --git a/Engine/Source/Editor/Kismet/Private/BlueprintActionMenuUtils.cpp b/Engine/Source/Editor/Kismet/Private/BlueprintActionMenuUtils.cpp
--- a/Engine/Source/Editor/Kismet/Private/BlueprintActionMenuUtils.cpp
+++ b/Engine/Source/Editor/Kismet/Private/BlueprintActionMenuUtils.cpp
@@ -412,6 +412,14 @@ void FBlueprintActionMenuUtils::MakeContextMenu(FBlueprintActionContext const& C
 
 	MenuOut.Empty();
 
+	if (!bIsContextSensitive)
+	{
+		MainMenuFilter.Context.Pins.Empty();
+	}
+	// for legacy purposes, we have to add the main menu section first (when 
+	// reconstructing the legacy menu, we pull the first menu system)
+	MenuOut.AddMenuSection(MainMenuFilter);
+
 	bool const bAddComponentsSection = bIsContextSensitive && bCanHaveActorComponents && (ComponentsFilter.Context.SelectedObjects.Num() > 0);
 	// add the components section to the menu (if we don't have any components
 	// selected, then inform the user through a dummy menu entry)
@@ -430,12 +438,7 @@ void FBlueprintActionMenuUtils::MakeContextMenu(FBlueprintActionContext const& C
 	if (bIsContextSensitive)
 	{
 		MenuOut.AddMenuSection(CallOnMemberFilter);
-	}
-	else
-	{
-		MainMenuFilter.Context.Pins.Empty();
-	}
-	MenuOut.AddMenuSection(MainMenuFilter);
+	}	
 
 	//--------------------------------------
 	// Building the Menu
