Fixing blueprint menu system to have the proper context when dragging off input object pins.

diff --git a/Engine/Source/Editor/Kismet/Private/BlueprintActionMenuUtils.cpp b/Engine/Source/Editor/Kismet/Private/BlueprintActionMenuUtils.cpp
--- a/Engine/Source/Editor/Kismet/Private/BlueprintActionMenuUtils.cpp
+++ b/Engine/Source/Editor/Kismet/Private/BlueprintActionMenuUtils.cpp
@@ -371,6 +371,13 @@ void FBlueprintActionMenuUtils::MakeContextMenu(FBlueprintActionContext const& C
 		// that pin and any other pins it sits beside 
 		for (UEdGraphPin* ContextPin : Context.Pins)
 		{
+			// we only want the pin to be the target class when it is an output
+			// (doesn't make sense to get members to plug into their parent)
+			if (ContextPin->Direction == EGPD_Input)
+			{
+				continue;
+			}
+
 			if (UClass* PinObjClass = GetPinClassType(ContextPin))
 			{
 				if (!bContextPinIsObj)
