Adding an option to disable the added "TargetClass" context for the new BP menus.

diff --git a/Engine/Source/Editor/BlueprintGraph/Private/BlueprintEditorSettings.cpp b/Engine/Source/Editor/BlueprintGraph/Private/BlueprintEditorSettings.cpp
--- a/Engine/Source/Editor/BlueprintGraph/Private/BlueprintEditorSettings.cpp
+++ b/Engine/Source/Editor/BlueprintGraph/Private/BlueprintEditorSettings.cpp
@@ -10,6 +10,7 @@ UBlueprintEditorSettings::UBlueprintEditorSettings(const class FPostConstructIni
 {
 	bShowGraphInstructionText = true;
 	NodeTemplateCacheCapMB    = 20.f;
+	bUseTargetContextForNodeMenu = true;
 
 	// settings that were moved out of experimental...
 	UEditorExperimentalSettings const* ExperimentalSettings = GetDefault<UEditorExperimentalSettings>();
diff --git a/Engine/Source/Editor/BlueprintGraph/Public/BlueprintEditorSettings.h b/Engine/Source/Editor/BlueprintGraph/Public/BlueprintEditorSettings.h
--- a/Engine/Source/Editor/BlueprintGraph/Public/BlueprintEditorSettings.h
+++ b/Engine/Source/Editor/BlueprintGraph/Public/BlueprintEditorSettings.h
@@ -24,6 +24,10 @@ public:
 	UPROPERTY(EditAnywhere, config, Category=UserExperience)
 	bool bShowGraphInstructionText;
 
+	/** If enabled, will use the blueprint's (or output pin's) class to narrow down context menu results. */
+	UPROPERTY(EditAnywhere, config, Category=UserExperience)
+	bool bUseTargetContextForNodeMenu;
+
 // Developer Settings
 public:
 	/** If enabled, tooltips on action menu items will show the associated action's signature id (can be used to setup custom favorites menus).*/
diff --git a/Engine/Source/Editor/Kismet/Private/BlueprintActionMenuUtils.cpp b/Engine/Source/Editor/Kismet/Private/BlueprintActionMenuUtils.cpp
--- a/Engine/Source/Editor/Kismet/Private/BlueprintActionMenuUtils.cpp
+++ b/Engine/Source/Editor/Kismet/Private/BlueprintActionMenuUtils.cpp
@@ -7,7 +7,7 @@
 #include "BlueprintNodeSpawner.h"
 #include "BlueprintEditorUtils.h"	// for DoesSupportComponents()
 #include "BlueprintPaletteFavorites.h"
-#include "EdGraphSchema_K2.h"		// for bUseLegacyActionMenus 
+#include "EdGraphSchema_K2.h"		// for PC_Object/PC_Interface 
 #include "K2Node_ActorBoundEvent.h"
 #include "K2Node_CallFunction.h"
 #include "K2Node_ComponentBoundEvent.h"
@@ -16,6 +16,7 @@
 #include "K2Node_VariableSet.h"
 #include "KismetEditorUtilities.h"	// for CanPasteNodes()
 #include "K2ActionMenuBuilder.h"
+#include "BlueprintEditorSettings.h"
 
 #define LOCTEXT_NAMESPACE "BlueprintActionMenuUtils"
 
@@ -344,7 +345,8 @@ void FBlueprintActionMenuUtils::MakeContextMenu(FBlueprintActionContext const& C
 		LevelActorsFilter.Context.SelectedObjects.AddUnique(LevelActor);
 	}
 
-	bool const bAddTargetContext  = bIsContextSensitive;
+	const UBlueprintEditorSettings* BlueprintSettings = GetDefault<UBlueprintEditorSettings>();
+	bool const bAddTargetContext  = bIsContextSensitive && BlueprintSettings->bUseTargetContextForNodeMenu;
 	bool bCanOperateOnLevelActors = bIsContextSensitive;
 	bool bCanHaveActorComponents  = bIsContextSensitive;
 	// determine if we can operate on certain object selections (level actors, 
