Blueprints: Adjust style of 'parent class' hyperlink and add some more padding between it and the search box

diff --git a/Engine/Source/Editor/Kismet/Private/BlueprintEditor.cpp b/Engine/Source/Editor/Kismet/Private/BlueprintEditor.cpp
--- a/Engine/Source/Editor/Kismet/Private/BlueprintEditor.cpp
+++ b/Engine/Source/Editor/Kismet/Private/BlueprintEditor.cpp
@@ -1400,9 +1400,11 @@ void FBlueprintEditor::PostRegenerateMenusAndToolbars()
 			]
 			+ SHorizontalBox::Slot()
 			.AutoWidth()
+			.VAlign(VAlign_Center)
 			[
 				SNew(SHyperlink)
-				.Style(FEditorStyle::Get(), "HoverOnlyHyperlink")
+				.Style(FEditorStyle::Get(), "EditBPHyperlink")
+				.TextStyle(FEditorStyle::Get(), "DetailsView.EditBlueprintHyperlinkStyle")
 				.IsEnabled(this, &FBlueprintEditor::IsNativeParentClassCodeLinkEnabled)
 				.Visibility(this, &FBlueprintEditor::GetNativeParentClassButtonsVisibility)
 				.OnNavigate(this, &FBlueprintEditor::OnEditParentClassNativeCodeClicked)
@@ -1414,7 +1416,7 @@ void FBlueprintEditor::PostRegenerateMenusAndToolbars()
 			.VAlign(VAlign_Center)
 			[
 				SNew(SSpacer)
-				.Size(FVector2D(2.0f, 1.0f))
+				.Size(FVector2D(8.0f, 1.0f))
 			]
 			;
 		SetMenuOverlay( MenuOverlayBox );
