BP: A crash occurs, when graph (including CreateDelegate node) is removed.

diff --git a/Engine/Source/Editor/Kismet/Private/SMyBlueprint.cpp b/Engine/Source/Editor/Kismet/Private/SMyBlueprint.cpp
--- a/Engine/Source/Editor/Kismet/Private/SMyBlueprint.cpp
+++ b/Engine/Source/Editor/Kismet/Private/SMyBlueprint.cpp
@@ -1672,9 +1672,12 @@ void SMyBlueprint::OnDeleteGraph(UEdGraph* InGraph, EEdGraphSchemaAction_K2Graph
 		FBlueprintEditorUtils::RemoveGraph(GetBlueprintObj(), InGraph, EGraphRemoveFlags::Recompile);
 		BlueprintEditorPtr.Pin()->CloseDocumentTab(InGraph);
 
-		for (TObjectIterator<UK2Node_CreateDelegate> It; It; ++It)
+		for (TObjectIterator<UK2Node_CreateDelegate> It(RF_ClassDefaultObject | RF_PendingKill); It; ++It)
 		{
-			It->HandleAnyChange();
+			if (It->GetGraph() != InGraph)
+			{
+				It->HandleAnyChange();
+			}
 		}
 
 		InGraph = NULL;
@@ -1696,7 +1699,7 @@ void SMyBlueprint::OnDeleteDelegate(FEdGraphSchemaAction_K2Delegate* InDelegateA
 		FBlueprintEditorUtils::RemoveMemberVariable(Blueprint, EdGraph->GetFName());
 		FBlueprintEditorUtils::RemoveGraph(Blueprint, EdGraph, EGraphRemoveFlags::Recompile);
 
-		for (TObjectIterator<UK2Node_CreateDelegate> It; It; ++It)
+		for (TObjectIterator<UK2Node_CreateDelegate> It(RF_ClassDefaultObject | RF_PendingKill); It; ++It)
 		{
 			It->HandleAnyChange();
 		}
