Finally giving in, and letting the new blueprint action database govern how those actions are presented in the ui :(

diff --git a/Engine/Source/Editor/BlueprintGraph/Classes/BlueprintBoundEventNodeSpawner.h b/Engine/Source/Editor/BlueprintGraph/Classes/BlueprintBoundEventNodeSpawner.h
--- a/Engine/Source/Editor/BlueprintGraph/Classes/BlueprintBoundEventNodeSpawner.h
+++ b/Engine/Source/Editor/BlueprintGraph/Classes/BlueprintBoundEventNodeSpawner.h
@@ -34,9 +34,6 @@ public:
 	// UBlueprintNodeSpawner interface
 	virtual FBlueprintNodeSignature GetSpawnerSignature() const override;
 	virtual UEdGraphNode* Invoke(UEdGraph* ParentGraph, FBindingSet const& Bindings, FVector2D const Location) const override;
-	virtual FText GetDefaultMenuName(FBindingSet const& Bindings) const override;
-	virtual FText GetDefaultMenuCategory() const override;
-	virtual FName GetDefaultMenuIcon(FLinearColor& ColorOut) const;
 	// End UBlueprintNodeSpawner interface
 
 	// UBlueprintEventNodeSpawner interface
@@ -48,11 +45,7 @@ public:
 	virtual bool BindToNode(UEdGraphNode* Node, UObject* Binding) const override;
 	// End IBlueprintNodeBinder interface
 
-	/**
-	 * 
-	 *
-	 * @return 
-	 */
+	/** @return  */
 	UMulticastDelegateProperty const* GetEventDelegate() const;
 
 private:
diff --git a/Engine/Source/Editor/BlueprintGraph/Classes/BlueprintBoundNodeSpawner.h b/Engine/Source/Editor/BlueprintGraph/Classes/BlueprintBoundNodeSpawner.h
--- a/Engine/Source/Editor/BlueprintGraph/Classes/BlueprintBoundNodeSpawner.h
+++ b/Engine/Source/Editor/BlueprintGraph/Classes/BlueprintBoundNodeSpawner.h
@@ -15,7 +15,6 @@ class BLUEPRINTGRAPH_API UBlueprintBoundNodeSpawner : public UBlueprintNodeSpawn
 	GENERATED_UCLASS_BODY()
 	DECLARE_DELEGATE_RetVal_OneParam(bool, FCanBindObjectDelegate, UObject const*);
 	DECLARE_DELEGATE_RetVal_TwoParams(bool, FOnBindObjectDelegate, UEdGraphNode*, UObject*);
-	DECLARE_DELEGATE_RetVal_OneParam(FText, FOnGenerateMenuDescriptionDelegate, const IBlueprintNodeBinder::FBindingSet& );
 
 public:
 	/**
@@ -25,7 +24,6 @@ public:
 
 	// UBlueprintNodeSpawner interface
 	virtual FBlueprintNodeSignature GetSpawnerSignature() const override;
-	virtual FText GetDefaultMenuName(FBindingSet const& Bindings) const override;
 	// End UBlueprintNodeSpawner interface
 
 	// FBlueprintNodeBinder interface
@@ -43,9 +41,4 @@ public:
 	 * A delegate to perform specialized node setup during binding
 	 */
 	FOnBindObjectDelegate OnBindObjectDelegate;
-
-	/**
-	 * A delegate to generate a description for the menu item. Executed any time bindings change
-	 */
-	FOnGenerateMenuDescriptionDelegate OnGenerateMenuDescriptionDelegate;
 };
diff --git a/Engine/Source/Editor/BlueprintGraph/Classes/BlueprintComponentNodeSpawner.h b/Engine/Source/Editor/BlueprintGraph/Classes/BlueprintComponentNodeSpawner.h
--- a/Engine/Source/Editor/BlueprintGraph/Classes/BlueprintComponentNodeSpawner.h
+++ b/Engine/Source/Editor/BlueprintGraph/Classes/BlueprintComponentNodeSpawner.h
@@ -33,12 +33,8 @@ public:
 
 	// UBlueprintNodeSpawner interface
 	virtual FBlueprintNodeSignature GetSpawnerSignature() const override;
+	virtual FBlueprintActionUiSpec GetUiSpec(FBlueprintActionContext const& Context, FBindingSet const& Bindings) const override;
 	virtual UEdGraphNode* Invoke(UEdGraph* ParentGraph, FBindingSet const& Bindings, FVector2D const Location) const override;
-	virtual FText GetDefaultMenuName(FBindingSet const& Bindings) const override;
-	virtual FText GetDefaultMenuTooltip() const;
-	virtual FText GetDefaultMenuCategory() const override;
-	virtual FString GetDefaultSearchKeywords() const override;
-	virtual FName GetDefaultMenuIcon(FLinearColor& ColorOut) const override;
 	// End UBlueprintNodeSpawner interface
 	
 	// IBlueprintNodeBinder interface
@@ -57,8 +53,4 @@ private:
 	/** The component class to configure new nodes with. */
 	UPROPERTY()
 	TSubclassOf<UActorComponent> ComponentClass;
-
-	/** Constructing FText strings can be costly, so we cache the default menu name/category */
-	FNodeTextCache CachedMenuName;
-	FNodeTextCache CachedCategory;
 };
diff --git a/Engine/Source/Editor/BlueprintGraph/Classes/BlueprintDelegateNodeSpawner.h b/Engine/Source/Editor/BlueprintGraph/Classes/BlueprintDelegateNodeSpawner.h
--- a/Engine/Source/Editor/BlueprintGraph/Classes/BlueprintDelegateNodeSpawner.h
+++ b/Engine/Source/Editor/BlueprintGraph/Classes/BlueprintDelegateNodeSpawner.h
@@ -31,17 +31,11 @@ public:
 	 */
 	static UBlueprintDelegateNodeSpawner* Create(TSubclassOf<UK2Node_BaseMCDelegate> NodeClass, UMulticastDelegateProperty const* const Property, UObject* Outer = nullptr);
 
-	// UBlueprintNodeSpawner interface
-	virtual FText GetDefaultMenuName(FBindingSet const& Bindings) const override;
-	virtual FText GetDefaultMenuCategory() const override;
-	virtual FName GetDefaultMenuIcon(FLinearColor& ColorOut) const override;
-	// End UBlueprintNodeSpawner interface
-
 	/**
 	 * Accessor to the delegate property that this spawner wraps (the delegate
 	 * that this will assign spawned nodes with).
 	 *
 	 * @return The delegate property that this was initialized with.
 	 */
-	UMulticastDelegateProperty const* GetProperty() const;
+	UMulticastDelegateProperty const* GetDelegateProperty() const;
 };
diff --git a/Engine/Source/Editor/BlueprintGraph/Classes/BlueprintEventNodeSpawner.h b/Engine/Source/Editor/BlueprintGraph/Classes/BlueprintEventNodeSpawner.h
--- a/Engine/Source/Editor/BlueprintGraph/Classes/BlueprintEventNodeSpawner.h
+++ b/Engine/Source/Editor/BlueprintGraph/Classes/BlueprintEventNodeSpawner.h
@@ -44,25 +44,9 @@ public:
 	 */
 	static UBlueprintEventNodeSpawner* Create(TSubclassOf<UK2Node_Event> NodeClass, FName CustomEventName, UObject* Outer = nullptr);
 
-	/**
-	  * Creates a new UBlueprintEventNodeSpawner for custom events. The
-	  * CustomEventName can be left blank if the node will pick one itself on
-	  * instantiation.
-	  *
-	  * @param CustomEventName			The name you want assigned to the event.
-	  * @param CustomizeNodeDelegate	Delegate that will be executed once the node spawner has spawned a node
-	  * @param Outer					Optional outer for the new spawner (if left null, the transient package will be used).
-	  * @return A newly allocated instance of this class.
-	  */
-	static UBlueprintEventNodeSpawner* CreateWithDelegate(FName CustomEventName, FCustomizeNodeDelegate CustomizeNodeDelegate, FText CustomCategory = FText(), UObject* Outer = nullptr);
-
 	// UBlueprintNodeSpawner interface
 	virtual FBlueprintNodeSignature GetSpawnerSignature() const override;
 	virtual UEdGraphNode* Invoke(UEdGraph* ParentGraph, FBindingSet const& Bindings, FVector2D const Location) const override;
-	virtual FText GetDefaultMenuName(FBindingSet const& Bindings) const override;
-	virtual FText GetDefaultMenuCategory() const override;
-	virtual FString GetDefaultSearchKeywords() const override;
-	virtual FName GetDefaultMenuIcon(FLinearColor& ColorOut) const;
 	// End UBlueprintNodeSpawner interface
 
 	/**
@@ -87,15 +71,7 @@ public:
 	 */
 	virtual UK2Node_Event const* FindPreExistingEvent(UBlueprint* Blueprint, FBindingSet const& Bindings) const;
 
-protected:
-	/** Constructing FText strings can be costly, so we cache the default menu name */
-	FNodeTextCache CachedMenuName;
-
 private:
-	/** Custom category to place nodes in: */
-	UPROPERTY()
-	FText CustomCategory;
-
 	/** The function to configure new nodes with. */
 	UPROPERTY()
     UFunction const* EventFunc;
diff --git a/Engine/Source/Editor/BlueprintGraph/Classes/BlueprintFunctionNodeSpawner.h b/Engine/Source/Editor/BlueprintGraph/Classes/BlueprintFunctionNodeSpawner.h
--- a/Engine/Source/Editor/BlueprintGraph/Classes/BlueprintFunctionNodeSpawner.h
+++ b/Engine/Source/Editor/BlueprintGraph/Classes/BlueprintFunctionNodeSpawner.h
@@ -43,12 +43,8 @@ public:
 
 	// UBlueprintNodeSpawner interface
 	virtual void Prime() override;
+	virtual FBlueprintActionUiSpec GetUiSpec(FBlueprintActionContext const& Context, FBindingSet const& Bindings) const override;
 	virtual UEdGraphNode* Invoke(UEdGraph* ParentGraph, FBindingSet const& Bindings, FVector2D const Location) const override;
-	virtual FText GetDefaultMenuName(FBindingSet const& Bindings) const override;
-	virtual FText GetDefaultMenuCategory() const override;
-	virtual FText GetDefaultMenuTooltip() const override;
-	virtual FString GetDefaultSearchKeywords() const override;
-	virtual FName GetDefaultMenuIcon(FLinearColor& ColorOut) const;
 	// End UBlueprintNodeSpawner interface
 
 	// IBlueprintNodeBinder interface
diff --git a/Engine/Source/Editor/BlueprintGraph/Classes/BlueprintUnloadedAssetNodeSpawner.h b/Engine/Source/Editor/BlueprintGraph/Classes/BlueprintUnloadedAssetNodeSpawner.h
deleted file mode 100644
diff --git a/Engine/Source/Editor/BlueprintGraph/Classes/BlueprintVariableNodeSpawner.h b/Engine/Source/Editor/BlueprintGraph/Classes/BlueprintVariableNodeSpawner.h
--- a/Engine/Source/Editor/BlueprintGraph/Classes/BlueprintVariableNodeSpawner.h
+++ b/Engine/Source/Editor/BlueprintGraph/Classes/BlueprintVariableNodeSpawner.h
@@ -48,12 +48,8 @@ public:
 	// UBlueprintNodeSpawner interface
 	virtual void Prime() override;
 	virtual FBlueprintNodeSignature GetSpawnerSignature() const override;
+	virtual FBlueprintActionUiSpec GetUiSpec(FBlueprintActionContext const& Context, FBindingSet const& Bindings) const override;
 	virtual UEdGraphNode* Invoke(UEdGraph* ParentGraph, FBindingSet const& Bindings, FVector2D const Location) const override;
-	virtual FText GetDefaultMenuName(FBindingSet const& Bindings) const override;
-	virtual FText GetDefaultMenuCategory() const override;
-	virtual FText GetDefaultMenuTooltip() const override;
-	virtual FString GetDefaultSearchKeywords() const override;
-	virtual FName GetDefaultMenuIcon(FLinearColor& ColorOut) const override;
 	// End UBlueprintNodeSpawner interface
 	
 	/**
@@ -107,9 +103,4 @@ private:
 	/** */
 	UPROPERTY()
 	FBPVariableDescription LocalVarDesc;
-
-	/** Constructing FText strings can be costly, so we cache the default menu name/category/tooltip */
-	FNodeTextCache CachedMenuName;
-	FNodeTextCache CachedCategory;
-	FNodeTextCache CachedTooltip;
 };
\ No newline at end of file
diff --git a/Engine/Source/Editor/BlueprintGraph/Private/BlueprintActionDatabase.cpp b/Engine/Source/Editor/BlueprintGraph/Private/BlueprintActionDatabase.cpp
--- a/Engine/Source/Editor/BlueprintGraph/Private/BlueprintActionDatabase.cpp
+++ b/Engine/Source/Editor/BlueprintGraph/Private/BlueprintActionDatabase.cpp
@@ -18,6 +18,7 @@
 #include "AssetRegistryModule.h"	// for OnAssetAdded()/OnAssetRemoved()
 #include "BlueprintEditorUtils.h"	// for FindBlueprintForGraph()
 #include "BlueprintActionDatabaseRegistrar.h"
+#include "BlueprintActionFilter.h"	// for FBlueprintActionContext
 
 // used below in FBlueprintNodeSpawnerFactory::MakeMacroNodeSpawner()
 #include "K2Node_MacroInstance.h"
@@ -80,6 +81,13 @@ namespace FBlueprintNodeSpawnerFactory
 	template <class DocNodeType>
 	static UBlueprintNodeSpawner* MakeDocumentationNodeSpawner();
 
+	/**
+	 * 
+	 * 
+	 * @return 
+	 */
+	static UBlueprintNodeSpawner* MakeCommentNodeSpawner();
+
 	/**
 	 * Constructs a delegate binding node along with a connected event that is  
 	 * triggered from the specified delegate.
@@ -190,6 +198,28 @@ static UBlueprintNodeSpawner* FBlueprintNodeSpawnerFactory::MakeDocumentationNod
 	return NodeSpawner;
 }
 
+//------------------------------------------------------------------------------
+static UBlueprintNodeSpawner* FBlueprintNodeSpawnerFactory::MakeCommentNodeSpawner()
+{
+	UBlueprintNodeSpawner* NodeSpawner = MakeDocumentationNodeSpawner<UEdGraphNode_Comment>();
+	NodeSpawner->DefaultMenuSignature.MenuName = LOCTEXT("AddCommentActionMenuName", "Add Comment...");
+
+	auto OverrideMenuNameLambda = [](FBlueprintActionContext const& Context, IBlueprintNodeBinder::FBindingSet const& /*Bindings*/, FBlueprintActionUiSpec* UiSpecOut)
+	{
+		for (UBlueprint* Blueprint : Context.Blueprints)
+		{
+			if (FKismetEditorUtilities::GetNumberOfSelectedNodes(Blueprint) > 0)
+			{
+				UiSpecOut->MenuName = LOCTEXT("AddCommentFromSelectionMenuName", "Add Comment to Selection");
+				break;
+			}
+		}
+	};
+	NodeSpawner->DynamicUiSignatureGetter = UBlueprintNodeSpawner::FUiSpecOverrideDelegate::CreateStatic(OverrideMenuNameLambda);
+
+	return NodeSpawner;
+}
+
 //------------------------------------------------------------------------------
 static UBlueprintNodeSpawner* FBlueprintNodeSpawnerFactory::MakeAssignDelegateNodeSpawner(UMulticastDelegateProperty* DelegateProperty)
 {
@@ -219,7 +249,11 @@ static UBlueprintNodeSpawner* FBlueprintNodeSpawnerFactory::MakeAnimOwnedEventSp
 		ActorRefNode->EventSignatureClass = UAnimInstance::StaticClass();
 	};
 
-	return UBlueprintEventNodeSpawner::CreateWithDelegate(SignatureName, UBlueprintNodeSpawner::FCustomizeNodeDelegate::CreateStatic( PostSpawnSetupLambda ), CustomCategory );
+	UBlueprintNodeSpawner* NodeSpawner = UBlueprintEventNodeSpawner::Create(UK2Node_Event::StaticClass(), SignatureName);
+	NodeSpawner->CustomizeNodeDelegate = UBlueprintNodeSpawner::FCustomizeNodeDelegate::CreateStatic(PostSpawnSetupLambda);
+	NodeSpawner->DefaultMenuSignature.Category = CustomCategory;
+
+	return NodeSpawner;
 }
 
 /*******************************************************************************
@@ -664,7 +698,7 @@ static void BlueprintActionDatabaseImpl::AddAnimBlueprintGraphActions(UAnimBluep
 //------------------------------------------------------------------------------
 static void BlueprintActionDatabaseImpl::GetNodeSpecificActions(TSubclassOf<UEdGraphNode const> const NodeClass, FBlueprintActionDatabaseRegistrar& Registrar)
 {
-	using namespace FBlueprintNodeSpawnerFactory; // for MakeDocumentationNodeSpawner()
+	using namespace FBlueprintNodeSpawnerFactory; // for MakeCommentNodeSpawner()/MakeDocumentationNodeSpawner()
 
 	if (NodeClass->IsChildOf<UK2Node>() && !NodeClass->HasAnyClassFlags(CLASS_Abstract))
 	{
@@ -680,7 +714,7 @@ static void BlueprintActionDatabaseImpl::GetNodeSpecificActions(TSubclassOf<UEdG
 	//        with a better (more generalized) solution.
 	else if (NodeClass == UEdGraphNode_Comment::StaticClass())
 	{
-		Registrar.AddBlueprintAction(MakeDocumentationNodeSpawner<UEdGraphNode_Comment>());
+		Registrar.AddBlueprintAction(MakeCommentNodeSpawner());
 	}
 	else if (NodeClass == UEdGraphNode_Documentation::StaticClass())
 	{
diff --git a/Engine/Source/Editor/BlueprintGraph/Private/BlueprintBoundEventNodeSpawner.cpp b/Engine/Source/Editor/BlueprintGraph/Private/BlueprintBoundEventNodeSpawner.cpp
--- a/Engine/Source/Editor/BlueprintGraph/Private/BlueprintBoundEventNodeSpawner.cpp
+++ b/Engine/Source/Editor/BlueprintGraph/Private/BlueprintBoundEventNodeSpawner.cpp
@@ -8,6 +8,36 @@
 
 #define LOCTEXT_NAMESPACE "BlueprintBoundEventNodeSpawner"
 
+/*******************************************************************************
+ * Static UBlueprintBoundEventNodeSpawner Helpers
+ ******************************************************************************/
+
+namespace BlueprintBoundEventNodeSpawnerImpl
+{
+	static FText GetDefaultMenuName(UMulticastDelegateProperty const* Delegate);
+	static FText GetDefaultMenuCategory(UMulticastDelegateProperty const* Delegate);
+}
+
+//------------------------------------------------------------------------------
+static FText BlueprintBoundEventNodeSpawnerImpl::GetDefaultMenuName(UMulticastDelegateProperty const* Delegate)
+{
+	bool const bShowFriendlyNames = GetDefault<UEditorStyleSettings>()->bShowFriendlyNames;
+	FText const DelegateName = bShowFriendlyNames ? FText::FromString(UEditorEngine::GetFriendlyName(Delegate)) : FText::FromName(Delegate->GetFName());
+
+	return FText::Format(LOCTEXT("ComponentEventName", "Add {0}"), DelegateName);
+}
+
+//------------------------------------------------------------------------------
+static FText BlueprintBoundEventNodeSpawnerImpl::GetDefaultMenuCategory(UMulticastDelegateProperty const* Delegate)
+{
+	FText DelegateCategory = FText::FromString(FObjectEditorUtils::GetCategory(Delegate));
+	if (DelegateCategory.IsEmpty())
+	{
+		DelegateCategory = FEditorCategoryUtils::GetCommonCategory(FCommonEditorCategory::Delegates);
+	}
+	return DelegateCategory;
+}
+
 /*******************************************************************************
  * UBlueprintBoundEventNodeSpawner
  ******************************************************************************/
@@ -24,6 +54,13 @@ UBlueprintBoundEventNodeSpawner* UBlueprintBoundEventNodeSpawner::Create(TSubcla
 	NodeSpawner->NodeClass     = NodeClass;
 	NodeSpawner->EventDelegate = EventDelegate;
 
+	FBlueprintActionUiSpec& MenuSignature = NodeSpawner->DefaultMenuSignature;
+	MenuSignature.MenuName = BlueprintBoundEventNodeSpawnerImpl::GetDefaultMenuName(EventDelegate);
+	MenuSignature.Category = BlueprintBoundEventNodeSpawnerImpl::GetDefaultMenuCategory(EventDelegate);
+	//MenuSignature.Tooltip,  will be pulled from the node template
+	//MenuSignature.Keywords, will be pulled from the node template
+	MenuSignature.IconName = TEXT("GraphEditor.Event_16x");
+
 	return NodeSpawner;
 }
 
@@ -54,36 +91,6 @@ UEdGraphNode* UBlueprintBoundEventNodeSpawner::Invoke(UEdGraph* ParentGraph, FBi
 	return EventNode;
 }
 
-//------------------------------------------------------------------------------
-FText UBlueprintBoundEventNodeSpawner::GetDefaultMenuName(FBindingSet const& Bindings) const
-{
-	if (CachedMenuName.IsOutOfDate())
-	{
-		FText const DelegateName = FText::FromName(EventDelegate->GetFName());
-		// FText::Format() is slow, so we cache this to save on performance
-		CachedMenuName = FText::Format(LOCTEXT("ComponentEventName", "Add {0}"), DelegateName);
-	}
-	return CachedMenuName;
-}
-
-//------------------------------------------------------------------------------
-FText UBlueprintBoundEventNodeSpawner::GetDefaultMenuCategory() const
-{
-	FText DelegateCategory = FEditorCategoryUtils::GetCommonCategory(FCommonEditorCategory::Delegates);
-	if (UMulticastDelegateProperty const* Delegate = GetEventDelegate())
-	{
-		DelegateCategory = FText::FromString(FObjectEditorUtils::GetCategory(Delegate));
-	}
-	return DelegateCategory;
-}
-
-//------------------------------------------------------------------------------
-FName UBlueprintBoundEventNodeSpawner::GetDefaultMenuIcon(FLinearColor& ColorOut) const
-{
-	ColorOut = FLinearColor::White;
-	return TEXT("GraphEditor.Event_16x");
-}
-
 //------------------------------------------------------------------------------
 UK2Node_Event const* UBlueprintBoundEventNodeSpawner::FindPreExistingEvent(UBlueprint* Blueprint, FBindingSet const& Bindings) const
 {
diff --git a/Engine/Source/Editor/BlueprintGraph/Private/BlueprintBoundNodeSpawner.cpp b/Engine/Source/Editor/BlueprintGraph/Private/BlueprintBoundNodeSpawner.cpp
--- a/Engine/Source/Editor/BlueprintGraph/Private/BlueprintBoundNodeSpawner.cpp
+++ b/Engine/Source/Editor/BlueprintGraph/Private/BlueprintBoundNodeSpawner.cpp
@@ -38,16 +38,6 @@ FBlueprintNodeSignature UBlueprintBoundNodeSpawner::GetSpawnerSignature() const
 	return FBlueprintNodeSignature();
 }
 
-//------------------------------------------------------------------------------
-FText UBlueprintBoundNodeSpawner::GetDefaultMenuName(FBindingSet const& Bindings) const
-{
-	if( OnGenerateMenuDescriptionDelegate.IsBound() )
-	{
-		return OnGenerateMenuDescriptionDelegate.Execute(Bindings);
-	}
-	return UBlueprintNodeSpawner::GetDefaultMenuName( Bindings );
-}
-
 //------------------------------------------------------------------------------
 bool UBlueprintBoundNodeSpawner::IsBindingCompatible(UObject const* BindingCandidate) const
 {
diff --git a/Engine/Source/Editor/BlueprintGraph/Private/BlueprintComponentNodeSpawner.cpp b/Engine/Source/Editor/BlueprintGraph/Private/BlueprintComponentNodeSpawner.cpp
--- a/Engine/Source/Editor/BlueprintGraph/Private/BlueprintComponentNodeSpawner.cpp
+++ b/Engine/Source/Editor/BlueprintGraph/Private/BlueprintComponentNodeSpawner.cpp
@@ -9,6 +9,35 @@
 
 #define LOCTEXT_NAMESPACE "BlueprintComponenetNodeSpawner"
 
+/*******************************************************************************
+ * Static UBlueprintComponentNodeSpawner Helpers
+ ******************************************************************************/
+
+namespace BlueprintComponentNodeSpawnerImpl
+{
+	static FText GetDefaultMenuCategory(TSubclassOf<UActorComponent> const ComponentClass);
+}
+
+//------------------------------------------------------------------------------
+static FText BlueprintComponentNodeSpawnerImpl::GetDefaultMenuCategory(TSubclassOf<UActorComponent> const ComponentClass)
+{
+	FText ClassGroup;
+	TArray<FString> ClassGroupNames;
+	ComponentClass->GetClassGroupNames(ClassGroupNames);
+
+	static FText const DefaultClassGroup(LOCTEXT("DefaultClassGroup", "Common"));
+	// 'Common' takes priority over other class groups
+	if (ClassGroupNames.Contains(DefaultClassGroup.ToString()) || (ClassGroupNames.Num() == 0))
+	{
+		ClassGroup = DefaultClassGroup;
+	}
+	else
+	{
+		ClassGroup = FText::FromString(ClassGroupNames[0]);
+	}
+	return FText::Format(LOCTEXT("ComponentCategory", "Add Component|{0}"), ClassGroup);
+}
+
 /*******************************************************************************
  * UBlueprintComponentNodeSpawner
  ******************************************************************************/
@@ -27,6 +56,17 @@ UBlueprintComponentNodeSpawner* UBlueprintComponentNodeSpawner::Create(TSubclass
 	NodeSpawner->ComponentClass = ComponentClass;
 	NodeSpawner->NodeClass      = UK2Node_AddComponent::StaticClass();
 
+	FBlueprintActionUiSpec& MenuSignature = NodeSpawner->DefaultMenuSignature;
+	FText const ComponentTypeName = FText::FromName(ComponentClass->GetFName());
+	MenuSignature.MenuName = FText::Format(LOCTEXT("AddComponentMenuName", "Add {0}"), ComponentTypeName);
+	MenuSignature.Category = BlueprintComponentNodeSpawnerImpl::GetDefaultMenuCategory(ComponentClass);
+	MenuSignature.Tooltip  = FText::Format(LOCTEXT("AddComponentTooltip", "Spawn a {0}"), ComponentTypeName);
+	MenuSignature.Keywords = ComponentClass->GetMetaData(FBlueprintMetadata::MD_FunctionKeywords);
+	// add at least one character, so that PrimeDefaultMenuSignature() doesn't 
+	// attempt to query the template node
+	MenuSignature.Keywords.AppendChar(TEXT(' '));
+	MenuSignature.IconName = FClassIconFinder::FindIconNameForClass(ComponentClass);
+
 	return NodeSpawner;
 }
 
@@ -97,15 +137,10 @@ UEdGraphNode* UBlueprintComponentNodeSpawner::Invoke(UEdGraph* ParentGraph, FBin
 }
 
 //------------------------------------------------------------------------------
-FText UBlueprintComponentNodeSpawner::GetDefaultMenuName(FBindingSet const& Bindings) const
+FBlueprintActionUiSpec UBlueprintComponentNodeSpawner::GetUiSpec(FBlueprintActionContext const& Context, FBindingSet const& Bindings) const
 {
-	check(ComponentClass != nullptr);
-	if (CachedMenuName.IsOutOfDate())
-	{
-		// FText::Format() is slow, so we cache this to save on performance
-		CachedMenuName = FText::Format(LOCTEXT("AddComponentMenuName", "Add {0}"), FText::FromName(ComponentClass->GetFName()));
-	}
-	FText MenuName = CachedMenuName.GetCachedText();
+	UEdGraph* TargetGraph = (Context.Graphs.Num() > 0) ? Context.Graphs[0] : nullptr;
+	FBlueprintActionUiSpec MenuSignature = PrimeDefaultUiSpec(TargetGraph);
 
 	if (Bindings.Num() > 0)
 	{
@@ -116,61 +151,11 @@ FText UBlueprintComponentNodeSpawner::GetDefaultMenuName(FBindingSet const& Bind
 		}
 
 		FText const ComponentTypeName = FText::FromName(ComponentClass->GetFName());
-		MenuName = FText::Format(LOCTEXT("AddBoundComponentMenuName", "Add {0} (as {1})"), AssetName, ComponentTypeName);
+		MenuSignature.MenuName = FText::Format(LOCTEXT("AddBoundComponentMenuName", "Add {0} (as {1})"), AssetName, ComponentTypeName);
+		MenuSignature.Tooltip  = FText::Format(LOCTEXT("AddBoundComponentTooltip", "Spawn {0} using {1}"), ComponentTypeName, AssetName);
 	}
-	return MenuName;
-}
-
-//------------------------------------------------------------------------------
-FText UBlueprintComponentNodeSpawner::GetDefaultMenuTooltip() const
-{
-	FText const ComponentTypeName = FText::FromName(ComponentClass->GetFName());
-	// @TODO: consider caching this, and provide a separate tooltip for when the
-	//        node would be bound ("Spawn a {0} using {1}").
-	return FText::Format(LOCTEXT("AddComponentTooltip", "Spawn a {0}"), ComponentTypeName);
-}
-
-//------------------------------------------------------------------------------
-FText UBlueprintComponentNodeSpawner::GetDefaultMenuCategory() const
-{
-	check(ComponentClass != nullptr);
-
-	if (CachedCategory.IsOutOfDate())
-	{
-		FText ClassGroup;
-		TArray<FString> ClassGroupNames;
-		ComponentClass->GetClassGroupNames(ClassGroupNames);
-
-		static FText const DefaultClassGroup(LOCTEXT("DefaultClassGroup", "Common"));
-		// 'Common' takes priority over other class groups
-		if (ClassGroupNames.Contains(DefaultClassGroup.ToString()) || (ClassGroupNames.Num() == 0))
-		{
-			ClassGroup = DefaultClassGroup;
-		}
-		else
-		{
-			ClassGroup = FText::FromString(ClassGroupNames[0]);
-		}
-		// FText::Format() is slow, so we cache this to save on performance
-		CachedCategory = FText::Format(LOCTEXT("ComponentCategory", "Add Component|{0}"), ClassGroup);
-	}
-	return CachedCategory;
-}
-
-//------------------------------------------------------------------------------
-FString UBlueprintComponentNodeSpawner::GetDefaultSearchKeywords() const
-{
-	FString SearchKeywords = ComponentClass->GetMetaData(FBlueprintMetadata::MD_FunctionKeywords);
-	// add at least one character, so that the menu item doesn't attempt to
-	// ping a template node
-	return SearchKeywords.AppendChar(TEXT(' '));
-}
-
-//------------------------------------------------------------------------------
-FName UBlueprintComponentNodeSpawner::GetDefaultMenuIcon(FLinearColor& ColorOut) const
-{
-	check(ComponentClass != nullptr);
-	return FClassIconFinder::FindIconNameForClass(ComponentClass);
+	DynamicUiSignatureGetter.ExecuteIfBound(Context, Bindings, &MenuSignature);
+	return MenuSignature;
 }
 
 //------------------------------------------------------------------------------
diff --git a/Engine/Source/Editor/BlueprintGraph/Private/BlueprintDelegateNodeSpawner.cpp b/Engine/Source/Editor/BlueprintGraph/Private/BlueprintDelegateNodeSpawner.cpp
--- a/Engine/Source/Editor/BlueprintGraph/Private/BlueprintDelegateNodeSpawner.cpp
+++ b/Engine/Source/Editor/BlueprintGraph/Private/BlueprintDelegateNodeSpawner.cpp
@@ -10,6 +10,44 @@
 
 #define LOCTEXT_NAMESPACE "BlueprintDelegateNodeSpawner"
 
+/*******************************************************************************
+ * Static UBlueprintDelegateNodeSpawner Helpers
+ ******************************************************************************/
+
+namespace BlueprintDelegateNodeSpawnerImpl
+{
+	static FText GetDefaultMenuName(UMulticastDelegateProperty const* Delegate);
+	static FText GetDefaultMenuCategory(UMulticastDelegateProperty const* Delegate);
+	static FName GetDefaultMenuIcon(UMulticastDelegateProperty const* Delegate, FLinearColor& ColorOut);
+}
+
+//------------------------------------------------------------------------------
+static FText BlueprintDelegateNodeSpawnerImpl::GetDefaultMenuName(UMulticastDelegateProperty const* Delegate)
+{
+	bool const bShowFriendlyNames = GetDefault<UEditorStyleSettings>()->bShowFriendlyNames;
+	return bShowFriendlyNames ? FText::FromString(UEditorEngine::GetFriendlyName(Delegate)) : FText::FromName(Delegate->GetFName());
+}
+
+//------------------------------------------------------------------------------
+static FText BlueprintDelegateNodeSpawnerImpl::GetDefaultMenuCategory(UMulticastDelegateProperty const* Delegate)
+{
+	FText DelegateCategory = FText::FromString(FObjectEditorUtils::GetCategory(Delegate));
+	if (DelegateCategory.IsEmpty())
+	{
+		DelegateCategory = FEditorCategoryUtils::GetCommonCategory(FCommonEditorCategory::Delegates);
+	}
+	return DelegateCategory;
+}
+
+//------------------------------------------------------------------------------
+static FName BlueprintDelegateNodeSpawnerImpl::GetDefaultMenuIcon(UMulticastDelegateProperty const* Delegate, FLinearColor& ColorOut)
+{
+	FName    const PropertyName = Delegate->GetFName();
+	UStruct* const PropertyOwner = CastChecked<UStruct>(Delegate->GetOuterUField());
+
+	return UK2Node_Variable::GetVariableIconAndColor(PropertyOwner, PropertyName, ColorOut);
+}
+
 /*******************************************************************************
  * UBlueprintDelegateNodeSpawner
  ******************************************************************************/
@@ -23,10 +61,29 @@ UBlueprintDelegateNodeSpawner* UBlueprintDelegateNodeSpawner::Create(TSubclassOf
 		Outer = GetTransientPackage();
 	}
 
+	//--------------------------------------
+	// Constructing the Spawner
+	//--------------------------------------
+
 	UBlueprintDelegateNodeSpawner* NodeSpawner = NewObject<UBlueprintDelegateNodeSpawner>(Outer);
 	NodeSpawner->Field     = Property;
 	NodeSpawner->NodeClass = NodeClass;
 
+	//--------------------------------------
+	// Default UI Signature
+	//--------------------------------------
+
+	FBlueprintActionUiSpec& MenuSignature = NodeSpawner->DefaultMenuSignature;
+	MenuSignature.MenuName = BlueprintDelegateNodeSpawnerImpl::GetDefaultMenuName(Property);
+	MenuSignature.Category = BlueprintDelegateNodeSpawnerImpl::GetDefaultMenuCategory(Property);
+	//MenuSignature.Tooltip  =
+	//MenuSignature.Keywords = 
+	MenuSignature.IconName = BlueprintDelegateNodeSpawnerImpl::GetDefaultMenuIcon(Property, MenuSignature.IconTint);
+
+	//--------------------------------------
+	// Post-Spawn Setup
+	//--------------------------------------
+
 	auto SetDelegateLambda = [](UEdGraphNode* NewNode, UField const* Field)
 	{
 		UMulticastDelegateProperty const* Property = Cast<UMulticastDelegateProperty>(Field);
@@ -53,56 +110,7 @@ UBlueprintDelegateNodeSpawner::UBlueprintDelegateNodeSpawner(class FPostConstruc
 }
 
 //------------------------------------------------------------------------------
-FText UBlueprintDelegateNodeSpawner::GetDefaultMenuName(FBindingSet const& Bindings) const
-{	
-	FText MenuName = Super::GetDefaultMenuName(Bindings);
-
-	UMulticastDelegateProperty const* Property = GetProperty();
-	if ((NodeClass == nullptr) && (Property != nullptr))
-	{
-		bool const bShowFriendlyNames = GetDefault<UEditorStyleSettings>()->bShowFriendlyNames;
-		MenuName = bShowFriendlyNames ? FText::FromString(UEditorEngine::GetFriendlyName(Property)) : FText::FromName(Property->GetFName());
-	}
-
-	return MenuName;
-}
-
-//------------------------------------------------------------------------------
-FText UBlueprintDelegateNodeSpawner::GetDefaultMenuCategory() const
-{
-	FText PropertyCategory;
-	if (UMulticastDelegateProperty const* Property = GetProperty())
-	{
-		PropertyCategory = FText::FromString(FObjectEditorUtils::GetCategory(Property));
-	}
-
-	if (PropertyCategory.IsEmpty())
-	{
-		PropertyCategory = FEditorCategoryUtils::GetCommonCategory(FCommonEditorCategory::Delegates);
-	}
-	return PropertyCategory;
-}
-
-//------------------------------------------------------------------------------
-FName UBlueprintDelegateNodeSpawner::GetDefaultMenuIcon(FLinearColor& ColorOut) const
-{
-	FName BrushName = Super::GetDefaultMenuIcon(ColorOut);
-	if (NodeClass == nullptr)
-	{
-		BrushName = TEXT("GraphEditor.Delegate_16x");
-	}
-	else if (UMulticastDelegateProperty const* Property = GetProperty())
-	{
-		FName    const PropertyName  = Property->GetFName();
-		UStruct* const PropertyOwner = CastChecked<UStruct>(Property->GetOuterUField());
-
-		BrushName = UK2Node_Variable::GetVariableIconAndColor(PropertyOwner, PropertyName, ColorOut);
-	}
-	return BrushName;
-}
-
-//------------------------------------------------------------------------------
-UMulticastDelegateProperty const* UBlueprintDelegateNodeSpawner::GetProperty() const
+UMulticastDelegateProperty const* UBlueprintDelegateNodeSpawner::GetDelegateProperty() const
 {
 	return Cast<UMulticastDelegateProperty>(GetField());
 }
diff --git a/Engine/Source/Editor/BlueprintGraph/Private/BlueprintEventNodeSpawner.cpp b/Engine/Source/Editor/BlueprintGraph/Private/BlueprintEventNodeSpawner.cpp
--- a/Engine/Source/Editor/BlueprintGraph/Private/BlueprintEventNodeSpawner.cpp
+++ b/Engine/Source/Editor/BlueprintGraph/Private/BlueprintEventNodeSpawner.cpp
@@ -63,6 +63,15 @@ UBlueprintEventNodeSpawner* UBlueprintEventNodeSpawner::Create(UFunction const*
 	NodeSpawner->EventFunc = EventFunc;
 	NodeSpawner->NodeClass = UK2Node_Event::StaticClass();
 
+	FBlueprintActionUiSpec& MenuSignature = NodeSpawner->DefaultMenuSignature;
+	FString const FuncName = UEdGraphSchema_K2::GetFriendlySignatureName(EventFunc);
+	MenuSignature.MenuName = FText::Format(LOCTEXT("EventWithSignatureName", "Event {0}"), FText::FromString(FuncName));
+	FString const FuncCategory = UK2Node_CallFunction::GetDefaultCategoryForFunction(EventFunc, TEXT(""));
+	MenuSignature.Category = FText::FromString(LOCTEXT("AddEventCategory", "Add Event").ToString() + TEXT("|") + FuncCategory);
+	//MenuSignature.Tooltip, will be pulled from the node template
+	MenuSignature.Keywords = UK2Node_CallFunction::GetKeywordsForFunction(EventFunc).AppendChar(TEXT(' '));
+	MenuSignature.IconName = TEXT("GraphEditor.Event_16x");
+
 	return NodeSpawner;
 }
 
@@ -78,22 +87,22 @@ UBlueprintEventNodeSpawner* UBlueprintEventNodeSpawner::Create(TSubclassOf<UK2No
 	NodeSpawner->NodeClass       = NodeClass;
 	NodeSpawner->CustomEventName = CustomEventName;
 
-	return NodeSpawner;
-}
-
-//------------------------------------------------------------------------------
-UBlueprintEventNodeSpawner* UBlueprintEventNodeSpawner::CreateWithDelegate(FName CustomEventName, FCustomizeNodeDelegate CustomizeNodeDelegate, FText CustomCategory /*= FText()*/, UObject* Outer /*= nullptr*/)
-{
-	if (Outer == nullptr)
+	FBlueprintActionUiSpec& MenuSignature = NodeSpawner->DefaultMenuSignature;
+	if (CustomEventName.IsNone())
 	{
-		Outer = GetTransientPackage();
+		MenuSignature.MenuName = LOCTEXT("AddCustomEvent", "Add Custom Event...");
+		MenuSignature.IconName = TEXT("GraphEditor.CustomEvent_16x");
+	}
+	else
+	{
+		FText const EventName = FText::FromName(CustomEventName);
+		MenuSignature.MenuName = FText::Format(LOCTEXT("EventWithSignatureName", "Event {0}"), EventName);
+		MenuSignature.IconName = TEXT("GraphEditor.Event_16x");
 	}
+	//MenuSignature.Category, will be pulled from the node template
+	//MenuSignature.Tooltip,  will be pulled from the node template 
+	//MenuSignature.Keywords, will be pulled from the node template
 
-	UBlueprintEventNodeSpawner* NodeSpawner = NewObject<UBlueprintEventNodeSpawner>(Outer);
-	NodeSpawner->NodeClass		= UK2Node_Event::StaticClass();
-	NodeSpawner->CustomCategory = CustomCategory;
-	NodeSpawner->CustomEventName = CustomEventName;
-	NodeSpawner->CustomizeNodeDelegate = CustomizeNodeDelegate;
 	return NodeSpawner;
 }
 
@@ -169,60 +178,6 @@ UEdGraphNode* UBlueprintEventNodeSpawner::Invoke(UEdGraph* ParentGraph, FBinding
 	return EventNode;
 }
 
-//------------------------------------------------------------------------------
-FText UBlueprintEventNodeSpawner::GetDefaultMenuName(FBindingSet const& Bindings) const
-{
-	if (CachedMenuName.IsOutOfDate())
-	{
-		FText EventName;
-		if (EventFunc != nullptr)
-		{
-			EventName = FText::FromString(UEdGraphSchema_K2::GetFriendlySignatureName(EventFunc));
-		}
-		else if (!CustomEventName.IsNone())
-		{
-			EventName = FText::FromName(CustomEventName);
-		}
-		else
-		{
-			CachedMenuName = LOCTEXT("AddCustomEvent", "Add Custom Event...");
-		}
-
-		if(CachedMenuName.IsOutOfDate())
-		{
-			CachedMenuName = FText::Format(LOCTEXT("EventWithSignatureName", "Event {0}"), EventName);
-		}
-	}
-	return CachedMenuName;
-}
-
-//------------------------------------------------------------------------------
-FText UBlueprintEventNodeSpawner::GetDefaultMenuCategory() const
-{
-	// certain events can have specialized categories (like input events), so 
-	// we choose to leave this in the node's hands (returns an empty text string
-	// so that the menu builder polls the node instead).
-	return CustomCategory;
-}
-
-//------------------------------------------------------------------------------
-FString UBlueprintEventNodeSpawner::GetDefaultSearchKeywords() const
-{
-	FString Keywords;
-	if (EventFunc != nullptr)
-	{
-		Keywords = UK2Node_CallFunction::GetKeywordsForFunction(EventFunc);
-	}
-	return Keywords;
-}
-
-//------------------------------------------------------------------------------
-FName UBlueprintEventNodeSpawner::GetDefaultMenuIcon(FLinearColor& ColorOut) const
-{
-	ColorOut = FLinearColor::White;
-	return (IsForCustomEvent() && CustomEventName.IsNone()) ? TEXT("GraphEditor.CustomEvent_16x") : TEXT("GraphEditor.Event_16x");
-}
-
 //------------------------------------------------------------------------------
 UFunction const* UBlueprintEventNodeSpawner::GetEventFunction() const
 {
diff --git a/Engine/Source/Editor/BlueprintGraph/Private/BlueprintFunctionNodeSpawner.cpp b/Engine/Source/Editor/BlueprintGraph/Private/BlueprintFunctionNodeSpawner.cpp
--- a/Engine/Source/Editor/BlueprintGraph/Private/BlueprintFunctionNodeSpawner.cpp
+++ b/Engine/Source/Editor/BlueprintGraph/Private/BlueprintFunctionNodeSpawner.cpp
@@ -6,6 +6,8 @@
 #include "GameFramework/Actor.h"
 #include "BlueprintVariableNodeSpawner.h"
 #include "BlueprintNodeTemplateCache.h" // for IsTemplateOuter()
+#include "BlueprintActionFilter.h" // for FBlueprintActionContext
+#include "EditorCategoryUtils.h" // for BuildCategoryString()
 
 #define LOCTEXT_NAMESPACE "BlueprintFunctionNodeSpawner"
 
@@ -17,6 +19,7 @@
 namespace BlueprintFunctionNodeSpawnerImpl
 {
 	FVector2D BindingOffset = FVector2D::ZeroVector;
+	static const FText FallbackCategory = LOCTEXT("UncategorizedFallbackCategory", "Call Function");
 
 	/**
 	 * 
@@ -181,6 +184,11 @@ UBlueprintFunctionNodeSpawner* UBlueprintFunctionNodeSpawner::Create(TSubclassOf
 	{
 		Outer = GetTransientPackage();
 	}
+
+	//--------------------------------------
+	// Constructing the Spawner
+	//--------------------------------------
+
 	UBlueprintFunctionNodeSpawner* NodeSpawner = NewObject<UBlueprintFunctionNodeSpawner>(Outer);
 	NodeSpawner->Field = Function;
 
@@ -193,6 +201,32 @@ UBlueprintFunctionNodeSpawner* UBlueprintFunctionNodeSpawner::Create(TSubclassOf
 		NodeSpawner->NodeClass = NodeClass;
 	}
 
+	//--------------------------------------
+	// Default UI Signature
+	//--------------------------------------
+
+	FBlueprintActionUiSpec& MenuSignature = NodeSpawner->DefaultMenuSignature;
+	MenuSignature.MenuName = FText::FromString( UK2Node_CallFunction::GetUserFacingFunctionName(Function) );
+	MenuSignature.Category = FText::FromString( UK2Node_CallFunction::GetDefaultCategoryForFunction(Function, TEXT("")) );
+	MenuSignature.Tooltip  = FText::FromString( UK2Node_CallFunction::GetDefaultTooltipForFunction(Function) );
+	// add at least one character, so that PrimeDefaultMenuSignature() doesn't attempt to query the template node
+	MenuSignature.Keywords = UK2Node_CallFunction::GetKeywordsForFunction(Function).AppendChar(TEXT(' '));
+	MenuSignature.IconName = UK2Node_CallFunction::GetPaletteIconForFunction(Function, MenuSignature.IconTint);
+
+	if (MenuSignature.Category.IsEmpty())
+	{
+		MenuSignature.Category = BlueprintFunctionNodeSpawnerImpl::FallbackCategory;
+	}
+
+	if (MenuSignature.Tooltip.IsEmpty())
+	{
+		MenuSignature.Tooltip = MenuSignature.MenuName;
+	}
+
+	//--------------------------------------
+	// Post-Spawn Setup
+	//--------------------------------------
+
 	auto SetNodeFunctionLambda = [](UEdGraphNode* NewNode, UField const* Field)
 	{
 		// user could have changed the node class (to something like
@@ -204,6 +238,7 @@ UBlueprintFunctionNodeSpawner* UBlueprintFunctionNodeSpawner::Create(TSubclassOf
 	};
 	NodeSpawner->SetNodeFieldDelegate = FSetNodeFieldDelegate::CreateStatic(SetNodeFunctionLambda);
 
+
 	return NodeSpawner;
 }
 
@@ -218,92 +253,78 @@ void UBlueprintFunctionNodeSpawner::Prime()
 {
 	// we expect that you don't need a node template to construct menu entries
 	// from this, so we choose not to pre-cache one here
-
-	// we don't have any expensive FText::Format() constructed strings to cache
-	// either
 }
 
 //------------------------------------------------------------------------------
-UEdGraphNode* UBlueprintFunctionNodeSpawner::Invoke(UEdGraph* ParentGraph, FBindingSet const& Bindings, FVector2D const Location) const
+FBlueprintActionUiSpec UBlueprintFunctionNodeSpawner::GetUiSpec(FBlueprintActionContext const& Context, FBindingSet const& Bindings) const
 {
-	auto PostSpawnSetupLambda = [](UEdGraphNode* NewNode, bool bIsTemplateNode, UFunction const* Function, FSetNodeFieldDelegate SetFieldDelegate, FCustomizeNodeDelegate UserDelegate)
-	{
-		SetFieldDelegate.ExecuteIfBound(NewNode, Function);
-		UserDelegate.ExecuteIfBound(NewNode, bIsTemplateNode);
-	};
-	FCustomizeNodeDelegate PostSpawnSetupDelegate = FCustomizeNodeDelegate::CreateStatic(PostSpawnSetupLambda, GetFunction(), SetNodeFieldDelegate, CustomizeNodeDelegate);
+	UEdGraph* TargetGraph = (Context.Graphs.Num() > 0) ? Context.Graphs[0] : nullptr;
+	FBlueprintActionUiSpec MenuSignature = PrimeDefaultUiSpec(TargetGraph);
 
-	UClass* SpawnClass = NodeClass;
-	if ((Bindings.Num() == 1) && (*Bindings.CreateConstIterator())->IsA<UObjectProperty>())
+	FString const CategoryString = MenuSignature.Category.ToString();
+	// FText compares are slow, so let's use FString compares
+	bool const bIsUncategorized = (CategoryString == BlueprintFunctionNodeSpawnerImpl::FallbackCategory.ToString());
+	if (bIsUncategorized)
 	{
-		SpawnClass = UK2Node_CallFunctionOnMember::StaticClass();
-	}
-	// if this spawner was set up to spawn a bound node, reset this so the 
-	// bound nodes get positioned properly
-	BlueprintFunctionNodeSpawnerImpl::BindingOffset = FVector2D::ZeroVector;
+		checkSlow(Context.Blueprints.Num() > 0);
+		UBlueprint* TargetBlueprint = Context.Blueprints[0];
 
-	return Super::SpawnNode(SpawnClass, ParentGraph, Bindings, Location, PostSpawnSetupDelegate);
-}
+		// @TODO: this is duplicated in a couple places, move it to some shared resource
+		UClass const* TargetClass = (TargetBlueprint->GeneratedClass != nullptr) ? TargetBlueprint->GeneratedClass : TargetBlueprint->ParentClass;
+		for (UEdGraphPin* Pin : Context.Pins)
+		{
+			if ((Pin->PinType.PinCategory == UEdGraphSchema_K2::PC_Object) &&
+				Pin->PinType.PinSubCategoryObject.IsValid())
+			{
+				TargetClass = CastChecked<UClass>(Pin->PinType.PinSubCategoryObject.Get());
+			}
+		}
+		
+		UFunction const* WrappedFunction = GetFunction();
+		checkSlow(WrappedFunction != nullptr);
+		UClass* FunctionClass = WrappedFunction->GetOwnerClass()->GetAuthoritativeClass();
 
-//------------------------------------------------------------------------------
-FText UBlueprintFunctionNodeSpawner::GetDefaultMenuName(FBindingSet const& Bindings) const
-{
-	UFunction const* Function = GetFunction();
-	check(Function != nullptr);
-	FString FunctionName = UK2Node_CallFunction::GetUserFacingFunctionName(Function);
+		if (!TargetClass->IsChildOf(FunctionClass))
+		{
+			MenuSignature.Category = FEditorCategoryUtils::BuildCategoryString(FCommonEditorCategory::Class, 
+				FText::FromString(FunctionClass->GetDisplayNameText().ToString() + TEXT("|") + CategoryString));
+		}
+	}
 
 	if (Bindings.Num() == 1)
 	{
 		UObjectProperty const* ObjectProperty = Cast<UObjectProperty>(Bindings.CreateConstIterator()->Get());
 		if (ObjectProperty != nullptr)
 		{
+			FString BoundFunctionName = FString::Printf(TEXT("%s (%s)"), *MenuSignature.MenuName.ToString(), *ObjectProperty->GetName());
 			// @TODO: this should probably be an FText::Format()
-			FunctionName = FString::Printf(TEXT("%s (%s)"), *FunctionName, *ObjectProperty->GetName());
+			MenuSignature.MenuName = FText::FromString(BoundFunctionName);
 		}
 	}
-
-	return FText::FromString(FunctionName);
-}
-
-//------------------------------------------------------------------------------
-FText UBlueprintFunctionNodeSpawner::GetDefaultMenuCategory() const
-{
-	UFunction const* Function = GetFunction();
-	check(Function != nullptr);
-	return FText::FromString(UK2Node_CallFunction::GetDefaultCategoryForFunction(Function, TEXT("")));
+	DynamicUiSignatureGetter.ExecuteIfBound(Context, Bindings, &MenuSignature);
+	return MenuSignature;
 }
 
 //------------------------------------------------------------------------------
-FText UBlueprintFunctionNodeSpawner::GetDefaultMenuTooltip() const
+UEdGraphNode* UBlueprintFunctionNodeSpawner::Invoke(UEdGraph* ParentGraph, FBindingSet const& Bindings, FVector2D const Location) const
 {
-	UFunction const* Function = GetFunction();
-	check(Function != nullptr);
+	auto PostSpawnSetupLambda = [](UEdGraphNode* NewNode, bool bIsTemplateNode, UFunction const* Function, FSetNodeFieldDelegate SetFieldDelegate, FCustomizeNodeDelegate UserDelegate)
+	{
+		SetFieldDelegate.ExecuteIfBound(NewNode, Function);
+		UserDelegate.ExecuteIfBound(NewNode, bIsTemplateNode);
+	};
+	FCustomizeNodeDelegate PostSpawnSetupDelegate = FCustomizeNodeDelegate::CreateStatic(PostSpawnSetupLambda, GetFunction(), SetNodeFieldDelegate, CustomizeNodeDelegate);
 
-	FText Tooltip = FText::FromString(UK2Node_CallFunction::GetDefaultTooltipForFunction(Function));
-	if (Tooltip.IsEmpty())
+	UClass* SpawnClass = NodeClass;
+	if ((Bindings.Num() == 1) && (*Bindings.CreateConstIterator())->IsA<UObjectProperty>())
 	{
-		FBindingSet EmptyContext;
-		Tooltip = GetDefaultMenuName(EmptyContext);
+		SpawnClass = UK2Node_CallFunctionOnMember::StaticClass();
 	}
-	return Tooltip;
-}
-
-//------------------------------------------------------------------------------
-FString UBlueprintFunctionNodeSpawner::GetDefaultSearchKeywords() const
-{
-	UFunction const* Function = GetFunction();
-	check(Function != nullptr);
-	
-	FString SearchKeywords = UK2Node_CallFunction::GetKeywordsForFunction(Function);
-	// add at least one character, so that the menu item doesn't attempt to
-	// ping a template node
-	return SearchKeywords.AppendChar(TEXT(' '));
-}
+	// if this spawner was set up to spawn a bound node, reset this so the 
+	// bound nodes get positioned properly
+	BlueprintFunctionNodeSpawnerImpl::BindingOffset = FVector2D::ZeroVector;
 
-//------------------------------------------------------------------------------
-FName UBlueprintFunctionNodeSpawner::GetDefaultMenuIcon(FLinearColor& ColorOut) const
-{
-	return UK2Node_CallFunction::GetPaletteIconForFunction(GetFunction(), ColorOut);
+	return Super::SpawnNode(SpawnClass, ParentGraph, Bindings, Location, PostSpawnSetupDelegate);
 }
 
 //------------------------------------------------------------------------------
diff --git a/Engine/Source/Editor/BlueprintGraph/Private/BlueprintNodeSpawner.cpp b/Engine/Source/Editor/BlueprintGraph/Private/BlueprintNodeSpawner.cpp
--- a/Engine/Source/Editor/BlueprintGraph/Private/BlueprintNodeSpawner.cpp
+++ b/Engine/Source/Editor/BlueprintGraph/Private/BlueprintNodeSpawner.cpp
@@ -11,7 +11,8 @@
 namespace BlueprintNodeSpawnerImpl
 {
 	/**
-	 * 
+	 * Retrieves a singleton like instance of FBlueprintNodeTemplateCache, will
+	 * spawn one if 
 	 * 
 	 * @return 
 	 */
@@ -82,22 +83,123 @@ void UBlueprintNodeSpawner::Prime()
 			// spawner (to help filter by pin context)
 			CachedTemplateNode->AllocateDefaultPins();
 		}
+	}
+	PrimeDefaultUiSpec();
+}
+
+//------------------------------------------------------------------------------
+FBlueprintActionUiSpec const& UBlueprintNodeSpawner::PrimeDefaultUiSpec(UEdGraph* TargetGraph) const
+{
+	bool bTemplateNodeFetched = false;
+	UEdGraphNode* NodeTemplate = nullptr;
+	// @TODO: boo! const cast... all to make this callable from GetUiSpec()
+	FBlueprintActionUiSpec& MenuSignature = const_cast<FBlueprintActionUiSpec&>(DefaultMenuSignature);
+
+	//--------------------------------------
+	// Verify MenuName
+	//--------------------------------------
+
+	if (MenuSignature.MenuName.IsEmpty())
+	{
+		NodeTemplate = bTemplateNodeFetched ? NodeTemplate : GetTemplateNode(TargetGraph);
+		if (NodeTemplate != nullptr)
+		{
+			MenuSignature.MenuName = NodeTemplate->GetNodeTitle(ENodeTitleType::MenuTitle);
+		}
+		// if a target graph was provided, then we've done all we can to spawn a
+		// template node... we have to default to something
+		if (MenuSignature.MenuName.IsEmpty() && (TargetGraph != nullptr))
+		{
+			MenuSignature.MenuName = FText::FromName(GetFName());
+		}
+		bTemplateNodeFetched = true;
+	}
+
+	//--------------------------------------
+	// Verify Category
+	//--------------------------------------
+
+	if (MenuSignature.Category.IsEmpty())
+	{
+		NodeTemplate = bTemplateNodeFetched ? NodeTemplate : GetTemplateNode(TargetGraph);
+		// @TODO: consider moving GetMenuCategory() up into UEdGraphNode
+		if (UK2Node* K2ReferenceNode = Cast<UK2Node>(NodeTemplate))
+		{
+			MenuSignature.Category = K2ReferenceNode->GetMenuCategory();
+		}
+		// if a target graph was provided, then we've done all we can to spawn a
+		// template node... we have to default to something
+		if (MenuSignature.Category.IsEmpty() && (TargetGraph != nullptr))
+		{
+			// want to set it to something so we won't end up back in this condition
+			MenuSignature.Category = NSLOCTEXT("BlueprintNodeSpawner", "EmptyMenuCategory", "|");
+		}
+		bTemplateNodeFetched = true;
+	}
+
+	//--------------------------------------
+	// Verify Tooltip
+	//--------------------------------------
+
+	if (MenuSignature.Tooltip.IsEmpty())
+	{
+		NodeTemplate = bTemplateNodeFetched ? NodeTemplate : GetTemplateNode(TargetGraph);
+		if (NodeTemplate != nullptr)
+		{
+			MenuSignature.Tooltip = NodeTemplate->GetTooltipText();
+		}
+		// if a target graph was provided, then we've done all we can to spawn a
+		// template node... we have to default to something
+		if (MenuSignature.Tooltip.IsEmpty() && (TargetGraph != nullptr))
+		{
+			MenuSignature.Tooltip = MenuSignature.MenuName;
+		}
+		bTemplateNodeFetched = true;
+	}
+
+	//--------------------------------------
+	// Verify Search Keywords
+	//--------------------------------------
+
+	if (MenuSignature.Keywords.IsEmpty())
+	{
+		NodeTemplate = bTemplateNodeFetched ? NodeTemplate : GetTemplateNode(TargetGraph);
+		if (NodeTemplate != nullptr)
+		{
+			MenuSignature.Keywords = NodeTemplate->GetKeywords();
+		}
+		// if a target graph was provided, then we've done all we can to spawn a
+		// template node... we have to default to something
+		if (MenuSignature.Keywords.IsEmpty() && (TargetGraph != nullptr))
+		{
+			// want to set it to something so we won't end up back in this condition
+			MenuSignature.Keywords = TEXT(" ");
+		}
+		bTemplateNodeFetched = true;
+	}
 
-		//
-		// let the node cache any FText::Format() operations that may be used...
-		if (UK2Node* K2NodeTemplate = Cast<UK2Node>(CachedTemplateNode))
+	//--------------------------------------
+	// Verify Icon Brush Name
+	//--------------------------------------
+
+	if (MenuSignature.IconName.IsNone())
+	{
+		NodeTemplate = bTemplateNodeFetched ? NodeTemplate : GetTemplateNode(TargetGraph);
+		if (NodeTemplate != nullptr)
 		{
-			K2NodeTemplate->GetMenuCategory();
+			MenuSignature.IconName = NodeTemplate->GetPaletteIcon(MenuSignature.IconTint);
 		}
-		CachedTemplateNode->GetTooltipText();
-		CachedTemplateNode->GetNodeTitle(ENodeTitleType::MenuTitle);
+		// if a target graph was provided, then we've done all we can to spawn a
+		// template node... we have to default to something
+		if (MenuSignature.IconName.IsNone() && (TargetGraph != nullptr))
+		{
+			// want to set it to something so we won't end up back in this condition
+			MenuSignature.IconName = TEXT("GraphEditor.Default_16x");
+		}
+		bTemplateNodeFetched = true;
 	}
 
-	// in case any of these cache FText::Format() operations
-	FBindingSet EmptyContext;
-	GetDefaultMenuName(EmptyContext);
-	GetDefaultMenuCategory();
-	GetDefaultMenuTooltip();
+	return MenuSignature;
 }
 
 //------------------------------------------------------------------------------
@@ -117,61 +219,29 @@ FBlueprintNodeSignature UBlueprintNodeSpawner::GetSpawnerSignature() const
 }
 
 //------------------------------------------------------------------------------
-UEdGraphNode* UBlueprintNodeSpawner::Invoke(UEdGraph* ParentGraph, FBindingSet const& Bindings, FVector2D const Location) const
-{
-	return SpawnNode(NodeClass, ParentGraph, Bindings, Location, CustomizeNodeDelegate);
-}
-
-//------------------------------------------------------------------------------
-FText UBlueprintNodeSpawner::GetDefaultMenuName(FBindingSet const& Bindings) const
+FBlueprintActionUiSpec UBlueprintNodeSpawner::GetUiSpec(FBlueprintActionContext const& Context, FBindingSet const& Bindings) const
 {
-	// DO NOT make a template node and query it here (the separate ui building
-	// code can do that if it likes)... this is meant to be an overridable
-	// alternative to that (for perf reasons)
-	return FText::GetEmpty();
+	FBlueprintActionUiSpec MenuSignature = PrimeDefaultUiSpec();
+	DynamicUiSignatureGetter.ExecuteIfBound(Context, Bindings, &MenuSignature);
+	return MenuSignature;
 }
 
 //------------------------------------------------------------------------------
-FText UBlueprintNodeSpawner::GetDefaultMenuCategory() const
-{
-	// DO NOT make a template node and query it here (the separate ui building
-	// code can do that if it likes)... this is meant to be an overridable
-	// alternative to that (for perf reasons)
-	return FText::GetEmpty();
-}
-
-//------------------------------------------------------------------------------
-FText UBlueprintNodeSpawner::GetDefaultMenuTooltip() const
-{
-	// DO NOT make a template node and query it here (the separate ui building
-	// code can do that if it likes)... this is meant to be an overridable
-	// alternative to that (for perf reasons)
-	return FText::GetEmpty();
-}
-
-//------------------------------------------------------------------------------
-FString UBlueprintNodeSpawner::GetDefaultSearchKeywords() const
+UEdGraphNode* UBlueprintNodeSpawner::Invoke(UEdGraph* ParentGraph, FBindingSet const& Bindings, FVector2D const Location) const
 {
-	// DO NOT make a template node and query it here (the separate ui building
-	// code can do that if it likes)... this is meant to be an overridable
-	// alternative to that (for perf reasons)
-	return FString();
+	return SpawnNode(NodeClass, ParentGraph, Bindings, Location, CustomizeNodeDelegate);
 }
 
 //------------------------------------------------------------------------------
-FName UBlueprintNodeSpawner::GetDefaultMenuIcon(FLinearColor& ColorOut) const
+UEdGraphNode* UBlueprintNodeSpawner::GetCachedTemplateNode() const
 {
-	ColorOut = FLinearColor::White;
-	// DO NOT make a template node and query it here (the separate ui building
-	// code can do that if it likes)... this is meant to be an overridable
-	// alternative to that (for perf reasons)
-	return FName();
+	return BlueprintNodeSpawnerImpl::GetSharedTemplateCache()->GetNodeTemplate(this, NoInit);
 }
 
 //------------------------------------------------------------------------------
-UEdGraphNode* UBlueprintNodeSpawner::GetTemplateNode(UEdGraph* Outer, FBindingSet const& Bindings) const 
+UEdGraphNode* UBlueprintNodeSpawner::GetTemplateNode(UEdGraph* TargetGraph, FBindingSet const& Bindings) const
 {       
-	UEdGraphNode* TemplateNode = BlueprintNodeSpawnerImpl::GetSharedTemplateCache()->GetNodeTemplate(this, Outer);
+	UEdGraphNode* TemplateNode = BlueprintNodeSpawnerImpl::GetSharedTemplateCache()->GetNodeTemplate(this, TargetGraph);
 
 	if (TemplateNode && Bindings.Num() > 0) 
 	{ 
@@ -182,12 +252,6 @@ UEdGraphNode* UBlueprintNodeSpawner::GetTemplateNode(UEdGraph* Outer, FBindingSe
 	return TemplateNode; 
 }
 
-//------------------------------------------------------------------------------
-UEdGraphNode* UBlueprintNodeSpawner::GetTemplateNode(ENoInit) const
-{
-	return BlueprintNodeSpawnerImpl::GetSharedTemplateCache()->GetNodeTemplate(this, NoInit);
-}
-
 //------------------------------------------------------------------------------
 UEdGraphNode* UBlueprintNodeSpawner::SpawnEdGraphNode(TSubclassOf<UEdGraphNode> InNodeClass, UEdGraph* ParentGraph, FBindingSet const& Bindings, FVector2D const Location, FCustomizeNodeDelegate PostSpawnDelegate) const
 {
diff --git a/Engine/Source/Editor/BlueprintGraph/Private/BlueprintNodeSpawnerUtils.cpp b/Engine/Source/Editor/BlueprintGraph/Private/BlueprintNodeSpawnerUtils.cpp
--- a/Engine/Source/Editor/BlueprintGraph/Private/BlueprintNodeSpawnerUtils.cpp
+++ b/Engine/Source/Editor/BlueprintGraph/Private/BlueprintNodeSpawnerUtils.cpp
@@ -50,7 +50,7 @@ UProperty const* FBlueprintNodeSpawnerUtils::GetAssociatedProperty(UBlueprintNod
 
 	if (UBlueprintDelegateNodeSpawner const* PropertySpawner = Cast<UBlueprintDelegateNodeSpawner>(BlueprintAction))
 	{
-		Property = PropertySpawner->GetProperty();
+		Property = PropertySpawner->GetDelegateProperty();
 	}
 	else if (UBlueprintVariableNodeSpawner const* VarSpawner = Cast<UBlueprintVariableNodeSpawner>(BlueprintAction))
 	{
diff --git a/Engine/Source/Editor/BlueprintGraph/Private/BlueprintUnloadedAssetNodeSpawner.cpp b/Engine/Source/Editor/BlueprintGraph/Private/BlueprintUnloadedAssetNodeSpawner.cpp
deleted file mode 100644
diff --git a/Engine/Source/Editor/BlueprintGraph/Private/BlueprintVariableNodeSpawner.cpp b/Engine/Source/Editor/BlueprintGraph/Private/BlueprintVariableNodeSpawner.cpp
--- a/Engine/Source/Editor/BlueprintGraph/Private/BlueprintVariableNodeSpawner.cpp
+++ b/Engine/Source/Editor/BlueprintGraph/Private/BlueprintVariableNodeSpawner.cpp
@@ -8,6 +8,7 @@
 #include "ObjectEditorUtils.h"		// for GetCategory()
 #include "EdGraphSchema_K2.h"		// for ConvertPropertyToPinType()
 #include "EditorCategoryUtils.h"	// for BuildCategoryString()
+#include "BlueprintActionFilter.h"	// for FBlueprintActionContext
 
 #define LOCTEXT_NAMESPACE "BlueprintVariableNodeSpawner"
 
@@ -24,10 +25,45 @@ UBlueprintVariableNodeSpawner* UBlueprintVariableNodeSpawner::Create(TSubclassOf
 		Outer = GetTransientPackage();
 	}
 
+	//--------------------------------------
+	// Constructing the Spawner
+	//--------------------------------------
+
 	UBlueprintVariableNodeSpawner* NodeSpawner = NewObject<UBlueprintVariableNodeSpawner>(Outer);
 	NodeSpawner->NodeClass = NodeClass;
 	NodeSpawner->Field     = VarProperty;
 
+	//--------------------------------------
+	// Default UI Signature
+	//--------------------------------------
+
+	FBlueprintActionUiSpec& MenuSignature = NodeSpawner->DefaultMenuSignature;
+	FString const VarSubCategory = FObjectEditorUtils::GetCategory(VarProperty);
+	MenuSignature.Category = FEditorCategoryUtils::BuildCategoryString(FCommonEditorCategory::Variables, FText::FromString(VarSubCategory));
+
+	FText const VarName = NodeSpawner->GetVariableName();
+	// @TODO: NodeClass could be modified post Create()
+	if (NodeClass->IsChildOf(UK2Node_VariableGet::StaticClass()))
+	{
+		MenuSignature.MenuName = FText::Format(LOCTEXT("GetterMenuName", "Get {0}"), VarName);
+		MenuSignature.Tooltip  = UK2Node_VariableGet::GetPropertyTooltip(VarProperty);
+	}
+	else if (NodeClass->IsChildOf(UK2Node_VariableSet::StaticClass()))
+	{
+		MenuSignature.MenuName = FText::Format(LOCTEXT("SetterMenuName", "Set {0}"), VarName);
+		MenuSignature.Tooltip  = UK2Node_VariableSet::GetPropertyTooltip(VarProperty);
+	}
+	// add at least one character, so that PrimeDefaultMenuSignature() doesn't 
+	// attempt to query the template node
+	//
+	// @TODO: maybe UPROPERTY() fields should have keyword metadata like functions
+	MenuSignature.Keywords = TEXT(" ");
+	MenuSignature.IconName = UK2Node_Variable::GetVarIconFromPinType(NodeSpawner->GetVarType(), MenuSignature.IconTint);
+
+	//--------------------------------------
+	// Post-Spawn Setup
+	//--------------------------------------
+
 	auto MemberVarSetupLambda = [](UEdGraphNode* NewNode, UField const* Field)
 	{
 		if (UProperty const* Property = Cast<UProperty>(Field))
@@ -54,6 +90,10 @@ UBlueprintVariableNodeSpawner* UBlueprintVariableNodeSpawner::Create(TSubclassOf
 		Outer = GetTransientPackage();
 	}
 
+	//--------------------------------------
+	// Constructing the Spawner
+	//--------------------------------------
+
 	// @TODO: consider splitting out local variable spawners (since they don't 
 	//        conform to UBlueprintFieldNodeSpawner
 	UBlueprintVariableNodeSpawner* NodeSpawner = NewObject<UBlueprintVariableNodeSpawner>(Outer);
@@ -61,6 +101,32 @@ UBlueprintVariableNodeSpawner* UBlueprintVariableNodeSpawner::Create(TSubclassOf
 	NodeSpawner->LocalVarOuter = VarContext;
 	NodeSpawner->LocalVarDesc  = VarDesc;
 
+	//--------------------------------------
+	// Default UI Signature
+	//--------------------------------------
+
+	FBlueprintActionUiSpec& MenuSignature = NodeSpawner->DefaultMenuSignature;
+	MenuSignature.Category = FEditorCategoryUtils::BuildCategoryString(FCommonEditorCategory::Variables, FText::FromName(VarDesc.Category));
+
+	FText const VarName = NodeSpawner->GetVariableName();
+	// @TODO: NodeClass could be modified post Create()
+	if (NodeClass->IsChildOf(UK2Node_VariableGet::StaticClass()))
+	{
+		MenuSignature.MenuName = FText::Format(LOCTEXT("LocalGetterMenuName", "Get {0}"), VarName);
+		MenuSignature.Tooltip  = UK2Node_VariableGet::GetBlueprintVarTooltip(VarDesc);
+	}
+	else if (NodeClass->IsChildOf(UK2Node_VariableSet::StaticClass()))
+	{
+		MenuSignature.MenuName = FText::Format(LOCTEXT("LocalSetterMenuName", "Set {0}"), VarName);
+		MenuSignature.Tooltip  = UK2Node_VariableSet::GetBlueprintVarTooltip(VarDesc);
+	}
+	// add at least one character, so that PrimeDefaultMenuSignature() doesn't 
+	// attempt to query the template node
+	//
+	// @TODO: maybe UPROPERTY() fields should have keyword metadata like functions
+	MenuSignature.Keywords = TEXT(" ");
+	MenuSignature.IconName = UK2Node_Variable::GetVarIconFromPinType(NodeSpawner->GetVarType(), MenuSignature.IconTint);
+
 	return NodeSpawner;
 }
 
@@ -75,13 +141,6 @@ void UBlueprintVariableNodeSpawner::Prime()
 {
 	// we expect that you don't need a node template to construct menu entries
 	// from this, so we choose not to pre-cache one here
-
-	// all of these perform expensive FText::Format() operations and cache the 
-	// results...
-	FBindingSet EmptyContext;
-	GetDefaultMenuName(EmptyContext);
-	GetDefaultMenuCategory();
-	GetDefaultMenuTooltip();
 }
 
 //------------------------------------------------------------------------------
@@ -98,6 +157,39 @@ FBlueprintNodeSignature UBlueprintVariableNodeSpawner::GetSpawnerSignature() con
 	return SpawnerSignature;
 }
 
+//------------------------------------------------------------------------------
+FBlueprintActionUiSpec UBlueprintVariableNodeSpawner::GetUiSpec(FBlueprintActionContext const& Context, FBindingSet const& Bindings) const
+{
+	UEdGraph* TargetGraph = (Context.Graphs.Num() > 0) ? Context.Graphs[0] : nullptr;
+	FBlueprintActionUiSpec MenuSignature = PrimeDefaultUiSpec(TargetGraph);
+
+	if (UProperty const* WrappedVariable = GetVarProperty())
+	{
+		checkSlow(Context.Blueprints.Num() > 0);
+		UBlueprint const* TargetBlueprint = Context.Blueprints[0];
+
+		// @TODO: this is duplicated in a couple places, move it to some shared resource
+		UClass const* TargetClass = (TargetBlueprint->GeneratedClass != nullptr) ? TargetBlueprint->GeneratedClass : TargetBlueprint->ParentClass;
+		for (UEdGraphPin* Pin : Context.Pins)
+		{
+			if ((Pin->PinType.PinCategory == UEdGraphSchema_K2::PC_Object) && 
+				 Pin->PinType.PinSubCategoryObject.IsValid())
+			{
+				TargetClass = CastChecked<UClass>(Pin->PinType.PinSubCategoryObject.Get());
+			}
+		}
+
+		UClass const* VariableClass = WrappedVariable->GetOwnerClass()->GetAuthoritativeClass();
+		if (!TargetClass->IsChildOf(VariableClass))
+		{
+			MenuSignature.Category = FEditorCategoryUtils::BuildCategoryString(FCommonEditorCategory::Class,
+				FText::FromString(VariableClass->GetDisplayNameText().ToString() + TEXT("|") + MenuSignature.Category.ToString()));
+		}
+	}
+	DynamicUiSignatureGetter.ExecuteIfBound(Context, Bindings, &MenuSignature);
+	return MenuSignature;
+}
+
 //------------------------------------------------------------------------------
 UEdGraphNode* UBlueprintVariableNodeSpawner::Invoke(UEdGraph* ParentGraph, FBindingSet const& Bindings, FVector2D const Location) const
 {
@@ -129,93 +221,6 @@ UEdGraphNode* UBlueprintVariableNodeSpawner::Invoke(UEdGraph* ParentGraph, FBind
 	return NewNode;
 }
 
-//------------------------------------------------------------------------------
-FText UBlueprintVariableNodeSpawner::GetDefaultMenuName(const FBindingSet& Bindings) const
-{	
-	if (CachedMenuName.IsOutOfDate())
-	{
-		FText VarName = GetVariableName();
-		check(NodeClass != nullptr);
-
-		// FText::Format() is slow, so we cache this to save on performance
-		if (NodeClass->IsChildOf<UK2Node_VariableGet>())
-		{
-			CachedMenuName = FText::Format(LOCTEXT("GetterMenuName", "Get {0}"), VarName);
-		}
-		else if (NodeClass->IsChildOf<UK2Node_VariableSet>())
-		{
-			CachedMenuName = FText::Format(LOCTEXT("SetterMenuName", "Set {0}"), VarName);
-		}
-	}
-	return CachedMenuName;
-}
-
-//------------------------------------------------------------------------------
-FText UBlueprintVariableNodeSpawner::GetDefaultMenuCategory() const
-{
-	if (CachedCategory.IsOutOfDate())
-	{
-		FText VarSubCategory;
-		if (IsLocalVariable())
-		{
-			VarSubCategory = FText::FromName(LocalVarDesc.Category);
-		}
-		else if (UProperty const* MemberVariable = GetVarProperty())
-		{
-			VarSubCategory = FText::FromString(FObjectEditorUtils::GetCategory(MemberVariable));
-		}
-		CachedCategory = FEditorCategoryUtils::BuildCategoryString(FCommonEditorCategory::Variables, VarSubCategory);
-	}
-	return CachedCategory;
-}
-
-//------------------------------------------------------------------------------
-FText UBlueprintVariableNodeSpawner::GetDefaultMenuTooltip() const
-{
-	if (CachedTooltip.IsOutOfDate())
-	{
-		if (NodeClass->IsChildOf<UK2Node_VariableSet>())
-		{
-			if (IsLocalVariable())
-			{
-				CachedTooltip = UK2Node_VariableSet::GetBlueprintVarTooltip(LocalVarDesc);
-			}
-			else if (UProperty const* MemberVariable = GetVarProperty())
-			{
-				CachedTooltip = UK2Node_VariableSet::GetPropertyTooltip(MemberVariable);
-			}
-		}
-		else if (NodeClass->IsChildOf<UK2Node_VariableGet>())
-		{
-			if (IsLocalVariable())
-			{
-				CachedTooltip = UK2Node_VariableGet::GetBlueprintVarTooltip(LocalVarDesc);
-			}
-			else if (UProperty const* MemberVariable = GetVarProperty())
-			{
-				CachedTooltip = UK2Node_VariableGet::GetPropertyTooltip(MemberVariable);
-			}
-		}
-	}	
-	return CachedTooltip;
-}
-
-//------------------------------------------------------------------------------
-FString UBlueprintVariableNodeSpawner::GetDefaultSearchKeywords() const
-{
-	// @TODO: maybe UPROPERTY() fields should have keyword metadata like functions
-	// 
-	// add at least one character, so that the menu item doesn't attempt to
-	// ping a node template 	
-	return TEXT(" ");
-}
-
-//------------------------------------------------------------------------------
-FName UBlueprintVariableNodeSpawner::GetDefaultMenuIcon(FLinearColor& ColorOut) const
-{
-	return UK2Node_Variable::GetVarIconFromPinType(GetVarType(), ColorOut);
-}
-
 //------------------------------------------------------------------------------
 bool UBlueprintVariableNodeSpawner::IsLocalVariable() const
 {
diff --git a/Engine/Source/Editor/BlueprintGraph/Private/K2Node_CallFunction.cpp b/Engine/Source/Editor/BlueprintGraph/Private/K2Node_CallFunction.cpp
--- a/Engine/Source/Editor/BlueprintGraph/Private/K2Node_CallFunction.cpp
+++ b/Engine/Source/Editor/BlueprintGraph/Private/K2Node_CallFunction.cpp
@@ -1616,7 +1616,11 @@ void UK2Node_CallFunction::GetNodeAttributes( TArray<TKeyValuePair<FString, FStr
 FText UK2Node_CallFunction::GetMenuCategory() const
 {
 	UFunction* TargetFunction = GetTargetFunction();
-	return FText::FromString(GetDefaultCategoryForFunction(TargetFunction, TEXT("")));
+	if (TargetFunction != nullptr)
+	{
+		return FText::FromString(GetDefaultCategoryForFunction(TargetFunction, TEXT("")));
+	}
+	return FText::GetEmpty();
 }
 
 bool UK2Node_CallFunction::HasExternalBlueprintDependencies(TArray<class UStruct*>* OptionalOutput) const
diff --git a/Engine/Source/Editor/BlueprintGraph/Private/K2Node_Literal.cpp b/Engine/Source/Editor/BlueprintGraph/Private/K2Node_Literal.cpp
--- a/Engine/Source/Editor/BlueprintGraph/Private/K2Node_Literal.cpp
+++ b/Engine/Source/Editor/BlueprintGraph/Private/K2Node_Literal.cpp
@@ -171,26 +171,28 @@ void UK2Node_Literal::GetMenuActions(FBlueprintActionDatabaseRegistrar& ActionRe
  		return true;
 	};
 
-	auto MenuDescriptionLambda = []( const IBlueprintNodeBinder::FBindingSet& BindingContext ) -> FText
+	auto UiSpecOverride = [](const FBlueprintActionContext& /*Context*/, const IBlueprintNodeBinder::FBindingSet& Bindings, FBlueprintActionUiSpec* UiSpecOut)
 	{
-		if (BindingContext.Num() == 1)
+		if (Bindings.Num() == 1)
 		{
-			return FText::Format(NSLOCTEXT("K2Node", "LiteralTitle", "Create a Reference to {0}"), FText::FromString((*(BindingContext.CreateConstIterator()))->GetName()));
+			UiSpecOut->MenuName = FText::Format(NSLOCTEXT("K2Node", "LiteralTitle", "Create a Reference to {0}"), 
+				FText::FromString( (*(Bindings.CreateConstIterator()))->GetName() ));
 		}
-		else if (BindingContext.Num() > 1)
+		else if (Bindings.Num() > 1)
 		{
-			return FText::Format(NSLOCTEXT("K2Node", "FallbackLiteralTitle", "Create References to {0} selected Actors"), FText::AsNumber(BindingContext.Num()));
+			UiSpecOut->MenuName = FText::Format(NSLOCTEXT("K2Node", "FallbackLiteralTitle", "Create References to {0} selected Actors"), 
+				FText::AsNumber(Bindings.Num()));
 		}
 		else
 		{
-			return NSLOCTEXT("K2Node", "FallbackLiteralTitle", "Error: No Actors in Context");
+			UiSpecOut->MenuName = NSLOCTEXT("K2Node", "FallbackLiteralTitle", "Error: No Actors in Context");
 		}
 	};
 
 	UBlueprintBoundNodeSpawner* NodeSpawner = UBlueprintBoundNodeSpawner::Create(GetClass());
 	NodeSpawner->CanBindObjectDelegate = UBlueprintBoundNodeSpawner::FCanBindObjectDelegate::CreateStatic(CanBindObjectLambda);
 	NodeSpawner->OnBindObjectDelegate = UBlueprintBoundNodeSpawner::FOnBindObjectDelegate::CreateStatic(PostBindSetupLambda);
-	NodeSpawner->OnGenerateMenuDescriptionDelegate = UBlueprintBoundNodeSpawner::FOnGenerateMenuDescriptionDelegate::CreateStatic(MenuDescriptionLambda);
+	NodeSpawner->DynamicUiSignatureGetter = UBlueprintBoundNodeSpawner::FUiSpecOverrideDelegate::CreateStatic(UiSpecOverride);
 	ActionRegistrar.AddBlueprintAction(NodeSpawner);
 }
 
diff --git a/Engine/Source/Editor/BlueprintGraph/Private/K2Node_MatineeController.cpp b/Engine/Source/Editor/BlueprintGraph/Private/K2Node_MatineeController.cpp
--- a/Engine/Source/Editor/BlueprintGraph/Private/K2Node_MatineeController.cpp
+++ b/Engine/Source/Editor/BlueprintGraph/Private/K2Node_MatineeController.cpp
@@ -142,19 +142,21 @@ void UK2Node_MatineeController::GetMenuActions(FBlueprintActionDatabaseRegistrar
 		return false;
 	};
 
-	auto MenuCategoryDescription = []( const IBlueprintNodeBinder::FBindingSet& BindingContext ) -> FText
+	auto UiSpecOverride = [](const FBlueprintActionContext& /*Context*/, const IBlueprintNodeBinder::FBindingSet& Bindings, FBlueprintActionUiSpec* UiSpecOut)
 	{
-		if( BindingContext.Num() == 1 )
+		if (Bindings.Num() == 1)
 		{
-			return FText::Format(NSLOCTEXT("K2Node", "MatineeeControllerTitle", "Create a Matinee Controller for {0}"), FText::FromString( (*(BindingContext.CreateConstIterator()))->GetName()) );
+			UiSpecOut->MenuName = FText::Format(NSLOCTEXT("K2Node", "MatineeeControllerTitle", "Create a Matinee Controller for {0}"),
+				FText::FromString((*(Bindings.CreateConstIterator()))->GetName()));
 		}
-		else if( BindingContext.Num() > 1 )
+		else if (Bindings.Num() > 1)
 		{
-			return FText::Format(NSLOCTEXT("K2Node", "MultipleMatineeeControllerTitle", "Create Matinee Controllers for {0} selected MatineeActors"), FText::AsNumber(BindingContext.Num()) );
+			UiSpecOut->MenuName = FText::Format(NSLOCTEXT("K2Node", "MultipleMatineeeControllerTitle", "Create Matinee Controllers for {0} selected MatineeActors"),
+				FText::AsNumber(Bindings.Num()));
 		}
 		else
 		{
-			return NSLOCTEXT("K2Node", "FallbackMatineeeControllerTitle", "Error: No MatineeActors in Context");
+			UiSpecOut->MenuName = NSLOCTEXT("K2Node", "FallbackMatineeeControllerTitle", "Error: No MatineeActors in Context");
 		}
 	};
 
@@ -167,8 +169,8 @@ void UK2Node_MatineeController::GetMenuActions(FBlueprintActionDatabaseRegistrar
 
 	UBlueprintBoundNodeSpawner* NodeSpawner = UBlueprintBoundNodeSpawner::Create(GetClass());
 	NodeSpawner->CanBindObjectDelegate = UBlueprintBoundNodeSpawner::FCanBindObjectDelegate::CreateStatic(CanBindObjectLambda);
-	NodeSpawner->OnBindObjectDelegate = UBlueprintBoundNodeSpawner::FOnBindObjectDelegate::CreateStatic(PostBindSetupLambda);
-	NodeSpawner->OnGenerateMenuDescriptionDelegate = UBlueprintBoundNodeSpawner::FOnGenerateMenuDescriptionDelegate::CreateStatic(MenuCategoryDescription);
+	NodeSpawner->OnBindObjectDelegate  = UBlueprintBoundNodeSpawner::FOnBindObjectDelegate::CreateStatic(PostBindSetupLambda);
+	NodeSpawner->DynamicUiSignatureGetter = UBlueprintBoundNodeSpawner::FUiSpecOverrideDelegate::CreateStatic(UiSpecOverride);
 	ActionRegistrar.AddBlueprintAction(NodeSpawner);
 }
 
diff --git a/Engine/Source/Editor/BlueprintGraph/Public/BlueprintNodeSpawner.h b/Engine/Source/Editor/BlueprintGraph/Public/BlueprintNodeSpawner.h
--- a/Engine/Source/Editor/BlueprintGraph/Public/BlueprintNodeSpawner.h
+++ b/Engine/Source/Editor/BlueprintGraph/Public/BlueprintNodeSpawner.h
@@ -7,34 +7,46 @@
 #include "BlueprintNodeSpawner.generated.h"
 
 // Forward declarations
-class UEdGraph;
-class UEdGraphNode;
-class FBlueprintNodeTemplateCache;
+class  UEdGraph;
+class  UEdGraphNode;
+struct FBlueprintActionContext;
 
 /*******************************************************************************
-* UBlueprintNodeSpawner Declaration
-******************************************************************************/
+ * FBlueprintActionUiSpec
+ ******************************************************************************/
+
+struct FBlueprintActionUiSpec
+{
+	FBlueprintActionUiSpec()
+		: IconTint(FLinearColor::White)
+	{}
+
+	FText   MenuName;
+	FText   Category;
+	FText   Tooltip;
+	FString Keywords;
+	FName   IconName;
+	FLinearColor IconTint;
+};
+
+/*******************************************************************************
+ * UBlueprintNodeSpawner
+ ******************************************************************************/
 
 /**
  * Intended to be wrapped and used by FBlueprintActionMenuItem. Rather than 
  * sub-classing the menu item, we choose to subclass this instead (for 
  * different node types). That way, we get the type inference that comes with  
  * UObjects (and we don't have to continuously compare identification strings). 
- *
- * Additionally, this decouples the "action" from ui settings (like the tooltip,
- * menu category, etc.). It gives us more leeway in how we choose to present 
- * these "actions".
  */
 UCLASS(Transient)
 class BLUEPRINTGRAPH_API UBlueprintNodeSpawner : public UObject, public IBlueprintNodeBinder
 {
 	GENERATED_UCLASS_BODY()
 	DECLARE_DELEGATE_TwoParams(FCustomizeNodeDelegate, UEdGraphNode*, bool);
+	DECLARE_DELEGATE_ThreeParams(FUiSpecOverrideDelegate, FBlueprintActionContext const&, FBindingSet const&, FBlueprintActionUiSpec*);
 
 public:
-	/** */
-	virtual ~UBlueprintNodeSpawner();
-
 	/**
 	 * Creates a new UBlueprintNodeSpawner for the specified node class. Sets
 	 * the allocated spawner's NodeClass and CustomizeNodeDelegate fields from
@@ -58,20 +70,22 @@ public:
 	template<class NodeType>
 	static UBlueprintNodeSpawner* Create(UObject* Outer = nullptr, FCustomizeNodeDelegate PostSpawnDelegate = FCustomizeNodeDelegate());
 	
-	/**
-	 * Holds the class of node to spawn. May be null for sub-classes that know
-	 * specifically what node types to spawn (if null for an instance of this 
-	 * class, then nothing will be spawned).
-	 */
+public:
+	/** */
+	virtual ~UBlueprintNodeSpawner();
+
+	/** Holds the node type that this spawner will instantiate. */
 	UPROPERTY()
 	TSubclassOf<UEdGraphNode> NodeClass;
-	
-	/**
-	 * A delegate to perform specialized node setup post-spawn (so we don't have
-	 * to sub-class this for every node type, and so we don't have to spawn a
-	 * template node for each to hold on to).
-	 */
-	FCustomizeNodeDelegate CustomizeNodeDelegate;
+
+	/** Defines how this spawner is presented in the ui */
+	FBlueprintActionUiSpec  DefaultMenuSignature;
+
+	/** A delegate to perform specialized node setup post-spawn (so you don't have to sub-class this for every node type). */
+	FCustomizeNodeDelegate  CustomizeNodeDelegate;
+
+	/** Provides a way to override DefaultMenuSignature based off blueprint/graph/menu context */
+	FUiSpecOverrideDelegate DynamicUiSignatureGetter;
 
 	/**
 	 * Not required, but intended to passively help speed up menu building 
@@ -80,6 +94,15 @@ public:
 	 */
 	virtual void Prime();
 
+	/**
+	 * Takes the FBlueprintActionUiSpec that this was spawned with and attempts 
+	 * to fill in any missing fields (by polling a template node).
+	 * 
+	 * @param  TargetGraph	Optional context for helping spawn the template node. 
+	 * @return This spawner's default ui spec.
+	 */
+	FBlueprintActionUiSpec const& PrimeDefaultUiSpec(UEdGraph* TargetGraph = nullptr) const;
+
 	/**
 	 * We want to be able to compare spawners, and have a signature that is 
 	 * rebuildable on subsequent runs. So, what makes each spawner unique is the 
@@ -90,6 +113,16 @@ public:
 	 */
 	virtual FBlueprintNodeSignature GetSpawnerSignature() const;
 
+	/**
+	 * Takes the default FBlueprintActionUiSpec and modifies it dynamically to 
+	 * accommodate the current context.
+	 * 
+	 * @param  Context		Context that this spawner will be presented under.
+	 * @param  Bindings		Bindings that will be applied to the node post-spawn.
+	 * @return A ui spec, detailing how to present this action in the menu.
+	 */
+	virtual FBlueprintActionUiSpec GetUiSpec(FBlueprintActionContext const& Context, FBindingSet const& Bindings) const;
+
 	/**
 	 * Takes care of spawning a node for the specified graph. Looks to see if 
 	 * the supplied graph is transient, and if so, spawns a NOT fully formed
@@ -106,65 +139,14 @@ public:
 	 * @return Null if it failed to spawn a node, otherwise a newly spawned node or possibly one that already existed.
 	 */
 	virtual UEdGraphNode* Invoke(UEdGraph* ParentGraph, FBindingSet const& Bindings, FVector2D const Location) const;	
-	
-	/**
-	 * To save on performance, certain sub-classes can concoct menu names from
-	 * static data (like function names, etc.). This unfortunately doesn't 
-	 * completely decouple these spawners from the ui we sought to separate 
-	 * from. However, this can be thought of as a suggestion that the ui 
-	 * building code is free to do with as it pleases.
-	 *
-	 * @return An empty text string (could be a localized name for sub-classes).
-	 */
-	virtual FText GetDefaultMenuName(FBindingSet const& Bindings) const;
-	
-	/**
-	 * To save on performance, certain sub-classes can concoct menu categories
-	 * from static data (like function metadata, etc.). This unfortunately
-	 * doesn't completely decouple these spawners from the ui we sought to
-	 * separate from. However, this can be thought of as a suggestion that the
-	 * ui building code is free to do with as it pleases (like concatenate with 
-	 * parent categories, etc.).
-	 *
-	 * @return An empty text string (could be a localized category for sub-classes).
-	 */
-	virtual FText GetDefaultMenuCategory() const;
-	
-	/**
-	 * To save on performance, certain sub-classes can construct tooltips from
-	 * static data (like function metadata, etc.). This unfortunately doesn't
-	 * completely decouple these spawners from the ui we sought to separate
-	 * them from. However, this can be thought of as a suggestion that the ui
-	 * building code is free to do with as it pleases.
-	 *
-	 * @return An empty text string (could be a localized tooltip for sub-classes).
-	 */
-	virtual FText GetDefaultMenuTooltip() const;
-	
-	/**
-	 * To save on performance, certain sub-classes can pull search keywords from
-	 * static data (like function metadata, etc.). This unfortunately doesn't
-	 * completely decouple these spawners from the ui we sought to separate
-	 * them from. Although, this can be thought of as a suggestion that the ui
-	 * building code is free to do with as it pleases (doesn't have to use it).
-	 *
-	 * @TODO: Should search keywords be localized? Probably.
-	 *
-	 * @return An empty text string (could be a localized keywords for sub-classes).
-	 */
-	virtual FString GetDefaultSearchKeywords() const;
 
 	/**
-	 * To save on performance, certain sub-classes can pull icons info from
-	 * static data (like function metadata, etc.). This unfortunately doesn't
-	 * completely decouple these spawners from the ui we sought to separate
-	 * them from. Although, this can be thought of as a suggestion that the ui
-	 * building code is free to do with as it pleases (doesn't have to use it).
+	 * Retrieves a cached template for the node that this is set to spawn. Will
+	 * NOT spawn one if it is not already cached.
 	 *
-	 * @param  ColorOut		The color to tint the icon with.
-	 * @return Name of the brush to use (use FEditorStyle::GetBrush() to resolve).
+	 * @return The cached template-node (if one already exists for this spawner).
 	 */
-	virtual FName GetDefaultMenuIcon(FLinearColor& ColorOut) const;
+	UEdGraphNode* GetCachedTemplateNode() const;
 
 	/**
 	 * Retrieves a cached template for the node that this is set to spawn. Will
@@ -178,14 +160,6 @@ public:
 	 */
 	UEdGraphNode* GetTemplateNode(UEdGraph* TargetGraph = nullptr, FBindingSet const& Bindings = FBindingSet()) const;
 
-	/**
-	 * Retrieves a cached template for the node that this is set to spawn. Will
-	 * NOT spawn one if it is not already cached.
-	 *
-	 * @return The cached template-node (if one already exists for this spawner).
-	 */
-	UEdGraphNode* GetTemplateNode(ENoInit) const;
-
 	// IBlueprintNodeBinder interface
 	virtual bool IsBindingCompatible(UObject const* BindingCandidate) const override { return false; }
 	virtual bool CanBindMultipleObjects() const override { return false; }
diff --git a/Engine/Source/Editor/Kismet/Private/BlueprintActionMenuBuilder.cpp b/Engine/Source/Editor/Kismet/Private/BlueprintActionMenuBuilder.cpp
--- a/Engine/Source/Editor/Kismet/Private/BlueprintActionMenuBuilder.cpp
+++ b/Engine/Source/Editor/Kismet/Private/BlueprintActionMenuBuilder.cpp
@@ -4,7 +4,6 @@
 #include "BlueprintActionMenuBuilder.h"
 #include "BlueprintActionMenuItem.h"
 #include "BlueprintDragDropMenuItem.h"
-#include "BlueprintBoundMenuItem.h"
 #include "BlueprintActionFilter.h"
 #include "BlueprintActionDatabase.h"
 #include "BlueprintNodeSpawner.h"
@@ -50,7 +49,7 @@ public:
 	 * @param  Action			The node-spawner that the new menu item should wrap.
 	 * @return A newly allocated FBlueprintActionMenuItem (which wraps the supplied action).
 	 */
-	TSharedPtr<FEdGraphSchemaAction> MakeActionMenuItem(TWeakPtr<FBlueprintEditor> EditorContext, FBlueprintActionInfo const& ActionInfo);
+	TSharedPtr<FBlueprintActionMenuItem> MakeActionMenuItem(TWeakPtr<FBlueprintEditor> EditorContext, FBlueprintActionInfo const& ActionInfo);
 
 	/**
 	 * Spawns a new FBlueprintDragDropMenuItem with the node-spawner. Constructs
@@ -67,64 +66,9 @@ public:
 	 * @param  BoundAction	
 	 * @return 
 	 */
-	TSharedPtr<FBlueprintBoundMenuItem> MakeBoundMenuItem(TWeakPtr<FBlueprintEditor> EditorContext, FBlueprintActionInfo const& ActionInfo);
+	TSharedPtr<FBlueprintActionMenuItem> MakeBoundMenuItem(TWeakPtr<FBlueprintEditor> EditorContext, FBlueprintActionInfo const& ActionInfo);
 	
 private:
-	/**
-	 * Attempts to pull a menu name from the supplied spawner. If one isn't 
-	 * provided, then it spawns a temporary node and pulls one from that node's 
-	 * title.
-	 * 
-	 * @param  EditorContext
-	 * @param  Action		The action you want to suss name information from.
-	 * @return A name for the menu item wrapping this action.
-	 */
-	FText GetMenuNameForAction(TWeakPtr<FBlueprintEditor> EditorContext, FBlueprintActionInfo const& ActionInfo);
-
-	/**
-	 * Attempts to pull a menu category from the supplied spawner. If one isn't 
-	 * provided, then it spawns a temporary node and pulls one from that node.
-	 * 
-	 * @param  EditorContext
-	 * @param  Action		The action you want to suss category information from.
-	 * @return A category for the menu item wrapping this action.
-	 */
-	FText GetCategoryForAction(TWeakPtr<FBlueprintEditor> EditorContext, UBlueprintNodeSpawner const* Action);
-
-	/**
-	 * Attempts to pull a menu tooltip from the supplied spawner. If one isn't 
-	 * provided, then it spawns a temporary node and pulls one from that node.
-	 * 
-	 * @param  EditorContext
-	 * @param  Action		The action you want to suss tooltip information from.
-	 * @return A tooltip for the menu item wrapping this action.
-	 */
-	FText GetTooltipForAction(TWeakPtr<FBlueprintEditor> EditorContext, UBlueprintNodeSpawner const* Action);
-
-	/**
-	 * Attempts to pull a keywords from the supplied spawner. If one isn't 
-	 * provided, then it spawns a temporary node and pulls them from that.
-	 * 
-	 * @TODO: Should search keywords be localized? Probably. 
-	 *
-	 * @param  EditorContext
-	 * @param  Action		The action you want to suss keyword information from.
-	 * @return A keyword text string for the menu item wrapping this action.
-	 */
-	FString GetSearchKeywordsForAction(TWeakPtr<FBlueprintEditor> EditorContext, UBlueprintNodeSpawner const* Action);
-
-	/**
-	 * Attempts to pull a menu icon information from the supplied spawner. If
-	 * info isn't provided, then it spawns a temporary node and pulls data from 
-	 * that node.
-	 * 
-	 * @param  EditorContext
-	 * @param  Action		The action you want to suss icon information from.
-	 * @param  ColorOut		The color to tint the icon with.
-	 * @return Name of the brush to use (use FEditorStyle::GetBrush() to resolve).
-	 */
-	FName GetMenuIconForAction(TWeakPtr<FBlueprintEditor> EditorContext, UBlueprintNodeSpawner const* Action, FLinearColor& ColorOut);
-
 	/**
 	 * Utility getter function that retrieves the blueprint context for the menu
 	 * items being made.
@@ -136,11 +80,19 @@ private:
 	/**
 	 * 
 	 * 
-	 * @param  Action	
 	 * @param  EditorContext	
 	 * @return 
 	 */
-	UEdGraphNode* GetTemplateNode(UBlueprintNodeSpawner const* Action, TWeakPtr<FBlueprintEditor> EditorContext) const;	
+	UEdGraph* GetTargetGraph(TWeakPtr<FBlueprintEditor> EditorContext) const;
+
+	/**
+	 *
+	 *
+	 * @param  EditorContext
+	 * @param  ActionInfo
+	 * @return
+	 */
+	FBlueprintActionUiSpec GetActionUiSignature(TWeakPtr<FBlueprintEditor> EditorContext, FBlueprintActionInfo const& ActionInfo);
 };
 
 //------------------------------------------------------------------------------
@@ -152,20 +104,16 @@ FBlueprintActionMenuItemFactory::FBlueprintActionMenuItemFactory(FBlueprintActio
 }
 
 //------------------------------------------------------------------------------
-TSharedPtr<FEdGraphSchemaAction> FBlueprintActionMenuItemFactory::MakeActionMenuItem(TWeakPtr<FBlueprintEditor> EditorContext, FBlueprintActionInfo const& ActionInfo)
+TSharedPtr<FBlueprintActionMenuItem> FBlueprintActionMenuItemFactory::MakeActionMenuItem(TWeakPtr<FBlueprintEditor> EditorContext, FBlueprintActionInfo const& ActionInfo)
 {
-	UBlueprintNodeSpawner const* Action = ActionInfo.NodeSpawner;
+	FBlueprintActionUiSpec UiSignature = GetActionUiSignature(EditorContext, ActionInfo);
 
-	FLinearColor IconTint = FLinearColor::White;
-	FName IconBrushName = GetMenuIconForAction(EditorContext, Action, IconTint);
+	UBlueprintNodeSpawner const* Action = ActionInfo.NodeSpawner;
+	FBlueprintActionMenuItem* NewMenuItem = new FBlueprintActionMenuItem(Action, UiSignature);
 
-	FBlueprintActionMenuItem * NewMenuItem = new FBlueprintActionMenuItem(Action, FEditorStyle::GetBrush(IconBrushName), IconTint, MenuGrouping);
-	NewMenuItem->MenuDescription    = GetMenuNameForAction(EditorContext, ActionInfo);
-	NewMenuItem->TooltipDescription = GetTooltipForAction(EditorContext, Action).ToString();
-	NewMenuItem->Category           = GetCategoryForAction(EditorContext, Action).ToString();
-	NewMenuItem->Keywords           = GetSearchKeywordsForAction(EditorContext, Action);
+	NewMenuItem->Grouping = MenuGrouping;
+	NewMenuItem->Category = FString::Printf(TEXT("%s|%s"), *RootCategory.ToString(), *NewMenuItem->Category);
 
-	NewMenuItem->Category = FString::Printf(TEXT("%s|%s"), *RootCategory.ToString(), *NewMenuItem->Category);	
 	return MakeShareable(NewMenuItem);
 }
 
@@ -180,200 +128,17 @@ TSharedPtr<FBlueprintDragDropMenuItem> FBlueprintActionMenuItemFactory::MakeDrag
 }
 
 //------------------------------------------------------------------------------
-TSharedPtr<FBlueprintBoundMenuItem> FBlueprintActionMenuItemFactory::MakeBoundMenuItem(TWeakPtr<FBlueprintEditor> EditorContext, FBlueprintActionInfo const& ActionInfo)
+TSharedPtr<FBlueprintActionMenuItem> FBlueprintActionMenuItemFactory::MakeBoundMenuItem(TWeakPtr<FBlueprintEditor> EditorContext, FBlueprintActionInfo const& ActionInfo)
 {
-	IBlueprintNodeBinder::FBindingSet const& Bindings = ActionInfo.GetBindings();
-	UBlueprintNodeSpawner const* Action = ActionInfo.NodeSpawner;
-
-	FBlueprintBoundMenuItem* NewMenuItem = new FBlueprintBoundMenuItem(Action, MenuGrouping);
-
-	// AddBindings() updates the MenuDescription, everytime a binding is added, 
-	// so set the default menu name before it
-	NewMenuItem->MenuDescription = GetMenuNameForAction(EditorContext, ActionInfo);
-	NewMenuItem->AddBindings(Bindings);
-
-	NewMenuItem->TooltipDescription = GetTooltipForAction(EditorContext, Action).ToString();
-	NewMenuItem->Category           = GetCategoryForAction(EditorContext, Action).ToString();
-	NewMenuItem->Keywords           = GetSearchKeywordsForAction(EditorContext, Action);
-
-	NewMenuItem->Category = FString::Printf(TEXT("%s|%s"), *RootCategory.ToString(), *NewMenuItem->Category);
-	return MakeShareable(NewMenuItem);
-}
+	FBlueprintActionUiSpec UiSignature = GetActionUiSignature(EditorContext, ActionInfo);
 
-//------------------------------------------------------------------------------
-FText FBlueprintActionMenuItemFactory::GetMenuNameForAction(TWeakPtr<FBlueprintEditor> EditorContext, FBlueprintActionInfo const& ActionInfo)
-{
 	UBlueprintNodeSpawner const* Action = ActionInfo.NodeSpawner;
-	check(Action != nullptr);
-	// give the action the chance to save on performance (to keep from having to spawn a template node)
-	FText MenuName = Action->GetDefaultMenuName(ActionInfo.GetBindings());
-	
-	if (MenuName.IsEmpty())
-	{
-		if (UEdGraphNode* NodeTemplate = GetTemplateNode(Action, EditorContext))
-		{
-			MenuName = NodeTemplate->GetNodeTitle(ENodeTitleType::MenuTitle);
-		}
-		else
-		{
-			// need to give it some name, this is as good as any I guess
-			MenuName = FText::FromName(Action->GetFName());
-		}
-	}
-	
-	return MenuName;
-}
-
-//------------------------------------------------------------------------------
-FText FBlueprintActionMenuItemFactory::GetCategoryForAction(TWeakPtr<FBlueprintEditor> EditorContext, UBlueprintNodeSpawner const* Action)
-{
-	check(Action != nullptr);
-	// give the action the chance to save on performance (to keep from having to spawn a template node)
-	FText MenuCategory = Action->GetDefaultMenuCategory();
-	
-	if (MenuCategory.IsEmpty())
-	{
-		// put uncategorized function calls in a member function category
-		// (sorted by their respective classes)
-		if (UBlueprintFunctionNodeSpawner const* FuncSpawner = Cast<UBlueprintFunctionNodeSpawner>(Action))
-		{
-			UFunction const* Function = FuncSpawner->GetFunction();
-			check(Function != nullptr);
-			UClass* FuncOwner = Function->GetOuterUClass();
-
-			UBlueprint* Blueprint = GetTargetBlueprint();
-			check(Blueprint != nullptr);
-			UClass* BlueprintClass = (Blueprint->SkeletonGeneratedClass != nullptr) ? Blueprint->SkeletonGeneratedClass : Blueprint->ParentClass;
+	FBlueprintActionMenuItem* NewMenuItem = new FBlueprintActionMenuItem(Action, UiSignature, ActionInfo.GetBindings());
 
-			// if this is NOT a self function call (self function calls
-			// don't get nested any deeper)
-			if (!BlueprintClass->IsChildOf(FuncOwner))
-			{
-				MenuCategory = FuncOwner->GetDisplayNameText();
-			}
-			MenuCategory = FText::Format(LOCTEXT("MemberFunctionsCategory", "{0}|Call Function"), MenuCategory);
-		}
-		else if (UBlueprintNodeSpawner const* NodeSpawner = Cast<UBlueprintNodeSpawner>(Action))
-		{
-			// Only for macro instances do we want to mess with the category to put it into a category
-			// named after the Blueprint if no other category is specified
-			if(NodeSpawner->NodeClass == UK2Node_MacroInstance::StaticClass())
-			{
-				UK2Node_MacroInstance* MacroInstance = Cast<UK2Node_MacroInstance>(NodeSpawner->GetTemplateNode());
-				if(UEdGraph* MacroGraph = MacroInstance->GetMacroGraph())
-				{
-					// Check if the MacroGraph has a category
-					FKismetUserDeclaredFunctionMetadata* MacroGraphMetadata = UK2Node_MacroInstance::GetAssociatedGraphMetadata(MacroGraph);
-					if ((MacroGraphMetadata != nullptr) && MacroGraphMetadata->Category.IsEmpty())
-					{
-						MenuCategory = MacroInstance->GetMenuCategory();
-
-						UBlueprint* MacroBlueprint = FBlueprintEditorUtils::FindBlueprintForGraph(MacroGraph);
-						if(MacroBlueprint != GetTargetBlueprint())
-						{
-							FText BlueprintDisplayName = MacroBlueprint->GeneratedClass->GetDisplayNameText();
-							FFormatNamedArguments Args;
-							Args.Add(TEXT("BlueprintDisplayName"), BlueprintDisplayName);
-							Args.Add(TEXT("MacroCategory"), MenuCategory);
-
-							MenuCategory = FText::Format(LOCTEXT("MemberMacroCategory", "{BlueprintDisplayName}|{MacroCategory}"), Args);
-						}
-					}
-				}
-			}
-		}
-
-		// If the menu category is still empty, fill it in with the node template's defined category
-		if(MenuCategory.IsEmpty())
-		{
-			// @TODO: consider moving GetMenuCategory() up into UEdGraphNode
-			if (UK2Node* NodeTemplate = Cast<UK2Node>(GetTemplateNode(Action, EditorContext)))
-			{
-				MenuCategory = NodeTemplate->GetMenuCategory();
-			}
-		}
-	}
-	else
-	{
-		if (UBlueprintVariableNodeSpawner const* VarSpawner = Cast<UBlueprintVariableNodeSpawner>(Action))
-		{
-			UProperty const* Property = VarSpawner->GetVarProperty();
-			if(Property)
-			{
-				check(Property != nullptr);
-				UClass* PropertyOwner = Property->GetTypedOuter<UClass>();
-
-				UBlueprint* Blueprint = GetTargetBlueprint();
-				check(Blueprint != nullptr);
-				UClass* BlueprintClass = (Blueprint->SkeletonGeneratedClass != nullptr) ? Blueprint->SkeletonGeneratedClass : Blueprint->ParentClass;
-
-				// if this is NOT a self function call (self function calls
-				// don't get nested any deeper)
-				if (!BlueprintClass->IsChildOf(PropertyOwner))
-				{
-					FFormatNamedArguments Args;
-					Args.Add(TEXT("PropertyDisplayName"), PropertyOwner->GetDisplayNameText());
-					Args.Add(TEXT("VariableCategory"), MenuCategory);
-					MenuCategory = FText::Format(LOCTEXT("MemberVariablesCategory", "{PropertyDisplayName}|{VariableCategory}"), Args);
-				}
-			}
-		}
-	}
-	
-	return MenuCategory;
-}
-
-//------------------------------------------------------------------------------
-FText FBlueprintActionMenuItemFactory::GetTooltipForAction(TWeakPtr<FBlueprintEditor> EditorContext, UBlueprintNodeSpawner const* Action)
-{
-	check(Action != nullptr);
-	// give the action the chance to save on performance (to keep from having to spawn a template node)
-	FText Tooltip = Action->GetDefaultMenuTooltip();
-	
-	if (Tooltip.IsEmpty())
-	{
-		if (UEdGraphNode* NodeTemplate = GetTemplateNode(Action, EditorContext))
-		{
-			Tooltip = NodeTemplate->GetTooltipText();
-		}
-	}
-	
-	return Tooltip;
-}
-
-//------------------------------------------------------------------------------
-FString FBlueprintActionMenuItemFactory::GetSearchKeywordsForAction(TWeakPtr<FBlueprintEditor> EditorContext, UBlueprintNodeSpawner const* Action)
-{
-	check(Action != nullptr);
-	// give the action the chance to save on performance (to keep from having to 
-	// spawn a template node)
-	//
-	// @TODO: Should search keywords be localized? Probably.
-	FString SearchKeywords = Action->GetDefaultSearchKeywords();
-	
-	if (SearchKeywords.IsEmpty())
-	{
-		if (UEdGraphNode* NodeTemplate = GetTemplateNode(Action, EditorContext))
-		{
-			SearchKeywords = NodeTemplate->GetKeywords();
-		}
-	}
-	
-	return SearchKeywords;
-}
+	NewMenuItem->Grouping = MenuGrouping;
+	NewMenuItem->Category = FString::Printf(TEXT("%s|%s"), *RootCategory.ToString(), *NewMenuItem->Category);
 
-//------------------------------------------------------------------------------
-FName FBlueprintActionMenuItemFactory::GetMenuIconForAction(TWeakPtr<FBlueprintEditor> EditorContext, UBlueprintNodeSpawner const* Action, FLinearColor& ColorOut)
-{
-	FName BrushName = Action->GetDefaultMenuIcon(ColorOut);
-	if (BrushName.IsNone())
-	{
-		if (UEdGraphNode* NodeTemplate = GetTemplateNode(Action, EditorContext))
-		{
-			BrushName = NodeTemplate->GetPaletteIcon(ColorOut);
-		}
-	}
-	return BrushName;
+	return MakeShareable(NewMenuItem);
 }
 
 //------------------------------------------------------------------------------
@@ -388,7 +153,7 @@ UBlueprint* FBlueprintActionMenuItemFactory::GetTargetBlueprint() const
 }
 
 //------------------------------------------------------------------------------
-UEdGraphNode* FBlueprintActionMenuItemFactory::GetTemplateNode(UBlueprintNodeSpawner const* Action, TWeakPtr<FBlueprintEditor> EditorContext) const
+UEdGraph* FBlueprintActionMenuItemFactory::GetTargetGraph(TWeakPtr<FBlueprintEditor> EditorContext) const
 {
 	UEdGraph* TargetGraph = nullptr;
 	if (Context.Graphs.Num() > 0)
@@ -409,11 +174,21 @@ UEdGraphNode* FBlueprintActionMenuItemFactory::GetTemplateNode(UBlueprintNodeSpa
 			TargetGraph = EditorContext.Pin()->GetFocusedGraph();
 		}
 	}
+	return TargetGraph;
+}
 
-	check(Action != nullptr);
-	return Action->GetTemplateNode(TargetGraph);
+//------------------------------------------------------------------------------
+FBlueprintActionUiSpec FBlueprintActionMenuItemFactory::GetActionUiSignature(TWeakPtr<FBlueprintEditor> EditorContext, FBlueprintActionInfo const& ActionInfo)
+{
+	UBlueprintNodeSpawner const* Action = ActionInfo.NodeSpawner;
+
+	UEdGraph* TargetGraph = GetTargetGraph(EditorContext);
+	Action->PrimeDefaultUiSpec(TargetGraph);
+
+	return Action->GetUiSpec(Context, ActionInfo.GetBindings());
 }
 
+
 /*******************************************************************************
  * Static FBlueprintActionMenuBuilder Helpers
  ******************************************************************************/
@@ -526,7 +301,7 @@ void FBlueprintActionMenuBuilderImpl::FMenuSectionDefinition::AddBoundMenuItems(
 {
 	UBlueprintNodeSpawner const* DatabaseAction = DatabaseActionInfo.NodeSpawner;
 
-	TSharedPtr<FBlueprintBoundMenuItem> LastMadeMenuItem;
+	TSharedPtr<FBlueprintActionMenuItem> LastMadeMenuItem;
 	bool const bConsolidate = (Flags & FBlueprintActionMenuBuilder::ConsolidateBoundActions) != 0;
 	
 	IBlueprintNodeBinder::FBindingSet CompatibleBindings;
@@ -545,12 +320,6 @@ void FBlueprintActionMenuBuilderImpl::FMenuSectionDefinition::AddBoundMenuItems(
 			// add bindings before filtering (in case tests accept/reject based off of this)
 			CompatibleBindings.Add(BindingObj);
 		}
-// 		else if (bConsolidate)
-// 		{
-// 			MadeMenuItems.Empty();
-// 			LastMadeMenuItem.Reset();
-// 			break;
-// 		}
 
 		// if BoundAction is now "full" (meaning it can take any more 
 		// bindings), or if this is the last binding to test...
@@ -577,7 +346,7 @@ void FBlueprintActionMenuBuilderImpl::FMenuSectionDefinition::AddBoundMenuItems(
 				{
 					// move these bindings over to the menu item (so we can 
 					// test the next set)
-					LastMadeMenuItem->AddBindings(CompatibleBindings);
+					LastMadeMenuItem->AppendBindings(Filter.Context, CompatibleBindings);
 				}
 			}
 			CompatibleBindings.Empty(); // do before we copy back over cached fields for DatabaseActionInfo
@@ -608,7 +377,7 @@ FBlueprintActionMenuBuilderImpl::MenuItemList FBlueprintActionMenuBuilderImpl::F
 		}
 		else if (UBlueprintDelegateNodeSpawner const* DelegateSpawner = Cast<UBlueprintDelegateNodeSpawner>(DatabaseAction.NodeSpawner))
 		{
-			ActionProperty = DelegateSpawner->GetProperty();
+			ActionProperty = DelegateSpawner->GetDelegateProperty();
 			bPassedFilter = (ActionProperty != nullptr);
 		}
 
diff --git a/Engine/Source/Editor/Kismet/Private/BlueprintActionMenuItem.cpp b/Engine/Source/Editor/Kismet/Private/BlueprintActionMenuItem.cpp
--- a/Engine/Source/Editor/Kismet/Private/BlueprintActionMenuItem.cpp
+++ b/Engine/Source/Editor/Kismet/Private/BlueprintActionMenuItem.cpp
@@ -23,6 +23,18 @@ namespace FBlueprintMenuActionItemImpl
 	 * @param  SpawnedNode	The node that was just added to the blueprint.
 	 */
 	static void DirtyBlueprintFromNewNode(UEdGraphNode* SpawnedNode);
+
+	/**
+	 * 
+	 * 
+	 * @param  Action			The action you wish to invoke.
+	 * @param  ParentGraph		The graph you want the action to spawn a node into.
+	 * @param  Location			The position in the graph that you want the node spawned.
+	 * @param  Bindings			Any bindings you want applied after the node has been spawned.
+	 * @param  bSelectNewNode	Determines if the new node should be selected after it have been spawned.
+	 * @return The spawned node (could be an existing one if the event was already placed).
+	 */
+	static UEdGraphNode* InvokeAction(const UBlueprintNodeSpawner* Action, UEdGraph* ParentGraph, FVector2D const Location, IBlueprintNodeBinder::FBindingSet const& Bindings, bool bSelectNewNode);
 }
 
 //------------------------------------------------------------------------------
@@ -48,28 +60,55 @@ static void FBlueprintMenuActionItemImpl::DirtyBlueprintFromNewNode(UEdGraphNode
 	}
 }
 
+//------------------------------------------------------------------------------
+static UEdGraphNode* FBlueprintMenuActionItemImpl::InvokeAction(const UBlueprintNodeSpawner* Action, UEdGraph* ParentGraph, FVector2D const Location, IBlueprintNodeBinder::FBindingSet const& Bindings, bool bSelectNewNode)
+{
+	// this could return an existing node
+	UEdGraphNode* SpawnedNode = Action->Invoke(ParentGraph, Bindings, Location);
+
+	// if a returned node hasn't been added to the graph yet (it must have been freshly spawned)
+	if (ParentGraph->Nodes.Find(SpawnedNode) == INDEX_NONE)
+	{
+		check(SpawnedNode != nullptr);
+
+		// @TODO: Move Modify()/AddNode() into UBlueprintNodeSpawner and pull 
+		//        selection functionality out to FBlueprintActionMenuItem::PerformAction()
+		ParentGraph->Modify();
+		ParentGraph->AddNode(SpawnedNode, /*bFromUI =*/true, bSelectNewNode);
+		// @TODO: if this spawned multiple nodes, then we should be selecting all of them
+
+		SpawnedNode->SnapToGrid(SNodePanel::GetSnapGridSize());
+
+		FBlueprintEditorUtils::AnalyticsTrackNewNode(SpawnedNode);
+		DirtyBlueprintFromNewNode(SpawnedNode);
+	}
+	// if this node already existed, then we just want to focus on that node...
+	// some node types are only allowed one instance per blueprint (like events)
+	else if (SpawnedNode != nullptr)
+	{
+		FKismetEditorUtilities::BringKismetToFocusAttentionOnObject(SpawnedNode);
+	}
+
+	return SpawnedNode;
+}
+
 /*******************************************************************************
  * FBlueprintMenuActionItem
  ******************************************************************************/
 
 //------------------------------------------------------------------------------
-FBlueprintActionMenuItem::FBlueprintActionMenuItem(UBlueprintNodeSpawner const* NodeSpawner, FSlateBrush const* MenuIcon, FSlateColor const& IconTintIn, int32 MenuGrouping/* = 0*/)
-	: IconBrush(MenuIcon)
-	, IconTint(IconTintIn)
-	, Action(NodeSpawner)
+FBlueprintActionMenuItem::FBlueprintActionMenuItem(UBlueprintNodeSpawner const* NodeSpawner, FBlueprintActionUiSpec const& UiSpec, IBlueprintNodeBinder::FBindingSet const& Bindings)
+	: Action(NodeSpawner)
+	, IconTint(UiSpec.IconTint)
+	, IconBrush(FEditorStyle::GetBrush(UiSpec.IconName))
+	, Bindings(Bindings)
 {
 	check(Action != nullptr);
-	Grouping  = MenuGrouping;
-}
 
-//------------------------------------------------------------------------------
-FBlueprintActionMenuItem::FBlueprintActionMenuItem(UBlueprintNodeSpawner const* NodeSpawner, IBlueprintNodeBinder::FBindingSet const& InBindings)
-	: IconBrush(nullptr)
-	, IconTint(FLinearColor::White)
-	, Action(NodeSpawner)
-	, Bindings(InBindings)
-{
-	check(Action != nullptr);
+	MenuDescription = UiSpec.MenuName;
+	Category = UiSpec.Category.ToString();
+	TooltipDescription = UiSpec.Tooltip.ToString();
+	Keywords = UiSpec.Keywords;
 }
 
 //------------------------------------------------------------------------------
@@ -95,47 +134,39 @@ UEdGraphNode* FBlueprintActionMenuItem::PerformAction(UEdGraph* ParentGraph, UEd
 				ModifiedLocation.X = FromNodeX - MinNodeDistance;
 			}
 		}
+
+		// modify before the call to AutowireNewNode() below
+		FromPin->Modify();
 	}
 
-	// this could return an existing node
-	UEdGraphNode* SpawnedNode = Action->Invoke(ParentGraph, Bindings, Location);
-	
-	// if a returned node hasn't been added to the graph yet (it must have been freshly spawned)
-	if (ParentGraph->Nodes.Find(SpawnedNode) == INDEX_NONE)
+	UEdGraphNode* LastSpawnedNode = nullptr;
+	auto BoundObjIt = Bindings.CreateConstIterator();
+	do
 	{
-		check(SpawnedNode != nullptr);
-
-		// @TODO: Move Modify()/AddNode() into UBlueprintNodeSpawner and pull selection functionality out here
-		ParentGraph->Modify();
-		ParentGraph->AddNode(SpawnedNode, /*bFromUI =*/true, bSelectNewNode);
-		// @TODO: if this spawned multiple nodes, then we should be selecting all of them
+		IBlueprintNodeBinder::FBindingSet BindingsSubset;
+		for (; BoundObjIt && (Action->CanBindMultipleObjects() || (BindingsSubset.Num() == 0)); ++BoundObjIt)
+		{
+			if (BoundObjIt->IsValid())
+			{
+				BindingsSubset.Add(BoundObjIt->Get());
+			}
+		}
 
-		SpawnedNode->SnapToGrid(SNodePanel::GetSnapGridSize());
-		
+		LastSpawnedNode = InvokeAction(Action, ParentGraph, ModifiedLocation, BindingsSubset, bSelectNewNode);
 		if (FromPin != nullptr)
 		{
-			// modify before the call to AutowireNewNode() below
-			FromPin->Modify();
 			// make sure to auto-wire after we position the new node (in case
 			// the auto-wire creates a conversion node to put between them)
-			SpawnedNode->AutowireNewNode(FromPin);
-		}
-		
-		FBlueprintEditorUtils::AnalyticsTrackNewNode(SpawnedNode);
-		DirtyBlueprintFromNewNode(SpawnedNode);
-	}
-	// if this node already existed, then we just want to focus on that node...
-	// some node types are only allowed one instance per blueprint (like events)
-	else if (SpawnedNode != nullptr)
-	{
-		FKismetEditorUtilities::BringKismetToFocusAttentionOnObject(SpawnedNode);
-		if (FromPin != nullptr)
-		{
-			SpawnedNode->AutowireNewNode(FromPin);
+			LastSpawnedNode->AutowireNewNode(FromPin);
 		}
-	}
+
+		// Increase the node location a safe distance so follow-up nodes are not stacked
+		ModifiedLocation.Y += UEdGraphSchema_K2::EstimateNodeHeight(LastSpawnedNode);
+
+	} while (BoundObjIt);
+	// @TODO: select ALL spawned nodes
 	
-	return SpawnedNode;
+	return LastSpawnedNode;
 }
 
 //------------------------------------------------------------------------------
@@ -167,6 +198,23 @@ void FBlueprintActionMenuItem::AddReferencedObjects(FReferenceCollector& Collect
 	Collector.AddReferencedObject(Action);
 }
 
+//------------------------------------------------------------------------------
+void FBlueprintActionMenuItem::AppendBindings(const FBlueprintActionContext& Context, IBlueprintNodeBinder::FBindingSet const& BindingSet)
+{
+	Bindings.Append(BindingSet);
+
+	FBlueprintActionUiSpec UiSpec = Action->GetUiSpec(Context, Bindings);
+	// ui signature could be dynamic, and change as bindings change
+	MenuDescription = UiSpec.MenuName;
+	// @TODO: would invalidate any category pre-pending that was done at the 
+	//        MenuBuilder level
+	//Category  = UiSpec.Category.ToString();
+	TooltipDescription = UiSpec.Tooltip.ToString();
+	Keywords  = UiSpec.Keywords;	
+	IconBrush = FEditorStyle::GetBrush(UiSpec.IconName);
+	IconTint  = UiSpec.IconTint;
+}
+
 //------------------------------------------------------------------------------
 FSlateBrush const* FBlueprintActionMenuItem::GetMenuIcon(FSlateColor& ColorOut)
 {
diff --git a/Engine/Source/Editor/Kismet/Private/BlueprintBoundMenuItem.cpp b/Engine/Source/Editor/Kismet/Private/BlueprintBoundMenuItem.cpp
deleted file mode 100644
diff --git a/Engine/Source/Editor/Kismet/Private/BlueprintDragDropMenuItem.cpp b/Engine/Source/Editor/Kismet/Private/BlueprintDragDropMenuItem.cpp
--- a/Engine/Source/Editor/Kismet/Private/BlueprintDragDropMenuItem.cpp
+++ b/Engine/Source/Editor/Kismet/Private/BlueprintDragDropMenuItem.cpp
@@ -30,7 +30,7 @@ FBlueprintDragDropMenuItem::FBlueprintDragDropMenuItem(FBlueprintActionContext c
 	UProperty const* SampleProperty = nullptr;
 	if (UBlueprintDelegateNodeSpawner const* DelegateSpawner = Cast<UBlueprintDelegateNodeSpawner const>(SampleAction))
 	{
-		SampleProperty = DelegateSpawner->GetProperty();
+		SampleProperty = DelegateSpawner->GetDelegateProperty();
 	}
 	else if (UBlueprintVariableNodeSpawner const* VariableSpawner = Cast<UBlueprintVariableNodeSpawner const>(SampleAction))
 	{
@@ -87,7 +87,7 @@ FBlueprintDragDropMenuItem::FBlueprintDragDropMenuItem(FBlueprintActionContext c
 UEdGraphNode* FBlueprintDragDropMenuItem::PerformAction(UEdGraph* ParentGraph, UEdGraphPin* FromPin, FVector2D const Location, bool bSelectNewNode/* = true*/)
 {
 	// we shouldn't get here (this should just be used for drag/drop ops), but just in case we do...
-	FBlueprintActionMenuItem BlueprintActionItem(GetSampleAction(), nullptr, FSlateColor(FLinearColor::White));
+	FBlueprintActionMenuItem BlueprintActionItem(GetSampleAction());
 	return BlueprintActionItem.PerformAction(ParentGraph, FromPin, Location, bSelectNewNode);
 }
 
@@ -172,7 +172,7 @@ TSharedPtr<FDragDropOperation> FBlueprintDragDropMenuItem::OnDragged(FNodeCreati
 	UBlueprintNodeSpawner const* SampleAction = GetSampleAction();
 	if (UBlueprintDelegateNodeSpawner const* DelegateSpawner = Cast<UBlueprintDelegateNodeSpawner const>(SampleAction))
 	{
-		if (UMulticastDelegateProperty const* Property = DelegateSpawner->GetProperty())
+		if (UMulticastDelegateProperty const* Property = DelegateSpawner->GetDelegateProperty())
 		{
 			UStruct* const PropertyOwner = CastChecked<UStruct>(Property->GetOuterUField());
 			TSharedRef<FDragDropOperation> DragDropOpRef = FKismetDelegateDragDropAction::New(Property->GetFName(), PropertyOwner, AnalyticsDelegate);
diff --git a/Engine/Source/Editor/Kismet/Private/BlueprintPaletteFavorites.cpp b/Engine/Source/Editor/Kismet/Private/BlueprintPaletteFavorites.cpp
--- a/Engine/Source/Editor/Kismet/Private/BlueprintPaletteFavorites.cpp
+++ b/Engine/Source/Editor/Kismet/Private/BlueprintPaletteFavorites.cpp
@@ -6,7 +6,6 @@
 #include "BlueprintActionMenuItem.h"
 #include "BlueprintActionFilter.h"  // for FBlueprintActionInfo
 #include "BlueprintDragDropMenuItem.h"
-#include "BlueprintBoundMenuItem.h"
 
 /*******************************************************************************
  * Static UBlueprintPaletteFavorites Helpers
@@ -168,15 +167,6 @@ FFavoritedBlueprintPaletteItem::FFavoritedBlueprintPaletteItem(TSharedPtr<FEdGra
 			static const FName CollectionSignatureKey(TEXT("ActionCollection"));
 			ActionSignature.AddNamedValue(CollectionSignatureKey, TEXT("true"));
 		}
-		else if (InPaletteAction->GetTypeId() == FBlueprintBoundMenuItem::StaticGetTypeId())
-		{
-			FBlueprintBoundMenuItem* BoundMenuItem = (FBlueprintBoundMenuItem*)InPaletteAction.Get();
-			// this is the bare action with out any bindings, you cannot 
-			// favorite actions with bindings (mainly because the database 
-			// actions are alway unbound, binding is done externally at menu 
-			// build time)
-			ActionSignature = BoundMenuItem->GetBoundAction()->GetSpawnerSignature();
-		}
 		else
 		{
 			ActionSignature = BlueprintPaletteFavoritesImpl::ConstructLegacySignature(InPaletteAction);
diff --git a/Engine/Source/Editor/Kismet/Private/SBlueprintPalette.cpp b/Engine/Source/Editor/Kismet/Private/SBlueprintPalette.cpp
--- a/Engine/Source/Editor/Kismet/Private/SBlueprintPalette.cpp
+++ b/Engine/Source/Editor/Kismet/Private/SBlueprintPalette.cpp
@@ -28,7 +28,6 @@
 #include "AnimationTransitionGraph.h"
 #include "BlueprintActionMenuItem.h"
 #include "BlueprintDragDropMenuItem.h"
-#include "BlueprintBoundMenuItem.h"
 #include "BlueprintNodeSpawner.h"
 #include "TutorialMetaData.h"
 #include "BlueprintEditorSettings.h" // for bShowActionMenuItemSignatures
@@ -370,11 +369,6 @@ static void GetPaletteItemIcon(TSharedPtr<FEdGraphSchemaAction> ActionIn, UBluep
 		FBlueprintDragDropMenuItem* DragDropAction = (FBlueprintDragDropMenuItem*)ActionIn.Get();
 		BrushOut = DragDropAction->GetMenuIcon(ColorOut);
 	}
-	else if (ActionIn->GetTypeId() == FBlueprintBoundMenuItem::StaticGetTypeId())
-	{
-		FBlueprintBoundMenuItem* BoundAction = (FBlueprintBoundMenuItem*)ActionIn.Get();
-		BrushOut = BoundAction->GetMenuIcon(ColorOut);
-	}
 	else if (ActionIn->GetTypeId() == FEdGraphSchemaAction_K2AddDocumentation::StaticGetTypeId())
 	{
 		ToolTipOut = LOCTEXT("Documentation_Tooltip", "Create a documentation node.").ToString();
diff --git a/Engine/Source/Editor/Kismet/Public/BlueprintActionMenuItem.h b/Engine/Source/Editor/Kismet/Public/BlueprintActionMenuItem.h
--- a/Engine/Source/Editor/Kismet/Public/BlueprintActionMenuItem.h
+++ b/Engine/Source/Editor/Kismet/Public/BlueprintActionMenuItem.h
@@ -8,9 +8,10 @@
 #include "BlueprintActionMenuItem.generated.h"
 
 // Forward declarations
-class UBlueprintNodeSpawner;
-class FDragDropOperation;
+class  UBlueprintNodeSpawner;
 struct FSlateBrush;
+struct FBlueprintActionUiSpec;
+struct FBlueprintActionContext;
 
 /**
  * Wrapper around a UBlueprintNodeSpawner, which takes care of specialized
@@ -22,14 +23,12 @@ USTRUCT()
 struct FBlueprintActionMenuItem : public FEdGraphSchemaAction
 {
 	GENERATED_USTRUCT_BODY();
-	
-public:
 	static FName StaticGetTypeId() { static FName const TypeId("FBlueprintActionMenuItem"); return TypeId; }
-	
+
+public:	
 	/** Constructors */
-	FBlueprintActionMenuItem() : IconBrush(nullptr), IconTint(FLinearColor::White), Action(nullptr) {}
-	FBlueprintActionMenuItem(UBlueprintNodeSpawner const* NodeSpawner, FSlateBrush const* MenuIcon, FSlateColor const& IconTint, int32 MenuGrouping = 0);
-	FBlueprintActionMenuItem(UBlueprintNodeSpawner const* NodeSpawner, IBlueprintNodeBinder::FBindingSet const& Bindings);
+	FBlueprintActionMenuItem(UBlueprintNodeSpawner const* NodeSpawner = nullptr) : Action(NodeSpawner), IconTint(FLinearColor::White), IconBrush(nullptr) {}
+	FBlueprintActionMenuItem(UBlueprintNodeSpawner const* NodeSpawner, FBlueprintActionUiSpec const& UiDef, IBlueprintNodeBinder::FBindingSet const& Bindings = IBlueprintNodeBinder::FBindingSet());
 	
 	// FEdGraphSchemaAction interface
 	virtual FName         GetTypeId() const final { return StaticGetTypeId(); }
@@ -38,6 +37,11 @@ public:
 	virtual void          AddReferencedObjects(FReferenceCollector& Collector) final;
 	// End FEdGraphSchemaAction interface
 
+	/** @return */
+	UBlueprintNodeSpawner const* GetRawAction() const;
+
+	void AppendBindings(const FBlueprintActionContext& Context, IBlueprintNodeBinder::FBindingSet const& BindingSet);
+
 	/**
 	 * Retrieves the icon brush for this menu entry (to be displayed alongside
 	 * in the menu).
@@ -47,16 +51,15 @@ public:
 	 */
 	FSlateBrush const* GetMenuIcon(FSlateColor& ColorOut);
 
-	/** @return */
-	UBlueprintNodeSpawner const* GetRawAction() const;
+	
 
 private:
-	/** Brush that should be used for the icon on this menu item. */
-	FSlateBrush const* IconBrush;
-	/** Tint to return along with the icon brush. */
-	FSlateColor IconTint;
 	/** Specialized node-spawner, that comprises the action portion of this menu entry. */
 	UBlueprintNodeSpawner const* Action;
+	/** Tint to return along with the icon brush. */
+	FSlateColor IconTint;
+	/** Brush that should be used for the icon on this menu item. */
+	FSlateBrush const* IconBrush;
 	/** */
 	IBlueprintNodeBinder::FBindingSet Bindings;
 };
diff --git a/Engine/Source/Editor/Kismet/Public/BlueprintBoundMenuItem.h b/Engine/Source/Editor/Kismet/Public/BlueprintBoundMenuItem.h
deleted file mode 100644
diff --git a/Engine/Source/Editor/UnrealEd/Private/Commandlets/DumpBlueprintsInfoCommandlet.cpp b/Engine/Source/Editor/UnrealEd/Private/Commandlets/DumpBlueprintsInfoCommandlet.cpp
--- a/Engine/Source/Editor/UnrealEd/Private/Commandlets/DumpBlueprintsInfoCommandlet.cpp
+++ b/Engine/Source/Editor/UnrealEd/Private/Commandlets/DumpBlueprintsInfoCommandlet.cpp
@@ -1161,7 +1161,7 @@ static bool DumpBlueprintInfoUtils::DumpActionDatabaseInfo(uint32 Indent, FArchi
 				}
 				SpawnerInfo.TotalPrimingTime += (TotalPrimingTime - OldPrimingTime);
 
-				if (UEdGraphNode* TemplateNode = BpAction->GetTemplateNode(NoInit))
+				if (UEdGraphNode* TemplateNode = BpAction->GetCachedTemplateNode())
 				{
 					UObject* TemplateOuter = TemplateNode->GetOuter();
 					while ((TemplateOuter != nullptr) && (Cast<UBlueprint>(TemplateOuter) == nullptr))
diff --git a/Engine/Source/Editor/UnrealEd/Private/EditorCategoryUtils.cpp b/Engine/Source/Editor/UnrealEd/Private/EditorCategoryUtils.cpp
--- a/Engine/Source/Editor/UnrealEd/Private/EditorCategoryUtils.cpp
+++ b/Engine/Source/Editor/UnrealEd/Private/EditorCategoryUtils.cpp
@@ -83,6 +83,7 @@ FEditorCategoryUtilsImpl::FCategoryInfoMap& FEditorCategoryUtilsImpl::GetCategor
 		RegisterCategoryKey("Utilities", LOCTEXT("UtilitiesCategory", "Utilities"));
 		RegisterCategoryKey("Delegates", LOCTEXT("DelegatesCategory", "Event Dispatchers"));
 		RegisterCategoryKey("Variables", LOCTEXT("VariablesCategory", "Variables"));
+		RegisterCategoryKey("Class", LOCTEXT("ClassCategory", "Class"));
 		RegisterCategoryKey("UserInterface", LOCTEXT("UserInterfaceCategory", "User Interface"));
 		RegisterCategoryKey("AnimNotify", LOCTEXT("AnimNotifyCategory", "Add AnimNotify Event"));
 		RegisterCategoryKey("BranchPoint", LOCTEXT("BranchPointCategory", "Add Montage Branching Point Event"));
@@ -196,6 +197,7 @@ FText const& FEditorCategoryUtils::GetCommonCategory(const FCommonEditorCategory
 		CommonCategoryKeys.Add(FCommonEditorCategory::Utilities, "Utilities");
 		CommonCategoryKeys.Add(FCommonEditorCategory::Delegates, "Delegates");
 		CommonCategoryKeys.Add(FCommonEditorCategory::Variables, "Variables");
+		CommonCategoryKeys.Add(FCommonEditorCategory::Class, "Class");
 		CommonCategoryKeys.Add(FCommonEditorCategory::UserInterface, "UserInterface");
 		CommonCategoryKeys.Add(FCommonEditorCategory::AnimNotify, "AnimNotify");
 		CommonCategoryKeys.Add(FCommonEditorCategory::BranchPoint, "BranchPoint");
@@ -233,6 +235,8 @@ FText FEditorCategoryUtils::BuildCategoryString(FCommonEditorCategory::EValue Ro
 	}
 	else
 	{
+		// @TODO: FText::Format() is expensive, since this is category 
+		//        concatenation, we could just do FString concatenation
 		ConstructedCategory = FText::Format(LOCTEXT("ConcatedCategory", "{0}|{1}"), RootCategory, SubCategory);
 	}
 	return ConstructedCategory;
diff --git a/Engine/Source/Editor/UnrealEd/Public/EditorCategoryUtils.h b/Engine/Source/Editor/UnrealEd/Public/EditorCategoryUtils.h
--- a/Engine/Source/Editor/UnrealEd/Public/EditorCategoryUtils.h
+++ b/Engine/Source/Editor/UnrealEd/Public/EditorCategoryUtils.h
@@ -46,6 +46,7 @@ struct FCommonEditorCategory
 		Macro,
 		Delegates,
 
+		Class,
 		Variables,
 	};
 };
