Adding in the option to use "ExposeFunctionCategories" metadata on properties for the new menu system.

diff --git a/Engine/Source/Editor/BlueprintGraph/Private/BlueprintActionDatabase.cpp b/Engine/Source/Editor/BlueprintGraph/Private/BlueprintActionDatabase.cpp
--- a/Engine/Source/Editor/BlueprintGraph/Private/BlueprintActionDatabase.cpp
+++ b/Engine/Source/Editor/BlueprintGraph/Private/BlueprintActionDatabase.cpp
@@ -718,7 +718,7 @@ static void BlueprintActionDatabaseImpl::GetNodeSpecificActions(TSubclassOf<UEdG
 	}
 	else if (NodeClass == UEdGraphNode_Documentation::StaticClass())
 	{
-		// @TODO: BOO! (see comment above)
+		// @TODO: BOOOOOOO! (see comment above)
 		Registrar.AddBlueprintAction(MakeDocumentationNodeSpawner<UEdGraphNode_Documentation>());
 	}
 }
diff --git a/Engine/Source/Editor/BlueprintGraph/Private/BlueprintEditorSettings.cpp b/Engine/Source/Editor/BlueprintGraph/Private/BlueprintEditorSettings.cpp
--- a/Engine/Source/Editor/BlueprintGraph/Private/BlueprintEditorSettings.cpp
+++ b/Engine/Source/Editor/BlueprintGraph/Private/BlueprintEditorSettings.cpp
@@ -11,6 +11,7 @@ UBlueprintEditorSettings::UBlueprintEditorSettings(const class FPostConstructIni
 	bShowGraphInstructionText = true;
 	NodeTemplateCacheCapMB    = 20.f;
 	bUseTargetContextForNodeMenu = true;
+	bExposeAllMemberComponentFunctions = true;
 
 	// settings that were moved out of experimental...
 	UEditorExperimentalSettings const* ExperimentalSettings = GetDefault<UEditorExperimentalSettings>();
diff --git a/Engine/Source/Editor/BlueprintGraph/Private/BlueprintVariableNodeSpawner.cpp b/Engine/Source/Editor/BlueprintGraph/Private/BlueprintVariableNodeSpawner.cpp
--- a/Engine/Source/Editor/BlueprintGraph/Private/BlueprintVariableNodeSpawner.cpp
+++ b/Engine/Source/Editor/BlueprintGraph/Private/BlueprintVariableNodeSpawner.cpp
@@ -122,8 +122,6 @@ UBlueprintVariableNodeSpawner* UBlueprintVariableNodeSpawner::Create(TSubclassOf
 	}
 	// add at least one character, so that PrimeDefaultMenuSignature() doesn't 
 	// attempt to query the template node
-	//
-	// @TODO: maybe UPROPERTY() fields should have keyword metadata like functions
 	MenuSignature.Keywords = TEXT(" ");
 	MenuSignature.IconName = UK2Node_Variable::GetVarIconFromPinType(NodeSpawner->GetVarType(), MenuSignature.IconTint);
 
diff --git a/Engine/Source/Editor/BlueprintGraph/Public/BlueprintEditorSettings.h b/Engine/Source/Editor/BlueprintGraph/Public/BlueprintEditorSettings.h
--- a/Engine/Source/Editor/BlueprintGraph/Public/BlueprintEditorSettings.h
+++ b/Engine/Source/Editor/BlueprintGraph/Public/BlueprintEditorSettings.h
@@ -28,6 +28,10 @@ public:
 	UPROPERTY(EditAnywhere, config, Category=UserExperience)
 	bool bUseTargetContextForNodeMenu;
 
+	/** If enabled, then ALL component functions are exposed to the context menu (when the contextual target is a component owner). Ignores "ExposeFunctionCategories" metadata for components. */
+	UPROPERTY(EditAnywhere, config, Category=UserExperience, meta=(DisplayName="Context Menu: Expose All Sub-Component Functions"))
+	bool bExposeAllMemberComponentFunctions;
+
 // Developer Settings
 public:
 	/** If enabled, tooltips on action menu items will show the associated action's signature id (can be used to setup custom favorites menus).*/
diff --git a/Engine/Source/Editor/Kismet/Private/BlueprintActionMenuUtils.cpp b/Engine/Source/Editor/Kismet/Private/BlueprintActionMenuUtils.cpp
--- a/Engine/Source/Editor/Kismet/Private/BlueprintActionMenuUtils.cpp
+++ b/Engine/Source/Editor/Kismet/Private/BlueprintActionMenuUtils.cpp
@@ -69,6 +69,15 @@ namespace BlueprintActionMenuUtilsImpl
 	 */
 	static bool IsUnexposedMemberAction(FBlueprintActionFilter const& Filter, FBlueprintActionInfo& BlueprintAction);
 
+	/**
+	 * 
+	 * 
+	 * @param  Filter	
+	 * @param  BlueprintAction	
+	 * @return 
+	 */
+	static bool IsUnexposedNonComponentAction(FBlueprintActionFilter const& Filter, FBlueprintActionInfo& BlueprintAction);
+
 	/**
 	 * Utility function to find a common base class from a set of given classes.
 	 * 
@@ -85,6 +94,14 @@ namespace BlueprintActionMenuUtilsImpl
 	 */
 	static UClass* GetPinClassType(UEdGraphPin const* Pin);
 
+	/**
+	 * 
+	 * 
+	 * @param  MainMenuFilter	
+	 * @return 
+	 */
+	static FBlueprintActionFilter MakeCallOnMemberFilter(FBlueprintActionFilter const& MainMenuFilter);
+
 	/**
 	 * 
 	 * 
@@ -145,9 +162,9 @@ static bool BlueprintActionMenuUtilsImpl::IsUnexposedMemberAction(FBlueprintActi
 		TArray<FString> AllExposedCategories;
 		for (TWeakObjectPtr<UObject> Binding : BlueprintAction.GetBindings())
 		{
-			if (UObjectProperty* ObjectProperty = Cast<UObjectProperty>(Binding.Get()))
+			if (UProperty* Property = Cast<UProperty>(Binding.Get()))
 			{
-				FString const ExposedCategoryMetadata = ObjectProperty->GetMetaData(FBlueprintMetadata::MD_ExposeFunctionCategories);
+				FString const ExposedCategoryMetadata = Property->GetMetaData(FBlueprintMetadata::MD_ExposeFunctionCategories);
 				if (ExposedCategoryMetadata.IsEmpty())
 				{
 					continue;
@@ -165,6 +182,31 @@ static bool BlueprintActionMenuUtilsImpl::IsUnexposedMemberAction(FBlueprintActi
 	return bIsFliteredOut;
 }
 
+//------------------------------------------------------------------------------
+static bool BlueprintActionMenuUtilsImpl::IsUnexposedNonComponentAction(FBlueprintActionFilter const& Filter, FBlueprintActionInfo& BlueprintAction)
+{
+	bool bIsFliteredOut = false;
+
+	for (TWeakObjectPtr<UObject> Binding : BlueprintAction.GetBindings())
+	{
+		if (UObjectProperty* ObjectProperty = Cast<UObjectProperty>(Binding.Get()))
+		{
+			bool const bIsComponent = ObjectProperty->PropertyClass->IsChildOf<UActorComponent>();
+			// ignoring components for this rejection test
+			if (bIsComponent)
+			{
+				continue;
+			}
+		}
+		// else, it's not a component... let's do this!
+
+		bIsFliteredOut = IsUnexposedMemberAction(Filter, BlueprintAction);
+		break;
+	}
+
+	return bIsFliteredOut;
+}
+
 //------------------------------------------------------------------------------
 static UClass* BlueprintActionMenuUtilsImpl::FindCommonBaseClass(TArray<UObject*> const& ObjectSet)
 {
@@ -207,6 +249,45 @@ static UClass* BlueprintActionMenuUtilsImpl::GetPinClassType(UEdGraphPin const*
 	return PinObjClass;
 }
 
+//------------------------------------------------------------------------------
+static FBlueprintActionFilter BlueprintActionMenuUtilsImpl::MakeCallOnMemberFilter(FBlueprintActionFilter const& MainMenuFilter)
+{
+	FBlueprintActionFilter CallOnMemberFilter;
+	CallOnMemberFilter.Context = MainMenuFilter.Context;
+	CallOnMemberFilter.PermittedNodeTypes.Add(UK2Node_CallFunction::StaticClass());
+	CallOnMemberFilter.AddRejectionTest(FBlueprintActionFilter::FRejectionTestDelegate::CreateStatic(IsUnBoundSpawner));
+
+	const UBlueprintEditorSettings* BlueprintSettings = GetDefault<UBlueprintEditorSettings>();
+	// instead of looking for "ExposeFunctionCategories" on component properties,
+	// we just expose functions for all components, but we still need to check
+	// for "ExposeFunctionCategories" on any non-component properties... 
+	if (BlueprintSettings->bExposeAllMemberComponentFunctions)
+	{
+		CallOnMemberFilter.AddRejectionTest(FBlueprintActionFilter::FRejectionTestDelegate::CreateStatic(IsUnexposedNonComponentAction));
+	}
+	else
+	{
+		CallOnMemberFilter.AddRejectionTest(FBlueprintActionFilter::FRejectionTestDelegate::CreateStatic(IsUnexposedMemberAction));
+	}
+
+	for (UClass const* TargetClass : MainMenuFilter.TargetClasses)
+	{
+		for (TFieldIterator<UObjectProperty> PropertyIt(TargetClass, EFieldIteratorFlags::IncludeSuper); PropertyIt; ++PropertyIt)
+		{
+			UObjectProperty* ObjectProperty = *PropertyIt;
+
+			if (!ObjectProperty->HasAnyPropertyFlags(CPF_BlueprintVisible) ||
+				!(ObjectProperty->PropertyClass->IsChildOf<UActorComponent>() || ObjectProperty->HasMetaData(FBlueprintMetadata::MD_ExposeFunctionCategories)))
+			{
+				continue;
+			}
+			CallOnMemberFilter.Context.SelectedObjects.Add(ObjectProperty);
+		}
+	}
+
+	return CallOnMemberFilter;
+}
+
 //------------------------------------------------------------------------------
 static void BlueprintActionMenuUtilsImpl::AddComponentSections(FBlueprintActionFilter const& ComponentsFilter, FBlueprintActionMenuBuilder& MenuOut)
 {
@@ -403,32 +484,8 @@ void FBlueprintActionMenuUtils::MakeContextMenu(FBlueprintActionContext const& C
 		}
 	}
 
-	FBlueprintActionFilter CallOnMemberFilter;
-	CallOnMemberFilter.Context = MainMenuFilter.Context;
-	CallOnMemberFilter.PermittedNodeTypes.Add(UK2Node_CallFunction::StaticClass());
-	CallOnMemberFilter.AddRejectionTest(FBlueprintActionFilter::FRejectionTestDelegate::CreateStatic(IsUnBoundSpawner));
-	// instead of looking for "ExposeFunctionCategories" on object properties,
-	// we just expose functions for all components, this is problematic if we 
-	// were using the "ExposeFunctionCategories" on any non-component object 
-	// properties... if we re-enable this, then make sure to remove the 
-	// IsChildOf<UActorComponent>() check below
-	//CallOnMemberFilter.AddRejectionTest(FBlueprintActionFilter::FRejectionTestDelegate::CreateStatic(IsUnexposedMemberAction));
-
-	for (UClass const* TargetClass : MainMenuFilter.TargetClasses)
-	{
-		for (TFieldIterator<UObjectProperty> PropertyIt(TargetClass, EFieldIteratorFlags::IncludeSuper); PropertyIt; ++PropertyIt)
-		{
-			UObjectProperty* ComponentProperty = *PropertyIt;
-			// @TODO: consider a more relaxed scenario where we allow other
-			//        (non-component) properties as options here too
-			if (!ComponentProperty->PropertyClass->IsChildOf<UActorComponent>() || 
-				!ComponentProperty->HasAnyPropertyFlags(CPF_BlueprintVisible))
-			{
-				continue;
-			}
-			CallOnMemberFilter.Context.SelectedObjects.Add(ComponentProperty);
-		}
-	}
+	// should be called AFTER the MainMenuFilter if fully constructed
+	FBlueprintActionFilter CallOnMemberFilter = MakeCallOnMemberFilter(MainMenuFilter);
 
 	FBlueprintActionFilter AddComponentFilter;
 	AddComponentFilter.Context = MainMenuFilter.Context;
