Consolidating where we call UEdGraphNode::AddNode() for newly spawned nodes (to avoid duplicated code... centralized in UBlueprintNodeSpawner).

diff --git a/Engine/Source/Editor/BlueprintGraph/Private/BlueprintFunctionNodeSpawner.cpp b/Engine/Source/Editor/BlueprintGraph/Private/BlueprintFunctionNodeSpawner.cpp
--- a/Engine/Source/Editor/BlueprintGraph/Private/BlueprintFunctionNodeSpawner.cpp
+++ b/Engine/Source/Editor/BlueprintGraph/Private/BlueprintFunctionNodeSpawner.cpp
@@ -102,10 +102,6 @@ static bool BlueprintFunctionNodeSpawnerImpl::BindFunctionNode(UK2Node_CallFunct
 
 	BindingOffset.Y += UEdGraphSchema_K2::EstimateNodeHeight(BindingNode);
 
-	// @TODO: Move this down into Invoke()
-	ParentGraph->Modify();
-	ParentGraph->AddNode(BindingNode, /*bFromUI =*/false, /*bSelectNewNode =*/false);
-
 	UEdGraphPin* LiteralOutput = BindingNode->GetValuePin();
 	UEdGraphPin* CallSelfInput = NewNode->FindPin(GetDefault<UEdGraphSchema_K2>()->PN_Self);
 	// connect the new "get-var" node with the spawned function node
diff --git a/Engine/Source/Editor/BlueprintGraph/Private/BlueprintNodeSpawner.cpp b/Engine/Source/Editor/BlueprintGraph/Private/BlueprintNodeSpawner.cpp
--- a/Engine/Source/Editor/BlueprintGraph/Private/BlueprintNodeSpawner.cpp
+++ b/Engine/Source/Editor/BlueprintGraph/Private/BlueprintNodeSpawner.cpp
@@ -275,6 +275,10 @@ UEdGraphNode* UBlueprintNodeSpawner::SpawnEdGraphNode(TSubclassOf<UEdGraphNode>
 			NewNode->SetFlags(RF_Transactional);
 			NewNode->AllocateDefaultPins();
 			NewNode->PostPlacedNewNode();
+
+			ParentGraph->Modify();
+			// the FBlueprintMenuActionItem should do the selecting
+			ParentGraph->AddNode(NewNode, /*bFromUI =*/true, /*bSelectNewNode =*/false);
 		}
 
 		ApplyBindings(NewNode, Bindings);
diff --git a/Engine/Source/Editor/GraphEditor/Private/SGraphPanel.cpp b/Engine/Source/Editor/GraphEditor/Private/SGraphPanel.cpp
--- a/Engine/Source/Editor/GraphEditor/Private/SGraphPanel.cpp
+++ b/Engine/Source/Editor/GraphEditor/Private/SGraphPanel.cpp
@@ -951,8 +951,7 @@ void SGraphPanel::AddNode(UEdGraphNode* Node)
 	TSharedPtr<SGraphNode> NewNode = FNodeFactory::CreateNodeWidget(Node);
 	check(NewNode.IsValid());
 
-	FEdGraphEditAction* GraphAction = UserAddedNodes.Find(Node);
-	const bool bWasUserAdded = (GraphAction != nullptr) ? true : false;
+	const bool bWasUserAdded = (UserAddedNodes.Find(Node) != nullptr);
 
 	NewNode->SetIsEditable(IsEditable);
 	NewNode->SetDoubleClickEvent(OnNodeDoubleClicked);
@@ -972,12 +971,6 @@ void SGraphPanel::AddNode(UEdGraphNode* Node)
 
 		NewNode->PlaySpawnEffect();
 
-		// Do not Select nodes unless they are marked for selection
-		if ((GraphAction->Action & GRAPHACTION_SelectNode) != 0)
-		{
-			SelectAndCenterObject(Node, false);
-		}
-
 		NewNode->UpdateGraphNode();
 		NewNode->RequestRenameOnSpawn();
 	}
@@ -997,22 +990,33 @@ void SGraphPanel::Update()
 	// Add widgets for all the nodes that don't have one.
 	if(GraphObj != NULL)
 	{
-		int32 NumNodesAdded = 0;
-
 		// Scan for all missing nodes
-		for ( int32 NodeIndex = 0; NodeIndex < GraphObj->Nodes.Num(); ++NodeIndex )
+		for (int32 NodeIndex = 0; NodeIndex < GraphObj->Nodes.Num(); ++NodeIndex)
 		{
 			UEdGraphNode* Node = GraphObj->Nodes[NodeIndex];
 			if (Node)
 			{
 				AddNode(Node);
-				++NumNodesAdded;
 			}
 			else
 			{
 				UE_LOG(LogGraphPanel, Warning, TEXT("Found NULL Node in GraphObj array. A node type has been deleted without creating an ActiveClassRedictor to K2Node_DeadClass."));
 			}
 		}
+
+		// find the last selection action, and execute it
+		for (int32 ActionIndex = UserActions.Num() - 1; ActionIndex >= 0; --ActionIndex)
+		{
+			if (UserActions[ActionIndex].Action & GRAPHACTION_SelectNode)
+			{
+				DeferredSelectionTargetObjects.Empty();
+				for (const UEdGraphNode* Node : UserActions[ActionIndex].Nodes)
+				{
+					DeferredSelectionTargetObjects.Add(Node);
+				}
+				break;
+			}
+		}
 	}
 	else
 	{
@@ -1021,6 +1025,7 @@ void SGraphPanel::Update()
 
 	// Clean out set of added nodes
 	UserAddedNodes.Empty();
+	UserActions.Empty();
 
 	// Invoke any delegate methods
 	OnUpdateGraphPanel.ExecuteIfBound();
@@ -1106,16 +1111,23 @@ void SGraphPanel::JumpToPin(const UEdGraphPin* JumpToMe)
 	}
 }
 
-void SGraphPanel::OnGraphChanged ( const FEdGraphEditAction& EditAction)
+void SGraphPanel::OnGraphChanged(const FEdGraphEditAction& EditAction)
 {
-	if ((EditAction.Graph == GraphObj ) 
-		&& (EditAction.Action & GRAPHACTION_AddNode) 
-		&& (EditAction.Node != NULL)
+	if ((EditAction.Graph == GraphObj) &&
+		(EditAction.Nodes.Num() > 0) &&
 		// We do not want to mark it as a UserAddedNode for graphs that do not currently have focus,
 		// this causes each one to want to do the effects and rename, which causes problems.
-		&& (HasKeyboardFocus() || HasFocusedDescendants()))
+		(HasKeyboardFocus() || HasFocusedDescendants()))
 	{
-		UserAddedNodes.Add(EditAction.Node, EditAction);
+		int32 ActionIndex = UserActions.Num();
+		if (EditAction.Action & GRAPHACTION_AddNode)
+		{
+			for (const UEdGraphNode* Node : EditAction.Nodes)
+			{
+				UserAddedNodes.Add(Node, ActionIndex);
+			}
+		}
+		UserActions.Add(EditAction);
 	}
 }
 
diff --git a/Engine/Source/Editor/GraphEditor/Private/SGraphPanel.h b/Engine/Source/Editor/GraphEditor/Private/SGraphPanel.h
--- a/Engine/Source/Editor/GraphEditor/Private/SGraphPanel.h
+++ b/Engine/Source/Editor/GraphEditor/Private/SGraphPanel.h
@@ -240,8 +240,11 @@ protected:
 	TAttribute<bool> ShowGraphStateOverlay;
 
 private:
-	/** Map of recently added nodes for the panel */
-	TMap<class UEdGraphNode*, FEdGraphEditAction> UserAddedNodes;
+	/** Ordered list of user actions, as they came in */
+	TArray<FEdGraphEditAction> UserActions;
+
+	/** Map of recently added nodes for the panel (maps from added nodes to UserActions indices) */
+	TMap<const class UEdGraphNode*, int32> UserAddedNodes;
 
 	FOnGraphChanged::FDelegate MyRegisteredGraphChangedDelegate;
 private:
diff --git a/Engine/Source/Editor/GraphEditor/Private/SNodePanel.cpp b/Engine/Source/Editor/GraphEditor/Private/SNodePanel.cpp
--- a/Engine/Source/Editor/GraphEditor/Private/SNodePanel.cpp
+++ b/Engine/Source/Editor/GraphEditor/Private/SNodePanel.cpp
@@ -297,7 +297,7 @@ void SNodePanel::Construct()
 	bAllowContinousZoomInterpolation = false;
 	bTeleportInsteadOfScrollingWhenZoomingToFit = false;
 
-	DeferredSelectionTargetObject = NULL;
+	DeferredSelectionTargetObjects.Empty();
 	DeferredMovementTargetObject = NULL;
 
 	bIsPanning = false;
@@ -406,19 +406,23 @@ void SNodePanel::OnEndNodeInteraction(const TSharedRef<SNode>& InNodeToDrag)
 // Ticks this widget.  Override in derived classes, but always call the parent implementation.
 void SNodePanel::Tick(const FGeometry& AllottedGeometry, const double InCurrentTime, const float InDeltaTime)
 {
-	if(DeferredSelectionTargetObject != NULL)
+	if(DeferredSelectionTargetObjects.Num() > 0)
 	{
-		TSharedRef<SNode>* pWidget = NodeToWidgetLookup.Find(DeferredSelectionTargetObject);
-		if (pWidget != NULL)
+		FGraphPanelSelectionSet NewSelectionSet;
+		for (const UObject* SelectionTarget : DeferredSelectionTargetObjects)
 		{
-			SelectionManager.SelectSingleNode(const_cast<SelectedItemType>(DeferredSelectionTargetObject));
-			DeferredSelectionTargetObject = NULL;
+			if (TSharedRef<SNode>* pWidget = NodeToWidgetLookup.Find(SelectionTarget))
+			{
+				NewSelectionSet.Add(const_cast<SelectedItemType>(SelectionTarget));
+			}
 		}
-		else
+
+		if (NewSelectionSet.Num() > 0)
 		{
-			// The selection target does not have a corresponding widget. Just end the selection.
-			DeferredSelectionTargetObject = NULL;
+			SelectionManager.SetSelectionSet(NewSelectionSet);
 		}
+		DeferredSelectionTargetObjects.Empty();
+
 
 		// Since we want to move to a target object, do not zoom to extent. Panning and zoom will not begin until next tick however due to the nodes potentially not having a size yet.
 		if( DeferredMovementTargetObject )
@@ -1056,7 +1060,8 @@ FSlateRect SNodePanel::ComputeSensibleGraphBounds() const
 
 void SNodePanel::SelectAndCenterObject(const UObject* ObjectToSelect, bool bCenter)
 {
-	DeferredSelectionTargetObject = ObjectToSelect;
+	DeferredSelectionTargetObjects.Empty();
+	DeferredSelectionTargetObjects.Add(ObjectToSelect);
 
 	if( bCenter )
 	{
diff --git a/Engine/Source/Editor/GraphEditor/Public/SNodePanel.h b/Engine/Source/Editor/GraphEditor/Public/SNodePanel.h
--- a/Engine/Source/Editor/GraphEditor/Public/SNodePanel.h
+++ b/Engine/Source/Editor/GraphEditor/Public/SNodePanel.h
@@ -623,8 +623,8 @@ protected:
 	/** Given a node, find the corresponding widget */
 	TMap< UObject*, TSharedRef<SNode> > NodeToWidgetLookup;
 
-	/** If non-null and a part of this panel, this node will be selected and brought into view on the next Tick */
-	const UObject* DeferredSelectionTargetObject;
+	/** If not empty and a part of this panel, this node will be selected and brought into view on the next Tick */
+	TSet<const UObject*> DeferredSelectionTargetObjects;
 	/** If non-null and a part of this panel, this node will be brought into view on the next Tick */
 	const UObject* DeferredMovementTargetObject;
 
diff --git a/Engine/Source/Editor/Kismet/Private/BlueprintActionMenuItem.cpp b/Engine/Source/Editor/Kismet/Private/BlueprintActionMenuItem.cpp
--- a/Engine/Source/Editor/Kismet/Private/BlueprintActionMenuItem.cpp
+++ b/Engine/Source/Editor/Kismet/Private/BlueprintActionMenuItem.cpp
@@ -31,10 +31,9 @@ namespace FBlueprintMenuActionItemImpl
 	 * @param  ParentGraph		The graph you want the action to spawn a node into.
 	 * @param  Location			The position in the graph that you want the node spawned.
 	 * @param  Bindings			Any bindings you want applied after the node has been spawned.
-	 * @param  bSelectNewNode	Determines if the new node should be selected after it have been spawned.
 	 * @return The spawned node (could be an existing one if the event was already placed).
 	 */
-	static UEdGraphNode* InvokeAction(const UBlueprintNodeSpawner* Action, UEdGraph* ParentGraph, FVector2D const Location, IBlueprintNodeBinder::FBindingSet const& Bindings, bool bSelectNewNode);
+	static UEdGraphNode* InvokeAction(const UBlueprintNodeSpawner* Action, UEdGraph* ParentGraph, FVector2D const Location, IBlueprintNodeBinder::FBindingSet const& Bindings);
 }
 
 //------------------------------------------------------------------------------
@@ -61,22 +60,16 @@ static void FBlueprintMenuActionItemImpl::DirtyBlueprintFromNewNode(UEdGraphNode
 }
 
 //------------------------------------------------------------------------------
-static UEdGraphNode* FBlueprintMenuActionItemImpl::InvokeAction(const UBlueprintNodeSpawner* Action, UEdGraph* ParentGraph, FVector2D const Location, IBlueprintNodeBinder::FBindingSet const& Bindings, bool bSelectNewNode)
+static UEdGraphNode* FBlueprintMenuActionItemImpl::InvokeAction(const UBlueprintNodeSpawner* Action, UEdGraph* ParentGraph, FVector2D const Location, IBlueprintNodeBinder::FBindingSet const& Bindings)
 {
+	int32 const PreSpawnNodeCount = ParentGraph->Nodes.Num();
 	// this could return an existing node
 	UEdGraphNode* SpawnedNode = Action->Invoke(ParentGraph, Bindings, Location);
 
-	// if a returned node hasn't been added to the graph yet (it must have been freshly spawned)
-	if (ParentGraph->Nodes.Find(SpawnedNode) == INDEX_NONE)
+	// if a returned node wasn't one that previously existed in the graph
+	if (PreSpawnNodeCount < ParentGraph->Nodes.Num())
 	{
 		check(SpawnedNode != nullptr);
-
-		// @TODO: Move Modify()/AddNode() into UBlueprintNodeSpawner and pull 
-		//        selection functionality out to FBlueprintActionMenuItem::PerformAction()
-		ParentGraph->Modify();
-		ParentGraph->AddNode(SpawnedNode, /*bFromUI =*/true, bSelectNewNode);
-		// @TODO: if this spawned multiple nodes, then we should be selecting all of them
-
 		SpawnedNode->SnapToGrid(SNodePanel::GetSnapGridSize());
 
 		FBlueprintEditorUtils::AnalyticsTrackNewNode(SpawnedNode);
@@ -139,6 +132,9 @@ UEdGraphNode* FBlueprintActionMenuItem::PerformAction(UEdGraph* ParentGraph, UEd
 		FromPin->Modify();
 	}
 
+	TSet<const UEdGraphNode*> NodesToFocus;
+	int32 const PreSpawnNodeCount = ParentGraph->Nodes.Num();
+
 	UEdGraphNode* LastSpawnedNode = nullptr;
 	auto BoundObjIt = Bindings.CreateConstIterator();
 	do
@@ -152,7 +148,11 @@ UEdGraphNode* FBlueprintActionMenuItem::PerformAction(UEdGraph* ParentGraph, UEd
 			}
 		}
 
-		LastSpawnedNode = InvokeAction(Action, ParentGraph, ModifiedLocation, BindingsSubset, bSelectNewNode);
+		LastSpawnedNode = InvokeAction(Action, ParentGraph, ModifiedLocation, BindingsSubset);
+		// could already be an existent node, so we have to add here (can't 
+		// catch it as we go through all new nodes)
+		NodesToFocus.Add(LastSpawnedNode);
+
 		if (FromPin != nullptr)
 		{
 			// make sure to auto-wire after we position the new node (in case
@@ -164,6 +164,16 @@ UEdGraphNode* FBlueprintActionMenuItem::PerformAction(UEdGraph* ParentGraph, UEd
 		ModifiedLocation.Y += UEdGraphSchema_K2::EstimateNodeHeight(LastSpawnedNode);
 
 	} while (BoundObjIt);
+
+	if (bSelectNewNode)
+	{
+		int32 const PostSpawnCount = ParentGraph->Nodes.Num();
+		for (int32 NodeIndex = PreSpawnNodeCount; NodeIndex < PostSpawnCount; ++NodeIndex)
+		{
+			NodesToFocus.Add(ParentGraph->Nodes[NodeIndex]);
+		}
+		ParentGraph->SelectNodeSet(NodesToFocus, /*bFromUI =*/true);
+	}
 	// @TODO: select ALL spawned nodes
 	
 	return LastSpawnedNode;
diff --git a/Engine/Source/Runtime/Engine/Classes/EdGraph/EdGraph.h b/Engine/Source/Runtime/Engine/Classes/EdGraph/EdGraph.h
--- a/Engine/Source/Runtime/Engine/Classes/EdGraph/EdGraph.h
+++ b/Engine/Source/Runtime/Engine/Classes/EdGraph/EdGraph.h
@@ -129,6 +129,14 @@ public:
 	 */
 	void AddNode( UEdGraphNode* NodeToAdd, bool bUserAction = false, bool bSelectNewNode = true );
 
+	/**
+	 * Queues up a select operation for a series of nodes in this graph.
+	 * 
+	 * @param  NodeSelection	The group of nodes you want selected
+	 * @param  bFromUI			True if the node was added as the result of a direct user action.
+	 */
+	void SelectNodeSet(TSet<const UEdGraphNode*> NodeSelection, bool bFromUI = false);
+
 	/** Remove a node from this graph */
 	bool RemoveNode( UEdGraphNode* NodeToRemove );
 
diff --git a/Engine/Source/Runtime/Engine/Private/EdGraph/EdGraph.cpp b/Engine/Source/Runtime/Engine/Private/EdGraph/EdGraph.cpp
--- a/Engine/Source/Runtime/Engine/Private/EdGraph/EdGraph.cpp
+++ b/Engine/Source/Runtime/Engine/Private/EdGraph/EdGraph.cpp
@@ -146,6 +146,22 @@ void UEdGraph::AddNode( UEdGraphNode* NodeToAdd, bool bFromUI/* = false*/, bool
 	NotifyGraphChanged( Action );
 }
 
+void UEdGraph::SelectNodeSet(TSet<const UEdGraphNode*> NodeSelection, bool bFromUI/*= false*/) 
+{
+	FEdGraphEditAction SelectionAction;
+
+	SelectionAction.Action = GRAPHACTION_SelectNode;
+	if (bFromUI)
+	{
+		SelectionAction.Action = (EEdGraphActionType)(((int32)SelectionAction.Action) | GRAPHACTION_UserInitiated);
+	}
+
+	SelectionAction.Graph = this;
+	SelectionAction.Nodes = NodeSelection;
+
+	NotifyGraphChanged(SelectionAction);
+}
+
 bool UEdGraph::RemoveNode( UEdGraphNode* NodeToRemove )
 {
 	Modify();
diff --git a/Engine/Source/Runtime/Engine/Public/GraphEditAction.h b/Engine/Source/Runtime/Engine/Public/GraphEditAction.h
--- a/Engine/Source/Runtime/Engine/Public/GraphEditAction.h
+++ b/Engine/Source/Runtime/Engine/Public/GraphEditAction.h
@@ -27,17 +27,17 @@ struct FEdGraphEditAction
 	/** The graph the action occurred on */
 	class UEdGraph* Graph;
 	/** [Optional] The object the action occurred on */
-	class UEdGraphNode* Node;
+	TSet<const class UEdGraphNode*> Nodes;
 
 	FEdGraphEditAction()
 		: Action(GRAPHACTION_Default)
 		, Graph (NULL)
-		, Node(NULL)
 	{}
 
 	explicit FEdGraphEditAction(EEdGraphActionType InAction, class UEdGraph* InGraph, class UEdGraphNode* InNode)
 		: Action(InAction) 
 		, Graph(InGraph)
-		, Node(InNode)
-	{}
+	{
+		Nodes.Add(InNode);
+	}
 };
\ No newline at end of file
