Made sure widgets are tagged correctly in tabs

This allows tab content in workflow tabs to properly tag their content.
Viewports now use the metadata from their parent implementation if there is one.

NOTE: TTP listed will not be addressed until a build is promoted with this change and content is edited to fix the issue (i.e. re-tag the viewport).

TTP# 347538 - TUTORIALS: Animation Editor Walkthrough: Multiple Viewports are highlighted when it's trying to showcase the one in persona

diff --git a/Engine/Source/Editor/Kismet/Private/WorkflowOrientedApp/WorkflowTabFactory.cpp b/Engine/Source/Editor/Kismet/Private/WorkflowOrientedApp/WorkflowTabFactory.cpp
--- a/Engine/Source/Editor/Kismet/Private/WorkflowOrientedApp/WorkflowTabFactory.cpp
+++ b/Engine/Source/Editor/Kismet/Private/WorkflowOrientedApp/WorkflowTabFactory.cpp
@@ -31,9 +31,13 @@ TSharedRef<SDockTab> FWorkflowTabFactory::SpawnTab(const FWorkflowTabSpawnInfo&
 	// Pad the content if requested
 	if (InsideTabPadding > 0.0f)
 	{
+		// propagate tag from original content, or use TabId
+		TSharedPtr<FTagMetaData> MetaData = TabBody->GetMetaData<FTagMetaData>();
+
 		TabBody = SNew(SBorder)
 			.Padding(InsideTabPadding)
 			.BorderImage(FEditorStyle::GetBrush("ToolPanel.GroupBorder"))
+			.AddMetaData<FTagMetaData>(MetaData.IsValid() ? MetaData->Tag : this->TabIdentifier)
 			[
 				TabBody
 			];
diff --git a/Engine/Source/Editor/Persona/Private/SAnimationEditorViewport.cpp b/Engine/Source/Editor/Persona/Private/SAnimationEditorViewport.cpp
--- a/Engine/Source/Editor/Persona/Private/SAnimationEditorViewport.cpp
+++ b/Engine/Source/Editor/Persona/Private/SAnimationEditorViewport.cpp
@@ -48,6 +48,7 @@ void SAnimationEditorViewport::Construct(const FArguments& InArgs, TSharedPtr<cl
 	SEditorViewport::Construct(
 		SEditorViewport::FArguments()
 			.IsEnabled(FSlateApplication::Get().GetNormalExecutionAttribute())
+			.AddMetaData<FTagMetaData>(TEXT("Persona.Viewport"))
 		);
 }
 
diff --git a/Engine/Source/Editor/Persona/Private/Shared/PersonaMode.cpp b/Engine/Source/Editor/Persona/Private/Shared/PersonaMode.cpp
--- a/Engine/Source/Editor/Persona/Private/Shared/PersonaMode.cpp
+++ b/Engine/Source/Editor/Persona/Private/Shared/PersonaMode.cpp
@@ -201,7 +201,8 @@ TSharedRef<SWidget> FPreviewViewportSummoner::CreateTabBody(const FWorkflowTabSp
 
 	TSharedRef<SAnimationEditorViewportTabBody> NewViewport = SNew(SAnimationEditorViewportTabBody)
 		.Persona(PersonaPtr)
-		.Skeleton(PersonaPtr->GetSkeleton());
+		.Skeleton(PersonaPtr->GetSkeleton())
+		.AddMetaData<FTagMetaData>(TEXT("Persona.Viewport"));
 
 	//@TODO:MODES:check(!PersonaPtr->Viewport.IsValid());
 	// mode switch data sharing
diff --git a/Engine/Source/Editor/UnrealEd/Private/SEditorViewport.cpp b/Engine/Source/Editor/UnrealEd/Private/SEditorViewport.cpp
--- a/Engine/Source/Editor/UnrealEd/Private/SEditorViewport.cpp
+++ b/Engine/Source/Editor/UnrealEd/Private/SEditorViewport.cpp
@@ -38,7 +38,7 @@ void SEditorViewport::Construct( const FArguments& InArgs )
 		SAssignNew( ViewportWidget, SViewport )
 		.ShowEffectWhenDisabled( false )
 		.EnableGammaCorrection( false ) // Scene rendering handles this
-		.AddMetaData<FTutorialMetaData>(FTutorialMetaData(TEXT("EditorViewports"), TEXT("LevelEditorViewport")))
+		.AddMetaData(InArgs.MetaData.Num() > 0 ? InArgs.MetaData[0] : MakeShareable(new FTagMetaData(TEXT("EditorViewports"))))
 		[
 			SAssignNew( ViewportOverlay, SOverlay )
 			+SOverlay::Slot()
