Adding option to show blueprint favorites in the context menu.

diff --git a/Engine/Source/Editor/BlueprintGraph/Private/BlueprintEditorSettings.cpp b/Engine/Source/Editor/BlueprintGraph/Private/BlueprintEditorSettings.cpp
--- a/Engine/Source/Editor/BlueprintGraph/Private/BlueprintEditorSettings.cpp
+++ b/Engine/Source/Editor/BlueprintGraph/Private/BlueprintEditorSettings.cpp
@@ -9,9 +9,11 @@ UBlueprintEditorSettings::UBlueprintEditorSettings(const class FPostConstructIni
 	: Super(PCIP)
 {
 	bShowGraphInstructionText = true;
-	NodeTemplateCacheCapMB    = 20.f;
+	NodeTemplateCacheCapMB = 20.f;
 	bUseTargetContextForNodeMenu = true;
 	bExposeAllMemberComponentFunctions = true;
+	bShowContextualFavorites = false;
+	bFlattenFavoritesMenu = false;
 
 	// settings that were moved out of experimental...
 	UEditorExperimentalSettings const* ExperimentalSettings = GetDefault<UEditorExperimentalSettings>();
diff --git a/Engine/Source/Editor/BlueprintGraph/Public/BlueprintEditorSettings.h b/Engine/Source/Editor/BlueprintGraph/Public/BlueprintEditorSettings.h
--- a/Engine/Source/Editor/BlueprintGraph/Public/BlueprintEditorSettings.h
+++ b/Engine/Source/Editor/BlueprintGraph/Public/BlueprintEditorSettings.h
@@ -32,6 +32,14 @@ public:
 	UPROPERTY(EditAnywhere, config, Category=UserExperience, meta=(DisplayName="Context Menu: Expose All Sub-Component Functions"))
 	bool bExposeAllMemberComponentFunctions;
 
+	/** If enabled, then a separate section with your Blueprint favorites will be pined to the top of the context menu. */
+	UPROPERTY(EditAnywhere, config, Category=UserExperience, meta=(DisplayName="Context Menu: Show Favorites Section"))
+	bool bShowContextualFavorites;
+
+	/** If enabled, then your Blueprint favorites will be uncategorized, leaving you with less nested categories to sort through. */
+	UPROPERTY(EditAnywhere, config, Category=UserExperience)
+	bool bFlattenFavoritesMenu;
+
 // Developer Settings
 public:
 	/** If enabled, tooltips on action menu items will show the associated action's signature id (can be used to setup custom favorites menus).*/
diff --git a/Engine/Source/Editor/Kismet/Private/BlueprintActionMenuBuilder.cpp b/Engine/Source/Editor/Kismet/Private/BlueprintActionMenuBuilder.cpp
--- a/Engine/Source/Editor/Kismet/Private/BlueprintActionMenuBuilder.cpp
+++ b/Engine/Source/Editor/Kismet/Private/BlueprintActionMenuBuilder.cpp
@@ -341,6 +341,11 @@ void FBlueprintActionMenuBuilderImpl::FMenuSectionDefinition::AddBoundMenuItems(
 				{
 					LastMadeMenuItem = ItemFactory.MakeBoundMenuItem(EditorContext, BoundActionInfo);
 					MenuItemsOut.Add(LastMadeMenuItem);
+
+					if (Flags & FBlueprintActionMenuBuilder::FlattenCategoryHierarcy)
+					{
+						LastMadeMenuItem->Category = ItemFactory.RootCategory.ToString();
+					}
 				}
 				else
 				{
@@ -402,6 +407,10 @@ FBlueprintActionMenuBuilderImpl::MenuItemList FBlueprintActionMenuBuilderImpl::F
 	if (!UnBoundMenuEntry.IsValid() && bPassedFilter)
 	{
 		UnBoundMenuEntry = ItemFactory.MakeActionMenuItem(EditorContext, DatabaseAction);
+		if (Flags & FBlueprintActionMenuBuilder::FlattenCategoryHierarcy)
+		{
+			UnBoundMenuEntry->Category = ItemFactory.RootCategory.ToString();
+		}
 	}
 
 	FBlueprintActionMenuBuilderImpl::MenuItemList MenuItems;
diff --git a/Engine/Source/Editor/Kismet/Private/BlueprintActionMenuUtils.cpp b/Engine/Source/Editor/Kismet/Private/BlueprintActionMenuUtils.cpp
--- a/Engine/Source/Editor/Kismet/Private/BlueprintActionMenuUtils.cpp
+++ b/Engine/Source/Editor/Kismet/Private/BlueprintActionMenuUtils.cpp
@@ -26,6 +26,7 @@
 
 namespace BlueprintActionMenuUtilsImpl
 {
+	static int32 const FavoritesSectionGroup  = 102;
 	static int32 const LevelActorSectionGroup = 101;
 	static int32 const ComponentsSectionGroup = 100;
 	static int32 const BoundAddComponentGroup = 002;
@@ -119,6 +120,8 @@ namespace BlueprintActionMenuUtilsImpl
 	 * @return 
 	 */
 	static void AddLevelActorSections(FBlueprintActionFilter const& LevelActorsFilter, FBlueprintActionMenuBuilder& MenuOut);
+
+	static void AddFavoritesSection(FBlueprintActionFilter const& MainMenuFilter, FBlueprintActionMenuBuilder& MenuOut);
 }
 
 //------------------------------------------------------------------------------
@@ -337,6 +340,28 @@ static void BlueprintActionMenuUtilsImpl::AddLevelActorSections(FBlueprintAction
 	MenuOut.AddMenuSection(ActorReferencesFilter, FText::GetEmpty(), LevelActorSectionGroup, FBlueprintActionMenuBuilder::ConsolidateBoundActions);
 }
 
+//------------------------------------------------------------------------------
+static void BlueprintActionMenuUtilsImpl::AddFavoritesSection(FBlueprintActionFilter const& MainMenuFilter, FBlueprintActionMenuBuilder& MenuOut)
+{
+	const UBlueprintEditorSettings* BlueprintSettings = GetDefault<UBlueprintEditorSettings>();
+	if (BlueprintSettings->bShowContextualFavorites)
+	{
+		FBlueprintActionFilter FavoritesFilter = MainMenuFilter;
+		FavoritesFilter.AddRejectionTest(FBlueprintActionFilter::FRejectionTestDelegate::CreateStatic(IsNonFavoritedAction));
+		
+		uint32 SectionFlags = 0x00;
+		FText  SectionHeading = LOCTEXT("ContextMenuFavoritesTitle", "Favorites");
+
+		if (BlueprintSettings->bFlattenFavoritesMenu)
+		{
+			SectionFlags |= FBlueprintActionMenuBuilder::FlattenCategoryHierarcy;
+			SectionHeading = FText::GetEmpty();
+		}
+
+		MenuOut.AddMenuSection(FavoritesFilter, SectionHeading, FavoritesSectionGroup, SectionFlags);
+	}
+}
+
 /*******************************************************************************
  * FBlueprintActionMenuUtils
  ******************************************************************************/
@@ -532,6 +557,7 @@ void FBlueprintActionMenuUtils::MakeContextMenu(FBlueprintActionContext const& C
 
 	if (bIsContextSensitive)
 	{
+		AddFavoritesSection(MainMenuFilter, MenuOut);
 		MenuOut.AddMenuSection(CallOnMemberFilter, FText::GetEmpty(), MainMenuSectionGroup);
 		MenuOut.AddMenuSection(AddComponentFilter, FText::GetEmpty(), BoundAddComponentGroup);
 	}	
@@ -586,7 +612,13 @@ void FBlueprintActionMenuUtils::MakeFavoritesMenu(FBlueprintActionContext const&
 		MenuFilter.Context = Context;
 		MenuFilter.AddRejectionTest(FBlueprintActionFilter::FRejectionTestDelegate::CreateStatic(BlueprintActionMenuUtilsImpl::IsNonFavoritedAction));
 
-		MenuOut.AddMenuSection(MenuFilter);
+		uint32 SectionFlags = 0x00;
+		if (GetDefault<UBlueprintEditorSettings>()->bFlattenFavoritesMenu)
+		{
+			SectionFlags = FBlueprintActionMenuBuilder::FlattenCategoryHierarcy;
+		}
+
+		MenuOut.AddMenuSection(MenuFilter, FText::GetEmpty(), BlueprintActionMenuUtilsImpl::MainMenuSectionGroup, SectionFlags);
 		MenuOut.RebuildActionList();
 	}
 	else
diff --git a/Engine/Source/Editor/Kismet/Public/BlueprintActionMenuBuilder.h b/Engine/Source/Editor/Kismet/Public/BlueprintActionMenuBuilder.h
--- a/Engine/Source/Editor/Kismet/Public/BlueprintActionMenuBuilder.h
+++ b/Engine/Source/Editor/Kismet/Public/BlueprintActionMenuBuilder.h
@@ -33,6 +33,9 @@ public:
 		// Rolls bound node spawners into a single menu entry that will spawn 
 		// multiple nodes, each bound to a single binding.
 		ConsolidateBoundActions    = (1<<1),
+
+		// Will clear all action categories (except the section's root) 
+		FlattenCategoryHierarcy    = (1<<2),
 	};
 	
 public:
