Enabling the new blueprint menu system for wide spread use.

diff --git a/Engine/Config/BaseEditor.ini b/Engine/Config/BaseEditor.ini
--- a/Engine/Config/BaseEditor.ini
+++ b/Engine/Config/BaseEditor.ini
@@ -93,7 +93,6 @@ MenuPriority=11
 MenuPriority=10
 
 [/Script/BlueprintGraph.EdGraphSchema_K2]
-bUseLegacyActionMenus=true
 +EditoronlyBPFunctionRedirects=(ClassName="AIBlueprintHelperLibrary", OldFunctionName="SpawnAI", NewFunctionName="SpawnAIFromClass", BlueprintParamName="Pawn", ClassParamName="PawnClass")
 +EditoronlyBPFunctionRedirects=(ClassName="GameplayStatics", OldFunctionName="CreateSaveGameObjectFromBlueprint", NewFunctionName="CreateSaveGameObject", BlueprintParamName="SaveGameBlueprint", ClassParamName="SaveGameClass")
 
diff --git a/Engine/Source/Editor/BlueprintGraph/Classes/EdGraphSchema_K2.h b/Engine/Source/Editor/BlueprintGraph/Classes/EdGraphSchema_K2.h
--- a/Engine/Source/Editor/BlueprintGraph/Classes/EdGraphSchema_K2.h
+++ b/Engine/Source/Editor/BlueprintGraph/Classes/EdGraphSchema_K2.h
@@ -289,10 +289,6 @@ class BLUEPRINTGRAPH_API UEdGraphSchema_K2 : public UEdGraphSchema
 	// Somewhat hacky mechanism to prevent tooltips created for pins from including the display name and type when generating BP API documentation
 	static bool bGeneratingDocumentation;
 
-	/** DEPRECATED - Use UEditorExperimentalSettings::bUseRefactoredBlueprintMenuingSystem instead */
-	UPROPERTY(GlobalConfig)
-	bool bUseLegacyActionMenus; 
-
 	UPROPERTY(globalconfig)
 	TArray<FBlueprintCallableFunctionRedirect> EditoronlyBPFunctionRedirects;
 
diff --git a/Engine/Source/Editor/BlueprintGraph/Private/BlueprintEditorSettings.cpp b/Engine/Source/Editor/BlueprintGraph/Private/BlueprintEditorSettings.cpp
--- a/Engine/Source/Editor/BlueprintGraph/Private/BlueprintEditorSettings.cpp
+++ b/Engine/Source/Editor/BlueprintGraph/Private/BlueprintEditorSettings.cpp
@@ -13,7 +13,8 @@ UBlueprintEditorSettings::UBlueprintEditorSettings(const class FPostConstructIni
 	bUseTargetContextForNodeMenu = true;
 	bExposeAllMemberComponentFunctions = true;
 	bShowContextualFavorites = false;
-	bFlattenFavoritesMenu = false;
+	bFlattenFavoritesMenus = true;
+	bUseLegacyMenuingSystem = false;
 
 	// settings that were moved out of experimental...
 	UEditorExperimentalSettings const* ExperimentalSettings = GetDefault<UEditorExperimentalSettings>();
diff --git a/Engine/Source/Editor/BlueprintGraph/Public/BlueprintEditorSettings.h b/Engine/Source/Editor/BlueprintGraph/Public/BlueprintEditorSettings.h
--- a/Engine/Source/Editor/BlueprintGraph/Public/BlueprintEditorSettings.h
+++ b/Engine/Source/Editor/BlueprintGraph/Public/BlueprintEditorSettings.h
@@ -38,7 +38,11 @@ public:
 
 	/** If enabled, then your Blueprint favorites will be uncategorized, leaving you with less nested categories to sort through. */
 	UPROPERTY(EditAnywhere, config, Category=UserExperience)
-	bool bFlattenFavoritesMenu;
+	bool bFlattenFavoritesMenus;
+
+	/** If set, then the new refactored menu system will be replaced with the old (legacy) system (as a fallback, in case the new system has unforeseen problems)*/
+	UPROPERTY(EditAnywhere, AdvancedDisplay, config, Category=UserExperience)
+	bool bUseLegacyMenuingSystem;
 
 // Developer Settings
 public:
diff --git a/Engine/Source/Editor/Kismet/Private/BlueprintActionMenuBuilder.cpp b/Engine/Source/Editor/Kismet/Private/BlueprintActionMenuBuilder.cpp
--- a/Engine/Source/Editor/Kismet/Private/BlueprintActionMenuBuilder.cpp
+++ b/Engine/Source/Editor/Kismet/Private/BlueprintActionMenuBuilder.cpp
@@ -16,6 +16,7 @@
 #include "SMyBlueprint.h"				// for SelectionAsVar()
 #include "BlueprintEditorUtils.h"		// for FindBlueprintForGraphChecked()
 #include "BlueprintEditor.h"			// for GetFocusedGraph()
+#include "BlueprintEditorSettings.h"	// for bUseLegacyMenuingSystem
 
 #define LOCTEXT_NAMESPACE "BlueprintActionMenuBuilder"
 
@@ -540,15 +541,14 @@ void FBlueprintActionMenuBuilder::RebuildActionList()
 		MenuSection->Empty();
 	}
 	
-	UEditorExperimentalSettings const* ExperimentalSettings = GetDefault<UEditorExperimentalSettings>();
-	if (ExperimentalSettings->bUseRefactoredBlueprintMenuingSystem)
+	const UBlueprintEditorSettings* BlueprintSettings = GetDefault<UBlueprintEditorSettings>();
+	if (!BlueprintSettings->bUseLegacyMenuingSystem)
 	{
 		FBlueprintActionDatabase::FActionRegistry const& ActionDatabase = FBlueprintActionDatabase::Get().GetAllActions();
 		for (auto const& ActionEntry : ActionDatabase)
 		{
 			for (UBlueprintNodeSpawner const* NodeSpawner : ActionEntry.Value)
 			{
-
 				FBlueprintActionInfo BlueprintAction(ActionEntry.Key, NodeSpawner);
 
 				// @TODO: could probably have a super filter that spreads across 
diff --git a/Engine/Source/Editor/Kismet/Private/BlueprintActionMenuUtils.cpp b/Engine/Source/Editor/Kismet/Private/BlueprintActionMenuUtils.cpp
--- a/Engine/Source/Editor/Kismet/Private/BlueprintActionMenuUtils.cpp
+++ b/Engine/Source/Editor/Kismet/Private/BlueprintActionMenuUtils.cpp
@@ -352,7 +352,7 @@ static void BlueprintActionMenuUtilsImpl::AddFavoritesSection(FBlueprintActionFi
 		uint32 SectionFlags = 0x00;
 		FText  SectionHeading = LOCTEXT("ContextMenuFavoritesTitle", "Favorites");
 
-		if (BlueprintSettings->bFlattenFavoritesMenu)
+		if (BlueprintSettings->bFlattenFavoritesMenus)
 		{
 			SectionFlags |= FBlueprintActionMenuBuilder::FlattenCategoryHierarcy;
 			SectionHeading = FText::GetEmpty();
@@ -568,8 +568,7 @@ void FBlueprintActionMenuUtils::MakeContextMenu(FBlueprintActionContext const& C
 
 	MenuOut.RebuildActionList();
 
-	UEditorExperimentalSettings const* ExperimentalSettings =  GetDefault<UEditorExperimentalSettings>();
-	if (ExperimentalSettings->bUseRefactoredBlueprintMenuingSystem)
+	if (!BlueprintSettings->bUseLegacyMenuingSystem)
 	{
 		for (UEdGraph const* Graph : Context.Graphs)
 		{
@@ -605,15 +604,15 @@ void FBlueprintActionMenuUtils::MakeFavoritesMenu(FBlueprintActionContext const&
 {
 	MenuOut.Empty();
 
-	UEditorExperimentalSettings const* ExperimentalSettings = GetDefault<UEditorExperimentalSettings>();
-	if (ExperimentalSettings->bUseRefactoredBlueprintMenuingSystem)
+	const UBlueprintEditorSettings* BlueprintSettings = GetDefault<UBlueprintEditorSettings>();
+	if (!BlueprintSettings->bUseLegacyMenuingSystem)
 	{
 		FBlueprintActionFilter MenuFilter;
 		MenuFilter.Context = Context;
 		MenuFilter.AddRejectionTest(FBlueprintActionFilter::FRejectionTestDelegate::CreateStatic(BlueprintActionMenuUtilsImpl::IsNonFavoritedAction));
 
 		uint32 SectionFlags = 0x00;
-		if (GetDefault<UBlueprintEditorSettings>()->bFlattenFavoritesMenu)
+		if (BlueprintSettings->bFlattenFavoritesMenus)
 		{
 			SectionFlags = FBlueprintActionMenuBuilder::FlattenCategoryHierarcy;
 		}
diff --git a/Engine/Source/Editor/UnrealEd/Classes/Settings/EditorExperimentalSettings.h b/Engine/Source/Editor/UnrealEd/Classes/Settings/EditorExperimentalSettings.h
--- a/Engine/Source/Editor/UnrealEd/Classes/Settings/EditorExperimentalSettings.h
+++ b/Engine/Source/Editor/UnrealEd/Classes/Settings/EditorExperimentalSettings.h
@@ -56,10 +56,6 @@ public:
 	UPROPERTY(/*EditAnywhere - deprecated (moved into UBlueprintEditorSettings), */config/*, Category=Blueprints, meta=(DisplayName="Draw midpoint arrows in Blueprints")*/)
 	bool bDrawMidpointArrowsInBlueprints;
 
-	/** Determines if the Blueprint editor should use the new (experimental) menu system, or the old (legacy) system */
-	UPROPERTY(EditAnywhere, config, Category=Blueprints, meta=(DisplayName="Use New Blueprint Menuing System"))
-	bool bUseRefactoredBlueprintMenuingSystem;
-
 	/** Whether to show Audio Streaming options for SoundWaves (disabling will not stop all audio streaming) */
 	UPROPERTY(EditAnywhere, config, Category=Audio)
 	bool bShowAudioStreamingOptions;
diff --git a/Engine/Source/Editor/UnrealEd/Private/Settings/SettingsClasses.cpp b/Engine/Source/Editor/UnrealEd/Private/Settings/SettingsClasses.cpp
--- a/Engine/Source/Editor/UnrealEd/Private/Settings/SettingsClasses.cpp
+++ b/Engine/Source/Editor/UnrealEd/Private/Settings/SettingsClasses.cpp
@@ -2,8 +2,6 @@
 
 #include "UnrealEd.h"
 #include "ISourceControlModule.h"
-#include "EdGraphSchema_K2.h" // for bUseLegacyActionMenus
-
 
 /* UContentBrowserSettings interface
  *****************************************************************************/
@@ -53,8 +51,6 @@ UDestructableMeshEditorSettings::UDestructableMeshEditorSettings( const class FP
 UEditorExperimentalSettings::UEditorExperimentalSettings( const class FPostConstructInitializeProperties& PCIP )
 	: Super(PCIP)
 { 
-	// setting was originally loaded from .ini into UEdGraphSchema_K2
-	bUseRefactoredBlueprintMenuingSystem = !GetDefault<UEdGraphSchema_K2>()->bUseLegacyActionMenus;
 }
 
 
