Guarding against null graph node entries (supposedly from a lack of node redirects?).

diff --git a/Engine/Source/Editor/Kismet/Private/BlueprintEditor.cpp b/Engine/Source/Editor/Kismet/Private/BlueprintEditor.cpp
--- a/Engine/Source/Editor/Kismet/Private/BlueprintEditor.cpp
+++ b/Engine/Source/Editor/Kismet/Private/BlueprintEditor.cpp
@@ -196,6 +196,11 @@ static bool BlueprintEditorImpl::GraphHasUserPlacedNodes(UEdGraph const* InGraph
 
 	for (UEdGraphNode const* Node : InGraph->Nodes)
 	{
+		if (Node == nullptr)
+		{
+			continue;
+		}
+
 		if (!Node->GetOutermost()->GetMetaData()->HasValue(Node, FNodeMetadata::DefaultGraphNode))
 		{
 			bHasUserPlacedNodes = true;
@@ -212,6 +217,11 @@ static bool BlueprintEditorImpl::GraphHasDefaultNode(UEdGraph const* InGraph)
 
 	for (UEdGraphNode const* Node : InGraph->Nodes)
 	{
+		if (Node == nullptr)
+		{
+			continue;
+		}
+
 		if (Node->GetOutermost()->GetMetaData()->HasValue(Node, FNodeMetadata::DefaultGraphNode))
 		{
 			bHasDefaultNodes = true;
