Preventing users from referencing BSP brushes in the level script.
TTP #348446

diff --git a/Engine/Source/Editor/Kismet/Private/BlueprintActionMenuUtils.cpp b/Engine/Source/Editor/Kismet/Private/BlueprintActionMenuUtils.cpp
--- a/Engine/Source/Editor/Kismet/Private/BlueprintActionMenuUtils.cpp
+++ b/Engine/Source/Editor/Kismet/Private/BlueprintActionMenuUtils.cpp
@@ -435,6 +435,11 @@ void FBlueprintActionMenuUtils::MakeContextMenu(FBlueprintActionContext const& C
 		else if (AActor* LevelActor = Cast<AActor>(Selection))
 		{
 			ComponentsFilter.Context.SelectedObjects.Remove(Selection);
+			if (!LevelActor->NeedsLoadForClient() && !LevelActor->NeedsLoadForServer())
+			{
+				// don't want to let the level script operate on actors that won't be loaded in game
+				LevelActorsFilter.Context.SelectedObjects.Remove(Selection);
+			}
 		}
 		else
 		{
@@ -448,7 +453,11 @@ void FBlueprintActionMenuUtils::MakeContextMenu(FBlueprintActionContext const& C
 	for (FSelectionIterator LvlActorIt(*GEditor->GetSelectedActors()); LvlActorIt; ++LvlActorIt)
 	{
 		AActor* LevelActor = Cast<AActor>(*LvlActorIt);
-		LevelActorsFilter.Context.SelectedObjects.AddUnique(LevelActor);
+		// don't want to let the level script operate on actors that won't be loaded in game
+		if (LevelActor->NeedsLoadForClient() || LevelActor->NeedsLoadForServer())
+		{
+			LevelActorsFilter.Context.SelectedObjects.AddUnique(LevelActor);
+		}
 	}
 
 	const UBlueprintEditorSettings* BlueprintSettings = GetDefault<UBlueprintEditorSettings>();
