Adding a way to extend blueprint editor key commands.

diff --git a/Engine/Source/Editor/Kismet/Private/BlueprintEditor.cpp b/Engine/Source/Editor/Kismet/Private/BlueprintEditor.cpp
--- a/Engine/Source/Editor/Kismet/Private/BlueprintEditor.cpp
+++ b/Engine/Source/Editor/Kismet/Private/BlueprintEditor.cpp
@@ -1775,6 +1775,10 @@ void FBlueprintEditor::CreateDefaultCommands()
 	FMyBlueprintCommands::Register();
 	FBlueprintSpawnNodeCommands::Register();
 
+	static const FName BpEditorModuleName("Kismet");
+	FBlueprintEditorModule& BlueprintEditorModule = FModuleManager::LoadModuleChecked<FBlueprintEditorModule>(BpEditorModuleName);
+	ToolkitCommands->Append(BlueprintEditorModule.GetsSharedBlueprintEditorCommands());
+
 	ToolkitCommands->MapAction(
 		FFullBlueprintEditorCommands::Get().Compile,
 		FExecuteAction::CreateSP(this, &FBlueprintEditor::Compile),
@@ -2303,6 +2307,7 @@ void FBlueprintEditor::OnGraphEditorDropActor(const TArray< TWeakObjectPtr<AActo
 	{
 		ULevel* BlueprintLevel = LevelBlueprint->GetLevel();
 
+		FVector2D NodeLocation = DropLocation;
 		for (int32 i = 0; i < Actors.Num(); i++)
 		{
 			AActor* DroppedActor = Actors[i].Get();
@@ -2311,8 +2316,8 @@ void FBlueprintEditor::OnGraphEditorDropActor(const TArray< TWeakObjectPtr<AActo
 				UK2Node_Literal* LiteralNodeTemplate = NewObject<UK2Node_Literal>();
 				LiteralNodeTemplate->SetObjectRef(DroppedActor);
 
-				const FVector2D NodeLocation = DropLocation + (i * FVector2D(0,30));
-				FEdGraphSchemaAction_K2NewNode::SpawnNodeFromTemplate<UK2Node_Literal>(Graph, LiteralNodeTemplate, NodeLocation);
+				UK2Node_Literal* ActorRefNode = FEdGraphSchemaAction_K2NewNode::SpawnNodeFromTemplate<UK2Node_Literal>(Graph, LiteralNodeTemplate, NodeLocation);
+				NodeLocation.Y += UEdGraphSchema_K2::EstimateNodeHeight(ActorRefNode);
 			}
 		}
 	}
@@ -6515,15 +6520,36 @@ AActor* FBlueprintEditor::GetPreviewActor() const
 FReply FBlueprintEditor::OnSpawnGraphNodeByShortcut(FInputGesture InGesture, const FVector2D& InPosition, UEdGraph* InGraph)
 {
 	UEdGraph* Graph = InGraph;
+	if (Graph == nullptr)
+	{
+		return FReply::Handled();
+	}
 
 	FBlueprintPaletteListBuilder PaletteBuilder(GetBlueprintObj());
-	TSharedPtr< FEdGraphSchemaAction > Action = FBlueprintSpawnNodeCommands::Get().GetGraphActionByGesture(InGesture, PaletteBuilder, InGraph);
+	FBlueprintSpawnNodeCommands::Get().GetGraphActionByGesture(InGesture, PaletteBuilder, InGraph);
 
-	if(Action.IsValid())
+	TSet<const UEdGraphNode*> NodesToSelect;
+	FVector2D NodeSpawnPos = InPosition;
+
+	for (int32 ActionIndex = 0; ActionIndex < PaletteBuilder.GetNumActions(); ++ActionIndex)
 	{
-		TArray<UEdGraphPin*> DummyPins;
-		Action->PerformAction(Graph, DummyPins, InPosition);
+		FGraphActionListBuilderBase::ActionGroup& ActionSet = PaletteBuilder.GetAction(ActionIndex);
+		if ((ActionSet.Actions.Num() > 0) && ActionSet.Actions[0].IsValid())
+		{
+			int32 const OldNodeCount = Graph->Nodes.Num();
+
+			TArray<UEdGraphPin*> DummyPins;
+			ActionSet.PerformAction(Graph, DummyPins, NodeSpawnPos);
+
+			for (int32 NodeIndex = OldNodeCount; NodeIndex < Graph->Nodes.Num(); ++NodeIndex)
+			{
+				UEdGraphNode* SpawnedNode = Graph->Nodes[NodeIndex];
+				NodeSpawnPos.Y = FMath::Max(NodeSpawnPos.Y, SpawnedNode->NodePosY + UEdGraphSchema_K2::EstimateNodeHeight(SpawnedNode));
+				NodesToSelect.Add(SpawnedNode);
+			}
+		}	
 	}
+	Graph->SelectNodeSet(NodesToSelect, /*bFromUI =*/true);
 
 	return FReply::Handled();
 }
diff --git a/Engine/Source/Editor/Kismet/Private/BlueprintEditorCommands.cpp b/Engine/Source/Editor/Kismet/Private/BlueprintEditorCommands.cpp
--- a/Engine/Source/Editor/Kismet/Private/BlueprintEditorCommands.cpp
+++ b/Engine/Source/Editor/Kismet/Private/BlueprintEditorCommands.cpp
@@ -156,7 +156,7 @@ public:
 	 *
 	 * @return						A fully prepared action containing the information to spawn the node
 	 */
-	virtual TSharedPtr< FEdGraphSchemaAction > GetAction(FBlueprintPaletteListBuilder& InPaletteBuilder, UEdGraph* InDestGraph) = 0;
+	virtual void GetActions(FBlueprintPaletteListBuilder& InPaletteBuilder, UEdGraph* InDestGraph) = 0;
 };
 
 class FEdGraphNodeSpawnInfo : public FNodeSpawnInfo
@@ -165,7 +165,7 @@ public:
 	FEdGraphNodeSpawnInfo(UClass* InClass) : NodeClass(InClass), GraphNode(NULL) {}
 
 	// FNodeSpawnInfo interface
-	virtual TSharedPtr< FEdGraphSchemaAction > GetAction(FBlueprintPaletteListBuilder& InPaletteBuilder, UEdGraph* InDestGraph) override
+	virtual void GetActions(FBlueprintPaletteListBuilder& InPaletteBuilder, UEdGraph* InDestGraph) override
 	{
 		TSharedPtr<FEdGraphSchemaAction_NewNode> NewActionNode = TSharedPtr<FEdGraphSchemaAction_NewNode>(new FEdGraphSchemaAction_NewNode);
 
@@ -178,7 +178,7 @@ public:
 		// the NodeTemplate UPROPERTY takes ownership of GraphNode's lifetime (hence it being a weak-pointer here)
 		NewActionNode->NodeTemplate = GraphNode.Get();
 
-		return NewActionNode;
+		InPaletteBuilder.AddAction(NewActionNode);
 	}
 	// End of FNodeSpawnInfo interface
 
@@ -196,10 +196,9 @@ public:
 	FFunctionNodeSpawnInfo(UFunction* InFunctionPtr) : FunctionPtr(InFunctionPtr) {}
 
 	// FNodeSpawnInfo interface
-	virtual TSharedPtr< FEdGraphSchemaAction > GetAction(FBlueprintPaletteListBuilder& InPaletteBuilder, UEdGraph* InDestGraph) override
+	virtual void GetActions(FBlueprintPaletteListBuilder& InPaletteBuilder, UEdGraph* InDestGraph) override
 	{
 		const UBlueprint* Blueprint = InPaletteBuilder.Blueprint;
-		FBlueprintPaletteListBuilder TempListBuilder(Blueprint);
 
 		const UEdGraphSchema_K2* K2Schema = GetDefault<UEdGraphSchema_K2>();
 		if(NodeSpawnInfoHelpers::IsFunctionAvailableAsEvent(Blueprint, FunctionPtr))
@@ -207,7 +206,7 @@ public:
 			UFunction* FunctionEvent = NodeSpawnInfoHelpers::FindEventFunctionForClass(Blueprint, FunctionPtr);
 			if(FunctionEvent)
 			{
-				FK2ActionMenuBuilder::AddEventForFunction(TempListBuilder, FunctionEvent);	
+				FK2ActionMenuBuilder::AddEventForFunction(InPaletteBuilder, FunctionEvent);
 			}
 		}
 		else
@@ -222,19 +221,9 @@ public:
 
 			if(K2Schema->CanFunctionBeUsedInClass(Blueprint->GeneratedClass, FunctionPtr, InDestGraph, FunctionTypes, true, false, FFunctionTargetInfo()))
 			{
-				FK2ActionMenuBuilder::AddSpawnInfoForFunction(FunctionPtr, false, FFunctionTargetInfo(), FMemberReference(), TEXT(""), K2Schema->AG_LevelReference, TempListBuilder);
+				FK2ActionMenuBuilder::AddSpawnInfoForFunction(FunctionPtr, false, FFunctionTargetInfo(), FMemberReference(), TEXT(""), K2Schema->AG_LevelReference, InPaletteBuilder);
 			}
 		}
-
-		TSharedPtr<FEdGraphSchemaAction> NewActionNode;
-		if(TempListBuilder.GetNumActions() == 1)
-		{
-			TArray<UEdGraphPin*> DummyPins;
-			FGraphActionListBuilderBase::ActionGroup& Action = TempListBuilder.GetAction(0);
-			NewActionNode = Action.Actions[0];
-		}
-
-		return NewActionNode;
 	}
 	// End of FNodeSpawnInfo interface
 
@@ -249,15 +238,9 @@ public:
 	FMacroNodeSpawnInfo(UEdGraph* InMacroGraph) : MacroGraph(InMacroGraph) {}
 
 	// FNodeSpawnInfo interface
-	virtual TSharedPtr< FEdGraphSchemaAction > GetAction(FBlueprintPaletteListBuilder& InPaletteBuilder, UEdGraph* InDestGraph) override
+	virtual void GetActions(FBlueprintPaletteListBuilder& InPaletteBuilder, UEdGraph* InDestGraph) override
 	{
-		const UBlueprint* Blueprint = InPaletteBuilder.Blueprint;
-		FBlueprintPaletteListBuilder TempListBuilder(Blueprint);
-
-		FK2ActionMenuBuilder::AttachMacroGraphAction(TempListBuilder, MacroGraph);
-
-		FGraphActionListBuilderBase::ActionGroup& Action = TempListBuilder.GetAction(0);
-		return Action.Actions[0];
+		FK2ActionMenuBuilder::AttachMacroGraphAction(InPaletteBuilder, MacroGraph);
 	}
 	// End of FNodeSpawnInfo interface
 
@@ -266,6 +249,32 @@ private:
 	UEdGraph* MacroGraph;
 };
 
+class FActorRefSpawnInfo : public FNodeSpawnInfo
+{
+public:
+	// FNodeSpawnInfo interface
+	virtual void GetActions(FBlueprintPaletteListBuilder& InPaletteBuilder, UEdGraph* InDestGraph) override
+	{
+		const UBlueprint* Blueprint = InPaletteBuilder.Blueprint;
+		USelection* SelectedLvlActors = GEditor->GetSelectedActors();
+
+		if ((Blueprint != nullptr) && Blueprint->ParentClass->IsChildOf<ALevelScriptActor>() && (SelectedLvlActors->Num() > 0))
+		{
+			for (FSelectionIterator LvlActorIt(*SelectedLvlActors); LvlActorIt; ++LvlActorIt)
+			{
+				UK2Node_Literal* TemplateRefNode = InPaletteBuilder.CreateTemplateNode<UK2Node_Literal>();
+				TemplateRefNode->SetObjectRef(*LvlActorIt);
+
+				TSharedPtr<FEdGraphSchemaAction_NewNode> NewActionNode = TSharedPtr<FEdGraphSchemaAction_NewNode>(new FEdGraphSchemaAction_NewNode);
+				NewActionNode->NodeTemplate = TemplateRefNode;
+
+				InPaletteBuilder.AddAction(NewActionNode);
+			}
+		}
+	}
+	// End of FNodeSpawnInfo interface
+};
+
 //////////////////////////////////////////////////////////////////////////
 // FBlueprintSpawnNodeCommands
 
@@ -377,9 +386,14 @@ void FBlueprintSpawnNodeCommands::RegisterCommands()
 			NodeCommands.Add(InfoPtr);
 		}
 	}
+
+	TSharedPtr<FNodeSpawnInfo> AddActorRefAction = MakeShareable(new FActorRefSpawnInfo);
+	UI_COMMAND(AddActorRefAction->CommandInfo, "Add Selected Actor Reference(s)", "Spawns node(s) which reference the currently selected actor(s) in the level.", EUserInterfaceActionType::Button, FInputGesture(EKeys::R));
+	NodeCommands.Add(AddActorRefAction);
 }
 
-TSharedPtr< FEdGraphSchemaAction > FBlueprintSpawnNodeCommands::GetGraphActionByGesture(FInputGesture& InGesture, FBlueprintPaletteListBuilder& InPaletteBuilder, UEdGraph* InDestGraph) const
+
+void FBlueprintSpawnNodeCommands::GetGraphActionByGesture(FInputGesture& InGesture, FBlueprintPaletteListBuilder& InPaletteBuilder, UEdGraph* InDestGraph) const
 {
 	if(InGesture.IsValidGesture())
 	{
@@ -387,10 +401,8 @@ TSharedPtr< FEdGraphSchemaAction > FBlueprintSpawnNodeCommands::GetGraphActionBy
 		{
 			if(NodeCommands[x]->CommandInfo->GetActiveGesture().Get() == InGesture)
 			{
-				return NodeCommands[x]->GetAction(InPaletteBuilder, InDestGraph);
+				NodeCommands[x]->GetActions(InPaletteBuilder, InDestGraph);
 			}
 		}
 	}
-
-	return TSharedPtr< FEdGraphSchemaAction >();
 }
diff --git a/Engine/Source/Editor/Kismet/Private/BlueprintEditorCommands.h b/Engine/Source/Editor/Kismet/Private/BlueprintEditorCommands.h
--- a/Engine/Source/Editor/Kismet/Private/BlueprintEditorCommands.h
+++ b/Engine/Source/Editor/Kismet/Private/BlueprintEditorCommands.h
@@ -93,7 +93,7 @@ public:
 	 * @param InPaletteBuilder	The Blueprint palette to create the graph action for, used for validation purposes and to link any important node data to the blueprint
 	 * @param InDestGraph		The graph to create the graph action for, used for validation purposes and to link any important node data to the graph
 	 */
-	TSharedPtr< FEdGraphSchemaAction > GetGraphActionByGesture(FInputGesture& InGesture, struct FBlueprintPaletteListBuilder& InPaletteBuilder, UEdGraph* InDestGraph) const;
+	void GetGraphActionByGesture(FInputGesture& InGesture, struct FBlueprintPaletteListBuilder& InPaletteBuilder, UEdGraph* InDestGraph) const;
 
 private:
 	/** An array of all the possible commands for spawning nodes */
diff --git a/Engine/Source/Editor/Kismet/Private/BlueprintEditorModule.cpp b/Engine/Source/Editor/Kismet/Private/BlueprintEditorModule.cpp
--- a/Engine/Source/Editor/Kismet/Private/BlueprintEditorModule.cpp
+++ b/Engine/Source/Editor/Kismet/Private/BlueprintEditorModule.cpp
@@ -67,6 +67,7 @@ static void FocusBlueprintEditorOnObject(const TSharedRef<IMessageToken>& Token)
 void FBlueprintEditorModule::StartupModule()
 {
 	MenuExtensibilityManager = MakeShareable(new FExtensibilityManager);
+	SharedBlueprintEditorCommands = MakeShareable(new FUICommandList);
 
 	const IWorkspaceMenuStructure& MenuStructure = WorkspaceMenu::GetMenuStructure();
 
@@ -109,6 +110,7 @@ void FBlueprintEditorModule::StartupModule()
 
 void FBlueprintEditorModule::ShutdownModule()
 {
+	SharedBlueprintEditorCommands.Reset();
 	MenuExtensibilityManager.Reset();
 
 	if (FSlateApplication::IsInitialized())
diff --git a/Engine/Source/Editor/Kismet/Public/BlueprintEditorModule.h b/Engine/Source/Editor/Kismet/Public/BlueprintEditorModule.h
--- a/Engine/Source/Editor/Kismet/Public/BlueprintEditorModule.h
+++ b/Engine/Source/Editor/Kismet/Public/BlueprintEditorModule.h
@@ -130,6 +130,13 @@ public:
 	DECLARE_EVENT_OneParam( FBlueprintEditorModule, FBlueprintEditorOpenedEvent, EBlueprintType );
 	FBlueprintEditorOpenedEvent& OnBlueprintEditorOpened() { return BlueprintEditorOpened; }
 
+	/** 
+	 * Exposes a way for other modules to fold in their own Blueprint editor 
+	 * commands (folded in with other BP editor commands, when the editor is 
+	 * first opened).
+	 */
+	virtual const TSharedRef<FUICommandList> GetsSharedBlueprintEditorCommands() const { return SharedBlueprintEditorCommands.ToSharedRef(); }
+
 private:
 	TSharedPtr<FExtensibilityManager> MenuExtensibilityManager;
 
@@ -138,4 +145,10 @@ private:
 
 	/** Customizations for the SCS editor */
 	TMap<FName, FSCSEditorCustomizationBuilder> SCSEditorCustomizations;
+
+	/** 
+	 * A command list that can be passed around and isn't bound to an instance 
+	 * of the blueprint editor. 
+	 */
+	TSharedPtr<FUICommandList> SharedBlueprintEditorCommands;
 };
