Left clicking in empty space in a Blueprint graph will not cause a visual refresh of all nodes.

#ttp 348826 - Regression: BP: Left clicking on empty space in the graph will refresh all the nodes

diff --git a/Engine/Source/Editor/Kismet/Private/BlueprintEditor.cpp b/Engine/Source/Editor/Kismet/Private/BlueprintEditor.cpp
--- a/Engine/Source/Editor/Kismet/Private/BlueprintEditor.cpp
+++ b/Engine/Source/Editor/Kismet/Private/BlueprintEditor.cpp
@@ -6534,6 +6534,8 @@ FReply FBlueprintEditor::OnSpawnGraphNodeByShortcut(FInputGesture InGesture, con
 	TSet<const UEdGraphNode*> NodesToSelect;
 	FVector2D NodeSpawnPos = InPosition;
 
+	bool bPerformedAction = false;
+
 	for (int32 ActionIndex = 0; ActionIndex < PaletteBuilder.GetNumActions(); ++ActionIndex)
 	{
 		FGraphActionListBuilderBase::ActionGroup& ActionSet = PaletteBuilder.GetAction(ActionIndex);
@@ -6550,9 +6552,16 @@ FReply FBlueprintEditor::OnSpawnGraphNodeByShortcut(FInputGesture InGesture, con
 				NodeSpawnPos.Y = FMath::Max(NodeSpawnPos.Y, SpawnedNode->NodePosY + UEdGraphSchema_K2::EstimateNodeHeight(SpawnedNode));
 				NodesToSelect.Add(SpawnedNode);
 			}
+
+			bPerformedAction = true;
 		}	
 	}
-	Graph->SelectNodeSet(NodesToSelect, /*bFromUI =*/true);
+
+	// Do not change node selection if no actions were performed
+	if(bPerformedAction)
+	{
+		Graph->SelectNodeSet(NodesToSelect, /*bFromUI =*/true);
+	}
 
 	return FReply::Handled();
 }
