Adding new "save on compile" option for blueprints. #NUX
TTP #345833

diff --git a/Engine/Source/Editor/BlueprintGraph/Private/BlueprintEditorSettings.cpp b/Engine/Source/Editor/BlueprintGraph/Private/BlueprintEditorSettings.cpp
--- a/Engine/Source/Editor/BlueprintGraph/Private/BlueprintEditorSettings.cpp
+++ b/Engine/Source/Editor/BlueprintGraph/Private/BlueprintEditorSettings.cpp
@@ -8,6 +8,7 @@
 UBlueprintEditorSettings::UBlueprintEditorSettings(const class FPostConstructInitializeProperties& PCIP)
 	: Super(PCIP)
 {
+	bSaveOnCompile = false;
 	bShowGraphInstructionText = true;
 	NodeTemplateCacheCapMB = 20.f;
 	bUseTargetContextForNodeMenu = true;
diff --git a/Engine/Source/Editor/BlueprintGraph/Public/BlueprintEditorSettings.h b/Engine/Source/Editor/BlueprintGraph/Public/BlueprintEditorSettings.h
--- a/Engine/Source/Editor/BlueprintGraph/Public/BlueprintEditorSettings.h
+++ b/Engine/Source/Editor/BlueprintGraph/Public/BlueprintEditorSettings.h
@@ -12,6 +12,12 @@ class BLUEPRINTGRAPH_API UBlueprintEditorSettings
 {
 	GENERATED_UCLASS_BODY()
 
+// General Settings
+public:
+	/** Determines if Blueprints are saved whenever you compile them */
+	UPROPERTY(EditAnywhere, config, Category=General)
+	bool bSaveOnCompile;
+
 // Style Settings
 public:
 	/** Should arrows indicating data/execution flow be drawn halfway along wires? */
diff --git a/Engine/Source/Editor/Kismet/Private/BlueprintEditor.cpp b/Engine/Source/Editor/Kismet/Private/BlueprintEditor.cpp
--- a/Engine/Source/Editor/Kismet/Private/BlueprintEditor.cpp
+++ b/Engine/Source/Editor/Kismet/Private/BlueprintEditor.cpp
@@ -1783,6 +1783,13 @@ void FBlueprintEditor::CreateDefaultCommands()
 		FFullBlueprintEditorCommands::Get().Compile,
 		FExecuteAction::CreateSP(this, &FBlueprintEditor::Compile),
 		FCanExecuteAction::CreateSP(this, &FBlueprintEditor::IsCompilingEnabled));
+
+	ToolkitCommands->MapAction(
+		FFullBlueprintEditorCommands::Get().SaveOnCompile,
+		FExecuteAction::CreateSP(this, &FBlueprintEditor::OnSaveOnCompileToggled),
+		FCanExecuteAction(),
+		FIsActionChecked::CreateSP(this, &FBlueprintEditor::IsSaveOnCompileChecked)
+	);
 	
 	ToolkitCommands->MapAction(
 		FFullBlueprintEditorCommands::Get().SwitchToScriptingMode,
@@ -2448,9 +2455,31 @@ void FBlueprintEditor::Compile()
 		// send record when player clicks compile and send the result
 		// this will make sure how the users activity is
 		AnalyticsTrackCompileEvent(BlueprintObj, LogResults.NumErrors, LogResults.NumWarnings);
+
+		UBlueprintEditorSettings* Settings = GetMutableDefault<UBlueprintEditorSettings>();
+		if (Settings->bSaveOnCompile)
+		{
+			TArray<UPackage*> PackagesToSave;
+			PackagesToSave.Add(BlueprintObj->GetOutermost());
+
+			FEditorFileUtils::PromptForCheckoutAndSave(PackagesToSave, /*bCheckDirty =*/false, /*bPromptToSave =*/false);
+		}
 	}
 }
 
+bool FBlueprintEditor::IsSaveOnCompileChecked() const
+{
+	UBlueprintEditorSettings const* Settings = GetDefault<UBlueprintEditorSettings>();
+	return Settings->bSaveOnCompile;
+}
+
+void FBlueprintEditor::OnSaveOnCompileToggled() const
+{
+	UBlueprintEditorSettings* Settings = GetMutableDefault<UBlueprintEditorSettings>();
+	Settings->bSaveOnCompile = !Settings->bSaveOnCompile;
+	Settings->SaveConfig();
+}
+
 FReply FBlueprintEditor::Compile_OnClickWithReply()
 {
 	Compile();
diff --git a/Engine/Source/Editor/Kismet/Private/SBlueprintEditorToolbar.cpp b/Engine/Source/Editor/Kismet/Private/SBlueprintEditorToolbar.cpp
--- a/Engine/Source/Editor/Kismet/Private/SBlueprintEditorToolbar.cpp
+++ b/Engine/Source/Editor/Kismet/Private/SBlueprintEditorToolbar.cpp
@@ -543,12 +543,32 @@ TSharedRef<SWidget> FKismet2Menu::MakeDiffMenu(FBlueprintEditor& Kismet)
 void FFullBlueprintEditorCommands::RegisterCommands() 
 {
 	UI_COMMAND(Compile, "Compile", "Compile the blueprint", EUserInterfaceActionType::Button, FInputGesture());
+	UI_COMMAND(SaveOnCompile, "Save on Compile", "Determines if the Blueprint is saved every time you compile it", EUserInterfaceActionType::ToggleButton, FInputGesture());
 	UI_COMMAND(SwitchToScriptingMode, "Graph", "Switches to Graph Editing Mode", EUserInterfaceActionType::ToggleButton, FInputGesture());
 	UI_COMMAND(SwitchToBlueprintDefaultsMode, "Defaults", "Switches to Blueprint Defaults Mode", EUserInterfaceActionType::ToggleButton, FInputGesture());
 	UI_COMMAND(SwitchToComponentsMode, "Components", "Switches to Components Mode", EUserInterfaceActionType::ToggleButton, FInputGesture());
 	UI_COMMAND(EditGlobalOptions, "Blueprint Props", "Edit Blueprint Options", EUserInterfaceActionType::Button, FInputGesture());
 }
 
+//////////////////////////////////////////////////////////////////////////
+// Static FBlueprintEditorToolbar Helpers
+
+namespace BlueprintEditorToolbarImpl
+{
+	static TSharedRef<SWidget> GenerateCompileOptionsWidget(TSharedRef<FUICommandList> CommandList);
+};
+
+
+static TSharedRef<SWidget> BlueprintEditorToolbarImpl::GenerateCompileOptionsWidget(TSharedRef<FUICommandList> CommandList)
+{
+	FMenuBuilder MenuBuilder(/*bShouldCloseWindowAfterMenuSelection =*/true, CommandList);
+
+	const FFullBlueprintEditorCommands& Commands = FFullBlueprintEditorCommands::Get();
+	MenuBuilder.AddMenuEntry(Commands.SaveOnCompile);
+
+	return MenuBuilder.MakeWidget();
+}
+
 
 //////////////////////////////////////////////////////////////////////////
 // FBlueprintEditorToolbar
@@ -677,6 +697,16 @@ void FBlueprintEditorToolbar::FillCompileToolbar(FToolBarBuilder& ToolbarBuilder
 									 TAttribute<FText>(this, &FBlueprintEditorToolbar::GetStatusTooltip),
 									 TAttribute<FSlateIcon>(this, &FBlueprintEditorToolbar::GetStatusImage),
 									 FName(TEXT("CompileBlueprint")));
+
+		FUIAction TempCompileOptionsCommand;
+		ToolbarBuilder.AddComboButton(
+			TempCompileOptionsCommand,
+			FOnGetContent::CreateStatic(&BlueprintEditorToolbarImpl::GenerateCompileOptionsWidget, BlueprintEditorPtr->GetToolkitCommands()),
+			LOCTEXT("BlupeintCompileOptions_ToolbarName",    "Compile Options"),
+			LOCTEXT("BlupeintCompileOptions_ToolbarTooltip", "Options to customize how Blueprints compile"),
+			TAttribute<FSlateIcon>(),
+			/*bSimpleComboBox =*/true
+		);
 	}
 	ToolbarBuilder.EndSection();
 }
diff --git a/Engine/Source/Editor/Kismet/Public/BlueprintEditor.h b/Engine/Source/Editor/Kismet/Public/BlueprintEditor.h
--- a/Engine/Source/Editor/Kismet/Public/BlueprintEditor.h
+++ b/Engine/Source/Editor/Kismet/Public/BlueprintEditor.h
@@ -315,6 +315,10 @@ public:
 	/** Called when Compile button is clicked */
 	virtual void Compile();
 
+	/** Helper functions used to construct/operate-on the "Save on Compile" command */
+	bool IsSaveOnCompileChecked() const;
+	void OnSaveOnCompileToggled() const;
+
 	/** Calls the above function, but returns an FReply::Handled(). Used in SButtons */
 	virtual FReply Compile_OnClickWithReply();
 
diff --git a/Engine/Source/Editor/Kismet/Public/SBlueprintEditorToolbar.h b/Engine/Source/Editor/Kismet/Public/SBlueprintEditorToolbar.h
--- a/Engine/Source/Editor/Kismet/Public/SBlueprintEditorToolbar.h
+++ b/Engine/Source/Editor/Kismet/Public/SBlueprintEditorToolbar.h
@@ -42,6 +42,7 @@ public:
 
 	/** Compile the blueprint */
 	TSharedPtr<FUICommandInfo> Compile;
+	TSharedPtr<FUICommandInfo> SaveOnCompile;
 
 	/** Switch between modes in the blueprint editor */
 	TSharedPtr<FUICommandInfo> SwitchToScriptingMode;
