Switching away from the component mode will no longer leave the details panel displaying the selected component.

#ttp 342918 - Regression: BP: A component's details are displayed in the details tab in graph view

diff --git a/Engine/Source/Editor/Kismet/Private/BlueprintEditorModes.cpp b/Engine/Source/Editor/Kismet/Private/BlueprintEditorModes.cpp
--- a/Engine/Source/Editor/Kismet/Private/BlueprintEditorModes.cpp
+++ b/Engine/Source/Editor/Kismet/Private/BlueprintEditorModes.cpp
@@ -306,6 +306,13 @@ void FBlueprintComponentsApplicationMode::PreDeactivateMode()
 	BP->GetInspector()->SetEnabled(true);
 	BP->GetInspector()->EnableComponentDetailsCustomization(false);
 	BP->EnableSCSPreview(false);
+	TSharedPtr<class SKismetInspector> Inspector = BP->GetInspector();
+
+	if(Inspector.IsValid())
+	{
+		// Clear the details panel of any selected components, they will be restored when the menu is re-activated
+		Inspector->ShowDetailsForObjects(TArray<UObject*>());
+	}
 }
 
 void FBlueprintComponentsApplicationMode::PostActivateMode()
