Further fixes for reselecting components when switching modes in Blueprints, holding onto component selection was causing a series of annoying bugs with the details panel.

diff --git a/Engine/Source/Editor/Kismet/Private/BlueprintEditorModes.cpp b/Engine/Source/Editor/Kismet/Private/BlueprintEditorModes.cpp
--- a/Engine/Source/Editor/Kismet/Private/BlueprintEditorModes.cpp
+++ b/Engine/Source/Editor/Kismet/Private/BlueprintEditorModes.cpp
@@ -306,13 +306,13 @@ void FBlueprintComponentsApplicationMode::PreDeactivateMode()
 	BP->GetInspector()->SetEnabled(true);
 	BP->GetInspector()->EnableComponentDetailsCustomization(false);
 	BP->EnableSCSPreview(false);
-	TSharedPtr<class SKismetInspector> Inspector = BP->GetInspector();
 
-	if(Inspector.IsValid())
+	// Cache component selection before clearing so it can be restored
+	for( auto& SCSNode : BP->GetSCSEditor()->GetSelectedNodes() )
 	{
-		// Clear the details panel of any selected components, they will be restored when the menu is re-activated
-		Inspector->ShowDetailsForObjects(TArray<UObject*>());
+		CachedComponentSelection.AddUnique(SCSNode->GetComponentTemplate());
 	}
+	BP->GetSCSEditor()->ClearSelection();
 }
 
 void FBlueprintComponentsApplicationMode::PostActivateMode()
@@ -323,6 +323,13 @@ void FBlueprintComponentsApplicationMode::PostActivateMode()
 	BP->UpdateSCSPreview();
 	BP->GetInspector()->EnableComponentDetailsCustomization(true);
 
+	// Reselect the cached components
+	TArray<TSharedPtr<FSCSEditorTreeNode>> Selection;
+	for( auto& Component : CachedComponentSelection )
+	{
+		BP->GetSCSEditor()->SCSTreeWidget->SetItemSelection(BP->GetSCSEditor()->GetNodeFromActorComponent(Component), true);
+	}
+
 	if (BP->GetSCSViewport()->GetIsSimulateEnabled())
 	{
 		BP->GetSCSEditor()->SetEnabled(false);
diff --git a/Engine/Source/Editor/Kismet/Public/BlueprintEditorModes.h b/Engine/Source/Editor/Kismet/Public/BlueprintEditorModes.h
--- a/Engine/Source/Editor/Kismet/Public/BlueprintEditorModes.h
+++ b/Engine/Source/Editor/Kismet/Public/BlueprintEditorModes.h
@@ -95,6 +95,8 @@ protected:
 
 	// Set of spawnable tabs in this mode
 	FWorkflowAllowedTabSet BlueprintComponentsTabFactories;
+
+	TArray<UActorComponent*> CachedComponentSelection;
 };
 
 class KISMET_API FBlueprintInterfaceApplicationMode : public FApplicationMode
