Double-clicking on variable set/get nodes that have a rep notify graph will jump to said graph.

#ttp 346423 - BP:  Double clicking on a set node for a rep notify variable should jump to the rep notify function

diff --git a/Engine/Source/Editor/Kismet/Private/BlueprintEditor.cpp b/Engine/Source/Editor/Kismet/Private/BlueprintEditor.cpp
--- a/Engine/Source/Editor/Kismet/Private/BlueprintEditor.cpp
+++ b/Engine/Source/Editor/Kismet/Private/BlueprintEditor.cpp
@@ -5027,6 +5027,22 @@ void FBlueprintEditor::OnNodeDoubleClicked(class UEdGraphNode* Node)
 			OpenDocument(Timeline, FDocumentTracker::OpenNewDocument);
 		}
 	}
+	else if(UK2Node_Variable* VariableNode = Cast<UK2Node_Variable>(Node))
+	{
+		// Jump to the RepNotify function graph
+		FName RepNotifyFunc = FBlueprintEditorUtils::GetBlueprintVariableRepNotifyFunc(GetBlueprintObj(), VariableNode->GetVarName());
+		if(RepNotifyFunc != NAME_None)
+		{
+			for( UEdGraph* Graph : GetBlueprintObj()->FunctionGraphs )
+			{
+				if(Graph->GetFName() == RepNotifyFunc)
+				{
+					FKismetEditorUtilities::BringKismetToFocusAttentionOnObject(Graph);
+					break;
+				}
+			}
+		}
+	}
 	else if (UObject* HyperlinkTarget = Node->GetJumpTargetForDoubleClick())
 	{
 		// Check to see if our outer chain contains a blueprint. If we're inside a blueprint (a graph, pin, etc.) then
