Suring up the new "save on compile" feature (making it so it is invoked when you PIE). #NUX
TTP #345833

diff --git a/Engine/Source/Editor/BlueprintGraph/Public/BlueprintEditorSettings.h b/Engine/Source/Editor/BlueprintGraph/Public/BlueprintEditorSettings.h
--- a/Engine/Source/Editor/BlueprintGraph/Public/BlueprintEditorSettings.h
+++ b/Engine/Source/Editor/BlueprintGraph/Public/BlueprintEditorSettings.h
@@ -14,7 +14,7 @@ class BLUEPRINTGRAPH_API UBlueprintEditorSettings
 
 // General Settings
 public:
-	/** Determines if Blueprints are saved whenever you compile them */
+	/** Determines if Blueprints are saved whenever you (successfully) compile them */
 	UPROPERTY(EditAnywhere, config, Category=General)
 	bool bSaveOnCompile;
 
diff --git a/Engine/Source/Editor/Kismet/Private/BlueprintEditor.cpp b/Engine/Source/Editor/Kismet/Private/BlueprintEditor.cpp
--- a/Engine/Source/Editor/Kismet/Private/BlueprintEditor.cpp
+++ b/Engine/Source/Editor/Kismet/Private/BlueprintEditor.cpp
@@ -2455,15 +2455,6 @@ void FBlueprintEditor::Compile()
 		// send record when player clicks compile and send the result
 		// this will make sure how the users activity is
 		AnalyticsTrackCompileEvent(BlueprintObj, LogResults.NumErrors, LogResults.NumWarnings);
-
-		UBlueprintEditorSettings* Settings = GetMutableDefault<UBlueprintEditorSettings>();
-		if (Settings->bSaveOnCompile)
-		{
-			TArray<UPackage*> PackagesToSave;
-			PackagesToSave.Add(BlueprintObj->GetOutermost());
-
-			FEditorFileUtils::PromptForCheckoutAndSave(PackagesToSave, /*bCheckDirty =*/false, /*bPromptToSave =*/false);
-		}
 	}
 }
 
diff --git a/Engine/Source/Editor/Kismet/Private/SBlueprintEditorToolbar.cpp b/Engine/Source/Editor/Kismet/Private/SBlueprintEditorToolbar.cpp
--- a/Engine/Source/Editor/Kismet/Private/SBlueprintEditorToolbar.cpp
+++ b/Engine/Source/Editor/Kismet/Private/SBlueprintEditorToolbar.cpp
@@ -543,7 +543,7 @@ TSharedRef<SWidget> FKismet2Menu::MakeDiffMenu(FBlueprintEditor& Kismet)
 void FFullBlueprintEditorCommands::RegisterCommands() 
 {
 	UI_COMMAND(Compile, "Compile", "Compile the blueprint", EUserInterfaceActionType::Button, FInputGesture());
-	UI_COMMAND(SaveOnCompile, "Save on Compile", "Determines if the Blueprint is saved every time you compile it", EUserInterfaceActionType::ToggleButton, FInputGesture());
+	UI_COMMAND(SaveOnCompile, "Save on Compile", "Determines if the Blueprint is saved every time you (successfully) compile it", EUserInterfaceActionType::ToggleButton, FInputGesture());
 	UI_COMMAND(SwitchToScriptingMode, "Graph", "Switches to Graph Editing Mode", EUserInterfaceActionType::ToggleButton, FInputGesture());
 	UI_COMMAND(SwitchToBlueprintDefaultsMode, "Defaults", "Switches to Blueprint Defaults Mode", EUserInterfaceActionType::ToggleButton, FInputGesture());
 	UI_COMMAND(SwitchToComponentsMode, "Components", "Switches to Components Mode", EUserInterfaceActionType::ToggleButton, FInputGesture());
diff --git a/Engine/Source/Editor/UnrealEd/Private/Kismet2/Kismet2.cpp b/Engine/Source/Editor/UnrealEd/Private/Kismet2/Kismet2.cpp
--- a/Engine/Source/Editor/UnrealEd/Private/Kismet2/Kismet2.cpp
+++ b/Engine/Source/Editor/UnrealEd/Private/Kismet2/Kismet2.cpp
@@ -32,6 +32,7 @@
 #include "ObjectEditorUtils.h"
 #include "DlgPickAssetPath.h"
 #include "ComponentAssetBroker.h"
+#include "BlueprintEditorSettings.h"
 
 #define LOCTEXT_NAMESPACE "UnrealEd.Editor"
 
@@ -473,6 +474,15 @@ void FKismetEditorUtilities::CompileBlueprint(UBlueprint* BlueprintObj, bool bIs
 	if (BlueprintPackage != NULL)
 	{
 		BlueprintPackage->SetDirtyFlag(bStartedWithUnsavedChanges);
+
+		UBlueprintEditorSettings* Settings = GetMutableDefault<UBlueprintEditorSettings>();
+		if (Settings->bSaveOnCompile && (BlueprintObj->Status == BS_UpToDate))
+		{
+			TArray<UPackage*> PackagesToSave;
+			PackagesToSave.Add(BlueprintPackage);
+
+			FEditorFileUtils::PromptForCheckoutAndSave(PackagesToSave, /*bCheckDirty =*/true, /*bPromptToSave =*/false);
+		}
 	}	
 }
 
diff --git a/Engine/Source/Editor/UnrealEd/Private/PlayLevel.cpp b/Engine/Source/Editor/UnrealEd/Private/PlayLevel.cpp
--- a/Engine/Source/Editor/UnrealEd/Private/PlayLevel.cpp
+++ b/Engine/Source/Editor/UnrealEd/Private/PlayLevel.cpp
@@ -1777,10 +1777,6 @@ void UEditorEngine::PlayInEditor( UWorld* InWorld, bool bInSimulateInEditor )
 						// ... is there a case where we'd want to recompile this again?
 					}
 
-					// Cache off the dirty flag for the package, so we can restore it later
-					UPackage* Package = Cast<UPackage>(Blueprint->GetOutermost());
-					const bool bIsPackageDirty = Package ? Package->IsDirty() : false;
-
 					Blueprint->BroadcastChanged();
 
 					UE_LOG(LogPlayLevel, Log, TEXT("[PIE] Compiling %s before PIE..."), *Blueprint->GetName());
@@ -1798,12 +1794,6 @@ void UEditorEngine::PlayInEditor( UWorld* InWorld, bool bInSimulateInEditor )
 
 						BlueprintLog.Info( FText::Format(LOCTEXT("BlueprintCompileFailed", "Blueprint {Name} failed to compile"), Arguments) );
 					}
-
-					// Restore the dirty package flag
-					if (Package)
-					{
-						Package->SetDirtyFlag(bIsPackageDirty);
-					}
 				}
 
 				CollectGarbage( GARBAGE_COLLECTION_KEEPFLAGS );
