Properly cleaning up level-script references on level change.
TTP #349207

diff --git a/Engine/Source/Editor/BlueprintGraph/Private/BlueprintActionDatabase.cpp b/Engine/Source/Editor/BlueprintGraph/Private/BlueprintActionDatabase.cpp
--- a/Engine/Source/Editor/BlueprintGraph/Private/BlueprintActionDatabase.cpp
+++ b/Engine/Source/Editor/BlueprintGraph/Private/BlueprintActionDatabase.cpp
@@ -821,7 +821,7 @@ static void BlueprintActionDatabaseImpl::OnWorldDestroyed(UWorld* DestroyedWorld
 {
 	for(auto It = DestroyedWorld->GetLevelIterator() ; It ; ++It)
 	{
-		if(const ULevelScriptBlueprint* LevelScript = (*It)->GetLevelScriptBlueprint(true))
+		if(const ULevelScriptBlueprint* LevelScript = (*It)->GetLevelScriptBlueprint(/*bDontCreate =*/true))
 		{
 			FBlueprintActionDatabase::Get().ClearAssetActions((UObject*)LevelScript);
 		}
@@ -1049,6 +1049,12 @@ void FBlueprintActionDatabase::RefreshClassActions(UClass* const Class)
 			ActionRegistry.Remove(Class);
 		}
 	}
+
+	// blueprints are handled in RefreshAssetActions()
+	if (!bIsInitializing && !bIsBlueprintClass)
+	{
+		EntryRefreshDelegate.Broadcast(Class);
+	}
 }
 
 //------------------------------------------------------------------------------
@@ -1109,18 +1115,40 @@ void FBlueprintActionDatabase::RefreshAssetActions(UObject* const AssetObject)
 	{
 		ClearAssetActions(AssetObject);
 	}
+
+	if (!bIsInitializing)
+	{
+		EntryRefreshDelegate.Broadcast(AssetObject);
+	}
 }
 
 //------------------------------------------------------------------------------
 void FBlueprintActionDatabase::ClearAssetActions(UObject* const AssetObject)
 {
-	check(BlueprintActionDatabaseImpl::IsObjectValidForDatabase(AssetObject));
-	ActionRegistry.Remove(AssetObject);
+	FActionList* ActionList = ActionRegistry.Find(AssetObject);
+
+	bool const bHasEntry = (ActionList != nullptr);
+	if (bHasEntry)
+	{
+		for (UBlueprintNodeSpawner* Action : *ActionList)
+		{
+			// because some asserts expect everything to be cleaned up in a 
+			// single GC pass, we can't wait for the GC'd Action to release its
+			// template node from the cache
+			Action->ClearCachedTemplateNode();
+		}
+		ActionRegistry.Remove(AssetObject);
+	}
 
 	if (UBlueprint* BlueprintAsset = Cast<UBlueprint>(AssetObject))
 	{
 		BlueprintAsset->OnChanged().RemoveStatic(&BlueprintActionDatabaseImpl::OnBlueprintChanged);
 	}
+
+	if (bHasEntry && (ActionList->Num() > 0) && !BlueprintActionDatabaseImpl::bIsInitializing)
+	{
+		EntryRemovedDelegate.Broadcast(AssetObject);
+	}
 }
 
 //------------------------------------------------------------------------------
@@ -1170,17 +1198,6 @@ FBlueprintActionDatabase::FActionRegistry const& FBlueprintActionDatabase::GetAl
 	return ActionRegistry;
 }
 
-//------------------------------------------------------------------------------
-int32 FBlueprintActionDatabase::EstimatedSize() const
-{
-	int32 TotalSize = sizeof(*this);
-	for (TObjectIterator<UBlueprintNodeSpawner> NodeSpawnerIt; NodeSpawnerIt; ++NodeSpawnerIt)
-	{
-		TotalSize += sizeof(**NodeSpawnerIt);
-	}
-	return TotalSize;
-}
-
 //------------------------------------------------------------------------------
 void FBlueprintActionDatabase::RegisterAllNodeActions(FBlueprintActionDatabaseRegistrar& Registrar)
 {
diff --git a/Engine/Source/Editor/BlueprintGraph/Private/BlueprintNodeSpawner.cpp b/Engine/Source/Editor/BlueprintGraph/Private/BlueprintNodeSpawner.cpp
--- a/Engine/Source/Editor/BlueprintGraph/Private/BlueprintNodeSpawner.cpp
+++ b/Engine/Source/Editor/BlueprintGraph/Private/BlueprintNodeSpawner.cpp
@@ -59,14 +59,10 @@ UBlueprintNodeSpawner::UBlueprintNodeSpawner(class FPostConstructInitializePrope
 }
 
 //------------------------------------------------------------------------------
-UBlueprintNodeSpawner::~UBlueprintNodeSpawner()
+void UBlueprintNodeSpawner::BeginDestroy()
 {
-	// @TODO: What if, on shutdown, the "SharedTemplateCache" is destroyed 
-	//        first? Then we'd be working with a bad pointer here
-	if (FBlueprintNodeTemplateCache* TemplateCache = BlueprintNodeSpawnerImpl::GetSharedTemplateCache(/*bNoInit =*/true))
-	{
-		TemplateCache->ClearCachedTemplate(this);
-	}
+	ClearCachedTemplateNode();
+	Super::BeginDestroy();
 }
 
 //------------------------------------------------------------------------------
@@ -252,6 +248,16 @@ UEdGraphNode* UBlueprintNodeSpawner::GetTemplateNode(UEdGraph* TargetGraph, FBin
 	return TemplateNode; 
 }
 
+//------------------------------------------------------------------------------
+void UBlueprintNodeSpawner::ClearCachedTemplateNode() const
+{
+	FBlueprintNodeTemplateCache* TemplateCache = BlueprintNodeSpawnerImpl::GetSharedTemplateCache(/*bNoInit =*/true);
+	if (TemplateCache != nullptr)
+	{
+		TemplateCache->ClearCachedTemplate(this);
+	}
+}
+
 //------------------------------------------------------------------------------
 UEdGraphNode* UBlueprintNodeSpawner::SpawnEdGraphNode(TSubclassOf<UEdGraphNode> InNodeClass, UEdGraph* ParentGraph, FBindingSet const& Bindings, FVector2D const Location, FCustomizeNodeDelegate PostSpawnDelegate) const
 {
diff --git a/Engine/Source/Editor/BlueprintGraph/Private/EdGraphSchema_K2.cpp b/Engine/Source/Editor/BlueprintGraph/Private/EdGraphSchema_K2.cpp
--- a/Engine/Source/Editor/BlueprintGraph/Private/EdGraphSchema_K2.cpp
+++ b/Engine/Source/Editor/BlueprintGraph/Private/EdGraphSchema_K2.cpp
@@ -838,6 +838,12 @@ bool UEdGraphSchema_K2::IsAllowableBlueprintVariableType(const class UClass* InC
 		{
 			return false;
 		}
+		
+		// cannot have level script variables
+		if (InClass->IsChildOf(ALevelScriptActor::StaticClass()))
+		{
+			return false;
+		}
 
 		const UClass* ParentClass = InClass;
 		while(ParentClass)
diff --git a/Engine/Source/Editor/BlueprintGraph/Private/K2Node_Literal.cpp b/Engine/Source/Editor/BlueprintGraph/Private/K2Node_Literal.cpp
--- a/Engine/Source/Editor/BlueprintGraph/Private/K2Node_Literal.cpp
+++ b/Engine/Source/Editor/BlueprintGraph/Private/K2Node_Literal.cpp
@@ -151,6 +151,12 @@ FLinearColor UK2Node_Literal::GetNodeTitleColor() const
 
 void UK2Node_Literal::GetMenuActions(FBlueprintActionDatabaseRegistrar& ActionRegistrar) const
 {
+	UClass* ActionKey = GetClass();
+	if (!ActionRegistrar.IsOpenForRegistration(ActionKey))
+	{
+		return;
+	}
+
 	auto CanBindObjectLambda = [](UObject const* BindingObject)
 	{
 		if(AActor const* Actor = Cast<AActor>(BindingObject))
@@ -220,7 +226,7 @@ void UK2Node_Literal::GetMenuActions(FBlueprintActionDatabaseRegistrar& ActionRe
 	NodeSpawner->OnBindObjectDelegate     = UBlueprintBoundNodeSpawner::FOnBindObjectDelegate::CreateStatic(PostBindSetupLambda);
 	NodeSpawner->DynamicUiSignatureGetter = UBlueprintBoundNodeSpawner::FUiSpecOverrideDelegate::CreateStatic(UiSpecOverride);
 
-	ActionRegistrar.AddBlueprintAction(NodeSpawner);
+	ActionRegistrar.AddBlueprintAction(ActionKey, NodeSpawner);
 }
 
 UK2Node::ERedirectType UK2Node_Literal::DoPinsMatchForReconstruction(const UEdGraphPin* NewPin, int32 NewPinIndex, const UEdGraphPin* OldPin, int32 OldPinIndex) const
diff --git a/Engine/Source/Editor/BlueprintGraph/Private/K2Node_MatineeController.cpp b/Engine/Source/Editor/BlueprintGraph/Private/K2Node_MatineeController.cpp
--- a/Engine/Source/Editor/BlueprintGraph/Private/K2Node_MatineeController.cpp
+++ b/Engine/Source/Editor/BlueprintGraph/Private/K2Node_MatineeController.cpp
@@ -130,6 +130,12 @@ void UK2Node_MatineeController::ExpandNode(FKismetCompilerContext& CompilerConte
 
 void UK2Node_MatineeController::GetMenuActions(FBlueprintActionDatabaseRegistrar& ActionRegistrar) const
 {
+	UClass* ActionKey = GetClass();
+	if (!ActionRegistrar.IsOpenForRegistration(ActionKey))
+	{
+		return;
+	}
+
 	auto CanBindObjectLambda = [](UObject const* BindingObject)
 	{
 		if (AMatineeActor const* Actor = Cast<AMatineeActor>(BindingObject))
@@ -187,7 +193,7 @@ void UK2Node_MatineeController::GetMenuActions(FBlueprintActionDatabaseRegistrar
 	NodeSpawner->OnBindObjectDelegate  = UBlueprintBoundNodeSpawner::FOnBindObjectDelegate::CreateStatic(PostBindSetupLambda);
 	NodeSpawner->DynamicUiSignatureGetter = UBlueprintBoundNodeSpawner::FUiSpecOverrideDelegate::CreateStatic(UiSpecOverride);
 	NodeSpawner->FindPreExistingNodeDelegate = UBlueprintBoundNodeSpawner::FFindPreExistingNodeDelegate::CreateStatic( FindPreExistingNodeLambda );
-	ActionRegistrar.AddBlueprintAction(NodeSpawner);
+	ActionRegistrar.AddBlueprintAction(ActionKey, NodeSpawner);
 }
 
 void UK2Node_MatineeController::OnEventKeyframeAdded(const AMatineeActor* InMatineeActor, const FName& InPinName, int32 InIndex)
diff --git a/Engine/Source/Editor/BlueprintGraph/Public/BlueprintActionDatabase.h b/Engine/Source/Editor/BlueprintGraph/Public/BlueprintActionDatabase.h
--- a/Engine/Source/Editor/BlueprintGraph/Public/BlueprintActionDatabase.h
+++ b/Engine/Source/Editor/BlueprintGraph/Public/BlueprintActionDatabase.h
@@ -24,6 +24,8 @@ class FBlueprintActionDatabaseRegistrar;
 class BLUEPRINTGRAPH_API FBlueprintActionDatabase : public FGCObject, public FTickableEditorObject
 {
 public:
+	DECLARE_MULTICAST_DELEGATE_OneParam(FOnDatabaseEntryUpdated, UObject*);
+
 	typedef TMap<TWeakObjectPtr<UObject>, int32>			FPrimingQueue;
 	typedef TArray<UBlueprintNodeSpawner*>					FActionList;
 	typedef TMap<UObject const*, FActionList>				FActionRegistry;
@@ -108,11 +110,10 @@ public:
 	 */
 	void MoveUnloadedAssetActions(FName SourceObjectPath, FName TargetObjectPath);
 
-	/**
-	 * 
-	 * @return An estimated memory footprint of this database (in bytes).
-	 */
-	int32 EstimatedSize() const;
+	/** */
+	FOnDatabaseEntryUpdated& OnEntryUpdated() { return EntryRefreshDelegate; }
+	/** */
+	FOnDatabaseEntryUpdated& OnEntryRemoved() { return EntryRemovedDelegate; }
 
 private:
 	/** Private constructor for singleton purposes. */
@@ -148,4 +149,8 @@ private:
 	 * caching each spawner's template-node, etc.).
 	 */
 	FPrimingQueue ActionPrimingQueue;
+
+	/** */
+	FOnDatabaseEntryUpdated EntryRefreshDelegate;
+	FOnDatabaseEntryUpdated EntryRemovedDelegate;
 };
diff --git a/Engine/Source/Editor/BlueprintGraph/Public/BlueprintNodeSpawner.h b/Engine/Source/Editor/BlueprintGraph/Public/BlueprintNodeSpawner.h
--- a/Engine/Source/Editor/BlueprintGraph/Public/BlueprintNodeSpawner.h
+++ b/Engine/Source/Editor/BlueprintGraph/Public/BlueprintNodeSpawner.h
@@ -71,9 +71,6 @@ public:
 	static UBlueprintNodeSpawner* Create(UObject* Outer = nullptr, FCustomizeNodeDelegate PostSpawnDelegate = FCustomizeNodeDelegate());
 	
 public:
-	/** */
-	virtual ~UBlueprintNodeSpawner();
-
 	/** Holds the node type that this spawner will instantiate. */
 	UPROPERTY()
 	TSubclassOf<UEdGraphNode> NodeClass;
@@ -87,6 +84,10 @@ public:
 	/** Provides a way to override DefaultMenuSignature based off blueprint/graph/menu context */
 	FUiSpecOverrideDelegate DynamicUiSignatureGetter;
 
+	// UObject interface
+	virtual void BeginDestroy() override;
+	// End UObject interface
+
 	/**
 	 * Not required, but intended to passively help speed up menu building 
 	 * operations. Will cache a node-template (via GetTemplateNode), along with  
@@ -160,6 +161,13 @@ public:
 	 */
 	UEdGraphNode* GetTemplateNode(UEdGraph* TargetGraph = nullptr, FBindingSet const& Bindings = FBindingSet()) const;
 
+	/**
+	 * Removes the spawner's cached template node (if it has one), meaning that
+	 * the next GetTemplateNode() call will spawn a new one (and that calls to 
+	 * GetCachedTemplateNode() will return null).
+	 */
+	void ClearCachedTemplateNode() const;
+
 	// IBlueprintNodeBinder interface
 	virtual bool IsBindingCompatible(UObject const* BindingCandidate) const override { return false; }
 	virtual bool CanBindMultipleObjects() const override { return false; }
diff --git a/Engine/Source/Editor/Kismet/Private/SBlueprintSubPalette.cpp b/Engine/Source/Editor/Kismet/Private/SBlueprintSubPalette.cpp
--- a/Engine/Source/Editor/Kismet/Private/SBlueprintSubPalette.cpp
+++ b/Engine/Source/Editor/Kismet/Private/SBlueprintSubPalette.cpp
@@ -13,6 +13,8 @@
 #include "BlueprintEditorUtils.h"
 #include "SBlueprintActionMenu.h" // for SBlueprintActionMenuExpander
 #include "BlueprintDragDropMenuItem.h"
+#include "BlueprintActionDatabase.h"
+#include "Engine/LevelScriptBlueprint.h"
 
 #define LOCTEXT_NAMESPACE "BlueprintSubPalette"
 
@@ -176,7 +178,9 @@ public:
 //------------------------------------------------------------------------------
 SBlueprintSubPalette::~SBlueprintSubPalette()
 {
-	GEditor->OnBlueprintCompiled().RemoveAll(this);
+	FBlueprintActionDatabase& ActionDatabase = FBlueprintActionDatabase::Get();
+	ActionDatabase.OnEntryRemoved().RemoveAll(this);
+	ActionDatabase.OnEntryUpdated().RemoveAll(this);
 }
 
 //------------------------------------------------------------------------------
@@ -237,10 +241,9 @@ void SBlueprintSubPalette::Construct(FArguments const& InArgs, TWeakPtr<FBluepri
 	// has to come after GraphActionMenu has been set
 	BindCommands(CommandList);
 
-	GEditor->OnBlueprintCompiled().AddSP(this, &SBlueprintSubPalette::RequestRefreshActionsList);
-	// Register with the Asset Registry to be informed when it is done loading up files.
-	FAssetRegistryModule& AssetRegistryModule = FModuleManager::GetModuleChecked<FAssetRegistryModule>(TEXT("AssetRegistry"));
-	AssetRegistryModule.Get().OnFilesLoaded().AddSP(this, &SBlueprintSubPalette::RequestRefreshActionsList);
+	FBlueprintActionDatabase& ActionDatabase = FBlueprintActionDatabase::Get();
+	ActionDatabase.OnEntryRemoved().AddSP(this, &SBlueprintSubPalette::OnDatabaseActionsRemoved);
+	ActionDatabase.OnEntryUpdated().AddSP(this, &SBlueprintSubPalette::OnDatabaseActionsUpdated);
 }
 
 //------------------------------------------------------------------------------
@@ -251,7 +254,7 @@ void SBlueprintSubPalette::Tick(const FGeometry& AllottedGeometry, const double
 	if(bNeedsRefresh)
 	{
 		bNeedsRefresh = false;
-		RefreshActionsList(true);
+		RefreshActionsList(/*bPreserveExpansion =*/true);
 	}
 }
 
@@ -361,7 +364,7 @@ void SBlueprintSubPalette::BindCommands(TSharedPtr<FUICommandList> CommandListIn
 
 	CommandListIn->MapAction(
 		PaletteCommands.RefreshPalette,
-		FExecuteAction::CreateSP(this, &SBlueprintSubPalette::RefreshActionsList, true)
+		FExecuteAction::CreateSP(this, &SBlueprintSubPalette::RefreshActionsList, /*bPreserveExpansion =*/true)
 	);
 }
 
@@ -386,10 +389,41 @@ void SBlueprintSubPalette::RequestRefreshActionsList()
 	bNeedsRefresh = true;
 }
 
+//------------------------------------------------------------------------------
+void SBlueprintSubPalette::OnDatabaseActionsUpdated(UObject* /*ActionsKey*/)
+{
+	RequestRefreshActionsList();
+}
+
+//------------------------------------------------------------------------------
+void SBlueprintSubPalette::OnDatabaseActionsRemoved(UObject* ActionsKey)
+{
+	ULevelScriptBlueprint* RemovedLevelScript = Cast<ULevelScriptBlueprint>(ActionsKey);
+	bool const bAssumeDestroyingWorld = (RemovedLevelScript != nullptr);
+
+	if (bAssumeDestroyingWorld)
+	{
+		// have to update the action list immediatly (cannot wait until Tick(), 
+		// because we have to handle level switching, which expects all references 
+		// to be cleared immediately)
+		ForceRefreshActionList();
+	}
+	else
+	{
+		RequestRefreshActionsList();
+	}
+}
+
 /*******************************************************************************
 * Private SBlueprintSubPalette Methods
 *******************************************************************************/
 
+//------------------------------------------------------------------------------
+void SBlueprintSubPalette::ForceRefreshActionList()
+{
+	RefreshActionsList(/*bPreserveExpansion =*/true);
+}
+
 //------------------------------------------------------------------------------
 TSharedRef<SVerticalBox> SBlueprintSubPalette::ConstructHeadingWidget(FSlateBrush const* const Icon, FText const& TitleText, FText const& ToolTipText)
 {
diff --git a/Engine/Source/Editor/Kismet/Private/SBlueprintSubPalette.h b/Engine/Source/Editor/Kismet/Private/SBlueprintSubPalette.h
--- a/Engine/Source/Editor/Kismet/Private/SBlueprintSubPalette.h
+++ b/Engine/Source/Editor/Kismet/Private/SBlueprintSubPalette.h
@@ -95,11 +95,17 @@ protected:
 
 	/** Delegate to call to request a refresh */
 	void RequestRefreshActionsList();
+	/** Delegate handlers for when the blueprint database is updated (so we can release references and refresh the list) */
+	void OnDatabaseActionsUpdated(UObject* ActionsKey);
+	void OnDatabaseActionsRemoved(UObject* ActionsKey);
 
 	/** Pointer back to the blueprint editor that owns us */
 	TWeakPtr<FBlueprintEditor> BlueprintEditorPtr;
 
 private:
+	/** Immediately calls RefreshActionsList(), doesn't not defer until Tick() like RequestRefreshActionsList() does. */
+	void ForceRefreshActionList();
+
 	/**
 	 * Constructs a slate widget for the right-click context menu in this 
 	 * palette. While this isn't virtual, sub-classes can override GenerateContextMenuEntries()
