Removed a bunch of stuff from Slate standard include, created SlateBasics.h

diff --git a/Engine/Source/Editor/GameplayAbilitiesEditor/Private/AttributeDetails.cpp b/Engine/Source/Editor/GameplayAbilitiesEditor/Private/AttributeDetails.cpp
--- a/Engine/Source/Editor/GameplayAbilitiesEditor/Private/AttributeDetails.cpp
+++ b/Engine/Source/Editor/GameplayAbilitiesEditor/Private/AttributeDetails.cpp
@@ -9,6 +9,8 @@
 #include "AbilitySystemGlobals.h"
 #include "GameplayAbilitiesModule.h"
 #include "SGameplayAttributeGraphPin.h"
+#include "SSearchBox.h"
+#include "STextComboBox.h"
 
 #define LOCTEXT_NAMESPACE "AttributeDetailsCustomization"
 
diff --git a/Engine/Source/Editor/GameplayAbilitiesEditor/Private/SGameplayAttributeWidget.cpp b/Engine/Source/Editor/GameplayAbilitiesEditor/Private/SGameplayAttributeWidget.cpp
--- a/Engine/Source/Editor/GameplayAbilitiesEditor/Private/SGameplayAttributeWidget.cpp
+++ b/Engine/Source/Editor/GameplayAbilitiesEditor/Private/SGameplayAttributeWidget.cpp
@@ -3,6 +3,7 @@
 #include "AbilitySystemEditorPrivatePCH.h"
 #include "KismetEditorUtilities.h"
 #include "SGameplayAttributeWidget.h"
+#include "STextComboBox.h"
 
 #define LOCTEXT_NAMESPACE "K2Node"
 
diff --git a/Engine/Source/Editor/GraphEditor/Private/SGraphActionMenu.cpp b/Engine/Source/Editor/GraphEditor/Private/SGraphActionMenu.cpp
--- a/Engine/Source/Editor/GraphEditor/Private/SGraphActionMenu.cpp
+++ b/Engine/Source/Editor/GraphEditor/Private/SGraphActionMenu.cpp
@@ -7,6 +7,8 @@
 #include "SScrollBorder.h"
 #include "IDocumentation.h"
 #include "EditorCategoryUtils.h"
+#include "SSearchBox.h"
+#include "SInlineEditableTextBlock.h"
 
 #define LOCTEXT_NAMESPACE "GraphActionMenu"
 
diff --git a/Engine/Source/Editor/GraphEditor/Private/SGraphEditorImpl.cpp b/Engine/Source/Editor/GraphEditor/Private/SGraphEditorImpl.cpp
--- a/Engine/Source/Editor/GraphEditor/Private/SGraphEditorImpl.cpp
+++ b/Engine/Source/Editor/GraphEditor/Private/SGraphEditorImpl.cpp
@@ -4,6 +4,8 @@
 #include "GraphEditorCommon.h"
 #include "SGraphEditorImpl.h"
 #include "GraphEditorModule.h"
+#include "SNotificationList.h"
+#include "SGraphPanel.h"
 
 #define LOCTEXT_NAMESPACE "GraphEditorModule"
 
diff --git a/Engine/Source/Editor/GraphEditor/Private/SGraphNode.cpp b/Engine/Source/Editor/GraphEditor/Private/SGraphNode.cpp
--- a/Engine/Source/Editor/GraphEditor/Private/SGraphNode.cpp
+++ b/Engine/Source/Editor/GraphEditor/Private/SGraphNode.cpp
@@ -10,6 +10,8 @@
 #include "SLevelOfDetailBranchNode.h"
 #include "IDocumentation.h"
 #include "TutorialMetaData.h"
+#include "SGraphPanel.h"
+#include "SInlineEditableTextBlock.h"
 
 /////////////////////////////////////////////////////
 // SNodeTitle
diff --git a/Engine/Source/Editor/GraphEditor/Private/SGraphNodeKnot.cpp b/Engine/Source/Editor/GraphEditor/Private/SGraphNodeKnot.cpp
--- a/Engine/Source/Editor/GraphEditor/Private/SGraphNodeKnot.cpp
+++ b/Engine/Source/Editor/GraphEditor/Private/SGraphNodeKnot.cpp
@@ -5,6 +5,8 @@
 #include "SGraphNodeKnot.h"
 #include "ScopedTransaction.h"
 #include "BlueprintEditorUtils.h"
+#include "GenericCommands.h"
+#include "SInlineEditableTextBlock.h"
 
 
 class FAmbivalentDirectionDragConnection : public FDragConnection
diff --git a/Engine/Source/Editor/GraphEditor/Private/SGraphPanel.cpp b/Engine/Source/Editor/GraphEditor/Private/SGraphPanel.cpp
--- a/Engine/Source/Editor/GraphEditor/Private/SGraphPanel.cpp
+++ b/Engine/Source/Editor/GraphEditor/Private/SGraphPanel.cpp
@@ -20,6 +20,7 @@
 
 DEFINE_LOG_CATEGORY_STATIC(LogGraphPanel, Log, All);
 
+DECLARE_CYCLE_STAT( TEXT("OnPaint SGraphPanel"), STAT_SlateOnPaint_SGraphPanel, STATGROUP_Slate );
 
 SGraphPanel::FGraphPinHandle::FGraphPinHandle(UEdGraphPin* InPin)
 {
diff --git a/Engine/Source/Editor/GraphEditor/Private/SNodePanel.cpp b/Engine/Source/Editor/GraphEditor/Private/SNodePanel.cpp
--- a/Engine/Source/Editor/GraphEditor/Private/SNodePanel.cpp
+++ b/Engine/Source/Editor/GraphEditor/Private/SNodePanel.cpp
@@ -3,6 +3,7 @@
 
 #include "GraphEditorCommon.h"
 #include "ScopedTransaction.h"
+#include "MarqueeOperation.h"
 
 struct FZoomLevelEntry
 {
diff --git a/Engine/Source/Editor/GraphEditor/Public/SNodePanel.h b/Engine/Source/Editor/GraphEditor/Public/SNodePanel.h
--- a/Engine/Source/Editor/GraphEditor/Public/SNodePanel.h
+++ b/Engine/Source/Editor/GraphEditor/Public/SNodePanel.h
@@ -3,6 +3,7 @@
 #pragma once
 
 #include "GraphEditorModule.h"
+#include "MarqueeOperation.h"
 
 //@TODO: Too generic of a name to expose at this scope
 typedef class UObject* SelectedItemType;
@@ -29,109 +30,6 @@ namespace EGraphRenderingLOD
 	};
 }
 
-/** Helper for managing marquee operations */
-struct FMarqueeOperation
-{
-	FMarqueeOperation()
-		: Operation(Add)
-	{
-	}
-
-	enum Type
-	{
-		/** Holding down Ctrl removes nodes */
-		Remove,
-		/** Holding down Shift adds to the selection */
-		Add,
-		/** When nothing is pressed, marquee replaces selection */
-		Replace
-	} Operation;
-
-	bool IsValid() const
-	{
-		return Rect.IsValid();
-	}
-
-	void Start(const FVector2D& InStartLocation, FMarqueeOperation::Type InOperationType)
-	{
-		Rect = FMarqueeRect(InStartLocation);
-		Operation = InOperationType;
-	}
-
-	void End()
-	{
-		Rect = FMarqueeRect();
-	}
-
-
-	/** Given a mouse event, figure out what the marquee selection should do based on the state of Shift and Ctrl keys */
-	static FMarqueeOperation::Type OperationTypeFromMouseEvent(const FPointerEvent& MouseEvent)
-	{
-		if (MouseEvent.IsControlDown())
-		{
-			return FMarqueeOperation::Remove;
-		}
-		else if (MouseEvent.IsShiftDown())
-		{
-			return FMarqueeOperation::Add;
-		}
-		else
-		{
-			return FMarqueeOperation::Replace;
-		}
-	}
-
-public:
-	/** The marquee rectangle being dragged by the user */
-	FMarqueeRect Rect;
-
-	/** Nodes that will be selected or unselected by the current marquee operation */
-	FGraphPanelSelectionSet AffectedNodes;
-};
-
-struct GRAPHEDITOR_API FGraphSelectionManager
-{
-	FGraphPanelSelectionSet SelectedNodes;
-
-	/** Invoked when the selected graph nodes have changed. */
-	SGraphEditor::FOnSelectionChanged OnSelectionChanged;
-public:
-	/** @return the set of selected nodes */
-	const FGraphPanelSelectionSet& GetSelectedNodes() const;
-
-	/** Select just the specified node */
-	void SelectSingleNode(SelectedItemType Node);
-
-	/** Reset the selection state of all nodes */
-	void ClearSelectionSet();
-
-	/** Returns true if any nodes are selected */
-	bool AreAnyNodesSelected() const
-	{
-		return SelectedNodes.Num() > 0;
-	}
-
-	/** Changes the selection set to contain exactly all of the passed in nodes */
-	void SetSelectionSet(FGraphPanelSelectionSet& NewSet);
-
-	/**
-	 * Add or remove a node from the selection set
-	 *
-	 * @param Node      Node the affect.
-	 * @param bSelect   true to select the node; false to unselect.
-	 */
-	void SetNodeSelection(SelectedItemType Node, bool bSelect);
-
-	/** @return true if Node is selected; false otherwise */
-	bool IsNodeSelected(SelectedItemType Node) const;
-
-	// Handle the selection mechanics of starting to drag a node
-	void StartDraggingNode(SelectedItemType NodeBeingDragged, const FPointerEvent& MouseEvent);
-
-	// Handle the selection mechanics when a node is clicked on
-	void ClickedOnNode(SelectedItemType Node, const FPointerEvent& MouseEvent);
-};
-
 // Context passed in when getting popup info
 struct FNodeInfoContext
 {
@@ -261,6 +159,49 @@ struct FZoomLevelsContainer
 	virtual ~FZoomLevelsContainer( void ) {};
 };
 
+struct GRAPHEDITOR_API FGraphSelectionManager
+{
+	FGraphPanelSelectionSet SelectedNodes;
+
+	/** Invoked when the selected graph nodes have changed. */
+	SGraphEditor::FOnSelectionChanged OnSelectionChanged;
+public:
+	/** @return the set of selected nodes */
+	const FGraphPanelSelectionSet& GetSelectedNodes() const;
+
+	/** Select just the specified node */
+	void SelectSingleNode(SelectedItemType Node);
+
+	/** Reset the selection state of all nodes */
+	void ClearSelectionSet();
+
+	/** Returns true if any nodes are selected */
+	bool AreAnyNodesSelected() const
+	{
+		return SelectedNodes.Num() > 0;
+	}
+
+	/** Changes the selection set to contain exactly all of the passed in nodes */
+	void SetSelectionSet(FGraphPanelSelectionSet& NewSet);
+
+	/**
+	 * Add or remove a node from the selection set
+	 *
+	 * @param Node      Node the affect.
+	 * @param bSelect   true to select the node; false to unselect.
+	 */
+	void SetNodeSelection(SelectedItemType Node, bool bSelect);
+
+	/** @return true if Node is selected; false otherwise */
+	bool IsNodeSelected(SelectedItemType Node) const;
+
+	// Handle the selection mechanics of starting to drag a node
+	void StartDraggingNode(SelectedItemType NodeBeingDragged, const FPointerEvent& MouseEvent);
+
+	// Handle the selection mechanics when a node is clicked on
+	void ClickedOnNode(SelectedItemType Node, const FPointerEvent& MouseEvent);
+};
+
 /**
  * This class is designed to serve as the base class for a panel/canvas that contains interactive widgets
  * which can be selected and moved around by the user.  It also manages zooming and panning, allowing a larger
diff --git a/Engine/Source/Editor/InputBindingEditor/Private/Widgets/SGestureTreeItem.cpp b/Engine/Source/Editor/InputBindingEditor/Private/Widgets/SGestureTreeItem.cpp
new file mode 100644
--- /dev/null
+++ b/Engine/Source/Editor/InputBindingEditor/Private/Widgets/SGestureTreeItem.cpp
@@ -0,0 +1,89 @@
+// Copyright 1998-2014 Epic Games, Inc. All Rights Reserved.
+
+#include "InputBindingEditorPrivatePCH.h"
+#include "SGestureTreeItem.h"
+#include "UICommandDragDropOp.h"
+
+
+
+
+void SGestureTreeItem::Construct(const FArguments& InArgs, const TSharedRef<STableViewBase>& InOwnerTable, TSharedPtr<FGestureTreeItem> InItem)
+{
+	TreeItem = InItem;
+
+	SMultiColumnTableRow< TSharedPtr<FGestureTreeItem> >::Construct(FSuperRowType::FArguments(), InOwnerTable);
+}
+
+TSharedRef<SWidget> SGestureTreeItem::GenerateWidgetForColumn(const FName& ColumnName)
+{
+	if (ColumnName == "Name")
+	{
+		return SNew(SHorizontalBox)
+
+			+ SHorizontalBox::Slot()
+			.AutoWidth()
+			.Padding(2.0f)
+			.VAlign(VAlign_Center)
+			[
+				SNew(SExpanderArrow, SharedThis(this))
+			]
+
+		+ SHorizontalBox::Slot()
+			.AutoWidth()
+			.VAlign(VAlign_Center)
+			[
+				SNew(SVerticalBox)
+
+				+ SVerticalBox::Slot()
+				.AutoHeight()
+				[
+					SNew(STextBlock)
+					.Text(TreeItem->CommandInfo->GetLabel())
+					.ToolTipText(TreeItem->CommandInfo->GetDescription())
+				]
+
+				+ SVerticalBox::Slot()
+					.Padding(0.0f, 0.0f, 0.0f, 5.0f)
+					.AutoHeight()
+					[
+						SNew(STextBlock)
+						.Font(FEditorStyle::GetFontStyle("InputBindingEditor.SmallFont"))
+						.ColorAndOpacity(FLinearColor::Gray)
+						.Text(TreeItem->CommandInfo->GetDescription())
+					]
+			];
+	}
+	else if (ColumnName == "Binding")
+	{
+		return SNew(SBox)
+			.Padding(2.0f)
+			.VAlign(VAlign_Center)
+			[
+				SNew(SGestureEditBox, TreeItem)
+			];
+	}
+	else
+	{
+		return SNew(STextBlock).Text(NSLOCTEXT("GestureTreeItem", "UnknownColumn", "Unknown Column"));
+	}
+}
+
+FReply SGestureTreeItem::OnMouseButtonDown(const FGeometry& MyGeometry, const FPointerEvent& MouseEvent)
+{
+	if (MouseEvent.GetEffectingButton() == EKeys::LeftMouseButton && FMultiBoxSettings::IsInToolbarEditMode())
+	{
+		return FReply::Handled().DetectDrag(SharedThis(this), MouseEvent.GetEffectingButton());
+	}
+
+	return FReply::Unhandled();
+}
+
+FReply SGestureTreeItem::OnDragDetected(const FGeometry& MyGeometry, const FPointerEvent& MouseEvent)
+{
+	FToolBarBuilder TempBuilder(MakeShareable(new FUICommandList), FMultiBoxCustomization::None);
+	TempBuilder.AddToolBarButton(TreeItem->CommandInfo.ToSharedRef());
+
+	TSharedRef<SWidget> CursorDecorator = TempBuilder.MakeWidget();
+
+	return FReply::Handled().BeginDragDrop(FUICommandDragDropOp::New(TreeItem->CommandInfo.ToSharedRef(), NAME_None, TempBuilder.MakeWidget(), FVector2D::ZeroVector));
+}
diff --git a/Engine/Source/Editor/InputBindingEditor/Private/Widgets/SGestureTreeItem.h b/Engine/Source/Editor/InputBindingEditor/Private/Widgets/SGestureTreeItem.h
--- a/Engine/Source/Editor/InputBindingEditor/Private/Widgets/SGestureTreeItem.h
+++ b/Engine/Source/Editor/InputBindingEditor/Private/Widgets/SGestureTreeItem.h
@@ -42,95 +42,22 @@ public:
 	 * @param InOwnerTable The table that owns this tree item.
 	 * @param InItem The actual item to be displayed.
 	 */
-	void Construct( const FArguments& InArgs, const TSharedRef<STableViewBase>& InOwnerTable, TSharedPtr<FGestureTreeItem> InItem )
-	{
-		TreeItem = InItem;
-
-		SMultiColumnTableRow< TSharedPtr<FGestureTreeItem> >::Construct( FSuperRowType::FArguments(), InOwnerTable );
-	}
+	void Construct( const FArguments& InArgs, const TSharedRef<STableViewBase>& InOwnerTable, TSharedPtr<FGestureTreeItem> InItem );
 
 	/**
 	 * Called to generate a widget for each column.
 	 *
 	 * @param ColumnName The name of the column that needs a widget.
 	 */
-	TSharedRef<SWidget> GenerateWidgetForColumn( const FName& ColumnName )
-	{
-		if( ColumnName == "Name")
-		{
-			return SNew(SHorizontalBox)
-
-				+SHorizontalBox::Slot()
-					.AutoWidth()
-					.Padding( 2.0f )
-					.VAlign(VAlign_Center)
-					[
-						SNew( SExpanderArrow, SharedThis( this ) )
-					]
-
-				+SHorizontalBox::Slot()
-					.AutoWidth()
-					.VAlign(VAlign_Center)
-					[
-						SNew( SVerticalBox )
-
-						+ SVerticalBox::Slot()
-							.AutoHeight()
-							[
-								SNew( STextBlock )
-								.Text( TreeItem->CommandInfo->GetLabel() )
-								.ToolTipText( TreeItem->CommandInfo->GetDescription() )
-							]
-
-						+ SVerticalBox::Slot()
-							.Padding(0.0f,0.0f,0.0f,5.0f)
-							.AutoHeight()
-							[
-								SNew( STextBlock )
-									.Font( FEditorStyle::GetFontStyle("InputBindingEditor.SmallFont") )
-									.ColorAndOpacity( FLinearColor::Gray )
-									.Text( TreeItem->CommandInfo->GetDescription() )
-							]
-					];
-		}
-		else if( ColumnName == "Binding")
-		{
-			return SNew( SBox )
-				.Padding(2.0f)
-				.VAlign( VAlign_Center )
-				[
-					SNew( SGestureEditBox, TreeItem )
-				];
-		}
-		else
-		{
-			return SNew(STextBlock) .Text(NSLOCTEXT("GestureTreeItem", "UnknownColumn", "Unknown Column"));
-		}
-	}
+	TSharedRef<SWidget> GenerateWidgetForColumn( const FName& ColumnName );
 
 public:
 
 	// SWidget interface
 
-	virtual FReply OnMouseButtonDown( const FGeometry& MyGeometry, const FPointerEvent& MouseEvent ) override
-	{
-		if( MouseEvent.GetEffectingButton() == EKeys::LeftMouseButton && FMultiBoxSettings::IsInToolbarEditMode() )
-		{
-			return FReply::Handled().DetectDrag( SharedThis(this), MouseEvent.GetEffectingButton() );
-		}
-
-		return FReply::Unhandled();
-	}
-
-	virtual FReply OnDragDetected( const FGeometry& MyGeometry, const FPointerEvent& MouseEvent ) override
-	{
-		FToolBarBuilder TempBuilder( MakeShareable(new FUICommandList), FMultiBoxCustomization::None );
-		TempBuilder.AddToolBarButton( TreeItem->CommandInfo.ToSharedRef() );
-
-		TSharedRef<SWidget> CursorDecorator = TempBuilder.MakeWidget();
+	virtual FReply OnMouseButtonDown( const FGeometry& MyGeometry, const FPointerEvent& MouseEvent ) override;
 
-		return FReply::Handled().BeginDragDrop( FUICommandDragDropOp::New( TreeItem->CommandInfo.ToSharedRef(), NAME_None, TempBuilder.MakeWidget(), FVector2D::ZeroVector ) );
-	}
+	virtual FReply OnDragDetected( const FGeometry& MyGeometry, const FPointerEvent& MouseEvent ) override;
 
 private:
 
diff --git a/Engine/Source/Editor/InputBindingEditor/Private/Widgets/SInputBindingEditorPanel.cpp b/Engine/Source/Editor/InputBindingEditor/Private/Widgets/SInputBindingEditorPanel.cpp
--- a/Engine/Source/Editor/InputBindingEditor/Private/Widgets/SInputBindingEditorPanel.cpp
+++ b/Engine/Source/Editor/InputBindingEditor/Private/Widgets/SInputBindingEditorPanel.cpp
@@ -1,7 +1,7 @@
 // Copyright 1998-2014 Epic Games, Inc. All Rights Reserved.
 
 #include "InputBindingEditorPrivatePCH.h"
-
+#include "SSearchBox.h"
 
 #define LOCTEXT_NAMESPACE "SInputBindingEditorPanel"
 
diff --git a/Engine/Source/Editor/IntroTutorials/Private/STutorialRoot.cpp b/Engine/Source/Editor/IntroTutorials/Private/STutorialRoot.cpp
--- a/Engine/Source/Editor/IntroTutorials/Private/STutorialRoot.cpp
+++ b/Engine/Source/Editor/IntroTutorials/Private/STutorialRoot.cpp
@@ -12,6 +12,8 @@
 #include "EngineAnalytics.h"
 #include "Runtime/Analytics/Analytics/Public/Interfaces/IAnalyticsProvider.h"
 #include "Editor/MainFrame/Public/Interfaces/IMainFrameModule.h"
+#include "NotificationManager.h"
+#include "SNotificationList.h"
 
 #define LOCTEXT_NAMESPACE "STutorialRoot"
 
diff --git a/Engine/Source/Editor/IntroTutorials/Private/STutorialsBrowser.cpp b/Engine/Source/Editor/IntroTutorials/Private/STutorialsBrowser.cpp
--- a/Engine/Source/Editor/IntroTutorials/Private/STutorialsBrowser.cpp
+++ b/Engine/Source/Editor/IntroTutorials/Private/STutorialsBrowser.cpp
@@ -10,6 +10,8 @@
 #include "AssetRegistryModule.h"
 #include "EngineAnalytics.h"
 #include "Runtime/Analytics/Analytics/Public/Interfaces/IAnalyticsProvider.h"
+#include "SSearchBox.h"
+#include "SBreadcrumbTrail.h"
 
 #define LOCTEXT_NAMESPACE "TutorialsBrowser"
 
diff --git a/Engine/Source/Editor/IntroTutorials/Private/TutorialHyperlinkRun.cpp b/Engine/Source/Editor/IntroTutorials/Private/TutorialHyperlinkRun.cpp
--- a/Engine/Source/Editor/IntroTutorials/Private/TutorialHyperlinkRun.cpp
+++ b/Engine/Source/Editor/IntroTutorials/Private/TutorialHyperlinkRun.cpp
@@ -2,6 +2,7 @@
 
 #include "IntroTutorialsPrivatePCH.h"
 #include "TutorialHyperlinkRun.h"
+#include "SRichTextHyperlink.h"
 
 TSharedRef< FTutorialHyperlinkRun > FTutorialHyperlinkRun::Create( const FRunInfo& InRunInfo, const TSharedRef< const FString >& InText, const FHyperlinkStyle& InStyle, FOnClick NavigateDelegate, FOnGenerateTooltip InTooltipDelegate, FOnGetTooltipText InTooltipTextDelegate )
 {
diff --git a/Engine/Source/Editor/Kismet/Private/BlueprintEditor.cpp b/Engine/Source/Editor/Kismet/Private/BlueprintEditor.cpp
--- a/Engine/Source/Editor/Kismet/Private/BlueprintEditor.cpp
+++ b/Engine/Source/Editor/Kismet/Private/BlueprintEditor.cpp
@@ -25,6 +25,7 @@
 #include "BlueprintEditorModule.h"
 #include "GraphEditorActions.h"
 #include "SNodePanel.h"
+#include "SDockTab.h"
 
 #include "SBlueprintEditorToolbar.h"
 #include "FindInBlueprints.h"
@@ -78,6 +79,10 @@
 
 // Blueprint merging
 #include "Merge.h"
+#include "SHyperlink.h"
+#include "GenericCommands.h"
+#include "SNotificationList.h"
+#include "NotificationManager.h"
 
 #define LOCTEXT_NAMESPACE "BlueprintEditor"
 
diff --git a/Engine/Source/Editor/Kismet/Private/BlueprintEditorTabFactories.cpp b/Engine/Source/Editor/Kismet/Private/BlueprintEditorTabFactories.cpp
new file mode 100644
--- /dev/null
+++ b/Engine/Source/Editor/Kismet/Private/BlueprintEditorTabFactories.cpp
@@ -0,0 +1,292 @@
+// Copyright 1998-2014 Epic Games, Inc. All Rights Reserved.
+
+#include "BlueprintEditorPrivatePCH.h"
+
+#include "BlueprintEditorTabFactories.h"
+#include "BlueprintEditorTabs.h"
+#include "SDockTab.h"
+#include "BlueprintEditor.h"
+#include "STimelineEditor.h"
+#include "Debugging/SKismetDebuggingView.h"
+#include "SKismetInspector.h"
+#include "SSCSEditor.h"
+#include "SSCSEditorViewport.h"
+#include "SBlueprintPalette.h"
+#include "FindInBlueprints.h"
+#include "SMyBlueprint.h"
+
+
+
+#define LOCTEXT_NAMESPACE "BlueprintEditor"
+
+void FGraphEditorSummoner::OnTabActivated(TSharedPtr<SDockTab> Tab) const
+{
+	TSharedRef<SGraphEditor> GraphEditor = StaticCastSharedRef<SGraphEditor>(Tab->GetContent());
+	BlueprintEditorPtr.Pin()->OnGraphEditorFocused(GraphEditor);
+}
+
+void FGraphEditorSummoner::OnTabRefreshed(TSharedPtr<SDockTab> Tab) const
+{
+	TSharedRef<SGraphEditor> GraphEditor = StaticCastSharedRef<SGraphEditor>(Tab->GetContent());
+	GraphEditor->NotifyGraphChanged();
+}
+
+void FGraphEditorSummoner::SaveState(TSharedPtr<SDockTab> Tab, TSharedPtr<FTabPayload> Payload) const
+{
+	TSharedRef<SGraphEditor> GraphEditor = StaticCastSharedRef<SGraphEditor>(Tab->GetContent());
+
+	FVector2D ViewLocation;
+	float ZoomAmount;
+	GraphEditor->GetViewLocation(ViewLocation, ZoomAmount);
+
+	UEdGraph* Graph = FTabPayload_UObject::CastChecked<UEdGraph>(Payload);
+
+	if (BlueprintEditorPtr.Pin()->IsGraphInCurrentBlueprint(Graph))
+	{
+		// Don't save references to external graphs.
+		BlueprintEditorPtr.Pin()->GetBlueprintObj()->LastEditedDocuments.Add(FEditedDocumentInfo(Graph, ViewLocation, ZoomAmount));
+	}
+}
+
+FGraphEditorSummoner::FGraphEditorSummoner(TSharedPtr<class FBlueprintEditor> InBlueprintEditorPtr, FOnCreateGraphEditorWidget CreateGraphEditorWidgetCallback) : FDocumentTabFactoryForObjects<UEdGraph>(FBlueprintEditorTabs::GraphEditorID, InBlueprintEditorPtr)
+, BlueprintEditorPtr(InBlueprintEditorPtr)
+, OnCreateGraphEditorWidget(CreateGraphEditorWidgetCallback)
+{
+
+}
+
+const FSlateBrush* FGraphEditorSummoner::GetTabIconForObject(const FWorkflowTabSpawnInfo& Info, UEdGraph* DocumentID) const
+{
+	return FBlueprintEditor::GetGlyphForGraph(DocumentID, false);
+}
+
+TSharedRef<FGenericTabHistory> FGraphEditorSummoner::CreateTabHistoryNode(TSharedPtr<FTabPayload> Payload)
+{
+	return MakeShareable(new FGraphTabHistory(SharedThis(this), Payload));
+}
+
+void FTimelineEditorSummoner::OnTabRefreshed(TSharedPtr<SDockTab> Tab) const
+{
+	TSharedRef<STimelineEditor> TimelineEditor = StaticCastSharedRef<STimelineEditor>(Tab->GetContent());
+	TimelineEditor->OnTimelineChanged();
+}
+
+FTimelineEditorSummoner::FTimelineEditorSummoner(TSharedPtr<class FBlueprintEditor> InBlueprintEditorPtr) : FDocumentTabFactoryForObjects<UTimelineTemplate>(FBlueprintEditorTabs::TimelineEditorID, InBlueprintEditorPtr)
+, BlueprintEditorPtr(InBlueprintEditorPtr)
+{
+
+}
+
+TSharedRef<SWidget> FTimelineEditorSummoner::CreateTabBodyForObject(const FWorkflowTabSpawnInfo& Info, UTimelineTemplate* DocumentID) const
+{
+	return SNew(STimelineEditor, BlueprintEditorPtr.Pin(), DocumentID);
+}
+
+const FSlateBrush* FTimelineEditorSummoner::GetTabIconForObject(const FWorkflowTabSpawnInfo& Info, UTimelineTemplate* DocumentID) const
+{
+	return FEditorStyle::GetBrush("GraphEditor.Timeline_16x");
+}
+
+void FTimelineEditorSummoner::SaveState(TSharedPtr<SDockTab> Tab, TSharedPtr<FTabPayload> Payload) const
+{
+	UTimelineTemplate* Timeline = FTabPayload_UObject::CastChecked<UTimelineTemplate>(Payload);
+	BlueprintEditorPtr.Pin()->GetBlueprintObj()->LastEditedDocuments.Add(FEditedDocumentInfo(Timeline));
+}
+
+TAttribute<FText> FTimelineEditorSummoner::ConstructTabNameForObject(UTimelineTemplate* DocumentID) const
+{
+	return TAttribute<FText>::Create(TAttribute<FText>::FGetter::CreateStatic<UObject*>(&FLocalKismetCallbacks::GetObjectName, DocumentID));
+}
+
+FDebugInfoSummoner::FDebugInfoSummoner(TSharedPtr<class FAssetEditorToolkit> InHostingApp) : FWorkflowTabFactory(FBlueprintEditorTabs::DebugID, InHostingApp)
+{
+	TabLabel = LOCTEXT("DebugTabTitle", "Debug");
+	TabIcon = FSlateIcon(FEditorStyle::GetStyleSetName(), "DebugTools.TabIcon");
+
+	EnableTabPadding();
+	bIsSingleton = true;
+
+	ViewMenuDescription = LOCTEXT("DebugView", "Debug");
+	ViewMenuTooltip = LOCTEXT("DebugView_ToolTip", "Shows the debugging view");
+}
+
+TSharedRef<SWidget> FDebugInfoSummoner::CreateTabBody(const FWorkflowTabSpawnInfo& Info) const
+{
+	TSharedPtr<FBlueprintEditor> BlueprintEditorPtr = StaticCastSharedPtr<FBlueprintEditor>(HostingApp.Pin());
+
+	return BlueprintEditorPtr->GetDebuggingView();
+}
+
+FDefaultsEditorSummoner::FDefaultsEditorSummoner(TSharedPtr<class FAssetEditorToolkit> InHostingApp) : FWorkflowTabFactory(FBlueprintEditorTabs::DefaultEditorID, InHostingApp)
+{
+	TabLabel = LOCTEXT("BlueprintDefaultsTabTitle", "Blueprint Defaults"); //@TODO: ANIMATION: GetDefaultEditorTitle(); !!!
+	TabIcon = FSlateIcon(FEditorStyle::GetStyleSetName(), "Kismet.Tabs.BlueprintDefaults");
+
+	bIsSingleton = true;
+
+	ViewMenuDescription = LOCTEXT("DefaultEditorView", "Defaults");
+	ViewMenuTooltip = LOCTEXT("DefaultEditorView_ToolTip", "Shows the default editor view");
+}
+
+TSharedRef<SWidget> FDefaultsEditorSummoner::CreateTabBody(const FWorkflowTabSpawnInfo& Info) const
+{
+	TSharedPtr<FBlueprintEditor> BlueprintEditorPtr = StaticCastSharedPtr<FBlueprintEditor>(HostingApp.Pin());
+
+	return BlueprintEditorPtr->GetDefaultEditor();
+}
+
+FConstructionScriptEditorSummoner::FConstructionScriptEditorSummoner(TSharedPtr<class FAssetEditorToolkit> InHostingApp) : FWorkflowTabFactory(FBlueprintEditorTabs::ConstructionScriptEditorID, InHostingApp)
+{
+	TabLabel = LOCTEXT("ComponentsTabLabel", "Components");
+	TabIcon = FSlateIcon(FEditorStyle::GetStyleSetName(), "Kismet.Tabs.Components");
+
+	bIsSingleton = true;
+
+	ViewMenuDescription = LOCTEXT("ComponentsView", "Components");
+	ViewMenuTooltip = LOCTEXT("ComponentsView_ToolTip", "Show the components view");
+}
+
+TSharedRef<SWidget> FConstructionScriptEditorSummoner::CreateTabBody(const FWorkflowTabSpawnInfo& Info) const
+{
+	TSharedPtr<FBlueprintEditor> BlueprintEditorPtr = StaticCastSharedPtr<FBlueprintEditor>(HostingApp.Pin());
+
+	return BlueprintEditorPtr->GetSCSEditor();
+}
+
+FSCSViewportSummoner::FSCSViewportSummoner(TSharedPtr<class FAssetEditorToolkit> InHostingApp) : FWorkflowTabFactory(FBlueprintEditorTabs::SCSViewportID, InHostingApp)
+{
+	TabLabel = LOCTEXT("SCSViewportTabLabel", "Viewport");
+	TabIcon = FSlateIcon(FEditorStyle::GetStyleSetName(), "LevelEditor.Tabs.Viewports");
+
+	bIsSingleton = true;
+
+	ViewMenuDescription = LOCTEXT("SCSViewportView", "Viewport");
+	ViewMenuTooltip = LOCTEXT("SCSViewportView_ToolTip", "Show the viewport view");
+}
+
+TSharedRef<SWidget> FSCSViewportSummoner::CreateTabBody(const FWorkflowTabSpawnInfo& Info) const
+{
+	TSharedPtr<FBlueprintEditor> BlueprintEditorPtr = StaticCastSharedPtr<FBlueprintEditor>(HostingApp.Pin());
+
+	TSharedPtr<SWidget> Result;
+	if (BlueprintEditorPtr->CanAccessComponentsMode())
+	{
+		Result = BlueprintEditorPtr->GetSCSViewport();
+	}
+
+	if (Result.IsValid())
+	{
+		return Result.ToSharedRef();
+	}
+	else
+	{
+		return SNew(SErrorText)
+			.BackgroundColor(FLinearColor::Transparent)
+			.ErrorText(LOCTEXT("SCSViewportView_Unavailable", "Viewport is not available for this Blueprint.").ToString());
+	}
+}
+
+FPaletteSummoner::FPaletteSummoner(TSharedPtr<class FAssetEditorToolkit> InHostingApp) : FWorkflowTabFactory(FBlueprintEditorTabs::PaletteID, InHostingApp)
+{
+	TabLabel = LOCTEXT("PaletteTabTitle", "Palette");
+	TabIcon = FSlateIcon(FEditorStyle::GetStyleSetName(), "Kismet.Tabs.Palette");
+
+	bIsSingleton = true;
+
+	ViewMenuDescription = LOCTEXT("PaletteView", "Palette");
+	ViewMenuTooltip = LOCTEXT("PaletteView_ToolTip", "Show palette of all functions and variables");
+}
+
+TSharedRef<SWidget> FPaletteSummoner::CreateTabBody(const FWorkflowTabSpawnInfo& Info) const
+{
+	TSharedPtr<FBlueprintEditor> BlueprintEditorPtr = StaticCastSharedPtr<FBlueprintEditor>(HostingApp.Pin());
+
+	return BlueprintEditorPtr->GetPalette();
+}
+
+FMyBlueprintSummoner::FMyBlueprintSummoner(TSharedPtr<class FAssetEditorToolkit> InHostingApp) : FWorkflowTabFactory(FBlueprintEditorTabs::MyBlueprintID, InHostingApp)
+{
+	TabLabel = LOCTEXT("MyBlueprintTabLabel", "My Blueprint");
+	TabIcon = FSlateIcon(FEditorStyle::GetStyleSetName(), "ClassIcon.BlueprintCore");
+
+	bIsSingleton = true;
+
+	ViewMenuDescription = LOCTEXT("MyBlueprintTabView", "My Blueprint");
+	ViewMenuTooltip = LOCTEXT("MyBlueprintTabView_ToolTip", "Show the my blueprint view");
+}
+
+TSharedRef<SWidget> FMyBlueprintSummoner::CreateTabBody(const FWorkflowTabSpawnInfo& Info) const
+{
+	TSharedPtr<FBlueprintEditor> BlueprintEditorPtr = StaticCastSharedPtr<FBlueprintEditor>(HostingApp.Pin());
+
+	return BlueprintEditorPtr->GetMyBlueprintWidget().ToSharedRef();
+}
+
+FCompilerResultsSummoner::FCompilerResultsSummoner(TSharedPtr<class FAssetEditorToolkit> InHostingApp) : FWorkflowTabFactory(FBlueprintEditorTabs::CompilerResultsID, InHostingApp)
+{
+	TabLabel = LOCTEXT("CompilerResultsTabTitle", "Compiler Results");
+	TabIcon = FSlateIcon(FEditorStyle::GetStyleSetName(), "Kismet.Tabs.CompilerResults");
+
+	bIsSingleton = true;
+
+	ViewMenuDescription = LOCTEXT("CompilerResultsView", "Compiler Results");
+	ViewMenuTooltip = LOCTEXT("CompilerResultsView_ToolTip", "Show compiler results of all functions and variables");
+}
+
+TSharedRef<SWidget> FCompilerResultsSummoner::CreateTabBody(const FWorkflowTabSpawnInfo& Info) const
+{
+	TSharedPtr<FBlueprintEditor> BlueprintEditorPtr = StaticCastSharedPtr<FBlueprintEditor>(HostingApp.Pin());
+
+	return BlueprintEditorPtr->GetCompilerResults();
+}
+
+FFindResultsSummoner::FFindResultsSummoner(TSharedPtr<class FAssetEditorToolkit> InHostingApp) : FWorkflowTabFactory(FBlueprintEditorTabs::FindResultsID, InHostingApp)
+{
+	TabLabel = LOCTEXT("FindResultsTabTitle", "Find Results");
+	TabIcon = FSlateIcon(FEditorStyle::GetStyleSetName(), "Kismet.Tabs.FindResults");
+
+	bIsSingleton = true;
+
+	ViewMenuDescription = LOCTEXT("FindResultsView", "Find Results");
+	ViewMenuTooltip = LOCTEXT("FindResultsView_ToolTip", "Show find results for searching in this blueprint or all blueprints");
+}
+
+TSharedRef<SWidget> FFindResultsSummoner::CreateTabBody(const FWorkflowTabSpawnInfo& Info) const
+{
+	TSharedPtr<FBlueprintEditor> BlueprintEditorPtr = StaticCastSharedPtr<FBlueprintEditor>(HostingApp.Pin());
+
+	return BlueprintEditorPtr->GetFindResults();
+}
+
+void FGraphTabHistory::EvokeHistory(TSharedPtr<FTabInfo> InTabInfo)
+{
+	FWorkflowTabSpawnInfo SpawnInfo;
+	SpawnInfo.Payload = Payload;
+	SpawnInfo.TabInfo = InTabInfo;
+
+	TSharedRef< SGraphEditor > GraphEditorRef = StaticCastSharedRef< SGraphEditor >(FactoryPtr.Pin()->CreateTabBody(SpawnInfo));
+	GraphEditor = GraphEditorRef;
+
+	FactoryPtr.Pin()->UpdateTab(InTabInfo->GetTab().Pin(), SpawnInfo, GraphEditorRef);
+}
+
+void FGraphTabHistory::SaveHistory()
+{
+	if (IsHistoryValid())
+	{
+		check(GraphEditor.IsValid());
+		GraphEditor.Pin()->GetViewLocation(SavedLocation, SavedZoomAmount);
+	}
+}
+
+void FGraphTabHistory::RestoreHistory()
+{
+	if (IsHistoryValid())
+	{
+		check(GraphEditor.IsValid());
+		GraphEditor.Pin()->SetViewLocation(SavedLocation, SavedZoomAmount);
+	}
+}
+
+
+#undef LOCTEXT_NAMESPACE
diff --git a/Engine/Source/Editor/Kismet/Private/BlueprintEditorTabFactories.h b/Engine/Source/Editor/Kismet/Private/BlueprintEditorTabFactories.h
--- a/Engine/Source/Editor/Kismet/Private/BlueprintEditorTabFactories.h
+++ b/Engine/Source/Editor/Kismet/Private/BlueprintEditorTabFactories.h
@@ -52,25 +52,9 @@ struct FLocalKismetCallbacks
 struct FDebugInfoSummoner : public FWorkflowTabFactory
 {
 public:
-	FDebugInfoSummoner(TSharedPtr<class FAssetEditorToolkit> InHostingApp)
-		: FWorkflowTabFactory(FBlueprintEditorTabs::DebugID, InHostingApp)
-	{
-		TabLabel = LOCTEXT("DebugTabTitle", "Debug");
-		TabIcon = FSlateIcon(FEditorStyle::GetStyleSetName(), "DebugTools.TabIcon");
-
-		EnableTabPadding();
-		bIsSingleton = true;
-
-		ViewMenuDescription = LOCTEXT("DebugView", "Debug");
-		ViewMenuTooltip = LOCTEXT("DebugView_ToolTip", "Shows the debugging view");
-	}
-
-	virtual TSharedRef<SWidget> CreateTabBody(const FWorkflowTabSpawnInfo& Info) const override
-	{
-		TSharedPtr<FBlueprintEditor> BlueprintEditorPtr = StaticCastSharedPtr<FBlueprintEditor>(HostingApp.Pin());
+	FDebugInfoSummoner(TSharedPtr<class FAssetEditorToolkit> InHostingApp);
 
-		return BlueprintEditorPtr->GetDebuggingView();
-	}
+	virtual TSharedRef<SWidget> CreateTabBody(const FWorkflowTabSpawnInfo& Info) const override;
 };
 
 /////////////////////////////////////////////////////
@@ -79,24 +63,9 @@ public:
 struct FDefaultsEditorSummoner : public FWorkflowTabFactory
 {
 public:
-	FDefaultsEditorSummoner(TSharedPtr<class FAssetEditorToolkit> InHostingApp)
-		: FWorkflowTabFactory(FBlueprintEditorTabs::DefaultEditorID, InHostingApp)
-	{
-		TabLabel = LOCTEXT("BlueprintDefaultsTabTitle", "Blueprint Defaults"); //@TODO: ANIMATION: GetDefaultEditorTitle(); !!!
-		TabIcon = FSlateIcon(FEditorStyle::GetStyleSetName(), "Kismet.Tabs.BlueprintDefaults");
-
-		bIsSingleton = true;
-
-		ViewMenuDescription = LOCTEXT("DefaultEditorView", "Defaults");
-		ViewMenuTooltip = LOCTEXT("DefaultEditorView_ToolTip", "Shows the default editor view");
-	}
+	FDefaultsEditorSummoner(TSharedPtr<class FAssetEditorToolkit> InHostingApp);
 
-	virtual TSharedRef<SWidget> CreateTabBody(const FWorkflowTabSpawnInfo& Info) const override
-	{
-		TSharedPtr<FBlueprintEditor> BlueprintEditorPtr = StaticCastSharedPtr<FBlueprintEditor>(HostingApp.Pin());
-
-		return BlueprintEditorPtr->GetDefaultEditor();
-	}
+	virtual TSharedRef<SWidget> CreateTabBody(const FWorkflowTabSpawnInfo& Info) const override;
 
 	virtual FText GetTabToolTipText(const FWorkflowTabSpawnInfo& Info) const override
 	{
@@ -122,35 +91,11 @@ public:
 
 	}
 
-	virtual void EvokeHistory(TSharedPtr<FTabInfo> InTabInfo) override
-	{
-		FWorkflowTabSpawnInfo SpawnInfo;
-		SpawnInfo.Payload = Payload;
-		SpawnInfo.TabInfo = InTabInfo;
-
-		TSharedRef< SGraphEditor > GraphEditorRef = StaticCastSharedRef< SGraphEditor >(FactoryPtr.Pin()->CreateTabBody(SpawnInfo));
-		GraphEditor = GraphEditorRef;
+	virtual void EvokeHistory(TSharedPtr<FTabInfo> InTabInfo) override;
 
-		FactoryPtr.Pin()->UpdateTab(InTabInfo->GetTab().Pin(), SpawnInfo, GraphEditorRef);
-	}
+	virtual void SaveHistory() override;
 
-	virtual void SaveHistory() override
-	{
-		if(IsHistoryValid())
-		{
-			check(GraphEditor.IsValid());
-			GraphEditor.Pin()->GetViewLocation(SavedLocation, SavedZoomAmount);
-		}
-	}
-
-	virtual void RestoreHistory() override
-	{
-		if(IsHistoryValid())
-		{
-			check(GraphEditor.IsValid());
-			GraphEditor.Pin()->SetViewLocation(SavedLocation, SavedZoomAmount);
-		}
-	}
+	virtual void RestoreHistory() override;
 
 private:
 	/** The graph editor represented by this history node. While this node is inactive, the graph editor is invalid */
@@ -169,41 +114,13 @@ struct FGraphEditorSummoner : public FDocumentTabFactoryForObjects<UEdGraph>
 public:
 	DECLARE_DELEGATE_RetVal_TwoParams(TSharedRef<SGraphEditor>, FOnCreateGraphEditorWidget, TSharedRef<FTabInfo>, UEdGraph*);
 public:
-	FGraphEditorSummoner(TSharedPtr<class FBlueprintEditor> InBlueprintEditorPtr, FOnCreateGraphEditorWidget CreateGraphEditorWidgetCallback)
-		: FDocumentTabFactoryForObjects<UEdGraph>(FBlueprintEditorTabs::GraphEditorID, InBlueprintEditorPtr)
-		, BlueprintEditorPtr(InBlueprintEditorPtr)
-		, OnCreateGraphEditorWidget(CreateGraphEditorWidgetCallback)
-	{
-	}
+	FGraphEditorSummoner(TSharedPtr<class FBlueprintEditor> InBlueprintEditorPtr, FOnCreateGraphEditorWidget CreateGraphEditorWidgetCallback);
 
-	virtual void OnTabActivated(TSharedPtr<SDockTab> Tab) const override
-	{
-		TSharedRef<SGraphEditor> GraphEditor = StaticCastSharedRef<SGraphEditor>(Tab->GetContent());
-		BlueprintEditorPtr.Pin()->OnGraphEditorFocused(GraphEditor);
-	}
+	virtual void OnTabActivated(TSharedPtr<SDockTab> Tab) const override;
 
-	virtual void OnTabRefreshed(TSharedPtr<SDockTab> Tab) const override
-	{
-		TSharedRef<SGraphEditor> GraphEditor = StaticCastSharedRef<SGraphEditor>(Tab->GetContent());
-		GraphEditor->NotifyGraphChanged();
-	}
+	virtual void OnTabRefreshed(TSharedPtr<SDockTab> Tab) const override;
 
-	virtual void SaveState(TSharedPtr<SDockTab> Tab, TSharedPtr<FTabPayload> Payload) const override
-	{
-		TSharedRef<SGraphEditor> GraphEditor = StaticCastSharedRef<SGraphEditor>(Tab->GetContent());
-
-		FVector2D ViewLocation;
-		float ZoomAmount;
-		GraphEditor->GetViewLocation(ViewLocation, ZoomAmount);
-
-		UEdGraph* Graph = FTabPayload_UObject::CastChecked<UEdGraph>(Payload);
-
-		if(BlueprintEditorPtr.Pin()->IsGraphInCurrentBlueprint(Graph))
-		{
-			// Don't save references to external graphs.
-			BlueprintEditorPtr.Pin()->GetBlueprintObj()->LastEditedDocuments.Add(FEditedDocumentInfo(Graph, ViewLocation, ZoomAmount));
-		}
-	}
+	virtual void SaveState(TSharedPtr<SDockTab> Tab, TSharedPtr<FTabPayload> Payload) const override;
 
 //	virtual void OnTabClosed(TSharedRef<SDockTab> Tab, TSharedPtr<FTabPayload> Payload) const override
 //	{
@@ -221,15 +138,9 @@ protected:
 		return OnCreateGraphEditorWidget.Execute(Info.TabInfo.ToSharedRef(), DocumentID);
 	}
 
-	virtual const FSlateBrush* GetTabIconForObject(const FWorkflowTabSpawnInfo& Info, UEdGraph* DocumentID) const override
-	{
-		return FBlueprintEditor::GetGlyphForGraph(DocumentID, false);
-	}
+	virtual const FSlateBrush* GetTabIconForObject(const FWorkflowTabSpawnInfo& Info, UEdGraph* DocumentID) const override;
 
-	virtual TSharedRef<FGenericTabHistory> CreateTabHistoryNode(TSharedPtr<FTabPayload> Payload) override
-	{
-		return MakeShareable(new FGraphTabHistory(SharedThis(this), Payload));
-	}
+	virtual TSharedRef<FGenericTabHistory> CreateTabHistoryNode(TSharedPtr<FTabPayload> Payload) override;
 
 protected:
 	TWeakPtr<class FBlueprintEditor> BlueprintEditorPtr;
@@ -242,42 +153,22 @@ protected:
 struct FTimelineEditorSummoner : public FDocumentTabFactoryForObjects<UTimelineTemplate>
 {
 public:
-	FTimelineEditorSummoner(TSharedPtr<class FBlueprintEditor> InBlueprintEditorPtr)
-		: FDocumentTabFactoryForObjects<UTimelineTemplate>(FBlueprintEditorTabs::TimelineEditorID, InBlueprintEditorPtr)
-		, BlueprintEditorPtr(InBlueprintEditorPtr)
-	{}
+	FTimelineEditorSummoner(TSharedPtr<class FBlueprintEditor> InBlueprintEditorPtr);
 protected:
-	virtual TAttribute<FText> ConstructTabNameForObject(UTimelineTemplate* DocumentID) const override
-	{
-		return TAttribute<FText>::Create(TAttribute<FText>::FGetter::CreateStatic<UObject*>(&FLocalKismetCallbacks::GetObjectName, DocumentID));
-	}
+	virtual TAttribute<FText> ConstructTabNameForObject(UTimelineTemplate* DocumentID) const override;
 
-	virtual TSharedRef<SWidget> CreateTabBodyForObject(const FWorkflowTabSpawnInfo& Info, UTimelineTemplate* DocumentID) const override
-	{
-		return SNew(STimelineEditor, BlueprintEditorPtr.Pin(), DocumentID);
-	}
+	virtual TSharedRef<SWidget> CreateTabBodyForObject(const FWorkflowTabSpawnInfo& Info, UTimelineTemplate* DocumentID) const override;
 
-	virtual const FSlateBrush* GetTabIconForObject(const FWorkflowTabSpawnInfo& Info, UTimelineTemplate* DocumentID) const override
-	{
-		return FEditorStyle::GetBrush("GraphEditor.Timeline_16x");
-	}
+	virtual const FSlateBrush* GetTabIconForObject(const FWorkflowTabSpawnInfo& Info, UTimelineTemplate* DocumentID) const override;
 
-	virtual void SaveState(TSharedPtr<SDockTab> Tab, TSharedPtr<FTabPayload> Payload) const override
-	{
-		UTimelineTemplate* Timeline = FTabPayload_UObject::CastChecked<UTimelineTemplate>(Payload);
-		BlueprintEditorPtr.Pin()->GetBlueprintObj()->LastEditedDocuments.Add(FEditedDocumentInfo(Timeline));
-	}
+	virtual void SaveState(TSharedPtr<SDockTab> Tab, TSharedPtr<FTabPayload> Payload) const override;
 
 // 	virtual void OnTabClosed(TSharedRef<SDockTab> Tab, TSharedPtr<FTabPayload> Payload) const override
 // 	{
 // 		BlueprintEditorPtr.Pin()->RequestSaveEditedObjectState();
 // 	}
 
-	virtual void OnTabRefreshed(TSharedPtr<SDockTab> Tab) const override
-	{
-		TSharedRef<STimelineEditor> TimelineEditor = StaticCastSharedRef<STimelineEditor>(Tab->GetContent());
-		TimelineEditor->OnTimelineChanged();
-	}
+	virtual void OnTabRefreshed(TSharedPtr<SDockTab> Tab) const override;
 
 	virtual FText GetTabToolTipText(const FWorkflowTabSpawnInfo& Info) const override
 	{
@@ -294,24 +185,9 @@ protected:
 struct FConstructionScriptEditorSummoner : public FWorkflowTabFactory
 {
 public:
-	FConstructionScriptEditorSummoner(TSharedPtr<class FAssetEditorToolkit> InHostingApp)
-		: FWorkflowTabFactory(FBlueprintEditorTabs::ConstructionScriptEditorID, InHostingApp)
-	{
-		TabLabel = LOCTEXT("ComponentsTabLabel", "Components");
-		TabIcon = FSlateIcon(FEditorStyle::GetStyleSetName(), "Kismet.Tabs.Components");
-
-		bIsSingleton = true;
+	FConstructionScriptEditorSummoner(TSharedPtr<class FAssetEditorToolkit> InHostingApp);
 
-		ViewMenuDescription = LOCTEXT("ComponentsView", "Components");
-		ViewMenuTooltip = LOCTEXT("ComponentsView_ToolTip", "Show the components view");
-	}
-
-	virtual TSharedRef<SWidget> CreateTabBody(const FWorkflowTabSpawnInfo& Info) const override
-	{
-		TSharedPtr<FBlueprintEditor> BlueprintEditorPtr = StaticCastSharedPtr<FBlueprintEditor>(HostingApp.Pin());
-
-		return BlueprintEditorPtr->GetSCSEditor();
-	}
+	virtual TSharedRef<SWidget> CreateTabBody(const FWorkflowTabSpawnInfo& Info) const override;
 
 	virtual FText GetTabToolTipText(const FWorkflowTabSpawnInfo& Info) const override
 	{
@@ -325,39 +201,9 @@ public:
 struct FSCSViewportSummoner : public FWorkflowTabFactory
 {
 public:
-	FSCSViewportSummoner(TSharedPtr<class FAssetEditorToolkit> InHostingApp)
-		: FWorkflowTabFactory(FBlueprintEditorTabs::SCSViewportID, InHostingApp)
-	{
-		TabLabel = LOCTEXT("SCSViewportTabLabel", "Viewport");
-		TabIcon = FSlateIcon(FEditorStyle::GetStyleSetName(), "LevelEditor.Tabs.Viewports");
+	FSCSViewportSummoner(TSharedPtr<class FAssetEditorToolkit> InHostingApp);
 
-		bIsSingleton = true;
-
-		ViewMenuDescription = LOCTEXT("SCSViewportView", "Viewport");
-		ViewMenuTooltip = LOCTEXT("SCSViewportView_ToolTip", "Show the viewport view");
-	}
-
-	virtual TSharedRef<SWidget> CreateTabBody(const FWorkflowTabSpawnInfo& Info) const override
-	{
-		TSharedPtr<FBlueprintEditor> BlueprintEditorPtr = StaticCastSharedPtr<FBlueprintEditor>(HostingApp.Pin());
-
-		TSharedPtr<SWidget> Result;
-		if (BlueprintEditorPtr->CanAccessComponentsMode())
-		{
-			Result = BlueprintEditorPtr->GetSCSViewport();
-		}
-
-		if (Result.IsValid())
-		{
-			return Result.ToSharedRef();
-		}
-		else
-		{
-			return SNew(SErrorText)
-					.BackgroundColor(FLinearColor::Transparent)
-					.ErrorText(LOCTEXT("SCSViewportView_Unavailable", "Viewport is not available for this Blueprint.").ToString());
-		}
-	}
+	virtual TSharedRef<SWidget> CreateTabBody(const FWorkflowTabSpawnInfo& Info) const override;
 };
 
 //////////////////////////////////////////////////////////////////////////
@@ -366,24 +212,9 @@ public:
 struct FPaletteSummoner : public FWorkflowTabFactory
 {
 public:
-	FPaletteSummoner(TSharedPtr<class FAssetEditorToolkit> InHostingApp)
-		: FWorkflowTabFactory(FBlueprintEditorTabs::PaletteID, InHostingApp)
-	{
-		TabLabel = LOCTEXT("PaletteTabTitle", "Palette");
-		TabIcon = FSlateIcon(FEditorStyle::GetStyleSetName(), "Kismet.Tabs.Palette");
-
-		bIsSingleton = true;
-
-		ViewMenuDescription = LOCTEXT("PaletteView", "Palette");
-		ViewMenuTooltip = LOCTEXT("PaletteView_ToolTip", "Show palette of all functions and variables");
-	}
-
-	virtual TSharedRef<SWidget> CreateTabBody(const FWorkflowTabSpawnInfo& Info) const override
-	{
-		TSharedPtr<FBlueprintEditor> BlueprintEditorPtr = StaticCastSharedPtr<FBlueprintEditor>(HostingApp.Pin());
+	FPaletteSummoner(TSharedPtr<class FAssetEditorToolkit> InHostingApp);
 
-		return BlueprintEditorPtr->GetPalette();
-	}
+	virtual TSharedRef<SWidget> CreateTabBody(const FWorkflowTabSpawnInfo& Info) const override;
 
 	virtual FText GetTabToolTipText(const FWorkflowTabSpawnInfo& Info) const override
 	{
@@ -397,24 +228,9 @@ public:
 struct FMyBlueprintSummoner : public FWorkflowTabFactory
 {
 public:
-	FMyBlueprintSummoner(TSharedPtr<class FAssetEditorToolkit> InHostingApp)
-		: FWorkflowTabFactory(FBlueprintEditorTabs::MyBlueprintID, InHostingApp)
-	{
-		TabLabel = LOCTEXT("MyBlueprintTabLabel", "My Blueprint");
-		TabIcon = FSlateIcon(FEditorStyle::GetStyleSetName(), "ClassIcon.BlueprintCore");
-
-		bIsSingleton = true;
+	FMyBlueprintSummoner(TSharedPtr<class FAssetEditorToolkit> InHostingApp);
 
-		ViewMenuDescription = LOCTEXT("MyBlueprintTabView", "My Blueprint");
-		ViewMenuTooltip = LOCTEXT("MyBlueprintTabView_ToolTip", "Show the my blueprint view");
-	}
-
-	virtual TSharedRef<SWidget> CreateTabBody(const FWorkflowTabSpawnInfo& Info) const override
-	{
-		TSharedPtr<FBlueprintEditor> BlueprintEditorPtr = StaticCastSharedPtr<FBlueprintEditor>(HostingApp.Pin());
-
-		return BlueprintEditorPtr->GetMyBlueprintWidget().ToSharedRef();
-	}
+	virtual TSharedRef<SWidget> CreateTabBody(const FWorkflowTabSpawnInfo& Info) const override;
 
 	virtual FText GetTabToolTipText(const FWorkflowTabSpawnInfo& Info) const override
 	{
@@ -428,24 +244,9 @@ public:
 struct FCompilerResultsSummoner : public FWorkflowTabFactory
 {
 public:
-	FCompilerResultsSummoner(TSharedPtr<class FAssetEditorToolkit> InHostingApp)
-		: FWorkflowTabFactory(FBlueprintEditorTabs::CompilerResultsID, InHostingApp)
-	{
-		TabLabel = LOCTEXT("CompilerResultsTabTitle", "Compiler Results");
-		TabIcon = FSlateIcon(FEditorStyle::GetStyleSetName(), "Kismet.Tabs.CompilerResults");
-
-		bIsSingleton = true;
-
-		ViewMenuDescription = LOCTEXT("CompilerResultsView", "Compiler Results");
-		ViewMenuTooltip = LOCTEXT("CompilerResultsView_ToolTip", "Show compiler results of all functions and variables");
-	}
-
-	virtual TSharedRef<SWidget> CreateTabBody(const FWorkflowTabSpawnInfo& Info) const override
-	{
-		TSharedPtr<FBlueprintEditor> BlueprintEditorPtr = StaticCastSharedPtr<FBlueprintEditor>(HostingApp.Pin());
+	FCompilerResultsSummoner(TSharedPtr<class FAssetEditorToolkit> InHostingApp);
 
-		return BlueprintEditorPtr->GetCompilerResults();
-	}
+	virtual TSharedRef<SWidget> CreateTabBody(const FWorkflowTabSpawnInfo& Info) const override;
 
 	virtual FText GetTabToolTipText(const FWorkflowTabSpawnInfo& Info) const override
 	{
@@ -459,24 +260,9 @@ public:
 struct FFindResultsSummoner : public FWorkflowTabFactory
 {
 public:
-	FFindResultsSummoner(TSharedPtr<class FAssetEditorToolkit> InHostingApp)
-		: FWorkflowTabFactory(FBlueprintEditorTabs::FindResultsID, InHostingApp)
-	{
-		TabLabel = LOCTEXT("FindResultsTabTitle", "Find Results");
-		TabIcon = FSlateIcon(FEditorStyle::GetStyleSetName(), "Kismet.Tabs.FindResults");
-
-		bIsSingleton = true;
+	FFindResultsSummoner(TSharedPtr<class FAssetEditorToolkit> InHostingApp);
 
-		ViewMenuDescription = LOCTEXT("FindResultsView", "Find Results");
-		ViewMenuTooltip = LOCTEXT("FindResultsView_ToolTip", "Show find results for searching in this blueprint or all blueprints");
-	}
-
-	virtual TSharedRef<SWidget> CreateTabBody(const FWorkflowTabSpawnInfo& Info) const override
-	{
-		TSharedPtr<FBlueprintEditor> BlueprintEditorPtr = StaticCastSharedPtr<FBlueprintEditor>(HostingApp.Pin());
-
-		return BlueprintEditorPtr->GetFindResults();
-	}
+	virtual TSharedRef<SWidget> CreateTabBody(const FWorkflowTabSpawnInfo& Info) const override;
 
 	virtual FText GetTabToolTipText(const FWorkflowTabSpawnInfo& Info) const override
 	{
diff --git a/Engine/Source/Editor/Kismet/Private/FindInBlueprints.cpp b/Engine/Source/Editor/Kismet/Private/FindInBlueprints.cpp
--- a/Engine/Source/Editor/Kismet/Private/FindInBlueprints.cpp
+++ b/Engine/Source/Editor/Kismet/Private/FindInBlueprints.cpp
@@ -9,6 +9,8 @@
 #include "AssetRegistryModule.h"
 #include "Json.h"
 #include "IDocumentation.h"
+#include "SSearchBox.h"
+#include "GenericCommands.h"
 
 #define LOCTEXT_NAMESPACE "FindInBlueprints"
 
diff --git a/Engine/Source/Editor/Kismet/Private/SMyBlueprint.cpp b/Engine/Source/Editor/Kismet/Private/SMyBlueprint.cpp
--- a/Engine/Source/Editor/Kismet/Private/SMyBlueprint.cpp
+++ b/Engine/Source/Editor/Kismet/Private/SMyBlueprint.cpp
@@ -36,6 +36,8 @@
 #include "IDocumentation.h"
 #include "Editor/UnrealEd/Public/SourceCodeNavigation.h"
 #include "EditorCategoryUtils.h"
+#include "SSearchBox.h"
+#include "GenericCommands.h"
 
 #define LOCTEXT_NAMESPACE "MyBlueprint"
 
diff --git a/Engine/Source/Editor/Kismet/Private/SSCSEditor.cpp b/Engine/Source/Editor/Kismet/Private/SSCSEditor.cpp
--- a/Engine/Source/Editor/Kismet/Private/SSCSEditor.cpp
+++ b/Engine/Source/Editor/Kismet/Private/SSCSEditor.cpp
@@ -26,6 +26,8 @@
 #include "Kismet2NameValidators.h"
 #include "UnrealExporter.h"
 #include "TutorialMetaData.h"
+#include "SInlineEditableTextBlock.h"
+#include "GenericCommands.h"
 
 #define LOCTEXT_NAMESPACE "SSCSEditor"
 
diff --git a/Engine/Source/Editor/Kismet/Private/STimelineEditor.cpp b/Engine/Source/Editor/Kismet/Private/STimelineEditor.cpp
--- a/Engine/Source/Editor/Kismet/Private/STimelineEditor.cpp
+++ b/Engine/Source/Editor/Kismet/Private/STimelineEditor.cpp
@@ -13,6 +13,10 @@
 #include "AssetRegistryModule.h"
 #include "AssetToolsModule.h"
 #include "DlgPickAssetPath.h"
+#include "SInlineEditableTextBlock.h"
+#include "GenericCommands.h"
+#include "SNotificationList.h"
+#include "NotificationManager.h"
 
 #define LOCTEXT_NAMESPACE "STimelineEditor"
 
diff --git a/Engine/Source/Editor/Kismet/Public/WorkflowOrientedApp/WorkflowTabFactory.h b/Engine/Source/Editor/Kismet/Public/WorkflowOrientedApp/WorkflowTabFactory.h
--- a/Engine/Source/Editor/Kismet/Public/WorkflowOrientedApp/WorkflowTabFactory.h
+++ b/Engine/Source/Editor/Kismet/Public/WorkflowOrientedApp/WorkflowTabFactory.h
@@ -58,7 +58,7 @@ class KISMET_API FWorkflowTabFactory : public TSharedFromThis<FWorkflowTabFactor
 {
 protected:
 	FName TabIdentifier;
-	ETabRole::Type TabRole;
+	ETabRole TabRole;
 
 	FText TabLabel;
 	FSlateIcon TabIcon;
@@ -163,9 +163,5 @@ public:
 	virtual TSharedRef<struct FGenericTabHistory> CreateTabHistoryNode(TSharedPtr<FTabPayload> Payload);
 
 protected:
-	FDocumentTabFactory(FName InIdentifier, TSharedPtr<class FAssetEditorToolkit> InHostingApp)
-		: FWorkflowTabFactory(InIdentifier, InHostingApp)
-	{
-		TabRole = ETabRole::DocumentTab;
-	}
+	FDocumentTabFactory(FName InIdentifier, TSharedPtr<class FAssetEditorToolkit> InHostingApp);
 };
diff --git a/Engine/Source/Editor/LandscapeEditor/Private/LandscapeEditorDetailWidgets.cpp b/Engine/Source/Editor/LandscapeEditor/Private/LandscapeEditorDetailWidgets.cpp
--- a/Engine/Source/Editor/LandscapeEditor/Private/LandscapeEditorDetailWidgets.cpp
+++ b/Engine/Source/Editor/LandscapeEditor/Private/LandscapeEditorDetailWidgets.cpp
@@ -2,6 +2,7 @@
 
 #include "LandscapeEditorPrivatePCH.h"
 #include "LandscapeEditorDetailWidgets.h"
+#include "SToolBarButtonBlock.h"
 
 
 // Based on a stripped-down FToolBarComboButtonBlock
diff --git a/Engine/Source/Editor/Layers/Private/SLayerBrowser.h b/Engine/Source/Editor/Layers/Private/SLayerBrowser.h
--- a/Engine/Source/Editor/Layers/Private/SLayerBrowser.h
+++ b/Engine/Source/Editor/Layers/Private/SLayerBrowser.h
@@ -4,8 +4,6 @@
 #include "ModuleManager.h"
 #include "Editor/SceneOutliner/Public/SceneOutlinerModule.h"
 
-#define LOCTEXT_NAMESPACE "LayerBrowser"
-
 typedef TTextFilter< const TSharedPtr< FLayerViewModel >& > LayerTextFilter;
 
 namespace ELayerBrowserMode
@@ -44,147 +42,7 @@ public:
 	 * @param	InArgs				Declaration used by the SNew() macro to construct this widget
 	 * @param	InViewModel			The UI logic not specific to slate
 	 */
-	void Construct( const FArguments& InArgs )
-	{
-		//@todo Create a proper ViewModel [7/27/2012 Justin.Sargent]
-		Mode = ELayerBrowserMode::Layers;
-
-		LayerCollectionViewModel = FLayerCollectionViewModel::Create( GEditor->Layers.ToSharedRef(), GEditor );
-		SelectedLayersFilter = MakeShareable( new FActorsAssignedToSpecificLayersFilter() );
-		SelectedLayerViewModel = FLayerViewModel::Create( NULL, GEditor->Layers.ToSharedRef(), GEditor ); //We'll set the datasource for this viewmodel later
-		
-		SearchBoxLayerFilter = MakeShareable( new LayerTextFilter( LayerTextFilter::FItemToStringArray::CreateSP( this, &SLayerBrowser::TransformLayerToString ) ) );
-		
-		LayerCollectionViewModel->AddFilter( SearchBoxLayerFilter.ToSharedRef() );
-		LayerCollectionViewModel->OnLayersChanged().AddSP( this, &SLayerBrowser::OnLayersChanged );
-		LayerCollectionViewModel->OnSelectionChanged().AddSP( this, &SLayerBrowser::UpdateSelectedLayer );
-		LayerCollectionViewModel->OnRenameRequested().AddSP( this, &SLayerBrowser::OnRenameRequested );
-
-		//////////////////////////////////////////////////////////////////////////
-		//	Layers View Section
-		SAssignNew( LayersSection, SBorder )
-		.Padding( 5 )
-		.BorderImage( FEditorStyle::GetBrush( "NoBrush" ) )
-		.Content()
-		[
-			SNew( SVerticalBox )
-			+SVerticalBox::Slot()
-			.AutoHeight()
-			[
-				SNew( SSearchBox )
-				.ToolTipText( LOCTEXT("FilterSearchToolTip", "Type here to search layers") )
-				.HintText( LOCTEXT( "FilterSearchHint", "Search Layers" ) )
-				.OnTextChanged( SearchBoxLayerFilter.Get(), &LayerTextFilter::SetRawFilterText )
-			]
-
-			+SVerticalBox::Slot()
-			.FillHeight( 1.0f )
-			[
-				SAssignNew( LayersView, SLayersView, LayerCollectionViewModel.ToSharedRef() )
-				.IsEnabled( FSlateApplication::Get().GetNormalExecutionAttribute() )
-				.ConstructContextMenu( FOnContextMenuOpening::CreateSP( this, &SLayerBrowser::ConstructLayerContextMenu ) )
-				.HighlightText( SearchBoxLayerFilter.Get(), &LayerTextFilter::GetRawFilterText )
-			]
-		];
-
-		//////////////////////////////////////////////////////////////////////////
-		//	Layer Contents Header
-		SAssignNew( LayerContentsHeader, SBorder )
-		.BorderImage( FEditorStyle::GetBrush("LayerBrowser.LayerContentsQuickbarBackground") )
-		.Visibility( TAttribute< EVisibility >( this, &SLayerBrowser::GetLayerContentsHeaderVisibility ) )
-		.Content()
-		[
-			SNew( SHorizontalBox )
-			+SHorizontalBox::Slot()
-			.AutoWidth()
-			.Padding( FMargin( 0, 0, 2, 0 ))
-			[
-				SAssignNew( ToggleModeButton, SButton )
-				.ContentPadding( FMargin( 2, 0, 2, 0 ) ) 
-				.ButtonStyle( FEditorStyle::Get(), "LayerBrowserButton" )
-				.OnClicked( this, &SLayerBrowser::ToggleLayerContents )
-				.ForegroundColor( FSlateColor::UseForeground() )
-				.VAlign( VAlign_Center )
-				.HAlign( HAlign_Center )
-				.Content()
-				[
-					SNew( SHorizontalBox )
-					+SHorizontalBox::Slot()
-					.AutoWidth()
-					.HAlign( HAlign_Center )
-					.VAlign( VAlign_Center )
-					.Padding(0, 1, 3, 1 )
-					[
-						SNew(SImage)
-						.Image( this, &SLayerBrowser::GetToggleModeButtonImageBrush )
-						.ColorAndOpacity( this, &SLayerBrowser::GetInvertedForegroundIfHovered )
-					]
-
-					+SHorizontalBox::Slot()
-					.AutoWidth()
-					.HAlign( HAlign_Center )
-					.VAlign( VAlign_Center )
-					[
-						SNew( STextBlock )
-						.Text( LOCTEXT("ContentsLabel", "See Contents") )
-						.Visibility( this, &SLayerBrowser::IsVisibleIfModeIs, ELayerBrowserMode::Layers )
-						.ColorAndOpacity( this, &SLayerBrowser::GetInvertedForegroundIfHovered )
-					]
-				]
-			]
-
-			+SHorizontalBox::Slot()
-			.AutoWidth()
-			.VAlign( VAlign_Center )
-			[
-				SNew( STextBlock )
-				.Text( this, &SLayerBrowser::GetLayerContentsHeaderText )
-				.Visibility( this, &SLayerBrowser::IsVisibleIfModeIs, ELayerBrowserMode::LayerContents )
-			]
-
-			+SHorizontalBox::Slot()
-			.HAlign( HAlign_Right )
-			.FillWidth( 1.0f )
-			[
-				SNew( SLayerStats, SelectedLayerViewModel.ToSharedRef() )
-			]
-		];
-
-
-		//////////////////////////////////////////////////////////////////////////
-		//	Layer Contents Section
-		FSceneOutlinerModule& SceneOutlinerModule = FModuleManager::LoadModuleChecked< FSceneOutlinerModule >( "SceneOutliner" );
-		FSceneOutlinerInitializationOptions InitOptions;
-		{
-			InitOptions.Mode = ESceneOutlinerMode::ActorBrowsing;
-
-			// We hide the header row to keep the UI compact.
-			InitOptions.bShowHeaderRow = false;
-			InitOptions.bShowParentTree = false;
-			InitOptions.CustomDelete = FCustomSceneOutlinerDeleteDelegate::CreateSP( this, &SLayerBrowser::RemoveActorsFromSelectedLayer );
-			InitOptions.CustomColumnFactory = FCreateSceneOutlinerColumnDelegate::CreateSP( this, &SLayerBrowser::CreateCustomLayerColumn );
-
-			InitOptions.Filters->Add( SelectedLayersFilter );
-		}
-
-		SAssignNew( LayerContentsSection, SBorder )
-		.Padding( 5 )
-		.BorderImage( FEditorStyle::GetBrush( "NoBrush" ) )
-		.Content()
-		[
-			SceneOutlinerModule.CreateSceneOutliner( InitOptions, FOnActorPicked() )
-		];
-
-
-		//////////////////////////////////////////////////////////////////////////
-		//	Layer Browser
-		ChildSlot
-		[
-			SAssignNew( ContentAreaBox, SVerticalBox )
-		];
-
-		SetupLayersMode();
-	}
+	void Construct( const FArguments& InArgs );
 
 
 protected:
@@ -333,10 +191,7 @@ private:
 		return ( Mode == ELayerBrowserMode::Layers ) ? FEditorStyle::GetBrush( "LayerBrowser.ExploreLayerContents" ) : FEditorStyle::GetBrush( "LayerBrowser.ReturnToLayersList" );
 	}
 
-	FText GetLayerContentsHeaderText() const
-	{
-		return FText::Format( LOCTEXT("SelectedContentsLabel", "{0} Contents"), FText::FromString(SelectedLayerViewModel->GetName()));
-	}
+	FText GetLayerContentsHeaderText() const;
 
 	/**	Returns the visibility of the See Contents label */
 	EVisibility IsVisibleIfModeIs( ELayerBrowserMode::Type DesiredMode ) const
@@ -507,5 +362,3 @@ private:
 	/** The layer view widget, displays all the layers in the level */
 	TSharedPtr< SLayersView > LayersView;
 };
-
-#undef LOCTEXT_NAMESPACE
diff --git a/Engine/Source/Editor/Layers/Private/SLayersCommandsMenu.h b/Engine/Source/Editor/Layers/Private/SLayersCommandsMenu.h
--- a/Engine/Source/Editor/Layers/Private/SLayersCommandsMenu.h
+++ b/Engine/Source/Editor/Layers/Private/SLayersCommandsMenu.h
@@ -3,8 +3,6 @@
 
 #include "LayersModule.h"
 
-#define LOCTEXT_NAMESPACE "LayersCommands"
-
 /**
  * The widget that represents a row in the LayersView's list view control.  Generates widgets for each column on demand.
  */
@@ -25,70 +23,10 @@ public:
 	 * @param	InArgs				Declaration used by the SNew() macro to construct this widget
 	 * @param	InImplementation	The UI logic not specific to slate
 	 */
-	void Construct( const FArguments& InArgs, const TSharedRef< FLayerCollectionViewModel > InViewModel/*, TSharedPtr< SLayersViewRow > InSingleSelectedRow */)
-	{
-		ViewModel = InViewModel;
-		const FLayersViewCommands& Commands = FLayersViewCommands::Get();
-
-		// Get all menu extenders for this context menu from the layers module
-		FLayersModule& LayersModule = FModuleManager::GetModuleChecked<FLayersModule>( TEXT("Layers") );
-		TArray<FLayersModule::FLayersMenuExtender> MenuExtenderDelegates = LayersModule.GetAllLayersMenuExtenders();
-
-		TArray<TSharedPtr<FExtender>> Extenders;
-		for (int32 i = 0; i < MenuExtenderDelegates.Num(); ++i)
-		{
-			if (MenuExtenderDelegates[i].IsBound())
-			{
-				Extenders.Add(MenuExtenderDelegates[i].Execute(ViewModel->GetCommandList()));
-			}
-		}
-		TSharedPtr<FExtender> MenuExtender = FExtender::Combine(Extenders);
-
-		// Build up the menu
-		FMenuBuilder MenuBuilder( InArgs._CloseWindowAfterMenuSelection, ViewModel->GetCommandList(), MenuExtender );
-		{ 
-			MenuBuilder.BeginSection("LayersCreate", LOCTEXT( "MenuHeader", "Layers" ) );
-			{
-				MenuBuilder.AddMenuEntry( Commands.CreateEmptyLayer );
-				MenuBuilder.AddMenuEntry( Commands.AddSelectedActorsToNewLayer );
-				MenuBuilder.AddMenuEntry( Commands.AddSelectedActorsToSelectedLayer );
-			}
-			MenuBuilder.EndSection();
-
-			MenuBuilder.BeginSection("LayersRemoveActors");
-			{
-				MenuBuilder.AddMenuEntry( Commands.RemoveSelectedActorsFromSelectedLayer );
-			}
-			MenuBuilder.EndSection();
-
-			MenuBuilder.AddMenuEntry( FGenericCommands::Get().Delete, "DeleteLayer", LOCTEXT( "DeleteLayer", "Delete Layer" ),  LOCTEXT( "DeleteLayerToolTip", "Removes all actors from the selected layers and then deletes the layers" ) );
-			MenuBuilder.AddMenuEntry( Commands.RequestRenameLayer );
-
-			MenuBuilder.BeginSection("LayersSelection", LOCTEXT( "SelectionMenuHeader", "Selection" ) );
-			{
-				MenuBuilder.AddMenuEntry( Commands.SelectActors );
-				MenuBuilder.AddMenuEntry( Commands.AppendActorsToSelection );
-				MenuBuilder.AddMenuEntry( Commands.DeselectActors );
-			}
-			MenuBuilder.EndSection();
-
-			MenuBuilder.BeginSection("LayersVisibility", LOCTEXT( "VisibilityMenuHeader", "Visibility" ) );
-			{
-				MenuBuilder.AddMenuEntry( Commands.MakeAllLayersVisible );
-			}
-			MenuBuilder.EndSection();
-		}
-
-		ChildSlot
-		[
-			MenuBuilder.MakeWidget()
-		];
-	}
+	void Construct( const FArguments& InArgs, const TSharedRef< FLayerCollectionViewModel > InViewModel/*, TSharedPtr< SLayersViewRow > InSingleSelectedRow */);
 
 private:
 	/** The UI logic of the LayersView that is not Slate specific */
 	TSharedPtr< FLayerCollectionViewModel > ViewModel;
 };
 
-
-#undef LOCTEXT_NAMESPACE
diff --git a/Engine/Source/Editor/Layers/Private/SLayersViewRow.cpp b/Engine/Source/Editor/Layers/Private/SLayersViewRow.cpp
new file mode 100644
--- /dev/null
+++ b/Engine/Source/Editor/Layers/Private/SLayersViewRow.cpp
@@ -0,0 +1,188 @@
+// Copyright 1998-2014 Epic Games, Inc. All Rights Reserved.
+
+#include "LayersPrivatePCH.h"
+#include "SLayersViewRow.h"
+#include "SInlineEditableTextBlock.h"
+
+#define LOCTEXT_NAMESPACE "LayersView"
+
+void SLayersViewRow::Construct(const FArguments& InArgs, TSharedRef< FLayerViewModel > InViewModel, TSharedRef< STableViewBase > InOwnerTableView)
+{
+	ViewModel = InViewModel;
+
+	HighlightText = InArgs._HighlightText;
+
+	SMultiColumnTableRow< TSharedPtr< FLayerViewModel > >::Construct(FSuperRowType::FArguments(), InOwnerTableView);
+}
+
+SLayersViewRow::~SLayersViewRow()
+{
+	ViewModel->OnRenamedRequest().RemoveSP(InlineTextBlock.Get(), &SInlineEditableTextBlock::EnterEditingMode);
+}
+
+TSharedRef< SWidget > SLayersViewRow::GenerateWidgetForColumn(const FName& ColumnID)
+{
+	TSharedPtr< SWidget > TableRowContent;
+
+	if (ColumnID == LayersView::ColumnID_LayerLabel)
+	{
+		TableRowContent =
+			SNew(SHorizontalBox)
+			+ SHorizontalBox::Slot()
+			.AutoWidth()
+			.Padding(0.0f, 1.0f, 3.0f, 1.0f)
+			[
+				SNew(SImage)
+				.Image(FEditorStyle::GetBrush(TEXT("Layer.Icon16x")))
+				.ColorAndOpacity(FSlateColor::UseForeground())
+			]
+
+		+ SHorizontalBox::Slot()
+			.FillWidth(1.0f)
+			[
+				SAssignNew(InlineTextBlock, SInlineEditableTextBlock)
+				.Font(FEditorStyle::GetFontStyle("LayersView.LayerNameFont"))
+				.Text(ViewModel.Get(), &FLayerViewModel::GetNameAsText)
+				.ColorAndOpacity(this, &SLayersViewRow::GetColorAndOpacity)
+				.HighlightText(HighlightText)
+				.ToolTipText(LOCTEXT("DoubleClickToolTip", "Double Click to Select All Actors"))
+				.OnVerifyTextChanged(this, &SLayersViewRow::OnRenameLayerTextChanged)
+				.OnTextCommitted(this, &SLayersViewRow::OnRenameLayerTextCommitted)
+				.IsSelected(this, &SLayersViewRow::IsSelectedExclusively)
+			]
+		;
+
+		ViewModel->OnRenamedRequest().AddSP(InlineTextBlock.Get(), &SInlineEditableTextBlock::EnterEditingMode);
+	}
+	else if (ColumnID == LayersView::ColumnID_Visibility)
+	{
+		TableRowContent =
+			SAssignNew(VisibilityButton, SButton)
+			.ContentPadding(0)
+			.ButtonStyle(FEditorStyle::Get(), "ToggleButton")
+			.OnClicked(this, &SLayersViewRow::OnToggleVisibility)
+			.ToolTipText(LOCTEXT("VisibilityButtonToolTip", "Toggle Layer Visibility"))
+			.ForegroundColor(FSlateColor::UseForeground())
+			.HAlign(HAlign_Center)
+			.VAlign(VAlign_Center)
+			.Content()
+			[
+				SNew(SImage)
+				.Image(this, &SLayersViewRow::GetVisibilityBrushForLayer)
+				.ColorAndOpacity(this, &SLayersViewRow::GetForegroundColorForButton)
+			]
+		;
+	}
+	else
+	{
+		checkf(false, TEXT("Unknown ColumnID provided to SLayersView"));
+	}
+
+	return TableRowContent.ToSharedRef();
+}
+
+void SLayersViewRow::OnRenameLayerTextCommitted(const FText& InText, ETextCommit::Type eInCommitType)
+{
+	if (!InText.IsEmpty())
+	{
+		ViewModel->RenameTo(*InText.ToString());
+	}
+}
+
+bool SLayersViewRow::OnRenameLayerTextChanged(const FText& NewText, FText& OutErrorMessage)
+{
+	FString OutMessage;
+	if (!ViewModel->CanRenameTo(*NewText.ToString(), OutMessage))
+	{
+		OutErrorMessage = FText::FromString(OutMessage);
+		return false;
+	}
+
+	return true;
+}
+
+void SLayersViewRow::OnDragLeave(const FDragDropEvent& DragDropEvent)
+{
+	TSharedPtr< FActorDragDropGraphEdOp > DragActorOp = DragDropEvent.GetOperationAs< FActorDragDropGraphEdOp >();
+	if (DragActorOp.IsValid())
+	{
+		DragActorOp->ResetToDefaultToolTip();
+	}
+}
+
+FReply SLayersViewRow::OnDragOver(const FGeometry& MyGeometry, const FDragDropEvent& DragDropEvent)
+{
+	TSharedPtr< FActorDragDropGraphEdOp > DragActorOp = DragDropEvent.GetOperationAs< FActorDragDropGraphEdOp >();
+	if (!DragActorOp.IsValid())
+	{
+		return FReply::Unhandled();
+	}
+
+	bool bCanAssign = false;
+	FText Message;
+	if (DragActorOp->Actors.Num() > 1)
+	{
+		bCanAssign = ViewModel->CanAssignActors(DragActorOp->Actors, OUT Message);
+	}
+	else
+	{
+		bCanAssign = ViewModel->CanAssignActor(DragActorOp->Actors[0], OUT Message);
+	}
+
+	if (bCanAssign)
+	{
+		DragActorOp->SetToolTip(FActorDragDropGraphEdOp::ToolTip_CompatibleGeneric, Message);
+	}
+	else
+	{
+		DragActorOp->SetToolTip(FActorDragDropGraphEdOp::ToolTip_IncompatibleGeneric, Message);
+	}
+
+	return FReply::Handled();
+}
+
+FReply SLayersViewRow::OnDrop(const FGeometry& MyGeometry, const FDragDropEvent& DragDropEvent)
+{
+	TSharedPtr< FActorDragDropGraphEdOp > DragActorOp = DragDropEvent.GetOperationAs< FActorDragDropGraphEdOp >();
+	if (!DragActorOp.IsValid())
+	{
+		return FReply::Unhandled();
+	}
+
+	ViewModel->AddActors(DragActorOp->Actors);
+
+	return FReply::Handled();
+}
+
+FSlateColor SLayersViewRow::GetForegroundColorForButton() const
+{
+	return (VisibilityButton.IsValid() && (VisibilityButton->IsHovered() || VisibilityButton->IsPressed())) ? FEditorStyle::GetSlateColor("InvertedForeground") : FSlateColor::UseForeground();
+}
+
+FSlateColor SLayersViewRow::GetColorAndOpacity() const
+{
+	if (!FSlateApplication::Get().IsDragDropping())
+	{
+		return FSlateColor::UseForeground();
+	}
+
+	bool bCanAcceptDrop = false;
+	TSharedPtr<FDragDropOperation> DragDropOp = FSlateApplication::Get().GetDragDroppingContent();
+
+	if (DragDropOp.IsValid() && DragDropOp->IsOfType<FActorDragDropGraphEdOp>())
+	{
+		TSharedPtr<FActorDragDropGraphEdOp> DragDropActorOp = StaticCastSharedPtr<FActorDragDropGraphEdOp>(DragDropOp);
+
+		FText Message;
+		bCanAcceptDrop = ViewModel->CanAssignActors(DragDropActorOp->Actors, OUT Message);
+	}
+
+	return (bCanAcceptDrop) ? FSlateColor::UseForeground() : FLinearColor(0.30f, 0.30f, 0.30f);
+}
+
+const FSlateBrush* SLayersViewRow::GetVisibilityBrushForLayer() const
+{
+	return (ViewModel->IsVisible()) ? FEditorStyle::GetBrush("Layer.VisibleIcon16x") : FEditorStyle::GetBrush("Layer.NotVisibleIcon16x");
+}
+
+#undef LOCTEXT_NAMESPACE
diff --git a/Engine/Source/Editor/Layers/Private/SLayersViewRow.h b/Engine/Source/Editor/Layers/Private/SLayersViewRow.h
--- a/Engine/Source/Editor/Layers/Private/SLayersViewRow.h
+++ b/Engine/Source/Editor/Layers/Private/SLayersViewRow.h
@@ -4,8 +4,6 @@
 #include "Editor/UnrealEd/Public/DragAndDrop/ActorDragDropGraphEdOp.h"
 #include "SLayerStats.h"
 
-#define LOCTEXT_NAMESPACE "LayersView"
-
 namespace LayersView
 {
 	/** IDs for list columns */
@@ -32,19 +30,9 @@ public:
 	 * @param	InViewModel			The layer the row widget is supposed to represent
 	 * @param	InOwnerTableView	The owner of the row widget
 	 */
-	void Construct( const FArguments& InArgs,  TSharedRef< FLayerViewModel > InViewModel, TSharedRef< STableViewBase > InOwnerTableView )
-	{
-		ViewModel = InViewModel;
-
-		HighlightText = InArgs._HighlightText;
+	void Construct( const FArguments& InArgs,  TSharedRef< FLayerViewModel > InViewModel, TSharedRef< STableViewBase > InOwnerTableView );
 
-		SMultiColumnTableRow< TSharedPtr< FLayerViewModel > >::Construct( FSuperRowType::FArguments(), InOwnerTableView );
-	}
-
-	~SLayersViewRow()
-	{
-		ViewModel->OnRenamedRequest().RemoveSP(InlineTextBlock.Get(), &SInlineEditableTextBlock::EnterEditingMode);
-	}
+	~SLayersViewRow();
 
 protected:
 
@@ -55,102 +43,20 @@ protected:
 	 *
 	 * @return a widget to represent the contents of a cell in this row of a TableView. 
 	 */
-	virtual TSharedRef< SWidget > GenerateWidgetForColumn( const FName& ColumnID ) override
-	{
-		TSharedPtr< SWidget > TableRowContent;
-
-		if( ColumnID == LayersView::ColumnID_LayerLabel )
-		{
-			TableRowContent =
-				SNew( SHorizontalBox )
-				+SHorizontalBox::Slot()
-				.AutoWidth()
-				.Padding( 0.0f, 1.0f, 3.0f, 1.0f )
-				[
-					SNew( SImage )
-					.Image( FEditorStyle::GetBrush( TEXT( "Layer.Icon16x" ) ) )
-					.ColorAndOpacity( FSlateColor::UseForeground() )
-				]
-
-				+SHorizontalBox::Slot()
-				.FillWidth( 1.0f )
-				[
-					SAssignNew(InlineTextBlock, SInlineEditableTextBlock)
-					.Font( FEditorStyle::GetFontStyle("LayersView.LayerNameFont") )
-					.Text( ViewModel.Get(), &FLayerViewModel::GetNameAsText )
-					.ColorAndOpacity( this, &SLayersViewRow::GetColorAndOpacity )
-					.HighlightText( HighlightText )
-					.ToolTipText( LOCTEXT("DoubleClickToolTip", "Double Click to Select All Actors") )
-					.OnVerifyTextChanged(this, &SLayersViewRow::OnRenameLayerTextChanged)
-					.OnTextCommitted(this, &SLayersViewRow::OnRenameLayerTextCommitted)
-					.IsSelected(this, &SLayersViewRow::IsSelectedExclusively)
-				]
-			;
-
-			ViewModel->OnRenamedRequest().AddSP(InlineTextBlock.Get(), &SInlineEditableTextBlock::EnterEditingMode);
-		}
-		else if( ColumnID == LayersView::ColumnID_Visibility )
-		{
-			TableRowContent = 
-				SAssignNew( VisibilityButton, SButton )
-				.ContentPadding( 0 )
-				.ButtonStyle( FEditorStyle::Get(), "ToggleButton" )
-				.OnClicked( this, &SLayersViewRow::OnToggleVisibility )
-				.ToolTipText( LOCTEXT("VisibilityButtonToolTip", "Toggle Layer Visibility") )
-				.ForegroundColor( FSlateColor::UseForeground() )
-				.HAlign( HAlign_Center )
-				.VAlign( VAlign_Center )
-				.Content()
-				[
-					SNew( SImage )
-					.Image( this, &SLayersViewRow::GetVisibilityBrushForLayer )
-					.ColorAndOpacity( this, &SLayersViewRow::GetForegroundColorForButton )
-				]
-			;
-		}
-		else
-		{
-			checkf(false, TEXT("Unknown ColumnID provided to SLayersView") );
-		}
-
-		return TableRowContent.ToSharedRef();
-	}
+	virtual TSharedRef< SWidget > GenerateWidgetForColumn( const FName& ColumnID ) override;
 
 	/** Callback when the SInlineEditableTextBlock is committed, to update the name of the layer this row represents. */
-	void OnRenameLayerTextCommitted(const FText& InText, ETextCommit::Type eInCommitType)
-	{
-		if (!InText.IsEmpty())
-		{
-			ViewModel->RenameTo(*InText.ToString());
-		}
-	}
+	void OnRenameLayerTextCommitted(const FText& InText, ETextCommit::Type eInCommitType);
 
 	/** Callback when the SInlineEditableTextBlock is changed, to check for error conditions. */
-	bool OnRenameLayerTextChanged(const FText& NewText, FText& OutErrorMessage)
-	{
-		FString OutMessage;
-		if ( !ViewModel->CanRenameTo( *NewText.ToString(), OutMessage ) )
-		{
-			OutErrorMessage = FText::FromString( OutMessage );
-			return false;
-		}
-
-		return true;
-	}
+	bool OnRenameLayerTextChanged(const FText& NewText, FText& OutErrorMessage);
 
 	/**
 	 * Called during drag and drop when the drag leaves a widget.
 	 *
 	 * @param DragDropEvent   The drag and drop event.
 	 */
-	virtual void OnDragLeave( const FDragDropEvent& DragDropEvent ) override
-	{
-		TSharedPtr< FActorDragDropGraphEdOp > DragActorOp = DragDropEvent.GetOperationAs< FActorDragDropGraphEdOp >();
-		if (DragActorOp.IsValid())
-		{
-			DragActorOp->ResetToDefaultToolTip();
-		}
-	}
+	virtual void OnDragLeave( const FDragDropEvent& DragDropEvent ) override;
 
 	/**
 	 * Called during drag and drop when the the mouse is being dragged over a widget.
@@ -160,36 +66,7 @@ protected:
 	 *
 	 * @return A reply that indicated whether this event was handled.
 	 */
-	virtual FReply OnDragOver( const FGeometry& MyGeometry, const FDragDropEvent& DragDropEvent ) override
-	{
-		TSharedPtr< FActorDragDropGraphEdOp > DragActorOp = DragDropEvent.GetOperationAs< FActorDragDropGraphEdOp >();
-		if (!DragActorOp.IsValid())
-		{
-			return FReply::Unhandled();
-		}
-
-		bool bCanAssign = false;
-		FText Message;
-		if( DragActorOp->Actors.Num() > 1 )
-		{
-			bCanAssign = ViewModel->CanAssignActors( DragActorOp->Actors, OUT Message );
-		}
-		else
-		{
-			bCanAssign = ViewModel->CanAssignActor( DragActorOp->Actors[ 0 ], OUT Message );
-		}
-
-		if ( bCanAssign )
-		{
-			DragActorOp->SetToolTip( FActorDragDropGraphEdOp::ToolTip_CompatibleGeneric, Message );
-		}
-		else
-		{
-			DragActorOp->SetToolTip( FActorDragDropGraphEdOp::ToolTip_IncompatibleGeneric, Message );
-		}
-
-		return FReply::Handled();
-	}
+	virtual FReply OnDragOver( const FGeometry& MyGeometry, const FDragDropEvent& DragDropEvent ) override;
 
 	/**
 	 * Called when the user is dropping something onto a widget; terminates drag and drop.
@@ -199,18 +76,7 @@ protected:
 	 *
 	 * @return A reply that indicated whether this event was handled.
 	 */
-	virtual FReply OnDrop( const FGeometry& MyGeometry, const FDragDropEvent& DragDropEvent ) override
-	{
-		TSharedPtr< FActorDragDropGraphEdOp > DragActorOp = DragDropEvent.GetOperationAs< FActorDragDropGraphEdOp >();
-		if (!DragActorOp.IsValid())
-		{
-			return FReply::Unhandled();
-		}
-
-		ViewModel->AddActors( DragActorOp->Actors );
-
-		return FReply::Handled();
-	}
+	virtual FReply OnDrop( const FGeometry& MyGeometry, const FDragDropEvent& DragDropEvent ) override;
 
 
 private:
@@ -221,10 +87,7 @@ private:
 	 *	@param	Button	The SButton Widget to get the Foreground Color for
 	 *	@return			The ForegroundColor
 	 */
-	FSlateColor GetForegroundColorForButton() const
-	{
-		return ( VisibilityButton.IsValid() && ( VisibilityButton->IsHovered() || VisibilityButton->IsPressed() ) ) ? FEditorStyle::GetSlateColor( "InvertedForeground" ) : FSlateColor::UseForeground();
-	}
+	FSlateColor GetForegroundColorForButton() const;
 
 	/**
 	 *	Returns the Color and Opacity for displaying the bound Layer's Name.
@@ -232,26 +95,7 @@ private:
 	 *
 	 *	@return	The SlateColor to render the Layer's name in
 	 */
-	FSlateColor GetColorAndOpacity() const
-	{
-		if ( !FSlateApplication::Get().IsDragDropping() )
-		{
-			return FSlateColor::UseForeground();
-		}
-
-		bool bCanAcceptDrop = false;
-		TSharedPtr<FDragDropOperation> DragDropOp = FSlateApplication::Get().GetDragDroppingContent();
-
-		if (DragDropOp.IsValid() && DragDropOp->IsOfType<FActorDragDropGraphEdOp>())
-		{
-			TSharedPtr<FActorDragDropGraphEdOp> DragDropActorOp = StaticCastSharedPtr<FActorDragDropGraphEdOp>(DragDropOp);
-
-			FText Message;
-			bCanAcceptDrop = ViewModel->CanAssignActors( DragDropActorOp->Actors, OUT Message );
-		}
-
-		return ( bCanAcceptDrop ) ? FSlateColor::UseForeground() : FLinearColor( 0.30f, 0.30f, 0.30f );
-	}
+	FSlateColor GetColorAndOpacity() const;
 
 	/**
 	 *	Called when the user clicks on the visibility icon for a layer's row widget
@@ -270,10 +114,7 @@ private:
 	 *
 	 *	@return	The SlateBrush representing the layer's visibility state
 	 */
-	const FSlateBrush* GetVisibilityBrushForLayer() const
-	{
-		return ( ViewModel->IsVisible() ) ? FEditorStyle::GetBrush( "Layer.VisibleIcon16x" ) : FEditorStyle::GetBrush( "Layer.NotVisibleIcon16x" );
-	}
+	const FSlateBrush* GetVisibilityBrushForLayer() const;
 
 
 private:
@@ -291,5 +132,3 @@ private:
 	TSharedPtr< SInlineEditableTextBlock > InlineTextBlock;
 };
 
-
-#undef LOCTEXT_NAMESPACE
diff --git a/Engine/Source/Editor/LevelEditor/Private/LevelEditorActions.cpp b/Engine/Source/Editor/LevelEditor/Private/LevelEditorActions.cpp
--- a/Engine/Source/Editor/LevelEditor/Private/LevelEditorActions.cpp
+++ b/Engine/Source/Editor/LevelEditor/Private/LevelEditorActions.cpp
@@ -50,6 +50,8 @@
 #include "EngineBuildSettings.h"
 #include "HotReloadInterface.h"
 #include "ISourceControlModule.h"
+#include "NotificationManager.h"
+#include "SNotificationList.h"
 
 DEFINE_LOG_CATEGORY_STATIC(LevelEditorActions, Log, All);
 
diff --git a/Engine/Source/Editor/LevelEditor/Private/LevelEditorToolBar.cpp b/Engine/Source/Editor/LevelEditor/Private/LevelEditorToolBar.cpp
--- a/Engine/Source/Editor/LevelEditor/Private/LevelEditorToolBar.cpp
+++ b/Engine/Source/Editor/LevelEditor/Private/LevelEditorToolBar.cpp
@@ -22,6 +22,7 @@
 #include "Editor/ClassViewer/Public/ClassViewerFilter.h"
 #include "KismetEditorUtilities.h"
 #include "ISourceControlModule.h"
+#include "SVolumeControl.h"
 
 namespace LevelEditorActionHelpers
 {
diff --git a/Engine/Source/Editor/LevelEditor/Private/SLevelEditorBuildAndSubmit.cpp b/Engine/Source/Editor/LevelEditor/Private/SLevelEditorBuildAndSubmit.cpp
--- a/Engine/Source/Editor/LevelEditor/Private/SLevelEditorBuildAndSubmit.cpp
+++ b/Engine/Source/Editor/LevelEditor/Private/SLevelEditorBuildAndSubmit.cpp
@@ -10,6 +10,8 @@
 #include "LightingBuildOptions.h"
 
 #include "AssetToolsModule.h"
+#include "SDockTab.h"
+#include "SExpandableArea.h"
 
 #define LOCTEXT_NAMESPACE "SLevelEditorBuildAndSubmit"
 
diff --git a/Engine/Source/Editor/LevelEditor/Private/SLevelViewport.cpp b/Engine/Source/Editor/LevelEditor/Private/SLevelViewport.cpp
--- a/Engine/Source/Editor/LevelEditor/Private/SLevelViewport.cpp
+++ b/Engine/Source/Editor/LevelEditor/Private/SLevelViewport.cpp
@@ -31,6 +31,9 @@
 #include "EditorViewportCommands.h"
 #include "Runtime/Engine/Classes/Engine/RendererSettings.h"
 #include "SScissorRectBox.h"
+#include "SDPIScaler.h"
+#include "SNotificationList.h"
+#include "NotificationManager.h"
 
 static const FName LevelEditorName("LevelEditor");
 
diff --git a/Engine/Source/Editor/LevelEditor/Private/SSurfaceProperties.cpp b/Engine/Source/Editor/LevelEditor/Private/SSurfaceProperties.cpp
--- a/Engine/Source/Editor/LevelEditor/Private/SSurfaceProperties.cpp
+++ b/Engine/Source/Editor/LevelEditor/Private/SSurfaceProperties.cpp
@@ -2,10 +2,12 @@
 
 #include "LevelEditor.h"
 #include "SSurfaceProperties.h"
-#include "Slate.h"
+#include "SlateBasics.h"
 #include "Editor/PropertyEditor/Public/PropertyEditorModule.h"
 #include "Editor/PropertyEditor/Public/IDetailsView.h"
 #include "SurfaceIterators.h"
+#include "SNumericEntryBox.h"
+#include "SHyperlink.h"
 
 #define LOCTEXT_NAMESPACE "SSurfaceProperties"
 
diff --git a/Engine/Source/Editor/Levels/Private/SLevelBrowser.cpp b/Engine/Source/Editor/Levels/Private/SLevelBrowser.cpp
new file mode 100644
--- /dev/null
+++ b/Engine/Source/Editor/Levels/Private/SLevelBrowser.cpp
@@ -0,0 +1,94 @@
+// Copyright 1998-2014 Epic Games, Inc. All Rights Reserved.
+
+#include "LevelsPrivatePCH.h"
+
+#include "SLevelBrowser.h"
+#include "SSearchBox.h"
+
+#define LOCTEXT_NAMESPACE "LevelBrowser"
+
+void SLevelBrowser::Construct(const FArguments& InArgs)
+{
+	Mode = ELevelBrowserMode::Levels;
+
+	LevelCollectionViewModel = FLevelCollectionViewModel::Create(GEditor);
+	SelectedLevelsFilter = MakeShareable(new FActorsAssignedToSpecificLevelsFilter());
+	SelectedLevelViewModel = FLevelViewModel::Create(NULL, NULL, GEditor); //We'll set the datasource for this viewmodel later
+
+	SearchBoxLevelFilter = MakeShareable(new LevelTextFilter(LevelTextFilter::FItemToStringArray::CreateSP(this, &SLevelBrowser::TransformLevelToString)));
+
+	LevelCollectionViewModel->AddFilter(SearchBoxLevelFilter.ToSharedRef());
+	LevelCollectionViewModel->OnLevelsChanged().AddSP(this, &SLevelBrowser::OnLevelsChanged);
+	LevelCollectionViewModel->OnSelectionChanged().AddSP(this, &SLevelBrowser::UpdateSelectedLevel);
+
+
+	//////////////////////////////////////////////////////////////////////////
+	//	Levels View Section
+	SAssignNew(LevelsSection, SBorder)
+		.Padding(5)
+		.BorderImage(FEditorStyle::GetBrush("NoBrush"))
+		.Content()
+		[
+			SNew(SVerticalBox)
+			+ SVerticalBox::Slot()
+			.AutoHeight()
+			[
+				SNew(SHorizontalBox)
+				//Shift Up Button
+				+ SHorizontalBox::Slot()
+				.VAlign(VAlign_Center)
+				.AutoWidth()
+				[
+					SNew(SButton)
+					.ToolTipText(LOCTEXT("ShiftUpHint", "Shift Streaming Level Up").ToString())
+					.ContentPadding(0)
+					.OnClicked(this, &SLevelBrowser::OnClickUp)
+					.IsEnabled(this, &SLevelBrowser::CanShiftSelection)
+					[
+						SNew(SImage).Image(FEditorStyle::GetBrush("Level.ArrowUp"))
+					]
+				]
+				//Shift Down Button
+				+ SHorizontalBox::Slot()
+					.VAlign(VAlign_Center)
+					.AutoWidth()
+					[
+						SNew(SButton)
+						.ToolTipText(LOCTEXT("ShiftDownHint", "Shift Streaming Level Down").ToString())
+						.ContentPadding(0)
+						.OnClicked(this, &SLevelBrowser::OnClickDown)
+						.IsEnabled(this, &SLevelBrowser::CanShiftSelection)
+						[
+							SNew(SImage).Image(FEditorStyle::GetBrush("Level.ArrowDown"))
+						]
+					]
+				+ SHorizontalBox::Slot()
+					.VAlign(VAlign_Center)
+					.FillWidth(1.f)
+					[
+						SNew(SSearchBox)
+						.ToolTipText(LOCTEXT("FilterSearchToolTip", "Type here to search Levels").ToString())
+						.HintText(LOCTEXT("FilterSearchHint", "Search Levels"))
+						.OnTextChanged(SearchBoxLevelFilter.Get(), &LevelTextFilter::SetRawFilterText)
+					]
+			]
+			+ SVerticalBox::Slot()
+				.FillHeight(1.0f)
+				[
+					SNew(SLevelsView, LevelCollectionViewModel.ToSharedRef())
+					.IsEnabled(FSlateApplication::Get().GetNormalExecutionAttribute())
+					.ConstructContextMenu(FOnContextMenuOpening::CreateSP(this, &SLevelBrowser::ConstructLevelContextMenu))
+					.HighlightText(SearchBoxLevelFilter.Get(), &LevelTextFilter::GetRawFilterText)
+				]
+		];
+
+	ChildSlot
+		[
+			SAssignNew(ContentAreaBox, SVerticalBox)
+		];
+
+
+	SetupLevelsMode();
+}
+
+#undef LOCTEXT_NAMESPACE
diff --git a/Engine/Source/Editor/Levels/Private/SLevelBrowser.h b/Engine/Source/Editor/Levels/Private/SLevelBrowser.h
--- a/Engine/Source/Editor/Levels/Private/SLevelBrowser.h
+++ b/Engine/Source/Editor/Levels/Private/SLevelBrowser.h
@@ -4,8 +4,6 @@
 #include "ModuleManager.h"
 #include "Editor/SceneOutliner/Public/SceneOutlinerModule.h"
 
-#define LOCTEXT_NAMESPACE "LevelBrowser"
-
 typedef TTextFilter< const TSharedPtr< FLevelViewModel >& > LevelTextFilter;
 
 namespace ELevelBrowserMode
@@ -40,89 +38,7 @@ public:
 	 * @param	InArgs				Declaration used by the SNew() macro to construct this widget
 	 * @param	InViewModel			The UI logic not specific to slate
 	 */
-	void Construct( const FArguments& InArgs )
-	{
-		Mode = ELevelBrowserMode::Levels;
-
-		LevelCollectionViewModel = FLevelCollectionViewModel::Create( GEditor );
-		SelectedLevelsFilter = MakeShareable( new FActorsAssignedToSpecificLevelsFilter() );
-		SelectedLevelViewModel = FLevelViewModel::Create( NULL, NULL, GEditor ); //We'll set the datasource for this viewmodel later
-		
-		SearchBoxLevelFilter = MakeShareable( new LevelTextFilter( LevelTextFilter::FItemToStringArray::CreateSP( this, &SLevelBrowser::TransformLevelToString ) ) );
-		
-		LevelCollectionViewModel->AddFilter( SearchBoxLevelFilter.ToSharedRef() );
-		LevelCollectionViewModel->OnLevelsChanged().AddSP( this, &SLevelBrowser::OnLevelsChanged );
-		LevelCollectionViewModel->OnSelectionChanged().AddSP( this, &SLevelBrowser::UpdateSelectedLevel );
-
-
-		//////////////////////////////////////////////////////////////////////////
-		//	Levels View Section
-		SAssignNew( LevelsSection, SBorder )
-		.Padding( 5 )
-		.BorderImage( FEditorStyle::GetBrush( "NoBrush" ) )
-		.Content()
-		[
-			SNew( SVerticalBox )
-			+SVerticalBox::Slot()
-			.AutoHeight()
-			[
-				SNew( SHorizontalBox )
-				//Shift Up Button
-				+SHorizontalBox::Slot()
-				.VAlign(VAlign_Center)
-				.AutoWidth()
-				[
-					SNew(SButton)
-					.ToolTipText( LOCTEXT("ShiftUpHint", "Shift Streaming Level Up").ToString() )
-					.ContentPadding( 0 )
-					.OnClicked(this, &SLevelBrowser::OnClickUp)
-					.IsEnabled(this, &SLevelBrowser::CanShiftSelection)
-					[
-						SNew(SImage) .Image(FEditorStyle::GetBrush("Level.ArrowUp"))
-					]
-				]
-				//Shift Down Button
-				+SHorizontalBox::Slot()
-				.VAlign(VAlign_Center)
-				.AutoWidth()
-				[
-					SNew(SButton)
-					.ToolTipText( LOCTEXT("ShiftDownHint", "Shift Streaming Level Down").ToString() )
-					.ContentPadding( 0 )
-					.OnClicked(this, &SLevelBrowser::OnClickDown)
-					.IsEnabled(this, &SLevelBrowser::CanShiftSelection)
-					[
-						SNew(SImage) .Image(FEditorStyle::GetBrush("Level.ArrowDown"))
-					]
-				]
-				+SHorizontalBox::Slot()
-				.VAlign(VAlign_Center)
-				.FillWidth(1.f)
-				[
-					SNew( SSearchBox )
-					.ToolTipText( LOCTEXT("FilterSearchToolTip", "Type here to search Levels").ToString() )
-					.HintText( LOCTEXT( "FilterSearchHint", "Search Levels" ) )
-					.OnTextChanged( SearchBoxLevelFilter.Get(), &LevelTextFilter::SetRawFilterText )
-				]
-			]
-			+SVerticalBox::Slot()
-			.FillHeight( 1.0f )
-			[
-				SNew( SLevelsView, LevelCollectionViewModel.ToSharedRef() )
-				.IsEnabled( FSlateApplication::Get().GetNormalExecutionAttribute() )
-				.ConstructContextMenu( FOnContextMenuOpening::CreateSP( this, &SLevelBrowser::ConstructLevelContextMenu ) )
-				.HighlightText( SearchBoxLevelFilter.Get(), &LevelTextFilter::GetRawFilterText )
-			]
-		];
-
-		ChildSlot
-		[
-			SAssignNew( ContentAreaBox, SVerticalBox )
-		];
-		
-
-		SetupLevelsMode();
-	}
+	void Construct( const FArguments& InArgs );
 
 
 protected:
@@ -310,5 +226,3 @@ private:
 	/**	The currently selected LevelViewModel if exactly one is selected, otherwise NULL */
 	TSharedPtr< FLevelViewModel > SelectedLevelViewModel;
 };
-
-#undef LOCTEXT_NAMESPACE
\ No newline at end of file
diff --git a/Engine/Source/Editor/Levels/Private/SLevelsViewRow.cpp b/Engine/Source/Editor/Levels/Private/SLevelsViewRow.cpp
new file mode 100644
--- /dev/null
+++ b/Engine/Source/Editor/Levels/Private/SLevelsViewRow.cpp
@@ -0,0 +1,279 @@
+// Copyright 1998-2014 Epic Games, Inc. All Rights Reserved.
+
+#include "LevelsPrivatePCH.h"
+#include "SLevelsViewRow.h"
+#include "SNumericEntryBox.h"
+
+#define LOCTEXT_NAMESPACE "LevelsView"
+
+
+BEGIN_SLATE_FUNCTION_BUILD_OPTIMIZATION
+TSharedRef< SWidget > SLevelsViewRow::GenerateWidgetForColumn(const FName& ColumnID)
+{
+	TSharedPtr< SWidget > TableRowContent;
+
+	if (ColumnID == LevelsView::ColumnID_LevelLabel)
+	{
+		TableRowContent =
+			SNew(SHorizontalBox)
+			+ SHorizontalBox::Slot()
+			.VAlign(VAlign_Center)
+			.AutoWidth()
+			[
+				SNew(STextBlock)
+				.Font(this, &SLevelsViewRow::GetFont)
+				.Text(ViewModel.Get(), &FLevelViewModel::GetDisplayName)
+				.ColorAndOpacity(this, &SLevelsViewRow::GetColorAndOpacity)
+				.HighlightText(HighlightText)
+				.ToolTipText(LOCTEXT("DoubleClickToolTip", "Double-Click to make this the current Level"))
+				.ShadowOffset(FVector2D(1.0f, 1.0f))
+			]
+		;
+	}
+	else if (ColumnID == LevelsView::ColumnID_ActorCount)
+	{
+		TableRowContent =
+			SNew(SHorizontalBox)
+			.Visibility(this, &SLevelsViewRow::IsActorColumnVisible)
+			+ SHorizontalBox::Slot()
+			.HAlign(HAlign_Right)
+			.VAlign(VAlign_Center)
+			[
+				SNew(STextBlock)
+				.Font(this, &SLevelsViewRow::GetFont)
+				.Text(ViewModel.Get(), &FLevelViewModel::GetActorCountString)
+				.ColorAndOpacity(this, &SLevelsViewRow::GetColorAndOpacity)
+				.HighlightText(HighlightText)
+				.ToolTipText(LOCTEXT("ActorCountToolTip", "The number of Actors in this level"))
+				.ShadowOffset(FVector2D(1.0f, 1.0f))
+			]
+		;
+	}
+	else if (ColumnID == LevelsView::ColumnID_LightmassSize)
+	{
+		TableRowContent =
+			SNew(SHorizontalBox)
+			.Visibility(this, &SLevelsViewRow::IsLightmassSizeColumnVisible)
+			+ SHorizontalBox::Slot()
+			.FillWidth(1.0f)
+			.HAlign(HAlign_Right)
+			.VAlign(VAlign_Center)
+			[
+				SNew(STextBlock)
+				.Font(this, &SLevelsViewRow::GetFont)
+				.Text(ViewModel.Get(), &FLevelViewModel::GetLightmassSizeString)
+				.ColorAndOpacity(this, &SLevelsViewRow::GetColorAndOpacity)
+				.HighlightText(HighlightText)
+				.ToolTipText(LOCTEXT("LightmassSizeToolTip", "The size of the lightmap data for this level"))
+				.ShadowOffset(FVector2D(1.0f, 1.0f))
+			]
+		;
+	}
+	else if (ColumnID == LevelsView::ColumnID_FileSize)
+	{
+		TableRowContent =
+			SNew(SHorizontalBox)
+			.Visibility(this, &SLevelsViewRow::IsFileSizeColumnVisible)
+			+ SHorizontalBox::Slot()
+			.HAlign(HAlign_Right)
+			.VAlign(VAlign_Center)
+			[
+				SNew(STextBlock)
+				.Font(this, &SLevelsViewRow::GetFont)
+				.Text(ViewModel.Get(), &FLevelViewModel::GetFileSizeString)
+				.ColorAndOpacity(this, &SLevelsViewRow::GetColorAndOpacity)
+				.HighlightText(HighlightText)
+				.ToolTipText(LOCTEXT("FileSizeToolTip", "The size of the file for this level"))
+				.ShadowOffset(FVector2D(1.0f, 1.0f))
+			]
+		;
+	}
+	else if (ColumnID == LevelsView::ColumnID_Visibility)
+	{
+		TableRowContent =
+			SAssignNew(VisibilityButton, SButton)
+			.ContentPadding(0)
+			.ButtonStyle(FEditorStyle::Get(), "ToggleButton")
+			.OnClicked(this, &SLevelsViewRow::OnToggleVisibility)
+			.ToolTipText(LOCTEXT("VisibilityButtonToolTip", "Toggle Level Visibility"))
+			.ForegroundColor(FSlateColor::UseForeground())
+			.HAlign(HAlign_Center)
+			.VAlign(VAlign_Center)
+			.Content()
+			[
+				SNew(SImage)
+				.Image(this, &SLevelsViewRow::GetVisibilityBrushForLevel)
+				.ColorAndOpacity(this, &SLevelsViewRow::GetForegroundColorForVisibilityButton)
+			]
+		;
+	}
+	else if (ColumnID == LevelsView::ColumnID_Lock)
+	{
+		TableRowContent =
+			SAssignNew(LockButton, SButton)
+			.ContentPadding(0)
+			.ButtonStyle(FEditorStyle::Get(), "ToggleButton")
+			.OnClicked(this, &SLevelsViewRow::OnToggleLock)
+			.ToolTipText(this, &SLevelsViewRow::GetLockToolTipForLevel)
+			.ForegroundColor(FSlateColor::UseForeground())
+			.HAlign(HAlign_Center)
+			.VAlign(VAlign_Center)
+			.Content()
+			[
+				SNew(SImage)
+				.Image(this, &SLevelsViewRow::GetLockBrushForLevel)
+				.ColorAndOpacity(this, &SLevelsViewRow::GetForegroundColorForLockButton)
+			]
+		;
+	}
+	else if (ColumnID == LevelsView::ColumnID_Color)
+	{
+		TableRowContent =
+			SAssignNew(ColorButton, SButton)
+			.ContentPadding(0)
+			.ButtonStyle(FEditorStyle::Get(), "ToggleButton")
+			.OnClicked(this, &SLevelsViewRow::OnChangeColor)
+			.ToolTipText(LOCTEXT("ColorButtonToolTip", "Change Level Color"))
+			.ForegroundColor(FSlateColor::UseForeground())
+			.HAlign(HAlign_Center)
+			.VAlign(VAlign_Center)
+			.Content()
+			[
+				SNew(SImage)
+				.Image(this, &SLevelsViewRow::GetColorBrushForLevel)
+				.ColorAndOpacity(TAttribute< FSlateColor >::Create(TAttribute< FSlateColor >::FGetter::CreateSP(this, &SLevelsViewRow::GetLevelColorAndOpacity)))
+			]
+		;
+	}
+	else if (ColumnID == LevelsView::ColumnID_SCCStatus)
+	{
+		TableRowContent =
+			SNew(SVerticalBox)
+			+ SVerticalBox::Slot()
+			.FillHeight(1.0f)
+			.HAlign(HAlign_Center)
+			.VAlign(VAlign_Center)
+			[
+				SNew(SHorizontalBox)
+				+ SHorizontalBox::Slot()
+				.FillWidth(1.0f)
+				.VAlign(VAlign_Center)
+				[
+					SNew(SImage)
+					.Image(this, &SLevelsViewRow::GetSCCStateImage)
+					.ToolTipText(this, &SLevelsViewRow::GetSCCStateTooltip)
+				]
+			]
+		;
+	}
+	else if (ColumnID == LevelsView::ColumnID_Save)
+	{
+		TableRowContent =
+			SAssignNew(SaveButton, SButton)
+			.ContentPadding(0)
+			.ButtonStyle(FEditorStyle::Get(), "ToggleButton")
+			.OnClicked(this, &SLevelsViewRow::OnSave)
+			.ToolTipText(LOCTEXT("SaveButtonToolTip", "Save Level"))
+			.ForegroundColor(FSlateColor::UseForeground())
+			.HAlign(HAlign_Center)
+			.VAlign(VAlign_Center)
+			.Content()
+			[
+				SNew(SImage)
+				.Image(this, &SLevelsViewRow::GetSaveBrushForLevel)
+				.ColorAndOpacity(TAttribute< FSlateColor >::Create(TAttribute< FSlateColor >::FGetter::CreateSP(this, &SLevelsViewRow::GetSaveButtonColorAndOpacity)))
+			]
+		;
+	}
+	else if (ColumnID == LevelsView::ColumnID_Kismet)
+	{
+		TableRowContent =
+			SAssignNew(KismetButton, SButton)
+			.ContentPadding(0)
+			.ButtonStyle(FEditorStyle::Get(), "ToggleButton")
+			.OnClicked(this, &SLevelsViewRow::OnOpenKismet)
+			.ToolTipText(LOCTEXT("KismetButtonToolTip", "Open Level Blueprint"))
+			.ForegroundColor(FSlateColor::UseForeground())
+			.HAlign(HAlign_Center)
+			.VAlign(VAlign_Center)
+			.Content()
+			[
+				SNew(SImage)
+				.Image(this, &SLevelsViewRow::GetKismetBrushForLevel)
+				.ColorAndOpacity(this, &SLevelsViewRow::GetForegroundColorForKismetButton)
+			]
+		;
+	}
+	else if (ColumnID == LevelsView::ColumnID_EditorOffset)
+	{
+		TableRowContent =
+			SNew(SHorizontalBox)
+			.Visibility(this, &SLevelsViewRow::IsEditorOffsetColumnVisible)
+			+ SHorizontalBox::Slot()
+			.Padding(FMargin(10, 0, 0, 0))
+			.HAlign(HAlign_Fill)
+			.VAlign(VAlign_Center)
+			.FillWidth(0.25f)
+			[
+				SNew(SButton)
+				.Text(LOCTEXT("EditLevelTransform", "Viewport Edit"))
+				.ToolTipText(LOCTEXT("EditLevelToolTip", "Edit level transform in viewport."))
+				.OnClicked(FOnClicked::CreateSP(this, &SLevelsViewRow::OnEditLevelClicked))
+				.IsEnabled(this, &SLevelsViewRow::LevelTransformAllowed)
+			]
+		+ SHorizontalBox::Slot()
+			.Padding(FMargin(10, 0, 0, 0))
+			.HAlign(HAlign_Fill)
+			.VAlign(VAlign_Center)
+			.FillWidth(0.55f)
+			[
+				SNew(SVectorInputBox)
+				.X(this, &SLevelsViewRow::GetLevelOffset, 0)
+				.Y(this, &SLevelsViewRow::GetLevelOffset, 1)
+				.Z(this, &SLevelsViewRow::GetLevelOffset, 2)
+				.bColorAxisLabels(true)
+				.OnXCommitted(this, &SLevelsViewRow::OnSetLevelOffset, 0)
+				.OnYCommitted(this, &SLevelsViewRow::OnSetLevelOffset, 1)
+				.OnZCommitted(this, &SLevelsViewRow::OnSetLevelOffset, 2)
+				.Font(FEditorStyle::GetFontStyle(TEXT("PropertyWindow.NormalFont"))) // Look the same as transform editing in the details panel
+				.IsEnabled(this, &SLevelsViewRow::LevelEditTextTransformAllowed)
+			]
+		+ SHorizontalBox::Slot()
+			.Padding(FMargin(10, 0, 0, 0))
+			.HAlign(HAlign_Fill)
+			.VAlign(VAlign_Center)
+			.FillWidth(0.2f)
+			[
+				SNew(SNumericEntryBox<float>)
+				.IsEnabled(this, &SLevelsViewRow::LevelEditTextTransformAllowed)
+				.Font(FEditorStyle::GetFontStyle(TEXT("PropertyWindow.NormalFont")))
+				.Delta(90.0f)
+				.AllowSpin(true)
+				.MinValue(0.0f)
+				.MaxValue(270.0f)
+				.MinSliderValue(0.0f)
+				.MaxSliderValue(270.0f)
+				.Value(this, &SLevelsViewRow::GetLevelRotation)
+				.OnValueChanged(this, &SLevelsViewRow::OnSetLevelRotation)
+				.OnValueCommitted(this, &SLevelsViewRow::OnCommitLevelRotation)
+				.OnBeginSliderMovement(this, &SLevelsViewRow::OnBeginLevelRotatonSlider)
+				.OnEndSliderMovement(this, &SLevelsViewRow::OnEndLevelRotatonSlider)
+				.LabelPadding(0)
+				.Label()
+				[
+					SNumericEntryBox<float>::BuildLabel(LOCTEXT("LevelRotation_Label", "Yaw"), FLinearColor::White, SNumericEntryBox<float>::BlueLabelBackgroundColor)
+				]
+			]
+		;
+	}
+	else
+	{
+		checkf(false, TEXT("Unknown ColumnID provided to SLevelsView"));
+	}
+
+	return TableRowContent.ToSharedRef();
+}
+
+END_SLATE_FUNCTION_BUILD_OPTIMIZATION
+
+#undef LOCTEXT_NAMESPACE
