BlueprintEdito - crash, more robust code

diff --git a/Engine/Source/Editor/Kismet/Private/BlueprintEditorModes.cpp b/Engine/Source/Editor/Kismet/Private/BlueprintEditorModes.cpp
--- a/Engine/Source/Editor/Kismet/Private/BlueprintEditorModes.cpp
+++ b/Engine/Source/Editor/Kismet/Private/BlueprintEditorModes.cpp
@@ -318,22 +318,29 @@ void FBlueprintComponentsApplicationMode::PreDeactivateMode()
 void FBlueprintComponentsApplicationMode::PostActivateMode()
 {
 	TSharedPtr<FBlueprintEditor> BP = MyBlueprintEditor.Pin();
-	BP->GetSCSEditor()->UpdateTree();
-	BP->EnableSCSPreview(true);
-	BP->UpdateSCSPreview();
-	BP->GetInspector()->EnableComponentDetailsCustomization(true);
-
-	// Reselect the cached components
-	TArray<TSharedPtr<FSCSEditorTreeNode>> Selection;
-	for( auto& Component : CachedComponentSelection )
-	{
-		BP->GetSCSEditor()->SCSTreeWidget->SetItemSelection(BP->GetSCSEditor()->GetNodeFromActorComponent(Component), true);
-	}
-
-	if (BP->GetSCSViewport()->GetIsSimulateEnabled())
+	if (BP.IsValid())
 	{
-		BP->GetSCSEditor()->SetEnabled(false);
-		BP->GetInspector()->SetEnabled(false);
+		auto SCSEditor = BP->GetSCSEditor();
+		SCSEditor->UpdateTree();
+		BP->EnableSCSPreview(true);
+		BP->UpdateSCSPreview();
+		BP->GetInspector()->EnableComponentDetailsCustomization(true);
+
+		// Reselect the cached components
+		TArray<TSharedPtr<FSCSEditorTreeNode>> Selection;
+		for (auto Component : CachedComponentSelection)
+		{
+			if (Component.IsValid())
+			{
+				SCSEditor->SCSTreeWidget->SetItemSelection(SCSEditor->GetNodeFromActorComponent(Component.Get()), true);
+			}
+		}
+
+		if (BP->GetSCSViewport()->GetIsSimulateEnabled())
+		{
+			SCSEditor->SetEnabled(false);
+			BP->GetInspector()->SetEnabled(false);
+		}
 	}
 
 	FApplicationMode::PostActivateMode();
diff --git a/Engine/Source/Editor/Kismet/Public/BlueprintEditorModes.h b/Engine/Source/Editor/Kismet/Public/BlueprintEditorModes.h
--- a/Engine/Source/Editor/Kismet/Public/BlueprintEditorModes.h
+++ b/Engine/Source/Editor/Kismet/Public/BlueprintEditorModes.h
@@ -96,7 +96,7 @@ protected:
 	// Set of spawnable tabs in this mode
 	FWorkflowAllowedTabSet BlueprintComponentsTabFactories;
 
-	TArray<UActorComponent*> CachedComponentSelection;
+	TArray<TWeakObjectPtr<UActorComponent>> CachedComponentSelection;
 };
 
 class KISMET_API FBlueprintInterfaceApplicationMode : public FApplicationMode
diff --git a/Engine/Source/Editor/UnrealEd/Private/Kismet2/BlueprintEditorUtils.cpp b/Engine/Source/Editor/UnrealEd/Private/Kismet2/BlueprintEditorUtils.cpp
--- a/Engine/Source/Editor/UnrealEd/Private/Kismet2/BlueprintEditorUtils.cpp
+++ b/Engine/Source/Editor/UnrealEd/Private/Kismet2/BlueprintEditorUtils.cpp
@@ -1666,7 +1666,8 @@ bool FBlueprintEditorUtils::ShouldRegenerateBlueprint(UBlueprint* Blueprint)
 // Helper function to get the blueprint that ultimately owns a node.
 UBlueprint* FBlueprintEditorUtils::FindBlueprintForNode(const UEdGraphNode* Node)
 {
-	return FindBlueprintForGraph(Node->GetGraph());
+	auto Graph = Node ? Cast<UEdGraph>(Node->GetOuter()) : NULL;
+	return FindBlueprintForGraph(Graph);
 }
 
 // Helper function to get the blueprint that ultimately owns a node.  Cannot fail.
