Adding PropertiesInOrderDisplayed accessor, Property Highlighting, ShowAllAdvancedProperties, Displayed Properties Changed delegate, Custom Layout disabling, property whitelisting to details view. Used by Diff and Merge tools.

Merge tool highlighting improved, now uses property paths.

diff --git a/Engine/Source/Developer/Merge/Merge.Build.cs b/Engine/Source/Developer/Merge/Merge.Build.cs
--- a/Engine/Source/Developer/Merge/Merge.Build.cs
+++ b/Engine/Source/Developer/Merge/Merge.Build.cs
@@ -6,47 +6,31 @@ namespace UnrealBuildTool.Rules
 	{
         public Merge(TargetInfo Target)
 		{
-			PublicIncludePaths.AddRange(
-				new string[] {
-				}
-				);
-
 			PrivateIncludePaths.AddRange(
 				new string[] {
 					"Developer/Merge/Private",
 				    "Kismet",
 				}
-				);
-
-			PublicDependencyModuleNames.AddRange(
-				new string[]
-				{
-					"Core",
-				}
-				);
+			);
 
 			PrivateDependencyModuleNames.AddRange(
 				new string[]
 				{
 				    "AssetTools",
+					"Core",
 				    "CoreUObject",
 				    "EditorStyle",
 				    "Engine", // needed so that we can clone blueprints...
 				    "GraphEditor",
 				    "InputCore",
 				    "Kismet",
+					"PropertyEditor",
 				    "Slate",
 				    "SlateCore",
 				    "SourceControl",
 				    "UnrealEd",
 				}
-				);
-
-			DynamicallyLoadedModuleNames.AddRange(
-				new string[]
-				{
-				}
-                );
+			);
 		}
 	}
 }
\ No newline at end of file
diff --git a/Engine/Source/Developer/Merge/Private/Merge.cpp b/Engine/Source/Developer/Merge/Private/Merge.cpp
--- a/Engine/Source/Developer/Merge/Private/Merge.cpp
+++ b/Engine/Source/Developer/Merge/Private/Merge.cpp
@@ -192,12 +192,14 @@ TSharedRef<SDockTab> FMerge::GenerateMergeWidget(const UBlueprint& Object, TShar
 
 		if (RemoteBlueprint && BaseBlueprint)
 		{
+			Object.bDuplicatingReadOnly = true;
 			FBlueprintMergeData Data(Editor 
 									, static_cast<const UBlueprint*>(StaticDuplicateObject(&Object, GetTransientPackage(), TEXT("None")))
 									, BaseBlueprint
 									, CurrentRevInfo
 									, RemoteBlueprint
-									, BaseRevInfo );
+									, BaseRevInfo);
+			Object.bDuplicatingReadOnly = false;
 
 			Contents = SNew(SBlueprintMerge, Data);
 		}
diff --git a/Engine/Source/Developer/Merge/Private/SBlueprintMerge.cpp b/Engine/Source/Developer/Merge/Private/SBlueprintMerge.cpp
--- a/Engine/Source/Developer/Merge/Private/SBlueprintMerge.cpp
+++ b/Engine/Source/Developer/Merge/Private/SBlueprintMerge.cpp
@@ -48,7 +48,7 @@ void SBlueprintMerge::Construct(const FArguments InArgs, const FBlueprintMergeDa
 	WriteBackup(*Data.BlueprintBase->GetOutermost(), BackupSubDir, TEXT("CommonBaseAsset") + FPackageName::GetAssetPackageExtension());
 	WriteBackup(*Data.BlueprintLocal->GetOutermost(), BackupSubDir, TEXT("LocalAsset") + FPackageName::GetAssetPackageExtension());
 
-	auto GraphView = SNew( SMergeGraphView, InData );
+	auto GraphView = SNew(SMergeGraphView, InData);
 	GraphControl.Widget = GraphView;
 	GraphControl.DiffControl = &GraphView.Get();
 
@@ -56,7 +56,7 @@ void SBlueprintMerge::Construct(const FArguments InArgs, const FBlueprintMergeDa
 	TreeControl.Widget = TreeView;
 	TreeControl.DiffControl = &TreeView.Get();
 
-	auto DetailsView = SNew( SMergeDetailsView, InData );
+	auto DetailsView = SNew(SMergeDetailsView, InData);
 	DetailsControl.Widget = DetailsView;
 	DetailsControl.DiffControl = &DetailsView.Get();
 
@@ -77,7 +77,7 @@ void SBlueprintMerge::Construct(const FArguments InArgs, const FBlueprintMergeDa
 	);
 	ToolbarBuilder.AddSeparator();
 	ToolbarBuilder.AddToolBarButton(
-		FUIAction(FExecuteAction::CreateSP(this, &SBlueprintMerge::OnAcceptResultClicked) )
+		FUIAction(FExecuteAction::CreateSP(this, &SBlueprintMerge::OnFinishMerge))
 		, NAME_None
 		, LOCTEXT("FinishMergeLabel", "Finish Merge")
 		, LOCTEXT("FinishMergeTooltip", "Complete the merge operation - saves the blueprint and resolves the conflict with the SCC provider")
@@ -139,7 +139,7 @@ bool SBlueprintMerge::HasPrevDiff() const
 	return CurrentDiffControl && CurrentDiffControl->HasPrevDifference();
 }
 
-void SBlueprintMerge::OnAcceptResultClicked()
+void SBlueprintMerge::OnFinishMerge()
 {
 	UPackage* Package = GetTargetBlueprint()->GetOutermost();
 	TArray<UPackage*> PackagesToSave;
diff --git a/Engine/Source/Developer/Merge/Private/SBlueprintMerge.h b/Engine/Source/Developer/Merge/Private/SBlueprintMerge.h
--- a/Engine/Source/Developer/Merge/Private/SBlueprintMerge.h
+++ b/Engine/Source/Developer/Merge/Private/SBlueprintMerge.h
@@ -25,7 +25,7 @@ private:
 	void PrevDiff();
 	bool HasNextDiff() const;
 	bool HasPrevDiff() const;
-	void OnAcceptResultClicked();
+	void OnFinishMerge();
 	void OnCancelClicked();
 	void OnModeChanged(FName NewMode);
 
diff --git a/Engine/Source/Developer/Merge/Private/SMergeDetailsView.cpp b/Engine/Source/Developer/Merge/Private/SMergeDetailsView.cpp
--- a/Engine/Source/Developer/Merge/Private/SMergeDetailsView.cpp
+++ b/Engine/Source/Developer/Merge/Private/SMergeDetailsView.cpp
@@ -3,28 +3,26 @@
 #include "MergePrivatePCH.h"
 
 #include "DetailsDiff.h"
+#include "ObjectEditorUtils.h"
 #include "SBlueprintDiff.h"
 #include "SMergeDetailsView.h"
 
-void SMergeDetailsView::Construct(const FArguments InArgs, const FBlueprintMergeData& InData)
+void SMergeDetailsView::Construct(const FArguments InArgs, const FBlueprintMergeData& InData )
 {
 	Data = InData;
 
 	const UObject* RemoteCDO = DiffUtils::GetCDO(InData.BlueprintRemote);
-	TMap< FName, const UProperty* > RemotePropertyMap = DiffUtils::GetProperties(RemoteCDO);
 	const UObject* BaseCDO = DiffUtils::GetCDO(InData.BlueprintBase);
-	TMap< FName, const UProperty* > BasePropertyMap = DiffUtils::GetProperties(BaseCDO);
 	const UObject* LocalCDO = DiffUtils::GetCDO(InData.BlueprintLocal);
-	TMap< FName, const UProperty* > LocalPropertyMap = DiffUtils::GetProperties(LocalCDO);
 
-	TArray<FName> RemoteIdenticalProperties, RemoteDifferingProperties;
-	DiffUtils::CompareUnrelatedObjects(BaseCDO, BasePropertyMap, RemoteCDO, RemotePropertyMap, RemoteIdenticalProperties, RemoteDifferingProperties);
-	TArray<FName> LocalIdenticalProperties, LocalDifferingProperties;
-	DiffUtils::CompareUnrelatedObjects(BaseCDO, BasePropertyMap, LocalCDO, LocalPropertyMap, LocalIdenticalProperties, LocalDifferingProperties);
+	TArray<FPropertySoftPath> RemoteDifferingProperties;
+	DiffUtils::CompareUnrelatedObjects(BaseCDO, RemoteCDO, RemoteDifferingProperties);
+	TArray<FPropertySoftPath> LocalDifferingProperties;
+	DiffUtils::CompareUnrelatedObjects(BaseCDO, LocalCDO, LocalDifferingProperties);
 
-	TSet<FName> RemoteDifferingPropertiesSet(RemoteDifferingProperties);
-	TSet<FName> LocalDifferingPropertiesSet(LocalDifferingProperties);
-	TSet<FName> BaseDifferingPropertiesSet(RemoteDifferingPropertiesSet.Union(LocalDifferingPropertiesSet));
+	FPropertySoftPathSet RemoteDifferingPropertiesSet(RemoteDifferingProperties);
+	FPropertySoftPathSet LocalDifferingPropertiesSet(LocalDifferingProperties);
+	FPropertySoftPathSet BaseDifferingPropertiesSet(RemoteDifferingPropertiesSet.Union(LocalDifferingPropertiesSet));
 	
 	/*	
 		DifferingProperties is an ordered list of all differing properties (properties added, removed or changed) in remote or local.
@@ -35,66 +33,63 @@ void SMergeDetailsView::Construct(const FArguments InArgs, const FBlueprintMerge
 			2. Iterate properties in remote, add any new properties to DifferingProperties
 			3. Iterate properties in local, add any new properties to DifferingProperties, if they have not already been added
 	*/
-	const auto AddPropertiesOrdered = [](UProperty const* Property, const TSet<FName>& DifferingProperties, TArray<FName>& ResultingProperties)
+	const auto AddPropertiesOrdered = [](FPropertySoftPath const& Property, const FPropertySoftPathSet& DifferingProperties, TArray<FPropertySoftPath>& ResultingProperties)
 	{
-		if ((!Property->HasAnyPropertyFlags(CPF_Parm) && Property->HasAnyPropertyFlags(CPF_Edit)))
+		// contains check here is O(n), so we're needlessly n^2:
+		if (!ResultingProperties.Contains(Property))
 		{
-			// contains check here is O(n), so we're needlessly n^2:
-			if (!ResultingProperties.Contains(Property->GetFName()))
+			if (DifferingProperties.Contains(Property))
 			{
-				if (DifferingProperties.Contains(Property->GetFName()))
-				{
-					ResultingProperties.Add(Property->GetFName());
-				}
+				ResultingProperties.Add(Property);
 			}
 		}
 	};
 
-	if( BaseCDO )
-	{
-		const UClass* Class = BaseCDO->GetClass();
-		for (TFieldIterator<UProperty> PropertyIt(Class, EFieldIteratorFlags::IncludeSuper); PropertyIt; ++PropertyIt)
-		{
-			AddPropertiesOrdered(*PropertyIt, BaseDifferingPropertiesSet, DifferingProperties);
-		}
-	}
-	if( RemoteCDO )
-	{
-		const UClass* Class = RemoteCDO->GetClass();
-		for (TFieldIterator<UProperty> PropertyIt(Class, EFieldIteratorFlags::IncludeSuper); PropertyIt; ++PropertyIt)
-		{
-			AddPropertiesOrdered(*PropertyIt, RemoteDifferingPropertiesSet, DifferingProperties);
-		}
-	}
-	if (LocalCDO)
-	{
-		const UClass* Class = LocalCDO->GetClass();
-		for (TFieldIterator<UProperty> PropertyIt(Class, EFieldIteratorFlags::IncludeSuper); PropertyIt; ++PropertyIt)
-		{
-			AddPropertiesOrdered(*PropertyIt, LocalDifferingPropertiesSet, DifferingProperties);
-		}
-	}
-	CurrentDifference = -1;
-
 	// EMergeParticipant::MERGE_PARTICIPANT_REMOTE
 	{
 		DetailsViews.Add(
-			FDetailsDiff(RemoteCDO, RemotePropertyMap, RemoteDifferingPropertiesSet)
+			FDetailsDiff(RemoteCDO, DiffUtils::ResolveAll( RemoteCDO, RemoteDifferingProperties), FDetailsDiff::FOnDisplayedPropertiesChanged() )
 		);
 	}
 	// EMergeParticipant::MERGE_PARTICIPANT_BASE
 	{
 		DetailsViews.Add(
-			FDetailsDiff(BaseCDO, BasePropertyMap, BaseDifferingPropertiesSet)
+			FDetailsDiff(BaseCDO, DiffUtils::ResolveAll( BaseCDO, BaseDifferingPropertiesSet.Array()), FDetailsDiff::FOnDisplayedPropertiesChanged())
 		);
 	}
 	// EMergeParticipant::MERGE_PARTICIPANT_LOCAL
 	{
 		DetailsViews.Add(
-			FDetailsDiff(LocalCDO, LocalPropertyMap, LocalDifferingPropertiesSet)
+			FDetailsDiff(LocalCDO, DiffUtils::ResolveAll( LocalCDO, LocalDifferingProperties), FDetailsDiff::FOnDisplayedPropertiesChanged())
 		);
 	}
 
+	TArray<FPropertySoftPath> RemoteVisibleProperties = GetRemoteDetails().GetDisplayedProperties();
+	TArray<FPropertySoftPath> BaseVisibleProperties = GetBaseDetails().GetDisplayedProperties();
+	TArray<FPropertySoftPath> LocalVisibleProperties = GetLocalDetails().GetDisplayedProperties();
+	if (BaseCDO)
+	{
+		for (const auto& Property : BaseVisibleProperties)
+		{
+			AddPropertiesOrdered(Property, BaseDifferingPropertiesSet, DifferingProperties);
+		}
+	}
+	if (RemoteCDO)
+	{
+		for (const auto& Property : RemoteVisibleProperties)
+		{
+			AddPropertiesOrdered(Property, RemoteDifferingPropertiesSet, DifferingProperties);
+		}
+	}
+	if (LocalCDO)
+	{
+		for (const auto& Property : LocalVisibleProperties)
+		{
+			AddPropertiesOrdered(Property, LocalDifferingPropertiesSet, DifferingProperties);
+		}
+	}
+	CurrentDifference = -1;
+
 	ChildSlot[
 		SNew(SSplitter)
 		+ SSplitter::Slot()
@@ -142,11 +137,9 @@ void SMergeDetailsView::PrevDiff()
 
 void SMergeDetailsView::HighlightCurrentDifference()
 {
-	const FName PropertyName = DifferingProperties[CurrentDifference];
-
 	for( auto& DetailDiff : DetailsViews )
 	{
-		DetailDiff.HighlightProperty(PropertyName);
+		DetailDiff.HighlightProperty(DifferingProperties[CurrentDifference]);
 	}
 }
 
diff --git a/Engine/Source/Developer/Merge/Private/SMergeDetailsView.h b/Engine/Source/Developer/Merge/Private/SMergeDetailsView.h
--- a/Engine/Source/Developer/Merge/Private/SMergeDetailsView.h
+++ b/Engine/Source/Developer/Merge/Private/SMergeDetailsView.h
@@ -15,7 +15,7 @@ public:
 	{}
 	SLATE_END_ARGS()
 
-	void Construct(const FArguments InArgs, const FBlueprintMergeData& InData);
+	void Construct(const FArguments InArgs, const FBlueprintMergeData& InData );
 private:
 	/** Implementation of IDiffControl: */
 	void NextDiff() override;
@@ -32,6 +32,6 @@ private:
 	TArray< FDetailsDiff > DetailsViews;
 
 	/** These have been duplicated from FCDODiffControl, opportunity to refactor exists: */
-	TArray< FName > DifferingProperties;
+	TArray< FPropertySoftPath > DifferingProperties;
 	int CurrentDifference;
 };
diff --git a/Engine/Source/Developer/Merge/Private/SMergeGraphView.cpp b/Engine/Source/Developer/Merge/Private/SMergeGraphView.cpp
--- a/Engine/Source/Developer/Merge/Private/SMergeGraphView.cpp
+++ b/Engine/Source/Developer/Merge/Private/SMergeGraphView.cpp
@@ -39,12 +39,13 @@ struct FMergeGraphRowEntry
 
 	// These lists contain shared ptrs because they are displayed by a
 	// SListView, and it currently does not support lists of values.
-	TArray< TSharedPtr<FDiffSingleResult> > RemoteDifferences;
-	TArray< TSharedPtr<FDiffSingleResult> > LocalDifferences;
+	TArray< TSharedPtr<FDiffSingleResult> > Differences;
 
 	bool bExistsInRemote;
 	bool bExistsInBase;
 	bool bExistsInLocal;
+	bool bDiffersInRemote;
+	bool bDiffersInLocal;
 };
 
 static UEdGraph* FindGraphByName(UBlueprint const& FromBlueprint, const FString& GraphName)
@@ -60,6 +61,15 @@ static UEdGraph* FindGraphByName(UBlueprint const& FromBlueprint, const FString&
 	return Ret;
 }
 
+static const UEdGraphNode* FindNodeInUnrelatedGraph(UEdGraph* Graph, UEdGraphNode* TargetNode)
+{
+	check(TargetNode->GetGraph() != Graph); // purpose of this function is to find a node in a different version of its parent graph
+
+	TArray< FGraphDiffControl::FNodeMatch > PriorMatches;
+	FGraphDiffControl::FNodeMatch NodeMatch = FGraphDiffControl::FindNodeMatch(Graph, TargetNode, PriorMatches);
+	return NodeMatch.NewNode;
+}
+
 static FDiffPanel InitializePanel(const FBlueprintRevPair& BlueprintRevPair)
 {
 	FDiffPanel Ret;
@@ -74,7 +84,6 @@ static FDiffPanel InitializePanel(const FBlueprintRevPair& BlueprintRevPair)
 	return Ret;
 }
 
-// @todo doc: can i avoid shared ptrs for individual entries?
 static TArray< TSharedPtr<FMergeGraphRowEntry> > GenerateDiffListItems(const FBlueprintRevPair& RemoteBlueprint, const FBlueprintRevPair& BaseBlueprint, const FBlueprintRevPair& LocalBlueprint)
 {
 	// Index all the graphs by name, we use the name of the graph as the 
@@ -132,18 +141,109 @@ static TArray< TSharedPtr<FMergeGraphRowEntry> > GenerateDiffListItems(const FBl
 					LocalDifferences = GenerateDifferences(*LocalGraph, BaseGraph);
 				}
 
+				// 'join' the local differences and remote differences by noting changes
+				// that affected the same common base:
+				{
+					for (const auto& RemoteDifference : RemoteDifferences)
+					{
+						bool bConflicted = false;
+						if (RemoteDifference.Diff)
+						for (const auto& LocalDifference : LocalDifferences)
+						{
+							if (RemoteDifference.Node1 == LocalDifference.Node1 )
+							{
+								bConflicted |=	RemoteDifference.Diff == EDiffType::NODE_REMOVED 
+												|| LocalDifference.Diff == EDiffType::NODE_REMOVED;
+								bConflicted |= RemoteDifference.Pin1 == LocalDifference.Pin1;
+							}
+						}
+
+						if (bConflicted)
+						{
+							// For now, we don't want to create a hard conflict for changes that don't effect runtime behavior:
+							if (RemoteDifference.Diff == EDiffType::NODE_MOVED ||
+								RemoteDifference.Diff == EDiffType::NODE_COMMENT)
+							{
+								continue;
+							}
+							// @todo doc: flag the conflict somehow...
+						}
+						else
+						{
+							// no conflict, we want to be able to automatically apply this remote change
+							// to the local revision:
+							switch (RemoteDifference.Diff)
+							{
+							case EDiffType::NODE_REMOVED:
+								{
+									// Find the corresponding object in the local graph:
+							
+									const UEdGraphNode* TargetNode = FindNodeInUnrelatedGraph(*LocalGraph, RemoteDifference.Node1);
+									check(TargetNode); // should have been a conflict!
+								}
+								break;
+							case EDiffType::NODE_ADDED:
+
+
+								break;
+							case EDiffType::PIN_LINKEDTO_NUM_DEC:
+								
+								break;
+							case EDiffType::PIN_LINKEDTO_NUM_INC:
+								break;
+							case EDiffType::PIN_DEFAULT_VALUE:
+								break;
+							case EDiffType::PIN_TYPE_CATEGORY:
+								break;
+							case EDiffType::PIN_TYPE_SUBCATEGORY:
+								break;
+							case EDiffType::PIN_TYPE_SUBCATEGORY_OBJECT:
+								break;
+							case EDiffType::PIN_TYPE_IS_ARRAY:
+								break;
+							case EDiffType::PIN_TYPE_IS_REF:
+								break;
+							case EDiffType::PIN_LINKEDTO_NODE:
+								break;
+							case EDiffType::NODE_MOVED:
+
+								break;
+							case EDiffType::TIMELINE_LENGTH:
+							case EDiffType::TIMELINE_AUTOPLAY:
+							case EDiffType::TIMELINE_LOOP:
+							case EDiffType::TIMELINE_NUM_TRACKS:
+							case EDiffType::TIMELINE_TRACK_MODIFIED:
+								check(false); // doesn't apply to blueprint
+								break;
+							case EDiffType::NODE_PIN_COUNT:
+								break;
+							case EDiffType::NODE_COMMENT:
+								break;
+							case EDiffType::NODE_PROPERTY:
+								break;
+							case EDiffType::NO_DIFFERENCE:
+								check(false);
+							}
+						}
+					}
+				}
+
 				bExistsInRemote = RemoteGraph != NULL;
 				bExistsInBase = BaseGraph != NULL;
 				bExistsInLocal = LocalGraph != NULL;
 			}
 
+			TArray<FDiffSingleResult> AllDifferences(RemoteDifferences);
+			AllDifferences.Append(LocalDifferences);
+
 			FMergeGraphRowEntry NewEntry = {
 				GraphName,
-				HackToShared(RemoteDifferences),
-				HackToShared(LocalDifferences),
+				HackToShared(AllDifferences),
 				bExistsInRemote,
 				bExistsInBase,
-				bExistsInLocal
+				bExistsInLocal,
+				RemoteDifferences.Num() != 0,
+				LocalDifferences.Num() != 0
 			};
 			Ret.Add(TSharedPtr<FMergeGraphRowEntry>(new FMergeGraphRowEntry(NewEntry)));
 		}
@@ -230,8 +330,8 @@ void SMergeGraphView::Construct(const FArguments InArgs, const FBlueprintMergeDa
 	}
 
 	DifferencesFromBase = GenerateDiffListItems(BlueprintsForDisplay[EMergeParticipant::MERGE_PARTICIPANT_REMOTE]
-		, BlueprintsForDisplay[EMergeParticipant::MERGE_PARTICIPANT_BASE]
-		, BlueprintsForDisplay[EMergeParticipant::MERGE_PARTICIPANT_LOCAL]);
+												, BlueprintsForDisplay[EMergeParticipant::MERGE_PARTICIPANT_BASE]
+												, BlueprintsForDisplay[EMergeParticipant::MERGE_PARTICIPANT_LOCAL]);
 
 	// This is the function we'll use to generate a row in the control that lists all the available graphs:
 	const auto RowGenerator = [](TSharedPtr<FMergeGraphRowEntry> ParamItem, const TSharedRef<STableViewBase>& OwnerTable) -> TSharedRef<ITableRow>
@@ -261,9 +361,9 @@ void SMergeGraphView::Construct(const FArguments InArgs, const FBlueprintMergeDa
 			return MergeWhite;
 		};
 
-		FLinearColor RemoteColor = ComputeColorForRevision(ParamItem->bExistsInBase, ParamItem->bExistsInRemote, ParamItem->RemoteDifferences.Num() != 0);
+		FLinearColor RemoteColor = ComputeColorForRevision(ParamItem->bExistsInBase, ParamItem->bExistsInRemote, ParamItem->bDiffersInRemote);
 		FLinearColor BaseColor = ComputeColorForRevision(ParamItem->bExistsInBase, ParamItem->bExistsInBase, false);
-		FLinearColor LocalColor = ComputeColorForRevision(ParamItem->bExistsInBase, ParamItem->bExistsInLocal, ParamItem->LocalDifferences.Num() != 0);
+		FLinearColor LocalColor = ComputeColorForRevision(ParamItem->bExistsInBase, ParamItem->bExistsInLocal, ParamItem->bDiffersInLocal);
 
 		const auto Box = [](bool bIsPresent, FLinearColor Color) -> SHorizontalBox::FSlot&
 		{
@@ -291,8 +391,8 @@ void SMergeGraphView::Construct(const FArguments InArgs, const FBlueprintMergeDa
 					.Text(ParamItem->GraphName)
 				]
 				+ Box(ParamItem->bExistsInRemote, RemoteColor)
-					+ Box(ParamItem->bExistsInBase, BaseColor)
-					+ Box(ParamItem->bExistsInLocal, LocalColor)
+				+ Box(ParamItem->bExistsInBase, BaseColor)
+				+ Box(ParamItem->bExistsInLocal, LocalColor)
 			];
 	};
 
@@ -416,11 +516,11 @@ void SMergeGraphView::OnGraphListSelectionChanged(TSharedPtr<FMergeGraphRowEntry
 
 	LockViews(DiffPanels, bViewsAreLocked);
 
-	DiffResultsListData = &(Item->LocalDifferences);
+	DiffResultsListData = &(Item->Differences);
 	DiffResultsWidget->SetContent(
 		SAssignNew(DiffResultList, SListView< TSharedPtr<FDiffSingleResult> >)
 		.ItemHeight(24)
-		.ListItemsSource(&Item->LocalDifferences)
+		.ListItemsSource(&Item->Differences)
 		.OnGenerateRow(SListView< TSharedPtr<FDiffSingleResult> >::FOnGenerateRow::CreateStatic(
 			[](TSharedPtr<FDiffSingleResult> ParamItem, const TSharedRef<STableViewBase>& OwnerTable) -> TSharedRef<ITableRow>
 			{
diff --git a/Engine/Source/Editor/Kismet/Private/DetailsDiff.cpp b/Engine/Source/Editor/Kismet/Private/DetailsDiff.cpp
--- a/Engine/Source/Editor/Kismet/Private/DetailsDiff.cpp
+++ b/Engine/Source/Editor/Kismet/Private/DetailsDiff.cpp
@@ -4,9 +4,12 @@
 #include "DetailsDiff.h"
 #include "IDetailsView.h"
 #include "PropertyEditorModule.h"
+//#include "SBlueprintDiff.h"
 
-FDetailsDiff::FDetailsDiff(const UObject* InObject, const TMap< FName, const UProperty* >& InPropertyMap, const TSet< FName >& InDifferingProperties)
-	: PropertyMap( InPropertyMap )
+FDetailsDiff::FDetailsDiff(const UObject* InObject, const TArray< FPropertyPath >& InDifferingProperties, FOnDisplayedPropertiesChanged InOnDisplayedPropertiesChanged )
+	: OnDisplayedPropertiesChanged( InOnDisplayedPropertiesChanged )
+	, DifferingProperties( InDifferingProperties )
+	, DisplayedObject( InObject )
 	, DetailsView()
 {
 	FDetailsViewArgs DetailsViewArgs;
@@ -14,19 +17,47 @@ FDetailsDiff::FDetailsDiff(const UObject* InObject, const TMap< FName, const UPr
 
 	FPropertyEditorModule& EditModule = FModuleManager::Get().GetModuleChecked<FPropertyEditorModule>("PropertyEditor");
 	DetailsView = EditModule.CreateDetailView(DetailsViewArgs);
+	DetailsView->SetDisableCustomDetailLayouts(true);
 	DetailsView->SetIsPropertyEditingEnabledDelegate(FIsPropertyEditingEnabled::CreateStatic([]{return false; }));
-	// This is a read only details view (see the property editing delegate), but it is not const correct:
+	// Forcing all advanced properties to be displayed for now, the logic to show changes made to advance properties
+	// conditionally is fragile and low priority for now:
+	DetailsView->ShowAllAdvancedProperties();
+	// This is a read only details view (see the property editing delegate), but const correctness at the details view
+	// level would stress the type system a little:
 	DetailsView->SetObject(const_cast<UObject*>(InObject));
-	DetailsView->UpdateDifferingProperties(InDifferingProperties);
+	DetailsView->SetOnDisplayedPropertiesChanged( ::FOnDisplayedPropertiesChanged::CreateRaw(this, &FDetailsDiff::HandlePropertiesChanged) );
+	DetailsView->UpdatePropertiesWhitelist( TSet<FPropertyPath>(InDifferingProperties) );
+
+	DifferingProperties = DetailsView->GetPropertiesInOrderDisplayed();
 }
 
-void FDetailsDiff::HighlightProperty(FName PropertyName)
+void FDetailsDiff::HighlightProperty(const FPropertySoftPath& PropertyName)
 {
-	const UProperty** Property = PropertyMap.Find(PropertyName);
-	DetailsView->HighlightProperty( Property ? *Property : NULL );
+	// resolve the property soft path:
+	const FPropertyPath ResolvedProperty = PropertyName.ResolvePath(DisplayedObject);
+	DetailsView->HighlightProperty(ResolvedProperty);
 }
 
 TSharedRef< SWidget > FDetailsDiff::DetailsWidget()
 {
 	return DetailsView.ToSharedRef();
 }
+
+void FDetailsDiff::HandlePropertiesChanged()
+{
+	if( OnDisplayedPropertiesChanged.IsBound() )
+	{
+		DifferingProperties = DetailsView->GetPropertiesInOrderDisplayed();
+		OnDisplayedPropertiesChanged.Execute();
+	}
+}
+
+TArray<FPropertySoftPath> FDetailsDiff::GetDisplayedProperties() const
+{
+	TArray<FPropertySoftPath> Ret;
+	for( const auto& Property : DifferingProperties )
+	{
+		Ret.Push( Property );
+	}
+	return Ret;
+}
diff --git a/Engine/Source/Editor/Kismet/Private/DiffUtils.cpp b/Engine/Source/Editor/Kismet/Private/DiffUtils.cpp
new file mode 100644
--- /dev/null
+++ b/Engine/Source/Editor/Kismet/Private/DiffUtils.cpp
@@ -0,0 +1,245 @@
+// Copyright 1998-2014 Epic Games, Inc. All Rights Reserved.
+
+#include "BlueprintEditorPrivatePCH.h"
+
+#include "DiffUtils.h"
+#include "EditorCategoryUtils.h"
+#include "ObjectEditorUtils.h"
+
+static const UProperty* Resolve( const UStruct* Class, FName PropertyName )
+{
+	for (TFieldIterator<UProperty> PropertyIt(Class); PropertyIt; ++PropertyIt)
+	{
+		if( PropertyIt->GetFName() == PropertyName )
+		{
+			return *PropertyIt;
+		}
+	}
+
+	return nullptr;
+}
+
+static FPropertySoftPathSet GetPropertyNameSet(const UObject* ForObj)
+{
+	return FPropertySoftPathSet(DiffUtils::GetVisiblePropertiesInOrderDeclared(ForObj));
+}
+
+FResolvedProperty FPropertySoftPath::Resolve(const UObject* Object) const
+{
+	// dig into the object, finding nested objects, etc:
+	const UProperty* Property = ::Resolve(Object->GetClass(), PropertyChain[0]);
+	// @todo: recurse for objects and structs:
+	return FResolvedProperty(Object, Property);
+}
+
+FPropertyPath FPropertySoftPath::ResolvePath(const UObject* Object) const
+{
+	const UStruct* ClassOverride = nullptr;
+	FPropertyPath Ret;
+	for( const auto& Property : PropertyChain )
+	{
+		const UProperty* ResolvedProperty = ::Resolve(ClassOverride ? ClassOverride : Object->GetClass(), Property);
+		if (const UObjectProperty* ObjectProperty = Cast<UObjectProperty>(ResolvedProperty))
+		{
+			// At the moment the cast of Object to void* just avoids an overzealous assert in ContainerPtrToValuePtr (for some of our
+			// properties, the outer is a UStruct, not a UClass). At some point I want to support diffing of individual array elements, 
+			// and that will necessitate the container be a void* anyway:
+			const UObject* const* BaseObject = reinterpret_cast<const UObject* const*>( ObjectProperty->ContainerPtrToValuePtr<void>( static_cast<const void*>(Object)) );
+			if( BaseObject && *BaseObject )
+			{
+				Object = *BaseObject;
+			}
+			ClassOverride = nullptr;
+		}
+		else if( const UStructProperty* StructProperty = Cast<UStructProperty>(ResolvedProperty) )
+		{
+			ClassOverride = StructProperty->Struct;
+		}
+		FPropertyInfo Info = { ResolvedProperty, -1 };
+		Ret.AddProperty(Info);
+	}
+	return Ret;
+}
+
+const UObject* DiffUtils::GetCDO(const UBlueprint* ForBlueprint)
+{
+	if (!ForBlueprint
+		|| !ForBlueprint->GeneratedClass)
+	{
+		return NULL;
+	}
+
+	return ForBlueprint->GeneratedClass->ClassDefaultObject;
+}
+
+void DiffUtils::CompareUnrelatedObjects(const UObject* A, UObject const* B, TArray<FPropertySoftPath>& OutDifferingProperties)
+{
+	FPropertySoftPathSet PropertiesInA = GetPropertyNameSet(A);
+	FPropertySoftPathSet PropertiesInB = GetPropertyNameSet(B);
+
+	// any properties in A that aren't in B are differing:
+	OutDifferingProperties.Append(PropertiesInA.Difference(PropertiesInB).Array());
+
+	// and the converse:
+	OutDifferingProperties.Append(PropertiesInB.Difference(PropertiesInA).Array());
+
+	// for properties in common, dig out the uproperties and determine if they're identical:
+	if (A && B)
+	{
+		FPropertySoftPathSet Common = PropertiesInA.Intersect(PropertiesInB);
+		for (const auto& PropertyName : Common)
+		{
+			FResolvedProperty AProp = PropertyName.Resolve(A);
+			FResolvedProperty BProp = PropertyName.Resolve(B);
+
+			check(AProp != FResolvedProperty() && BProp != FResolvedProperty());
+			if (!DiffUtils::Identical(AProp, BProp))
+			{
+				OutDifferingProperties.Add(PropertyName);
+			}
+		}
+	}
+}
+
+void DiffUtils::CompareUnrelatedSCS(const UBlueprint* A, const UBlueprint* B, TArray< FSCSDiffEntry >& OutDifferingEntries, TArray< int >* OptionalOutSortKeys)
+{
+	const auto FindObjectPropertyByName = [](const UObject* Instance, FName Name) -> UObject const*
+	{
+		for (TFieldIterator<UProperty> PropertyIter(Instance->GetClass()); PropertyIter; ++PropertyIter)
+		{
+			if (PropertyIter->GetFName() == Name)
+			{
+				const UObject* const* SCS = PropertyIter->ContainerPtrToValuePtr<const UObject*>(Instance);
+				return SCS ? *SCS : nullptr;
+			}
+		}
+		return nullptr;
+	};
+
+	FName SCS = TEXT("SimpleConstructionScript");
+	UObject const* SCSA = FindObjectPropertyByName(A, SCS);
+	UObject const* SCSB = FindObjectPropertyByName(B, SCS);
+
+	// we get bitten by const shallowness here, avoid mutation:
+	const TArray<USCS_Node*> NodesInA = A && A->SimpleConstructionScript ? A->SimpleConstructionScript->GetAllNodes() : TArray<USCS_Node*>();
+	const TArray<USCS_Node*> NodesInB = B && B->SimpleConstructionScript ? B->SimpleConstructionScript->GetAllNodes() : TArray<USCS_Node*>();
+
+	const auto FindByName = [](const TArray<USCS_Node*>& NodeList, FName Name) -> USCS_Node*
+	{
+		for (auto Node : NodeList)
+		{
+			if (Node->VariableName == Name)
+			{
+				return Node;
+			}
+		}
+		return nullptr;
+	};
+
+	int32 SortKey = 0;
+	for (auto NodeInA : NodesInA)
+	{
+		auto NodeInB = FindByName(NodesInB, NodeInA->VariableName);
+		if (NodeInB)
+		{
+			// diff NodeA vs NodeB:
+			TArray<FPropertySoftPath> DifferingProperties;
+
+			DiffUtils::CompareUnrelatedObjects(A, B, DifferingProperties);
+			for (auto Property : DifferingProperties)
+			{
+				FSCSDiffEntry Diff = { NodeInA->VariableName, Property.LastPropertyName() };
+				OutDifferingEntries.Push(Diff);
+
+				if (OptionalOutSortKeys)
+				{
+					OptionalOutSortKeys->Push(SortKey);
+				}
+			}
+		}
+		else
+		{
+			// Node A was added (removed from node B):
+			FSCSDiffEntry Diff = { NodeInA->VariableName, FName() };
+			OutDifferingEntries.Push(Diff);
+
+			if (OptionalOutSortKeys)
+			{
+				OptionalOutSortKeys->Push(SortKey);
+			}
+		}
+
+		++SortKey;
+	}
+
+	// add nodes that were added in B:
+	for (auto NodeInB : NodesInB)
+	{
+		if (FindByName(NodesInA, NodeInB->VariableName) == nullptr)
+		{
+			// Node B was added:
+			FSCSDiffEntry Diff = { NodeInB->VariableName, FName() };
+			OutDifferingEntries.Push(Diff);
+
+			if (OptionalOutSortKeys)
+			{
+				OptionalOutSortKeys->Push(SortKey);
+			}
+		}
+
+		++SortKey;
+	}
+}
+
+bool DiffUtils::Identical(const FResolvedProperty& AProp, const FResolvedProperty& BProp)
+{
+	const void* AValue = AProp.Property->ContainerPtrToValuePtr<void>(AProp.Object);
+	const void* BValue = BProp.Property->ContainerPtrToValuePtr<void>(BProp.Object);
+
+	return AProp.Property->Identical(AValue, BValue, PPF_DeepComparison);
+}
+
+TArray<FPropertySoftPath> DiffUtils::GetVisiblePropertiesInOrderDeclared(const UObject* ForObj, const TArray<FName>& Scope /*= TArray<FName>()*/)
+{
+	TArray<FPropertySoftPath> Ret;
+	if (ForObj)
+	{
+		const UClass* Class = ForObj->GetClass();
+		TSet<FString> HiddenCategories = FEditorCategoryUtils::GetHiddenCategories(Class);
+		for (TFieldIterator<UProperty> PropertyIt(Class); PropertyIt; ++PropertyIt)
+		{
+			FName CategoryName = FObjectEditorUtils::GetCategoryFName(*PropertyIt);
+			if (!HiddenCategories.Contains(CategoryName.ToString()))
+			{
+				if (PropertyIt->PropertyFlags&CPF_Edit)
+				{
+					TArray<FName> NewPath(Scope);
+					NewPath.Push(PropertyIt->GetFName());
+					if (const UObjectProperty* ObjectProperty = Cast<UObjectProperty>(*PropertyIt))
+					{
+						const UObject* const* BaseObject = reinterpret_cast<const UObject* const*>( ObjectProperty->ContainerPtrToValuePtr<void>(ForObj) );
+						if (BaseObject && *BaseObject)
+						{
+							Ret.Append( GetVisiblePropertiesInOrderDeclared(*BaseObject, NewPath) );
+						}
+					}
+					else
+					{
+						Ret.Push(NewPath);
+					}
+				}
+			}
+		}
+	}
+	return Ret;
+}
+
+TArray<FPropertyPath> DiffUtils::ResolveAll(const UObject* Object, const TArray<FPropertySoftPath>& InSoftProperties)
+{
+	TArray< FPropertyPath > Ret;
+	for (const auto& Path : InSoftProperties)
+	{
+		Ret.Push(Path.ResolvePath(Object));
+	}
+	return Ret;
+}
diff --git a/Engine/Source/Editor/Kismet/Private/SBlueprintDiff.cpp b/Engine/Source/Editor/Kismet/Private/SBlueprintDiff.cpp
--- a/Engine/Source/Editor/Kismet/Private/SBlueprintDiff.cpp
+++ b/Engine/Source/Editor/Kismet/Private/SBlueprintDiff.cpp
@@ -104,10 +104,8 @@ class FCDODiffControl	: public TSharedFromThis<FCDODiffControl>
 {
 public:
 	FCDODiffControl( const UObject* InOldCDO
-					, FNamePropertyMap InOldProperties
 					, const UObject* InNewCDO
-					, FNamePropertyMap InNewProperities
-					, TArray<FName> InDifferingProperties );
+					, const TArray<FPropertySoftPath>& InDifferingProperties );
 
 	TSharedRef<SWidget> OldDetailsWidget() { return OldDetails.DetailsWidget(); }
 	TSharedRef<SWidget> NewDetailsWidget() { return NewDetails.DetailsWidget(); }
@@ -120,24 +118,28 @@ private:
 	bool HasPrevDifference() const override;
 
 	void HighlightCurrentDifference();
+
+	void UpdateDisplayedDifferences();
+
 	FDetailsDiff OldDetails;
 	FDetailsDiff NewDetails;
 
-	TArray< FName > DifferingProperties;
+	TArray<FPropertySoftPath> DisplayedDifferingProperties;
+	FPropertySoftPathSet DifferingProperties;
 	int CurrentDifference;
 };
 
 FCDODiffControl::FCDODiffControl( 
 		const UObject* InOldCDO
-		, FNamePropertyMap InOldProperties
 		, const UObject* InNewCDO
-		, FNamePropertyMap InNewProperities
-		, TArray<FName> InDifferingProperties )
-	: OldDetails(InOldCDO, InOldProperties, TSet<FName>(InDifferingProperties))
-	, NewDetails(InNewCDO, InNewProperities, TSet<FName>(InDifferingProperties))
-	, DifferingProperties( InDifferingProperties )
+		, const TArray<FPropertySoftPath>& InDifferingProperties)
+	: OldDetails(InOldCDO, DiffUtils::ResolveAll( InOldCDO, InDifferingProperties), FDetailsDiff::FOnDisplayedPropertiesChanged::CreateRaw(this, &FCDODiffControl::UpdateDisplayedDifferences) )
+	, NewDetails(InNewCDO, DiffUtils::ResolveAll( InNewCDO, InDifferingProperties), FDetailsDiff::FOnDisplayedPropertiesChanged::CreateRaw(this, &FCDODiffControl::UpdateDisplayedDifferences) )
+	, DisplayedDifferingProperties()
+	, DifferingProperties(InDifferingProperties)
 	, CurrentDifference( -1 )
 {
+	UpdateDisplayedDifferences();
 }
 
 void FCDODiffControl::NextDiff()
@@ -154,22 +156,67 @@ void FCDODiffControl::PrevDiff()
 
 bool FCDODiffControl::HasNextDifference() const
 {
-	return DifferingProperties.IsValidIndex(CurrentDifference + 1);
+	return DisplayedDifferingProperties.IsValidIndex(CurrentDifference + 1);
 }
 
 bool FCDODiffControl::HasPrevDifference() const
 {
-	return DifferingProperties.IsValidIndex(CurrentDifference - 1);
+	return DisplayedDifferingProperties.IsValidIndex(CurrentDifference - 1);
 }
 
 void FCDODiffControl::HighlightCurrentDifference()
 {
-	const FName PropertyName = DifferingProperties[CurrentDifference];
+	const FPropertySoftPath& PropertyName = DisplayedDifferingProperties[CurrentDifference];
 
 	OldDetails.HighlightProperty(PropertyName);
 	NewDetails.HighlightProperty(PropertyName);
 }
 
+void FCDODiffControl::UpdateDisplayedDifferences()
+{
+	DisplayedDifferingProperties.Empty();
+	CurrentDifference = 0;
+
+	// create differing properties list based on what is displayed by the old properties..
+	TArray<FPropertySoftPath> OldProperties = OldDetails.GetDisplayedProperties();
+	TArray<FPropertySoftPath> NewProperties = NewDetails.GetDisplayedProperties();
+
+	// zip the two sets of properties, zip iterators are not trivial to write in c++,
+	// so this procedural stuff will have to do:
+	int IterOld = 0;
+	int IterNew = 0;
+	while (IterOld != OldProperties.Num() || IterNew != NewProperties.Num())
+	{
+		if (IterOld != OldProperties.Num() && IterNew != NewProperties.Num()
+			&& OldProperties[IterOld] == NewProperties[IterNew])
+		{
+			if( DifferingProperties.Contains(OldProperties[IterOld]) )
+			{
+				DisplayedDifferingProperties.Push(OldProperties[IterOld]);
+			}
+			++IterOld;
+			++IterNew;
+		}
+		else if (IterOld != OldProperties.Num())
+		{
+			if (DifferingProperties.Contains(OldProperties[IterOld]))
+			{
+				DisplayedDifferingProperties.Push(OldProperties[IterOld]);
+			}
+			++IterOld;
+		}
+		else
+		{
+			check(IterNew != NewProperties.Num());
+			if (DifferingProperties.Contains(NewProperties[IterNew]))
+			{
+				DisplayedDifferingProperties.Push(NewProperties[IterNew]);
+			}
+			++IterNew;
+		}
+	}
+}
+
 /*List item that entry for a graph*/
 struct KISMET_API FListItemGraphToDiff	: public TSharedFromThis<FListItemGraphToDiff>
 										, public IDiffControl
@@ -258,48 +305,6 @@ private:
 	TSharedPtr< FUICommandList > KeyCommands;
 };
 
-FNamePropertyMap DiffUtils::GetProperties( const UObject* ForObj )
-{
-	FNamePropertyMap Ret;
-	if( ForObj )
-	{
-		const UClass* Class = ForObj->GetClass();
-		for (TFieldIterator<UProperty> PropertyIt(Class, EFieldIteratorFlags::IncludeSuper); PropertyIt; ++PropertyIt)
-		{
-			if ((!PropertyIt->HasAnyPropertyFlags(CPF_Parm) && PropertyIt->HasAnyPropertyFlags(CPF_Edit)))
-			{
-				Ret.Add(PropertyIt->GetFName(), *PropertyIt);
-			}
-		}
-	}
-	return Ret;
-}
-
-static void GetPropertyNameSet( const UObject& ForObj, TSet< FName >& OutProperties )
-{
-	if( const UClass* Class = ForObj.GetClass() )
-	{
-		for (TFieldIterator<UProperty> PropertyIt(Class, EFieldIteratorFlags::IncludeSuper); PropertyIt; ++PropertyIt)
-		{
-			if ((!PropertyIt->HasAnyPropertyFlags(CPF_Parm) && PropertyIt->HasAnyPropertyFlags(CPF_Edit)))
-			{
-				OutProperties.Add(PropertyIt->GetFName());
-			}
-		}
-	}
-}
-
-const UObject* DiffUtils::GetCDO( const UBlueprint* ForBlueprint )
-{
-	if( !ForBlueprint 
-		|| !ForBlueprint->GeneratedClass )
-	{
-		return NULL;
-	}
-
-	return ForBlueprint->GeneratedClass->ClassDefaultObject;
-}
-
 TSharedRef<SWidget>	FDiffResultItem::GenerateWidget() const
 {
 	FString ToolTip = Result.ToolTip;
@@ -617,132 +622,6 @@ FDiffPanel::FDiffPanel()
 	LastFocusedPin = NULL;
 }
 
-void DiffUtils::CompareUnrelatedObjects( const UObject* A, const TMap< FName, const UProperty* >& PropertyMapA, UObject const* B, const TMap< FName, const UProperty* >& PropertyMapB, TArray<FName> &OutIdenticalProperties, TArray<FName> &OutDifferingProperties )
-{
-	TSet<FName> PropertiesInA, PropertiesInB;
-	if( A )
-	{
-		GetPropertyNameSet(*A, PropertiesInA);
-	}
-	if( B )
-	{
-		GetPropertyNameSet(*B, PropertiesInB);
-	}
-
-	// any properties in A that aren't in B are differing:
-	OutDifferingProperties.Append( PropertiesInA.Difference(PropertiesInB).Array() );
-
-	// and the converse:
-	OutDifferingProperties.Append( PropertiesInB.Difference(PropertiesInA).Array() );
-
-	// for properties in common, dig out the uproperties and determine if they're identical:
-	if( A && B )
-	{
-		TSet<FName> Common = PropertiesInA.Intersect(PropertiesInB);
-		for (const auto& PropertyName : Common)
-		{
-			const UProperty* const* AProp = PropertyMapA.Find(PropertyName);
-			const UProperty* const* BProp = PropertyMapB.Find(PropertyName);
-			
-			if( !AProp && !BProp )
-			{
-				OutIdenticalProperties.Add(PropertyName);
-				continue;
-			}
-
-			check( AProp && BProp );
-			if( (*AProp)->SameType(*BProp) )
-			{
-				const void* AValue = (*AProp)->ContainerPtrToValuePtr<void>(A);
-				const void* BValue = (*BProp)->ContainerPtrToValuePtr<void>(B);
-
-				if ((*AProp)->Identical(AValue, BValue, PPF_DeepComparison))
-				{
-					OutIdenticalProperties.Add(PropertyName);
-					continue;
-				}
-			}
-
-			OutDifferingProperties.Add(PropertyName);
-		}
-	}
-}
-
-void DiffUtils::CompareUnrelatedSCS(const UBlueprint* A, const UBlueprint* B, TArray< FSCSDiffEntry >& OutDifferingEntries, TArray< int >* OptionalOutSortKeys )
-{
-	// we get bitten by const shallowness here, avoid mutation:
-	const TArray<USCS_Node*> NodesInA = A && A->SimpleConstructionScript ? A->SimpleConstructionScript->GetAllNodes() : TArray<USCS_Node*>();
-	const TArray<USCS_Node*> NodesInB = B && B->SimpleConstructionScript ? B->SimpleConstructionScript->GetAllNodes() : TArray<USCS_Node*>();
-
-	const auto FindByName = []( const TArray<USCS_Node*>& NodeList, FName Name ) -> USCS_Node*
-	{
-		for( auto Node : NodeList )
-		{
-			if( Node->VariableName == Name )
-			{
-				return Node;
-			}
-		}
-		return nullptr;
-	};
-
-	int32 SortKey = 0;
-	for( auto NodeInA : NodesInA )
-	{
-		auto NodeInB = FindByName( NodesInB, NodeInA->VariableName );
-		if( NodeInB )
-		{
-			// diff NodeA vs NodeB:
-			FNamePropertyMap PropertyMapA = DiffUtils::GetProperties(NodeInA->ComponentTemplate);
-			FNamePropertyMap PropertyMapB = DiffUtils::GetProperties(NodeInB->ComponentTemplate);
-			TArray<FName> IdenticalProperties, DifferingProperties;
-
-			DiffUtils::CompareUnrelatedObjects(A, PropertyMapA, B, PropertyMapB, IdenticalProperties, DifferingProperties);
-			for( auto Property : DifferingProperties )
-			{
-				FSCSDiffEntry Diff = { NodeInA->VariableName, Property };
-				OutDifferingEntries.Push( Diff );
-
-				if (OptionalOutSortKeys)
-				{
-					OptionalOutSortKeys->Push(SortKey);
-				}
-			}
-		}
-		else
-		{
-			// Node A was added:
-			FSCSDiffEntry Diff = { NodeInA->VariableName, FName() };
-			OutDifferingEntries.Push( Diff );
-
-			if (OptionalOutSortKeys)
-			{
-				OptionalOutSortKeys->Push(SortKey);
-			}
-		}
-
-		++SortKey;
-	}
-
-	// add nodes that were added in B:
-	for( auto NodeInB : NodesInB )
-	{
-		if( FindByName( NodesInA, NodeInB->VariableName ) == nullptr )
-		{
-			// Node B was added:
-			FSCSDiffEntry Diff = { NodeInB->VariableName, FName() };
-			OutDifferingEntries.Push(Diff);
-
-			if (OptionalOutSortKeys)
-			{
-				OptionalOutSortKeys->Push(SortKey);
-			}
-		}
-
-		++SortKey;
-	}
-}
-
 int32 GetCurrentIndex( SListView< TSharedPtr< struct FDiffSingleResult> > const& ListView, const TArray< TSharedPtr< FDiffSingleResult > >& ListViewSource )
 {
 	const auto& Selected = ListView.GetSelectedItems();
@@ -1374,13 +1253,11 @@ TSharedRef<SWidget> SBlueprintDiff::GenerateGraphPanel()
 TSharedRef<SWidget> SBlueprintDiff::GenerateDefaultsPanel()
 {
 	const UObject* A = DiffUtils::GetCDO(PanelOld.Blueprint);
-	FNamePropertyMap PropertyMapA = DiffUtils::GetProperties(A);
 	const UObject* B = DiffUtils::GetCDO(PanelNew.Blueprint);
-	FNamePropertyMap PropertyMapB = DiffUtils::GetProperties(B);
-	TArray<FName> IdenticalProperties, DifferingProperties;
-	DiffUtils::CompareUnrelatedObjects( A, PropertyMapA, B, PropertyMapB, IdenticalProperties, DifferingProperties);
+	TArray<FPropertySoftPath> DifferingProperties;
+	DiffUtils::CompareUnrelatedObjects( A, B, DifferingProperties);
 
-	auto NewDiffControl = TSharedPtr<FCDODiffControl>(new FCDODiffControl(A, PropertyMapA, B, PropertyMapB, DifferingProperties) );
+	auto NewDiffControl = TSharedPtr<FCDODiffControl>(new FCDODiffControl(A, B, DifferingProperties) );
 	//Splitter for left and right blueprint. Current convention is for the local (probably newer?) blueprint to be on the right:
 	DiffControl = NewDiffControl;
 	return SNew(SSplitter)
diff --git a/Engine/Source/Editor/Kismet/Private/SSCSEditor.cpp b/Engine/Source/Editor/Kismet/Private/SSCSEditor.cpp
--- a/Engine/Source/Editor/Kismet/Private/SSCSEditor.cpp
+++ b/Engine/Source/Editor/Kismet/Private/SSCSEditor.cpp
@@ -11,6 +11,7 @@
 #include "SKismetInspector.h"
 #include "SSCSEditorViewport.h"
 #include "SComponentClassCombo.h"
+#include "PropertyPath.h"
 
 #include "AssetSelection.h"
 #include "Editor/SceneOutliner/Private/SSocketChooser.h"
@@ -2492,8 +2493,13 @@ void SSCSEditor::HighlightTreeNode(const USCS_Node* Node, FName Property)
 		check( KismetInspectorSPtr.IsValid() );
 		UActorComponent* Component = TreeNode->GetComponentTemplate();
 		UProperty* CurrentProp = FindField<UProperty>(Component->GetClass(), Property);
-		check( CurrentProp );
-		KismetInspectorSPtr->GetPropertyView()->HighlightProperty( CurrentProp );
+		FPropertyPath Path;
+		if( CurrentProp )
+		{
+			FPropertyInfo NewInfo = { CurrentProp, -1 };
+			Path.ExtendPath(NewInfo);
+		}
+		KismetInspectorSPtr->GetPropertyView()->HighlightProperty( Path );
 	}
 }
 
diff --git a/Engine/Source/Editor/Kismet/Public/DetailsDiff.h b/Engine/Source/Editor/Kismet/Public/DetailsDiff.h
--- a/Engine/Source/Editor/Kismet/Public/DetailsDiff.h
+++ b/Engine/Source/Editor/Kismet/Public/DetailsDiff.h
@@ -1,15 +1,27 @@
 // Copyright 1998-2014 Epic Games, Inc. All Rights Reserved.
 #pragma once
 
+#include "PropertyPath.h"
+#include "DiffUtils.h"
+
 class KISMET_API FDetailsDiff
 {
 public:
-	FDetailsDiff( const UObject* InObject, const TMap< FName, const UProperty* >& InPropertyMap, const TSet< FName >& InDifferingProperties );
+	DECLARE_DELEGATE( FOnDisplayedPropertiesChanged );
+
+	FDetailsDiff( const UObject* InObject, const TArray< FPropertyPath >& InDifferingProperties, FOnDisplayedPropertiesChanged InOnDisplayedPropertiesChanged );
 
-	void HighlightProperty( FName PropertyName );
+	void HighlightProperty( const FPropertySoftPath& PropertyName );
 	TSharedRef< SWidget > DetailsWidget();
+	TArray<FPropertySoftPath> GetDisplayedProperties() const;
 
 private:
-	TMap< FName, const UProperty* > PropertyMap;
+	void HandlePropertiesChanged();
+
+	FOnDisplayedPropertiesChanged OnDisplayedPropertiesChanged;
+
+	TArray< FPropertyPath > DifferingProperties;
+	const UObject* DisplayedObject;
+
 	TSharedPtr< class IDetailsView > DetailsView;
 };
diff --git a/Engine/Source/Editor/Kismet/Public/DiffUtils.h b/Engine/Source/Editor/Kismet/Public/DiffUtils.h
new file mode 100644
--- /dev/null
+++ b/Engine/Source/Editor/Kismet/Public/DiffUtils.h
@@ -0,0 +1,107 @@
+// Copyright 1998-2014 Epic Games, Inc. All Rights Reserved.
+#pragma once
+
+#include "PropertyPath.h"
+
+struct FResolvedProperty
+{
+	explicit FResolvedProperty()
+	: Object(nullptr)
+	, Property(nullptr)
+	{
+	}
+
+	FResolvedProperty(const UObject* InObject, const UProperty* InProperty)
+		: Object(InObject)
+		, Property(InProperty)
+	{
+	}
+
+	inline bool operator==(const FResolvedProperty& RHS) const
+	{
+		return Object == RHS.Object && Property == RHS.Property;
+	}
+
+	inline bool operator!=(const FResolvedProperty& RHS) const { return !(*this == RHS); }
+
+	const UObject* Object;
+	const UProperty* Property;
+};
+
+/**
+ * FPropertySoftPath is a string of identifiers used to identify a single member of a UObject. It is primarily
+ * used when comparing unrelated UObjects for Diffing and Merging, but can also be used as a key select
+ * a property in a SDetailsView.
+ */
+struct FPropertySoftPath
+{
+	FPropertySoftPath(TArray<FName> InPropertyChain)
+		: PropertyChain(InPropertyChain)
+	{
+	}
+
+	FPropertySoftPath( FPropertyPath InPropertyPath )
+	{
+		for( int32 i = 0, end = InPropertyPath.GetNumProperties(); i != end; ++i )
+		{
+			PropertyChain.Push( InPropertyPath.GetPropertyInfo(i).Property->GetFName() );
+		}
+	}
+
+	FResolvedProperty Resolve(const UObject* Object) const;
+	FPropertyPath ResolvePath(const UObject* Object) const;
+	FName LastPropertyName() const;
+
+	inline bool operator==(FPropertySoftPath const& RHS) const
+	{
+		return PropertyChain == RHS.PropertyChain;
+	}
+private:
+	friend uint32 GetTypeHash( FPropertySoftPath const& Path );
+	TArray<FName> PropertyChain;
+};
+
+FORCEINLINE FName FPropertySoftPath::LastPropertyName() const
+{
+	if( PropertyChain.Num() != 0 )
+	{
+		return PropertyChain.Last();
+	}
+	return FName();
+}
+
+FORCEINLINE uint32 GetTypeHash( FPropertySoftPath const& Path )
+{
+	uint32 Ret = 0;
+	for( const auto& ProperytName : Path.PropertyChain )
+	{
+		Ret = Ret ^ GetTypeHash(ProperytName);
+	}
+	return Ret;
+}
+
+// Trying to restrict us to this typedef because I'm a little skeptical about hashing FPropertySoftPath safely
+typedef TSet< FPropertySoftPath > FPropertySoftPathSet;
+
+struct FSCSDiffEntry
+{
+	FName TreeNodeName;
+	FName PropertyName;
+};
+
+inline bool operator ==(const FSCSDiffEntry A, const FSCSDiffEntry B)
+{
+	return A.TreeNodeName == B.TreeNodeName && A.PropertyName == B.PropertyName;
+}
+
+namespace DiffUtils
+{
+	KISMET_API const UObject* GetCDO(const UBlueprint* ForBlueprint);
+	KISMET_API void CompareUnrelatedObjects(const UObject* A, const UObject* B, TArray<FPropertySoftPath>& OutDifferingProperties);
+	KISMET_API void CompareUnrelatedSCS(const UBlueprint* Old, const UBlueprint* New, TArray< FSCSDiffEntry >& OutDifferingEntries, TArray< int >* OptionalOutSortKeys = nullptr);
+	KISMET_API bool Identical(const FResolvedProperty& AProp, const FResolvedProperty& BProp);
+	TArray<FPropertySoftPath> GetVisiblePropertiesInOrderDeclared(const UObject* ForObj, const TArray<FName>& Scope = TArray<FName>());
+
+	KISMET_API TArray<FPropertyPath> ResolveAll(const UObject* Object, const TArray<FPropertySoftPath>& InSoftProperties);
+
+}
diff --git a/Engine/Source/Editor/Kismet/Public/SBlueprintDiff.h b/Engine/Source/Editor/Kismet/Public/SBlueprintDiff.h
--- a/Engine/Source/Editor/Kismet/Public/SBlueprintDiff.h
+++ b/Engine/Source/Editor/Kismet/Public/SBlueprintDiff.h
@@ -30,25 +30,6 @@ struct FDiffResultItem : public TSharedFromThis<FDiffResultItem>
 	TSharedRef<SWidget> KISMET_API GenerateWidget() const;
 };
 
-struct FSCSDiffEntry
-{
-	FName TreeNodeName;
-	FName PropertyName;
-};
-
-inline bool operator ==(const FSCSDiffEntry A, const FSCSDiffEntry B)
-{
-	return A.TreeNodeName == B.TreeNodeName && A.PropertyName == B.PropertyName;
-}
-
-namespace DiffUtils
-{
-	KISMET_API TMap< FName, const UProperty* > GetProperties( const UObject* ForObj );
-	KISMET_API const UObject* GetCDO( const UBlueprint* ForBlueprint );
-	KISMET_API void CompareUnrelatedObjects( const UObject* A, const TMap< FName, const UProperty* >& PropertyMapA, const UObject* B, const TMap< FName, const UProperty* >& PropertyMapB, TArray<FName> &OutIdenticalProperties, TArray<FName> &OutDifferingProperties );
-	KISMET_API void CompareUnrelatedSCS(const UBlueprint* Old, const UBlueprint* New, TArray< struct FSCSDiffEntry >& OutDifferingEntries, TArray< int >* OptionalOutSortKeys = nullptr );
-}
-
 namespace DiffWidgetUtils
 {
 	KISMET_API void SelectNextRow(SListView< TSharedPtr< struct FDiffSingleResult> >& ListView, const TArray< TSharedPtr< struct FDiffSingleResult > >& ListViewSource );
diff --git a/Engine/Source/Editor/Kismet/Public/SCSDiff.h b/Engine/Source/Editor/Kismet/Public/SCSDiff.h
--- a/Engine/Source/Editor/Kismet/Public/SCSDiff.h
+++ b/Engine/Source/Editor/Kismet/Public/SCSDiff.h
@@ -1,7 +1,7 @@
 // Copyright 1998-2014 Epic Games, Inc. All Rights Reserved.
 #pragma once
 
-#include "SBlueprintDiff.h"
+#include "DiffUtils.h"
 
 class KISMET_API FSCSDiff
 {
diff --git a/Engine/Source/Editor/PropertyEditor/Private/DetailItemNode.cpp b/Engine/Source/Editor/PropertyEditor/Private/DetailItemNode.cpp
--- a/Engine/Source/Editor/PropertyEditor/Private/DetailItemNode.cpp
+++ b/Engine/Source/Editor/PropertyEditor/Private/DetailItemNode.cpp
@@ -22,6 +22,7 @@ FDetailItemNode::FDetailItemNode(const FDetailLayoutCustomization& InCustomizati
 	, bShouldBeVisibleDueToChildFiltering( false )
 	, bTickable( false )
 	, bIsExpanded( InCustomization.HasCustomBuilder() ? !InCustomization.CustomBuilderRow->IsInitiallyCollapsed() : false )
+	, bIsHighlighted( false )
 {
 	
 }
@@ -79,8 +80,6 @@ FDetailItemNode::~FDetailItemNode()
 
 void FDetailItemNode::InitPropertyEditor()
 {
-	Customization.GetPropertyNode()->SetTreeNode( this->AsShared() );
-
 	if( Customization.GetPropertyNode()->GetProperty() && Customization.GetPropertyNode()->GetProperty()->IsA<UArrayProperty>() )
 	{
 		const bool bUpdateFilteredNodes = false;
@@ -296,7 +295,11 @@ static bool PassesAllFilters( const FDetailLayoutCustomization& InCustomization,
 
 			const bool bPassesSearchFilter = bSearchFilterIsEmpty || ( bIsNotBeingFiltered || bIsSeenDueToFiltering || bIsParentSeenDueToFiltering );
 			const bool bPassesModifiedFilter = bPassesSearchFilter && ( InFilter.bShowOnlyModifiedProperties == false || PropertyNodePin->GetDiffersFromDefault() == true );
-			const bool bPassesDifferingFilter = InFilter.bShowOnlyDiffering == false || InFilter.DifferingProperties.Find(PropertyNodePin->GetProperty()->GetFName()) != NULL;
+			bool bPassesDifferingFilter = true;
+			if( InFilter.bShowOnlyDiffering )
+			{
+				 bPassesDifferingFilter = InFilter.WhitelistedProperties.Find(*FPropertyNode::CreatePropertyPath(PropertyNodePin.ToSharedRef())) != NULL;
+			}
 
 			// The property node is visible (note categories are never visible unless they have a child that is visible )
 			bPassesAllFilters = bPassesSearchFilter && bPassesModifiedFilter && bPassesDifferingFilter;
@@ -383,6 +386,17 @@ FName FDetailItemNode::GetNodeName() const
 	return NAME_None;
 }
 
+FPropertyPath FDetailItemNode::GetPropertyPath() const
+{
+	FPropertyPath Ret;
+	auto PropertyNode = Customization.GetPropertyNode();
+	if( PropertyNode.IsValid() )
+	{
+		Ret = *FPropertyNode::CreatePropertyPath( PropertyNode.ToSharedRef() );
+	}
+	return Ret;
+}
+
 void FDetailItemNode::FilterNode( const FDetailFilter& InFilter )
 {
 	bShouldBeVisibleDueToFiltering = PassesAllFilters( Customization, InFilter );
diff --git a/Engine/Source/Editor/PropertyEditor/Private/DetailItemNode.h b/Engine/Source/Editor/PropertyEditor/Private/DetailItemNode.h
--- a/Engine/Source/Editor/PropertyEditor/Private/DetailItemNode.h
+++ b/Engine/Source/Editor/PropertyEditor/Private/DetailItemNode.h
@@ -41,6 +41,11 @@ public:
 	virtual void Tick( float DeltaTime ) override;
 	virtual bool ShouldShowOnlyChildren() const override;
 	virtual FName GetNodeName() const override;
+	virtual TSharedPtr<FDetailCategoryImpl> GetParentCategory() override { return ParentCategory.Pin(); }
+	virtual FPropertyPath GetPropertyPath() const override;
+	virtual void SetIsHighlighted(bool bInIsHighlighted) { bIsHighlighted = bInIsHighlighted; }
+	virtual bool IsHighlighted() const { return bIsHighlighted; }
+	virtual bool IsLeaf() { return true; }
 
 private:
 
@@ -78,4 +83,6 @@ private:
 	bool bTickable;
 	/** True if this node is expanded */
 	bool bIsExpanded;
+	/** True if this node is highlighted */
+	bool bIsHighlighted;
 };
\ No newline at end of file
diff --git a/Engine/Source/Editor/PropertyEditor/Private/IDetailTreeNode.h b/Engine/Source/Editor/PropertyEditor/Private/IDetailTreeNode.h
--- a/Engine/Source/Editor/PropertyEditor/Private/IDetailTreeNode.h
+++ b/Engine/Source/Editor/PropertyEditor/Private/IDetailTreeNode.h
@@ -76,5 +76,30 @@ public:
 	 * The identifier name of the node                                                              
 	 */
 	virtual FName GetNodeName() const = 0;
+
+	/**
+	 * @return The category node that this node is nested in, if any:
+	 */
+	virtual TSharedPtr<FDetailCategoryImpl> GetParentCategory() { return TSharedPtr<FDetailCategoryImpl>(); }
+	
+	/**
+	 * @return The property path that this node is associate with, if any:
+	 */
+	virtual FPropertyPath GetPropertyPath() const { return FPropertyPath(); }
+	
+	/**
+	 * Called when the node should appear 'highlighted' to draw the users attention to it
+	 */
+	virtual void SetIsHighlighted(bool bInIsHighlighted) {}
+	
+	/**
+	 * @return true if the node has been highlighted
+	 */
+	virtual bool IsHighlighted() const { return false; }
+
+	/**
+	 * @return true if this is a leaf node:
+	 */
+	virtual bool IsLeaf() { return false; }
 };
 
diff --git a/Engine/Source/Editor/PropertyEditor/Private/PropertyNode.cpp b/Engine/Source/Editor/PropertyEditor/Private/PropertyNode.cpp
--- a/Engine/Source/Editor/PropertyEditor/Private/PropertyNode.cpp
+++ b/Engine/Source/Editor/PropertyEditor/Private/PropertyNode.cpp
@@ -46,7 +46,6 @@ FPropertyNode::FPropertyNode(void)
 	, MaxChildDepthAllowed(FPropertyNodeConstants::NoDepthRestrictions)
 	, PropertyNodeFlags (EPropertyNodeFlags::NoFlags)
 	, bRebuildChildrenRequested( false )
-	, bIsHighlighted( false )
 {
 }
 
diff --git a/Engine/Source/Editor/PropertyEditor/Private/PropertyNode.h b/Engine/Source/Editor/PropertyEditor/Private/PropertyNode.h
--- a/Engine/Source/Editor/PropertyEditor/Private/PropertyNode.h
+++ b/Engine/Source/Editor/PropertyEditor/Private/PropertyNode.h
@@ -649,18 +649,6 @@ public:
 
 	FPropertyChangedEvent& FixPropertiesInEvent(FPropertyChangedEvent& Event);
 
-	/** Returns true if this property node should be highlighted by the view */
-	bool IsHighlighted() const { return bIsHighlighted; }
-
-	/** Sets the highlighted flag to the requested value */
-	void SetIsHighlighted( bool bInIsHighlighted ) { bIsHighlighted = bInIsHighlighted; }
-
-	/** Returns the IDetailTreeNode associated with this property node */
-	TSharedPtr< class IDetailTreeNode > GetTreeNode() { return TreeNode.Pin(); }
-
-	/** Sets the IDetailTreeNode associated with this property node */
-	void SetTreeNode( TSharedPtr< class IDetailTreeNode > InTreeNode ) { TreeNode = InTreeNode; }
-
 protected:
 
 	TSharedRef<FEditPropertyChain> BuildPropertyChain( UProperty* PropertyAboutToChange );
@@ -775,9 +763,6 @@ protected:
 	/** If true, children of this node will be rebuilt next tick. */
 	bool bRebuildChildrenRequested;
 
-	/** Used to display a highlight on this node */
-	bool bIsHighlighted;
-
 	/** An array of restrictions limiting this property's potential values in property editors.*/
 	TArray<TSharedRef<const class FPropertyRestriction>> Restrictions;
 
diff --git a/Engine/Source/Editor/PropertyEditor/Private/SDetailSingleItemRow.cpp b/Engine/Source/Editor/PropertyEditor/Private/SDetailSingleItemRow.cpp
--- a/Engine/Source/Editor/PropertyEditor/Private/SDetailSingleItemRow.cpp
+++ b/Engine/Source/Editor/PropertyEditor/Private/SDetailSingleItemRow.cpp
@@ -386,5 +386,5 @@ FReply SDetailSingleItemRow::OnAddKeyframeClicked()
 
 bool SDetailSingleItemRow::IsHighlighted() const
 {
-	return Customization->IsValidCustomization() && Customization->HasPropertyNode() && Customization->GetPropertyNode()->IsHighlighted();
+	return OwnerTreeNode.Pin()->IsHighlighted();
 }
diff --git a/Engine/Source/Editor/PropertyEditor/Private/SDetailsView.cpp b/Engine/Source/Editor/PropertyEditor/Private/SDetailsView.cpp
--- a/Engine/Source/Editor/PropertyEditor/Private/SDetailsView.cpp
+++ b/Engine/Source/Editor/PropertyEditor/Private/SDetailsView.cpp
@@ -36,17 +36,13 @@ SDetailsView::~SDetailsView()
 void SDetailsView::Construct(const FArguments& InArgs)
 {
 	DetailsViewArgs = InArgs._DetailsViewArgs;
-	bHasActiveFilter = false;
-	bIsLocked = false;
 	bViewingClassDefaultObject = false;
-	bHasOpenColorPicker = false;
 
 	// Create the root property now
 	RootPropertyNode = MakeShareable( new FObjectPropertyNode );
 		
 	PropertyUtilities = MakeShareable( new FPropertyDetailsUtilities( *this ) );
 	
-	ColumnWidth = 0.65f;
 	ColumnSizeData.LeftColumnWidth = TAttribute<float>( this, &SDetailsView::OnGetLeftColumnWidth );
 	ColumnSizeData.RightColumnWidth = TAttribute<float>( this, &SDetailsView::OnGetRightColumnWidth );
 	ColumnSizeData.OnWidthChanged = SSplitter::FOnSlotResized::CreateSP( this, &SDetailsView::OnSetColumnWidth );
diff --git a/Engine/Source/Editor/PropertyEditor/Private/SDetailsViewBase.cpp b/Engine/Source/Editor/PropertyEditor/Private/SDetailsViewBase.cpp
--- a/Engine/Source/Editor/PropertyEditor/Private/SDetailsViewBase.cpp
+++ b/Engine/Source/Editor/PropertyEditor/Private/SDetailsViewBase.cpp
@@ -88,42 +88,93 @@ void SDetailsViewBase::HideFilterArea(bool bHide)
 	DetailsViewArgs.bAllowSearch = !bHide;
 }
 
-void SDetailsViewBase::HighlightProperty(const UProperty* Property)
+static void GetPropertiesInOrderDisplayedRecursive(const TArray< TSharedRef<IDetailTreeNode> >& TreeNodes, TArray< FPropertyPath > &OutLeaves)
 {
-	// Clear the previously enabled highlight:
-	auto PrevPropertyPtr = PrevHighlightedProperty.Pin();
-	if( PrevPropertyPtr.IsValid() )
+	for (auto& TreeNode : TreeNodes)
 	{
-		PrevPropertyPtr->SetIsHighlighted( false );
+		if (TreeNode->IsLeaf())
+		{
+			FPropertyPath Path = TreeNode->GetPropertyPath();
+			// Some leaf nodes are not associated with properties, specifically the collision presets.
+			// @todo doc: investigate what we can do about this, result is that for these fields
+			// we can't highlight hte property in the diff tool.
+			if( Path.GetNumProperties() != 0 )
+			{
+				OutLeaves.Push(Path);
+			}
+		}
+		else
+		{
+			TArray< TSharedRef<IDetailTreeNode> > Children;
+			TreeNode->GetChildren(Children);
+			GetPropertiesInOrderDisplayedRecursive(Children, OutLeaves);
+		}
 	}
+}
+TArray< FPropertyPath > SDetailsViewBase::GetPropertiesInOrderDisplayed() const
+{
+	TArray< FPropertyPath > Ret;
+	GetPropertiesInOrderDisplayedRecursive(RootTreeNodes, Ret);
+	return Ret;
+}
 
-	if( !Property )
+static TSharedPtr< IDetailTreeNode > FindTreeNodeFromPropertyRecursive( const TArray< TSharedRef<IDetailTreeNode> >& Nodes, const FPropertyPath& Property )
+{
+	for (auto& TreeNode : Nodes)
 	{
-		return;
+		if (TreeNode->IsLeaf())
+		{
+			FPropertyPath tmp = TreeNode->GetPropertyPath();
+			if( Property == tmp )
+			{
+				return TreeNode;
+			}
+		}
+		else
+		{
+			TArray< TSharedRef<IDetailTreeNode> > Children;
+			TreeNode->GetChildren(Children);
+			auto Result = FindTreeNodeFromPropertyRecursive(Children, Property);
+			if( Result.IsValid() )
+			{
+				return Result;
+			}
+		}
 	}
 
-	// Find the newly highlighted node, and enable the highlight:
-	FName PropertyOwnerFName = Property->GetOwnerStruct()->GetFName();
-	FClassInstanceToPropertyMap* PropertyNodeMap = ClassToPropertyMap.Find(PropertyOwnerFName);
-	if( PropertyNodeMap )
+	return TSharedPtr< IDetailTreeNode >();
+}
+
+void SDetailsViewBase::HighlightProperty(const FPropertyPath& Property)
+{
+	auto PrevHighlightedNodePtr = CurrentlyHighlightedNode.Pin();
+	if (PrevHighlightedNodePtr.IsValid())
 	{
-		FName InstanceName = NAME_None;
-		FPropertyNodeMap* NodeMap = PropertyNodeMap->Find(InstanceName);
-		if( NodeMap )
+		PrevHighlightedNodePtr->SetIsHighlighted(false);
+	}
+
+	auto NextNodePtr = FindTreeNodeFromPropertyRecursive( RootTreeNodes, Property );
+	if (NextNodePtr.IsValid())
+	{
+		NextNodePtr->SetIsHighlighted(true);
+		auto ParentCategory = NextNodePtr->GetParentCategory();
+		if (ParentCategory.IsValid())
 		{
-			TSharedPtr<FPropertyNode>* PropertyNode = NodeMap->PropertyNameToNode.Find(Property->GetFName());
-			if( PropertyNode )
-			{
-				(*PropertyNode)->SetIsHighlighted( true );
-				PrevHighlightedProperty = *PropertyNode;
-				auto TreeNodePtr = (*PropertyNode)->GetTreeNode();
-				if ( TreeNodePtr.IsValid() )
-				{
-					DetailTree->RequestScrollIntoView(TreeNodePtr.ToSharedRef());
-				}
-			}
+			DetailTree->SetItemExpansion(ParentCategory.ToSharedRef(), true);
 		}
+		DetailTree->RequestScrollIntoView(NextNodePtr.ToSharedRef());
 	}
+	CurrentlyHighlightedNode = NextNodePtr;
+}
+
+void SDetailsViewBase::ShowAllAdvancedProperties()
+{
+	CurrentFilter.bShowAllAdvanced = true;
+}
+
+void SDetailsViewBase::SetOnDisplayedPropertiesChanged(FOnDisplayedPropertiesChanged InOnDisplayedPropertiesChangedDelegate)
+{
+	OnDisplayedPropertiesChangedDelegate = InOnDisplayedPropertiesChangedDelegate;
 }
 
 EVisibility SDetailsViewBase::GetTreeVisibility() const
@@ -199,7 +250,6 @@ void SDetailsViewBase::CreateColorPickerWindow(const TSharedRef< FPropertyEditor
 	OpenColorPicker(PickerArgs);
 }
 
-
 void SDetailsViewBase::SetColorPropertyFromColorPicker(FLinearColor NewColor)
 {
 	const TSharedPtr< FPropertyNode > PinnedColorPropertyNode = ColorPropertyNode.Pin();
@@ -310,6 +360,11 @@ void SDetailsViewBase::RequestItemExpanded(TSharedRef<IDetailTreeNode> TreeNode,
 
 void SDetailsViewBase::RefreshTree()
 {
+	if (OnDisplayedPropertiesChangedDelegate.IsBound())
+	{
+		OnDisplayedPropertiesChangedDelegate.Execute();
+	}
+
 	DetailTree->RequestTreeRefresh();
 }
 
@@ -500,7 +555,8 @@ void SDetailsViewBase::QueryCustomDetailLayout(FDetailLayoutBuilderImpl& CustomD
 				FClassInstanceToPropertyMap& InstancedPropertyMap = ClassToPropertyMap.FindChecked(Class->GetFName());
 				for (FClassInstanceToPropertyMap::TIterator InstanceIt(InstancedPropertyMap); InstanceIt; ++InstanceIt)
 				{
-					CustomDetailLayout.SetCurrentCustomizationClass(CastChecked<UClass>(Class), InstanceIt.Key());
+					FName Key = InstanceIt.Key();
+					CustomDetailLayout.SetCurrentCustomizationClass(CastChecked<UClass>(Class), Key);
 
 					const FOnGetDetailCustomizationInstance& DetailDelegate = LayoutIt.Value()->DetailLayoutDelegate;
 
@@ -885,7 +941,8 @@ void SDetailsViewBase::UpdateFilteredDetails()
 
 		RootTreeNodes = DetailLayout->GetRootTreeNodes();
 	}
-	DetailTree->RequestTreeRefresh();
+
+	RefreshTree();
 }
 
 /**
@@ -1093,9 +1150,15 @@ void SDetailsViewBase::UpdatePropertyMap()
 
 	CustomUpdatePropertyMap();
 
-	// Ask for custom detail layouts
-	QueryCustomDetailLayout(*DetailLayout);
-
+	// Ask for custom detail layouts, unless disabled. One reason for disabling custom layouts is that the custom layouts
+	// inhibit our ability to find a single property's tree node. This is problematic for the diff and merge tools, that need
+	// to display and highlight each changed property for the user. We could whitelist 'good' customizations here if 
+	// we can make them work with the diff/merge tools.
+	if( !bDisableCustomDetailLayouts )
+	{
+		QueryCustomDetailLayout(*DetailLayout);
+	}
+	
 	DetailLayout->GenerateDetailLayout();
 
 	UpdateFilteredDetails();
diff --git a/Engine/Source/Editor/PropertyEditor/Private/SDetailsViewBase.h b/Engine/Source/Editor/PropertyEditor/Private/SDetailsViewBase.h
--- a/Engine/Source/Editor/PropertyEditor/Private/SDetailsViewBase.h
+++ b/Engine/Source/Editor/PropertyEditor/Private/SDetailsViewBase.h
@@ -66,7 +66,7 @@ struct FDetailFilter
 	bool bShowAllAdvanced;
 	/** If we should only show differing properties */
 	bool bShowOnlyDiffering;
-	TSet<FName> DifferingProperties;
+	TSet<FPropertyPath> WhitelistedProperties;
 };
 
 struct FDetailColumnSizeData
@@ -81,6 +81,15 @@ struct FDetailColumnSizeData
 class SDetailsViewBase : public IDetailsViewPrivate
 {
 public:
+	SDetailsViewBase()
+		: ColumnWidth(.65f)
+		, bHasActiveFilter(false)
+		, bIsLocked(false)
+		, bHasOpenColorPicker(false)
+		, bDisableCustomDetailLayouts( false )
+	{
+	}
+
 	/**
 	* @return true of the details view can be updated from editor selection
 	*/
@@ -101,9 +110,13 @@ public:
 	virtual void HideFilterArea(bool bHide) override;
 
 	/** 
-	 * Creates a box around the property with PropertyName and scrolls the property into view
+	 * Implementation of IDetailsView:
 	 */
-	virtual void HighlightProperty(const UProperty* Property) override;
+	virtual TArray< FPropertyPath > GetPropertiesInOrderDisplayed() const override;
+	virtual void HighlightProperty(const FPropertyPath& Property) override;
+	virtual void ShowAllAdvancedProperties() override;
+	virtual void SetOnDisplayedPropertiesChanged(FOnDisplayedPropertiesChanged InOnDisplayedPropertiesChangedDelegate) override;
+	virtual void SetDisableCustomDetailLayouts( bool bInDisableCustomDetailLayouts ) override { bDisableCustomDetailLayouts = bInDisableCustomDetailLayouts; }
 
 	virtual FOnFinishedChangingProperties& OnFinishedChangingProperties() override
 	{ 
@@ -272,7 +285,6 @@ public:
 	void CreateColorPickerWindow(const TSharedRef< class FPropertyEditor >& PropertyEditor, bool bUseAlpha) override;
 
 protected:
-
 	/**
 	 * Called when a color property is changed from a color picker
 	 */
@@ -356,7 +368,7 @@ protected:
 	void OnFilterTextChanged(const FText& InFilterText);
 
 	/** Called when the list of currently differing properties changes */
-	virtual void UpdateDifferingProperties(const TSet<FName> InDifferingProperties) override { CurrentFilter.DifferingProperties = InDifferingProperties; }
+	virtual void UpdatePropertiesWhitelist(const TSet<FPropertyPath> InWhitelistedProperties) override { CurrentFilter.WhitelistedProperties = InWhitelistedProperties; }
 
 	/** 
 	 * Hides or shows properties based on the passed in filter text
@@ -367,7 +379,6 @@ protected:
 
 	/** Called to get the visibility of the filter box */
 	EVisibility GetFilterBoxVisibility() const;
-
 protected:
 	/** The user defined args for the details view */
 	FDetailsViewArgs DetailsViewArgs;
@@ -436,5 +447,12 @@ protected:
 	/** External property nodes which need to validated each tick */
 	TArray< TWeakPtr<FObjectPropertyNode> > ExternalRootPropertyNodes;
 
-	TWeakPtr< FPropertyNode > PrevHighlightedProperty;
+	/** The tree node that is currently highlighted, may be none: */
+	TWeakPtr< IDetailTreeNode > CurrentlyHighlightedNode;
+
+	/** Executed when the tree is refreshed */
+	FOnDisplayedPropertiesChanged OnDisplayedPropertiesChangedDelegate;
+
+	/** True if we want to skip generation of custom layouts for displayed object */
+	bool bDisableCustomDetailLayouts;
 };
\ No newline at end of file
diff --git a/Engine/Source/Editor/PropertyEditor/Private/SStructureDetailsView.cpp b/Engine/Source/Editor/PropertyEditor/Private/SStructureDetailsView.cpp
--- a/Engine/Source/Editor/PropertyEditor/Private/SStructureDetailsView.cpp
+++ b/Engine/Source/Editor/PropertyEditor/Private/SStructureDetailsView.cpp
@@ -31,10 +31,7 @@ UScriptStruct* SStructureDetailsView::GetBaseScriptStruct() const
 void SStructureDetailsView::Construct(const FArguments& InArgs)
 {
 	DetailsViewArgs = InArgs._DetailsViewArgs;
-	bHasActiveFilter = false;
-	bIsLocked = false;
-	bHasOpenColorPicker = false;
-
+	
 	CustomName = InArgs._CustomName;
 
 	// Create the root property now
@@ -42,7 +39,6 @@ void SStructureDetailsView::Construct(const FArguments& InArgs)
 		
 	PropertyUtilities = MakeShareable( new FPropertyDetailsUtilities( *this ) );
 	
-	ColumnWidth = 0.65f;
 	ColumnSizeData.LeftColumnWidth = TAttribute<float>(this, &SStructureDetailsView::OnGetLeftColumnWidth);
 	ColumnSizeData.RightColumnWidth = TAttribute<float>(this, &SStructureDetailsView::OnGetRightColumnWidth);
 	ColumnSizeData.OnWidthChanged = SSplitter::FOnSlotResized::CreateSP(this, &SStructureDetailsView::OnSetColumnWidth);
diff --git a/Engine/Source/Editor/PropertyEditor/Public/IDetailsView.h b/Engine/Source/Editor/PropertyEditor/Public/IDetailsView.h
--- a/Engine/Source/Editor/PropertyEditor/Public/IDetailsView.h
+++ b/Engine/Source/Editor/PropertyEditor/Public/IDetailsView.h
@@ -172,12 +172,32 @@ public:
 	virtual void HideFilterArea(bool bIsVisible) = 0;
 
 	/**
-	 * Creates a box around the property with PropertyName and scrolls the property into view
+	 * Returns a list of all the properties displayed (via full path), order in list corresponds to draw order:
 	 */
-	virtual void HighlightProperty(const UProperty* Property) = 0;
+	virtual TArray< FPropertyPath > GetPropertiesInOrderDisplayed() const = 0;
+
+	/**
+	 * Creates a box around the treenode corresponding to Property and scrolls the treenode into view
+	 */
+	virtual void HighlightProperty(const FPropertyPath& Property) = 0;
+	
+	/**
+	 * Forces all advanced property sections to be in expanded state:
+	 */
+	virtual void ShowAllAdvancedProperties() = 0;
+	
+	/**
+	 * Assigns delegate called when view is filtered, useful for updating external control logic:
+	 */
+	virtual void SetOnDisplayedPropertiesChanged(FOnDisplayedPropertiesChanged InOnDisplayedPropertiesChangedDelegate) = 0;
+
+	/**
+	 * Disables or enables customization of the details view:
+	 */
+	virtual void SetDisableCustomDetailLayouts(bool bInDisableCustomDetailLayouts) = 0;
 
 	/**
 	 * Sets the set of properties that are considered differing, used when filtering out identical properties
 	 */
-	virtual void UpdateDifferingProperties(const TSet<FName> DifferingProperties) = 0;
+	virtual void UpdatePropertiesWhitelist(const TSet<FPropertyPath> InWhitelistedProperties) = 0;
 };
diff --git a/Engine/Source/Editor/PropertyEditor/Public/PropertyEditorDelegates.h b/Engine/Source/Editor/PropertyEditor/Public/PropertyEditorDelegates.h
--- a/Engine/Source/Editor/PropertyEditor/Public/PropertyEditorDelegates.h
+++ b/Engine/Source/Editor/PropertyEditor/Public/PropertyEditorDelegates.h
@@ -30,6 +30,9 @@ DECLARE_DELEGATE_RetVal( TSharedRef<class IPropertyTypeCustomization>, FOnGetPro
 /** Notification for when a property view changes */
 DECLARE_DELEGATE_TwoParams( FOnObjectArrayChanged, const FString&, const TArray< TWeakObjectPtr< UObject > >& );
 
+/** Notification for when displayed properties changes (for instance, because the user has filtered some properties */
+DECLARE_DELEGATE( FOnDisplayedPropertiesChanged );
+
 /** Notification for when a property selection changes. */
 DECLARE_DELEGATE_OneParam( FOnPropertySelectionChanged, UProperty* )
 
diff --git a/Engine/Source/Runtime/Engine/Classes/Engine/Blueprint.h b/Engine/Source/Runtime/Engine/Classes/Engine/Blueprint.h
--- a/Engine/Source/Runtime/Engine/Classes/Engine/Blueprint.h
+++ b/Engine/Source/Runtime/Engine/Classes/Engine/Blueprint.h
@@ -341,6 +341,15 @@ class ENGINE_API UBlueprint : public UBlueprintCore
 	/** Set of functions actually compiled for this class */
 	UPROPERTY(transient, duplicatetransient)
 	TArray<class UEdGraph*> EventGraphs;
+
+	/** 
+	 * Flag indicating that a read only duplicate of this blueprint is being created, used to disable logic in ::PostDuplicate,
+	 *
+	 * This flag needs to be copied on duplication (because it's the duplicated object that we're disabling on PostDuplicate),
+	 * but we don't *need* to serialize it for permanent objects.
+	 */
+	UPROPERTY()
+	mutable bool bDuplicatingReadOnly;
 #endif // WITH_EDITORONLY_DATA
 
 	/** Array of component template objects, used by AddComponent function */
diff --git a/Engine/Source/Runtime/Engine/Private/Blueprint.cpp b/Engine/Source/Runtime/Engine/Private/Blueprint.cpp
--- a/Engine/Source/Runtime/Engine/Private/Blueprint.cpp
+++ b/Engine/Source/Runtime/Engine/Private/Blueprint.cpp
@@ -273,6 +273,7 @@ UBlueprint::UBlueprint(const FObjectInitializer& ObjectInitializer)
 	, bGenerateConstClass(false)
 #endif
 #if WITH_EDITORONLY_DATA
+	, bDuplicatingReadOnly(false)
 	, bCachedDependenciesUpToDate(false)
 #endif
 {
@@ -413,7 +414,10 @@ bool UBlueprint::Rename( const TCHAR* InName, UObject* NewOuter, ERenameFlags Fl
 void UBlueprint::PostDuplicate(bool bDuplicateForPIE)
 {
 	Super::PostDuplicate(bDuplicateForPIE);
-	FBlueprintEditorUtils::PostDuplicateBlueprint(this);
+	if( !bDuplicatingReadOnly )
+	{
+		FBlueprintEditorUtils::PostDuplicateBlueprint(this);
+	}
 }
 
 UClass* UBlueprint::RegenerateClass(UClass* ClassToRegenerate, UObject* PreviousCDO, TArray<UObject*>& ObjLoaded)
