[UE-1551] Disabling "save on compile" for level blueprints (since it saves the level).

diff --git a/Engine/Source/Editor/Kismet/Private/BlueprintEditor.cpp b/Engine/Source/Editor/Kismet/Private/BlueprintEditor.cpp
--- a/Engine/Source/Editor/Kismet/Private/BlueprintEditor.cpp
+++ b/Engine/Source/Editor/Kismet/Private/BlueprintEditor.cpp
@@ -1792,7 +1792,7 @@ void FBlueprintEditor::CreateDefaultCommands()
 	ToolkitCommands->MapAction(
 		FFullBlueprintEditorCommands::Get().SaveOnCompile,
 		FExecuteAction::CreateSP(this, &FBlueprintEditor::OnSaveOnCompileToggled),
-		FCanExecuteAction(),
+		FCanExecuteAction::CreateSP(this, &FBlueprintEditor::IsSaveOnCompileEnabled),
 		FIsActionChecked::CreateSP(this, &FBlueprintEditor::IsSaveOnCompileChecked)
 	);
 	
@@ -2466,7 +2466,7 @@ void FBlueprintEditor::Compile()
 bool FBlueprintEditor::IsSaveOnCompileChecked() const
 {
 	UBlueprintEditorSettings const* Settings = GetDefault<UBlueprintEditorSettings>();
-	return Settings->bSaveOnCompile;
+	return Settings->bSaveOnCompile && IsSaveOnCompileEnabled();
 }
 
 void FBlueprintEditor::OnSaveOnCompileToggled() const
@@ -2476,6 +2476,14 @@ void FBlueprintEditor::OnSaveOnCompileToggled() const
 	Settings->SaveConfig();
 }
 
+bool FBlueprintEditor::IsSaveOnCompileEnabled() const
+{
+	UBlueprint* Blueprint = GetBlueprintObj();
+	bool const bIsLevelScript = (Cast<ULevelScriptBlueprint>(Blueprint) != nullptr);
+
+	return !bIsLevelScript;
+}
+
 FReply FBlueprintEditor::Compile_OnClickWithReply()
 {
 	Compile();
diff --git a/Engine/Source/Editor/Kismet/Public/BlueprintEditor.h b/Engine/Source/Editor/Kismet/Public/BlueprintEditor.h
--- a/Engine/Source/Editor/Kismet/Public/BlueprintEditor.h
+++ b/Engine/Source/Editor/Kismet/Public/BlueprintEditor.h
@@ -318,6 +318,7 @@ public:
 	/** Helper functions used to construct/operate-on the "Save on Compile" command */
 	bool IsSaveOnCompileChecked() const;
 	void OnSaveOnCompileToggled() const;
+	bool IsSaveOnCompileEnabled() const;
 
 	/** Calls the above function, but returns an FReply::Handled(). Used in SButtons */
 	virtual FReply Compile_OnClickWithReply();
diff --git a/Engine/Source/Editor/UnrealEd/Private/Kismet2/Kismet2.cpp b/Engine/Source/Editor/UnrealEd/Private/Kismet2/Kismet2.cpp
--- a/Engine/Source/Editor/UnrealEd/Private/Kismet2/Kismet2.cpp
+++ b/Engine/Source/Editor/UnrealEd/Private/Kismet2/Kismet2.cpp
@@ -481,10 +481,18 @@ void FKismetEditorUtilities::CompileBlueprint(UBlueprint* BlueprintObj, bool bIs
 		UBlueprintEditorSettings* Settings = GetMutableDefault<UBlueprintEditorSettings>();
 		if (Settings->bSaveOnCompile && (BlueprintObj->Status == BS_UpToDate))
 		{
-			TArray<UPackage*> PackagesToSave;
-			PackagesToSave.Add(BlueprintPackage);
+			bool const bIsLevelPackage = (UWorld::FindWorldInPackage(BlueprintPackage) != nullptr);
+			// we don't want to save the entire level (especially is this 
+			// compile was already kicked off as a result of a level save, as it
+			// could cause a recursive save)... let the "bSaveOnCompile" setting 
+			// only save blueprint assets
+			if (!bIsLevelPackage)
+			{
+				TArray<UPackage*> PackagesToSave;
+				PackagesToSave.Add(BlueprintPackage);
 
-			FEditorFileUtils::PromptForCheckoutAndSave(PackagesToSave, /*bCheckDirty =*/true, /*bPromptToSave =*/false);
+				FEditorFileUtils::PromptForCheckoutAndSave(PackagesToSave, /*bCheckDirty =*/true, /*bPromptToSave =*/false);
+			}			
 		}
 	}	
 }
