Extending the BP editor's "save on compile" option, so it can be applied to blueprints that fail to compile.

diff --git a/Engine/Source/Editor/BlueprintGraph/Private/BlueprintEditorSettings.cpp b/Engine/Source/Editor/BlueprintGraph/Private/BlueprintEditorSettings.cpp
--- a/Engine/Source/Editor/BlueprintGraph/Private/BlueprintEditorSettings.cpp
+++ b/Engine/Source/Editor/BlueprintGraph/Private/BlueprintEditorSettings.cpp
@@ -7,18 +7,24 @@
 
 UBlueprintEditorSettings::UBlueprintEditorSettings(const FObjectInitializer& ObjectInitializer)
 	: Super(ObjectInitializer)
+	// Style Settings
+	, bDrawMidpointArrowsInBlueprints(false)
+	// UX Settings
+	, bShowGraphInstructionText(true)
+	, bUseTargetContextForNodeMenu(true)
+	, bExposeAllMemberComponentFunctions(true)
+	, bShowContextualFavorites(false)
+	, bFlattenFavoritesMenus(true)
+	, bUseLegacyMenuingSystem(false)
+	// Compiler Settings
+	, SaveOnCompile(SoC_Never)
+	// Developer Settings
+	, bShowActionMenuItemSignatures(false)
+	// Perf Settings
+	, bShowDetailedCompileResults(false)
+	, CompileEventDisplayThresholdMs(5)
+	, NodeTemplateCacheCapMB(20.f)
 {
-	bSaveOnCompile = false;
-	bShowGraphInstructionText = true;
-	NodeTemplateCacheCapMB = 20.f;
-	bUseTargetContextForNodeMenu = true;
-	bExposeAllMemberComponentFunctions = true;
-	bShowContextualFavorites = false;
-	bFlattenFavoritesMenus = true;
-	bUseLegacyMenuingSystem = false;
-	bShowDetailedCompileResults = false;
-	CompileEventDisplayThresholdMs = 5;
-
 	// settings that were moved out of experimental...
 	UEditorExperimentalSettings const* ExperimentalSettings = GetDefault<UEditorExperimentalSettings>();
 	bDrawMidpointArrowsInBlueprints = ExperimentalSettings->bDrawMidpointArrowsInBlueprints;
@@ -26,4 +32,13 @@ UBlueprintEditorSettings::UBlueprintEditorSettings(const FObjectInitializer& Obj
 	// settings that were moved out of editor-user settings...
 	UEditorUserSettings const* UserSettings = GetDefault<UEditorUserSettings>();
 	bShowActionMenuItemSignatures = UserSettings->bDisplayActionListItemRefIds;
+
+	FString const ClassConfigKey = GetClass()->GetPathName();
+
+	bool bOldSaveOnCompileVal = false;
+	// backwards compatibility: handle the case where users have already switched this on
+	if (GConfig->GetBool(*ClassConfigKey, TEXT("bSaveOnCompile"), bOldSaveOnCompileVal, GEditorUserSettingsIni) && bOldSaveOnCompileVal)
+	{
+		SaveOnCompile = SoC_SuccessOnly;
+	}
 }
diff --git a/Engine/Source/Editor/BlueprintGraph/Public/BlueprintEditorSettings.h b/Engine/Source/Editor/BlueprintGraph/Public/BlueprintEditorSettings.h
--- a/Engine/Source/Editor/BlueprintGraph/Public/BlueprintEditorSettings.h
+++ b/Engine/Source/Editor/BlueprintGraph/Public/BlueprintEditorSettings.h
@@ -5,19 +5,20 @@
 #include "EdGraph/EdGraphPin.h" // for EBlueprintPinStyleType
 #include "BlueprintEditorSettings.generated.h"
 
-
+UENUM()
+enum ESaveOnCompile
+{
+	SoC_Never UMETA(DisplayName="Never"),
+	SoC_SuccessOnly UMETA(DisplayName="On Success Only"),
+	SoC_Always UMETA(DisplayName = "Always"),
+};
+ 
 UCLASS(config=EditorUserSettings)
 class BLUEPRINTGRAPH_API UBlueprintEditorSettings
 	:	public UObject
 {
 	GENERATED_UCLASS_BODY()
 
-// General Settings
-public:
-	/** Determines if Blueprints are saved whenever you (successfully) compile them */
-	UPROPERTY(EditAnywhere, config, Category=General)
-	bool bSaveOnCompile;
-
 // Style Settings
 public:
 	/** Should arrows indicating data/execution flow be drawn halfway along wires? */
@@ -31,7 +32,7 @@ public:
 	bool bShowGraphInstructionText;
 
 	/** If enabled, will use the blueprint's (or output pin's) class to narrow down context menu results. */
-	UPROPERTY(EditAnywhere, config, Category=UserExperience)
+	UPROPERTY(EditAnywhere, config, Category=UserExperience, meta=(DisplayName="Context Menu: Use Target Context"))
 	bool bUseTargetContextForNodeMenu;
 
 	/** If enabled, then ALL component functions are exposed to the context menu (when the contextual target is a component owner). Ignores "ExposeFunctionCategories" metadata for components. */
@@ -50,6 +51,12 @@ public:
 	UPROPERTY(EditAnywhere, AdvancedDisplay, config, Category=UserExperience)
 	bool bUseLegacyMenuingSystem;
 
+// Compiler Settings
+public:
+	/** Determines when to save Blueprints post-compile */
+	UPROPERTY(EditAnywhere, config, Category=Compiler)
+	TEnumAsByte<ESaveOnCompile> SaveOnCompile;
+
 // Developer Settings
 public:
 	/** If enabled, tooltips on action menu items will show the associated action's signature id (can be used to setup custom favorites menus).*/
diff --git a/Engine/Source/Editor/Kismet/Private/BlueprintEditor.cpp b/Engine/Source/Editor/Kismet/Private/BlueprintEditor.cpp
--- a/Engine/Source/Editor/Kismet/Private/BlueprintEditor.cpp
+++ b/Engine/Source/Editor/Kismet/Private/BlueprintEditor.cpp
@@ -179,20 +179,40 @@ namespace BlueprintEditorImpl
 	static const float InstructionFadeDuration = 0.5f;
 
 	/**
-	*
-	*
-	* @param  InGraph
-	* @return
-	*/
+	 * Utility function that will check to see if the specified graph has any 
+	 * nodes other than those that come default, pre-placed, in the graph.
+	 *
+	 * @param  InGraph  The graph to check.
+	 * @return True if the graph has any nodes added by the user, otherwise false.
+	 */
 	static bool GraphHasUserPlacedNodes(UEdGraph const* InGraph);
 
 	/**
-	*
-	*
-	* @param  InGraph
-	* @return
-	*/
+	 * Utility function that will check to see if the specified graph has any
+	 * nodes that were default, pre-placed, in the graph.
+	 *
+	 * @param  InGraph  The graph to check.
+	 * @return True if the graph has any pre-placed nodes, otherwise false.
+	 */
 	static bool GraphHasDefaultNode(UEdGraph const* InGraph);
+
+	/**
+	 * Utility function that will set the global save-on-compile setting to the
+	 * specified value.
+	 * 
+	 * @param  NewSetting	The new save-on-compile setting that you want applied.
+	 */
+	static void SetSaveOnCompileSetting(ESaveOnCompile NewSetting);
+
+	/**
+	 * Utility function used to determine what save-on-compile setting should be 
+	 * presented to the user.
+	 * 
+	 * @param  Editor	The editor currently querying for the setting value.
+	 * @param  Option	The setting to check for.
+	 * @return False if the option isn't set, or if the save-on-compile is disabled for the blueprint being edited (otherwise true). 
+	 */
+	static bool IsSaveOnCompileOptionSet(TWeakPtr<FBlueprintEditor> Editor, ESaveOnCompile Option);
 }
 
 static bool BlueprintEditorImpl::GraphHasUserPlacedNodes(UEdGraph const* InGraph)
@@ -237,6 +257,30 @@ static bool BlueprintEditorImpl::GraphHasDefaultNode(UEdGraph const* InGraph)
 	return bHasDefaultNodes;
 }
 
+static void BlueprintEditorImpl::SetSaveOnCompileSetting(ESaveOnCompile NewSetting)
+{
+	UBlueprintEditorSettings* Settings = GetMutableDefault<UBlueprintEditorSettings>();
+	Settings->SaveOnCompile = NewSetting;
+	Settings->SaveConfig();
+}
+
+static bool BlueprintEditorImpl::IsSaveOnCompileOptionSet(TWeakPtr<FBlueprintEditor> Editor, ESaveOnCompile Option)
+{
+	const UBlueprintEditorSettings* Settings = GetDefault<UBlueprintEditorSettings>();
+
+	ESaveOnCompile CurrentSetting = Settings->SaveOnCompile;
+	if (!Editor.IsValid() || !Editor.Pin()->IsSaveOnCompileEnabled())
+	{
+		// if save-on-compile is disabled for the blueprint, then we want to 
+		// show "Never" as being selected
+		// 
+		// @TODO: a tooltip explaining why would be nice too
+		CurrentSetting = SoC_Never;
+	}
+
+	return (CurrentSetting == Option);
+}
+
 bool FBlueprintEditor::IsASubGraph( const UEdGraph* GraphPtr )
 {
 	if( GraphPtr && GraphPtr->GetOuter() )
@@ -1789,11 +1833,24 @@ void FBlueprintEditor::CreateDefaultCommands()
 		FExecuteAction::CreateSP(this, &FBlueprintEditor::Compile),
 		FCanExecuteAction::CreateSP(this, &FBlueprintEditor::IsCompilingEnabled));
 
+	TWeakPtr<FBlueprintEditor> WeakThisPtr = SharedThis(this);
+	ToolkitCommands->MapAction(
+		FFullBlueprintEditorCommands::Get().SaveOnCompile_Never,
+		FExecuteAction::CreateStatic(&BlueprintEditorImpl::SetSaveOnCompileSetting, (ESaveOnCompile)SoC_Never),
+		FCanExecuteAction::CreateSP(this, &FBlueprintEditor::IsSaveOnCompileEnabled),
+		FIsActionChecked::CreateStatic(&BlueprintEditorImpl::IsSaveOnCompileOptionSet, WeakThisPtr, (ESaveOnCompile)SoC_Never)
+	);
+	ToolkitCommands->MapAction(
+		FFullBlueprintEditorCommands::Get().SaveOnCompile_SuccessOnly,
+		FExecuteAction::CreateStatic(&BlueprintEditorImpl::SetSaveOnCompileSetting, (ESaveOnCompile)SoC_SuccessOnly),
+		FCanExecuteAction::CreateSP(this, &FBlueprintEditor::IsSaveOnCompileEnabled),
+		FIsActionChecked::CreateStatic(&BlueprintEditorImpl::IsSaveOnCompileOptionSet, WeakThisPtr, (ESaveOnCompile)SoC_SuccessOnly)
+	);
 	ToolkitCommands->MapAction(
-		FFullBlueprintEditorCommands::Get().SaveOnCompile,
-		FExecuteAction::CreateSP(this, &FBlueprintEditor::OnSaveOnCompileToggled),
+		FFullBlueprintEditorCommands::Get().SaveOnCompile_Always,
+		FExecuteAction::CreateStatic(&BlueprintEditorImpl::SetSaveOnCompileSetting, (ESaveOnCompile)SoC_Always),
 		FCanExecuteAction::CreateSP(this, &FBlueprintEditor::IsSaveOnCompileEnabled),
-		FIsActionChecked::CreateSP(this, &FBlueprintEditor::IsSaveOnCompileChecked)
+		FIsActionChecked::CreateStatic(&BlueprintEditorImpl::IsSaveOnCompileOptionSet, WeakThisPtr, (ESaveOnCompile)SoC_Always)
 	);
 	
 	ToolkitCommands->MapAction(
@@ -2463,19 +2520,6 @@ void FBlueprintEditor::Compile()
 	}
 }
 
-bool FBlueprintEditor::IsSaveOnCompileChecked() const
-{
-	UBlueprintEditorSettings const* Settings = GetDefault<UBlueprintEditorSettings>();
-	return Settings->bSaveOnCompile && IsSaveOnCompileEnabled();
-}
-
-void FBlueprintEditor::OnSaveOnCompileToggled() const
-{
-	UBlueprintEditorSettings* Settings = GetMutableDefault<UBlueprintEditorSettings>();
-	Settings->bSaveOnCompile = !Settings->bSaveOnCompile;
-	Settings->SaveConfig();
-}
-
 bool FBlueprintEditor::IsSaveOnCompileEnabled() const
 {
 	UBlueprint* Blueprint = GetBlueprintObj();
diff --git a/Engine/Source/Editor/Kismet/Private/SBlueprintEditorToolbar.cpp b/Engine/Source/Editor/Kismet/Private/SBlueprintEditorToolbar.cpp
--- a/Engine/Source/Editor/Kismet/Private/SBlueprintEditorToolbar.cpp
+++ b/Engine/Source/Editor/Kismet/Private/SBlueprintEditorToolbar.cpp
@@ -544,7 +544,11 @@ TSharedRef<SWidget> FKismet2Menu::MakeDiffMenu(FBlueprintEditor& Kismet)
 void FFullBlueprintEditorCommands::RegisterCommands() 
 {
 	UI_COMMAND(Compile, "Compile", "Compile the blueprint", EUserInterfaceActionType::Button, FInputGesture());
-	UI_COMMAND(SaveOnCompile, "Save on Compile", "Determines if the Blueprint is saved every time you (successfully) compile it", EUserInterfaceActionType::ToggleButton, FInputGesture());
+	
+	UI_COMMAND(SaveOnCompile_Never, "Never", "Sets the save-on-compile option to 'Never', meaning that your Blueprints will not be saved when they are compiled", EUserInterfaceActionType::RadioButton, FInputGesture());
+	UI_COMMAND(SaveOnCompile_SuccessOnly, "On Success Only", "Sets the save-on-compile option to 'Success Only', meaning that your Blueprints will be saved whenever they are successfully compiled", EUserInterfaceActionType::RadioButton, FInputGesture());
+	UI_COMMAND(SaveOnCompile_Always, "Always", "Sets the save-on-compile option to 'Always', meaning that your Blueprints will be saved whenever they ar compiled (even if there were errors)", EUserInterfaceActionType::RadioButton, FInputGesture());
+
 	UI_COMMAND(SwitchToScriptingMode, "Graph", "Switches to Graph Editing Mode", EUserInterfaceActionType::ToggleButton, FInputGesture());
 	UI_COMMAND(SwitchToBlueprintDefaultsMode, "Defaults", "Switches to Blueprint Defaults Mode", EUserInterfaceActionType::ToggleButton, FInputGesture());
 	UI_COMMAND(SwitchToComponentsMode, "Components", "Switches to Components Mode", EUserInterfaceActionType::ToggleButton, FInputGesture());
@@ -557,19 +561,31 @@ void FFullBlueprintEditorCommands::RegisterCommands()
 namespace BlueprintEditorToolbarImpl
 {
 	static TSharedRef<SWidget> GenerateCompileOptionsWidget(TSharedRef<FUICommandList> CommandList);
+	static void MakeSaveOnCompileSubMenu(FMenuBuilder& InMenuBuilder);
 };
 
-
 static TSharedRef<SWidget> BlueprintEditorToolbarImpl::GenerateCompileOptionsWidget(TSharedRef<FUICommandList> CommandList)
 {
 	FMenuBuilder MenuBuilder(/*bShouldCloseWindowAfterMenuSelection =*/true, CommandList);
 
 	const FFullBlueprintEditorCommands& Commands = FFullBlueprintEditorCommands::Get();
-	MenuBuilder.AddMenuEntry(Commands.SaveOnCompile);
+
+	// @TODO: disable the menu and change up the tooltip when all sub items are disabled
+	MenuBuilder.AddSubMenu(
+		LOCTEXT("SaveOnCompileSubMenu", "Save on Compile"),
+		LOCTEXT("SaveOnCompileSubMenu_ToolTip", "Determines how the Blueprint is saved whenever you compile it"),
+		FNewMenuDelegate::CreateStatic(&BlueprintEditorToolbarImpl::MakeSaveOnCompileSubMenu));
 
 	return MenuBuilder.MakeWidget();
 }
 
+static void BlueprintEditorToolbarImpl::MakeSaveOnCompileSubMenu(FMenuBuilder& InMenuBuilder)
+{
+	const FFullBlueprintEditorCommands& Commands = FFullBlueprintEditorCommands::Get();
+	InMenuBuilder.AddMenuEntry(Commands.SaveOnCompile_Never);
+	InMenuBuilder.AddMenuEntry(Commands.SaveOnCompile_SuccessOnly);
+	InMenuBuilder.AddMenuEntry(Commands.SaveOnCompile_Always);
+}
 
 //////////////////////////////////////////////////////////////////////////
 // FBlueprintEditorToolbar
diff --git a/Engine/Source/Editor/Kismet/Public/BlueprintEditor.h b/Engine/Source/Editor/Kismet/Public/BlueprintEditor.h
--- a/Engine/Source/Editor/Kismet/Public/BlueprintEditor.h
+++ b/Engine/Source/Editor/Kismet/Public/BlueprintEditor.h
@@ -316,8 +316,6 @@ public:
 	virtual void Compile();
 
 	/** Helper functions used to construct/operate-on the "Save on Compile" command */
-	bool IsSaveOnCompileChecked() const;
-	void OnSaveOnCompileToggled() const;
 	bool IsSaveOnCompileEnabled() const;
 
 	/** Calls the above function, but returns an FReply::Handled(). Used in SButtons */
diff --git a/Engine/Source/Editor/Kismet/Public/SBlueprintEditorToolbar.h b/Engine/Source/Editor/Kismet/Public/SBlueprintEditorToolbar.h
--- a/Engine/Source/Editor/Kismet/Public/SBlueprintEditorToolbar.h
+++ b/Engine/Source/Editor/Kismet/Public/SBlueprintEditorToolbar.h
@@ -44,7 +44,9 @@ public:
 
 	/** Compile the blueprint */
 	TSharedPtr<FUICommandInfo> Compile;
-	TSharedPtr<FUICommandInfo> SaveOnCompile;
+	TSharedPtr<FUICommandInfo> SaveOnCompile_Never;
+	TSharedPtr<FUICommandInfo> SaveOnCompile_SuccessOnly;
+	TSharedPtr<FUICommandInfo> SaveOnCompile_Always;
 
 	/** Switch between modes in the blueprint editor */
 	TSharedPtr<FUICommandInfo> SwitchToScriptingMode;
