[UE-3622] (Regression) Can no longer drag-and-drop incompatible nodes from palette into graph.

Pass #1:
- Updated FK2SchemaActionUtils::ExtractNodeTemplateFromAction() to support the new FBlueprintActionMenuItem type.
- Exported the FBlueprintActionMenuItem::GetRawAction() API.

Pass #2:
- Relocated ExtractNodeTemplateFromAction() into the Kismet module and FBlueprintActionMenuUtils.
- Updated existing calls to refer to the new location.

CR: Mike.Beach

diff --git a/Engine/Source/Editor/BlueprintGraph/Classes/EdGraphSchema_K2_Actions.h b/Engine/Source/Editor/BlueprintGraph/Classes/EdGraphSchema_K2_Actions.h
--- a/Engine/Source/Editor/BlueprintGraph/Classes/EdGraphSchema_K2_Actions.h
+++ b/Engine/Source/Editor/BlueprintGraph/Classes/EdGraphSchema_K2_Actions.h
@@ -4,23 +4,6 @@
 #include "K2Node.h"
 #include "EdGraphSchema_K2_Actions.generated.h"
 
-/*******************************************************************************
-* FK2SchemaActionUtils
-*******************************************************************************/
-
-struct FK2SchemaActionUtils
-{
-	/**
-	 * A number of different palette actions hold onto node-templates in different 
-	 * ways. This handles most of those cases and looks to extract said node- 
-	 * template from the specified action.
-	 * 
-	 * @param  PaletteAction	The action you want a node-template for.
-	 * @return A pointer to the extracted node (NULL if the action doesn't have one, or we don't support the specific action type yet)
-	 */
-	static BLUEPRINTGRAPH_API UK2Node const* ExtractNodeTemplateFromAction(TSharedPtr<FEdGraphSchemaAction> PaletteAction);
-};
-
 /*******************************************************************************
 * FEdGraphSchemaAction_K2NewNode
 *******************************************************************************/
diff --git a/Engine/Source/Editor/BlueprintGraph/Private/EdGraphSchema_K2_Actions.cpp b/Engine/Source/Editor/BlueprintGraph/Private/EdGraphSchema_K2_Actions.cpp
--- a/Engine/Source/Editor/BlueprintGraph/Private/EdGraphSchema_K2_Actions.cpp
+++ b/Engine/Source/Editor/BlueprintGraph/Private/EdGraphSchema_K2_Actions.cpp
@@ -25,43 +25,6 @@ namespace
 	const float NodeLiteralHeight = 48.0f;
 }
 
-
-/////////////////////////////////////////////////////
-// FK2SchemaActionUtils
-
-//------------------------------------------------------------------------------
-UK2Node const* FK2SchemaActionUtils::ExtractNodeTemplateFromAction(TSharedPtr<FEdGraphSchemaAction> PaletteAction)
-{
-	UK2Node const* TemplateNode = NULL;
-	if (PaletteAction.IsValid())
-	{
-		FName const ActionId = PaletteAction->GetTypeId();
-		// if this action inherits from FEdGraphSchemaAction_K2NewNode
-		if (ActionId == FEdGraphSchemaAction_K2NewNode::StaticGetTypeId() ||
-			ActionId == FEdGraphSchemaAction_K2AssignDelegate::StaticGetTypeId() ||
-			ActionId == FEdGraphSchemaAction_K2AddComponent::StaticGetTypeId() ||
-			ActionId == FEdGraphSchemaAction_K2AddTimeline::StaticGetTypeId() ||
-			ActionId == FEdGraphSchemaAction_K2AddCustomEvent::StaticGetTypeId() ||
-			ActionId == FEdGraphSchemaAction_K2AddCallOnActor::StaticGetTypeId() ||
-			ActionId == FEdGraphSchemaAction_K2AddCallOnVariable::StaticGetTypeId() ||
-			ActionId == FEdGraphSchemaAction_K2TargetNode::StaticGetTypeId() ||
-			ActionId == FEdGraphSchemaAction_K2PasteHere::StaticGetTypeId() ||
-			ActionId == FEdGraphSchemaAction_K2Event::StaticGetTypeId() || 
-			ActionId == FEdGraphSchemaAction_K2AddEvent::StaticGetTypeId())
-		{
-			FEdGraphSchemaAction_K2NewNode* NewNodeAction = (FEdGraphSchemaAction_K2NewNode*)PaletteAction.Get();
-			TemplateNode = NewNodeAction->NodeTemplate;
-		}
-		else if (ActionId == FEdGraphSchemaAction_K2ViewNode::StaticGetTypeId())
-		{
-			FEdGraphSchemaAction_K2ViewNode* FocusNodeAction = (FEdGraphSchemaAction_K2ViewNode*)PaletteAction.Get();
-			TemplateNode = FocusNodeAction->NodePtr;
-		}
-	}
-
-	return TemplateNode;
-}
-
 /////////////////////////////////////////////////////
 // FEdGraphSchemaAction_K2ViewNode
 
diff --git a/Engine/Source/Editor/Kismet/Private/BlueprintActionMenuUtils.cpp b/Engine/Source/Editor/Kismet/Private/BlueprintActionMenuUtils.cpp
--- a/Engine/Source/Editor/Kismet/Private/BlueprintActionMenuUtils.cpp
+++ b/Engine/Source/Editor/Kismet/Private/BlueprintActionMenuUtils.cpp
@@ -3,6 +3,7 @@
 #include "BlueprintEditorPrivatePCH.h"
 #include "BlueprintActionMenuUtils.h"
 #include "BlueprintActionMenuBuilder.h"
+#include "BlueprintActionMenuItem.h"
 #include "BlueprintActionFilter.h"
 #include "BlueprintNodeSpawner.h"
 #include "BlueprintEditorUtils.h"	// for DoesSupportComponents()
@@ -663,4 +664,42 @@ void FBlueprintActionMenuUtils::MakeFavoritesMenu(FBlueprintActionContext const&
 	} 	
 }
 
+//------------------------------------------------------------------------------
+const UK2Node* FBlueprintActionMenuUtils::ExtractNodeTemplateFromAction(TSharedPtr<FEdGraphSchemaAction> PaletteAction)
+{
+	UK2Node const* TemplateNode = NULL;
+	if (PaletteAction.IsValid())
+	{
+		FName const ActionId = PaletteAction->GetTypeId();
+		if (ActionId == FBlueprintActionMenuItem::StaticGetTypeId())
+		{
+			FBlueprintActionMenuItem* NewNodeActionMenuItem = (FBlueprintActionMenuItem*)PaletteAction.Get();
+			TemplateNode = Cast<UK2Node>(NewNodeActionMenuItem->GetRawAction()->GetTemplateNode());
+		}
+		// if this action inherits from FEdGraphSchemaAction_K2NewNode
+		else if (ActionId == FEdGraphSchemaAction_K2NewNode::StaticGetTypeId() ||
+			ActionId == FEdGraphSchemaAction_K2AssignDelegate::StaticGetTypeId() ||
+			ActionId == FEdGraphSchemaAction_K2AddComponent::StaticGetTypeId() ||
+			ActionId == FEdGraphSchemaAction_K2AddTimeline::StaticGetTypeId() ||
+			ActionId == FEdGraphSchemaAction_K2AddCustomEvent::StaticGetTypeId() ||
+			ActionId == FEdGraphSchemaAction_K2AddCallOnActor::StaticGetTypeId() ||
+			ActionId == FEdGraphSchemaAction_K2AddCallOnVariable::StaticGetTypeId() ||
+			ActionId == FEdGraphSchemaAction_K2TargetNode::StaticGetTypeId() ||
+			ActionId == FEdGraphSchemaAction_K2PasteHere::StaticGetTypeId() ||
+			ActionId == FEdGraphSchemaAction_K2Event::StaticGetTypeId() || 
+			ActionId == FEdGraphSchemaAction_K2AddEvent::StaticGetTypeId())
+		{
+			FEdGraphSchemaAction_K2NewNode* NewNodeAction = (FEdGraphSchemaAction_K2NewNode*)PaletteAction.Get();
+			TemplateNode = NewNodeAction->NodeTemplate;
+		}
+		else if (ActionId == FEdGraphSchemaAction_K2ViewNode::StaticGetTypeId())
+		{
+			FEdGraphSchemaAction_K2ViewNode* FocusNodeAction = (FEdGraphSchemaAction_K2ViewNode*)PaletteAction.Get();
+			TemplateNode = FocusNodeAction->NodePtr;
+		}
+	}
+
+	return TemplateNode;
+}
+
 #undef LOCTEXT_NAMESPACE
diff --git a/Engine/Source/Editor/Kismet/Private/BlueprintPaletteFavorites.cpp b/Engine/Source/Editor/Kismet/Private/BlueprintPaletteFavorites.cpp
--- a/Engine/Source/Editor/Kismet/Private/BlueprintPaletteFavorites.cpp
+++ b/Engine/Source/Editor/Kismet/Private/BlueprintPaletteFavorites.cpp
@@ -6,6 +6,7 @@
 #include "BlueprintActionMenuItem.h"
 #include "BlueprintActionFilter.h"  // for FBlueprintActionInfo
 #include "BlueprintDragDropMenuItem.h"
+#include "BlueprintActionMenuUtils.h"
 
 /*******************************************************************************
  * Static UBlueprintPaletteFavorites Helpers
@@ -58,7 +59,7 @@ static FBlueprintNodeSignature BlueprintPaletteFavoritesImpl::ConstructLegacySig
 		SignatureSubObject = DelegateAction->GetDelegatePoperty();
 	}
 	// if we can pull out a node associated with this action
-	else if (UK2Node const* NodeTemplate = FK2SchemaActionUtils::ExtractNodeTemplateFromAction(InPaletteAction))
+	else if (UK2Node const* NodeTemplate = FBlueprintActionMenuUtils::ExtractNodeTemplateFromAction(InPaletteAction))
 	{
 		bool bIsSupported = false;
 		// now, if we need more info to help identify the node, let's fill 
diff --git a/Engine/Source/Editor/Kismet/Private/SBlueprintPalette.cpp b/Engine/Source/Editor/Kismet/Private/SBlueprintPalette.cpp
--- a/Engine/Source/Editor/Kismet/Private/SBlueprintPalette.cpp
+++ b/Engine/Source/Editor/Kismet/Private/SBlueprintPalette.cpp
@@ -27,6 +27,7 @@
 #include "AnimStateConduitNode.h"
 #include "AnimationTransitionGraph.h"
 #include "BlueprintActionMenuItem.h"
+#include "BlueprintActionMenuUtils.h"
 #include "BlueprintDragDropMenuItem.h"
 #include "BlueprintNodeSpawner.h"
 #include "TutorialMetaData.h"
@@ -296,7 +297,7 @@ static void GetPaletteItemIcon(TSharedPtr<FEdGraphSchemaAction> ActionIn, UBluep
 			ToolTipOut = ComponentClass->GetName();
 		}
 	}
-	else if (UK2Node const* const NodeTemplate = FK2SchemaActionUtils::ExtractNodeTemplateFromAction(ActionIn))
+	else if (UK2Node const* const NodeTemplate = FBlueprintActionMenuUtils::ExtractNodeTemplateFromAction(ActionIn))
 	{
 		// If the node wants to create tooltip text, use that instead, because its probably more detailed
 		FString NodeToolTipText = NodeTemplate->GetTooltipText().ToString();
@@ -1297,7 +1298,7 @@ FText SBlueprintPaletteItem::GetToolTipText() const
 				ToolTipText = ComponentClass->GetToolTipText().ToString();
 			}
 		}
-		else if (UK2Node const* const NodeTemplate = FK2SchemaActionUtils::ExtractNodeTemplateFromAction(PaletteAction))
+		else if (UK2Node const* const NodeTemplate = FBlueprintActionMenuUtils::ExtractNodeTemplateFromAction(PaletteAction))
 		{
 			// Display the native title of the node when alt is held
 			if(FSlateApplication::Get().GetModifierKeys().IsAltDown())
@@ -1407,7 +1408,7 @@ TSharedPtr<SToolTip> SBlueprintPaletteItem::ConstructToolTipWidget() const
 
 	if (PaletteAction.IsValid())
 	{
-		if (UEdGraphNode const* const NodeTemplate = FK2SchemaActionUtils::ExtractNodeTemplateFromAction(PaletteAction))
+		if (UEdGraphNode const* const NodeTemplate = FBlueprintActionMenuUtils::ExtractNodeTemplateFromAction(PaletteAction))
 		{
 			// Take rich tooltip from node
 			DocLink = NodeTemplate->GetDocumentationLink();
@@ -1466,7 +1467,7 @@ TSharedPtr<SToolTip> SBlueprintPaletteItem::ConstructToolTipWidget() const
 	// English speakers have no real need to know this exists.
 	if(FInternationalization::Get().GetCurrentCulture()->GetTwoLetterISOLanguageName() != TEXT("en"))
 	{
-		if (UK2Node const* const NodeTemplate = FK2SchemaActionUtils::ExtractNodeTemplateFromAction(PaletteAction))
+		if (UK2Node const* const NodeTemplate = FBlueprintActionMenuUtils::ExtractNodeTemplateFromAction(PaletteAction))
 		{
 			struct Local
 			{
diff --git a/Engine/Source/Editor/Kismet/Private/SBlueprintSubPalette.cpp b/Engine/Source/Editor/Kismet/Private/SBlueprintSubPalette.cpp
--- a/Engine/Source/Editor/Kismet/Private/SBlueprintSubPalette.cpp
+++ b/Engine/Source/Editor/Kismet/Private/SBlueprintSubPalette.cpp
@@ -14,6 +14,7 @@
 #include "SBlueprintActionMenu.h" // for SBlueprintActionMenuExpander
 #include "BlueprintDragDropMenuItem.h"
 #include "BlueprintActionDatabase.h"
+#include "BlueprintActionMenuUtils.h"
 #include "Engine/LevelScriptBlueprint.h"
 
 #define LOCTEXT_NAMESPACE "BlueprintSubPalette"
@@ -58,7 +59,7 @@ static bool CanPaletteItemBePlaced(TSharedPtr<FEdGraphSchemaAction> DropActionIn
 		bCanBePlaced = false;
 		ImpededReasonOut = LOCTEXT("DropOnlyInGraph", "Nodes can only be placed inside the blueprint graph");
 	}
-	else if (UK2Node const* NodeToBePlaced = FK2SchemaActionUtils::ExtractNodeTemplateFromAction(DropActionIn))
+	else if (UK2Node const* NodeToBePlaced = FBlueprintActionMenuUtils::ExtractNodeTemplateFromAction(DropActionIn))
 	{
 		UEdGraphSchema const* const GraphSchema = HoveredGraphIn->GetSchema();
 		check(GraphSchema != nullptr);
diff --git a/Engine/Source/Editor/Kismet/Public/BlueprintActionMenuUtils.h b/Engine/Source/Editor/Kismet/Public/BlueprintActionMenuUtils.h
--- a/Engine/Source/Editor/Kismet/Public/BlueprintActionMenuUtils.h
+++ b/Engine/Source/Editor/Kismet/Public/BlueprintActionMenuUtils.h
@@ -39,5 +39,15 @@ struct FBlueprintActionMenuUtils
 	 * @param  MenuOut	The structure that will be populated with favorite menu items.
 	 */
 	KISMET_API static void MakeFavoritesMenu(FBlueprintActionContext const& Context, FBlueprintActionMenuBuilder& MenuOut);
+
+	/**
+	 * A number of different palette actions hold onto node-templates in different 
+	 * ways. This handles most of those cases and looks to extract said node- 
+	 * template from the specified action.
+	 * 
+	 * @param  PaletteAction	The action you want a node-template for.
+	 * @return A pointer to the extracted node (NULL if the action doesn't have one, or we don't support the specific action type yet)
+	 */
+	KISMET_API static const UK2Node* ExtractNodeTemplateFromAction(TSharedPtr<FEdGraphSchemaAction> PaletteAction);
 };
 
diff --git a/Engine/Source/Editor/UnrealEd/Private/Commandlets/DumpBlueprintsInfoCommandlet.cpp b/Engine/Source/Editor/UnrealEd/Private/Commandlets/DumpBlueprintsInfoCommandlet.cpp
--- a/Engine/Source/Editor/UnrealEd/Private/Commandlets/DumpBlueprintsInfoCommandlet.cpp
+++ b/Engine/Source/Editor/UnrealEd/Private/Commandlets/DumpBlueprintsInfoCommandlet.cpp
@@ -1705,7 +1705,7 @@ static void DumpBlueprintInfoUtils::DumpActionMenuItem(uint32 Indent, FGraphActi
 		ActionEntry += IndentedNewline + FString::Printf(TEXT("\"Grouping\"    : %d"), PrimeAction->Grouping);
 		
 		// Get action node type info
-		UK2Node const* NodeTemplate = FK2SchemaActionUtils::ExtractNodeTemplateFromAction(PrimeAction);
+		UK2Node const* NodeTemplate = FBlueprintActionMenuUtils::ExtractNodeTemplateFromAction(PrimeAction);
 		if (NodeTemplate != nullptr)
 		{
 			// Build action node type data
diff --git a/Engine/Source/Editor/UnrealEd/Private/Commandlets/GenerateBlueprintAPICommandlet.cpp b/Engine/Source/Editor/UnrealEd/Private/Commandlets/GenerateBlueprintAPICommandlet.cpp
--- a/Engine/Source/Editor/UnrealEd/Private/Commandlets/GenerateBlueprintAPICommandlet.cpp
+++ b/Engine/Source/Editor/UnrealEd/Private/Commandlets/GenerateBlueprintAPICommandlet.cpp
@@ -822,7 +822,7 @@ static void GenerateBlueprintAPIUtils::DumpActionMenuItem(uint32 Indent, FGraphA
 	ActionEntry += IndentedNewline + TooltipFieldLabel + MakeJsonString(TooltipStr) + "\"";
 		
 	// Get action node type info
-	UK2Node const* NodeTemplate = FK2SchemaActionUtils::ExtractNodeTemplateFromAction(PrimeAction);
+	UK2Node const* NodeTemplate = FBlueprintActionMenuUtils::ExtractNodeTemplateFromAction(PrimeAction);
 	if (NodeTemplate != nullptr)
 	{
 		UK2Node* Node = DuplicateObject<UK2Node>(NodeTemplate,ActionList.OwnerOfTemporaries);
