Save  c++ code generated from BP to files - WIP (early version of UI).
Minor fixes in kismet compiler cpp backend.

diff --git a/Engine/Config/BaseEngine.ini b/Engine/Config/BaseEngine.ini
--- a/Engine/Config/BaseEngine.ini
+++ b/Engine/Config/BaseEngine.ini
@@ -1705,6 +1705,7 @@ bUseLocalGraphVariables=false
 bPersistentUberGraphFrame=false
 bReinstanceOnlyWhenNecessary=false
 bDisplaysLayout=false
+bNativeCodeGenerationTool=false
 
 [/Script/Engine.Blueprint]
 bRecompileOnLoad=true
diff --git a/Engine/Source/Editor/Kismet/Private/BlueprintEditor.cpp b/Engine/Source/Editor/Kismet/Private/BlueprintEditor.cpp
--- a/Engine/Source/Editor/Kismet/Private/BlueprintEditor.cpp
+++ b/Engine/Source/Editor/Kismet/Private/BlueprintEditor.cpp
@@ -83,6 +83,7 @@
 #include "GenericCommands.h"
 #include "SNotificationList.h"
 #include "NotificationManager.h"
+#include "NativeCodeGenerationTool.h"
 
 #define LOCTEXT_NAMESPACE "BlueprintEditor"
 
@@ -2038,7 +2039,7 @@ void FBlueprintEditor::CreateDefaultCommands()
 		FIsActionChecked(),
 		FIsActionButtonVisible::CreateSP(this, &FBlueprintEditor::NewDocument_IsVisibleForType, CGT_NewAnimationGraph));
 	*/
-
+	
 	ToolkitCommands->MapAction(FBlueprintEditorCommands::Get().SaveIntermediateBuildProducts,
 		FExecuteAction::CreateSP(this, &FBlueprintEditor::ToggleSaveIntermediateBuildProducts),
 		FCanExecuteAction(),
@@ -2059,6 +2060,25 @@ void FBlueprintEditor::CreateDefaultCommands()
 	ToolkitCommands->MapAction(FBlueprintEditorCommands::Get().BeginBlueprintMerge,
 		FExecuteAction::CreateSP(this, &FBlueprintEditor::CreateMergeToolTab),
 		FCanExecuteAction());
+
+	ToolkitCommands->MapAction(FBlueprintEditorCommands::Get().GenerateNativeCode,
+		FExecuteAction::CreateSP(this, &FBlueprintEditor::OpenNativeCodeGenerationTool),
+		FCanExecuteAction::CreateSP(this, &FBlueprintEditor::CanGenerateNativeCode));
+}
+
+void FBlueprintEditor::OpenNativeCodeGenerationTool()
+{
+	auto Blueprint = GetBlueprintObj();
+	if (Blueprint)
+	{
+		FNativeCodeGenerationTool::Open(*Blueprint, SharedThis(this));
+	}
+}
+
+bool FBlueprintEditor::CanGenerateNativeCode() const
+{
+	auto Blueprint = GetBlueprintObj();
+	return Blueprint && FNativeCodeGenerationTool::CanGenerate(*Blueprint);
 }
 
 void FBlueprintEditor::FindInBlueprint_Clicked()
diff --git a/Engine/Source/Editor/Kismet/Private/BlueprintEditorCommands.cpp b/Engine/Source/Editor/Kismet/Private/BlueprintEditorCommands.cpp
--- a/Engine/Source/Editor/Kismet/Private/BlueprintEditorCommands.cpp
+++ b/Engine/Source/Editor/Kismet/Private/BlueprintEditorCommands.cpp
@@ -62,6 +62,8 @@ void FBlueprintEditorCommands::RegisterCommands()
 	UI_COMMAND( RecompileBlueprintEditor, "Recompile Blueprint Editor", "Recompiles and reloads C++ code for the blueprint editor", EUserInterfaceActionType::Button, FInputGesture() );
 	UI_COMMAND( RecompilePersona, "Recompile Persona", "Recompiles and reloads C++ code for Persona", EUserInterfaceActionType::Button, FInputGesture() );
 
+	UI_COMMAND(GenerateNativeCode, "Generate Native Code", "Generate C++ code from the blueprint", EUserInterfaceActionType::Button, FInputGesture());
+
 	// SCC commands
 	UI_COMMAND( BeginBlueprintMerge, "Merge", "Shows the Blueprint merge panel and toolbar, allowing the user to resolve conflicted blueprints", EUserInterfaceActionType::Button, FInputGesture() );
 }
diff --git a/Engine/Source/Editor/Kismet/Private/BlueprintEditorCommands.h b/Engine/Source/Editor/Kismet/Private/BlueprintEditorCommands.h
--- a/Engine/Source/Editor/Kismet/Private/BlueprintEditorCommands.h
+++ b/Engine/Source/Editor/Kismet/Private/BlueprintEditorCommands.h
@@ -66,6 +66,7 @@ public:
 	TSharedPtr< FUICommandInfo > RecompileKismetCompiler;
 	TSharedPtr< FUICommandInfo > RecompileBlueprintEditor;
 	TSharedPtr< FUICommandInfo > RecompilePersona;
+	TSharedPtr< FUICommandInfo > GenerateNativeCode;
 
 	// SSC commands
 	TSharedPtr< FUICommandInfo > BeginBlueprintMerge;
diff --git a/Engine/Source/Editor/Kismet/Private/NativeCodeGenerationTool.cpp b/Engine/Source/Editor/Kismet/Private/NativeCodeGenerationTool.cpp
new file mode 100644
--- /dev/null
+++ b/Engine/Source/Editor/Kismet/Private/NativeCodeGenerationTool.cpp
@@ -0,0 +1,344 @@
+// Copyright 1998-2014 Epic Games, Inc. All Rights Reserved.
+
+#include "BlueprintEditorPrivatePCH.h"
+#include "NativeCodeGenerationTool.h"
+#include "Dialogs.h"
+
+#define LOCTEXT_NAMESPACE "NativeCodeGenerationTool"
+
+struct FGeneratedCodeData
+{
+	FGeneratedCodeData(UBlueprint& Blueprint)
+	{
+		FName GeneratedClassName, SkeletonClassName;
+		Blueprint.GetBlueprintClassNames(GeneratedClassName, SkeletonClassName);
+		ClassName = GeneratedClassName.ToString();
+
+		GatherUserDefinedDependencies(Blueprint);
+	}
+
+	FString TypeDependencies;
+	FString HeaderSource;
+	FString CppSource;
+	FString ErrorString;
+	FString ClassName;
+
+	void GatherUserDefinedDependencies(UBlueprint& Blueprint)
+	{
+		TArray<UObject*> ReferencedObjects;
+		{
+			FReferenceFinder ReferenceFinder(ReferencedObjects, NULL, false, false, false, false);
+
+			{
+				TArray<UObject*> ObjectsToCheck;
+				GetObjectsWithOuter(Blueprint.GeneratedClass, ObjectsToCheck, true);
+				for (auto Obj : ObjectsToCheck)
+				{
+					if (IsValid(Obj))
+					{
+						ReferenceFinder.FindReferences(Obj);
+					}
+				}
+			}
+
+			for (UClass* Class = Blueprint.GeneratedClass->GetSuperClass(); Class && !Class->HasAnyClassFlags(CLASS_Native); Class = Class->GetSuperClass())
+			{
+				ReferencedObjects.Add(Class);
+			}
+		}
+
+		TypeDependencies.Empty();
+		for (auto Obj : ReferencedObjects)
+		{
+			if (IsValid(Obj) && !Obj->IsIn(Blueprint.GeneratedClass) && (Obj != Blueprint.GeneratedClass))
+			{
+				if (Obj->IsA<UBlueprintGeneratedClass>() || Obj->IsA<UUserDefinedEnum>() || Obj->IsA<UUserDefinedStruct>())
+				{
+					TypeDependencies += Obj->GetPathName();
+					TypeDependencies += TEXT("\n");
+				}
+			}
+		}
+
+		if (TypeDependencies.IsEmpty())
+		{
+			TypeDependencies += LOCTEXT("NoNonNativeDependencies", "No non-native dependencies.").ToString();
+		}
+	}
+
+	static FString DefaultHeaderDir()
+	{
+		auto DefaultSourceDir = FPaths::ConvertRelativePathToFull(FPaths::GameSourceDir());
+		return FPaths::Combine(*DefaultSourceDir, FApp::GetGameName(), TEXT("Classes"));
+	}
+
+	static FString DefaultSourceDir()
+	{
+		auto DefaultSourceDir = FPaths::ConvertRelativePathToFull(FPaths::GameSourceDir());
+		return FPaths::Combine(*DefaultSourceDir, FApp::GetGameName(), TEXT("Private"));
+	}
+
+	FString HeaderFileName() const
+	{
+		return ClassName + TEXT(".h");
+	}
+
+	FString SourceFileName() const
+	{
+		return ClassName + TEXT(".cpp");
+	}
+
+	bool Save(const FString& HeaderDirPath, const FString& CppDirPath)
+	{
+		const bool bHeaderSaved = FFileHelper::SaveStringToFile(HeaderSource, *FPaths::Combine(*HeaderDirPath, *HeaderFileName()));
+		if (!bHeaderSaved)
+		{
+			ErrorString += LOCTEXT("HeaderNotSaved", "Header file wasn't saved. Check log for details.\n").ToString();
+		}
+		const bool bCppSaved = FFileHelper::SaveStringToFile(CppSource, *FPaths::Combine(*CppDirPath, *SourceFileName()));
+		if (!bCppSaved)
+		{
+			ErrorString += LOCTEXT("CppNotSaved", "Cpp file wasn't saved. Check log for details.\n").ToString();
+		}
+
+		return bHeaderSaved && bCppSaved;
+	}
+};
+
+class SSimpleDirectoryPicker : public SCompoundWidget
+{
+public:
+	SLATE_BEGIN_ARGS(SSimpleDirectoryPicker){}
+		SLATE_ARGUMENT(FString, Directory)
+		SLATE_ARGUMENT(FString, File)
+		SLATE_ARGUMENT(FText, Message)
+		SLATE_ATTRIBUTE(bool, IsEnabled)
+	SLATE_END_ARGS()
+
+	FString File;
+	FString Directory;
+	FText Message;
+
+	FString GetFilePath() const
+	{
+		return FPaths::Combine(*Directory, *File);
+	}
+
+	FText GetFilePathText() const
+	{
+		return FText::FromString(GetFilePath());
+	}
+
+	FReply BrowseHeaderDirectory()
+	{
+		PromptUserForDirectory(Directory, Message.ToString(), Directory);
+		return FReply::Handled();
+	}
+
+	void Construct(const FArguments& InArgs)
+	{
+		Directory = InArgs._Directory;
+		File = InArgs._File;
+		Message = InArgs._Message;
+
+		TSharedPtr<SButton> OpenButton;
+		ChildSlot
+		[
+			SNew(SHorizontalBox)
+			+ SHorizontalBox::Slot()
+			.FillWidth(1.0f)
+			.Padding(4.0f, 0.0f, 0.0f, 0.0f)
+			[
+				SNew(SEditableTextBox)
+				.Text(this, &SSimpleDirectoryPicker::GetFilePathText)
+				.IsReadOnly(true)
+			]
+			+ SHorizontalBox::Slot()
+			.AutoWidth()
+			.Padding(4.0f, 0.0f, 0.0f, 0.0f)
+			[
+				SAssignNew(OpenButton, SButton)
+				.ToolTipText(Message)
+				.OnClicked(this, &SSimpleDirectoryPicker::BrowseHeaderDirectory)
+				[
+					SNew(SImage)
+					.Image(FEditorStyle::GetBrush("PropertyWindow.Button_Ellipsis"))
+				]
+			]
+		];
+
+		OpenButton->SetEnabled(InArgs._IsEnabled);
+	}
+};
+
+class SNativeCodeGenerationDialog : public SCompoundWidget
+{
+public:
+	SLATE_BEGIN_ARGS(SNativeCodeGenerationDialog){}
+		SLATE_ARGUMENT(TSharedPtr<SWindow>, ParentWindow)
+		SLATE_ARGUMENT(TSharedPtr<FGeneratedCodeData>, GeneratedCodeData)
+	SLATE_END_ARGS()
+
+private:
+	
+	// Child widgets
+	TSharedPtr<SSimpleDirectoryPicker> HeaderDirectoryBrowser;
+	TSharedPtr<SSimpleDirectoryPicker> SourceDirectoryBrowser;
+	TSharedPtr<SErrorText> ErrorWidget;
+	// 
+	TWeakPtr<SWindow> WeakParentWindow;
+	TSharedPtr<FGeneratedCodeData> GeneratedCodeData;
+	bool bSaved;
+
+	void CloseParentWindow()
+	{
+		auto ParentWindow = WeakParentWindow.Pin();
+		if (ParentWindow.IsValid())
+		{
+			ParentWindow->RequestDestroyWindow();
+		}
+	}
+
+	bool IsEditable() const
+	{
+		return !bSaved && GeneratedCodeData->ErrorString.IsEmpty();
+	}
+
+	FReply OnButtonClicked()
+	{
+		if (IsEditable())
+		{
+			bSaved = GeneratedCodeData->Save(HeaderDirectoryBrowser->Directory, SourceDirectoryBrowser->Directory);
+			ErrorWidget->SetError(GeneratedCodeData->ErrorString);
+		}
+		else
+		{
+			CloseParentWindow();
+		}
+		
+		return FReply::Handled();
+	}
+
+	FText ButtonText() const
+	{
+		return IsEditable() ? LOCTEXT("Generate", "Generate") : LOCTEXT("Close", "Close");
+	}
+
+public:
+
+	void Construct(const FArguments& InArgs)
+	{
+		GeneratedCodeData = InArgs._GeneratedCodeData;
+		bSaved = false;
+		WeakParentWindow = InArgs._ParentWindow;
+
+		ChildSlot
+		[
+			SNew(SBorder)
+			.Padding(4.f)
+			.BorderImage(FEditorStyle::GetBrush("ToolPanel.GroupBorder"))
+			[
+				SNew(SVerticalBox)
+				+ SVerticalBox::Slot()
+				.Padding(4.0f)
+				.AutoHeight()
+				[
+					SNew(STextBlock)
+					.Text(LOCTEXT("HeaderPath", "Header Path").ToString())
+				]
+				+ SVerticalBox::Slot()
+				.Padding(4.0f)
+				.AutoHeight()
+				[
+					SAssignNew(HeaderDirectoryBrowser, SSimpleDirectoryPicker)
+					.Directory(FGeneratedCodeData::DefaultHeaderDir())
+					.File(GeneratedCodeData->HeaderFileName())
+					.Message(LOCTEXT("HeaderDirectory", "Header Directory"))
+					.IsEnabled(this, &SNativeCodeGenerationDialog::IsEditable)
+				]
+				+ SVerticalBox::Slot()
+				.Padding(4.0f)
+				.AutoHeight()
+				[
+					SNew(STextBlock)
+					.Text(LOCTEXT("SourcePath", "Source Path").ToString())
+				]
+				+ SVerticalBox::Slot()
+				.Padding(4.0f)
+				.AutoHeight()
+				[
+					SAssignNew(SourceDirectoryBrowser, SSimpleDirectoryPicker)
+					.Directory(FGeneratedCodeData::DefaultSourceDir())
+					.File(GeneratedCodeData->SourceFileName())
+					.Message(LOCTEXT("SourceDirectory", "Source Directory"))
+					.IsEnabled(this, &SNativeCodeGenerationDialog::IsEditable)
+				]
+				+ SVerticalBox::Slot()
+				.Padding(4.0f)
+				.AutoHeight()
+				[
+					SNew(STextBlock)
+					.Text(LOCTEXT("Dependencies", "Dependencies"))
+				]
+				+ SVerticalBox::Slot()
+				.Padding(4.0f)
+				.AutoHeight()
+				[
+					SNew(SBox)
+					.WidthOverride(360.0f)
+					.HeightOverride(200.0f)
+					[
+						SNew(SMultiLineEditableTextBox)
+						.IsReadOnly(true)
+						.Text(FText::FromString(GeneratedCodeData->TypeDependencies))
+					]
+				]
+				+ SVerticalBox::Slot()
+				.Padding(4.0f)
+				.AutoHeight()
+				[
+					SAssignNew(ErrorWidget, SErrorText)
+				]
+				+ SVerticalBox::Slot()
+				.Padding(4.0f)
+				.AutoHeight()
+				.HAlign(HAlign_Right)
+				[
+					SNew(SButton)
+					.Text(this, &SNativeCodeGenerationDialog::ButtonText)
+					.OnClicked(this, &SNativeCodeGenerationDialog::OnButtonClicked)
+				]
+			]
+		];
+
+		ErrorWidget->SetError(GeneratedCodeData->ErrorString);
+	}
+};
+
+void FNativeCodeGenerationTool::Open(UBlueprint& Blueprint, TSharedRef< class FBlueprintEditor> Editor)
+{
+	TSharedRef<FGeneratedCodeData> GeneratedCodeData(new FGeneratedCodeData(Blueprint));
+
+	TSharedRef<SWindow> PickerWindow = SNew(SWindow)
+		.Title(LOCTEXT("GenerateNativeCode", "Generate Native Code"))
+		.SizingRule(ESizingRule::Autosized)
+		.ClientSize(FVector2D(0.f, 300.f))
+		.SupportsMaximize(false)
+		.SupportsMinimize(false);
+
+	TSharedRef<SNativeCodeGenerationDialog> CodeGenerationDialog = SNew(SNativeCodeGenerationDialog)
+		.ParentWindow(PickerWindow)
+		.GeneratedCodeData(GeneratedCodeData);
+
+	PickerWindow->SetContent(CodeGenerationDialog);
+	GEditor->EditorAddModalWindow(PickerWindow);
+}
+
+bool FNativeCodeGenerationTool::CanGenerate(const UBlueprint& Blueprint)
+{
+	return (Blueprint.Status == EBlueprintStatus::BS_UpToDate)
+		&& (Blueprint.BlueprintType == EBlueprintType::BPTYPE_Normal)
+		&& (Blueprint.GeneratedClass->GetClass() == UBlueprintGeneratedClass::StaticClass());
+}
+
+#undef LOCTEXT_NAMESPACE
\ No newline at end of file
diff --git a/Engine/Source/Editor/Kismet/Private/NativeCodeGenerationTool.h b/Engine/Source/Editor/Kismet/Private/NativeCodeGenerationTool.h
new file mode 100644
--- /dev/null
+++ b/Engine/Source/Editor/Kismet/Private/NativeCodeGenerationTool.h
@@ -0,0 +1,15 @@
+// Copyright 1998-2014 Epic Games, Inc. All Rights Reserved.
+
+#pragma once
+
+#include "BlueprintEditor.h"
+
+//////////////////////////////////////////////////////////////////////////
+// FNativeCodeGenerationTool
+
+struct FNativeCodeGenerationTool
+{
+	static void Open(UBlueprint& Object, TSharedRef< class FBlueprintEditor> Editor);
+
+	static bool CanGenerate(const UBlueprint& Object);
+};
\ No newline at end of file
diff --git a/Engine/Source/Editor/Kismet/Private/SBlueprintEditorToolbar.cpp b/Engine/Source/Editor/Kismet/Private/SBlueprintEditorToolbar.cpp
--- a/Engine/Source/Editor/Kismet/Private/SBlueprintEditorToolbar.cpp
+++ b/Engine/Source/Editor/Kismet/Private/SBlueprintEditorToolbar.cpp
@@ -112,6 +112,16 @@ void FKismet2Menu::FillDeveloperMenu( FMenuBuilder& MenuBuilder )
 	}
 	MenuBuilder.EndSection();
 
+	static const FBoolConfigValueHelper NativeCodeGenerationTool(TEXT("Kismet"), TEXT("bNativeCodeGenerationTool"), GEngineIni);
+	if (NativeCodeGenerationTool)
+	{
+		MenuBuilder.BeginSection("GenerateNativeCode", LOCTEXT("Cpp", "C++"));
+		{
+			MenuBuilder.AddMenuEntry(FBlueprintEditorCommands::Get().GenerateNativeCode);
+		}
+		MenuBuilder.EndSection();
+	}
+
 	if (false)
 	{
 		MenuBuilder.BeginSection("FileDeveloperFindReferences");
diff --git a/Engine/Source/Editor/Kismet/Public/BlueprintEditor.h b/Engine/Source/Editor/Kismet/Public/BlueprintEditor.h
--- a/Engine/Source/Editor/Kismet/Public/BlueprintEditor.h
+++ b/Engine/Source/Editor/Kismet/Public/BlueprintEditor.h
@@ -468,6 +468,12 @@ public:
 	/** Closes the merge tool, rather than simply hiding it */
 	void CloseMergeTool();
 
+	/** Opens a native code generation modal window */
+	void OpenNativeCodeGenerationTool();
+
+	/** Can generate native code for current blueprint */
+	bool CanGenerateNativeCode() const;
+
 	/** 
 	 * Check to see if we can customize the SCS editor for the passed-in scene component 
 	 * @param	InComponentToCustomize	The component to check to see if a customization exists
diff --git a/Engine/Source/Editor/KismetCompiler/Private/KismetCompiler.cpp b/Engine/Source/Editor/KismetCompiler/Private/KismetCompiler.cpp
--- a/Engine/Source/Editor/KismetCompiler/Private/KismetCompiler.cpp
+++ b/Engine/Source/Editor/KismetCompiler/Private/KismetCompiler.cpp
@@ -704,7 +704,7 @@ void FKismetCompilerContext::CreateLocalVariablesForFunction(FKismetFunctionCont
 	ensure(Context.IsEventGraph() || !Context.EventGraphLocals.Num());
 	ensure(!Context.IsEventGraph() || !Context.Locals.Num());
 
-	const bool bPersistentUberGraphFrame = UBlueprintGeneratedClass::UsePersistentUberGraphFrame() && Context.bIsUbergraph;
+	const bool bPersistentUberGraphFrame = UsePersistentUberGraphFrame() && Context.bIsUbergraph;
 	// Local stack frame (or maybe class for the ubergraph)
 	{
 		const bool bArePropertiesLocal = true;
@@ -1179,7 +1179,7 @@ void FKismetCompilerContext::PrecompileFunction(FKismetFunctionContext& Context)
 
 		// Add the function to it's owner class function name -> function map
 		Context.NewClass->AddFunctionToFunctionMap(Context.Function);
-		if (UBlueprintGeneratedClass::UsePersistentUberGraphFrame() && Context.bIsUbergraph)
+		if (UsePersistentUberGraphFrame() && Context.bIsUbergraph)
 		{
 			ensure(!NewClass->UberGraphFunction);
 			NewClass->UberGraphFunction = Context.Function;
@@ -2355,7 +2355,7 @@ void FKismetCompilerContext::CreateFunctionStubForEvent(UK2Node_Event* SrcEventN
 			if (AssignmentNode == NULL)
 			{
 				// Create a variable write node to store the parameters into the ubergraph frame storage
-				if (UBlueprintGeneratedClass::UsePersistentUberGraphFrame())
+				if (UsePersistentUberGraphFrame())
 				{
 					AssignmentNode = SpawnIntermediateNode<UK2Node_SetVariableOnPersistentFrame>(SrcEventNode, ChildStubGraph);
 				}
@@ -3293,7 +3293,7 @@ void FKismetCompilerContext::Compile()
 		}
 	}
 
-	if (UBlueprintGeneratedClass::UsePersistentUberGraphFrame() && UbergraphContext)
+	if (UsePersistentUberGraphFrame() && UbergraphContext)
 	{
 		//UBER GRAPH PERSISTENT FRAME
 		FEdGraphPinType Type(TEXT("struct"), TEXT(""), FPointerToUberGraphFrame::StaticStruct(), false, false);
@@ -3664,5 +3664,10 @@ void FKismetCompilerContext::SetCanEverTickForActor()
 	}
 }
 
+bool FKismetCompilerContext::UsePersistentUberGraphFrame() const
+{
+	return UBlueprintGeneratedClass::UsePersistentUberGraphFrame();
+}
+
 #undef LOCTEXT_NAMESPACE
 //////////////////////////////////////////////////////////////////////////
diff --git a/Engine/Source/Editor/KismetCompiler/Private/KismetCompilerCppBackend.cpp b/Engine/Source/Editor/KismetCompiler/Private/KismetCompilerCppBackend.cpp
--- a/Engine/Source/Editor/KismetCompiler/Private/KismetCompilerCppBackend.cpp
+++ b/Engine/Source/Editor/KismetCompiler/Private/KismetCompilerCppBackend.cpp
@@ -86,6 +86,10 @@ FString FKismetCppBackend::TermToText(const FBPTerminal* Term, const UProperty*
 		{
 			return FString::Printf(TEXT("TEXT(\"%s\")"), *(Term->Name));
 		}
+		else if(CoerceProperty->IsA(UTextProperty::StaticClass()))
+		{
+			return FString::Printf(TEXT("FText::FromString(TEXT(\"%s\"))"), *(Term->Name));
+		}
 		else if (CoerceProperty->IsA(UFloatProperty::StaticClass()))
 		{
 			float Value = FCString::Atof(*(Term->Name));
@@ -434,9 +438,10 @@ struct FEmitHelper
 	{
 		TArray<FString> Tags;
 
-		//TODO: BlueprintNativeEvent, BlueprintImplementableEvent
+		//Pointless: BlueprintNativeEvent, BlueprintImplementableEvent
+		//Pointless: CustomThunk
+
 		//TODO: SealedEvent
-		//TODO: CustomThunk
 		//TODO: Unreliable
 		//TODO: ServiceRequest, ServiceResponse
 		
@@ -535,6 +540,38 @@ struct FEmitHelper
 
 		return Results;
 	}
+
+	static FString MakeCppClassName(UClass* SourceClass)
+	{
+		check(SourceClass);
+		return FString(SourceClass->GetPrefixCPP()) + SourceClass->GetName();
+	}
+
+	static FString EmitLifetimeReplicatedPropsImpl(UClass* SourceClass, const TCHAR* InCurrentIndent)
+	{
+		FString Result;
+		bool bFunctionInitilzed = false;
+		const FString CppClassName = MakeCppClassName(SourceClass);
+		for (TFieldIterator<UProperty> It(SourceClass, EFieldIteratorFlags::ExcludeSuper); It; ++It)
+		{
+			if ((It->PropertyFlags & CPF_Net) != 0)
+			{
+				if (!bFunctionInitilzed)
+				{
+					Result += FString::Printf(TEXT("%svoid %s::GetLifetimeReplicatedProps(TArray< FLifetimeProperty > & OutLifetimeProps) const\n"), InCurrentIndent, *CppClassName);
+					Result += FString::Printf(TEXT("%s{\n"), InCurrentIndent);
+					Result += FString::Printf(TEXT("%s\tSuper::GetLifetimeReplicatedProps(OutLifetimeProps);\n"), InCurrentIndent);
+					bFunctionInitilzed = true;
+				}
+				Result += FString::Printf(TEXT("%s\tDOREPLIFETIME( %s, %s);\n"), InCurrentIndent, *CppClassName, *It->GetNameCPP());
+			}
+		}
+		if (bFunctionInitilzed)
+		{
+			Result += FString::Printf(TEXT("%s}\n"), InCurrentIndent);
+		}
+		return Result;
+	}
 };
 
 void FKismetCppBackend::EmitClassProperties(FStringOutputDevice& Target, UClass* SourceClass)
@@ -565,11 +602,12 @@ void FKismetCppBackend::EmitClassProperties(FStringOutputDevice& Target, UClass*
 
 void FKismetCppBackend::GenerateCodeFromClass(UClass* SourceClass, TIndirectArray<FKismetFunctionContext>& Functions, bool bGenerateStubsOnly)
 {
-	CppClassName = FString(SourceClass->GetPrefixCPP()) + SourceClass->GetName();
+	CppClassName = FEmitHelper::MakeCppClassName(SourceClass);
 
 	UClass* SuperClass = SourceClass->GetSuperClass();
 
 	Emit(Header, TEXT("#pragma once\n\n"));
+	//Emit(Header, TEXT("#inlcude \"Public/Engine.h\"\n"));
 	Emit(Header, *FString::Printf(TEXT("#include \"%s.generated.h\"\n\n"), *SourceClass->GetName()));
 
 	// MC DELEGATE DECLARATION
@@ -651,6 +689,8 @@ void FKismetCppBackend::GenerateCodeFromClass(UClass* SourceClass, TIndirectArra
 	}
 
 	Emit(Header, TEXT("};\n\n"));
+
+	Emit(Body, *FEmitHelper::EmitLifetimeReplicatedPropsImpl(SourceClass, TEXT("")));
 }
 
 void FKismetCppBackend::EmitCallDelegateStatment(FKismetFunctionContext& FunctionContext, FBlueprintCompiledStatement& Statement)
@@ -762,8 +802,7 @@ void FKismetCppBackend::EmitCallStatment(FKismetFunctionContext& FunctionContext
 	Statement.FunctionToCall->GetName(FunctionNameToCall);
 	if( Statement.bIsParentContext )
 	{
-		const FString ContextString = Statement.FunctionToCall->GetOuter()->GetName();
-		FunctionNameToCall = ContextString + TEXT("::") + FunctionNameToCall;
+		FunctionNameToCall = TEXT("Super::") + FunctionNameToCall;
 	}
 	Emit(Body, *FString::Printf(TEXT("%s"), *FunctionNameToCall));
 
diff --git a/Engine/Source/Editor/KismetCompiler/Public/KismetCompiler.h b/Engine/Source/Editor/KismetCompiler/Public/KismetCompiler.h
--- a/Engine/Source/Editor/KismetCompiler/Public/KismetCompiler.h
+++ b/Engine/Source/Editor/KismetCompiler/Public/KismetCompiler.h
@@ -133,6 +133,8 @@ public:
 
 	UK2Node_TemporaryVariable* SpawnInternalVariable(UEdGraphNode* SourceNode, FString Category, FString SubCategory = TEXT(""), UObject* SubcategoryObject = NULL, bool bIsArray = false);
 
+	bool UsePersistentUberGraphFrame() const;
+
 protected:
 	virtual UEdGraphSchema_K2* CreateSchema();
 	virtual void PostCreateSchema();
diff --git a/Engine/Source/Runtime/Engine/Public/GeneratedCodeHelpers.h b/Engine/Source/Runtime/Engine/Public/GeneratedCodeHelpers.h
--- a/Engine/Source/Runtime/Engine/Public/GeneratedCodeHelpers.h
+++ b/Engine/Source/Runtime/Engine/Public/GeneratedCodeHelpers.h
@@ -2,8 +2,12 @@
 
 #pragma once
 
+// Common libraries
 #include "Kismet/KismetArrayLibrary.h"
 #include "Kismet/KismetSystemLibrary.h"
+#include "Kismet/KismetMathLibrary.h"
+
+// Special libraries
 #include "Kismet/DataTableFunctionLibrary.h"
 
 FORCEINLINE UClass* DynamicMetaCast(const UClass* DesiredClass, UClass* SourceClass)
