MathExpression graphs are read-only.

diff --git a/Engine/Source/Editor/BlueprintGraph/Public/BlueprintGraphDefinitions.h b/Engine/Source/Editor/BlueprintGraph/Public/BlueprintGraphDefinitions.h
--- a/Engine/Source/Editor/BlueprintGraph/Public/BlueprintGraphDefinitions.h
+++ b/Engine/Source/Editor/BlueprintGraph/Public/BlueprintGraphDefinitions.h
@@ -44,6 +44,7 @@
 #include "K2Node_MacroInstance.h"
 #include "K2Node_MakeArray.h"
 #include "K2Node_MakeStruct.h"
+#include "K2Node_MathExpression.h"
 #include "K2Node_MatineeController.h"
 #include "K2Node_RemoveDelegate.h"
 #include "K2Node_Select.h"
diff --git a/Engine/Source/Editor/Kismet/Private/BlueprintEditor.cpp b/Engine/Source/Editor/Kismet/Private/BlueprintEditor.cpp
--- a/Engine/Source/Editor/Kismet/Private/BlueprintEditor.cpp
+++ b/Engine/Source/Editor/Kismet/Private/BlueprintEditor.cpp
@@ -6806,8 +6806,9 @@ bool FBlueprintEditor::IsGraphPanelEnabled(UEdGraph* InGraph) const
 {
 	bool const bIsInterface = (FBlueprintEditorUtils::FindBlueprintForGraph(InGraph)->BlueprintType == BPTYPE_Interface);
 	bool const bIsDelegate  = FBlueprintEditorUtils::IsDelegateSignatureGraph(InGraph);
+	bool const bIsMathExpression = FBlueprintEditorUtils::IsMathExpressionGraph(InGraph);
 
-	return !bIsInterface && !bIsDelegate;
+	return !bIsInterface && !bIsDelegate && !bIsMathExpression;
 }
 
 float FBlueprintEditor::GetInstructionTextOpacity(UEdGraph* InGraph) const
diff --git a/Engine/Source/Editor/UnrealEd/Private/Kismet2/BlueprintEditorUtils.cpp b/Engine/Source/Editor/UnrealEd/Private/Kismet2/BlueprintEditorUtils.cpp
--- a/Engine/Source/Editor/UnrealEd/Private/Kismet2/BlueprintEditorUtils.cpp
+++ b/Engine/Source/Editor/UnrealEd/Private/Kismet2/BlueprintEditorUtils.cpp
@@ -2379,6 +2379,15 @@ bool FBlueprintEditorUtils::IsDelegateSignatureGraph(const UEdGraph* Graph)
 	return false;
 }
 
+bool FBlueprintEditorUtils::IsMathExpressionGraph(const UEdGraph* InGraph)
+{
+	if(InGraph)
+	{
+		return InGraph->GetOuter()->GetClass() == UK2Node_MathExpression::StaticClass();
+	}
+	return false;
+}
+
 bool FBlueprintEditorUtils::IsInterfaceBlueprint(const UBlueprint* Blueprint)
 {
 	return (Blueprint->BlueprintType == BPTYPE_Interface);
diff --git a/Engine/Source/Editor/UnrealEd/Public/Kismet2/BlueprintEditorUtils.h b/Engine/Source/Editor/UnrealEd/Public/Kismet2/BlueprintEditorUtils.h
--- a/Engine/Source/Editor/UnrealEd/Public/Kismet2/BlueprintEditorUtils.h
+++ b/Engine/Source/Editor/UnrealEd/Public/Kismet2/BlueprintEditorUtils.h
@@ -442,6 +442,10 @@ public:
 
 	/** Checks if given graph contains a delegate signature */
 	static bool IsDelegateSignatureGraph(const UEdGraph* Graph);
+
+	/** Checks if given graph is owned by a Math Expression node */
+	static bool IsMathExpressionGraph(const UEdGraph* InGraph);
+
 	/**
 	 * Gets a list of pins that should hidden for a given function
 	 *
