Newly created actor BPs will open into the components mode the first time and data only BPs will open into the defaults mode.

diff --git a/Engine/Source/Developer/AssetTools/Private/AssetTypeActions/AssetTypeActions_Blueprint.cpp b/Engine/Source/Developer/AssetTools/Private/AssetTypeActions/AssetTypeActions_Blueprint.cpp
--- a/Engine/Source/Developer/AssetTools/Private/AssetTypeActions/AssetTypeActions_Blueprint.cpp
+++ b/Engine/Source/Developer/AssetTools/Private/AssetTypeActions/AssetTypeActions_Blueprint.cpp
@@ -70,14 +70,8 @@ void FAssetTypeActions_Blueprint::OpenAssetEditor( const TArray<UObject*>& InObj
 		auto Blueprint = Cast<UBlueprint>(*ObjIt);
 		if (Blueprint && Blueprint->SkeletonGeneratedClass && Blueprint->GeneratedClass )
 		{
-			// Open to the full editor
-			const bool bSavedForceFullEditor = Blueprint->bForceFullEditor;
-			Blueprint->bForceFullEditor = true;
-
 			FBlueprintEditorModule& BlueprintEditorModule = FModuleManager::LoadModuleChecked<FBlueprintEditorModule>( "Kismet" );
 			TSharedRef< IBlueprintEditor > NewKismetEditor = BlueprintEditorModule.CreateBlueprintEditor( Mode, EditWithinLevelEditor, Blueprint, ShouldUseDataOnlyEditor(Blueprint) );
-			
-			Blueprint->bForceFullEditor = bSavedForceFullEditor;
 		}
 		else
 		{
diff --git a/Engine/Source/Editor/Kismet/Private/BlueprintEditor.cpp b/Engine/Source/Editor/Kismet/Private/BlueprintEditor.cpp
--- a/Engine/Source/Editor/Kismet/Private/BlueprintEditor.cpp
+++ b/Engine/Source/Editor/Kismet/Private/BlueprintEditor.cpp
@@ -1269,6 +1269,9 @@ void FBlueprintEditor::InitBlueprintEditor(const EToolkitMode::Type Mode, const
 {
 	check(InBlueprints.Num() == 1 || bShouldOpenInDefaultsMode);
 
+	// TRUE if a single Blueprint is being opened and is marked as newly created
+	bool bNewlyCreated = InBlueprints.Num() == 1 && InBlueprints[0]->bIsNewlyCreated;
+
 	TArray< UObject* > Objects;
 	for( auto BlueprintIter = InBlueprints.CreateConstIterator(); BlueprintIter; ++BlueprintIter )
 	{
@@ -1304,7 +1307,7 @@ void FBlueprintEditor::InitBlueprintEditor(const EToolkitMode::Type Mode, const
 	
 	RegenerateMenusAndToolbars();
 
-	RegisterApplicationModes(InBlueprints, bShouldOpenInDefaultsMode);
+	RegisterApplicationModes(InBlueprints, bShouldOpenInDefaultsMode, bNewlyCreated);
 
 	// Cache the project type ( Blueprint or Code Based )
 	if( FPaths::IsProjectFilePathSet() )
@@ -1323,10 +1326,10 @@ void FBlueprintEditor::InitBlueprintEditor(const EToolkitMode::Type Mode, const
 	FBlueprintEditorUtils::FindAndSetDebuggableBlueprintInstances();	
 }
 
-void FBlueprintEditor::RegisterApplicationModes(const TArray<UBlueprint*>& InBlueprints, bool bShouldOpenInDefaultsMode)
+void FBlueprintEditor::RegisterApplicationModes(const TArray<UBlueprint*>& InBlueprints, bool bShouldOpenInDefaultsMode, bool bNewlyCreated/* = false*/)
 {
 	// Newly-created Blueprints will open in Components mode rather than Standard mode
-	bool bShouldOpenInComponentsMode = !bShouldOpenInDefaultsMode && InBlueprints.Num() == 1 && InBlueprints[0]->bIsNewlyCreated;
+	bool bShouldOpenInComponentsMode = !bShouldOpenInDefaultsMode && bNewlyCreated;
 
 	// Create the modes and activate one (which will populate with a real layout)
 	if ( UBlueprint* SingleBP = GetBlueprintObj() )
diff --git a/Engine/Source/Editor/Kismet/Public/BlueprintEditor.h b/Engine/Source/Editor/Kismet/Public/BlueprintEditor.h
--- a/Engine/Source/Editor/Kismet/Public/BlueprintEditor.h
+++ b/Engine/Source/Editor/Kismet/Public/BlueprintEditor.h
@@ -503,7 +503,7 @@ public:
 
 protected:
 	/** Called during initialization of the blueprint editor to register any application modes. */
-	virtual void RegisterApplicationModes(const TArray<UBlueprint*>& InBlueprints, bool bShouldOpenInDefaultsMode);
+	virtual void RegisterApplicationModes(const TArray<UBlueprint*>& InBlueprints, bool bShouldOpenInDefaultsMode, bool bNewlyCreated = false);
 	
 	// Zooming to fit the entire graph
 	void ZoomToWindow_Clicked();
diff --git a/Engine/Source/Editor/UMGEditor/Private/WidgetBlueprintEditor.cpp b/Engine/Source/Editor/UMGEditor/Private/WidgetBlueprintEditor.cpp
--- a/Engine/Source/Editor/UMGEditor/Private/WidgetBlueprintEditor.cpp
+++ b/Engine/Source/Editor/UMGEditor/Private/WidgetBlueprintEditor.cpp
@@ -93,7 +93,7 @@ void FWidgetBlueprintEditor::InitWidgetBlueprintEditor(const EToolkitMode::Type
 		);
 }
 
-void FWidgetBlueprintEditor::RegisterApplicationModes(const TArray<UBlueprint*>& InBlueprints, bool bShouldOpenInDefaultsMode)
+void FWidgetBlueprintEditor::RegisterApplicationModes(const TArray<UBlueprint*>& InBlueprints, bool bShouldOpenInDefaultsMode, bool bNewlyCreated/* = false*/)
 {
 	//FBlueprintEditor::RegisterApplicationModes(InBlueprints, bShouldOpenInDefaultsMode);
 
diff --git a/Engine/Source/Editor/UMGEditor/Private/WidgetBlueprintEditor.h b/Engine/Source/Editor/UMGEditor/Private/WidgetBlueprintEditor.h
--- a/Engine/Source/Editor/UMGEditor/Private/WidgetBlueprintEditor.h
+++ b/Engine/Source/Editor/UMGEditor/Private/WidgetBlueprintEditor.h
@@ -121,7 +121,7 @@ public:
 
 protected:
 	// Begin FBlueprintEditor
-	virtual void RegisterApplicationModes(const TArray<UBlueprint*>& InBlueprints, bool bShouldOpenInDefaultsMode) override;
+	virtual void RegisterApplicationModes(const TArray<UBlueprint*>& InBlueprints, bool bShouldOpenInDefaultsMode, bool bNewlyCreated = false) override;
 	virtual FGraphAppearanceInfo GetGraphAppearance() const override;
 	// End FBlueprintEditor
 
