[UE-2889] New option to focus on error'ing nodes after manually compiling a blueprint.

diff --git a/Engine/Source/Editor/BlueprintGraph/Private/BlueprintEditorSettings.cpp b/Engine/Source/Editor/BlueprintGraph/Private/BlueprintEditorSettings.cpp
--- a/Engine/Source/Editor/BlueprintGraph/Private/BlueprintEditorSettings.cpp
+++ b/Engine/Source/Editor/BlueprintGraph/Private/BlueprintEditorSettings.cpp
@@ -18,6 +18,7 @@ UBlueprintEditorSettings::UBlueprintEditorSettings(const FObjectInitializer& Obj
 	, bUseLegacyMenuingSystem(false)
 	// Compiler Settings
 	, SaveOnCompile(SoC_Never)
+	, bJumpToNodeErrors(false)
 	// Developer Settings
 	, bShowActionMenuItemSignatures(false)
 	// Perf Settings
diff --git a/Engine/Source/Editor/BlueprintGraph/Public/BlueprintEditorSettings.h b/Engine/Source/Editor/BlueprintGraph/Public/BlueprintEditorSettings.h
--- a/Engine/Source/Editor/BlueprintGraph/Public/BlueprintEditorSettings.h
+++ b/Engine/Source/Editor/BlueprintGraph/Public/BlueprintEditorSettings.h
@@ -57,6 +57,10 @@ public:
 	UPROPERTY(EditAnywhere, config, Category=Compiler)
 	TEnumAsByte<ESaveOnCompile> SaveOnCompile;
 
+	/** When enabled, if a blueprint has compiler errors, then the graph will jump and focus on the first node generating an error */
+	UPROPERTY(EditAnywhere, config, Category=Compiler)
+	bool bJumpToNodeErrors;
+
 // Developer Settings
 public:
 	/** If enabled, tooltips on action menu items will show the associated action's signature id (can be used to setup custom favorites menus).*/
diff --git a/Engine/Source/Editor/Kismet/Private/BlueprintEditor.cpp b/Engine/Source/Editor/Kismet/Private/BlueprintEditor.cpp
--- a/Engine/Source/Editor/Kismet/Private/BlueprintEditor.cpp
+++ b/Engine/Source/Editor/Kismet/Private/BlueprintEditor.cpp
@@ -214,6 +214,25 @@ namespace BlueprintEditorImpl
 	 * @return False if the option isn't set, or if the save-on-compile is disabled for the blueprint being edited (otherwise true). 
 	 */
 	static bool IsSaveOnCompileOptionSet(TWeakPtr<FBlueprintEditor> Editor, ESaveOnCompile Option);
+
+	/**  Flips the value of the editor's "JumpToNodeErrors" setting. */
+	static void ToggleJumpToErrorNodeSetting();
+
+	/**
+	 * Utility function that will check to see if the "Jump to Error Nodes" 
+	 * setting is enabled.
+	 * 
+	 * @return True if UBlueprintEditorSettings::bJumpToNodeErrors is set, otherwise false.
+	 */
+	static bool IsJumpToErrorNodeOptionSet();
+
+	/**
+	 * Searches through a blueprint, looking for the most severe error'ing node.
+	 * 
+	 * @param  Blueprint	The blueprint to search through.
+	 * @param  Severity		Defines the severity of the error/warning to search for.
+	 */
+	static UEdGraphNode* FindNodeWithError(UBlueprint* Blueprint, EMessageSeverity::Type Severity = EMessageSeverity::Error);
 }
 
 static bool BlueprintEditorImpl::GraphHasUserPlacedNodes(UEdGraph const* InGraph)
@@ -282,6 +301,45 @@ static bool BlueprintEditorImpl::IsSaveOnCompileOptionSet(TWeakPtr<FBlueprintEdi
 	return (CurrentSetting == Option);
 }
 
+static void BlueprintEditorImpl::ToggleJumpToErrorNodeSetting()
+{
+	UBlueprintEditorSettings* Settings = GetMutableDefault<UBlueprintEditorSettings>();
+	Settings->bJumpToNodeErrors = !Settings->bJumpToNodeErrors;
+	Settings->SaveConfig();
+}
+
+static bool BlueprintEditorImpl::IsJumpToErrorNodeOptionSet()
+{
+	const UBlueprintEditorSettings* Settings = GetDefault<UBlueprintEditorSettings>();
+	return Settings->bJumpToNodeErrors;
+}
+
+static UEdGraphNode* BlueprintEditorImpl::FindNodeWithError(UBlueprint* Blueprint, EMessageSeverity::Type Severity/* = EMessageSeverity::Error*/)
+{
+	TArray<UEdGraph*> Graphs;
+	Blueprint->GetAllGraphs(Graphs);
+
+	UEdGraphNode* ChoiceNode = nullptr;
+	for (UEdGraph* Graph : Graphs)
+	{
+		for (UEdGraphNode* Node : Graph->Nodes)
+		{
+			if (Node->bHasCompilerMessage && !Node->ErrorMsg.IsEmpty() && (Node->ErrorType <= Severity))
+			{
+				if ((ChoiceNode == nullptr) || (ChoiceNode->ErrorType > Node->ErrorType))
+				{
+					ChoiceNode = Node;
+					if (ChoiceNode->ErrorType == 0)
+					{
+						break;
+					}
+				}
+			}
+		}
+	}
+	return ChoiceNode;
+}
+
 bool FBlueprintEditor::IsASubGraph( const UEdGraph* GraphPtr )
 {
 	if( GraphPtr && GraphPtr->GetOuter() )
@@ -1662,6 +1720,19 @@ void FBlueprintEditor::SetupViewForBlueprintEditingMode()
 	// Make sure the inspector is always on top
 	//@TODO: This is necessary right now because of a bug in restoring layouts not remembering which tab is on top (to get it right initially), but do we want this behavior always?
 	TabManager->InvokeTab(FBlueprintEditorTabs::DetailsID);
+
+	UBlueprint* Blueprint = GetBlueprintObj();
+	if ((Blueprint != nullptr) && (Blueprint->Status == EBlueprintStatus::BS_Error))
+	{
+		UBlueprintEditorSettings const* BpEditorSettings = GetDefault<UBlueprintEditorSettings>();
+		if (BpEditorSettings->bJumpToNodeErrors)
+		{
+			if (UEdGraphNode* NodeWithError = BlueprintEditorImpl::FindNodeWithError(Blueprint))
+			{
+				JumpToNode(NodeWithError, /*bRequestRename =*/false);
+			}
+		}
+	}
 }
 
 FBlueprintEditor::~FBlueprintEditor()
@@ -1856,6 +1927,13 @@ void FBlueprintEditor::CreateDefaultCommands()
 		FCanExecuteAction::CreateSP(this, &FBlueprintEditor::IsSaveOnCompileEnabled),
 		FIsActionChecked::CreateStatic(&BlueprintEditorImpl::IsSaveOnCompileOptionSet, WeakThisPtr, (ESaveOnCompile)SoC_Always)
 	);
+
+	ToolkitCommands->MapAction(
+		FFullBlueprintEditorCommands::Get().JumpToErrorNode,
+		FExecuteAction::CreateStatic(&BlueprintEditorImpl::ToggleJumpToErrorNodeSetting),
+		FCanExecuteAction(),
+		FIsActionChecked::CreateStatic(&BlueprintEditorImpl::IsJumpToErrorNodeOptionSet)
+	);
 	
 	ToolkitCommands->MapAction(
 		FFullBlueprintEditorCommands::Get().SwitchToScriptingMode,
@@ -2537,6 +2615,15 @@ void FBlueprintEditor::Compile()
 		bool bForceMessageDisplay = ((LogResults.NumWarnings > 0) || (LogResults.NumErrors > 0)) && !BlueprintObj->bIsRegeneratingOnLoad;
 		DumpMessagesToCompilerLog(LogResults.Messages, bForceMessageDisplay);
 
+		UBlueprintEditorSettings const* BpEditorSettings = GetDefault<UBlueprintEditorSettings>();
+		if ((LogResults.NumErrors > 0) && BpEditorSettings->bJumpToNodeErrors)
+		{
+			if (UEdGraphNode* NodeWithError = BlueprintEditorImpl::FindNodeWithError(BlueprintObj))
+			{
+				JumpToNode(NodeWithError, /*bRequestRename =*/false);
+			}
+		}
+
 		// send record when player clicks compile and send the result
 		// this will make sure how the users activity is
 		AnalyticsTrackCompileEvent(BlueprintObj, LogResults.NumErrors, LogResults.NumWarnings);
@@ -6353,19 +6440,20 @@ void FBlueprintEditor::RestoreEditedObjectState()
 {
 	check(IsEditingSingleBlueprint());	
 
-	if (GetBlueprintObj()->LastEditedDocuments.Num() == 0)
+	UBlueprint* Blueprint = GetBlueprintObj();
+	if (Blueprint->LastEditedDocuments.Num() == 0)
 	{
-		if(FBlueprintEditorUtils::SupportsConstructionScript(GetBlueprintObj()))
+		if (FBlueprintEditorUtils::SupportsConstructionScript(Blueprint))
 		{
-			GetBlueprintObj()->LastEditedDocuments.Add(FBlueprintEditorUtils::FindUserConstructionScript(GetBlueprintObj()));
+			Blueprint->LastEditedDocuments.Add(FBlueprintEditorUtils::FindUserConstructionScript(Blueprint));
 		}
 
-		GetBlueprintObj()->LastEditedDocuments.Add(FBlueprintEditorUtils::FindEventGraph(GetBlueprintObj()));
+		Blueprint->LastEditedDocuments.Add(FBlueprintEditorUtils::FindEventGraph(Blueprint));
 	}
 
-	for (int32 i = 0; i < GetBlueprintObj()->LastEditedDocuments.Num(); i++)
+	for (int32 i = 0; i < Blueprint->LastEditedDocuments.Num(); i++)
 	{
-		if (UObject* Obj = GetBlueprintObj()->LastEditedDocuments[i].EditedObject)
+		if (UObject* Obj = Blueprint->LastEditedDocuments[i].EditedObject)
 		{
 			if(UEdGraph* Graph = Cast<UEdGraph>(Obj))
 			{
@@ -6397,7 +6485,7 @@ void FBlueprintEditor::RestoreEditedObjectState()
 				TSharedPtr<SDockTab> TabWithGraph = LocalStruct::OpenGraphTree(this, Graph);
 
 				TSharedRef<SGraphEditor> GraphEditor = StaticCastSharedRef<SGraphEditor>(TabWithGraph->GetContent());
-				GraphEditor->SetViewLocation(GetBlueprintObj()->LastEditedDocuments[i].SavedViewOffset, GetBlueprintObj()->LastEditedDocuments[i].SavedZoomAmount);
+				GraphEditor->SetViewLocation(Blueprint->LastEditedDocuments[i].SavedViewOffset, Blueprint->LastEditedDocuments[i].SavedZoomAmount);
 			}
 			else
 			{
diff --git a/Engine/Source/Editor/Kismet/Private/BlueprintEditorCommands.cpp b/Engine/Source/Editor/Kismet/Private/BlueprintEditorCommands.cpp
--- a/Engine/Source/Editor/Kismet/Private/BlueprintEditorCommands.cpp
+++ b/Engine/Source/Editor/Kismet/Private/BlueprintEditorCommands.cpp
@@ -55,7 +55,7 @@ void FBlueprintEditorCommands::RegisterCommands()
 	UI_COMMAND( AddNewDelegate, "Event Dispatcher", "Add a new event dispatcher", EUserInterfaceActionType::Button, FInputGesture() );
 
 	// Development commands
-	UI_COMMAND( SaveIntermediateBuildProducts, "Save Intermediate Build Products", "Should the compiler save intermediate build products for debugging.", EUserInterfaceActionType::Check, FInputGesture() );
+	UI_COMMAND( SaveIntermediateBuildProducts, "Save Intermediate Build Products", "Should the compiler save intermediate build products for debugging.", EUserInterfaceActionType::ToggleButton, FInputGesture() );
 
 	UI_COMMAND( RecompileGraphEditor, "Recompile Graph Editor", "Recompiles and reloads C++ code for the graph editor", EUserInterfaceActionType::Button, FInputGesture() );
 	UI_COMMAND( RecompileKismetCompiler, "Recompile Blueprint Compiler", "Recompiles and reloads C++ code for the blueprint compiler", EUserInterfaceActionType::Button, FInputGesture() );
diff --git a/Engine/Source/Editor/Kismet/Private/SBlueprintEditorToolbar.cpp b/Engine/Source/Editor/Kismet/Private/SBlueprintEditorToolbar.cpp
--- a/Engine/Source/Editor/Kismet/Private/SBlueprintEditorToolbar.cpp
+++ b/Engine/Source/Editor/Kismet/Private/SBlueprintEditorToolbar.cpp
@@ -563,6 +563,8 @@ void FFullBlueprintEditorCommands::RegisterCommands()
 	UI_COMMAND(SwitchToBlueprintDefaultsMode, "Defaults", "Switches to Blueprint Defaults Mode", EUserInterfaceActionType::ToggleButton, FInputGesture());
 	UI_COMMAND(SwitchToComponentsMode, "Components", "Switches to Components Mode", EUserInterfaceActionType::ToggleButton, FInputGesture());
 	UI_COMMAND(EditGlobalOptions, "Blueprint Props", "Edit Blueprint Options", EUserInterfaceActionType::Button, FInputGesture());
+
+	UI_COMMAND(JumpToErrorNode, "Jump to Error Node", "When enabled, then the Blueprint will snap focus to nodes producing an error during compilation", EUserInterfaceActionType::ToggleButton, FInputGesture());
 }
 
 //////////////////////////////////////////////////////////////////////////
@@ -572,6 +574,7 @@ namespace BlueprintEditorToolbarImpl
 {
 	static TSharedRef<SWidget> GenerateCompileOptionsWidget(TSharedRef<FUICommandList> CommandList);
 	static void MakeSaveOnCompileSubMenu(FMenuBuilder& InMenuBuilder);
+	static void MakeCompileDeveloperSubMenu(FMenuBuilder& InMenuBuilder);
 };
 
 static TSharedRef<SWidget> BlueprintEditorToolbarImpl::GenerateCompileOptionsWidget(TSharedRef<FUICommandList> CommandList)
@@ -583,9 +586,16 @@ static TSharedRef<SWidget> BlueprintEditorToolbarImpl::GenerateCompileOptionsWid
 	// @TODO: disable the menu and change up the tooltip when all sub items are disabled
 	MenuBuilder.AddSubMenu(
 		LOCTEXT("SaveOnCompileSubMenu", "Save on Compile"),
-		LOCTEXT("SaveOnCompileSubMenu_ToolTip", "Determines how the Blueprint is saved whenever you compile it"),
+		LOCTEXT("SaveOnCompileSubMenu_ToolTip", "Determines how the Blueprint is saved whenever you compile it."),
 		FNewMenuDelegate::CreateStatic(&BlueprintEditorToolbarImpl::MakeSaveOnCompileSubMenu));
 
+	MenuBuilder.AddMenuEntry(Commands.JumpToErrorNode);
+
+	MenuBuilder.AddSubMenu(
+		LOCTEXT("DevCompileSubMenu", "Developer"),
+		LOCTEXT("DevCompileSubMenu_ToolTip", "Advanced settings that aid in devlopment/debugging of the Blueprint system as a whole."),
+		FNewMenuDelegate::CreateStatic(&BlueprintEditorToolbarImpl::MakeCompileDeveloperSubMenu));
+
 	return MenuBuilder.MakeWidget();
 }
 
@@ -597,6 +607,13 @@ static void BlueprintEditorToolbarImpl::MakeSaveOnCompileSubMenu(FMenuBuilder& I
 	InMenuBuilder.AddMenuEntry(Commands.SaveOnCompile_Always);
 }
 
+static void BlueprintEditorToolbarImpl::MakeCompileDeveloperSubMenu(FMenuBuilder& InMenuBuilder)
+{
+	const FBlueprintEditorCommands& EditorCommands = FBlueprintEditorCommands::Get();
+	InMenuBuilder.AddMenuEntry(EditorCommands.SaveIntermediateBuildProducts);
+}
+
+
 //////////////////////////////////////////////////////////////////////////
 // FBlueprintEditorToolbar
 
diff --git a/Engine/Source/Editor/Kismet/Public/SBlueprintEditorToolbar.h b/Engine/Source/Editor/Kismet/Public/SBlueprintEditorToolbar.h
--- a/Engine/Source/Editor/Kismet/Public/SBlueprintEditorToolbar.h
+++ b/Engine/Source/Editor/Kismet/Public/SBlueprintEditorToolbar.h
@@ -47,6 +47,7 @@ public:
 	TSharedPtr<FUICommandInfo> SaveOnCompile_Never;
 	TSharedPtr<FUICommandInfo> SaveOnCompile_SuccessOnly;
 	TSharedPtr<FUICommandInfo> SaveOnCompile_Always;
+	TSharedPtr<FUICommandInfo> JumpToErrorNode;
 
 	/** Switch between modes in the blueprint editor */
 	TSharedPtr<FUICommandInfo> SwitchToScriptingMode;
diff --git a/Engine/Source/Editor/UnrealEd/Private/Kismet2/Kismet2.cpp b/Engine/Source/Editor/UnrealEd/Private/Kismet2/Kismet2.cpp
--- a/Engine/Source/Editor/UnrealEd/Private/Kismet2/Kismet2.cpp
+++ b/Engine/Source/Editor/UnrealEd/Private/Kismet2/Kismet2.cpp
@@ -483,9 +483,9 @@ void FKismetEditorUtilities::CompileBlueprint(UBlueprint* BlueprintObj, bool bIs
 			((Settings->SaveOnCompile == SoC_SuccessOnly) && (BlueprintObj->Status == BS_UpToDate)) )
 		{
 			bool const bIsLevelPackage = (UWorld::FindWorldInPackage(BlueprintPackage) != nullptr);
-			// we don't want to save the entire level (especially is this 
+			// we don't want to save the entire level (especially if this 
 			// compile was already kicked off as a result of a level save, as it
-			// could cause a recursive save)... let the "bSaveOnCompile" setting 
+			// could cause a recursive save)... let the "SaveOnCompile" setting 
 			// only save blueprint assets
 			if (!bIsLevelPackage)
 			{
