Right click context menu for deleting nodes will gray out if none of the selected nodes can be deleted.

diff --git a/Engine/Source/Editor/Kismet/Private/BlueprintEditor.cpp b/Engine/Source/Editor/Kismet/Private/BlueprintEditor.cpp
--- a/Engine/Source/Editor/Kismet/Private/BlueprintEditor.cpp
+++ b/Engine/Source/Editor/Kismet/Private/BlueprintEditor.cpp
@@ -4445,7 +4445,23 @@ bool FBlueprintEditor::CanDeleteNodes() const
 {
 	const FGraphPanelSelectionSet SelectedNodes = GetSelectedNodes();
 
-	return IsEditable(GetFocusedGraph()) && (SelectedNodes.Num() > 0);
+	bool bCanUserDeleteNode = false;
+
+	if(IsEditable(GetFocusedGraph()) && SelectedNodes.Num() > 0)
+	{
+		for( UObject* NodeObject : SelectedNodes )
+		{
+			// If any nodes allow deleting, then do not disable the delete option
+			UEdGraphNode* Node = Cast<UEdGraphNode>(NodeObject);
+			if(Node->CanUserDeleteNode())
+			{
+				bCanUserDeleteNode = true;
+				break;
+			}
+		}
+	}
+
+	return bCanUserDeleteNode;
 }
 
 void FBlueprintEditor::DeleteSelectedDuplicatableNodes()
