Selecting custom events in the My Blueprint window will function as if selecting the event node.

diff --git a/Engine/Source/Editor/Kismet/Private/BlueprintDetailsCustomization.cpp b/Engine/Source/Editor/Kismet/Private/BlueprintDetailsCustomization.cpp
--- a/Engine/Source/Editor/Kismet/Private/BlueprintDetailsCustomization.cpp
+++ b/Engine/Source/Editor/Kismet/Private/BlueprintDetailsCustomization.cpp
@@ -4701,21 +4701,12 @@ void FBlueprintComponentDetails::OnSocketSelection( FName SocketName )
 BEGIN_SLATE_FUNCTION_BUILD_OPTIMIZATION
 void FBlueprintGraphNodeDetails::CustomizeDetails( IDetailLayoutBuilder& DetailLayout )
 {
-	// Cache the GraphNode
-	if( BlueprintEditorPtr.IsValid() )
+	const TArray<TWeakObjectPtr<UObject>> SelectedObjects = DetailLayout.GetDetailsView().GetSelectedObjects();
+	if( SelectedObjects.Num() == 1 )
 	{
-		/** Get the currently selected set of nodes */
-		TSet<UObject*> Objects = BlueprintEditorPtr.Pin()->GetSelectedNodes();
-
-		if (Objects.Num() == 1)
+		if (SelectedObjects[0].IsValid() && SelectedObjects[0]->IsA<UEdGraphNode>())
 		{
-			TSet<UObject*>::TIterator Iter(Objects);
-			UObject* Object = *Iter;
-
-			if (Object && Object->IsA<UEdGraphNode>())
-			{
-				GraphNodePtr = Cast<UEdGraphNode>(Object);
-			}
+			GraphNodePtr = Cast<UEdGraphNode>(SelectedObjects[0].Get());
 		}
 	}
 
diff --git a/Engine/Source/Editor/Kismet/Private/SKismetInspector.cpp b/Engine/Source/Editor/Kismet/Private/SKismetInspector.cpp
--- a/Engine/Source/Editor/Kismet/Private/SKismetInspector.cpp
+++ b/Engine/Source/Editor/Kismet/Private/SKismetInspector.cpp
@@ -279,15 +279,13 @@ void SKismetInspector::EnableComponentDetailsCustomization(bool bEnable)
 void SKismetInspector::ShowDetailsForSingleObject(UObject* Object, const FShowDetailsOptions& Options)
 {
 	TArray<UObject*> PropertyObjects;
-	FKismetSelectionInfo SelectionInfo;
 
 	if (Object != NULL)
 	{
 		PropertyObjects.Add(Object);
-		SelectionInfo.ObjectsForPropertyEditing.Add(Object);
 	}
 
-	UpdateFromObjects(PropertyObjects, SelectionInfo, Options);
+	ShowDetailsForObjects(PropertyObjects, Options);
 }
 
 void SKismetInspector::ShowDetailsForObjects(const TArray<UObject*>& PropertyObjects, const FShowDetailsOptions& Options)
diff --git a/Engine/Source/Editor/Kismet/Private/SMyBlueprint.cpp b/Engine/Source/Editor/Kismet/Private/SMyBlueprint.cpp
--- a/Engine/Source/Editor/Kismet/Private/SMyBlueprint.cpp
+++ b/Engine/Source/Editor/Kismet/Private/SMyBlueprint.cpp
@@ -1227,7 +1227,7 @@ void SMyBlueprint::OnActionSelected( const TArray< TSharedPtr<FEdGraphSchemaActi
 
 			Inspector->ShowDetailsForSingleObject(StructAction->Struct, Options);
 		}
-		else if (InAction->GetTypeId() == FEdGraphSchemaAction_K2TargetNode::StaticGetTypeId())
+		else if (InAction->GetTypeId() == FEdGraphSchemaAction_K2TargetNode::StaticGetTypeId() || InAction->GetTypeId() == FEdGraphSchemaAction_K2Event::StaticGetTypeId())
 		{
 			FEdGraphSchemaAction_K2TargetNode* TargetNodeAction = (FEdGraphSchemaAction_K2TargetNode*)InAction.Get();
 			SKismetInspector::FShowDetailsOptions Options(TargetNodeAction->NodeTemplate->GetNodeTitle(ENodeTitleType::EditableTitle).ToString());
