Can no longer directly modify input or output pins on Math Expression nodes.

diff --git a/Engine/Source/Editor/Kismet/Private/BlueprintDetailsCustomization.cpp b/Engine/Source/Editor/Kismet/Private/BlueprintDetailsCustomization.cpp
--- a/Engine/Source/Editor/Kismet/Private/BlueprintDetailsCustomization.cpp
+++ b/Engine/Source/Editor/Kismet/Private/BlueprintDetailsCustomization.cpp
@@ -1739,7 +1739,7 @@ void FBlueprintGraphArgumentLayout::GenerateHeaderRowContent( FDetailWidgetRow&
 		.HAlign(HAlign_Right)
 		.VAlign(VAlign_Center)
 		[
-			PropertyCustomizationHelpers::MakeClearButton(FSimpleDelegate::CreateSP(this, &FBlueprintGraphArgumentLayout::OnRemoveClicked))
+			PropertyCustomizationHelpers::MakeClearButton(FSimpleDelegate::CreateSP(this, &FBlueprintGraphArgumentLayout::OnRemoveClicked), FText(), !IsPinEditingReadOnly())
 		]
 	];
 }
@@ -1782,38 +1782,42 @@ void FBlueprintGraphArgumentLayout::GenerateChildContent( IDetailChildrenBuilder
 		];
 	}
 		
-	ChildrenBuilder.AddChildContent( *LOCTEXT( "FunctionArgDetailsMoving", "Moving" ).ToString() )
-	[
-		SNew(SHorizontalBox)
-		+SHorizontalBox::Slot()
-		[
-			SNew(SSpacer)
-		]
-		+SHorizontalBox::Slot()
-		.AutoWidth()
-		.Padding(2, 0)
-		[
-			SNew(SButton)
-			.ContentPadding(0)
-			.OnClicked(this, &FBlueprintGraphArgumentLayout::OnArgMoveUp)
-			[
-				SNew(SImage)
-				.Image(FEditorStyle::GetBrush("BlueprintEditor.Details.ArgUpButton"))
-			]
-		]
-		+SHorizontalBox::Slot()
-		.AutoWidth()
-		.Padding(2, 0)
-		[
-			SNew(SButton)
-			.ContentPadding(0)
-			.OnClicked(this, &FBlueprintGraphArgumentLayout::OnArgMoveDown)
+	// Read only graphs can't have their pins re-organized
+	if ( !IsPinEditingReadOnly() )
+	{
+		ChildrenBuilder.AddChildContent( *LOCTEXT( "FunctionArgDetailsMoving", "Moving" ).ToString() )
 			[
-				SNew(SImage)
-				.Image(FEditorStyle::GetBrush("BlueprintEditor.Details.ArgDownButton"))
-			]
-		]
-	];
+				SNew(SHorizontalBox)
+				+SHorizontalBox::Slot()
+				[
+					SNew(SSpacer)
+				]
+				+SHorizontalBox::Slot()
+					.AutoWidth()
+					.Padding(2, 0)
+					[
+						SNew(SButton)
+						.ContentPadding(0)
+						.OnClicked(this, &FBlueprintGraphArgumentLayout::OnArgMoveUp)
+						[
+							SNew(SImage)
+							.Image(FEditorStyle::GetBrush("BlueprintEditor.Details.ArgUpButton"))
+						]
+					]
+				+SHorizontalBox::Slot()
+					.AutoWidth()
+					.Padding(2, 0)
+					[
+						SNew(SButton)
+						.ContentPadding(0)
+						.OnClicked(this, &FBlueprintGraphArgumentLayout::OnArgMoveDown)
+						[
+							SNew(SImage)
+							.Image(FEditorStyle::GetBrush("BlueprintEditor.Details.ArgDownButton"))
+						]
+					]
+			];
+	}
 }
 
 void FBlueprintGraphArgumentLayout::OnRemoveClicked()
@@ -1877,11 +1881,29 @@ bool FBlueprintGraphArgumentLayout::ShouldPinBeReadOnly() const
 		{
 			return true;
 		}
+		else
+		{
+			// Check if pin editing is read only
+			return IsPinEditingReadOnly();
+		}
 	}
 	
 	return false;
 }
 
+bool FBlueprintGraphArgumentLayout::IsPinEditingReadOnly() const
+{
+	if(UEdGraph* NodeGraph = TargetNode->GetGraph())
+	{
+		// Math expression should not be modified directly, do not let the user tweak the parameters
+		if ( Cast<UK2Node_MathExpression>(NodeGraph->GetOuter()) )
+		{
+			return true;
+		}
+	}
+	return false;
+}
+
 FText FBlueprintGraphArgumentLayout::OnGetArgNameText() const
 {
 	if (ParamItemPtr.IsValid())
@@ -2330,6 +2352,7 @@ void FBlueprintGraphActionDetails::CustomizeDetails( IDetailLayoutBuilder& Detai
 				SNew(SButton)
 				.Text(LOCTEXT("FunctionNewInputArg", "New").ToString())
 				.OnClicked(this, &FBlueprintGraphActionDetails::OnAddNewInputClicked)
+				.Visibility(this, &FBlueprintGraphActionDetails::GetAddNewInputOutputVisibility)
 			]
 		];
 
@@ -2352,6 +2375,7 @@ void FBlueprintGraphActionDetails::CustomizeDetails( IDetailLayoutBuilder& Detai
 					SNew(SButton)
 					.Text(LOCTEXT("FunctionNewOutputArg", "New").ToString())
 					.OnClicked(this, &FBlueprintGraphActionDetails::OnAddNewOutputClicked)
+					.Visibility(this, &FBlueprintGraphActionDetails::GetAddNewInputOutputVisibility)
 				]
 			];
 		}
@@ -3691,6 +3715,20 @@ FReply FBaseBlueprintGraphActionDetails::OnAddNewInputClicked()
 	return FReply::Handled();
 }
 
+EVisibility FBlueprintGraphActionDetails::GetAddNewInputOutputVisibility() const
+{
+	UK2Node_EditablePinBase * FunctionEntryNode = FunctionEntryNodePtr.Get();
+	if(UEdGraph* Graph = FunctionEntryNode->GetGraph())
+	{
+		// Math expression graphs are read only, do not allow adding or removing of pins
+		if(Cast<UK2Node_MathExpression>(Graph->GetOuter()))
+		{
+			return EVisibility::Collapsed;
+		}
+	}
+	return EVisibility::Visible;
+}
+
 FReply FBlueprintGraphActionDetails::OnAddNewOutputClicked()
 {
 	FScopedTransaction Transaction( LOCTEXT( "AddOutParam", "Add Out Parameter" ) );
diff --git a/Engine/Source/Editor/Kismet/Private/BlueprintDetailsCustomization.h b/Engine/Source/Editor/Kismet/Private/BlueprintDetailsCustomization.h
--- a/Engine/Source/Editor/Kismet/Private/BlueprintDetailsCustomization.h
+++ b/Engine/Source/Editor/Kismet/Private/BlueprintDetailsCustomization.h
@@ -340,6 +340,9 @@ private:
 	/** Determines if this pin should not be editable */
 	bool ShouldPinBeReadOnly() const;
 	
+	/** Determines if editing the pins on the node should be read only */
+	bool IsPinEditingReadOnly() const;
+
 	/** Callbacks for all the functionality for modifying arguments */
 	void OnRemoveClicked();
 	FReply OnArgMoveUp();
@@ -474,6 +477,9 @@ private:
 	
 	FReply OnAddNewOutputClicked();
 
+	/** Callback to determine if the "New" button for adding input/output pins is visible */
+	EVisibility GetAddNewInputOutputVisibility() const;
+
 	/** Called to set the replication type from the details view combo */
 	static void SetNetFlags( TWeakObjectPtr<UK2Node_EditablePinBase> FunctionEntryNode, uint32 NetFlags);
 
