Changed the SCSDiff container widget so that the SSCSEditor widget gets a scrollbar.

diff --git a/Engine/Source/Editor/Kismet/Private/SCSDiff.cpp b/Engine/Source/Editor/Kismet/Private/SCSDiff.cpp
--- a/Engine/Source/Editor/Kismet/Private/SCSDiff.cpp
+++ b/Engine/Source/Editor/Kismet/Private/SCSDiff.cpp
@@ -13,13 +13,13 @@ FSCSDiff::FSCSDiff(const UBlueprint* InBlueprint)
 		.ViewIdentifier(FName("BlueprintInspector"))
 		.IsPropertyEditingEnabledDelegate(FIsPropertyEditingEnabled::CreateStatic([] { return false; }));
 
-	ContainerWidget = SNew(SVerticalBox)
-		+ SVerticalBox::Slot()
-		.AutoHeight()
+	ContainerWidget = SNew(SSplitter)
+		.Orientation(Orient_Vertical)
+		+ SSplitter::Slot()
 		[
 			SAssignNew(SCSEditor, SSCSEditor, TSharedPtr<FBlueprintEditor>(), InBlueprint->SimpleConstructionScript, const_cast<UBlueprint*>(InBlueprint), Inspector)
 		]
-	+ SVerticalBox::Slot()
+		+ SSplitter::Slot()
 		[
 			Inspector
 		];
