Find in Blueprints uses the asset registry once more instead of the DDC.

Level Blueprints are now fully supported when searching all assets.

Component variables are now searchable.

diff --git a/Engine/Source/Editor/Kismet/Private/BlueprintEditor.cpp b/Engine/Source/Editor/Kismet/Private/BlueprintEditor.cpp
--- a/Engine/Source/Editor/Kismet/Private/BlueprintEditor.cpp
+++ b/Engine/Source/Editor/Kismet/Private/BlueprintEditor.cpp
@@ -7008,7 +7008,6 @@ void FBlueprintEditor::SaveAsset_Execute()
 		if (Blueprint)
 		{
 			// Update the Blueprint's search data
-			Blueprint->SearchGuid = FGuid::NewGuid();
 			FFindInBlueprintSearchManager::Get().AddOrUpdateBlueprintSearchMetadata(Blueprint);
 		}
 	}
diff --git a/Engine/Source/Editor/Kismet/Private/FindInBlueprintManager.cpp b/Engine/Source/Editor/Kismet/Private/FindInBlueprintManager.cpp
--- a/Engine/Source/Editor/Kismet/Private/FindInBlueprintManager.cpp
+++ b/Engine/Source/Editor/Kismet/Private/FindInBlueprintManager.cpp
@@ -17,6 +17,9 @@ FFindInBlueprintSearchManager* FFindInBlueprintSearchManager::Instance = NULL;
 
 const FText FFindInBlueprintSearchTags::FiB_Properties = LOCTEXT("Properties", "Properties");
 
+const FText FFindInBlueprintSearchTags::FiB_Components = LOCTEXT("Components", "Components");
+const FText FFindInBlueprintSearchTags::FiB_IsSCSComponent = LOCTEXT("IsSCSComponent", "IsSCSComponent");
+
 const FText FFindInBlueprintSearchTags::FiB_Nodes = LOCTEXT("Nodes", "Nodes");
 
 const FText FFindInBlueprintSearchTags::FiB_SchemaName = LOCTEXT("SchemaName", "SchemaName");
@@ -142,6 +145,48 @@ namespace BlueprintSearchMetaDataHelpers
 			this->PreviousTokenWritten = EJsonToken::String;
 		}
 
+		/** Converts the lookup table of ints (which are stored as identifiers and string values in the Json) and the FText's they represent to an FString. */
+		FString GetSerializedLookupTable()
+		{
+			TArray<uint8> SerializedData;
+			FMemoryWriter Ar(SerializedData, /*bIsPersistent=*/ true);
+
+			Ar << LookupTable;
+			Ar.Close();
+
+			FString Result = BytesToString(SerializedData.GetData(), SerializedData.Num());
+
+			SerializedData.Empty();
+
+			// Attach, as the first value and in an FString, the length of the Lookup table, this is to help sort through the data during deserialization
+			FMemoryWriter ArWithLength(SerializedData, /*bIsPersistent=*/ true);
+			int32 Length = Result.Len();
+			ArWithLength << Length;
+			return BytesToString(SerializedData.GetData(), SerializedData.Num()) + Result;
+		}
+
+		struct FLookupTableItem
+		{
+			FText Text;
+
+			FLookupTableItem(FText InText)
+				: Text(InText)
+			{
+
+			}
+
+			bool operator==(const FLookupTableItem& InObject) const
+			{
+				if(!Text.CompareTo(InObject.Text) &&
+					!FTextInspector::GetSourceString(Text)->Compare(*FTextInspector::GetDisplayString(InObject.Text), ESearchCase::CaseSensitive) )
+				{
+					return true;
+				}
+
+				return false;
+			}
+		};
+
 	protected:
 		TJsonFindInBlueprintStringWriter( FString* const InOutString )
 			: TJsonStringWriter<PrintPolicy>( InOutString, 0 )
@@ -151,13 +196,25 @@ namespace BlueprintSearchMetaDataHelpers
 		virtual void WriteStringValue( const FString& String ) override
 		{
 			// We just want to make sure all strings are converted into FText hex strings, used by the FiB system
-			TJsonStringWriter<PrintPolicy>::WriteStringValue( FFindInBlueprintSearchManager::ConvertFTextToHexString(FText::FromString(String)) );
+			WriteTextValue(FText::FromString(String));
 		}
 
 		void WriteTextValue( const FText& Text )
 		{
-			// We just want to make sure all strings are converted into FText hex strings, used by the FiB system
-			TJsonStringWriter<PrintPolicy>::WriteStringValue( FFindInBlueprintSearchManager::ConvertFTextToHexString(Text) );
+			// Check to see if the value has already been added.
+			int32* TableLookupValuePtr = ReverseLookupTable.Find(FLookupTableItem(Text));
+			if(TableLookupValuePtr)
+			{
+				TJsonStringWriter<PrintPolicy>::WriteStringValue( FString::FromInt(*TableLookupValuePtr) );
+			}
+			else
+			{
+				// Add the FText to the table and write to the Json the ID to look the item up using
+				int32 TableLookupValue = LookupTable.Num();
+				LookupTable.Add(TableLookupValue, Text);
+				ReverseLookupTable.Add(FLookupTableItem(Text), TableLookupValue);
+				TJsonStringWriter<PrintPolicy>::WriteStringValue( FString::FromInt(TableLookupValue) );
+			}
 		}
 
 		FORCEINLINE void WriteIdentifier( const FText& Identifier )
@@ -170,11 +227,103 @@ namespace BlueprintSearchMetaDataHelpers
 			WriteTextValue( Identifier );
 			PrintPolicy::WriteChar(this->Stream, TCHAR(':'));
 		}
+		
+		// This gets serialized
+		TMap< int32, FText > LookupTable;
+
+		// This is just locally needed for the write, to lookup the integer value by using the string of the FText
+		TMap< FLookupTableItem, int32 > ReverseLookupTable;
 	};
 
 	typedef TJsonWriter<TCHAR, TCondensedJsonPrintPolicy<TCHAR>> SearchMetaDataWriterParentClass;
 	typedef TJsonFindInBlueprintStringWriter<TCondensedJsonPrintPolicy<TCHAR>> SearchMetaDataWriter;
 
+	/** Json Writer used for serializing FText's in the correct format for Find-in-Blueprints */
+	template <class CharType = TCHAR>
+	class TJsonFindInBlueprintStringReader : public TJsonReader<CharType>
+	{
+	public:
+		static TSharedRef< TJsonFindInBlueprintStringReader< TCHAR > > Create( FArchive* const Stream, TMap< int32, FText >& InLookupTable  )
+		{
+			return MakeShareable( new TJsonFindInBlueprintStringReader( Stream, InLookupTable ) );
+		}
+
+		TJsonFindInBlueprintStringReader( FArchive* InStream,  TMap< int32, FText >& InLookupTable )
+			: TJsonReader(InStream)
+			, LookupTable(MoveTemp(InLookupTable))
+		{
+
+		}
+
+		FORCEINLINE virtual const FString& GetIdentifier() const override
+		{
+			// The identifier from Json is a Hex value that must be looked up in the LookupTable to find the FText it represents
+			if(const FText* LookupText = LookupTable.Find(FCString::Atoi(*Identifier)))
+			{
+				return LookupText->ToString();
+			}
+			return Identifier;
+		}
+
+		FORCEINLINE virtual  const FString& GetValueAsString() const override
+		{ 
+			check( CurrentToken == EJsonToken::String ); 
+			// The string value from Json is a Hex value that must be looked up in the LookupTable to find the FText it represents
+			if(const FText* LookupText = LookupTable.Find(FCString::Atoi(*StringValue)))
+			{
+				return LookupText->ToString();
+			}
+			return StringValue;
+		}
+
+		TMap< int32, FText > LookupTable;
+	};
+
+	typedef TJsonFindInBlueprintStringReader<TCHAR> SearchMetaDataReader;
+
+	/** Json Writer used for serializing FText's in the correct format for Find-in-Blueprints */
+	template <class CharType = TCHAR>
+	class TJsonFindInBlueprintStringReaderDDC : public FJsonStringReader
+	{
+	public:
+		static TSharedRef< TJsonFindInBlueprintStringReaderDDC< TCHAR > > Create( const FString& InJsonString  )
+		{
+			return MakeShareable( new TJsonFindInBlueprintStringReaderDDC(InJsonString) );
+		}
+
+		TJsonFindInBlueprintStringReaderDDC( const FString& InJsonString )
+			: FJsonStringReader(InJsonString)
+		{
+
+		}
+
+		FORCEINLINE virtual const FString& GetIdentifier() const override
+		{
+			if(Identifier.Len())
+			{
+				// Remove the const, non-ideal, but this code is to bridge between the old DDC method and the new Asset Registry method of story FiB data.
+				FString& IdentifierNonConst = const_cast<FString&>(Identifier);
+				IdentifierNonConst = FFindInBlueprintSearchManager::ConvertHexStringToFText(Identifier).ToString();
+			}
+			return Identifier;
+		}
+
+		FORCEINLINE virtual  const FString& GetValueAsString() const override
+		{ 
+			check( CurrentToken == EJsonToken::String );
+			if(StringValue.Len())
+			{
+				// Remove the const, non-ideal, but this code is to bridge between the old DDC method and the new Asset Registry method of story FiB data.
+				FString& StringValueNonConst = const_cast<FString&>(StringValue);
+				StringValueNonConst = FFindInBlueprintSearchManager::ConvertHexStringToFText(StringValue).ToString();
+			}
+			return StringValue;
+		}
+	};
+
+	typedef TJsonFindInBlueprintStringReaderDDC<TCHAR> SearchMetaDataReaderDDC;
+
+
 	/**
 	 * Checks if Json value is searchable, eliminating data that not considered useful to search for
 	 *
@@ -357,13 +506,17 @@ namespace BlueprintSearchMetaDataHelpers
 					InWriter->WriteArrayStart(FFindInBlueprintSearchTags::FiB_Pins);
 					for(UEdGraphPin* Pin : Node->Pins)
 					{
-						InWriter->WriteObjectStart();
+						// Hidden pins are not searchable
+						if(Pin->bHidden == false)
 						{
-							InWriter->WriteValue(FFindInBlueprintSearchTags::FiB_Name, FText::FromString(Pin->GetSchema()->GetPinDisplayName(Pin)));
-							InWriter->WriteValue(FFindInBlueprintSearchTags::FiB_DefaultValue, FText::FromString(Pin->GetDefaultAsString()));
+							InWriter->WriteObjectStart();
+							{
+								InWriter->WriteValue(FFindInBlueprintSearchTags::FiB_Name, FText::FromString(Pin->GetSchema()->GetPinDisplayName(Pin)));
+								InWriter->WriteValue(FFindInBlueprintSearchTags::FiB_DefaultValue, FText::FromString(Pin->GetDefaultAsString()));
+							}
+							SavePinTypeToJson(InWriter, Pin->PinType);
+							InWriter->WriteObjectEnd();
 						}
-						SavePinTypeToJson(InWriter, Pin->PinType);
-						InWriter->WriteObjectEnd();
 					}
 					InWriter->WriteArrayEnd();
 
@@ -438,134 +591,10 @@ namespace BlueprintSearchMetaDataHelpers
 	}
 }
 
-class FCacheAllBlueprintsTickableObject
-	: public FTickableEditorObject
+static uint32 GetTypeHash(const BlueprintSearchMetaDataHelpers::TJsonFindInBlueprintStringWriter<TCondensedJsonPrintPolicy<TCHAR>>::FLookupTableItem& InObject)
 {
-
-public:
-
-	FCacheAllBlueprintsTickableObject(TArray<FName>& InUncachedBlueprints)
-		: TickCacheIndex(0)
-		, UncachedBlueprints(InUncachedBlueprints)
-	{
-		// Start the Blueprint indexing 'progress' notification
-		FNotificationInfo Info( LOCTEXT("BlueprintIndexMessage", "Indexing Blueprints...") );
-		Info.bFireAndForget = false;
-		Info.ButtonDetails.Add(FNotificationButtonInfo(
-			LOCTEXT("BlueprintIndexCancel","Cancel"),
-			LOCTEXT("BlueprintIndexCancelToolTip","Cancels indexing Blueprints."), FSimpleDelegate::CreateRaw(this, &FCacheAllBlueprintsTickableObject::OnCancelCaching)));
-
-		ProgressNotification = FSlateNotificationManager::Get().AddNotification(Info);
-		if (ProgressNotification.IsValid())
-		{
-			ProgressNotification.Pin()->SetCompletionState(SNotificationItem::CS_Pending);
-		}
-	}
-
-	~FCacheAllBlueprintsTickableObject()
-	{
-		
-	}
-
-	/** Returns the current cache index of the object */
-	int32 GetCurrentCacheIndex() const
-	{
-		return TickCacheIndex + 1;
-	}
-
-	/** Returns the name of the current Blueprint being cached */
-	FString GetCurrentCacheBlueprintName() const
-	{
-		if(UncachedBlueprints.Num() && TickCacheIndex >= 0)
-		{
-			return UncachedBlueprints[TickCacheIndex].ToString();
-		}
-		return FString();
-	}
-
-	/** Returns the progress as a percent */
-	float GetCacheProgress() const
-	{
-		return (float)TickCacheIndex / (float)UncachedBlueprints.Num();
-	}
-
-	/** Cancels caching and destroys this object */
-	void OnCancelCaching()
-	{
-		ProgressNotification.Pin()->SetText(LOCTEXT("BlueprintIndexCancelled", "Cancelled Indexing Blueprints!"));
-
-		ProgressNotification.Pin()->SetCompletionState(SNotificationItem::CS_Fail);
-		ProgressNotification.Pin()->ExpireAndFadeout();
-
-		FFindInBlueprintSearchManager::Get().FinishedCachingBlueprints(TickCacheIndex);
-	}
-
-protected:
-	/** FTickableEditorObject interface */
-	virtual void Tick(float DeltaTime) override
-	{
-		if(GWarn->ReceivedUserCancel())
-		{
-			FFindInBlueprintSearchManager::Get().FinishedCachingBlueprints(TickCacheIndex);
-		}
-		else
-		{
-			FAssetRegistryModule* AssetRegistryModule = &FModuleManager::LoadModuleChecked<FAssetRegistryModule>(TEXT("AssetRegistry"));
-
-			FAssetData AssetData = AssetRegistryModule->Get().GetAssetByObjectPath(UncachedBlueprints[TickCacheIndex]);
-
-			bool bIsLoaded = AssetData.IsAssetLoaded();
-
-			UObject* Asset = AssetData.GetAsset();
-
-			// If the Blueprint was not just loaded, cache it again, otherwise it was cached on load.
-			if(!bIsLoaded)
-			{
-				FFindInBlueprintSearchManager::Get().AddOrUpdateBlueprintSearchMetadata(Cast<UBlueprint>(Asset), true);
-			}
-
-			++TickCacheIndex;
-
-			// Check if done caching Blueprints
-			if(TickCacheIndex == UncachedBlueprints.Num())
-			{
-				ProgressNotification.Pin()->SetCompletionState(SNotificationItem::CS_Success);
-				ProgressNotification.Pin()->ExpireAndFadeout();
-
-				ProgressNotification.Pin()->SetText(LOCTEXT("BlueprintIndexComplete", "Finished indexing Blueprints!"));
-
-				FFindInBlueprintSearchManager::Get().FinishedCachingBlueprints(TickCacheIndex);
-			}
-			else
-			{
-				FFormatNamedArguments Args;
-				Args.Add(TEXT("Percent"), FText::AsPercent(GetCacheProgress()));
-				ProgressNotification.Pin()->SetText(FText::Format(LOCTEXT("BlueprintIndexProgress", "Indexing Blueprints... ({Percent})"), Args));
-			}
-		}
-	}
-
-	virtual bool IsTickable() const override
-	{
-		return true;
-	}
-
-	virtual TStatId GetStatId() const override
-	{
-		RETURN_QUICK_DECLARE_CYCLE_STAT(FCacheAllBlueprintsTickableObject, STATGROUP_Tickables);
-	}
-
-private:
-
-	/** The current index, increases at a rate of once per tick */
-	int32 TickCacheIndex;
-
-	/** The list of uncached Blueprints that are in the process of being cached */
-	TArray<FName> UncachedBlueprints;
-
-	/** Notification that appears and details progress */
-	TWeakPtr<SNotificationItem> ProgressNotification;
-};
+	return FCrc::MemCrc32(&InObject.Text, sizeof(InObject.Text));
+}
 
 FFindInBlueprintSearchManager& FFindInBlueprintSearchManager::Get()
 {
@@ -580,14 +609,11 @@ FFindInBlueprintSearchManager& FFindInBlueprintSearchManager::Get()
 
 FFindInBlueprintSearchManager::FFindInBlueprintSearchManager()
 	: bIsPausing(false)
-	, CachingObject(NULL)
 {
 }
 
 FFindInBlueprintSearchManager::~FFindInBlueprintSearchManager()
 {
-	FinishedCachingBlueprints(0);
-
 	AssetRegistryModule->Get().OnAssetAdded().RemoveAll(this);
 	AssetRegistryModule->Get().OnAssetRemoved().RemoveAll(this);
 	AssetRegistryModule->Get().OnAssetRenamed().RemoveAll(this);
@@ -606,13 +632,16 @@ void FFindInBlueprintSearchManager::Initialize()
 	FCoreUObjectDelegates::PostGarbageCollect.AddRaw(this, &FFindInBlueprintSearchManager::UnpauseFindInBlueprintSearch);
 	FCoreUObjectDelegates::OnAssetLoaded.AddRaw(this, &FFindInBlueprintSearchManager::OnAssetLoaded);
 
-	// Do an immediate load of the cache to catch any Blueprints that were discovered by the asset registry before we initialized.
-	BuildCache();
+	if(!GIsSavingPackage)
+	{
+		// Do an immediate load of the cache to catch any Blueprints that were discovered by the asset registry before we initialized.
+		BuildCache();
+	}
 }
 
 void FFindInBlueprintSearchManager::OnAssetAdded(const FAssetData& InAssetData)
 {
-	if(InAssetData.GetClass()->IsChildOf(UBlueprint::StaticClass()))
+	if(InAssetData.GetClass()->IsChildOf(UBlueprint::StaticClass()) || InAssetData.GetClass()->IsChildOf(UWorld::StaticClass()))
 	{
 		FString BlueprintPackagePath = FPaths::GetPath(InAssetData.ObjectPath.ToString()) / FPaths::GetBaseFilename(InAssetData.ObjectPath.ToString());
 
@@ -620,31 +649,64 @@ void FFindInBlueprintSearchManager::OnAssetAdded(const FAssetData& InAssetData)
 		int32* IndexPtr = SearchMap.Find(BlueprintPackagePath);
 		if(!IndexPtr)
 		{
-			FString SearchGuid;
-			const FString* SearchGuidPtr = InAssetData.TagsAndValues.Find("SearchGuid");
-			if(SearchGuidPtr)
-			{
-				SearchGuid = *SearchGuidPtr;
-			}
-			else
+			FSearchData NewSearchData;
+
+			NewSearchData.BlueprintPath = InAssetData.ObjectPath.ToString();
+			NewSearchData.bMarkedForDeletion = false;
+
+			if(InAssetData.IsAssetLoaded())
 			{
-				SearchGuid = FMD5::HashAnsiString(*InAssetData.ObjectPath.ToString());
+				UBlueprint* BlueprintAsset = nullptr;
 
-				if(InAssetData.IsAssetLoaded())
+				if(InAssetData.GetClass()->IsChildOf(UBlueprint::StaticClass()))
 				{
-					AddOrUpdateBlueprintSearchMetadata(Cast<UBlueprint>(InAssetData.GetAsset()));
+					BlueprintAsset = Cast<UBlueprint>(InAssetData.GetAsset());
+				}
+				else if(InAssetData.GetClass()->IsChildOf(UWorld::StaticClass()))
+				{
+					UWorld* WorldAsset = Cast<UWorld>(InAssetData.GetAsset());
+					if(WorldAsset->PersistentLevel)
+					{
+						BlueprintAsset = Cast<UBlueprint>((UObject*)WorldAsset->PersistentLevel->GetLevelScriptBlueprint(true));
+					}
 				}
-			}
 
-			FDerivedDataCacheInterface& DDC = GetDerivedDataCacheRef();
-			FSearchData NewSearchData;
-			if(InAssetData.IsAssetLoaded())
+				// Levels do not always have Blueprints
+				if(BlueprintAsset)
+				{
+					// Cache the searchable data
+					AddOrUpdateBlueprintSearchMetadata(BlueprintAsset);
+
+					NewSearchData.Blueprint = Cast<UBlueprint>(BlueprintAsset);
+				}
+
+				// No more work, we have added the Blueprint if there is one, and if it's a level without one, it need not be added.
+				return;
+			}
+			else if(const FString* FiBSearchData = InAssetData.TagsAndValues.Find("FiB"))
 			{
-				NewSearchData.Blueprint = Cast<UBlueprint>(InAssetData.GetAsset());
+				// Since the asset was not loaded, pull out the searchable data stored in the asset
+				NewSearchData.Value = *FiBSearchData;
+			}
+			else
+			{
+				// We will look for DDC information about this asset
+				FString SearchGuid;
+				const FString* SearchGuidPtr = InAssetData.TagsAndValues.Find("SearchGuid");
+				if(SearchGuidPtr)
+				{
+					SearchGuid = *SearchGuidPtr;
+
+					FDerivedDataCacheInterface& DDC = GetDerivedDataCacheRef();
+					if(InAssetData.IsAssetLoaded())
+					{
+						NewSearchData.Blueprint = Cast<UBlueprint>(InAssetData.GetAsset());
+					}
+					NewSearchData.BlueprintPath = InAssetData.ObjectPath.ToString();
+					NewSearchData.DDCRetrievalID = DDC.GetAsynchronous(*SearchGuid);
+					NewSearchData.bMarkedForDeletion = false;
+				}
 			}
-			NewSearchData.BlueprintPath = InAssetData.ObjectPath.ToString();
-			NewSearchData.DDCRetrievalID = DDC.GetAsynchronous(*SearchGuid);
-			NewSearchData.bMarkedForDeletion = false;
 
 			int32 ArrayIndex = SearchArray.Add(NewSearchData);
 
@@ -668,8 +730,10 @@ void FFindInBlueprintSearchManager::OnAssetRemoved(const class FAssetData& InAss
 
 void FFindInBlueprintSearchManager::OnAssetRenamed(const class FAssetData& InAssetData, const FString& InOldName)
 {
+	// Renaming removes the item from the manager, it will be re-added in the OnAssetAdded event under the new name.
 	if(InAssetData.IsAssetLoaded())
 	{
+		FString Test = FPaths::GetPath(InOldName) / FPaths::GetBaseFilename(InOldName);
 		int32* SearchIdx = SearchMap.Find(FPaths::GetPath(InOldName) / FPaths::GetBaseFilename(InOldName));
 		if(SearchIdx)
 		{
@@ -680,22 +744,39 @@ void FFindInBlueprintSearchManager::OnAssetRenamed(const class FAssetData& InAss
 
 void FFindInBlueprintSearchManager::OnAssetLoaded(UObject* InAsset)
 {
-	if(UBlueprint* Blueprint = Cast<UBlueprint>(InAsset))
+	UBlueprint* BlueprintAsset = nullptr;
+	FString BlueprintPackagePath;
+
+	if(UWorld* WorldAsset = Cast<UWorld>(InAsset))
+	{
+		if(WorldAsset->PersistentLevel)
+		{
+			BlueprintAsset = Cast<UBlueprint>((UObject*)WorldAsset->PersistentLevel->GetLevelScriptBlueprint(true));
+			if(BlueprintAsset)
+			{
+				BlueprintPackagePath = BlueprintAsset->GetOutermost()->GetPathName();
+			}
+		}
+	}
+	else
+	{
+		BlueprintAsset = Cast<UBlueprint>(InAsset);
+		BlueprintPackagePath = BlueprintAsset->GetPathName();
+	}
+	if(BlueprintAsset)
 	{
 		// Find and update the item in the search array. Searches may currently be active, this will do no harm to them
 
-		FString BlueprintPackagePath = InAsset->GetPathName();
-
 		// Confirm that the Blueprint has not been added already, this can occur during duplication of Blueprints.
 		int32* IndexPtr = SearchMap.Find(BlueprintPackagePath);
 
 		// The asset registry might not have informed us of this asset yet.
 		if(IndexPtr)
 		{
-			// That index should never have a Blueprint already
-			check(!SearchArray[*IndexPtr].Blueprint);
-
-			SearchArray[*IndexPtr].Blueprint = Blueprint;
+			// That index should never have a Blueprint already, but if it does, it should be the same Blueprint!
+			ensure(!SearchArray[*IndexPtr].Blueprint || SearchArray[*IndexPtr].Blueprint == BlueprintAsset);
+			//check(!SearchArray[*IndexPtr].Blueprint);
+			SearchArray[*IndexPtr].Blueprint = BlueprintAsset;
 		}
 	}
 }
@@ -731,9 +812,41 @@ FString FFindInBlueprintSearchManager::GatherBlueprintSearchMetadata(const UBlue
 	// Sub graphs are processed separately so that they do not become children in the TreeView, cluttering things up if the tree is deep
 	BlueprintSearchMetaDataHelpers::GatherGraphSearchData(Writer, Blueprint, SubGraphs, FFindInBlueprintSearchTags::FiB_SubGraphs, NULL);
 
+	// Gather all SCS components
+	// If we have an SCS but don't support it, then we remove it
+	if(Blueprint->SimpleConstructionScript)
+	{
+		// Remove any SCS variable nodes
+		TArray<USCS_Node*> AllSCSNodes = Blueprint->SimpleConstructionScript->GetAllNodes();
+		Writer->WriteArrayStart(FFindInBlueprintSearchTags::FiB_Components);
+		for (TFieldIterator<UProperty> PropertyIt(Blueprint->SkeletonGeneratedClass, EFieldIteratorFlags::ExcludeSuper); PropertyIt; ++PropertyIt)
+		{
+			UProperty* Property = *PropertyIt;
+			UObjectPropertyBase* Obj = Cast<UObjectPropertyBase>(Property);
+			const bool bComponentProperty = Obj && Obj->PropertyClass ? Obj->PropertyClass->IsChildOf<UActorComponent>() : false;
+			FName PropName = Property->GetFName();
+			if(bComponentProperty && FBlueprintEditorUtils::FindSCS_Node(Blueprint, PropName) != INDEX_NONE)
+			{
+				FEdGraphPinType PropertyPinType;
+				if(UEdGraphSchema_K2::StaticClass()->GetDefaultObject<UEdGraphSchema_K2>()->ConvertPropertyToPinType(Property, PropertyPinType))
+				{
+					Writer->WriteObjectStart();
+					{
+						Writer->WriteValue(FFindInBlueprintSearchTags::FiB_Name, FText::FromName(PropName));
+						Writer->WriteValue(FFindInBlueprintSearchTags::FiB_IsSCSComponent, true);
+						SavePinTypeToJson(Writer,  PropertyPinType);
+					}
+					Writer->WriteObjectEnd();
+				}
+			}
+		}
+		Writer->WriteArrayEnd(); // Components
+	}
+
 	Writer->WriteObjectEnd();
 	Writer->Close();
 
+	SearchMetaData = Writer->GetSerializedLookupTable() + SearchMetaData;
 	return SearchMetaData;
 }
 
@@ -743,18 +856,31 @@ void FFindInBlueprintSearchManager::AddOrUpdateBlueprintSearchMetadata(UBlueprin
 
 	check(InBlueprint);
 
-	// Allow only one thread to query the DDC at a single time, this is important in the case that two threads are waiting for the same DDC information, once the information is pulled, it is deleted and is no longer accessible.
+	// Allow only one thread modify the search data at a time
 	FScopeLock ScopeLock(&SafeModifyCacheCriticalSection);
 
 	FString BlueprintPackagePath = InBlueprint->GetOutermost()->GetPathName();
 
+	FString BlueprintPath;
+	if(FBlueprintEditorUtils::IsLevelScriptBlueprint(InBlueprint))
+	{
+		if(UWorld* World = InBlueprint->GetTypedOuter<UWorld>())
+		{
+			BlueprintPath = World->GetPathName();
+		}
+	}
+	else
+	{
+		BlueprintPath = InBlueprint->GetPathName();
+	}
+
 	int32* IndexPtr = SearchMap.Find(BlueprintPackagePath);
 	int32 Index = 0;
 	if(!IndexPtr)
 	{
 		FSearchData SearchData;
 		SearchData.Blueprint = InBlueprint;
-		SearchData.BlueprintPath = InBlueprint->GetPathName();
+		SearchData.BlueprintPath = BlueprintPath;
 		Index = SearchArray.Add(SearchData);
 
 		SearchMap.Add(BlueprintPackagePath, Index);
@@ -764,63 +890,10 @@ void FFindInBlueprintSearchManager::AddOrUpdateBlueprintSearchMetadata(UBlueprin
 		Index = *IndexPtr;
 	}
 
-	FDerivedDataCacheInterface& DDC = GetDerivedDataCacheRef();
-
-	FAssetData AssetData = AssetRegistryModule->Get().GetAssetByObjectPath(*InBlueprint->GetPathName());
-
-	FString* SearchGuidPtr = AssetData.TagsAndValues.Find("SearchGuid");
-	FString SearchGuid;
-
-	// If there is no search Guid stored, generate one using the asset's path. This transitions Blueprints from the old search system to the new one without resaving
-	if(!SearchGuidPtr)
-	{
-		SearchGuid = FMD5::HashAnsiString(*AssetData.ObjectPath.ToString());
-	}
-	else
-	{
-		SearchGuid = *SearchGuidPtr;
-	}
-
-	// If the Blueprint has been cached, don't cache it again, it's wasted time
-	if(!DDC.CachedDataProbablyExists(*SearchGuid) || bInForceReCache)
-	{
-		// The data is valid, but we may still need to wait on the DDC.
-		if(SearchArray[Index].DDCRetrievalID != INDEX_NONE)
-		{
-			// Allow only one thread to query the DDC at a single time, this is important in the case that two threads are waiting for the same DDC information, once the information is pulled, it is deleted and is no longer accessible.
-			FScopeLock ScopeLock(&SafeModifyCacheCriticalSection);
-
-			// If there was a DDC retrieval request, we have to see it through
-			if(SearchArray[Index].DDCRetrievalID != INDEX_NONE)
-			{
-				// Wait for completion of the DDC retrieval, we are on a separate thread so this will not stall the editor
-				DDC.WaitAsynchronousCompletion(SearchArray[Index].DDCRetrievalID);
-
-				// Fetch the data, we don't need it but the DDC doesn't need to manage the query anymore
-				TArray<uint8> DerivedData;
-				DDC.GetAsynchronousResults(SearchArray[Index].DDCRetrievalID, DerivedData);
-
-				SearchArray[Index].DDCRetrievalID = INDEX_NONE;
-			}
-		}
-
-		// Build the search data
-		SearchArray[Index].BlueprintPath = InBlueprint->GetPathName();
-		SearchArray[Index].Value = GatherBlueprintSearchMetadata(InBlueprint);
-		SearchArray[Index].bMarkedForDeletion = false;
-
-		// Save the search data to the DDC
-		TArray<uint8> DerivedData;
-		FMemoryWriter Ar(DerivedData, /*bIsPersistent=*/ true);
-		Ar << SearchArray[Index].BlueprintPath;
-		Ar << SearchArray[Index].Value;
-
-		// Insert the search data into the DDC
-		check(DerivedData.Num());
-		DDC.Put(*SearchGuid, DerivedData);
-
-		Ar.Close();
-	}
+	// Build the search data
+	SearchArray[Index].BlueprintPath = BlueprintPath;
+	SearchArray[Index].Value = GatherBlueprintSearchMetadata(InBlueprint);
+	SearchArray[Index].bMarkedForDeletion = false;
 }
 
 void FFindInBlueprintSearchManager::BeginSearchQuery(const FStreamSearch* InSearchOriginator)
@@ -863,17 +936,22 @@ bool FFindInBlueprintSearchManager::ContinueSearchQuery(const FStreamSearch* InS
 				// Mark it for deletion, it will be removed on next save
 				SearchArray[SearchIdx].bMarkedForDeletion = true;
 			}
-			else if(SearchArray[SearchIdx].Blueprint)
-			{
-				RetrieveDDCData(SearchArray[SearchIdx], *SearchArray[SearchIdx].Blueprint->GetPathName());
-
- 				OutSearchData = SearchArray[SearchIdx++];
-				return true;
-			}
 			else
 			{
-				RetrieveDDCData(SearchArray[SearchIdx], *SearchArray[SearchIdx].BlueprintPath);
-				OutSearchData = SearchArray[SearchIdx++];
+ 				OutSearchData = SearchArray[SearchIdx++];
+				if(OutSearchData.Value.IsEmpty())
+				{
+					if(OutSearchData.Blueprint)
+					{
+						RetrieveDDCData(SearchArray[SearchIdx - 1], *OutSearchData.Blueprint->GetPathName());
+					}
+					else
+					{
+						// Attempt to retrieve the data from the DDC, if there is none still, it
+						RetrieveDDCData(SearchArray[SearchIdx - 1], *SearchArray[SearchIdx - 1].BlueprintPath);
+					}
+					OutSearchData = SearchArray[SearchIdx - 1];
+				}
 				return true;
 			}
 
@@ -938,12 +1016,7 @@ void FFindInBlueprintSearchManager::RetrieveDDCData(FSearchData& InOutSearchData
 
 			TArray<uint8> DerivedData;
 
-			if( !DDC.GetAsynchronousResults(InOutSearchData.DDCRetrievalID, DerivedData) ) 
-			{
-				// Since the asset cannot be loaded at this point, cache the Blueprint, we'll index it later.
-				UncachedBlueprints.AddUnique(InBlueprintPath);
-			}
-			else
+			if( DDC.GetAsynchronousResults(InOutSearchData.DDCRetrievalID, DerivedData) ) 
 			{
 				check(DerivedData.Num());
 
@@ -958,20 +1031,25 @@ void FFindInBlueprintSearchManager::RetrieveDDCData(FSearchData& InOutSearchData
 	}
 }
 
-FString FFindInBlueprintSearchManager::QuerySingleBlueprint(UBlueprint* InBlueprint)
+FString FFindInBlueprintSearchManager::QuerySingleBlueprint(UBlueprint* InBlueprint, bool bInRebuildSearchData/* = true*/)
 {
-	// Update the Blueprint, make sure it is fully up-to-date
-	AddOrUpdateBlueprintSearchMetadata(InBlueprint, true);
+	if(bInRebuildSearchData)
+	{
+		// Update the Blueprint, make sure it is fully up-to-date
+		AddOrUpdateBlueprintSearchMetadata(InBlueprint, true);
+	}
 
 	FString Key = InBlueprint->GetOutermost()->GetPathName();
 
 	int32* ArrayIdx = SearchMap.Find(Key);
-	// This should always be true since we make sure to refresh the search data for this Blueprint when doing the search
-	check(ArrayIdx && *ArrayIdx < SearchArray.Num());
+	// This should always be true since we make sure to refresh the search data for this Blueprint when doing the search, unless we do not rebuild the searchable data
+	check((bInRebuildSearchData && ArrayIdx && *ArrayIdx < SearchArray.Num()) || !bInRebuildSearchData);
 
-	RetrieveDDCData(SearchArray[*ArrayIdx], *InBlueprint->GetPathName());
-
-	return SearchArray[*ArrayIdx].Value;
+	if(ArrayIdx)
+	{
+		return SearchArray[*ArrayIdx].Value;
+	}
+	return FString();
 }
 
 void FFindInBlueprintSearchManager::PauseFindInBlueprintSearch()
@@ -1019,19 +1097,45 @@ void FFindInBlueprintSearchManager::CleanCache()
 	 	}
 	}
 
-	FDerivedDataCacheInterface& DDC = GetDerivedDataCacheRef();
-
 	TMap<FString, int> NewSearchMap;
 	TArray<FSearchData> NewSearchArray;
 
 	for(auto& SearchValuePair : SearchMap)
 	{
 		// Here it builds the new map/array, clean of deleted content.
+
+		// If the database item is not marked for deletion and not pending kill (if loaded), keep it in the database
 		if( !SearchArray[SearchValuePair.Value].bMarkedForDeletion && !(SearchArray[SearchValuePair.Value].Blueprint && SearchArray[SearchValuePair.Value].Blueprint->HasAllFlags(RF_PendingKill)) )
 		{
 			// Build the new map/array
 			NewSearchMap.Add(SearchValuePair.Key, NewSearchArray.Add(SearchArray[SearchValuePair.Value]) );
 		}
+		else
+		{
+			// Level Blueprints are destroyed when you open a new level, we need to re-add it as an unloaded asset so long as they were not marked for deletion
+			if(!SearchArray[SearchValuePair.Value].bMarkedForDeletion)
+			{
+				SearchArray[SearchValuePair.Value].Blueprint = nullptr;
+
+				AssetRegistryModule = &FModuleManager::LoadModuleChecked<FAssetRegistryModule>(TEXT("AssetRegistry"));
+
+				// The asset was not user deleted, so this should usually find the asset. New levels can be deleted if they were not saved
+				FAssetData AssetData = AssetRegistryModule->Get().GetAssetByObjectPath(*SearchArray[SearchValuePair.Value].BlueprintPath);
+				if(AssetData.IsValid())
+				{
+					if(const FString* FiBSearchData = AssetData.TagsAndValues.Find("FiB"))
+					{
+						SearchArray[SearchValuePair.Value].Value = *FiBSearchData;
+					}
+					// Build the new map/array
+					NewSearchMap.Add(SearchValuePair.Key, NewSearchArray.Add(SearchArray[SearchValuePair.Value]) );
+				}
+				else
+				{
+					SearchArray[SearchValuePair.Value].Blueprint = SearchArray[SearchValuePair.Value].Blueprint;
+				}
+			}
+		}
 	}
 
 	SearchMap = NewSearchMap;
@@ -1064,103 +1168,19 @@ void FFindInBlueprintSearchManager::BuildCache()
 	AssetRegistryModule = &FModuleManager::LoadModuleChecked<FAssetRegistryModule>(TEXT("AssetRegistry"));
 
 	TArray< FAssetData > BlueprintAssets;
-	AssetRegistryModule->Get().GetAssetsByClass(UBlueprint::StaticClass()->GetFName(), BlueprintAssets, true);
+	FARFilter ClassFilter;
+	ClassFilter.bRecursiveClasses = true;
+	ClassFilter.ClassNames.Add(UBlueprint::StaticClass()->GetFName());
+	ClassFilter.ClassNames.Add(UWorld::StaticClass()->GetFName());
+
+	AssetRegistryModule->Get().GetAssets(ClassFilter, BlueprintAssets);
 	
-	FDerivedDataCacheInterface& DDC = GetDerivedDataCacheRef();
 	for( FAssetData& Asset : BlueprintAssets )
 	{
-		if(FString* SearchGuid = Asset.TagsAndValues.Find("SearchGuid"))
-		{
-			OnAssetAdded(Asset);
-		}
-	}
-}
-
-void FFindInBlueprintSearchManager::CacheAllUncachedBlueprints(TWeakPtr< SFindInBlueprints > InSourceWidget)
-{
-	// Do not start another caching process if one is in progress
-	if(!IsCacheInProgress())
-	{
-		// Multiple threads can be adding to this at the same time
-		FScopeLock ScopeLock(&SafeModifyCacheCriticalSection);
-
-		CachingObject = new FCacheAllBlueprintsTickableObject(UncachedBlueprints);
-	}
-
-	// Add the widget to the callback list, to be called when caching is complete
-	OnCachingCompleteCallbackWidgets.Add(InSourceWidget);
-}
-
-void FFindInBlueprintSearchManager::CancelCacheAll()
-{
-	if(IsCacheInProgress())
-	{
-		CachingObject->OnCancelCaching();
+		OnAssetAdded(Asset);
 	}
 }
 
-int32 FFindInBlueprintSearchManager::GetCurrentCacheIndex() const
-{
-	int32 CachingIndex = 0;
-	if(CachingObject)
-	{
-		CachingIndex = CachingObject->GetCurrentCacheIndex();
-	}
-
-	return CachingIndex;
-}
-
-FString FFindInBlueprintSearchManager::GetCurrentCacheBlueprintName() const
-{
-	FString CachingBPName;
-	if(CachingObject)
-	{
-		CachingBPName = CachingObject->GetCurrentCacheBlueprintName();
-	}
-
-	return CachingBPName;
-}
-
-float  FFindInBlueprintSearchManager::GetCacheProgress() const
-{
-	float ReturnCacheValue = 1.0f;
-
-	if(CachingObject)
-	{
-		ReturnCacheValue = CachingObject->GetCacheProgress();
-	}
-
-	return ReturnCacheValue;
-}
-
-void FFindInBlueprintSearchManager::FinishedCachingBlueprints(int32 InNumberCached)
-{
-	// Multiple threads could be adding to this at the same time
-	FScopeLock ScopeLock(&SafeModifyCacheCriticalSection);
-
-	// Clean out the list by the count that were successfully cached, this allows us to continue the caching later.
-	UncachedBlueprints.RemoveAt(0, InNumberCached);
-
-	// Signal the callback, so all FindInBlueprints can resubmit their search queries
-	for(TWeakPtr< SFindInBlueprints > FindInBlueprintsWidget : OnCachingCompleteCallbackWidgets)
-	{
-		if(FindInBlueprintsWidget.IsValid())
-		{
-			FindInBlueprintsWidget.Pin()->OnCacheComplete();
-		}
-	}
-	OnCachingCompleteCallbackWidgets.Empty();
-
-	// Delete the object and NULL it out so we can do it again in the future if needed (if it was canceled)
-	delete CachingObject;
-	CachingObject = NULL;
-}
-
-bool FFindInBlueprintSearchManager::IsCacheInProgress() const
-{
-	return CachingObject != NULL;
-}
-
 FText FFindInBlueprintSearchManager::ConvertHexStringToFText(FString InHexString)
 {
 	TArray<uint8> SerializedData;
@@ -1188,4 +1208,66 @@ FString FFindInBlueprintSearchManager::ConvertFTextToHexString(FText InValue)
 	return BytesToHex(SerializedData.GetData(), SerializedData.Num());
 }
 
+TSharedPtr< FJsonObject > FFindInBlueprintSearchManager::ConvertJsonStringToObject(FString InJsonString)
+{
+	if(InJsonString.Len() && InJsonString[0] == '{')
+	{
+		// This is a DDC version Json string, use old method of deserialization
+		TSharedPtr< FJsonObject > JsonObject = NULL;
+		TSharedRef< TJsonReader<> > Reader = BlueprintSearchMetaDataHelpers::SearchMetaDataReaderDDC::Create( InJsonString );
+		FJsonSerializer::Deserialize( Reader, JsonObject );
+
+		return JsonObject;
+	}
+
+	/** The searchable data is more complicated than a Json string, the Json being the main searchable body that is parsed. Below is a diagram of the full data:
+	 *  | int32 "Size" | TMap "Lookup Table" | Json String |
+	 *
+	 * Size: The size of the TMap in bytes
+	 * Lookup Table: The Json's identifiers and string values are in Hex strings and stored in a TMap, the Json stores these values as ints and uses them as the Key into the TMap
+	 * Json String: The Json string to be deserialized in full
+	 */
+	TArray<uint8> DerivedData;
+
+	// SearchData is currently the full string
+	// We want to first extract the size of the TMap we will be serializing
+	int32 SizeOfData;
+	FBufferReader ReaderStream((void*)*InJsonString, InJsonString.Len() * sizeof(TCHAR), false);
+
+	// Read, as a byte string, the number of characters composing the Lookup Table for the Json.
+	FString SizeOfDataAsHex;
+	SizeOfDataAsHex.GetCharArray().AddUninitialized(5);
+	SizeOfDataAsHex[4] = TEXT('\0');
+	ReaderStream.Serialize((char*)SizeOfDataAsHex.GetCharArray().GetData(), sizeof(TCHAR) * 4);
+
+	// Convert the number (which is stored in 1 serialized byte per TChar) into an int32
+	DerivedData.Empty();
+	DerivedData.AddUninitialized(4);
+	StringToBytes(SizeOfDataAsHex, DerivedData.GetData(), 4);
+	FMemoryReader SizeOfDataAr(DerivedData, /*bIsPersistent=*/ true);
+	SizeOfDataAr << SizeOfData;
+
+	// With the size of the TMap in hand, let's serialize JUST that (as a byte string)
+	FString MapAsByteString;
+	MapAsByteString.GetCharArray().AddUninitialized(SizeOfData + 1);
+	MapAsByteString[SizeOfData] = TEXT('\0');
+	ReaderStream.Serialize((char*)MapAsByteString.GetCharArray().GetData(), sizeof(TCHAR) * SizeOfData);
+
+	// With the hex of the TMap, serialize that into a TMap
+	DerivedData.Empty();
+	DerivedData.AddUninitialized(SizeOfData);
+	StringToBytes(MapAsByteString, DerivedData.GetData(), SizeOfData);
+
+	TMap<int32, FText> LookupTable;
+	FMemoryReader LookupTableAr(DerivedData, /*bIsPersistent=*/ true);
+	LookupTableAr << LookupTable;
+
+	// The original BufferReader should be positioned at the Json
+	TSharedPtr< FJsonObject > JsonObject = NULL;
+	TSharedRef< TJsonReader<> > Reader = BlueprintSearchMetaDataHelpers::SearchMetaDataReader::Create( &ReaderStream, LookupTable );
+	FJsonSerializer::Deserialize( Reader, JsonObject );
+
+	return JsonObject;
+}
+
 #undef LOCTEXT_NAMESPACE
diff --git a/Engine/Source/Editor/Kismet/Private/FindInBlueprints.cpp b/Engine/Source/Editor/Kismet/Private/FindInBlueprints.cpp
--- a/Engine/Source/Editor/Kismet/Private/FindInBlueprints.cpp
+++ b/Engine/Source/Editor/Kismet/Private/FindInBlueprints.cpp
@@ -117,15 +117,14 @@ namespace FindInBlueprintsHelpers
 		}
 
 		TSharedPtr< FJsonObject > JsonObject = NULL;
-		TSharedRef< TJsonReader<> > Reader = TJsonReaderFactory<TCHAR>::Create( InBuffer );
-		FJsonSerializer::Deserialize( Reader, JsonObject );
+		JsonObject = FFindInBlueprintSearchManager::ConvertJsonStringToObject(InBuffer);
 
 		if(JsonObject.IsValid())
 		{
-			if(JsonObject->HasField(FFindInBlueprintSearchManager::ConvertFTextToHexString(FFindInBlueprintSearchTags::FiB_Properties)))
+			if(JsonObject->HasField(FFindInBlueprintSearchTags::FiB_Properties.ToString()))
 			{
 				// Pulls out all properties (variables) for this Blueprint
-				TArray<TSharedPtr< FJsonValue > > PropertyList = JsonObject->GetArrayField(FFindInBlueprintSearchManager::ConvertFTextToHexString(FFindInBlueprintSearchTags::FiB_Properties));
+				TArray<TSharedPtr< FJsonValue > > PropertyList = JsonObject->GetArrayField(FFindInBlueprintSearchTags::FiB_Properties.ToString());
 				for( TSharedPtr< FJsonValue > PropertyValue : PropertyList )
 				{
 					FSearchResult Result(new FFindInBlueprintsProperty(FText::GetEmpty(), OutBlueprintCategory) );
@@ -137,10 +136,33 @@ namespace FindInBlueprintsHelpers
 				}
 			}
 
-			ExtractGraph(JsonObject, FFindInBlueprintSearchManager::ConvertFTextToHexString(FFindInBlueprintSearchTags::FiB_Functions), InTokens, EGraphType::GT_Function, OutBlueprintCategory);
-			ExtractGraph(JsonObject, FFindInBlueprintSearchManager::ConvertFTextToHexString(FFindInBlueprintSearchTags::FiB_Macros), InTokens, EGraphType::GT_Macro, OutBlueprintCategory);
-			ExtractGraph(JsonObject, FFindInBlueprintSearchManager::ConvertFTextToHexString(FFindInBlueprintSearchTags::FiB_UberGraphs), InTokens, EGraphType::GT_Ubergraph, OutBlueprintCategory);
-			ExtractGraph(JsonObject, FFindInBlueprintSearchManager::ConvertFTextToHexString(FFindInBlueprintSearchTags::FiB_SubGraphs), InTokens, EGraphType::GT_Ubergraph, OutBlueprintCategory);
+			ExtractGraph(JsonObject, FFindInBlueprintSearchTags::FiB_Functions.ToString(), InTokens, EGraphType::GT_Function, OutBlueprintCategory);
+			ExtractGraph(JsonObject, FFindInBlueprintSearchTags::FiB_Macros.ToString(), InTokens, EGraphType::GT_Macro, OutBlueprintCategory);
+			ExtractGraph(JsonObject, FFindInBlueprintSearchTags::FiB_UberGraphs.ToString(), InTokens, EGraphType::GT_Ubergraph, OutBlueprintCategory);
+			ExtractGraph(JsonObject, FFindInBlueprintSearchTags::FiB_SubGraphs.ToString(), InTokens, EGraphType::GT_Ubergraph, OutBlueprintCategory);
+
+			if(JsonObject->HasField(FFindInBlueprintSearchTags::FiB_Components.ToString()))
+			{
+				FSearchResult ComponentsResult(new FFindInBlueprintsResult(FFindInBlueprintSearchTags::FiB_Components, OutBlueprintCategory) );
+
+				// Pulls out all properties (variables) for this Blueprint
+				TArray<TSharedPtr< FJsonValue > > ComponentList = JsonObject->GetArrayField(FFindInBlueprintSearchTags::FiB_Components.ToString());
+				for( TSharedPtr< FJsonValue > ComponentValue : ComponentList )
+				{
+					FSearchResult Result(new FFindInBlueprintsProperty(FText::GetEmpty(), ComponentsResult) );
+
+					if( Result->ExtractContent(ComponentValue->AsObject(), InTokens) )
+					{
+						ComponentsResult->Children.Add(Result);
+					}
+				}
+
+				// Only add the component list category if there were components
+				if(ComponentsResult->Children.Num())
+				{
+					OutBlueprintCategory->Children.Add(ComponentsResult);
+				}
+			}
 		}
 	}
 
@@ -254,7 +276,7 @@ bool FFindInBlueprintsResult::ExtractJsonValue(const TArray<FString>& InTokens,
 		{
 			if(Array[ArrayIdx]->Type == EJson::String)
 			{
-				bArrayHasMatches |= ParseSearchInfo(InTokens, FString::FromInt(ArrayIdx), FFindInBlueprintSearchManager::ConvertHexStringToFText(Array[ArrayIdx]->AsString()), ArrayHeader);
+				bArrayHasMatches |= ParseSearchInfo(InTokens, FString::FromInt(ArrayIdx), FText::FromString(Array[ArrayIdx]->AsString()), ArrayHeader);
 			}
 			else
 			{
@@ -282,7 +304,7 @@ bool FFindInBlueprintsResult::ExtractJsonValue(const TArray<FString>& InTokens,
 	}
 	else if(InValue->Type == EJson::String)
 	{
-		bMatchesSearchResults |= ParseSearchInfo(InTokens, InKey.ToString(), FFindInBlueprintSearchManager::ConvertHexStringToFText(InValue->AsString()), ParentResultPtr);
+		bMatchesSearchResults |= ParseSearchInfo(InTokens, InKey.ToString(), FText::FromString(InValue->AsString()), ParentResultPtr);
 	}
 	else
 	{
@@ -300,7 +322,7 @@ bool FFindInBlueprintsResult::ExtractContent(TSharedPtr< FJsonObject > InJsonNod
 
 	for( auto MapValues : InJsonNode->Values )
 	{
-		bMatchesSearchResults |= ExtractJsonValue(InTokens, FFindInBlueprintSearchManager::ConvertHexStringToFText(MapValues.Key), MapValues.Value, ParentResultPtr);
+		bMatchesSearchResults |= ExtractJsonValue(InTokens, FText::FromString(MapValues.Key), MapValues.Value, ParentResultPtr);
 	}
 	return bMatchesSearchResults;
 }
@@ -328,14 +350,31 @@ FString FFindInBlueprintsResult::GetCommentText() const
 
 UBlueprint* FFindInBlueprintsResult::GetParentBlueprint() const
 {
+	UBlueprint* ResultBlueprint = nullptr;
 	if (Parent.IsValid())
 	{
-		return Parent.Pin()->GetParentBlueprint();
+		ResultBlueprint = Parent.Pin()->GetParentBlueprint();
 	}
 	else
 	{
-		return Cast<UBlueprint>(StaticLoadObject(UBlueprint::StaticClass(), NULL, *DisplayText.ToString()));
+		GIsEditorLoadingPackage = true;
+		UObject* Object = LoadObject<UObject>(NULL, *DisplayText.ToString(), NULL, 0, NULL); //StaticLoadObject(UObject::StaticClass(), NULL, *DisplayText.ToString());
+		GIsEditorLoadingPackage = false;
+
+		if(UBlueprint* BlueprintObj = Cast<UBlueprint>(Object))
+		{
+			ResultBlueprint = BlueprintObj;
+		}
+		else if(UWorld* WorldObj = Cast<UWorld>(Object))
+		{
+			if(WorldObj->PersistentLevel)
+			{
+				ResultBlueprint = Cast<UBlueprint>((UObject*)WorldObj->PersistentLevel->GetLevelScriptBlueprint(true));
+			}
+		}
+
 	}
+	return ResultBlueprint;
 }
 
 void FFindInBlueprintsResult::ExpandAllChildren( TSharedPtr< STreeView< TSharedPtr< FFindInBlueprintsResult > > > InTreeView )
@@ -440,17 +479,17 @@ bool FFindInBlueprintsGraphNode::ExtractContent(TSharedPtr< FJsonObject > InJson
 {
 	// Very important to get the schema first, other bits of data depend on it
 	FString SchemaName;
-	TSharedPtr< FJsonValue > SchemaNameValue = InJsonNode->GetField< EJson::String >(FFindInBlueprintSearchManager::ConvertFTextToHexString(FFindInBlueprintSearchTags::FiB_SchemaName));
+	TSharedPtr< FJsonValue > SchemaNameValue = InJsonNode->GetField< EJson::String >(FFindInBlueprintSearchTags::FiB_SchemaName.ToString());
 	if(SchemaNameValue.IsValid())
 	{
-		SchemaName = FFindInBlueprintSearchManager::ConvertHexStringToFText(SchemaNameValue->AsString()).ToString();
+		SchemaName = SchemaNameValue->AsString();
 	}
 
 	bool bMatchesSearchResults = false;
 
 	for( auto MapValues : InJsonNode->Values )
 	{
-		FText Key = FFindInBlueprintSearchManager::ConvertHexStringToFText(MapValues.Key);
+		FText Key = FText::FromString(MapValues.Key);
 
 		if(Key.CompareTo(FFindInBlueprintSearchTags::FiB_Pins) == 0)
 		{
@@ -474,7 +513,7 @@ bool FFindInBlueprintsGraphNode::ExtractContent(TSharedPtr< FJsonObject > InJson
 		}
 		else
 		{
-			bMatchesSearchResults |= ExtractJsonValue(InTokens, FFindInBlueprintSearchManager::ConvertHexStringToFText(MapValues.Key), MapValues.Value, SharedThis( this ));
+			bMatchesSearchResults |= ExtractJsonValue(InTokens, FText::FromString(MapValues.Key), MapValues.Value, SharedThis( this ));
 		}
 	}
 	return bMatchesSearchResults;
@@ -559,8 +598,11 @@ bool FFindInBlueprintsPin::ParseSearchInfo(const TArray<FString> &InTokens, FStr
 	}
 	else if(FindInBlueprintsHelpers::ParsePinType(InKey, InValue, PinType, PinSubCategory))
 	{
-		// PinType information is not searchable
-		bMatchesTokens = false;
+		// For pins, the only information in the Pin Type we care about is the object class, so return false if the Key is not that
+		if(InKey != FFindInBlueprintSearchTags::FiB_ObjectClass.ToString())
+		{
+			return false;
+		}
 	}
 	else
 	{
@@ -609,6 +651,7 @@ void FFindInBlueprintsPin::FinalizeSearchData()
 
 FFindInBlueprintsProperty::FFindInBlueprintsProperty(const FText& InValue, TSharedPtr<FFindInBlueprintsResult> InParent)
 	: FFindInBlueprintsResult(InValue, InParent)
+	, bIsSCSComponent(false)
 {
 }
 
@@ -634,7 +677,18 @@ bool FFindInBlueprintsProperty::ParseSearchInfo(const TArray<FString> &InTokens,
 	}
 	else if(FindInBlueprintsHelpers::ParsePinType(InKey, InValue, PinType, PinSubCategory))
 	{
-		// Do not need to do anything
+		// For properties, we care about all the information in the pin type, except whether its a reference or an array
+		if(InKey == FFindInBlueprintSearchTags::FiB_IsArray.ToString() || InKey == FFindInBlueprintSearchTags::FiB_IsReference.ToString())
+		{
+			return false;
+		}
+	}
+	else if(InKey == FFindInBlueprintSearchTags::FiB_IsSCSComponent.ToString())
+	{
+		bIsSCSComponent = true;
+
+		// This value is not searchable
+		bMatchesTokens = false;
 	}
 	else
 	{
@@ -649,6 +703,10 @@ bool FFindInBlueprintsProperty::ParseSearchInfo(const TArray<FString> &InTokens,
 
 FString FFindInBlueprintsProperty::GetCategory() const
 {
+	if(bIsSCSComponent)
+	{
+		return LOCTEXT("Component", "Component").ToString();
+	}
 	return LOCTEXT("Variable", "Variable").ToString();
 }
 
@@ -724,7 +782,7 @@ bool FFindInBlueprintsGraph::ExtractContent(TSharedPtr< FJsonObject > InJsonNode
 
 	for( auto MapValues : InJsonNode->Values )
 	{
-		FText Key = FFindInBlueprintSearchManager::ConvertHexStringToFText(MapValues.Key);
+		FText Key = FText::FromString(MapValues.Key);
 
 		if(Key.CompareTo(FFindInBlueprintSearchTags::FiB_Nodes) == 0)
 		{
@@ -763,7 +821,7 @@ bool FFindInBlueprintsGraph::ExtractContent(TSharedPtr< FJsonObject > InJsonNode
 		}
 		else
 		{
-			bMatchesSearchResults |= ExtractJsonValue( InTokens, FFindInBlueprintSearchManager::ConvertHexStringToFText(MapValues.Key), MapValues.Value, SharedThis( this ) );
+			bMatchesSearchResults |= ExtractJsonValue( InTokens, FText::FromString(MapValues.Key), MapValues.Value, SharedThis( this ) );
 		}
 	}
 	return bMatchesSearchResults;
@@ -1015,115 +1073,6 @@ void SFindInBlueprints::Construct( const FArguments& InArgs, TSharedPtr<FBluepri
 		];
 }
 
-void SFindInBlueprints::ConditionallyAddCacheBar()
-{
-	// Do not add a second cache bar and do not add it when there are no uncached Blueprints
-	if(FFindInBlueprintSearchManager::Get().GetNumberUncachedBlueprints() > 0)
-	{
-		if(MainVerticalBox.IsValid() && !CacheBarSlot.IsValid())
-		{
-			MainVerticalBox.Pin()->AddSlot()
-				.AutoHeight()
-				[
-					SAssignNew(CacheBarSlot, SBorder)
-					.Visibility( this, &SFindInBlueprints::GetCachingBarVisibility )
-					.BorderBackgroundColor( this, &SFindInBlueprints::GetCachingBarColor )
-					.BorderImage( FCoreStyle::Get().GetBrush("ErrorReporting.Box") )
-					.Padding( FMargin(3,1) )
-					[
-						SNew(SVerticalBox)
-
-						+SVerticalBox::Slot()
-							.AutoHeight()
-						[
-							SNew(SHorizontalBox)
-							+SHorizontalBox::Slot()
-							.VAlign(EVerticalAlignment::VAlign_Center)
-							.AutoWidth()
-							[
-								SNew(STextBlock)
-								.Text(this, &SFindInBlueprints::GetUncachedBlueprintWarningText)
-								.ColorAndOpacity( FCoreStyle::Get().GetColor("ErrorReporting.ForegroundColor") )
-							]
-
-							// Cache All button
-							+SHorizontalBox::Slot()
-								.AutoWidth()
-								.VAlign(EVerticalAlignment::VAlign_Center)
-								.Padding(6.0f, 2.0f, 4.0f, 2.0f)
-								[
-									SNew(SButton)
-									.Text(LOCTEXT("IndexAllBlueprints", "Index All"))
-									.OnClicked( this, &SFindInBlueprints::OnCacheAllBlueprints )
-									.Visibility( this, &SFindInBlueprints::GetCacheAllButtonVisibility )
-									.ToolTip(IDocumentation::Get()->CreateToolTip(
-										LOCTEXT("IndexAlLBlueprints_Tooltip", "Loads all non-indexed Blueprints and stores the data with the DDC. This can be a very slow process, feel free to cancel it whenever and benefit from the Blueprints it indexes."),
-										NULL,
-										TEXT("Shared/Editors/BlueprintEditor"),
-										TEXT("FindInBlueprint_IndexAll")))
-								]
-
-							// Cache progress bar
-							+SHorizontalBox::Slot()
-								.FillWidth(1.0f)
-								.Padding(4.0f, 2.0f, 4.0f, 2.0f)
-								[
-									SNew(SProgressBar)
-									.Percent( this, &SFindInBlueprints::GetPercentCompleteCache )
-									.Visibility( this, &SFindInBlueprints::GetCachingProgressBarVisiblity )
-								]
-
-							// Cancel button
-							+SHorizontalBox::Slot()
-								.AutoWidth()
-								.Padding(4.0f, 2.0f, 0.0f, 2.0f)
-								[
-									SNew(SButton)
-									.Text(LOCTEXT("CancelCacheAll", "Cancel"))
-									.OnClicked( this, &SFindInBlueprints::OnCancelCacheAll )
-									.Visibility( this, &SFindInBlueprints::GetCachingProgressBarVisiblity )
-									.ToolTipText( LOCTEXT("CancelCacheAll_Tooltip", "Stops the caching process from where ever it is, can be started back up where it left off when needed.") )
-								]
-
-							// "X" to remove the bar
-							+SHorizontalBox::Slot()
-								.HAlign(HAlign_Right)
-								[
-									SNew(SButton)
-									.ButtonStyle( FCoreStyle::Get(), "NoBorder" )
-									.ContentPadding(0)
-									.HAlign(HAlign_Center)
-									.VAlign(VAlign_Center)
-									.OnClicked( this, &SFindInBlueprints::OnRemoveCacheBar )
-									.ForegroundColor( FSlateColor::UseForeground() )
-									[
-										SNew(SImage)
-										.Image( FCoreStyle::Get().GetBrush("EditableComboBox.Delete") )
-										.ColorAndOpacity( FSlateColor::UseForeground() )
-									]
-								]
-						]
-
-						+SVerticalBox::Slot()
-							.AutoHeight()
-							.Padding(8.0f, 0.0f, 0.0f, 2.0f)
-						[
-							SNew(STextBlock)
-							.Text(this, &SFindInBlueprints::GetCurrentCacheBlueprintName)
-							.Visibility( this, &SFindInBlueprints::GetCachingBlueprintNameVisiblity )
-							.ColorAndOpacity( FCoreStyle::Get().GetColor("ErrorReporting.ForegroundColor") )
-						]
-					]
-				];
-		}
-	}
-	else
-	{
-		// Because there are no uncached Blueprints, remove the bar
-		OnRemoveCacheBar();
-	}
-}
-
 FReply SFindInBlueprints::OnRemoveCacheBar()
 {
 	if(MainVerticalBox.IsValid() && CacheBarSlot.IsValid())
@@ -1175,9 +1124,6 @@ void SFindInBlueprints::Tick( const FGeometry& AllottedGeometry, const double In
 				TreeView->RequestTreeRefresh();
 			}
 
-			// Add the cache bar if needed.
-			ConditionallyAddCacheBar();
-
 			StreamSearch->EnsureCompletion();
 			StreamSearch.Reset();
 		}
@@ -1215,7 +1161,16 @@ void SFindInBlueprints::FocusForUse(bool bSetFindWithinBlueprint, FString NewSea
 		// Select the first result
 		if (bSelectFirstResult && ItemsFound.Num())
 		{
-			TreeView->SetSelection(ItemsFound[0]);
+			auto ItemToFocusOn = ItemsFound[0];
+
+			// We want the first childmost item to select, as that is the item that is most-likely to be what was searched for (parents being graphs).
+			// Will fail back upward as neccessary to focus on a focusable item
+			while(ItemToFocusOn->Children.Num())
+			{
+				ItemToFocusOn = ItemToFocusOn->Children[0];
+			}
+			TreeView->SetSelection(ItemToFocusOn);
+			ItemToFocusOn->OnClick();
 		}
 	}
 }
@@ -1324,6 +1279,7 @@ void SFindInBlueprints::MatchTokensWithinBlueprint(const TArray<FString>& Tokens
 	{
 		// Insert a fake result to inform user if none found
 		ItemsFound.Add(FSearchResult(new FFindInBlueprintsResult(LOCTEXT("BlueprintSearchNoResults", "No Results found"))));
+		HighlightText = FText::GetEmpty();
 		TreeView->RequestTreeRefresh();
 	}
 	else
@@ -1452,100 +1408,6 @@ EVisibility SFindInBlueprints::GetSearchbarVisiblity() const
 	return StreamSearch.IsValid()? EVisibility::Visible : EVisibility::Collapsed;
 }
 
-FReply SFindInBlueprints::OnCacheAllBlueprints()
-{
-	FFindInBlueprintSearchManager::Get().CacheAllUncachedBlueprints(SharedThis(this));
-	return FReply::Handled();
-}
-
-FReply SFindInBlueprints::OnCancelCacheAll()
-{
-	FFindInBlueprintSearchManager::Get().CancelCacheAll();
-
-	// Resubmit the last search
-	OnSearchTextCommitted(SearchTextField->GetText(), ETextCommit::OnEnter);
-
-	return FReply::Handled();
-}
-
-int32 SFindInBlueprints::GetCurrentCacheIndex() const
-{
-	return FFindInBlueprintSearchManager::Get().GetCurrentCacheIndex();
-}
-
-TOptional<float> SFindInBlueprints::GetPercentCompleteCache() const
-{
-	return FFindInBlueprintSearchManager::Get().GetCacheProgress();
-}
-
-EVisibility SFindInBlueprints::GetCachingProgressBarVisiblity() const
-{
-	return IsCacheInProgress()? EVisibility::Visible : EVisibility::Hidden;
-}
-
-EVisibility SFindInBlueprints::GetCacheAllButtonVisibility() const
-{
-	return IsCacheInProgress()? EVisibility::Collapsed : EVisibility::Visible;
-}
-
-EVisibility SFindInBlueprints::GetCachingBarVisibility() const
-{
-	return (FFindInBlueprintSearchManager::Get().GetNumberUncachedBlueprints() > 0)? EVisibility::Visible : EVisibility::Collapsed;
-}
-
-EVisibility SFindInBlueprints::GetCachingBlueprintNameVisiblity() const
-{
-	return IsCacheInProgress()? EVisibility::Visible : EVisibility::Collapsed;
-}
-
-bool SFindInBlueprints::IsCacheInProgress() const
-{
-	return FFindInBlueprintSearchManager::Get().IsCacheInProgress();
-}
-
-FSlateColor SFindInBlueprints::GetCachingBarColor() const
-{
-	// The caching bar's default color is a darkish red
-	FSlateColor ReturnColor = FSlateColor(FLinearColor(0.4f, 0.0f, 0.0f));
-	if(IsCacheInProgress())
-	{
-		// It turns yellow when in progress
-		ReturnColor = FSlateColor(FLinearColor(0.4f, 0.4f, 0.0f));
-	}
-	return ReturnColor;
-}
-
-FText SFindInBlueprints::GetUncachedBlueprintWarningText() const
-{
-	FFormatNamedArguments Args;
-	Args.Add(TEXT("Count"), FFindInBlueprintSearchManager::Get().GetNumberUncachedBlueprints());
-
-	FText ReturnDisplayText;
-	if(IsCacheInProgress())
-	{
-		Args.Add(TEXT("CurrentIndex"), FFindInBlueprintSearchManager::Get().GetCurrentCacheIndex());
-
-		ReturnDisplayText = FText::Format(LOCTEXT("CachingBlueprints", "Indexing Blueprints... {CurrentIndex}/{Count}"), Args);
-	}
-	else
-	{
-		ReturnDisplayText = FText::Format(LOCTEXT("UncachedBlueprints", "Search incomplete. {Count} Blueprints are not indexed!"), Args);
-	}
-
-	return ReturnDisplayText;
-}
-
-FText SFindInBlueprints::GetCurrentCacheBlueprintName() const
-{
-	return FText::FromString(FFindInBlueprintSearchManager::Get().GetCurrentCacheBlueprintName());
-}
-
-void SFindInBlueprints::OnCacheComplete()
-{
-	// Resubmit the last search, which will also remove the bar if needed
-	OnSearchTextCommitted(SearchTextField->GetText(), ETextCommit::OnEnter);
-}
-
 TSharedPtr<SWidget> SFindInBlueprints::OnContextMenuOpening()
 {
 	const bool bShouldCloseWindowAfterMenuSelection = true;
diff --git a/Engine/Source/Editor/Kismet/Public/FindInBlueprintManager.h b/Engine/Source/Editor/Kismet/Public/FindInBlueprintManager.h
--- a/Engine/Source/Editor/Kismet/Public/FindInBlueprintManager.h
+++ b/Engine/Source/Editor/Kismet/Public/FindInBlueprintManager.h
@@ -10,6 +10,11 @@ struct KISMET_API FFindInBlueprintSearchTags
 	/** Properties tag, for Blueprint variables */
 	static const FText FiB_Properties;
 
+	/** Components tags */
+	static const FText FiB_Components;
+	static const FText FiB_IsSCSComponent;
+	/** End Components tags */
+
 	/** Nodes tag */
 	static const FText FiB_Nodes;
 
@@ -153,48 +158,23 @@ public:
 	/**
 	 * Query for a single, specific Blueprint's search block
 	 *
-	 * @param InBlueprint		The Blueprint to search for
-	 * @return					The search block
-	 */
-	FString QuerySingleBlueprint(UBlueprint* InBlueprint);
-
-	/** Returns the number of uncached Blueprints */
-	int32 GetNumberUncachedBlueprints() const { return UncachedBlueprints.Num(); };
-
-	/**
-	 * Starts caching all uncached Blueprints at a rate of 1 per tick
-	 *
-	 * @param InSourceWidget		The source FindInBlueprints widget, this widget will be informed when caching is complete
+	 * @param InBlueprint				The Blueprint to search for
+	 * @param bInRebuildSearchData		When TRUE the search data will be freshly collected
+	 * @return							The search block
 	 */
-	void CacheAllUncachedBlueprints(TWeakPtr< class SFindInBlueprints > InSourceWidget);
-
-	/** Stops the caching process where it currently is at, the rest can be continued later */
-	void CancelCacheAll();
-
-	/** Returns the current index in the caching */
-	int32 GetCurrentCacheIndex() const;
-
-	/** Returns the name of the current Blueprint being cached */
-	FString GetCurrentCacheBlueprintName() const;
-
-	/** Returns the progress complete on the caching */
-	float GetCacheProgress() const;
-
-	/**
-	 * Callback to note that Blueprint caching is complete
-	 *
-	 * @param InNumberCached		The number of Blueprints cached, to be chopped off the existing array so the rest (if any) can be finished later
-	 */
-	void FinishedCachingBlueprints(int32 InNumberCached);
-
-	/** Returns TRUE if Blueprints are being cached. */
-	bool IsCacheInProgress() const;
+	FString QuerySingleBlueprint(UBlueprint* InBlueprint, bool bInRebuildSearchData = true);
 
 	/** Converts a string of hex characters, previously converted by ConvertFTextToHexString, to an FText. */
 	static FText ConvertHexStringToFText(FString InHexString);
 
 	/** Serializes an FText to memory and converts the memory into a string of hex characters */
 	static FString ConvertFTextToHexString(FText InValue);
+
+	/** Serializes an FString to memory and converts the memory into a string of hex characters */
+	static FString ConvertFStringToHexString(FString InValue);
+
+	static TSharedPtr< class FJsonObject > ConvertJsonStringToObject(FString InJsonString);
+
 private:
 	/** Initializes the FiB manager */
 	void Initialize();
@@ -257,12 +237,6 @@ protected:
 	/** Because we are unable to query for the module on another thread, cache it for use later */
 	class FAssetRegistryModule* AssetRegistryModule;
 
-	/** Blueprint paths that have not been cached for searching due to lack of data, this means that they are either older Blueprints, or the DDC cannot find the data */
-	TArray<FName> UncachedBlueprints;
-
-	/** Tickable object that does the caching of uncached Blueprints at a rate of once per tick */
-	class FCacheAllBlueprintsTickableObject* CachingObject;
-
 	/** A list of all the FindInBlueprints widgets that need to be informed that caching is complete */
 	TArray< TWeakPtr<class SFindInBlueprints> > OnCachingCompleteCallbackWidgets;
-};
\ No newline at end of file
+};
diff --git a/Engine/Source/Editor/Kismet/Public/FindInBlueprints.h b/Engine/Source/Editor/Kismet/Public/FindInBlueprints.h
--- a/Engine/Source/Editor/Kismet/Public/FindInBlueprints.h
+++ b/Engine/Source/Editor/Kismet/Public/FindInBlueprints.h
@@ -191,6 +191,9 @@ protected:
 
 	/** The Property's sub-category name, for the pin display type */
 	FString PinSubCategory;
+
+	/** TRUE if the property is an SCS_Component */
+	bool bIsSCSComponent;
 };
 
 /** Graphs, such as functions and macros, are stored here */
@@ -213,7 +216,6 @@ protected:
 	EGraphType GraphType;
 };
 
-
 /*Widget for searching for (functions/events) across all blueprints or just a single blueprint */
 class SFindInBlueprints: public SCompoundWidget
 {
@@ -231,8 +233,6 @@ public:
 	/** Focuses this widget's search box, and changes the mode as well, and optionally the search terms */
 	void FocusForUse(bool bSetFindWithinBlueprint, FString NewSearchTerms = FString(), bool bSelectFirstResult = false);
 
-	/** Called when caching Blueprints is complete, if this widget initiated the indexing */
-	void OnCacheComplete();
 private:
 
 	/** Register any Find-in-Blueprint commands */
@@ -274,48 +274,9 @@ private:
 	/** Returns the progress bar visiblity */
 	EVisibility GetSearchbarVisiblity() const;
 
-	/** Adds the "cache" bar at the bottom of the Find-in-Blueprints widget, to notify the user that the search is incomplete */
-	void ConditionallyAddCacheBar();
-
 	/** Callback to remove the "cache" bar when a button is pressed */
 	FReply OnRemoveCacheBar();
 
-	/** Callback to return the cache bar's display text, informing the user of the situation */
-	FText GetUncachedBlueprintWarningText() const;
-
-	/** Callback to return the cache bar's current indexing Blueprint name */
-	FText GetCurrentCacheBlueprintName() const;
-
-	/** Callback to cache all uncached Blueprints */
-	FReply OnCacheAllBlueprints();
-
-	/** Callback to cancel the caching process */
-	FReply OnCancelCacheAll();
-
-	/** Retrieves the current index of the Blueprint caching process */
-	int32 GetCurrentCacheIndex() const;
-
-	/** Gets the percent complete of the caching process */
-	TOptional<float> GetPercentCompleteCache() const;
-
-	/** Returns the visibility of the caching progress bar, visible when in progress, hidden when not */
-	EVisibility GetCachingProgressBarVisiblity() const;
-
-	/** Returns the visibility of the "Cache All" button, visible when not caching, collapsed when caching is in progress */
-	EVisibility GetCacheAllButtonVisibility() const;
-
-	/** Returns the caching bar's visibility, it goes invisible when there is nothing to be cached. The next search will remove this bar or make it visible again */
-	EVisibility GetCachingBarVisibility() const;
-	
-	/** Returns the visibility of the caching Blueprint name, visible when in progress, collapsed when not */
-	EVisibility GetCachingBlueprintNameVisiblity() const;
-
-	/** Returns TRUE if Blueprint caching is in progress */
-	bool IsCacheInProgress() const;
-
-	/** Returns the color of the caching bar */
-	FSlateColor GetCachingBarColor() const;
-
 	/** Callback to build the context menu when right clicking in the tree */
 	TSharedPtr<SWidget> OnContextMenuOpening();
 
diff --git a/Engine/Source/Editor/UnrealEd/Private/Kismet2/Kismet2.cpp b/Engine/Source/Editor/UnrealEd/Private/Kismet2/Kismet2.cpp
--- a/Engine/Source/Editor/UnrealEd/Private/Kismet2/Kismet2.cpp
+++ b/Engine/Source/Editor/UnrealEd/Private/Kismet2/Kismet2.cpp
@@ -355,12 +355,11 @@ void FKismetEditorUtilities::CompileBlueprint(UBlueprint* BlueprintObj, bool bIs
 	bool const bStartedWithUnsavedChanges = (BlueprintPackage != NULL) ? BlueprintPackage->IsDirty() : true;
 #if WITH_EDITOR
 	// Do not want to run this code without the editor present nor when running commandlets.
-	if (GEditor && GIsEditor && !IsRunningCommandlet())
+	if (GEditor && GIsEditor)
 	{
 		// We do not want to regenerate a search Guid during loads, nothing has changed in the Blueprint and it is cached elsewhere
 		if (!bIsRegeneratingOnLoad)
 		{
-			BlueprintObj->SearchGuid = FGuid::NewGuid();
 			FFindInBlueprintSearchManager::Get().AddOrUpdateBlueprintSearchMetadata(BlueprintObj);
 		}
 	}
@@ -1254,7 +1253,7 @@ void FKismetEditorUtilities::ShowActorReferencesInLevelScript(const AActor* Acto
 
 				const bool bSetFindWithinBlueprint = true;
 				const bool bSelectFirstResult = true;
-				BlueprintEditor->SummonSearchUI(bSetFindWithinBlueprint, Actor->GetName(), bSelectFirstResult);
+				BlueprintEditor->SummonSearchUI(bSetFindWithinBlueprint, Actor->GetActorLabel(), bSelectFirstResult);
 			}
 		}
 
diff --git a/Engine/Source/Runtime/Core/Public/Containers/UnrealString.h b/Engine/Source/Runtime/Core/Public/Containers/UnrealString.h
--- a/Engine/Source/Runtime/Core/Public/Containers/UnrealString.h
+++ b/Engine/Source/Runtime/Core/Public/Containers/UnrealString.h
@@ -1427,6 +1427,52 @@ FORCEINLINE uint32 GetTypeHash( const FString& S )
 	return FCrc::Strihash_DEPRECATED(*S);
 }
 
+/** 
+ * Convert an array of bytes to a TCHAR
+ * @param In byte array values to convert
+ * @param Count number of bytes to convert
+ * @return Valid string representing bytes.
+ */
+inline FString BytesToString(const uint8* In, int32 Count)
+{
+	FString Result;
+	Result.Empty(Count);
+
+	while (Count)
+	{
+		// Put the byte into an int16 and add 1 to it, this keeps anything from being put into the string as a null terminator
+		int16 Value = *In;
+		Value += 1;
+
+		Result += TCHAR(Value);
+
+		++In;
+		Count--;
+	}
+	return Result;
+}
+
+/** 
+ * Convert FString of bytes into the byte array.
+ * @param String		The FString of byte values
+ * @param OutBytes		Ptr to memory must be preallocated large enough
+ * @param MaxBufferSize	Max buffer size of the OutBytes array, to prevent overflow
+ * @return	The number of bytes copied
+ */
+inline int32 StringToBytes( const FString& String, uint8* OutBytes, int32 MaxBufferSize )
+{
+	int32 NumBytes = 0;
+	const TCHAR* CharPos = *String;
+
+	while( *CharPos && NumBytes < MaxBufferSize)
+	{
+		OutBytes[ NumBytes ] = (int8)(*CharPos - 1);
+		CharPos++;
+		++NumBytes;
+	}
+	return NumBytes - 1;
+}
+
 /** @return Char value of Nibble */
 inline TCHAR NibbleToTChar(uint8 Num)
 {
diff --git a/Engine/Source/Runtime/Engine/Classes/Engine/Blueprint.h b/Engine/Source/Runtime/Engine/Classes/Engine/Blueprint.h
--- a/Engine/Source/Runtime/Engine/Classes/Engine/Blueprint.h
+++ b/Engine/Source/Runtime/Engine/Classes/Engine/Blueprint.h
@@ -504,10 +504,6 @@ public:
 
 	virtual void GetReparentingRules(TSet< const UClass* >& AllowedChildrenOfClasses, TSet< const UClass* >& DisallowedChildrenOfClasses) const;
 
-protected:
-	/** Gets asset registry tags */
-	virtual void GetAssetRegistryTags(TArray<FAssetRegistryTag>& OutTags) const override;
-
 private:
 	/** Sets the current object being debugged */
 	void DebuggingWorldRegistrationHelper(UObject* ObjectProvidingWorld, UObject* ValueToRegister);
@@ -529,6 +525,7 @@ public:
 	virtual void PostLoad() override;
 	virtual void PostLoadSubobjects( FObjectInstancingGraph* OuterInstanceGraph ) override;
 	virtual bool Modify(bool bAlwaysMarkDirty = true) override;
+	virtual void GetAssetRegistryTags(TArray<FAssetRegistryTag>& OutTags) const override;
 	// End of UObject interface
 
 	/** Consigns the GeneratedClass and the SkeletonGeneratedClass to oblivion, and nulls their references */
diff --git a/Engine/Source/Runtime/Engine/Classes/Engine/World.h b/Engine/Source/Runtime/Engine/Classes/Engine/World.h
--- a/Engine/Source/Runtime/Engine/Classes/Engine/World.h
+++ b/Engine/Source/Runtime/Engine/Classes/Engine/World.h
@@ -1663,6 +1663,7 @@ public:
 	static void AddReferencedObjects(UObject* InThis, FReferenceCollector& Collector);
 #if WITH_EDITOR
 	virtual bool Rename(const TCHAR* NewName = NULL, UObject* NewOuter = NULL, ERenameFlags Flags = REN_None) override;
+	virtual void GetAssetRegistryTags(TArray<FAssetRegistryTag>& OutTags) const override;
 #endif
 	virtual void PostDuplicate(bool bDuplicateForPIE) override;
 	// End UObject Interface
diff --git a/Engine/Source/Runtime/Engine/Private/Blueprint.cpp b/Engine/Source/Runtime/Engine/Private/Blueprint.cpp
--- a/Engine/Source/Runtime/Engine/Private/Blueprint.cpp
+++ b/Engine/Source/Runtime/Engine/Private/Blueprint.cpp
@@ -11,6 +11,7 @@
 #include "Editor/UnrealEd/Public/Kismet2/KismetEditorUtilities.h"
 #include "Editor/UnrealEd/Public/Kismet2/CompilerResultsLog.h"
 #include "Editor/UnrealEd/Public/Kismet2/StructureEditorUtils.h"
+#include "Editor/Kismet/Public/FindInBlueprintManager.h"
 #include "Editor/UnrealEd/Public/Editor.h"
 #include "Crc.h"
 #include "MessageLog.h"
@@ -360,6 +361,12 @@ void UBlueprint::Serialize(FArchive& Ar)
 			}
 		}
 	}
+
+	if(Ar.IsSaving())
+	{
+		// Cache the BP for use
+		FFindInBlueprintSearchManager::Get().AddOrUpdateBlueprintSearchMetadata(this);
+	}
 #endif // WITH_EDITORONLY_DATA
 }
 
@@ -690,10 +697,7 @@ void UBlueprint::GetAssetRegistryTags(TArray<FAssetRegistryTag>& OutTags) const
 		FBlueprintEditorUtils::IsDataOnlyBlueprint(this) ? TEXT("True") : TEXT("False"),
 		FAssetRegistryTag::TT_Alphabetical ) );
 
-	if(SearchGuid.IsValid())
-	{
-		OutTags.Add( FAssetRegistryTag("SearchGuid", SearchGuid.ToString(), FAssetRegistryTag::TT_Hidden) );
-	}
+	OutTags.Add( FAssetRegistryTag("FiB", FFindInBlueprintSearchManager::Get().QuerySingleBlueprint((UBlueprint*)this, false), FAssetRegistryTag::TT_Hidden) );
 }
 
 FString UBlueprint::GetFriendlyName() const
diff --git a/Engine/Source/Runtime/Engine/Private/World.cpp b/Engine/Source/Runtime/Engine/Private/World.cpp
--- a/Engine/Source/Runtime/Engine/Private/World.cpp
+++ b/Engine/Source/Runtime/Engine/Private/World.cpp
@@ -5354,6 +5354,19 @@ void UWorld::ChangeFeatureLevel(ERHIFeatureLevel::Type InFeatureLevel)
 }
 #endif // WITH_EDITOR
 
+#if WITH_EDITOR
+void UWorld::GetAssetRegistryTags(TArray<FAssetRegistryTag>& OutTags) const
+{
+	if(PersistentLevel && PersistentLevel->OwningWorld)
+	{
+		if(ULevelScriptBlueprint* LevelBlueprint = PersistentLevel->GetLevelScriptBlueprint(true))
+		{
+			LevelBlueprint->GetAssetRegistryTags(OutTags);
+		}
+	}
+}
+#endif
+
 /**
 * Dump visible actors in current world.
 */
diff --git a/Engine/Source/Runtime/Json/Public/Serialization/JsonReader.h b/Engine/Source/Runtime/Json/Public/Serialization/JsonReader.h
--- a/Engine/Source/Runtime/Json/Public/Serialization/JsonReader.h
+++ b/Engine/Source/Runtime/Json/Public/Serialization/JsonReader.h
@@ -132,9 +132,9 @@ public:
 		return ReadUntilMatching(EJsonNotation::ArrayEnd);
 	}
 
-	FORCEINLINE const FString& GetIdentifier() const { return Identifier; }
+	FORCEINLINE virtual const FString& GetIdentifier() const { return Identifier; }
 
-	FORCEINLINE const FString& GetValueAsString() const 
+	FORCEINLINE virtual  const FString& GetValueAsString() const 
 	{ 
 		check(CurrentToken == EJsonToken::String);
 		return StringValue;
@@ -785,7 +785,7 @@ public:
 		}
 	}
 
-private:
+protected:
 
 	/**
 	 * Parses a string containing Json information.
@@ -807,7 +807,7 @@ private:
 		Stream = Reader;
 	}
 
-private:
+protected:
 
 	const FString Content;
 	FBufferReader* Reader;
