Fixing Mac Compiler error with template classes.

diff --git a/Engine/Source/Editor/Kismet/Private/FindInBlueprintManager.cpp b/Engine/Source/Editor/Kismet/Private/FindInBlueprintManager.cpp
--- a/Engine/Source/Editor/Kismet/Private/FindInBlueprintManager.cpp
+++ b/Engine/Source/Editor/Kismet/Private/FindInBlueprintManager.cpp
@@ -249,7 +249,7 @@ namespace BlueprintSearchMetaDataHelpers
 		}
 
 		TJsonFindInBlueprintStringReader( FArchive* InStream,  TMap< int32, FText >& InLookupTable )
-			: TJsonReader(InStream)
+			: TJsonReader<CharType>(InStream)
 			, LookupTable(MoveTemp(InLookupTable))
 		{
 
