More fixes for templates in Clang.

diff --git a/Engine/Source/Editor/Kismet/Private/FindInBlueprintManager.cpp b/Engine/Source/Editor/Kismet/Private/FindInBlueprintManager.cpp
--- a/Engine/Source/Editor/Kismet/Private/FindInBlueprintManager.cpp
+++ b/Engine/Source/Editor/Kismet/Private/FindInBlueprintManager.cpp
@@ -258,11 +258,11 @@ namespace BlueprintSearchMetaDataHelpers
 		FORCEINLINE virtual const FString& GetIdentifier() const override
 		{
 			// The identifier from Json is a Hex value that must be looked up in the LookupTable to find the FText it represents
-			if(const FText* LookupText = LookupTable.Find(FCString::Atoi(*Identifier)))
+			if(const FText* LookupText = LookupTable.Find(FCString::Atoi(*this->Identifier)))
 			{
 				return LookupText->ToString();
 			}
-			return Identifier;
+			return this->Identifier;
 		}
 
 		FORCEINLINE virtual  const FString& GetValueAsString() const override
@@ -299,13 +299,13 @@ namespace BlueprintSearchMetaDataHelpers
 
 		FORCEINLINE virtual const FString& GetIdentifier() const override
 		{
-			if(Identifier.Len())
+			if(this->Identifier.Len())
 			{
 				// Remove the const, non-ideal, but this code is to bridge between the old DDC method and the new Asset Registry method of story FiB data.
-				FString& IdentifierNonConst = const_cast<FString&>(Identifier);
-				IdentifierNonConst = FFindInBlueprintSearchManager::ConvertHexStringToFText(Identifier).ToString();
+				FString& IdentifierNonConst = const_cast<FString&>(this->Identifier);
+				IdentifierNonConst = FFindInBlueprintSearchManager::ConvertHexStringToFText(this->Identifier).ToString();
 			}
-			return Identifier;
+			return this->Identifier;
 		}
 
 		FORCEINLINE virtual  const FString& GetValueAsString() const override
