Template issues with Clang

diff --git a/Engine/Source/Editor/Kismet/Private/FindInBlueprintManager.cpp b/Engine/Source/Editor/Kismet/Private/FindInBlueprintManager.cpp
--- a/Engine/Source/Editor/Kismet/Private/FindInBlueprintManager.cpp
+++ b/Engine/Source/Editor/Kismet/Private/FindInBlueprintManager.cpp
@@ -267,13 +267,13 @@ namespace BlueprintSearchMetaDataHelpers
 
 		FORCEINLINE virtual  const FString& GetValueAsString() const override
 		{ 
-			check( CurrentToken == EJsonToken::String ); 
+			check( this->CurrentToken == EJsonToken::String ); 
 			// The string value from Json is a Hex value that must be looked up in the LookupTable to find the FText it represents
-			if(const FText* LookupText = LookupTable.Find(FCString::Atoi(*StringValue)))
+			if(const FText* LookupText = LookupTable.Find(FCString::Atoi(*this->StringValue)))
 			{
 				return LookupText->ToString();
 			}
-			return StringValue;
+			return this->StringValue;
 		}
 
 		TMap< int32, FText > LookupTable;
@@ -310,14 +310,14 @@ namespace BlueprintSearchMetaDataHelpers
 
 		FORCEINLINE virtual  const FString& GetValueAsString() const override
 		{ 
-			check( CurrentToken == EJsonToken::String );
-			if(StringValue.Len())
+			check( this->CurrentToken == EJsonToken::String );
+			if(this->StringValue.Len())
 			{
 				// Remove the const, non-ideal, but this code is to bridge between the old DDC method and the new Asset Registry method of story FiB data.
-				FString& StringValueNonConst = const_cast<FString&>(StringValue);
-				StringValueNonConst = FFindInBlueprintSearchManager::ConvertHexStringToFText(StringValue).ToString();
+				FString& StringValueNonConst = const_cast<FString&>(this->StringValue);
+				StringValueNonConst = FFindInBlueprintSearchManager::ConvertHexStringToFText(this->StringValue).ToString();
 			}
-			return StringValue;
+			return this->StringValue;
 		}
 	};
 
