UE-5017 Editor crashes when creating an input boolean inside of a blueprint interface with a default value of 10

diff --git a/Engine/Source/Editor/Kismet/Private/BlueprintDetailsCustomization.cpp b/Engine/Source/Editor/Kismet/Private/BlueprintDetailsCustomization.cpp
--- a/Engine/Source/Editor/Kismet/Private/BlueprintDetailsCustomization.cpp
+++ b/Engine/Source/Editor/Kismet/Private/BlueprintDetailsCustomization.cpp
@@ -2013,13 +2013,17 @@ FText FBlueprintGraphArgumentLayout::OnGetArgDefaultValueText() const
 
 void FBlueprintGraphArgumentLayout::OnArgDefaultValueCommitted(const FText& NewText, ETextCommit::Type InTextCommit)
 {
-	if (!NewText.IsEmpty() && !ShouldPinBeReadOnly() && (InTextCommit == ETextCommit::OnEnter || InTextCommit == ETextCommit::OnUserMovedFocus) && ParamItemPtr.IsValid())
+	auto GraphActionDetailsPinned = GraphActionDetailsPtr.Pin();
+	if (!NewText.IsEmpty() 
+		&& !ShouldPinBeReadOnly() 
+		&& (InTextCommit == ETextCommit::OnEnter || InTextCommit == ETextCommit::OnUserMovedFocus) 
+		&& ParamItemPtr.IsValid() 
+		&& GraphActionDetailsPinned.IsValid())
 	{
 		bool bSuccess = TargetNode->ModifyUserDefinedPinDefaultValue(ParamItemPtr.Pin(), NewText.ToString());
-		
 		if (bSuccess)
 		{
-			GraphActionDetailsPtr.Pin()->OnParamsChanged(TargetNode);
+			GraphActionDetailsPinned->OnParamsChanged(TargetNode);
 		}
 	}
 }
