Double clicking on actor references will select and focus the actor in the level viewport and the main editor details panel.

diff --git a/Engine/Source/Editor/Kismet/Private/BlueprintEditor.cpp b/Engine/Source/Editor/Kismet/Private/BlueprintEditor.cpp
--- a/Engine/Source/Editor/Kismet/Private/BlueprintEditor.cpp
+++ b/Engine/Source/Editor/Kismet/Private/BlueprintEditor.cpp
@@ -4987,7 +4987,7 @@ void FBlueprintEditor::OnSelectReferenceInLevel()
 			// Now select the actors.
 			for (int32 iActor = 0; iActor < ActorsToSelect.Num(); iActor++)
 			{
-				GEditor->SelectActor(ActorsToSelect[ iActor ], true, false, false);
+				GEditor->SelectActor(ActorsToSelect[ iActor ], true, true, false);
 			}
 
 			// Execute the command to move camera to the object(s).
@@ -5258,28 +5258,45 @@ void FBlueprintEditor::OnNodeDoubleClicked(class UEdGraphNode* Node)
 	}
 	else if (UObject* HyperlinkTarget = Node->GetJumpTargetForDoubleClick())
 	{
-		// Check to see if our outer chain contains a blueprint. If we're inside a blueprint (a graph, pin, etc.) then
-		// focus on the target; otherwise open the editor for the target.
-		UBlueprint* TargetBP = Cast<UBlueprint>(const_cast<UObject*>(HyperlinkTarget));
-		if(TargetBP == NULL)
+		// true if jumping to a level actor, do not handle the node further
+		bool bJumpToLevelActor = false;
+
+		// If double clicking on a node that references level actors, jump to hyperlink so that they are selected in the level
+		UK2Node* K2Node = Cast<UK2Node>(Node);
+		if(K2Node)
 		{
-			for(UObject* TestOuter = HyperlinkTarget->GetOuter(); TestOuter; TestOuter = TestOuter->GetOuter())
+			if (AActor* Actor = K2Node->GetReferencedLevelActor())
 			{
-				TargetBP = Cast<UBlueprint>(TestOuter);
-				if(TargetBP != NULL)
-				{
-					break;
-				}
+				JumpToHyperlink(Actor);
+				bJumpToLevelActor = true;
 			}
 		}
 
-		if(TargetBP)
+		if(!bJumpToLevelActor)
 		{
-			FKismetEditorUtilities::BringKismetToFocusAttentionOnObject(HyperlinkTarget);
-		}
-		else
-		{
-			FAssetEditorManager::Get().OpenEditorForAsset(HyperlinkTarget);
+			// Check to see if our outer chain contains a blueprint. If we're inside a blueprint (a graph, pin, etc.) then
+			// focus on the target; otherwise open the editor for the target.
+			UBlueprint* TargetBP = Cast<UBlueprint>(const_cast<UObject*>(HyperlinkTarget));
+			if(TargetBP == NULL)
+			{
+				for(UObject* TestOuter = HyperlinkTarget->GetOuter(); TestOuter; TestOuter = TestOuter->GetOuter())
+				{
+					TargetBP = Cast<UBlueprint>(TestOuter);
+					if(TargetBP != NULL)
+					{
+						break;
+					}
+				}
+			}
+
+			if(TargetBP)
+			{
+				FKismetEditorUtilities::BringKismetToFocusAttentionOnObject(HyperlinkTarget);
+			}
+			else
+			{
+				FAssetEditorManager::Get().OpenEditorForAsset(HyperlinkTarget);
+			}
 		}
 	}
 }
