Removed "developer" settings from compile options sub-menu.

diff --git a/Engine/Source/Editor/Kismet/Private/SBlueprintEditorToolbar.cpp b/Engine/Source/Editor/Kismet/Private/SBlueprintEditorToolbar.cpp
--- a/Engine/Source/Editor/Kismet/Private/SBlueprintEditorToolbar.cpp
+++ b/Engine/Source/Editor/Kismet/Private/SBlueprintEditorToolbar.cpp
@@ -591,10 +591,10 @@ static TSharedRef<SWidget> BlueprintEditorToolbarImpl::GenerateCompileOptionsWid
 
 	MenuBuilder.AddMenuEntry(Commands.JumpToErrorNode);
 
-	MenuBuilder.AddSubMenu(
-		LOCTEXT("DevCompileSubMenu", "Developer"),
-		LOCTEXT("DevCompileSubMenu_ToolTip", "Advanced settings that aid in devlopment/debugging of the Blueprint system as a whole."),
-		FNewMenuDelegate::CreateStatic(&BlueprintEditorToolbarImpl::MakeCompileDeveloperSubMenu));
+// 	MenuBuilder.AddSubMenu(
+// 		LOCTEXT("DevCompileSubMenu", "Developer"),
+// 		LOCTEXT("DevCompileSubMenu_ToolTip", "Advanced settings that aid in devlopment/debugging of the Blueprint system as a whole."),
+// 		FNewMenuDelegate::CreateStatic(&BlueprintEditorToolbarImpl::MakeCompileDeveloperSubMenu));
 
 	return MenuBuilder.MakeWidget();
 }
