UE-5372, UE-4893: Custom Event created from event dispatcher pin has the same signature as the ED.

diff --git a/Engine/Source/Editor/BlueprintGraph/Classes/K2Node_CustomEvent.h b/Engine/Source/Editor/BlueprintGraph/Classes/K2Node_CustomEvent.h
--- a/Engine/Source/Editor/BlueprintGraph/Classes/K2Node_CustomEvent.h
+++ b/Engine/Source/Editor/BlueprintGraph/Classes/K2Node_CustomEvent.h
@@ -26,6 +26,7 @@ class UK2Node_CustomEvent : public UK2Node_Event
 	virtual FString GetDocumentationLink() const override;
 	virtual FString GetDocumentationExcerptName() const override;
 	virtual FName GetPaletteIcon(FLinearColor& OutColor) const override;
+	virtual void AutowireNewNode(UEdGraphPin* FromPin) override;
 	// End UEdGraphNode interface
 
 	// Begin UK2Node interface
diff --git a/Engine/Source/Editor/BlueprintGraph/Private/K2Node_CustomEvent.cpp b/Engine/Source/Editor/BlueprintGraph/Private/K2Node_CustomEvent.cpp
--- a/Engine/Source/Editor/BlueprintGraph/Private/K2Node_CustomEvent.cpp
+++ b/Engine/Source/Editor/BlueprintGraph/Private/K2Node_CustomEvent.cpp
@@ -389,4 +389,17 @@ FName UK2Node_CustomEvent::GetPaletteIcon(FLinearColor& OutColor) const
 	return bCallInEditor ? TEXT("GraphEditor.CallInEditorEvent_16x") : TEXT("GraphEditor.CustomEvent_16x");
 }
 
+void UK2Node_CustomEvent::AutowireNewNode(UEdGraphPin* FromPin)
+{
+	Super::AutowireNewNode(FromPin);
+
+	if (auto DelegateOutPin = FindPin(DelegateOutputName))
+	{
+		if (DelegateOutPin->LinkedTo.Num())
+		{
+			ReconstructNode();
+		}
+	}
+}
+
 #undef LOCTEXT_NAMESPACE
diff --git a/Engine/Source/Editor/Kismet/Private/BlueprintActionMenuItem.cpp b/Engine/Source/Editor/Kismet/Private/BlueprintActionMenuItem.cpp
--- a/Engine/Source/Editor/Kismet/Private/BlueprintActionMenuItem.cpp
+++ b/Engine/Source/Editor/Kismet/Private/BlueprintActionMenuItem.cpp
@@ -33,7 +33,7 @@ namespace FBlueprintMenuActionItemImpl
 	 * @param  Bindings			Any bindings you want applied after the node has been spawned.
 	 * @return The spawned node (could be an existing one if the event was already placed).
 	 */
-	static UEdGraphNode* InvokeAction(const UBlueprintNodeSpawner* Action, UEdGraph* ParentGraph, FVector2D const Location, IBlueprintNodeBinder::FBindingSet const& Bindings);
+	static UEdGraphNode* InvokeAction(const UBlueprintNodeSpawner* Action, UEdGraph* ParentGraph, FVector2D const Location, IBlueprintNodeBinder::FBindingSet const& Bindings, bool& bOutNewNode);
 }
 
 //------------------------------------------------------------------------------
@@ -60,20 +60,21 @@ static void FBlueprintMenuActionItemImpl::DirtyBlueprintFromNewNode(UEdGraphNode
 }
 
 //------------------------------------------------------------------------------
-static UEdGraphNode* FBlueprintMenuActionItemImpl::InvokeAction(const UBlueprintNodeSpawner* Action, UEdGraph* ParentGraph, FVector2D const Location, IBlueprintNodeBinder::FBindingSet const& Bindings)
+static UEdGraphNode* FBlueprintMenuActionItemImpl::InvokeAction(const UBlueprintNodeSpawner* Action, UEdGraph* ParentGraph, FVector2D const Location, IBlueprintNodeBinder::FBindingSet const& Bindings, bool& bOutNewNode)
 {
 	int32 const PreSpawnNodeCount = ParentGraph->Nodes.Num();
 	// this could return an existing node
 	UEdGraphNode* SpawnedNode = Action->Invoke(ParentGraph, Bindings, Location);
 
 	// if a returned node wasn't one that previously existed in the graph
-	if (PreSpawnNodeCount < ParentGraph->Nodes.Num())
+	const bool bNewNode = PreSpawnNodeCount < ParentGraph->Nodes.Num();
+	bOutNewNode = bNewNode;
+	if (bNewNode)
 	{
 		check(SpawnedNode != nullptr);
 		SpawnedNode->SnapToGrid(SNodePanel::GetSnapGridSize());
 
 		FBlueprintEditorUtils::AnalyticsTrackNewNode(SpawnedNode);
-		DirtyBlueprintFromNewNode(SpawnedNode);
 	}
 	// if this node already existed, then we just want to focus on that node...
 	// some node types are only allowed one instance per blueprint (like events)
@@ -148,11 +149,14 @@ UEdGraphNode* FBlueprintActionMenuItem::PerformAction(UEdGraph* ParentGraph, UEd
 			}
 		}
 
-		LastSpawnedNode = InvokeAction(Action, ParentGraph, ModifiedLocation, BindingsSubset);
+		bool bNewNode = false;
+		LastSpawnedNode = InvokeAction(Action, ParentGraph, ModifiedLocation, BindingsSubset, /*out*/ bNewNode);
 		// could already be an existent node, so we have to add here (can't 
 		// catch it as we go through all new nodes)
 		NodesToFocus.Add(LastSpawnedNode);
 
+		//NOTE: Between the new node is spawned and AutowireNewNode is called, the blueprint should not be compiled.
+
 		if (FromPin != nullptr)
 		{
 			// make sure to auto-wire after we position the new node (in case
@@ -160,6 +164,11 @@ UEdGraphNode* FBlueprintActionMenuItem::PerformAction(UEdGraph* ParentGraph, UEd
 			LastSpawnedNode->AutowireNewNode(FromPin);
 		}
 
+		if (bNewNode)
+		{
+			DirtyBlueprintFromNewNode(LastSpawnedNode);
+		}
+
 		// Increase the node location a safe distance so follow-up nodes are not stacked
 		ModifiedLocation.Y += UEdGraphSchema_K2::EstimateNodeHeight(LastSpawnedNode);
 
