Fixes potential crash during garbage collection with Find-in-Blueprints if a Blueprint is cleaned up.

diff --git a/Engine/Source/Editor/Kismet/Private/FindInBlueprintManager.cpp b/Engine/Source/Editor/Kismet/Private/FindInBlueprintManager.cpp
--- a/Engine/Source/Editor/Kismet/Private/FindInBlueprintManager.cpp
+++ b/Engine/Source/Editor/Kismet/Private/FindInBlueprintManager.cpp
@@ -775,7 +775,7 @@ void FFindInBlueprintSearchManager::OnAssetLoaded(UObject* InAsset)
 		if(IndexPtr)
 		{
 			// That index should never have a Blueprint already, but if it does, it should be the same Blueprint!
-			ensure(!SearchArray[*IndexPtr].Blueprint || SearchArray[*IndexPtr].Blueprint == BlueprintAsset);
+			ensure(!SearchArray[*IndexPtr].Blueprint.IsValid() || SearchArray[*IndexPtr].Blueprint == BlueprintAsset);
 			//check(!SearchArray[*IndexPtr].Blueprint);
 			SearchArray[*IndexPtr].Blueprint = BlueprintAsset;
 		}
@@ -932,7 +932,7 @@ bool FFindInBlueprintSearchManager::ContinueSearchQuery(const FStreamSearch* InS
 		while(SearchIdx < SearchArray.Num())
 		{
 			// If the Blueprint is not marked for deletion, and the asset is valid, we will check to see if we want to refresh the searchable data.
-			if( SearchArray[SearchIdx].bMarkedForDeletion || (SearchArray[SearchIdx].Blueprint && SearchArray[SearchIdx].Blueprint->HasAllFlags(RF_PendingKill)) )
+			if( SearchArray[SearchIdx].bMarkedForDeletion || (SearchArray[SearchIdx].Blueprint.IsValid() && SearchArray[SearchIdx].Blueprint->HasAllFlags(RF_PendingKill)) )
 			{
 				// Mark it for deletion, it will be removed on next save
 				SearchArray[SearchIdx].bMarkedForDeletion = true;
@@ -942,7 +942,7 @@ bool FFindInBlueprintSearchManager::ContinueSearchQuery(const FStreamSearch* InS
  				OutSearchData = SearchArray[SearchIdx++];
 				if(OutSearchData.Value.IsEmpty())
 				{
-					if(OutSearchData.Blueprint)
+					if(OutSearchData.Blueprint.IsValid())
 					{
 						RetrieveDDCData(SearchArray[SearchIdx - 1], *OutSearchData.Blueprint->GetPathName());
 					}
@@ -1106,7 +1106,7 @@ void FFindInBlueprintSearchManager::CleanCache()
 		// Here it builds the new map/array, clean of deleted content.
 
 		// If the database item is not marked for deletion and not pending kill (if loaded), keep it in the database
-		if( !SearchArray[SearchValuePair.Value].bMarkedForDeletion && !(SearchArray[SearchValuePair.Value].Blueprint && SearchArray[SearchValuePair.Value].Blueprint->HasAllFlags(RF_PendingKill)) )
+		if( !SearchArray[SearchValuePair.Value].bMarkedForDeletion && !(SearchArray[SearchValuePair.Value].Blueprint.IsValid() && SearchArray[SearchValuePair.Value].Blueprint->HasAllFlags(RF_PendingKill)) )
 		{
 			// Build the new map/array
 			NewSearchMap.Add(SearchValuePair.Key, NewSearchArray.Add(SearchArray[SearchValuePair.Value]) );
diff --git a/Engine/Source/Editor/Kismet/Public/FindInBlueprintManager.h b/Engine/Source/Editor/Kismet/Public/FindInBlueprintManager.h
--- a/Engine/Source/Editor/Kismet/Public/FindInBlueprintManager.h
+++ b/Engine/Source/Editor/Kismet/Public/FindInBlueprintManager.h
@@ -70,7 +70,7 @@ struct KISMET_API FFindInBlueprintSearchTags
 struct FSearchData
 {
 	/** The Blueprint this search data points to, if available */
-	UBlueprint* Blueprint;
+	TWeakObjectPtr<UBlueprint> Blueprint;
 
 	/** The full Blueprint path this search data is associated with */
 	FString BlueprintPath;
