Leveraging SMyBlueprint and SKismetInspector in the Diff/Merge tools. Rquired introduction of tabs into those tools because we ran out of screenspace

diff --git a/Engine/Source/Developer/Merge/Private/BlueprintMergeData.h b/Engine/Source/Developer/Merge/Private/BlueprintMergeData.h
--- a/Engine/Source/Developer/Merge/Private/BlueprintMergeData.h
+++ b/Engine/Source/Developer/Merge/Private/BlueprintMergeData.h
@@ -48,9 +48,11 @@ namespace EMergeParticipant
 {
 	enum Type
 	{
-		MERGE_PARTICIPANT_REMOTE,
-		MERGE_PARTICIPANT_BASE,
-		MERGE_PARTICIPANT_LOCAL,
+		Remote,
+		Base,
+		Local,
+
+		Max_None,
 	};
 }
 
diff --git a/Engine/Source/Developer/Merge/Private/SMergeDetailsView.cpp b/Engine/Source/Developer/Merge/Private/SMergeDetailsView.cpp
--- a/Engine/Source/Developer/Merge/Private/SMergeDetailsView.cpp
+++ b/Engine/Source/Developer/Merge/Private/SMergeDetailsView.cpp
@@ -47,19 +47,19 @@ void SMergeDetailsView::Construct(const FArguments InArgs, const FBlueprintMerge
 		}
 	};
 
-	// EMergeParticipant::MERGE_PARTICIPANT_REMOTE
+	// EMergeParticipant::Remote
 	{
 		DetailsViews.Add(
 			FDetailsDiff(RemoteCDO, DiffUtils::ResolveAll( RemoteCDO, RemoteDifferingProperties), FDetailsDiff::FOnDisplayedPropertiesChanged() )
 		);
 	}
-	// EMergeParticipant::MERGE_PARTICIPANT_BASE
+	// EMergeParticipant::Base
 	{
 		DetailsViews.Add(
 			FDetailsDiff(BaseCDO, DiffUtils::ResolveAll( BaseCDO, BaseDifferingPropertiesSet.Array()), FDetailsDiff::FOnDisplayedPropertiesChanged())
 		);
 	}
-	// EMergeParticipant::MERGE_PARTICIPANT_LOCAL
+	// EMergeParticipant::Local
 	{
 		DetailsViews.Add(
 			FDetailsDiff(LocalCDO, DiffUtils::ResolveAll( LocalCDO, LocalDifferingProperties), FDetailsDiff::FOnDisplayedPropertiesChanged())
@@ -195,16 +195,16 @@ bool SMergeDetailsView::HasPrevConflict() const
 
 FDetailsDiff& SMergeDetailsView::GetRemoteDetails()
 {
-	return DetailsViews[EMergeParticipant::MERGE_PARTICIPANT_REMOTE];
+	return DetailsViews[EMergeParticipant::Remote];
 }
 
 FDetailsDiff& SMergeDetailsView::GetBaseDetails()
 {
-	return DetailsViews[EMergeParticipant::MERGE_PARTICIPANT_BASE];
+	return DetailsViews[EMergeParticipant::Base];
 }
 
 FDetailsDiff& SMergeDetailsView::GetLocalDetails()
 {
-	return DetailsViews[EMergeParticipant::MERGE_PARTICIPANT_LOCAL];
+	return DetailsViews[EMergeParticipant::Local];
 }
 
diff --git a/Engine/Source/Developer/Merge/Private/SMergeGraphView.cpp b/Engine/Source/Developer/Merge/Private/SMergeGraphView.cpp
--- a/Engine/Source/Developer/Merge/Private/SMergeGraphView.cpp
+++ b/Engine/Source/Developer/Merge/Private/SMergeGraphView.cpp
@@ -3,10 +3,15 @@
 #include "MergePrivatePCH.h"
 #include "DiffResults.h"
 #include "GraphDiffControl.h"
+#include "SDockTab.h"
 #include "SMergeGraphView.h"
+#include "SKismetInspector.h"
 
 #define LOCTEXT_NAMESPACE "SMergeGraphView"
 
+const FName MergeMyBluerpintTabId = FName(TEXT("MergeMyBluerpintTab"));
+const FName MergeGraphTabId = FName(TEXT("MergeGraphTab"));
+
 struct FGraphMergeConflict
 {
 	FGraphMergeConflict(FName InGraphName, const FDiffSingleResult& InRemoteDifference, const FDiffSingleResult& InLocalDifference )
@@ -87,20 +92,6 @@ static const UEdGraphNode* FindNodeInUnrelatedGraph(UEdGraph* Graph, UEdGraphNod
 	return NodeMatch.NewNode;
 }
 
-static FDiffPanel InitializePanel(const FBlueprintRevPair& BlueprintRevPair)
-{
-	FDiffPanel Ret;
-	Ret.Blueprint = BlueprintRevPair.Blueprint;
-	SAssignNew(Ret.GraphEditorBorder, SBorder)
-		.VAlign(VAlign_Fill)
-		[
-			SBlueprintDiff::DefaultEmptyPanel()
-		];
-	Ret.RevisionInfo = BlueprintRevPair.RevData;
-	Ret.bShowAssetName = false;
-	return Ret;
-}
-
 static TArray< TSharedPtr<FMergeGraphRowEntry> > GenerateDiffListItems(const FBlueprintRevPair& RemoteBlueprint, const FBlueprintRevPair& BaseBlueprint, const FBlueprintRevPair& LocalBlueprint,  TArray< TSharedPtr<FGraphMergeConflict> >& OutMergeConflicts )
 {
 	// Index all the graphs by name, we use the name of the graph as the 
@@ -279,6 +270,21 @@ const FLinearColor MergeWhite = FLinearColor::White;
 
 void SMergeGraphView::Construct(const FArguments InArgs, const FBlueprintMergeData& InData )
 {
+	const TSharedRef<SDockTab> MajorTab = SNew(SDockTab)
+		.TabRole(ETabRole::MajorTab);
+
+	TabManager = FGlobalTabmanager::Get()->NewTabManager(MajorTab);
+
+	TabManager->RegisterTabSpawner(MergeGraphTabId,
+		FOnSpawnTab::CreateRaw(this, &SMergeGraphView::CreateGraphDiffViews))
+		.SetDisplayName(LOCTEXT("MergeGraphsTabTitle", "Graphs"))
+		.SetTooltipText(LOCTEXT("MergeGraphsTooltipText", "Differences in the various graphs present in the blueprint"));
+
+	TabManager->RegisterTabSpawner(MergeMyBluerpintTabId,
+		FOnSpawnTab::CreateRaw(this, &SMergeGraphView::CreateMyBlueprintsViews))
+		.SetDisplayName(LOCTEXT("MergeMyBlueprintTabTitle", "My Blueprint"))
+		.SetTooltipText(LOCTEXT("MergeMyBlueprintTooltipText", "Differences in the 'My Blueprints' attributes of the blueprint"));
+
 	CurrentMergeConflict = INDEX_NONE;
 
 	Data = InData;
@@ -287,27 +293,53 @@ void SMergeGraphView::Construct(const FArguments InArgs, const FBlueprintMergeDa
 	LocalDiffResultsListData = nullptr;
 
 	TArray<FBlueprintRevPair> BlueprintsForDisplay;
-	// EMergeParticipant::MERGE_PARTICIPANT_REMOTE
+	// EMergeParticipant::Remote
 	BlueprintsForDisplay.Add(FBlueprintRevPair(InData.BlueprintRemote, InData.RevisionRemote));
-	// EMergeParticipant::MERGE_PARTICIPANT_BASE
+	// EMergeParticipant::Base
 	BlueprintsForDisplay.Add(FBlueprintRevPair(InData.BlueprintBase, InData.RevisionBase));
-	// EMergeParticipant::MERGE_PARTICIPANT_LOCAL
+	// EMergeParticipant::Local
 	BlueprintsForDisplay.Add(FBlueprintRevPair(InData.BlueprintLocal, FRevisionInfo()));
+	
+	const TSharedRef<FTabManager::FLayout> DefaultLayout = FTabManager::NewLayout("BlueprintMerge_Layout_v1")
+	->AddArea
+	(
+		FTabManager::NewPrimaryArea()
+		->Split
+		(
+			FTabManager::NewStack()
+			->AddTab(MergeMyBluerpintTabId, ETabState::OpenedTab)
+			->AddTab(MergeGraphTabId, ETabState::OpenedTab)
+		)
+	);
+
+	for (int32 i = 0; i < EMergeParticipant::Max_None; ++i)
+	{
+		DiffPanels.Add(FDiffPanel());
+		FDiffPanel& NewPanel = DiffPanels[i];
+		NewPanel.Blueprint = BlueprintsForDisplay[i].Blueprint;
+		NewPanel.RevisionInfo = BlueprintsForDisplay[i].RevData;
+		NewPanel.bShowAssetName = false;
+	}
+
+	auto GraphPanelContainer = TabManager->RestoreFrom(DefaultLayout, TSharedPtr<SWindow>()).ToSharedRef();
 
-	for (const auto& BlueprintRevPair : BlueprintsForDisplay)
+	for (auto& Panel : DiffPanels )
 	{
-		DiffPanels.Add(InitializePanel(BlueprintRevPair));
+		Panel.InitializeDiffPanel();
 	}
 
-	TSharedRef<SSplitter> PanelContainer = SNew(SSplitter);
-	for (const auto& Panel : DiffPanels)
+	auto DetailsPanelContainer = SNew(SSplitter);
+	for( auto& Panel : DiffPanels )
 	{
-		PanelContainer->AddSlot()[Panel.GraphEditorBorder.ToSharedRef()];
+		DetailsPanelContainer->AddSlot()
+		[
+			Panel.DetailsView.ToSharedRef()
+		];
 	}
 
-	DifferencesFromBase = GenerateDiffListItems(BlueprintsForDisplay[EMergeParticipant::MERGE_PARTICIPANT_REMOTE]
-												, BlueprintsForDisplay[EMergeParticipant::MERGE_PARTICIPANT_BASE]
-												, BlueprintsForDisplay[EMergeParticipant::MERGE_PARTICIPANT_LOCAL]
+	DifferencesFromBase = GenerateDiffListItems(BlueprintsForDisplay[EMergeParticipant::Remote]
+												, BlueprintsForDisplay[EMergeParticipant::Base]
+												, BlueprintsForDisplay[EMergeParticipant::Local]
 												, MergeConflicts);
 
 	// This is the function we'll use to generate a row in the control that lists all the available graphs:
@@ -465,7 +497,18 @@ void SMergeGraphView::Construct(const FArguments InArgs, const FBlueprintMergeDa
 			+ SSplitter::Slot()
 				.Value(0.9f)
 				[
-					PanelContainer
+					SNew(SSplitter)
+					.Orientation(Orient_Vertical)
+					+ SSplitter::Slot()
+					.Value(.8f)
+					[
+						GraphPanelContainer
+					]
+					+ SSplitter::Slot()
+					.Value(.2f)
+					[
+						DetailsPanelContainer
+					]
 				]
 		];
 }
@@ -612,9 +655,9 @@ void SMergeGraphView::OnGraphListSelectionChanged(TSharedPtr<FMergeGraphRowEntry
 	UEdGraph* GraphBase = FindGraphByName(*GetBasePanel().Blueprint, GraphName);
 	UEdGraph* GraphLocal = FindGraphByName(*GetLocalPanel().Blueprint, GraphName);
 
-	GetBasePanel().GeneratePanel(GraphBase, nullptr);
-	GetRemotePanel().GeneratePanel(GraphRemote, GraphBase);
-	GetLocalPanel().GeneratePanel(GraphLocal, GraphBase);
+	GetBasePanel().GeneratePanel(GraphBase, nullptr );
+	GetRemotePanel().GeneratePanel(GraphRemote, GraphBase );
+	GetLocalPanel().GeneratePanel(GraphLocal, GraphBase );
 
 	LockViews(DiffPanels, bViewsAreLocked);
 
@@ -688,6 +731,44 @@ void SMergeGraphView::OnDiffListSelectionChanged(TSharedPtr<FDiffSingleResult> I
 	}
 }
 
+TSharedRef<SDockTab> SMergeGraphView::CreateGraphDiffViews(const FSpawnTabArgs& Args)
+{
+	auto PanelContainer = SNew(SSplitter);
+	for (auto& Panel : DiffPanels)
+	{
+		PanelContainer->AddSlot()
+		[
+			SAssignNew(Panel.GraphEditorBorder, SBox)
+			.VAlign(VAlign_Fill)
+			[
+				SBlueprintDiff::DefaultEmptyPanel()
+			]
+		];
+	}
+
+	return SNew(SDockTab)
+	[
+		PanelContainer
+	];
+}
+
+TSharedRef<SDockTab> SMergeGraphView::CreateMyBlueprintsViews(const FSpawnTabArgs& Args)
+{
+	auto PanelContainer = SNew(SSplitter);
+	for (auto& Panel : DiffPanels)
+	{
+		PanelContainer->AddSlot()
+		[
+			Panel.GenerateMyBlueprintPanel()
+		];
+	}
+
+	return SNew(SDockTab)
+	[
+		PanelContainer
+	];
+}
+
 FReply SMergeGraphView::OnToggleLockView()
 {
 	bViewsAreLocked = !bViewsAreLocked;
diff --git a/Engine/Source/Developer/Merge/Private/SMergeGraphView.h b/Engine/Source/Developer/Merge/Private/SMergeGraphView.h
--- a/Engine/Source/Developer/Merge/Private/SMergeGraphView.h
+++ b/Engine/Source/Developer/Merge/Private/SMergeGraphView.h
@@ -10,8 +10,7 @@ class SMergeGraphView : public SCompoundWidget
 						, public IMergeControl
 {
 public:
-	SLATE_BEGIN_ARGS(SMergeGraphView)
-	{}
+	SLATE_BEGIN_ARGS(SMergeGraphView){}
 	SLATE_END_ARGS()
 
 	void Construct(const FArguments InArgs, const FBlueprintMergeData& InData );
@@ -33,10 +32,12 @@ private:
 	bool HasNoDifferences() const;
 	void OnGraphListSelectionChanged(TSharedPtr<struct FMergeGraphRowEntry> Item, ESelectInfo::Type SelectionType);
 	void OnDiffListSelectionChanged(TSharedPtr<struct FDiffSingleResult> Item, ESelectInfo::Type SelectionType);
+	TSharedRef<SDockTab> CreateGraphDiffViews(const FSpawnTabArgs& Args);
+	TSharedRef<SDockTab> CreateMyBlueprintsViews(const FSpawnTabArgs& Args);
 
-	FDiffPanel& GetRemotePanel() { return DiffPanels[EMergeParticipant::MERGE_PARTICIPANT_REMOTE]; }
-	FDiffPanel& GetBasePanel() { return DiffPanels[EMergeParticipant::MERGE_PARTICIPANT_BASE]; }
-	FDiffPanel& GetLocalPanel() { return DiffPanels[EMergeParticipant::MERGE_PARTICIPANT_LOCAL]; }
+	FDiffPanel& GetRemotePanel() { return DiffPanels[EMergeParticipant::Remote]; }
+	FDiffPanel& GetBasePanel() { return DiffPanels[EMergeParticipant::Base]; }
+	FDiffPanel& GetLocalPanel() { return DiffPanels[EMergeParticipant::Local]; }
 
 	FReply OnToggleLockView();
 	const FSlateBrush*  GetLockViewImage() const;
@@ -73,4 +74,7 @@ private:
 	TArray< TSharedPtr< struct FDiffSingleResult> > const* LocalDiffResultsListData;
 
 	bool bViewsAreLocked;
+
+	/** We can't use the global tab manager because we need to instance the merge control, so we have our own tab manager: */
+	TSharedPtr<FTabManager> TabManager;
 };
diff --git a/Engine/Source/Developer/Merge/Private/SMergeTreeView.cpp b/Engine/Source/Developer/Merge/Private/SMergeTreeView.cpp
--- a/Engine/Source/Developer/Merge/Private/SMergeTreeView.cpp
+++ b/Engine/Source/Developer/Merge/Private/SMergeTreeView.cpp
@@ -11,19 +11,19 @@ void SMergeTreeView::Construct(const FArguments InArgs, const FBlueprintMergeDat
 	CurrentMergeConflict = -1;
 
 	// generate controls:
-	// EMergeParticipant::MERGE_PARTICIPANT_REMOTE
+	// EMergeParticipant::Remote
 	{
 		SCSViews.Push(
 			FSCSDiff(InData.BlueprintRemote)
 			);
 	}
-	// EMergeParticipant::MERGE_PARTICIPANT_BASE
+	// EMergeParticipant::Base
 	{
 		SCSViews.Push(
 			FSCSDiff(InData.BlueprintBase)
 			);
 	}
-	// EMergeParticipant::MERGE_PARTICIPANT_LOCAL
+	// EMergeParticipant::Local
 	{
 		SCSViews.Push(
 			FSCSDiff(InData.BlueprintLocal)
@@ -213,16 +213,16 @@ void SMergeTreeView::HighlightCurrentDifference()
 
 FSCSDiff& SMergeTreeView::GetRemoteView()
 {
-	return SCSViews[EMergeParticipant::MERGE_PARTICIPANT_REMOTE];
+	return SCSViews[EMergeParticipant::Remote];
 }
 
 FSCSDiff& SMergeTreeView::GetBaseView()
 {
-	return SCSViews[EMergeParticipant::MERGE_PARTICIPANT_BASE];
+	return SCSViews[EMergeParticipant::Base];
 }
 
 FSCSDiff& SMergeTreeView::GetLocalView()
 {
-	return SCSViews[EMergeParticipant::MERGE_PARTICIPANT_LOCAL];
+	return SCSViews[EMergeParticipant::Local];
 }
 
diff --git a/Engine/Source/Editor/Kismet/Private/SBlueprintDiff.cpp b/Engine/Source/Editor/Kismet/Private/SBlueprintDiff.cpp
--- a/Engine/Source/Editor/Kismet/Private/SBlueprintDiff.cpp
+++ b/Engine/Source/Editor/Kismet/Private/SBlueprintDiff.cpp
@@ -17,7 +17,9 @@
 #include "IDiffControl.h"
 #include "Kismet2/BlueprintEditorUtils.h"
 #include "SBlueprintDiff.h"
+#include "SDockTab.h"
 #include "SGraphTitleBar.h"
+#include "SMyBlueprint.h"
 #include "SSCSEditor.h"
 #include "SCSDiff.h"
 #include "WorkflowOrientedApp/SModeWidget.h"
@@ -27,6 +29,9 @@
 
 typedef TMap< FName, const UProperty* > FNamePropertyMap;
 
+const FName DiffMyBluerpintTabId = FName(TEXT("DiffMyBluerpintTab"));
+const FName DiffGraphTabId = FName(TEXT("DiffGraphTab"));
+
 // Each difference in the tree will either be a tree node that is added in one Blueprint 
 // or a tree node and an FName of a property that has been added or edited in one Blueprint
 class FSCSDiffControl	: public TSharedFromThis< FSCSDiffControl >
@@ -623,6 +628,34 @@ FDiffPanel::FDiffPanel()
 	LastFocusedPin = NULL;
 }
 
+TSharedRef<SWidget> FDiffPanel::InitializeDiffPanel()
+{
+	TSharedRef< SKismetInspector > Inspector = SNew(SKismetInspector)
+		.HideNameArea(true)
+		.ViewIdentifier(FName("BlueprintInspector"))
+		.MyBlueprintWidget(MyBlueprint)
+		.IsPropertyEditingEnabledDelegate(FIsPropertyEditingEnabled::CreateStatic([] { return false; }));
+	DetailsView = Inspector;
+	MyBlueprint->SetInspector(DetailsView);
+
+	return SNew(SSplitter)
+		.Orientation(Orient_Vertical)
+		+ SSplitter::Slot()
+		.Value(0.8f)
+		[
+			SAssignNew(GraphEditorBorder, SBox)
+			.VAlign(VAlign_Fill)
+			[
+				SBlueprintDiff::DefaultEmptyPanel()
+			]
+		]
+	+ SSplitter::Slot()
+		.Value(0.2f)
+		[
+			Inspector
+		];
+}
+
 int32 GetCurrentIndex( SListView< TSharedPtr< struct FDiffSingleResult> > const& ListView, const TArray< TSharedPtr< FDiffSingleResult > >& ListViewSource )
 {
 	const auto& Selected = ListView.GetSelectedItems();
@@ -680,6 +713,20 @@ bool DiffWidgetUtils::HasPrevDifference(SListView< TSharedPtr< struct FDiffSingl
 BEGIN_SLATE_FUNCTION_BUILD_OPTIMIZATION
 void SBlueprintDiff::Construct( const FArguments& InArgs)
 {
+	const TSharedRef<SDockTab> MajorTab = SNew(SDockTab)
+		.TabRole(ETabRole::MajorTab);
+	TabManager = FGlobalTabmanager::Get()->NewTabManager(MajorTab);
+
+	TabManager->RegisterTabSpawner(DiffGraphTabId,
+		FOnSpawnTab::CreateRaw(this, &SBlueprintDiff::CreateGraphDiffViews ))
+		.SetDisplayName(NSLOCTEXT("SBlueprintDiff", "GraphsTabTitle", "Graphs"))
+		.SetTooltipText(NSLOCTEXT("SBlueprintDiff", "GraphsTooltipText", "Differences in the various graphs present in the blueprint"));
+
+	TabManager->RegisterTabSpawner(DiffMyBluerpintTabId,
+		FOnSpawnTab::CreateRaw(this, &SBlueprintDiff::CreateMyBlueprintsViews))
+		.SetDisplayName(NSLOCTEXT("SBlueprintDiff", "MyBlueprintTabTitle", "My Blueprint"))
+		.SetTooltipText(NSLOCTEXT("SBlueprintDiff", "MyBlueprintTooltipText", "Differences in the 'My Blueprints' attributes of the blueprint"));
+
 	FDiffListCommands::Register();
 	check(InArgs._BlueprintOld && InArgs._BlueprintNew);
 	PanelOld.Blueprint = InArgs._BlueprintOld;
@@ -876,6 +923,46 @@ bool SBlueprintDiff::HasPrevDiff() const
 	return DiffControl.IsValid() && DiffControl->HasPrevDifference();
 }
 
+TSharedRef<SDockTab> SBlueprintDiff::CreateGraphDiffViews(const FSpawnTabArgs& Args)
+{
+	return SNew(SDockTab)
+	[
+		SNew(SSplitter) 
+		+ SSplitter::Slot()
+		[
+			SAssignNew(PanelOld.GraphEditorBorder, SBox)
+			.VAlign(VAlign_Fill)
+			[
+				DefaultEmptyPanel()
+			]
+		]
+		+ SSplitter::Slot()
+		[
+			SAssignNew(PanelNew.GraphEditorBorder, SBox)
+			.VAlign(VAlign_Fill)
+			[
+				DefaultEmptyPanel()
+			]
+		]
+	];
+}
+
+TSharedRef<SDockTab> SBlueprintDiff::CreateMyBlueprintsViews(const FSpawnTabArgs& Args)
+{
+	return SNew(SDockTab)
+	[
+		SNew(SSplitter)
+		+ SSplitter::Slot()
+		[
+			PanelOld.GenerateMyBlueprintPanel()
+		]
+		+ SSplitter::Slot()
+		[
+			PanelNew.GenerateMyBlueprintPanel()
+		]
+	];
+}
+
 void SBlueprintDiff::FocusOnGraphRevisions( class UEdGraph* GraphOld, class UEdGraph* GraphNew, FListItemGraphToDiff* Diff )
 {
 	HandleGraphChanged( GraphOld ? GraphOld->GetFName() : GraphNew->GetFName() );
@@ -950,7 +1037,7 @@ void SBlueprintDiff::ResetGraphEditors()
 	}
 }
 
-void FDiffPanel::GeneratePanel(UEdGraph* Graph, UEdGraph* GraphToDiff)
+void FDiffPanel::GeneratePanel(UEdGraph* Graph, UEdGraph* GraphToDiff )
 {
 	if( LastFocusedPin )
 	{
@@ -964,9 +1051,17 @@ void FDiffPanel::GeneratePanel(UEdGraph* Graph, UEdGraph* GraphToDiff)
 								[
 									SNew(STextBlock).Text( LOCTEXT("BPDifPanelNoGraphTip", "Graph does not exist in this revision"))
 								];
+
 	if(Graph)
 	{
 		SGraphEditor::FGraphEditorEvents InEvents;
+		{
+			const auto SelectionChangedHandler = [](const FGraphPanelSelectionSet& SelectionSet, TSharedPtr<SKismetInspector> Container)
+			{
+				Container->ShowDetailsForObjects(SelectionSet.Array());
+			};
+			InEvents.OnSelectionChanged = SGraphEditor::FOnSelectionChanged::CreateStatic(SelectionChangedHandler, DetailsView);
+		}
 
 		FGraphAppearanceInfo AppearanceInfo;
 		AppearanceInfo.CornerText = LOCTEXT("AppearanceCornerText_BlueprintDif", "DIFF").ToString();
@@ -981,6 +1076,9 @@ void FDiffPanel::GeneratePanel(UEdGraph* Graph, UEdGraph* GraphToDiff)
 				);
 		}
 
+		MyBlueprint->SetFocusedGraph(Graph);
+		MyBlueprint->Refresh();
+
 		auto Editor = SNew(SGraphEditor)
 			.AdditionalCommands(GraphEditorCommands)
 			.GraphToEdit(Graph)
@@ -1001,6 +1099,11 @@ void FDiffPanel::GeneratePanel(UEdGraph* Graph, UEdGraph* GraphToDiff)
 	GraphEditorBorder->SetContent(Widget.ToSharedRef());
 }
 
+TSharedRef<SWidget> FDiffPanel::GenerateMyBlueprintPanel()
+{
+	return SAssignNew(MyBlueprint, SMyBlueprint, TWeakPtr<FBlueprintEditor>(), Blueprint);
+}
+
 FGraphPanelSelectionSet FDiffPanel::GetSelectedNodes() const
 {
 	FGraphPanelSelectionSet CurrentSelection;
@@ -1133,6 +1236,35 @@ void SBlueprintDiff::HandleGraphChanged( const FName GraphName )
 
 TSharedRef<SWidget> SBlueprintDiff::GenerateGraphPanel()
 {
+	const TSharedRef<FTabManager::FLayout> DefaultLayout = FTabManager::NewLayout("BlueprintDiff_Layout_v1")
+	->AddArea
+	(
+		FTabManager::NewPrimaryArea()
+		->Split
+		(
+			FTabManager::NewStack()
+			->AddTab(DiffMyBluerpintTabId, ETabState::OpenedTab)
+			->AddTab(DiffGraphTabId, ETabState::OpenedTab)
+		)
+	);
+
+	// SMyBlueprint needs to be created *before* the KismetInspector, because the KismetInspector's customizations
+	// need a reference to the SMyBlueprint widget that is controlling them...
+	TSharedRef<SWidget> TabControl = TabManager->RestoreFrom(DefaultLayout, TSharedPtr<SWindow>()).ToSharedRef();
+
+	const auto CreateInspector = [](TSharedPtr<SMyBlueprint> InMyBlueprint) {
+		return SNew( SKismetInspector)
+			.HideNameArea(true)
+			.ViewIdentifier(FName("BlueprintInspector"))
+			.MyBlueprintWidget(InMyBlueprint)
+			.IsPropertyEditingEnabledDelegate(FIsPropertyEditingEnabled::CreateStatic([] { return false; }));
+	};
+
+	PanelOld.DetailsView = CreateInspector(PanelOld.MyBlueprint);
+	PanelOld.MyBlueprint->SetInspector(PanelOld.DetailsView);
+	PanelNew.DetailsView = CreateInspector(PanelNew.MyBlueprint);
+	PanelNew.MyBlueprint->SetInspector(PanelNew.DetailsView);
+
 	return SNew(SBorder)
 	.BorderImage(FEditorStyle::GetBrush("ToolPanel.GroupBorder"))
 	.Content()
@@ -1224,24 +1356,24 @@ TSharedRef<SWidget> SBlueprintDiff::GenerateGraphPanel()
 					[
 						//diff window
 						SNew(SSplitter)
-						+ SSplitter::Slot()
-						.Value(0.5f)
+						.Orientation(Orient_Vertical)
+						+SSplitter::Slot()
+						.Value(.8f)
 						[
-							//left blueprint
-							SAssignNew(PanelOld.GraphEditorBorder, SBorder)
-							.VAlign(VAlign_Fill)
-							[
-								DefaultEmptyPanel()
-							]
+							// graph and my blueprint views:
+							TabControl
 						]
 						+ SSplitter::Slot()
-						.Value(0.5f)
+						.Value(.2f)
 						[
-							//right blueprint
-							SAssignNew(PanelNew.GraphEditorBorder, SBorder)
-							.VAlign(VAlign_Fill)
+							SNew( SSplitter )
+							+SSplitter::Slot()
+							[
+								PanelOld.DetailsView.ToSharedRef()
+							]
+							+ SSplitter::Slot()
 							[
-								DefaultEmptyPanel()
+								PanelNew.DetailsView.ToSharedRef()
 							]
 						]
 					]
diff --git a/Engine/Source/Editor/Kismet/Private/SBlueprintPalette.cpp b/Engine/Source/Editor/Kismet/Private/SBlueprintPalette.cpp
--- a/Engine/Source/Editor/Kismet/Private/SBlueprintPalette.cpp
+++ b/Engine/Source/Editor/Kismet/Private/SBlueprintPalette.cpp
@@ -689,10 +689,11 @@ public:
 	 * @param  ActionPtrIn		The FEdGraphSchemaAction that the parent item represents.
 	 * @param  BlueprintEdPtrIn	A pointer to the blueprint editor that the palette belongs to.
 	 */
-	void Construct(const FArguments& InArgs, TWeakPtr<FEdGraphSchemaAction> ActionPtrIn, TWeakPtr<FBlueprintEditor> BlueprintEdPtrIn)
+	void Construct(const FArguments& InArgs, TWeakPtr<FEdGraphSchemaAction> ActionPtrIn, TWeakPtr<FBlueprintEditor> InBlueprintEditor, UBlueprint* InBlueprint)
 	{
 		ActionPtr = ActionPtrIn;
-		BlueprintEditorPtr = BlueprintEdPtrIn;
+		BlueprintEditorPtr = InBlueprintEditor;
+		BlueprintObj = InBlueprint;
 		TSharedPtr<FEdGraphSchemaAction> PaletteAction = ActionPtrIn.Pin();
 
 		bool bShouldHaveAVisibilityToggle = false;
@@ -702,7 +703,7 @@ public:
 			UObjectProperty* VariableObjProp = Cast<UObjectProperty>(VariableProp);
 
 			UStruct* VarSourceScope = (VariableProp != NULL) ? CastChecked<UStruct>(VariableProp->GetOuter()) : NULL;
-			const bool bIsBlueprintVariable = (VarSourceScope == BlueprintEditorPtr.Pin()->GetBlueprintObj()->SkeletonGeneratedClass);
+			const bool bIsBlueprintVariable = (VarSourceScope == BlueprintObj->SkeletonGeneratedClass);
 			const bool bIsComponentVar = (VariableObjProp != NULL) && (VariableObjProp->PropertyClass != NULL) && (VariableObjProp->PropertyClass->IsChildOf(UActorComponent::StaticClass()));
 			bShouldHaveAVisibilityToggle = bIsBlueprintVariable && !bIsComponentVar;
 		}
@@ -766,13 +767,17 @@ private:
 	 */
 	void OnVisibilityToggleFlipped(ESlateCheckBoxState::Type InNewState)
 	{
+		if( !BlueprintEditorPtr.IsValid() )
+		{
+			return;
+		}
+
 		TSharedPtr<FEdGraphSchemaAction_K2Var> VarAction = StaticCastSharedPtr<FEdGraphSchemaAction_K2Var>(ActionPtr.Pin());
 
 		// Toggle the flag on the blueprint's version of the variable description, based on state
 		const bool bVariableIsExposed = (InNewState == ESlateCheckBoxState::Checked);
 
-		UBlueprint* Blueprint = BlueprintEditorPtr.Pin()->GetBlueprintObj();
-		FBlueprintEditorUtils::SetBlueprintOnlyEditableFlag(Blueprint, VarAction->GetVariableName(), !bVariableIsExposed);
+		FBlueprintEditorUtils::SetBlueprintOnlyEditableFlag(BlueprintObj, VarAction->GetVariableName(), !bVariableIsExposed);
 	}
 
 	/**
@@ -806,7 +811,7 @@ private:
 			TSharedPtr<FEdGraphSchemaAction_K2Var> VarAction = StaticCastSharedPtr<FEdGraphSchemaAction_K2Var>(ActionPtr.Pin());
 
 			FString Result;
-			FBlueprintEditorUtils::GetBlueprintVariableMetaData( BlueprintEditorPtr.Pin()->GetBlueprintObj(), VarAction->GetVariableName(), NULL, TEXT("tooltip"), Result);
+			FBlueprintEditorUtils::GetBlueprintVariableMetaData( BlueprintObj, VarAction->GetVariableName(), NULL, TEXT("tooltip"), Result);
 
 			if(!Result.IsEmpty())
 			{
@@ -838,7 +843,7 @@ private:
 			TSharedPtr<FEdGraphSchemaAction_K2Var> VarAction = StaticCastSharedPtr<FEdGraphSchemaAction_K2Var>(ActionPtr.Pin());
 
 			FString Result;
-			FBlueprintEditorUtils::GetBlueprintVariableMetaData( BlueprintEditorPtr.Pin()->GetBlueprintObj(), VarAction->GetVariableName(), NULL, TEXT("tooltip"), Result);
+			FBlueprintEditorUtils::GetBlueprintVariableMetaData( BlueprintObj, VarAction->GetVariableName(), NULL, TEXT("tooltip"), Result);
 			if(!Result.IsEmpty())
 			{
 				ToolTip = LOCTEXT("VariablePrivacy_is_public_Tooltip", "Variable is public and is editable on each instance of this Blueprint.").ToString();
@@ -855,8 +860,11 @@ private:
 	/** The action that the owning palette entry represents */
 	TWeakPtr<FEdGraphSchemaAction> ActionPtr;
 
-	/** Pointer back to the blueprint editor that owns this */
+	/** Pointer back to the blueprint editor that owns this, optional because of diff and merge views: */
 	TWeakPtr<FBlueprintEditor>     BlueprintEditorPtr;
+
+	/** Pointer back to the blueprint that is being diplayed: */
+	UBlueprint* BlueprintObj;
 };
 
 /*******************************************************************************
@@ -866,9 +874,21 @@ private:
 //------------------------------------------------------------------------------
 BEGIN_SLATE_FUNCTION_BUILD_OPTIMIZATION
 void SBlueprintPaletteItem::Construct(const FArguments& InArgs, FCreateWidgetForActionData* const InCreateData, TWeakPtr<FBlueprintEditor> InBlueprintEditor)
+{
+	Construct(InArgs, InCreateData, InBlueprintEditor.Pin()->GetBlueprintObj(), InBlueprintEditor);
+}
+
+void SBlueprintPaletteItem::Construct(const FArguments& InArgs, FCreateWidgetForActionData* const InCreateData, UBlueprint* InBlueprint)
+{
+	Construct( InArgs, InCreateData, InBlueprint, TWeakPtr<FBlueprintEditor>() );
+}
+
+void SBlueprintPaletteItem::Construct(const FArguments& InArgs, FCreateWidgetForActionData* const InCreateData, UBlueprint* InBlueprint, TWeakPtr<FBlueprintEditor> InBlueprintEditor)
 {
 	check(InCreateData->Action.IsValid());
-	check(InBlueprintEditor.IsValid());
+	check(InBlueprint);
+
+	Blueprint = InBlueprint;
 
 	bShowClassInTooltip = InArgs._ShowClassInTooltip;	
 
@@ -881,7 +901,7 @@ void SBlueprintPaletteItem::Construct(const FArguments& InArgs, FCreateWidgetFor
 	FSlateColor        IconColor   = FSlateColor::UseForeground();
 	FString            IconToolTip = GraphAction->TooltipDescription;
 	FString			   IconDocLink, IconDocExcerpt;
-	GetPaletteItemIcon(GraphAction, InBlueprintEditor.Pin()->GetBlueprintObj(), IconBrush, IconColor, IconToolTip, IconDocLink, IconDocExcerpt);
+	GetPaletteItemIcon(GraphAction, Blueprint, IconBrush, IconColor, IconToolTip, IconDocLink, IconDocExcerpt);
 	TSharedRef<SWidget> IconWidget = CreateIconWidget(IconToolTip, IconBrush, IconColor, IconDocLink, IconDocExcerpt);
 
 	// Setup a meta tag for this node
@@ -893,7 +913,7 @@ void SBlueprintPaletteItem::Construct(const FArguments& InArgs, FCreateWidgetFor
 	}
 	// construct the text widget
 	FSlateFontInfo NameFont = FSlateFontInfo(FPaths::EngineContentDir() / TEXT("Slate/Fonts/Roboto-Regular.ttf"), 10);
-	bool bIsReadOnly = FBlueprintEditorUtils::IsPaletteActionReadOnly(GraphAction, InBlueprintEditor.Pin());
+	bool bIsReadOnly = InBlueprintEditor.IsValid() ? FBlueprintEditorUtils::IsPaletteActionReadOnly(GraphAction, InBlueprintEditor.Pin()) : true;
 	TSharedRef<SWidget> NameSlotWidget = CreateTextSlotWidget( NameFont, InCreateData, bIsReadOnly );
 	
 	// now, create the actual widget
@@ -921,7 +941,7 @@ void SBlueprintPaletteItem::Construct(const FArguments& InArgs, FCreateWidgetFor
 			.Padding(FMargin(3,0))
 			.VAlign(VAlign_Center)
 		[
-			SNew(SPaletteItemVisibilityToggle, ActionPtr, BlueprintEditorPtr)
+			SNew(SPaletteItemVisibilityToggle, ActionPtr, InBlueprintEditor, InBlueprint)
 		]
 	];
 }
@@ -1345,7 +1365,7 @@ FText SBlueprintPaletteItem::GetToolTipText() const
 			}
 			else
 			{
-				FString Result = GetVarTooltip(BlueprintEditorPtr.Pin()->GetBlueprintObj(), VarClass, VarAction->GetVariableName());
+				FString Result = GetVarTooltip(Blueprint, VarClass, VarAction->GetVariableName());
 				// Only use the variable tooltip if it has been filled out.
 				ToolTipText = !Result.IsEmpty() ? Result : GetVarType(VarClass, VarAction->GetVariableName(), true, true);
 			}
@@ -1362,7 +1382,7 @@ FText SBlueprintPaletteItem::GetToolTipText() const
 			else
 			{
 				FString Result;
-				FBlueprintEditorUtils::GetBlueprintVariableMetaData(BlueprintEditorPtr.Pin()->GetBlueprintObj(), LocalVarAction->GetVariableName(), LocalVarAction->GetVariableScope(), TEXT("tooltip"), Result);
+				FBlueprintEditorUtils::GetBlueprintVariableMetaData(Blueprint, LocalVarAction->GetVariableName(), LocalVarAction->GetVariableScope(), TEXT("tooltip"), Result);
 				// Only use the variable tooltip if it has been filled out.
 				ToolTipText = !Result.IsEmpty() ? Result : GetVarType(LocalVarAction->GetVariableScope(), LocalVarAction->GetVariableName(), true, true);
 			}
@@ -1371,7 +1391,7 @@ FText SBlueprintPaletteItem::GetToolTipText() const
 		{
 			FEdGraphSchemaAction_K2Delegate* DelegateAction = (FEdGraphSchemaAction_K2Delegate*)PaletteAction.Get();
 			
-			FString Result = GetVarTooltip(BlueprintEditorPtr.Pin()->GetBlueprintObj(), DelegateAction->GetDelegateClass(), DelegateAction->GetDelegateName());
+			FString Result = GetVarTooltip(Blueprint, DelegateAction->GetDelegateClass(), DelegateAction->GetDelegateName());
 			ToolTipText = !Result.IsEmpty() ? Result : DelegateAction->GetDelegateName().ToString();
 		}
 		else if (PaletteAction->GetTypeId() == FEdGraphSchemaAction_K2Enum::StaticGetTypeId())
diff --git a/Engine/Source/Editor/Kismet/Private/SBlueprintPalette.h b/Engine/Source/Editor/Kismet/Private/SBlueprintPalette.h
--- a/Engine/Source/Editor/Kismet/Private/SBlueprintPalette.h
+++ b/Engine/Source/Editor/Kismet/Private/SBlueprintPalette.h
@@ -31,8 +31,11 @@ public:
 	 * @param  InBlueprintEditor	A pointer to the blueprint editor that the palette belongs to.
 	 */
 	void Construct(const FArguments& InArgs, FCreateWidgetForActionData* const InCreateData, TWeakPtr<FBlueprintEditor> InBlueprintEditor);
+	void Construct(const FArguments& InArgs, FCreateWidgetForActionData* const InCreateData, UBlueprint* InBlueprint);
 
 private:
+	void Construct(const FArguments& InArgs, FCreateWidgetForActionData* const InCreateData, UBlueprint* InBlueprint, TWeakPtr<FBlueprintEditor> InBlueprintEditor);
+
 	// SGraphPaletteItem Interface
 	virtual TSharedRef<SWidget> CreateTextSlotWidget( const FSlateFontInfo& NameFont,  FCreateWidgetForActionData* const InCreateData, bool bIsReadOnly ) override;
 	virtual FText GetDisplayText() const override;
@@ -54,6 +57,8 @@ private:
 	/** True if the class should be displayed in the tooltip */
 	bool bShowClassInTooltip;
 
+	UBlueprint* Blueprint;
+
 	/** Pointer back to the blueprint editor that owns this */
 	TWeakPtr<FBlueprintEditor> BlueprintEditorPtr;
 };
diff --git a/Engine/Source/Editor/Kismet/Private/SKismetInspector.cpp b/Engine/Source/Editor/Kismet/Private/SKismetInspector.cpp
--- a/Engine/Source/Editor/Kismet/Private/SKismetInspector.cpp
+++ b/Engine/Source/Editor/Kismet/Private/SKismetInspector.cpp
@@ -211,23 +211,28 @@ void SKismetInspector::Construct(const FArguments& InArgs)
 	PropertyView->SetIsPropertyVisibleDelegate( FIsPropertyVisible::CreateSP(this, &SKismetInspector::IsPropertyVisible) );
 	PropertyView->SetIsPropertyEditingEnabledDelegate(InArgs._IsPropertyEditingEnabledDelegate);
 	PropertyView->OnFinishedChangingProperties().Add( InArgs._OnFinishedChangingProperties );
+
+	TWeakPtr<SMyBlueprint> MyBlueprint = Kismet2.IsValid() ? Kismet2->GetMyBlueprintWidget() : InArgs._MyBlueprintWidget;
 	
-	if (Kismet2.IsValid() && Kismet2->IsEditingSingleBlueprint())
+	if( MyBlueprint.IsValid() )
 	{
-		FOnGetDetailCustomizationInstance LayoutDelegateDetails = FOnGetDetailCustomizationInstance::CreateStatic(&FBlueprintDelegateActionDetails::MakeInstance, TWeakPtr<SMyBlueprint>(Kismet2->GetMyBlueprintWidget()));
+		FOnGetDetailCustomizationInstance LayoutDelegateDetails = FOnGetDetailCustomizationInstance::CreateStatic(&FBlueprintDelegateActionDetails::MakeInstance, MyBlueprint);
 		PropertyView->RegisterInstancedCustomPropertyLayout(UMulticastDelegateProperty::StaticClass(), LayoutDelegateDetails);
 
 		// Register function and variable details customization
-		FOnGetDetailCustomizationInstance LayoutGraphDetails = FOnGetDetailCustomizationInstance::CreateStatic(&FBlueprintGraphActionDetails::MakeInstance, TWeakPtr<SMyBlueprint>(Kismet2->GetMyBlueprintWidget()));
+		FOnGetDetailCustomizationInstance LayoutGraphDetails = FOnGetDetailCustomizationInstance::CreateStatic(&FBlueprintGraphActionDetails::MakeInstance, MyBlueprint);
 		PropertyView->RegisterInstancedCustomPropertyLayout(UEdGraph::StaticClass(), LayoutGraphDetails);
 		PropertyView->RegisterInstancedCustomPropertyLayout(UK2Node_EditablePinBase::StaticClass(), LayoutGraphDetails);
 		PropertyView->RegisterInstancedCustomPropertyLayout(UK2Node_CallFunction::StaticClass(), LayoutGraphDetails);
-	
-		FOnGetDetailCustomizationInstance LayoutVariableDetails = FOnGetDetailCustomizationInstance::CreateStatic(&FBlueprintVarActionDetails::MakeInstance, TWeakPtr<SMyBlueprint>(Kismet2->GetMyBlueprintWidget()));
+
+		FOnGetDetailCustomizationInstance LayoutVariableDetails = FOnGetDetailCustomizationInstance::CreateStatic(&FBlueprintVarActionDetails::MakeInstance, MyBlueprint);
 		PropertyView->RegisterInstancedCustomPropertyLayout(UProperty::StaticClass(), LayoutVariableDetails);
 		PropertyView->RegisterInstancedCustomPropertyLayout(UK2Node_VariableGet::StaticClass(), LayoutVariableDetails);
 		PropertyView->RegisterInstancedCustomPropertyLayout(UK2Node_VariableSet::StaticClass(), LayoutVariableDetails);
+	}
 
+	if (Kismet2.IsValid() && Kismet2->IsEditingSingleBlueprint())
+	{
 		FOnGetDetailCustomizationInstance LayoutOptionDetails = FOnGetDetailCustomizationInstance::CreateStatic(&FBlueprintGlobalOptionsDetails::MakeInstance, Kismet2Ptr);
 		PropertyView->RegisterInstancedCustomPropertyLayout(UBlueprint::StaticClass(), LayoutOptionDetails);
 
diff --git a/Engine/Source/Editor/Kismet/Private/SMyBlueprint.cpp b/Engine/Source/Editor/Kismet/Private/SMyBlueprint.cpp
--- a/Engine/Source/Editor/Kismet/Private/SMyBlueprint.cpp
+++ b/Engine/Source/Editor/Kismet/Private/SMyBlueprint.cpp
@@ -127,64 +127,85 @@ public:
 
 //////////////////////////////////////////////////////////////////////////
 
-void SMyBlueprint::Construct(const FArguments& InArgs, TWeakPtr<FBlueprintEditor> InBlueprintEditor)
+void SMyBlueprint::Construct(const FArguments& InArgs, TWeakPtr<FBlueprintEditor> InBlueprintEditor, const UBlueprint* InBlueprint )
 {
 	BlueprintEditorPtr = InBlueprintEditor;
-	TSharedPtr<FUICommandList> ToolKitCommandList = InBlueprintEditor.Pin()->GetToolkitCommands();
+	EdGraph = nullptr;
+	
+	TSharedPtr<SWidget> ToolbarBuilderWidget = TSharedPtr<SWidget>();
+
+	if( InBlueprintEditor.IsValid() )
+	{
+		Blueprint = BlueprintEditorPtr.Pin()->GetBlueprintObj();
 
-	ToolKitCommandList->MapAction( FMyBlueprintCommands::Get().OpenGraph,
-		FExecuteAction::CreateSP(this, &SMyBlueprint::OnOpenGraph),
-		FCanExecuteAction(), FIsActionChecked(),
-		FIsActionButtonVisible::CreateSP(this, &SMyBlueprint::CanOpenGraph) );
+		TSharedPtr<FUICommandList> ToolKitCommandList = InBlueprintEditor.Pin()->GetToolkitCommands();
+
+		ToolKitCommandList->MapAction( FMyBlueprintCommands::Get().OpenGraph,
+			FExecuteAction::CreateSP(this, &SMyBlueprint::OnOpenGraph),
+			FCanExecuteAction(), FIsActionChecked(),
+			FIsActionButtonVisible::CreateSP(this, &SMyBlueprint::CanOpenGraph) );
 	
-	ToolKitCommandList->MapAction( FMyBlueprintCommands::Get().OpenGraphInNewTab,
-		FExecuteAction::CreateSP(this, &SMyBlueprint::OnOpenGraphInNewTab),
-		FCanExecuteAction(), FIsActionChecked(),
-		FIsActionButtonVisible::CreateSP(this, &SMyBlueprint::CanOpenGraph) );
-
-	ToolKitCommandList->MapAction( FMyBlueprintCommands::Get().FocusNode,
-		FExecuteAction::CreateSP(this, &SMyBlueprint::OnFocusNode),
-		FCanExecuteAction(), FIsActionChecked(),
-		FIsActionButtonVisible::CreateSP(this, &SMyBlueprint::CanFocusOnNode) );
-
-	ToolKitCommandList->MapAction( FMyBlueprintCommands::Get().FocusNodeInNewTab,
-		FExecuteAction::CreateSP(this, &SMyBlueprint::OnFocusNodeInNewTab),
-		FCanExecuteAction(), FIsActionChecked(),
-		FIsActionButtonVisible::CreateSP(this, &SMyBlueprint::CanFocusOnNode) );
-
-	ToolKitCommandList->MapAction( FMyBlueprintCommands::Get().ImplementFunction,
-		FExecuteAction::CreateSP(this, &SMyBlueprint::OnImplementFunction),
-		FCanExecuteAction(), FIsActionChecked(),
-		FIsActionButtonVisible::CreateSP(this, &SMyBlueprint::CanImplementFunction) );
+		ToolKitCommandList->MapAction( FMyBlueprintCommands::Get().OpenGraphInNewTab,
+			FExecuteAction::CreateSP(this, &SMyBlueprint::OnOpenGraphInNewTab),
+			FCanExecuteAction(), FIsActionChecked(),
+			FIsActionButtonVisible::CreateSP(this, &SMyBlueprint::CanOpenGraph) );
+
+		ToolKitCommandList->MapAction( FMyBlueprintCommands::Get().FocusNode,
+			FExecuteAction::CreateSP(this, &SMyBlueprint::OnFocusNode),
+			FCanExecuteAction(), FIsActionChecked(),
+			FIsActionButtonVisible::CreateSP(this, &SMyBlueprint::CanFocusOnNode) );
+
+		ToolKitCommandList->MapAction( FMyBlueprintCommands::Get().FocusNodeInNewTab,
+			FExecuteAction::CreateSP(this, &SMyBlueprint::OnFocusNodeInNewTab),
+			FCanExecuteAction(), FIsActionChecked(),
+			FIsActionButtonVisible::CreateSP(this, &SMyBlueprint::CanFocusOnNode) );
+
+		ToolKitCommandList->MapAction( FMyBlueprintCommands::Get().ImplementFunction,
+			FExecuteAction::CreateSP(this, &SMyBlueprint::OnImplementFunction),
+			FCanExecuteAction(), FIsActionChecked(),
+			FIsActionButtonVisible::CreateSP(this, &SMyBlueprint::CanImplementFunction) );
 	
-	ToolKitCommandList->MapAction( FMyBlueprintCommands::Get().FindEntry,
-		FExecuteAction::CreateSP(this, &SMyBlueprint::OnFindEntry),
-		FCanExecuteAction(),
-		FIsActionButtonVisible::CreateSP(this, &SMyBlueprint::CanFindEntry) );
+		ToolKitCommandList->MapAction( FMyBlueprintCommands::Get().FindEntry,
+			FExecuteAction::CreateSP(this, &SMyBlueprint::OnFindEntry),
+			FCanExecuteAction(),
+			FIsActionButtonVisible::CreateSP(this, &SMyBlueprint::CanFindEntry) );
 	
-	ToolKitCommandList->MapAction( FMyBlueprintCommands::Get().DeleteEntry,
-		FExecuteAction::CreateSP(this, &SMyBlueprint::OnDeleteEntry),
-		FCanExecuteAction::CreateSP(this, &SMyBlueprint::CanDeleteEntry) );
-
-	ToolKitCommandList->MapAction( FGenericCommands::Get().Duplicate,
-		FExecuteAction::CreateSP(this, &SMyBlueprint::OnDuplicateAction),
-		FCanExecuteAction::CreateSP(this, &SMyBlueprint::CanDuplicateAction),
-		FIsActionChecked(),
-		FIsActionButtonVisible::CreateSP(this, &SMyBlueprint::IsDuplicateActionVisible) );
-
-	ToolKitCommandList->MapAction( FMyBlueprintCommands::Get().GotoNativeVarDefinition,
-		FExecuteAction::CreateSP(this, &SMyBlueprint::GotoNativeCodeVarDefinition),
-		FCanExecuteAction(),
-		FIsActionChecked(),
-		FIsActionButtonVisible::CreateSP(this, &SMyBlueprint::IsNativeVariable) );
-
-	TSharedPtr<FBlueprintEditorToolbar> Toolbar = MakeShareable(new FBlueprintEditorToolbar(InBlueprintEditor.Pin()));
-	TSharedPtr<FExtender> Extender = MakeShareable(new FExtender);
-	Toolbar->AddNewToolbar(Extender);
-
-	FToolBarBuilder ToolbarBuilder(ToolKitCommandList, FMultiBoxCustomization::None, Extender);
-	ToolbarBuilder.BeginSection("MyBlueprint");
-	ToolbarBuilder.EndSection();
+		ToolKitCommandList->MapAction( FMyBlueprintCommands::Get().DeleteEntry,
+			FExecuteAction::CreateSP(this, &SMyBlueprint::OnDeleteEntry),
+			FCanExecuteAction::CreateSP(this, &SMyBlueprint::CanDeleteEntry) );
+
+		ToolKitCommandList->MapAction( FGenericCommands::Get().Duplicate,
+			FExecuteAction::CreateSP(this, &SMyBlueprint::OnDuplicateAction),
+			FCanExecuteAction::CreateSP(this, &SMyBlueprint::CanDuplicateAction),
+			FIsActionChecked(),
+			FIsActionButtonVisible::CreateSP(this, &SMyBlueprint::IsDuplicateActionVisible) );
+
+		ToolKitCommandList->MapAction( FMyBlueprintCommands::Get().GotoNativeVarDefinition,
+			FExecuteAction::CreateSP(this, &SMyBlueprint::GotoNativeCodeVarDefinition),
+			FCanExecuteAction(),
+			FIsActionChecked(),
+			FIsActionButtonVisible::CreateSP(this, &SMyBlueprint::IsNativeVariable) );
+
+		TSharedPtr<FBlueprintEditorToolbar> Toolbar = MakeShareable(new FBlueprintEditorToolbar(InBlueprintEditor.Pin()));
+		TSharedPtr<FExtender> Extender = MakeShareable(new FExtender);
+		Toolbar->AddNewToolbar(Extender);
+
+		FToolBarBuilder ToolbarBuilder(ToolKitCommandList, FMultiBoxCustomization::None, Extender);
+		ToolbarBuilder.BeginSection("MyBlueprint");
+		ToolbarBuilder.EndSection();
+		ToolbarBuilderWidget = ToolbarBuilder.MakeWidget();
+
+		ToolKitCommandList->MapAction(FGenericCommands::Get().Rename,
+			FExecuteAction::CreateSP(this, &SMyBlueprint::OnRequestRenameOnActionNode),
+			FCanExecuteAction::CreateSP(this, &SMyBlueprint::CanRequestRenameOnActionNode));
+	}
+	else
+	{
+		// we're in read only mode when there's no blueprint editor:
+		Blueprint = const_cast<UBlueprint*>(InBlueprint);
+		check(Blueprint);
+		ToolbarBuilderWidget = SNew(SBox);
+	}
 
 	SAssignNew(FilterBox, SSearchBox)
 		.OnTextChanged( this, &SMyBlueprint::OnFilterTextChanged );
@@ -217,7 +238,7 @@ void SMyBlueprint::Construct(const FArguments& InArgs, TWeakPtr<FBlueprintEditor
 			+ SVerticalBox::Slot()
 			.AutoHeight()
 			[
-				ToolbarBuilder.MakeWidget()
+				ToolbarBuilderWidget.ToSharedRef()
 			]
 				
 			+ SVerticalBox::Slot()
@@ -267,11 +288,12 @@ void SMyBlueprint::Construct(const FArguments& InArgs, TWeakPtr<FBlueprintEditor
 		];
 	}
 	
-	ToolKitCommandList->MapAction( FGenericCommands::Get().Rename,
-		FExecuteAction::CreateSP(this, &SMyBlueprint::OnRequestRenameOnActionNode),
-		FCanExecuteAction::CreateSP(this, &SMyBlueprint::CanRequestRenameOnActionNode) );
-
 	ResetLastPinType();
+
+	if( !BlueprintEditorPtr.IsValid() )
+	{
+		Refresh();
+	}
 }
 
 void SMyBlueprint::OnCategoryNameCommitted(const FText& InNewText, ETextCommit::Type InTextCommit, TWeakPtr< FGraphActionNode > InAction )
@@ -491,9 +513,7 @@ void SMyBlueprint::Refresh()
 
 TSharedRef<SWidget> SMyBlueprint::OnCreateWidgetForAction(FCreateWidgetForActionData* const InCreateData)
 {
-	TSharedRef<SBlueprintPaletteItem> BlueprintPaletteItem = SNew(SBlueprintPaletteItem, InCreateData, BlueprintEditorPtr.Pin());
-
-	return BlueprintPaletteItem;
+	return BlueprintEditorPtr.IsValid() ? SNew(SBlueprintPaletteItem, InCreateData, BlueprintEditorPtr.Pin()) : SNew(SBlueprintPaletteItem, InCreateData, GetBlueprintObj());
 }
 
 TSharedRef<SWidget> SMyBlueprint::ConstructLocalActionPanel()
@@ -622,7 +642,7 @@ void SMyBlueprint::GetChildEvents(UEdGraph const* EdGraph, int32 const SectionId
 void SMyBlueprint::GetLocalVariables(FGraphActionListBuilderBase& OutAllActions) const
 {
 	// We want to pull local variables from the top level function graphs
-	UEdGraph* EdGraph = FBlueprintEditorUtils::GetTopLevelGraph(BlueprintEditorPtr.Pin()->GetFocusedGraph());
+	UEdGraph* EdGraph = FBlueprintEditorUtils::GetTopLevelGraph(GetFocusedGraph());
 	if( EdGraph )
 	{
 		// grab the parent graph's name
@@ -647,7 +667,7 @@ void SMyBlueprint::GetLocalVariables(FGraphActionListBuilderBase& OutAllActions)
 				}
 
 				TSharedPtr<FEdGraphSchemaAction_K2LocalVar> NewVarAction = MakeShareable(new FEdGraphSchemaAction_K2LocalVar(Category, FText::FromName(Variable.VarName), TEXT(""), 0));
-				UFunction* Func = FindField<UFunction>(BlueprintEditorPtr.Pin()->GetBlueprintObj()->SkeletonGeneratedClass, EdGraph->GetFName());
+				UFunction* Func = FindField<UFunction>(GetBlueprintObj()->SkeletonGeneratedClass, EdGraph->GetFName());
 				NewVarAction->SetVariableInfo(Variable.VarName, Func);
 				OutAllActions.AddAction(NewVarAction);
 			}
@@ -657,6 +677,11 @@ void SMyBlueprint::GetLocalVariables(FGraphActionListBuilderBase& OutAllActions)
 
 EVisibility SMyBlueprint::GetLocalActionsListVisibility() const
 {
+	if( !BlueprintEditorPtr.IsValid())
+	{
+		return EVisibility::Visible;
+	}
+
 	if( BlueprintEditorPtr.IsValid() && BlueprintEditorPtr.Pin()->NewDocument_IsVisibleForType(FBlueprintEditor::CGT_NewLocalVariable))
 	{
 		return EVisibility::Visible;
@@ -668,9 +693,7 @@ void SMyBlueprint::CollectAllActions(FGraphActionListBuilderBase& OutAllActions)
 {
 	const UEdGraphSchema_K2* K2Schema = GetDefault<UEdGraphSchema_K2>();
 
-	if (!BlueprintEditorPtr.IsValid()) {return;}
-
-	UBlueprint* Blueprint = BlueprintEditorPtr.Pin()->GetBlueprintObj();
+	UBlueprint* Blueprint = GetBlueprintObj();
 	check(Blueprint);
 
 	EFieldIteratorFlags::SuperClassFlags FieldIteratorSuperFlag = EFieldIteratorFlags::IncludeSuper;
@@ -1161,17 +1184,27 @@ void SMyBlueprint::OnLocalActionSelected(const TArray< TSharedPtr<FEdGraphSchema
 
 void SMyBlueprint::OnActionSelected( const TArray< TSharedPtr<FEdGraphSchemaAction> >& InActions )
 {
-	TSharedPtr<FEdGraphSchemaAction> InAction( InActions.Num() > 0 ? InActions[0] : NULL );
+	TSharedPtr<FEdGraphSchemaAction> InAction(InActions.Num() > 0 ? InActions[0] : NULL);
+	UBlueprint* CurrentBlueprint = Blueprint;
+	TSharedPtr<SKismetInspector> CurrentInspector = Inspector.Pin();
 
-	if (BlueprintEditorPtr.Pin()->GetUISelectionState() == FBlueprintEditor::GraphPanel)
+	if( BlueprintEditorPtr.IsValid() )
 	{
-		// clear graph panel selection
-		BlueprintEditorPtr.Pin()->ClearSelectionInAllEditors();
+		if (BlueprintEditorPtr.Pin()->GetUISelectionState() == FBlueprintEditor::GraphPanel)
+		{
+			// clear graph panel selection
+			BlueprintEditorPtr.Pin()->ClearSelectionInAllEditors();
+		}
+		BlueprintEditorPtr.Pin()->GetUISelectionState() = InAction.IsValid() ? FBlueprintEditor::MyBlueprint : FBlueprintEditor::NoSelection;
+		CurrentBlueprint = BlueprintEditorPtr.Pin()->GetBlueprintObj();
+		CurrentInspector = BlueprintEditorPtr.Pin()->GetInspector();
 	}
-	BlueprintEditorPtr.Pin()->GetUISelectionState() = InAction.IsValid() ? FBlueprintEditor::MyBlueprint : FBlueprintEditor::NoSelection;
-	
-	UBlueprint* Blueprint = BlueprintEditorPtr.Pin()->GetBlueprintObj();
-	TSharedRef<SKismetInspector> Inspector = BlueprintEditorPtr.Pin()->GetInspector();
+
+	OnActionSelectedHelper( InAction, Blueprint, CurrentInspector.ToSharedRef() );
+}
+
+void SMyBlueprint::OnActionSelectedHelper(TSharedPtr<FEdGraphSchemaAction> InAction, UBlueprint* Blueprint, TSharedRef<SKismetInspector> Inspector)
+{
 	if(InAction.IsValid())
 	{
 		if (InAction->GetTypeId() == FEdGraphSchemaAction_K2Graph::StaticGetTypeId())
@@ -1248,6 +1281,11 @@ void SMyBlueprint::OnActionSelected( const TArray< TSharedPtr<FEdGraphSchemaActi
 
 void SMyBlueprint::OnActionDoubleClicked( const TArray< TSharedPtr<FEdGraphSchemaAction> >& InActions )
 {
+	if( !BlueprintEditorPtr.IsValid() )
+	{
+		return;
+	}
+
 	TSharedPtr<FEdGraphSchemaAction> InAction( InActions.Num() > 0 ? InActions[0] : NULL );
 
 	UBlueprint* Blueprint = BlueprintEditorPtr.Pin()->GetBlueprintObj();
@@ -1432,6 +1470,11 @@ void SMyBlueprint::GetSelectedItemsForContextMenu(TArray<FComponentEventConstruc
 
 TSharedPtr<SWidget> SMyBlueprint::OnContextMenuOpening()
 {
+	if( !BlueprintEditorPtr.IsValid() )
+	{
+		return TSharedPtr<SWidget>();
+	}
+
 	const bool bShouldCloseWindowAfterMenuSelection = true;
 	FMenuBuilder MenuBuilder( bShouldCloseWindowAfterMenuSelection,BlueprintEditorPtr.Pin()->GetToolkitCommands());
 	
@@ -1499,7 +1542,7 @@ bool SMyBlueprint::CanOpenGraph() const
 	const bool bGraph = GraphAction && GraphAction->EdGraph;
 	const FEdGraphSchemaAction_K2Delegate* DelegateAction = SelectionAsDelegate();
 	const bool bDelegate = DelegateAction && DelegateAction->EdGraph;
-	return bGraph || bDelegate;
+	return (bGraph || bDelegate) && BlueprintEditorPtr.IsValid();
 }
 
 void SMyBlueprint::OpenGraph(FDocumentTracker::EOpenDocumentCause InCause)
@@ -1688,6 +1731,17 @@ void SMyBlueprint::OnDeleteGraph(UEdGraph* InGraph, EEdGraphSchemaAction_K2Graph
 	}
 }
 
+UEdGraph* SMyBlueprint::GetFocusedGraph() const
+{
+	auto BlueprintEditorPtrPinned = BlueprintEditorPtr.Pin();
+	if( BlueprintEditorPtrPinned.IsValid() )
+	{
+		return BlueprintEditorPtrPinned->GetFocusedGraph();
+	}
+
+	return EdGraph;
+}
+
 void SMyBlueprint::OnDeleteDelegate(FEdGraphSchemaAction_K2Delegate* InDelegateAction)
 {
 	UEdGraph* EdGraph = InDelegateAction->EdGraph;
@@ -1769,7 +1823,7 @@ void SMyBlueprint::OnDeleteEntry()
 
 		GetBlueprintObj()->Modify();
 
-		UEdGraph* FunctionGraph = FBlueprintEditorUtils::GetTopLevelGraph(BlueprintEditorPtr.Pin()->GetFocusedGraph());
+		UEdGraph* FunctionGraph = FBlueprintEditorUtils::GetTopLevelGraph(GetFocusedGraph());
 		TArray<UK2Node_FunctionEntry*> FunctionEntryNodes;
 		FunctionGraph->GetNodesOfClass<UK2Node_FunctionEntry>(FunctionEntryNodes);
 		check(FunctionEntryNodes.Num() == 1);
diff --git a/Engine/Source/Editor/Kismet/Private/SMyBlueprint.h b/Engine/Source/Editor/Kismet/Private/SMyBlueprint.h
--- a/Engine/Source/Editor/Kismet/Private/SMyBlueprint.h
+++ b/Engine/Source/Editor/Kismet/Private/SMyBlueprint.h
@@ -2,6 +2,8 @@
 
 #pragma once
 
+#include "DiffUtils.h"
+
 /* Enums to use when grouping the blueprint members in the list panel. The order here will determine the order in the list */
 namespace NodeSectionID
 {
@@ -46,21 +48,21 @@ public:
 	virtual void RegisterCommands() override;
 };
 
-
-
 class SMyBlueprint : public SCompoundWidget
 {
 public:
-	SLATE_BEGIN_ARGS( SMyBlueprint ) {};
+	SLATE_BEGIN_ARGS( SMyBlueprint ) {}
 	SLATE_END_ARGS()
 
-	void Construct(const FArguments& InArgs, TWeakPtr<FBlueprintEditor> InBlueprintEditor);
+	void Construct(const FArguments& InArgs, TWeakPtr<FBlueprintEditor> InBlueprintEditor, const UBlueprint* InBlueprint = nullptr);
+	void SetInspector( TSharedPtr<SKismetInspector> InInspector ) { Inspector = InInspector ; }
 
 	/* Reset the last pin type settings to default. */
 	void ResetLastPinType();
 
 	/** Refreshes the graph action menu */
 	void Refresh();
+	void SetFocusedGraph(UEdGraph* InEdGraph) { EdGraph = InEdGraph; }
 	
 	/** Accessor for getting the current selection as a K2 graph */
 	FEdGraphSchemaAction_K2Graph* SelectionAsGraph() const;
@@ -93,7 +95,7 @@ public:
 	FEdGraphPinType& GetLastFunctionPinTypeUsed() {EnsureLastPinTypeValid(); return LastFunctionPinType;}
 
 	/** Accessor the blueprint object from the main editor */
-	UBlueprint* GetBlueprintObj() const {return BlueprintEditorPtr.Pin()->GetBlueprintObj();}
+	UBlueprint* GetBlueprintObj() const {return Blueprint;}
 
 	/** Gets whether we are showing user variables only or not */
 	bool ShowUserVarsOnly() const {return bShowUserVarsOnly;}
@@ -124,7 +126,6 @@ public:
 
 	/** Move the category before the target category */
 	bool MoveCategoryBeforeCategory( const FString& CategoryToMove, const FString& TargetCategory );
-
 private:
 	/** Creates widgets for the graph schema actions */
 	TSharedRef<SWidget> OnCreateWidgetForAction(struct FCreateWidgetForActionData* const InCreateData);
@@ -145,6 +146,7 @@ private:
 	FReply OnActionDragged(const TArray< TSharedPtr<FEdGraphSchemaAction> >& InActions, const FPointerEvent& MouseEvent);
 	FReply OnCategoryDragged(const FString& InCategory, const FPointerEvent& MouseEvent);
 	void OnActionSelected(const TArray< TSharedPtr<FEdGraphSchemaAction> >& InActions);
+	static void OnActionSelectedHelper(TSharedPtr<FEdGraphSchemaAction> InAction, UBlueprint* Blueprint, TSharedRef<SKismetInspector> Inspector);
 	void OnGlobalActionSelected(const TArray< TSharedPtr<FEdGraphSchemaAction> >& InActions);
 	void OnLocalActionSelected(const TArray< TSharedPtr<FEdGraphSchemaAction> >& InActions);
 	void OnActionDoubleClicked(const TArray< TSharedPtr<FEdGraphSchemaAction> >& InActions);
@@ -202,6 +204,8 @@ private:
 
 	/** Helper function to delete a delegate in the MyBlueprint window */
 	void OnDeleteDelegate(FEdGraphSchemaAction_K2Delegate* InDelegateAction);
+
+	UEdGraph* GetFocusedGraph() const;
 private:
 	/** Pointer back to the blueprint editor that owns us */
 	TWeakPtr<FBlueprintEditor> BlueprintEditorPtr;
@@ -230,4 +234,13 @@ private:
 
 	/** Enums created from 'blueprint' level */
 	TArray<TWeakObjectPtr<UUserDefinedStruct>> StructsAddedToBlueprint;
+
+	/** The blueprint being displayed: */
+	UBlueprint* Blueprint;
+
+	/** The Ed Graph being displayed: */
+	UEdGraph* EdGraph;
+
+	/** The Kismet Inspector used to display properties: */
+	TWeakPtr<SKismetInspector> Inspector;
 };
diff --git a/Engine/Source/Editor/Kismet/Public/SBlueprintDiff.h b/Engine/Source/Editor/Kismet/Public/SBlueprintDiff.h
--- a/Engine/Source/Editor/Kismet/Public/SBlueprintDiff.h
+++ b/Engine/Source/Editor/Kismet/Public/SBlueprintDiff.h
@@ -30,6 +30,8 @@ struct FDiffResultItem : public TSharedFromThis<FDiffResultItem>
 	TSharedRef<SWidget> KISMET_API GenerateWidget() const;
 };
 
+DECLARE_DELEGATE_OneParam(FOnMyBlueprintActionSelected, UObject*);
+
 namespace DiffWidgetUtils
 {
 	KISMET_API void SelectNextRow(SListView< TSharedPtr< struct FDiffSingleResult> >& ListView, const TArray< TSharedPtr< struct FDiffSingleResult > >& ListViewSource );
@@ -43,8 +45,14 @@ struct KISMET_API FDiffPanel
 {
 	FDiffPanel();
 
-	/*Generate this panel based on the specified graph*/
-	void	GeneratePanel(UEdGraph* Graph, UEdGraph* GraphToDiff);
+	/** Initializes the panel, can be moved into constructor if diff and merge clients are made more uniform: */
+	TSharedRef<SWidget> InitializeDiffPanel();
+
+	/* Generate this panel based on the specified graph */
+	void GeneratePanel(UEdGraph* Graph, UEdGraph* GraphToDiff);
+
+	/* Generate the 'MyBlueprint' widget, which is private to this module */
+	TSharedRef<SWidget> GenerateMyBlueprintPanel();
 
 	/*Get the title to show at the top of the panel*/
 	FString GetTitle() const;
@@ -66,7 +74,13 @@ struct KISMET_API FDiffPanel
 	const class UBlueprint*				Blueprint;
 
 	/*The border around the graph editor, used to change the content when new graphs are set */
-	TSharedPtr<SBorder>				GraphEditorBorder;
+	TSharedPtr<SBox>				GraphEditorBorder;
+
+	/* The border around the my blueprint panel, used to regenerate the panel when the new graphs are set */
+	TSharedPtr<class SMyBlueprint>	MyBlueprint;
+
+	/*The box around the the details view associated with the graph editor */
+	TSharedPtr<class SKismetInspector>	DetailsView;
 
 	/*The graph editor which does the work of displaying the graph*/
 	TWeakPtr<class SGraphEditor>	GraphEditor;
@@ -87,7 +101,6 @@ private:
 /* Visual Diff between two Blueprints*/
 class  KISMET_API SBlueprintDiff: public SCompoundWidget
 {
-
 public:
 	DECLARE_DELEGATE_TwoParams( FOpenInDefaults, const class UBlueprint* , const class UBlueprint* );
 
@@ -118,6 +131,10 @@ protected:
 	bool HasNextDiff() const;
 	bool HasPrevDiff() const;
 
+	/** Spawns the tabs that contain the Graph views and blueprints views respectively: */
+	TSharedRef<SDockTab> CreateGraphDiffViews( const FSpawnTabArgs& Args );
+	TSharedRef<SDockTab> CreateMyBlueprintsViews( const FSpawnTabArgs& Args );
+
 	/* Need to process keys for shortcuts to buttons */
 	virtual FReply OnKeyDown( const FGeometry& MyGeometry, const FKeyEvent& InKeyEvent ) override;
 
@@ -137,6 +154,7 @@ protected:
 	/*Called when user clicks on a new graph list item */
 	void OnSelectionChanged(FGraphToDiff Item, ESelectInfo::Type SelectionType);
 
+	/** Called when user clicks on an entry in the listview of differences */
 	void OnDiffListSelectionChanged(const TSharedPtr<struct FDiffResultItem>& TheDiff, struct FListItemGraphToDiff* GraphDiffer);
 		
 	/** Disable the focus on a particular pin */
@@ -192,6 +210,9 @@ protected:
 
 	/** Helper class for highlighting diffs in different types of controls (graph view, details view, etc) */
 	TSharedPtr< class IDiffControl > DiffControl;
+
+	/** We can't use the global tab manager because we need to instance the diff control, so we have our own tab manager: */
+	TSharedPtr<FTabManager> TabManager;
 };
 
 
diff --git a/Engine/Source/Editor/Kismet/Public/SKismetInspector.h b/Engine/Source/Editor/Kismet/Public/SKismetInspector.h
--- a/Engine/Source/Editor/Kismet/Public/SKismetInspector.h
+++ b/Engine/Source/Editor/Kismet/Public/SKismetInspector.h
@@ -22,7 +22,8 @@ public:
 		, _ShowTitleArea(false)
 		{}
 
-		SLATE_ARGUMENT( TWeakPtr<FBlueprintEditor>, Kismet2 )
+		SLATE_ARGUMENT(TWeakPtr<FBlueprintEditor>, Kismet2)
+		SLATE_ARGUMENT(TWeakPtr<SMyBlueprint>, MyBlueprintWidget)
 		SLATE_ARGUMENT( bool, ShowPublicViewControl )
 		SLATE_ARGUMENT( bool, HideNameArea )
 		SLATE_ARGUMENT( FIsPropertyEditingEnabled, IsPropertyEditingEnabledDelegate )
