No longer hiding the filter area in the details panel for variables selected in SMyBlueprint, fixes a bug reported by [REDACTED_PERSON] in PR#432 (stale filter string stuck in a hidden filter area!)

diff --git a/Engine/Source/Editor/Kismet/Private/SMyBlueprint.cpp b/Engine/Source/Editor/Kismet/Private/SMyBlueprint.cpp
--- a/Engine/Source/Editor/Kismet/Private/SMyBlueprint.cpp
+++ b/Engine/Source/Editor/Kismet/Private/SMyBlueprint.cpp
@@ -1231,7 +1231,6 @@ void SMyBlueprint::OnActionSelectedHelper(TSharedPtr<FEdGraphSchemaAction> InAct
 			FEdGraphSchemaAction_K2Var* VarAction = (FEdGraphSchemaAction_K2Var*)InAction.Get();
 
 			SKismetInspector::FShowDetailsOptions Options(VarAction->GetVariableName().ToString());
-			Options.bHideFilterArea = true;
 
 			Inspector->ShowDetailsForSingleObject(VarAction->GetProperty(), Options);
 		}
@@ -1240,7 +1239,6 @@ void SMyBlueprint::OnActionSelectedHelper(TSharedPtr<FEdGraphSchemaAction> InAct
 			FEdGraphSchemaAction_K2LocalVar* VarAction = (FEdGraphSchemaAction_K2LocalVar*)InAction.Get();
 
 			SKismetInspector::FShowDetailsOptions Options(VarAction->GetVariableName().ToString());
-			Options.bHideFilterArea = true;
 
 			Inspector->ShowDetailsForSingleObject(VarAction->GetProperty(), Options);
 		}
