Converted the property editor to use FText

diff --git a/Engine/Source/Editor/DetailCustomizations/Private/MathStructProxyCustomizations.cpp b/Engine/Source/Editor/DetailCustomizations/Private/MathStructProxyCustomizations.cpp
--- a/Engine/Source/Editor/DetailCustomizations/Private/MathStructProxyCustomizations.cpp
+++ b/Engine/Source/Editor/DetailCustomizations/Private/MathStructProxyCustomizations.cpp
@@ -116,12 +116,12 @@ void FMatrixStructCustomization::CustomizeChildren(TSharedRef<class IPropertyHan
 
 	TWeakPtr<IPropertyHandle> WeakHandlePtr = StructPropertyHandle;
 
-	StructBuilder.AddChildContent(LOCTEXT("LocationLabel", "Location").ToString())
+	StructBuilder.AddChildContent(LOCTEXT("LocationLabel", "Location"))
 		.CopyAction(FUIAction(FExecuteAction::CreateSP(this, &FMatrixStructCustomization::OnCopy, FTransformField::Location, WeakHandlePtr)))
 		.PasteAction(FUIAction(FExecuteAction::CreateSP(this, &FMatrixStructCustomization::OnPaste, FTransformField::Location, WeakHandlePtr)))
 		.NameContent()
 		[
-			StructPropertyHandle->CreatePropertyNameWidget(LOCTEXT("LocationLabel", "Location").ToString())
+			StructPropertyHandle->CreatePropertyNameWidget(LOCTEXT("LocationLabel", "Location"))
 		]
 		.ValueContent()
 		.MinDesiredWidth(375.0f)
@@ -145,12 +145,12 @@ void FMatrixStructCustomization::CustomizeChildren(TSharedRef<class IPropertyHan
 			]
 		];
 
-	StructBuilder.AddChildContent(LOCTEXT("RotationLabel", "Rotation").ToString())
+	StructBuilder.AddChildContent(LOCTEXT("RotationLabel", "Rotation"))
 		.CopyAction(FUIAction(FExecuteAction::CreateSP(this, &FMatrixStructCustomization::OnCopy, FTransformField::Rotation, WeakHandlePtr)))
 		.PasteAction(FUIAction(FExecuteAction::CreateSP(this, &FMatrixStructCustomization::OnPaste, FTransformField::Rotation, WeakHandlePtr)))
 		.NameContent()
 		[
-			StructPropertyHandle->CreatePropertyNameWidget(LOCTEXT("RotationLabel", "Rotation").ToString())
+			StructPropertyHandle->CreatePropertyNameWidget(LOCTEXT("RotationLabel", "Rotation"))
 		]
 		.ValueContent()
 		.MinDesiredWidth(375.0f)
@@ -174,12 +174,12 @@ void FMatrixStructCustomization::CustomizeChildren(TSharedRef<class IPropertyHan
 			]
 		];
 
-	StructBuilder.AddChildContent(LOCTEXT("ScaleLabel", "Scale").ToString())
+	StructBuilder.AddChildContent(LOCTEXT("ScaleLabel", "Scale"))
 		.CopyAction(FUIAction(FExecuteAction::CreateSP(this, &FMatrixStructCustomization::OnCopy, FTransformField::Scale, WeakHandlePtr)))
 		.PasteAction(FUIAction(FExecuteAction::CreateSP(this, &FMatrixStructCustomization::OnPaste, FTransformField::Scale, WeakHandlePtr)))
 		.NameContent()
 		[
-			StructPropertyHandle->CreatePropertyNameWidget(LOCTEXT("ScaleLabel", "Scale").ToString())
+			StructPropertyHandle->CreatePropertyNameWidget(LOCTEXT("ScaleLabel", "Scale"))
 		]
 		.ValueContent()
 		.MinDesiredWidth(375.0f)
diff --git a/Engine/Source/Editor/DetailCustomizations/Private/MatineeActorDetails.cpp b/Engine/Source/Editor/DetailCustomizations/Private/MatineeActorDetails.cpp
--- a/Engine/Source/Editor/DetailCustomizations/Private/MatineeActorDetails.cpp
+++ b/Engine/Source/Editor/DetailCustomizations/Private/MatineeActorDetails.cpp
@@ -31,8 +31,8 @@ void FMatineeActorDetails::CustomizeDetails( IDetailLayoutBuilder& DetailLayout
 		}
 	}
 	
-	DetailLayout.EditCategory( "MatineeActor", NSLOCTEXT("MatineeActorDetails", "MatineeActor", "Matinee Actor").ToString(), ECategoryPriority::Important )
-	.AddCustomRow( NSLOCTEXT("MatineeActorDetails", "OpenMatinee", "Open Matinee").ToString() )
+	DetailLayout.EditCategory( "MatineeActor", NSLOCTEXT("MatineeActorDetails", "MatineeActor", "Matinee Actor"), ECategoryPriority::Important )
+	.AddCustomRow( NSLOCTEXT("MatineeActorDetails", "OpenMatinee", "Open Matinee") )
 	[
 		SNew(SHorizontalBox)
 		+SHorizontalBox::Slot()
@@ -44,7 +44,7 @@ void FMatineeActorDetails::CustomizeDetails( IDetailLayoutBuilder& DetailLayout
 			.VAlign(VAlign_Center)
 			.HAlign(HAlign_Center)
 			.OnClicked( this, &FMatineeActorDetails::OnOpenMatineeForActor )
-			.Text( NSLOCTEXT("MatineeActorDetails", "OpenMatinee", "Open Matinee").ToString() )
+			.Text( NSLOCTEXT("MatineeActorDetails", "OpenMatinee", "Open Matinee") )
 		]
 	];
 }
diff --git a/Engine/Source/Editor/DetailCustomizations/Private/MediaPlayerCustomization.cpp b/Engine/Source/Editor/DetailCustomizations/Private/MediaPlayerCustomization.cpp
--- a/Engine/Source/Editor/DetailCustomizations/Private/MediaPlayerCustomization.cpp
+++ b/Engine/Source/Editor/DetailCustomizations/Private/MediaPlayerCustomization.cpp
@@ -20,14 +20,14 @@ void FMediaPlayerCustomization::CustomizeDetails( IDetailLayoutBuilder& DetailBu
 	DetailBuilder.GetObjectsBeingCustomized(CustomizedMediaPlayers);
 
 	// customize 'Source' category
-	IDetailCategoryBuilder& SourceCategory = DetailBuilder.EditCategory("Source", TEXT(""));
+	IDetailCategoryBuilder& SourceCategory = DetailBuilder.EditCategory("Source");
 	{
 		// URL
 		UrlProperty = DetailBuilder.GetProperty("URL");
 		{
 			IDetailPropertyRow& UrlRow = SourceCategory.AddProperty(UrlProperty);
 
-			UrlRow.DisplayName(TEXT("URL"));
+			UrlRow.DisplayName(LOCTEXT("URLLabel", "URL"));
 			UrlRow.CustomWidget()
 				.NameContent()
 				[
@@ -73,10 +73,10 @@ void FMediaPlayerCustomization::CustomizeDetails( IDetailLayoutBuilder& DetailBu
 	}
 
 	// customize 'Information' category
-	IDetailCategoryBuilder& InformationCategory = DetailBuilder.EditCategory("Information", TEXT(""), ECategoryPriority::Uncommon);
+	IDetailCategoryBuilder& InformationCategory = DetailBuilder.EditCategory("Information", FText::GetEmpty(), ECategoryPriority::Uncommon);
 	{
 		// duration
-		InformationCategory.AddCustomRow(TEXT("Duration"))
+		InformationCategory.AddCustomRow(LOCTEXT("Duration", "Duration"))
 			.NameContent()
 			[
 				SNew(STextBlock)
@@ -93,7 +93,7 @@ void FMediaPlayerCustomization::CustomizeDetails( IDetailLayoutBuilder& DetailBu
 			];
 
 		// forward rates (thinned)
-		InformationCategory.AddCustomRow(TEXT("ForwardRatesThinned"))
+		InformationCategory.AddCustomRow(LOCTEXT("ForwardRatesThinned", "Forward Rates (Thinned)"))
 			.NameContent()
 			[
 				SNew(STextBlock)
@@ -110,7 +110,7 @@ void FMediaPlayerCustomization::CustomizeDetails( IDetailLayoutBuilder& DetailBu
 			];
 
 		// forward rates (unthinned)
-		InformationCategory.AddCustomRow(TEXT("ForwardRatesUnthinned"))
+		InformationCategory.AddCustomRow(LOCTEXT("ForwardRatesUnthinned", "Forward Rates (Unthinned)"))
 			.NameContent()
 			[
 				SNew(STextBlock)
@@ -127,7 +127,7 @@ void FMediaPlayerCustomization::CustomizeDetails( IDetailLayoutBuilder& DetailBu
 			];
 
 		// reverse rates (thinned)
-		InformationCategory.AddCustomRow(TEXT("ReverseRatesThinned"))
+		InformationCategory.AddCustomRow(LOCTEXT("ReverseRatesThinned", "Reverse Rates (Thinned)"))
 			.NameContent()
 			[
 				SNew(STextBlock)
@@ -144,7 +144,7 @@ void FMediaPlayerCustomization::CustomizeDetails( IDetailLayoutBuilder& DetailBu
 			];
 
 		// reverse rates (unthinned)
-		InformationCategory.AddCustomRow(TEXT("ReverseRatesUnthinned"))
+		InformationCategory.AddCustomRow(LOCTEXT("ReverseRatesUnthinned", "Reverse Rates (Unthinned)"))
 			.NameContent()
 			[
 				SNew(STextBlock)
@@ -161,7 +161,7 @@ void FMediaPlayerCustomization::CustomizeDetails( IDetailLayoutBuilder& DetailBu
 			];
 
 		// supports scrubbing
-		InformationCategory.AddCustomRow(TEXT("SupportsScrubbing"))
+		InformationCategory.AddCustomRow(LOCTEXT("SupportsScrubbing", "Supports Scrubbing"))
 			.NameContent()
 			[
 				SNew(STextBlock)
@@ -178,7 +178,7 @@ void FMediaPlayerCustomization::CustomizeDetails( IDetailLayoutBuilder& DetailBu
 			];
 
 		// supports seeking
-		InformationCategory.AddCustomRow(TEXT("SupportsSeeking"))
+		InformationCategory.AddCustomRow(LOCTEXT("SupportsSeeking", "Supports Seeking"))
 			.NameContent()
 			[
 				SNew(STextBlock)
diff --git a/Engine/Source/Editor/DetailCustomizations/Private/NavLinkStructCustomization.cpp b/Engine/Source/Editor/DetailCustomizations/Private/NavLinkStructCustomization.cpp
--- a/Engine/Source/Editor/DetailCustomizations/Private/NavLinkStructCustomization.cpp
+++ b/Engine/Source/Editor/DetailCustomizations/Private/NavLinkStructCustomization.cpp
@@ -54,8 +54,8 @@ void FNavLinkStructCustomization::CustomizeChildren( TSharedRef<class IPropertyH
 
 			if (AgentIdx >= 0 && AgentIdx < NumAgents)
 			{
-				FString PropName = LOCTEXT("SupportedAgent", "Supports Agent: ").ToString() + *DefNavSys->SupportedAgents[AgentIdx].Name.ToString();
-				StructBuilder.AddChildContent(TEXT("SupportedAgent"))
+				FText PropName = FText::Format(LOCTEXT("SupportedAgentFmt", "Supports Agent: {0}"), FText::FromName(DefNavSys->SupportedAgents[AgentIdx].Name));
+				StructBuilder.AddChildContent(PropName)
 					.NameContent()
 					[
 						SNew(STextBlock)
diff --git a/Engine/Source/Editor/DetailCustomizations/Private/ParticleSysParamStructCustomization.cpp b/Engine/Source/Editor/DetailCustomizations/Private/ParticleSysParamStructCustomization.cpp
--- a/Engine/Source/Editor/DetailCustomizations/Private/ParticleSysParamStructCustomization.cpp
+++ b/Engine/Source/Editor/DetailCustomizations/Private/ParticleSysParamStructCustomization.cpp
@@ -23,8 +23,8 @@ FParticleSysParamStructCustomization::FParticleSysParamStructCustomization()
 void FParticleSysParamStructCustomization::CustomizeHeader(TSharedRef<IPropertyHandle> StructPropertyHandle, FDetailWidgetRow& HeaderRow, IPropertyTypeCustomizationUtils& StructCustomizationUtils)
 {
 	const bool bDisplayResetToDefault = false;
-	const FString DisplayNameOverride = TEXT("");
-	const FString DisplayToolTipOverride = TEXT("");
+	const FText DisplayNameOverride = FText::GetEmpty();
+	const FText DisplayToolTipOverride = FText::GetEmpty();
 
 	HeaderRow
 		.NameContent()
@@ -63,7 +63,7 @@ void FParticleSysParamStructCustomization::CustomizeChildren(TSharedRef<IPropert
 	TArray<bool> RestrictedList;
 	ParamTypeHandle->GeneratePossibleValues(ParameterTypeNames, ParameterTypeToolTips, RestrictedList);
 
-	StructBuilder.AddChildContent(TEXT("ParamType"))
+	StructBuilder.AddChildContent(LOCTEXT("ParamType", "Param Type"))
 		.NameContent()
 		[
 			ParamTypeHandle->CreatePropertyNameWidget()
diff --git a/Engine/Source/Editor/DetailCustomizations/Private/ParticleSystemComponentDetails.cpp b/Engine/Source/Editor/DetailCustomizations/Private/ParticleSystemComponentDetails.cpp
--- a/Engine/Source/Editor/DetailCustomizations/Private/ParticleSystemComponentDetails.cpp
+++ b/Engine/Source/Editor/DetailCustomizations/Private/ParticleSystemComponentDetails.cpp
@@ -16,11 +16,11 @@ void FParticleSystemComponentDetails::CustomizeDetails(IDetailLayoutBuilder& InD
 {
 	this->DetailLayout = &InDetailLayout;
 
-	InDetailLayout.EditCategory("Particles", TEXT(""), ECategoryPriority::Important);
+	InDetailLayout.EditCategory("Particles", FText::GetEmpty(), ECategoryPriority::Important);
 
-	IDetailCategoryBuilder& CustomCategory = InDetailLayout.EditCategory("EmitterActions", NSLOCTEXT("ParticleSystemComponentDetails", "EmitterActionCategoryName", "Emitter Actions").ToString(), ECategoryPriority::Important);
+	IDetailCategoryBuilder& CustomCategory = InDetailLayout.EditCategory("EmitterActions", NSLOCTEXT("ParticleSystemComponentDetails", "EmitterActionCategoryName", "Emitter Actions"), ECategoryPriority::Important);
 
-	CustomCategory.AddCustomRow(TEXT(""))
+	CustomCategory.AddCustomRow(FText::GetEmpty())
 		[
 			SNew(SUniformGridPanel)
 			.SlotPadding(2.0f)
diff --git a/Engine/Source/Editor/DetailCustomizations/Private/PhysicsConstraintComponentDetails.cpp b/Engine/Source/Editor/DetailCustomizations/Private/PhysicsConstraintComponentDetails.cpp
--- a/Engine/Source/Editor/DetailCustomizations/Private/PhysicsConstraintComponentDetails.cpp
+++ b/Engine/Source/Editor/DetailCustomizations/Private/PhysicsConstraintComponentDetails.cpp
@@ -48,7 +48,7 @@ void FPhysicsConstraintComponentDetails::CustomizeDetails( IDetailLayoutBuilder&
 		ConstraintComponent = TWeakObjectPtr<UObject>(ConstraintComp);
 
 		DetailBuilder.EditCategory( "Joint Presets" )
-			.AddCustomRow( TEXT("") )
+			.AddCustomRow( FText::GetEmpty() )
 			[
 				SNew(SVerticalBox)
 				+ SVerticalBox::Slot()
@@ -99,7 +99,7 @@ void FPhysicsConstraintComponentDetails::CustomizeDetails( IDetailLayoutBuilder&
 
 	// Linear Limits
 	{
-		IDetailCategoryBuilder& LinearLimitCat = DetailBuilder.EditCategory("Linear Limits", TEXT(""));
+		IDetailCategoryBuilder& LinearLimitCat = DetailBuilder.EditCategory("Linear Limits");
 
 		LinearXMotionProperty = ConstraintInstance->GetChildHandle("LinearXMotion");
 		LinearYMotionProperty = ConstraintInstance->GetChildHandle("LinearYMotion");
@@ -193,7 +193,7 @@ void FPhysicsConstraintComponentDetails::CustomizeDetails( IDetailLayoutBuilder&
 
 	// Angular Limits
 	{
-		IDetailCategoryBuilder& AngularLimitCat = DetailBuilder.EditCategory("Angular Limits", TEXT(""));
+		IDetailCategoryBuilder& AngularLimitCat = DetailBuilder.EditCategory("Angular Limits");
 
 		AngularSwing1MotionProperty = ConstraintInstance->GetChildHandle(GET_MEMBER_NAME_CHECKED(FConstraintInstance, AngularSwing1Motion));
 		AngularSwing2MotionProperty = ConstraintInstance->GetChildHandle(GET_MEMBER_NAME_CHECKED(FConstraintInstance, AngularSwing2Motion));
@@ -313,12 +313,12 @@ void FPhysicsConstraintComponentDetails::CustomizeDetails( IDetailLayoutBuilder&
 
 	// Linear Drive
 	{
-		IDetailCategoryBuilder& LinearMotorCat = DetailBuilder.EditCategory("LinearMotor", TEXT(""));
+		IDetailCategoryBuilder& LinearMotorCat = DetailBuilder.EditCategory("LinearMotor");
 
 		LinearPositionDriveProperty = ConstraintInstance->GetChildHandle(GET_MEMBER_NAME_CHECKED(FConstraintInstance, bLinearPositionDrive));
 		LinearVelocityDriveProperty = ConstraintInstance->GetChildHandle(GET_MEMBER_NAME_CHECKED(FConstraintInstance, bLinearVelocityDrive));
 
-		IDetailGroup& PositionGroup = LinearMotorCat.AddGroup("Linear Position Drive", TEXT(""));
+		IDetailGroup& PositionGroup = LinearMotorCat.AddGroup("Linear Position Drive", FText::GetEmpty());
 		PositionGroup.HeaderProperty(LinearPositionDriveProperty.ToSharedRef());
 		PositionGroup.AddPropertyRow(ConstraintInstance->GetChildHandle(GET_MEMBER_NAME_CHECKED(FConstraintInstance, bLinearXPositionDrive)).ToSharedRef())
 			.Visibility(TAttribute<EVisibility>::Create(TAttribute<EVisibility>::FGetter::CreateSP(this, &FPhysicsConstraintComponentDetails::IsPropertyVisible, EPropertyType::LinearPositionDrive)));
@@ -329,7 +329,7 @@ void FPhysicsConstraintComponentDetails::CustomizeDetails( IDetailLayoutBuilder&
 		PositionGroup.AddPropertyRow(ConstraintInstance->GetChildHandle(GET_MEMBER_NAME_CHECKED(FConstraintInstance, LinearPositionTarget)).ToSharedRef())
 			.Visibility(TAttribute<EVisibility>::Create(TAttribute<EVisibility>::FGetter::CreateSP(this, &FPhysicsConstraintComponentDetails::IsPropertyVisible, EPropertyType::LinearPositionDrive)));
 
-		IDetailGroup& VelocityGroup = LinearMotorCat.AddGroup("Linear Velocity Drive", TEXT(""));
+		IDetailGroup& VelocityGroup = LinearMotorCat.AddGroup("Linear Velocity Drive", FText::GetEmpty());
 		VelocityGroup.HeaderProperty(LinearVelocityDriveProperty.ToSharedRef());
 		VelocityGroup.AddPropertyRow(ConstraintInstance->GetChildHandle(GET_MEMBER_NAME_CHECKED(FConstraintInstance, LinearVelocityTarget)).ToSharedRef())
 			.Visibility(TAttribute<EVisibility>::Create(TAttribute<EVisibility>::FGetter::CreateSP(this, &FPhysicsConstraintComponentDetails::IsPropertyVisible, EPropertyType::LinearVelocityDrive)));
@@ -347,15 +347,15 @@ void FPhysicsConstraintComponentDetails::CustomizeDetails( IDetailLayoutBuilder&
 		AngularPositionDriveProperty = ConstraintInstance->GetChildHandle(GET_MEMBER_NAME_CHECKED(FConstraintInstance, bAngularOrientationDrive));
 		AngularVelocityDriveProperty = ConstraintInstance->GetChildHandle(GET_MEMBER_NAME_CHECKED(FConstraintInstance, bAngularVelocityDrive));
 
-		IDetailCategoryBuilder& AngularMotorCat = DetailBuilder.EditCategory("AngularMotor", TEXT(""));
+		IDetailCategoryBuilder& AngularMotorCat = DetailBuilder.EditCategory("AngularMotor");
 
-		IDetailGroup& PositionGroup = AngularMotorCat.AddGroup("Angular Orientation Drive", TEXT(""));
+		IDetailGroup& PositionGroup = AngularMotorCat.AddGroup("Angular Orientation Drive", FText::GetEmpty());
 		PositionGroup.HeaderProperty(AngularPositionDriveProperty.ToSharedRef());
 		
 		PositionGroup.AddPropertyRow(ConstraintInstance->GetChildHandle(GET_MEMBER_NAME_CHECKED(FConstraintInstance, AngularOrientationTarget)).ToSharedRef())
 			.Visibility(TAttribute<EVisibility>::Create(TAttribute<EVisibility>::FGetter::CreateSP(this, &FPhysicsConstraintComponentDetails::IsPropertyVisible, EPropertyType::AngularPositionDrive)));
 
-		IDetailGroup& VelocityGroup = AngularMotorCat.AddGroup("Angular Velocity Drive", TEXT(""));
+		IDetailGroup& VelocityGroup = AngularMotorCat.AddGroup("Angular Velocity Drive", FText::GetEmpty());
 		VelocityGroup.HeaderProperty(AngularVelocityDriveProperty.ToSharedRef());
 		VelocityGroup.AddPropertyRow(ConstraintInstance->GetChildHandle(GET_MEMBER_NAME_CHECKED(FConstraintInstance, AngularVelocityTarget)).ToSharedRef())
 			.Visibility(TAttribute<EVisibility>::Create(TAttribute<EVisibility>::FGetter::CreateSP(this, &FPhysicsConstraintComponentDetails::IsPropertyVisible, EPropertyType::AngularVelocityDrive)));
diff --git a/Engine/Source/Editor/DetailCustomizations/Private/PhysicsSettingsDetails.cpp b/Engine/Source/Editor/DetailCustomizations/Private/PhysicsSettingsDetails.cpp
--- a/Engine/Source/Editor/DetailCustomizations/Private/PhysicsSettingsDetails.cpp
+++ b/Engine/Source/Editor/DetailCustomizations/Private/PhysicsSettingsDetails.cpp
@@ -124,7 +124,7 @@ TSharedRef<IDetailCustomization> FPhysicsSettingsDetails::MakeInstance()
 
 void FPhysicsSettingsDetails::CustomizeDetails( IDetailLayoutBuilder& DetailBuilder )
 {
-  	IDetailCategoryBuilder& PhysicalSurfaceCategory = DetailBuilder.EditCategory("Physical Surface", TEXT(""), ECategoryPriority::Uncommon);
+  	IDetailCategoryBuilder& PhysicalSurfaceCategory = DetailBuilder.EditCategory("Physical Surface", FText::GetEmpty(), ECategoryPriority::Uncommon);
 
 	PhysicsSettings= UPhysicsSettings::Get();
 	check(PhysicsSettings);
@@ -138,7 +138,7 @@ void FPhysicsSettingsDetails::CustomizeDetails( IDetailLayoutBuilder& DetailBuil
 	TSharedPtr<SToolTip> PhysicalSurfaceTooltip = IDocumentation::Get()->CreateToolTip(LOCTEXT("PhysicalSurface", "Edit physical surface."), NULL, PhysicalSurfaceDocLink, TEXT("PhysicalSurface"));
 
 	// Customize collision section
-	PhysicalSurfaceCategory.AddCustomRow(LOCTEXT("FPhysicsSettingsDetails_PhysicalSurface", "Physical Surface").ToString())
+	PhysicalSurfaceCategory.AddCustomRow(LOCTEXT("FPhysicsSettingsDetails_PhysicalSurface", "Physical Surface"))
 	[
 		SNew(SVerticalBox)
 
diff --git a/Engine/Source/Editor/DetailCustomizations/Private/PrimitiveComponentDetails.cpp b/Engine/Source/Editor/DetailCustomizations/Private/PrimitiveComponentDetails.cpp
--- a/Engine/Source/Editor/DetailCustomizations/Private/PrimitiveComponentDetails.cpp
+++ b/Engine/Source/Editor/DetailCustomizations/Private/PrimitiveComponentDetails.cpp
@@ -161,7 +161,7 @@ EVisibility FPrimitiveComponentDetails::IsMassVisible(bool bOverrideMass) const
 void FPrimitiveComponentDetails::CustomizeDetails( IDetailLayoutBuilder& DetailBuilder )
 {
 	TSharedRef<IPropertyHandle> MobilityHandle = DetailBuilder.GetProperty(GET_MEMBER_NAME_CHECKED(UPrimitiveComponent, Mobility), USceneComponent::StaticClass());
-	MobilityHandle->SetToolTipText(LOCTEXT("PrimitiveMobilityTooltip", "Mobility for primitive components controls how they can be modified in game and therefore how they interact with lighting and physics.\n● A movable primitive component can be changed in game, but requires dynamic lighting and shadowing from lights which have a large performance cost.\n● A static primitive component can't be changed in game, but can have its lighting baked, which allows rendering to be very efficient.").ToString());
+	MobilityHandle->SetToolTipText(LOCTEXT("PrimitiveMobilityTooltip", "Mobility for primitive components controls how they can be modified in game and therefore how they interact with lighting and physics.\n● A movable primitive component can be changed in game, but requires dynamic lighting and shadowing from lights which have a large performance cost.\n● A static primitive component can't be changed in game, but can have its lighting baked, which allows rendering to be very efficient."));
 
 	if ( DetailBuilder.GetProperty(GET_MEMBER_NAME_CHECKED(UPrimitiveComponent, BodyInstance))->IsValidHandle() )
 	{
@@ -246,7 +246,7 @@ void FPrimitiveComponentDetails::CustomizeDetails( IDetailLayoutBuilder& DetailB
 						
 						if (bDisplayMass)
 						{
-							PhysicsCategory.AddCustomRow(TEXT("Mass"), false)
+							PhysicsCategory.AddCustomRow(LOCTEXT("Mass", "Mass"), false)
 								.IsEnabled(TAttribute<bool>(this, &FPrimitiveComponentDetails::IsBodyMassEnabled))
 								.NameContent()
 								[
@@ -324,7 +324,7 @@ void FPrimitiveComponentDetails::AddAdvancedSubCategory( IDetailLayoutBuilder& D
 		IDetailCategoryBuilder& MainCategory = DetailBuilder.EditCategory(MainCategoryName);
 
 		const bool bForAdvanced = true;
-		IDetailGroup& Group = MainCategory.AddGroup( SubCategoryName, SubCategoryName.ToString(), bForAdvanced );
+		IDetailGroup& Group = MainCategory.AddGroup( SubCategoryName, FText::FromName(SubCategoryName), bForAdvanced );
 
 		for( int32 PropertyIndex = 0; PropertyIndex < SubCategoryProperties.Num(); ++PropertyIndex )
 		{
diff --git a/Engine/Source/Editor/DetailCustomizations/Private/RawDistributionVectorStructCustomization.cpp b/Engine/Source/Editor/DetailCustomizations/Private/RawDistributionVectorStructCustomization.cpp
--- a/Engine/Source/Editor/DetailCustomizations/Private/RawDistributionVectorStructCustomization.cpp
+++ b/Engine/Source/Editor/DetailCustomizations/Private/RawDistributionVectorStructCustomization.cpp
@@ -323,7 +323,7 @@ TSharedRef<SWidget> FReplaceVectorWithLinearColorBuilder::CreateColorWidget(cons
 
 void FReplaceVectorWithLinearColorBuilder::AddColorChildProperty(const TSharedPtr<IPropertyHandle>& StructHandle, const FText& Text, int32 ElementIndex, IDetailChildrenBuilder& ChildrenBuilder)
 {
-	ChildrenBuilder.AddChildContent(TEXT("Color"))
+	ChildrenBuilder.AddChildContent(LOCTEXT("Color", "Color"))
 		.NameContent()
 		[
 			SNew(STextBlock)
@@ -490,8 +490,8 @@ TSharedRef<IPropertyTypeCustomization> FRawDistributionVectorStructCustomization
 void FRawDistributionVectorStructCustomization::CustomizeHeader(TSharedRef<IPropertyHandle> StructPropertyHandle, FDetailWidgetRow& HeaderRow, IPropertyTypeCustomizationUtils& StructCustomizationUtils)
 {
 	const bool bDisplayResetToDefault = false;
-	const FString DisplayNameOverride = TEXT("");
-	const FString DisplayToolTipOverride = TEXT("");
+	const FText DisplayNameOverride = FText::GetEmpty();
+	const FText DisplayToolTipOverride = FText::GetEmpty();
 
 	HeaderRow
 	.NameContent()
diff --git a/Engine/Source/Editor/DetailCustomizations/Private/ReflectionCaptureDetails.cpp b/Engine/Source/Editor/DetailCustomizations/Private/ReflectionCaptureDetails.cpp
--- a/Engine/Source/Editor/DetailCustomizations/Private/ReflectionCaptureDetails.cpp
+++ b/Engine/Source/Editor/DetailCustomizations/Private/ReflectionCaptureDetails.cpp
@@ -31,7 +31,7 @@ void FReflectionCaptureDetails::CustomizeDetails( IDetailLayoutBuilder& DetailLa
 	}
 
 	DetailLayout.EditCategory("ReflectionCapture")
-	.AddCustomRow(NSLOCTEXT("ReflectionCaptureDetails", "UpdateReflectionCaptures", "Update Captures").ToString())
+	.AddCustomRow(NSLOCTEXT("ReflectionCaptureDetails", "UpdateReflectionCaptures", "Update Captures"))
 	.NameContent()
 	[
 		SNullWidget::NullWidget
@@ -48,7 +48,7 @@ void FReflectionCaptureDetails::CustomizeDetails( IDetailLayoutBuilder& DetailLa
 			.OnClicked(this, &FReflectionCaptureDetails::OnUpdateReflectionCaptures)
 			[
 				SNew(STextBlock)
-				.Text(NSLOCTEXT("ReflectionCaptureDetails", "UpdateReflectionCaptures", "Update Captures").ToString())
+				.Text(NSLOCTEXT("ReflectionCaptureDetails", "UpdateReflectionCaptures", "Update Captures"))
 				.Font(IDetailLayoutBuilder::GetDetailFont())
 			]
 		]
diff --git a/Engine/Source/Editor/DetailCustomizations/Private/RigDetails.cpp b/Engine/Source/Editor/DetailCustomizations/Private/RigDetails.cpp
--- a/Engine/Source/Editor/DetailCustomizations/Private/RigDetails.cpp
+++ b/Engine/Source/Editor/DetailCustomizations/Private/RigDetails.cpp
@@ -68,7 +68,7 @@ void FRigDetails::CustomizeDetails(IDetailLayoutBuilder& DetailBuilder)
 	// add custom menu
 	// -> set all to world
 	// -> set all to default parent
-	TransformBaseCategory.AddCustomRow(TEXT(""))
+	TransformBaseCategory.AddCustomRow(FText::GetEmpty())
 	[
 		// two button 1. view 2. save to base pose
 		SNew(SHorizontalBox)
@@ -117,7 +117,7 @@ void FRigDetails::GenerateNodeArrayElementWidget(TSharedRef<IPropertyHandle> Pro
 	check (ParentNameProp->GetValueAsDisplayString(ParentNodeName) != FPropertyAccess::Fail);
 	check (DisplayNameProp->GetValueAsDisplayString(DisplayString) != FPropertyAccess::Fail);
 
-	ChildrenBuilder.AddChildContent(TEXT(""))
+	ChildrenBuilder.AddChildContent(FText::GetEmpty())
 	[
 		SNew(SHorizontalBox)
 
@@ -314,7 +314,7 @@ void FRigDetails::GenerateTransformBaseArrayElementWidget(TSharedRef<IPropertyHa
 		NodeIndex++;
 	}
 
-	ChildrenBuilder.AddChildContent(TEXT(""))
+	ChildrenBuilder.AddChildContent(FText::GetEmpty())
 	[
 		SNew(SHorizontalBox)
 
diff --git a/Engine/Source/Editor/DetailCustomizations/Private/SceneCaptureDetails.cpp b/Engine/Source/Editor/DetailCustomizations/Private/SceneCaptureDetails.cpp
--- a/Engine/Source/Editor/DetailCustomizations/Private/SceneCaptureDetails.cpp
+++ b/Engine/Source/Editor/DetailCustomizations/Private/SceneCaptureDetails.cpp
@@ -165,7 +165,7 @@ void FSceneCaptureDetails::CustomizeDetails( IDetailLayoutBuilder& DetailLayout
 			}
 
 			FName GroupFName = FName(*(GroupName.ToString()));
-			IDetailGroup& Group = SceneCaptureCategoryBuilder.AddGroup(GroupFName, GroupName.ToString(), true);
+			IDetailGroup& Group = SceneCaptureCategoryBuilder.AddGroup(GroupFName, GroupName, true);
 
 			// Add each show flag for this group
 			for (FString& FlagName : ShowFlagsByGroup[GroupIndex])
@@ -187,7 +187,7 @@ void FSceneCaptureDetails::CustomizeDetails( IDetailLayoutBuilder& DetailLayout
 						.OnCheckStateChanged(this, &FSceneCaptureDetails::OnShowFlagCheckStateChanged, FlagName)
 						.IsChecked(this, &FSceneCaptureDetails::OnGetDisplayCheckState, FlagName)
 					]
-				.FilterString(FlagName);
+				.FilterString(LocalizedText);
 			}
 		}
 	}
diff --git a/Engine/Source/Editor/DetailCustomizations/Private/SceneComponentDetails.cpp b/Engine/Source/Editor/DetailCustomizations/Private/SceneComponentDetails.cpp
--- a/Engine/Source/Editor/DetailCustomizations/Private/SceneComponentDetails.cpp
+++ b/Engine/Source/Editor/DetailCustomizations/Private/SceneComponentDetails.cpp
@@ -388,7 +388,7 @@ FText FSceneComponentDetails::GetMobilityToolTip(TWeakPtr<IPropertyHandle> Mobil
 {
 	if ( MobilityHandle.IsValid() )
 	{
-		return FText::FromString(MobilityHandle.Pin()->GetToolTipText());
+		return MobilityHandle.Pin()->GetToolTipText();
 	}
 
 	return FText::GetEmpty();
@@ -399,7 +399,7 @@ void FSceneComponentDetails::CustomizeDetails( IDetailLayoutBuilder& DetailBuild
 	MakeTransformDetails( DetailBuilder );
 
 	// Put mobility property in Transform section
-	IDetailCategoryBuilder& TransformCategory = DetailBuilder.EditCategory( "TransformCommon", LOCTEXT("TransformCommonCategory", "Transform").ToString(), ECategoryPriority::Transform );
+	IDetailCategoryBuilder& TransformCategory = DetailBuilder.EditCategory( "TransformCommon", LOCTEXT("TransformCommonCategory", "Transform"), ECategoryPriority::Transform );
 	TSharedPtr<IPropertyHandle> MobilityProperty = DetailBuilder.GetProperty("Mobility");
 
 	uint8 RestrictedMobilityBits = 0u;
@@ -671,7 +671,7 @@ void FSceneComponentDetails::MakeTransformDetails( IDetailLayoutBuilder& DetailB
 			TransformDetails->HideTransformField(ETransformField::Rotation);
 		}
 
-		IDetailCategoryBuilder& TransformCategory = DetailBuilder.EditCategory( "TransformCommon", LOCTEXT("TransformCommonCategory", "Transform").ToString(), ECategoryPriority::Transform );
+		IDetailCategoryBuilder& TransformCategory = DetailBuilder.EditCategory( "TransformCommon", LOCTEXT("TransformCommonCategory", "Transform"), ECategoryPriority::Transform );
 
 		TransformCategory.AddCustomBuilder( TransformDetails );
 	}
diff --git a/Engine/Source/Editor/DetailCustomizations/Private/SkeletalControlNodeDetails.cpp b/Engine/Source/Editor/DetailCustomizations/Private/SkeletalControlNodeDetails.cpp
--- a/Engine/Source/Editor/DetailCustomizations/Private/SkeletalControlNodeDetails.cpp
+++ b/Engine/Source/Editor/DetailCustomizations/Private/SkeletalControlNodeDetails.cpp
@@ -32,7 +32,7 @@ void FSkeletalControlNodeDetails::CustomizeDetails(class IDetailLayoutBuilder& D
 void FSkeletalControlNodeDetails::OnGenerateElementForPropertyPin(TSharedRef<IPropertyHandle> ElementProperty, int32 ElementIndex, IDetailChildrenBuilder& ChildrenBuilder )
 {
 	TSharedPtr<IPropertyHandle> PropertyNameHandle = ElementProperty->GetChildHandle("PropertyFriendlyName");
-	FString PropertyFriendlyName(TEXT("Invalid"));
+	FText PropertyFriendlyName(LOCTEXT("Invalid", "Invalid"));
 	 
 	if (PropertyNameHandle.IsValid())
 	{
@@ -41,14 +41,14 @@ void FSkeletalControlNodeDetails::OnGenerateElementForPropertyPin(TSharedRef<IPr
 		{
 		case FPropertyAccess::Success:
 			{
-				PropertyFriendlyName = DisplayFriendlyName;
+				PropertyFriendlyName = FText::FromString(DisplayFriendlyName);
 
 				//DetailBuilder.EditCategory
 				//DisplayNameAsName
 			}
 			break;
 		case FPropertyAccess::MultipleValues:
-			ChildrenBuilder.AddChildContent( TEXT("") ) 
+			ChildrenBuilder.AddChildContent( FText::GetEmpty() ) 
 			[
 				 SNew(STextBlock).Text(LOCTEXT("OnlyWorksInSingleSelectMode", "Multiple types selected"))
 			];
diff --git a/Engine/Source/Editor/DetailCustomizations/Private/SkeletalMeshComponentDetails.cpp b/Engine/Source/Editor/DetailCustomizations/Private/SkeletalMeshComponentDetails.cpp
--- a/Engine/Source/Editor/DetailCustomizations/Private/SkeletalMeshComponentDetails.cpp
+++ b/Engine/Source/Editor/DetailCustomizations/Private/SkeletalMeshComponentDetails.cpp
@@ -55,12 +55,12 @@ void FSkeletalMeshComponentDetails::CustomizeDetails(IDetailLayoutBuilder& Detai
 	{
 		CurrentDetailBuilder = &DetailBuilder;
 	}
-	DetailBuilder.EditCategory("SkeletalMesh", TEXT(""), ECategoryPriority::TypeSpecific);
-	DetailBuilder.EditCategory("Materials", TEXT(""), ECategoryPriority::TypeSpecific);
-	DetailBuilder.EditCategory("Physics", TEXT(""), ECategoryPriority::TypeSpecific);
+	DetailBuilder.EditCategory("SkeletalMesh", FText::GetEmpty(), ECategoryPriority::TypeSpecific);
+	DetailBuilder.EditCategory("Materials", FText::GetEmpty(), ECategoryPriority::TypeSpecific);
+	DetailBuilder.EditCategory("Physics", FText::GetEmpty(), ECategoryPriority::TypeSpecific);
 	DetailBuilder.HideProperty("bCastStaticShadow", UPrimitiveComponent::StaticClass());
 	DetailBuilder.HideProperty("bLightAsIfStatic", UPrimitiveComponent::StaticClass());
-	DetailBuilder.EditCategory("Animation", TEXT(""), ECategoryPriority::Important);
+	DetailBuilder.EditCategory("Animation", FText::GetEmpty(), ECategoryPriority::Important);
 
 	OnSkeletalMeshPropertyChanged = USkeletalMeshComponent::FOnSkeletalMeshPropertyChanged::CreateSP(this, &FSkeletalMeshComponentDetails::SkeletalMeshPropertyChanged);
 
@@ -114,7 +114,7 @@ void FSkeletalMeshComponentDetails::UpdateAnimationCategory( IDetailLayoutBuilde
 		SelectedSkeletonName = FString::Printf(TEXT("%s'%s'"), *Skeleton->GetClass()->GetName(), *Skeleton->GetPathName());
 	}
 
-	IDetailCategoryBuilder& AnimationCategory = DetailBuilder.EditCategory("Animation", TEXT(""), ECategoryPriority::Important);
+	IDetailCategoryBuilder& AnimationCategory = DetailBuilder.EditCategory("Animation", FText::GetEmpty(), ECategoryPriority::Important);
 
 	// Force the mode switcher to be first
 	AnimationModeHandle = DetailBuilder.GetProperty(GET_MEMBER_NAME_CHECKED(USkeletalMeshComponent, AnimationMode));
diff --git a/Engine/Source/Editor/DetailCustomizations/Private/SkeletonNotifyDetails.cpp b/Engine/Source/Editor/DetailCustomizations/Private/SkeletonNotifyDetails.cpp
--- a/Engine/Source/Editor/DetailCustomizations/Private/SkeletonNotifyDetails.cpp
+++ b/Engine/Source/Editor/DetailCustomizations/Private/SkeletonNotifyDetails.cpp
@@ -12,10 +12,10 @@ TSharedRef<IDetailCustomization> FSkeletonNotifyDetails::MakeInstance()
 
 void FSkeletonNotifyDetails::CustomizeDetails( IDetailLayoutBuilder& DetailBuilder )
 {
-	IDetailCategoryBuilder& Category = DetailBuilder.EditCategory("Skeleton Notify", TEXT("Skeleton Notify") );
+	IDetailCategoryBuilder& Category = DetailBuilder.EditCategory("Skeleton Notify", LOCTEXT("SkeletonNotifyCategoryName", "Skeleton Notify") );
 	const FSlateFontInfo DetailFontInfo = IDetailLayoutBuilder::GetDetailFont();
 
-	Category.AddProperty("Name").DisplayName( TEXT("Notify Name") );
+	Category.AddProperty("Name").DisplayName( LOCTEXT("SkeletonNotifyName", "Notify Name") );
 
 	TSharedPtr<IPropertyHandle> InPropertyHandle = DetailBuilder.GetProperty("AnimationNames");
 
@@ -34,7 +34,7 @@ void FSkeletonNotifyDetails::CustomizeDetails( IDetailLayoutBuilder& DetailBuild
 
 	if(EdObj)
 	{
-		Category.AddCustomRow(TEXT("Animations"))
+		Category.AddCustomRow(LOCTEXT("AnimationsLabel","Animations"))
 		.NameContent()
 		[
 			SNew(STextBlock)
diff --git a/Engine/Source/Editor/DetailCustomizations/Private/SkinnedMeshComponentDetails.cpp b/Engine/Source/Editor/DetailCustomizations/Private/SkinnedMeshComponentDetails.cpp
--- a/Engine/Source/Editor/DetailCustomizations/Private/SkinnedMeshComponentDetails.cpp
+++ b/Engine/Source/Editor/DetailCustomizations/Private/SkinnedMeshComponentDetails.cpp
@@ -14,16 +14,16 @@ TSharedRef<IDetailCustomization> FSkinnedMeshComponentDetails::MakeInstance()
 
 void FSkinnedMeshComponentDetails::CustomizeDetails(IDetailLayoutBuilder& DetailBuilder)
 {
-	DetailBuilder.EditCategory("Mesh", TEXT(""), ECategoryPriority::Important);
-	IDetailCategoryBuilder& PhysicsCategory = DetailBuilder.EditCategory("Physics", TEXT(""), ECategoryPriority::TypeSpecific);
-	IDetailCategoryBuilder& LODCategory = DetailBuilder.EditCategory("LOD", TEXT("Level of Detail"), ECategoryPriority::Default);
+	DetailBuilder.EditCategory("Mesh", FText::GetEmpty(), ECategoryPriority::Important);
+	IDetailCategoryBuilder& PhysicsCategory = DetailBuilder.EditCategory("Physics", FText::GetEmpty(), ECategoryPriority::TypeSpecific);
+	IDetailCategoryBuilder& LODCategory = DetailBuilder.EditCategory("LOD", LOCTEXT("LODCategoryName", "Level of Detail"), ECategoryPriority::Default);
 	
 	// show extra field about actually used physics asset, but make sure to show it under physics asset override
 	TSharedRef<IPropertyHandle> PhysicsAssetProperty =  DetailBuilder.GetProperty(GET_MEMBER_NAME_CHECKED(USkinnedMeshComponent, PhysicsAssetOverride));
 	if (PhysicsAssetProperty->IsValidHandle())
 	{
 		PhysicsCategory.AddProperty(GET_MEMBER_NAME_CHECKED(USkinnedMeshComponent, PhysicsAssetOverride));
-		CreateActuallyUsedPhysicsAssetWidget(PhysicsCategory.AddCustomRow( LOCTEXT("CurrentPhysicsAsset", "Currently used Physics Asset").ToString(), true), &DetailBuilder);
+		CreateActuallyUsedPhysicsAssetWidget(PhysicsCategory.AddCustomRow( LOCTEXT("CurrentPhysicsAsset", "Currently used Physics Asset"), true), &DetailBuilder);
 	}
 }
 
diff --git a/Engine/Source/Editor/DetailCustomizations/Private/SkyLightComponentDetails.cpp b/Engine/Source/Editor/DetailCustomizations/Private/SkyLightComponentDetails.cpp
--- a/Engine/Source/Editor/DetailCustomizations/Private/SkyLightComponentDetails.cpp
+++ b/Engine/Source/Editor/DetailCustomizations/Private/SkyLightComponentDetails.cpp
@@ -19,7 +19,7 @@ void FSkyLightComponentDetails::CustomizeDetails( IDetailLayoutBuilder& DetailLa
 	// Mobility property is on the scene component base class not the light component and that is why we have to use USceneComponent::StaticClass
 	TSharedRef<IPropertyHandle> MobilityHandle = DetailLayout.GetProperty(GET_MEMBER_NAME_CHECKED(USkyLightComponent, Mobility), USceneComponent::StaticClass());
 	// Set a mobility tooltip specific to lights
-	MobilityHandle->SetToolTipText(LOCTEXT("SkyLightMobilityTooltip", "Mobility for sky light components determines what rendering methods will be used.  A Stationary sky light has its shadowing baked into Bent Normal AO by Lightmass, but its lighting can be changed in game.").ToString());
+	MobilityHandle->SetToolTipText(LOCTEXT("SkyLightMobilityTooltip", "Mobility for sky light components determines what rendering methods will be used.  A Stationary sky light has its shadowing baked into Bent Normal AO by Lightmass, but its lighting can be changed in game."));
 
 	TSharedPtr<IPropertyHandle> LightIntensityProperty = DetailLayout.GetProperty(GET_MEMBER_NAME_CHECKED(USkyLightComponent, Intensity), ULightComponentBase::StaticClass());
 	// Point lights need to override the ui min and max for units of lumens, so we have to undo that
@@ -43,7 +43,7 @@ void FSkyLightComponentDetails::CustomizeDetails( IDetailLayoutBuilder& DetailLa
 	}
 
 	DetailLayout.EditCategory( "SkyLight" )
-	.AddCustomRow( NSLOCTEXT("SkyLightDetails", "UpdateSkyLight", "Recapture Scene").ToString() )
+	.AddCustomRow( NSLOCTEXT("SkyLightDetails", "UpdateSkyLight", "Recapture Scene") )
 	[
 		SNew(SHorizontalBox)
 		+SHorizontalBox::Slot()
@@ -55,7 +55,7 @@ void FSkyLightComponentDetails::CustomizeDetails( IDetailLayoutBuilder& DetailLa
 			.VAlign(VAlign_Center)
 			.HAlign(HAlign_Center)
 			.OnClicked( this, &FSkyLightComponentDetails::OnUpdateSkyCapture )
-			.Text( NSLOCTEXT("SkyLightDetails", "UpdateSkyCapture", "Recapture Scene").ToString() )
+			.Text( NSLOCTEXT("SkyLightDetails", "UpdateSkyCapture", "Recapture Scene") )
 		]
 	];
 }
diff --git a/Engine/Source/Editor/DetailCustomizations/Private/SlateBrushCustomization.cpp b/Engine/Source/Editor/DetailCustomizations/Private/SlateBrushCustomization.cpp
--- a/Engine/Source/Editor/DetailCustomizations/Private/SlateBrushCustomization.cpp
+++ b/Engine/Source/Editor/DetailCustomizations/Private/SlateBrushCustomization.cpp
@@ -1208,7 +1208,7 @@ void FSlateBrushStructCustomization::CustomizeChildren( TSharedRef<IPropertyHand
 	TSharedPtr<IPropertyHandle> TintProperty = StructPropertyHandle->GetChildHandle( TEXT("TintColor") );
 	TSharedPtr<IPropertyHandle> ResourceObjectProperty = StructPropertyHandle->GetChildHandle( TEXT("ResourceObject") );
 	
-	StructBuilder.AddChildContent( NSLOCTEXT( "SlateBrushCustomization", "ResourceObjectFilterString", "Resource" ).ToString() )
+	StructBuilder.AddChildContent( NSLOCTEXT( "SlateBrushCustomization", "ResourceObjectFilterString", "Resource" ) )
 	.NameContent()
 	[
 		ResourceObjectProperty->CreatePropertyNameWidget()
@@ -1248,13 +1248,13 @@ void FSlateBrushStructCustomization::CustomizeChildren( TSharedRef<IPropertyHand
 				.ResourceObjectProperty(ResourceObjectProperty)
 				.SlateBrush(Brush);
 
-			IDetailGroup& PreviewGroup = StructBuilder.AddChildGroup(TEXT("Preview"), TEXT(""));
+			IDetailGroup& PreviewGroup = StructBuilder.AddChildGroup(TEXT("Preview"), FText::GetEmpty());
 
 			PreviewGroup
 				.HeaderRow()
 				.NameContent()
 				[
-					StructPropertyHandle->CreatePropertyNameWidget(NSLOCTEXT("UnrealEd", "Preview", "Preview").ToString(), TEXT( "" ), false)
+					StructPropertyHandle->CreatePropertyNameWidget(NSLOCTEXT("UnrealEd", "Preview", "Preview"), FText::GetEmpty(), false)
 				]
 			.ValueContent()
 				.MinDesiredWidth(1)
diff --git a/Engine/Source/Editor/DetailCustomizations/Private/SourceCodeAccessSettingsDetails.cpp b/Engine/Source/Editor/DetailCustomizations/Private/SourceCodeAccessSettingsDetails.cpp
--- a/Engine/Source/Editor/DetailCustomizations/Private/SourceCodeAccessSettingsDetails.cpp
+++ b/Engine/Source/Editor/DetailCustomizations/Private/SourceCodeAccessSettingsDetails.cpp
@@ -34,7 +34,7 @@ void FSourceCodeAccessSettingsDetails::CustomizeDetails( IDetailLayoutBuilder& D
 	}
 
 	IDetailCategoryBuilder& AccessorCategory = DetailLayout.EditCategory( "Accessor" );
-	AccessorCategory.AddCustomRow( LOCTEXT("PreferredAccessor", "Preferred Accessor").ToString() )
+	AccessorCategory.AddCustomRow( LOCTEXT("PreferredAccessor", "Preferred Accessor") )
 	.NameContent()
 	[
 		PreferredProviderPropertyHandle->CreatePropertyNameWidget()
diff --git a/Engine/Source/Editor/DetailCustomizations/Private/StaticMeshActorDetails.cpp b/Engine/Source/Editor/DetailCustomizations/Private/StaticMeshActorDetails.cpp
--- a/Engine/Source/Editor/DetailCustomizations/Private/StaticMeshActorDetails.cpp
+++ b/Engine/Source/Editor/DetailCustomizations/Private/StaticMeshActorDetails.cpp
@@ -42,7 +42,7 @@ void FStaticMeshActorDetails::CustomizeDetails( IDetailLayoutBuilder& DetailBuil
 	// The blocking volume menu is advanced
 	const bool bForAdvanced = true;
 
-	const FString CreateBlockingVolumeString = LOCTEXT("BlockingVolumeMenu", "Create Blocking Volume").ToString();
+	const FText CreateBlockingVolumeString = LOCTEXT("BlockingVolumeMenu", "Create Blocking Volume");
 
 	StaticMeshCategory.AddCustomRow( CreateBlockingVolumeString, bForAdvanced )
 	.NameContent()
diff --git a/Engine/Source/Editor/DetailCustomizations/Private/StaticMeshComponentDetails.cpp b/Engine/Source/Editor/DetailCustomizations/Private/StaticMeshComponentDetails.cpp
--- a/Engine/Source/Editor/DetailCustomizations/Private/StaticMeshComponentDetails.cpp
+++ b/Engine/Source/Editor/DetailCustomizations/Private/StaticMeshComponentDetails.cpp
@@ -14,7 +14,7 @@ TSharedRef<IDetailCustomization> FStaticMeshComponentDetails::MakeInstance()
 void FStaticMeshComponentDetails::CustomizeDetails( IDetailLayoutBuilder& DetailBuilder )
 {
 	// Create a category so this is displayed early in the properties
-	DetailBuilder.EditCategory( "StaticMesh", TEXT(""), ECategoryPriority::Important);
+	DetailBuilder.EditCategory( "StaticMesh", FText::GetEmpty(), ECategoryPriority::Important);
 }
 
 
diff --git a/Engine/Source/Editor/DetailCustomizations/Private/TransitionPoseEvaluatorNodeDetails.cpp b/Engine/Source/Editor/DetailCustomizations/Private/TransitionPoseEvaluatorNodeDetails.cpp
--- a/Engine/Source/Editor/DetailCustomizations/Private/TransitionPoseEvaluatorNodeDetails.cpp
+++ b/Engine/Source/Editor/DetailCustomizations/Private/TransitionPoseEvaluatorNodeDetails.cpp
@@ -31,7 +31,7 @@ void FTransitionPoseEvaluatorNodeDetails::CustomizeDetails( IDetailLayoutBuilder
 		}
 	}
 
-	IDetailCategoryBuilder& PoseCategory = DetailBuilder.EditCategory("Pose", TEXT("Pose") );
+	IDetailCategoryBuilder& PoseCategory = DetailBuilder.EditCategory("Pose", LOCTEXT("PoseCategoryName", "Pose") );
 	TSharedPtr<IPropertyHandle> FramesToCachePosePropety = DetailBuilder.GetProperty(TEXT("Node.FramesToCachePose"));
 
 	//@TODO: CONDUIT: try both
diff --git a/Engine/Source/Editor/DetailCustomizations/Private/VehicleTransmissionDataCustomization.cpp b/Engine/Source/Editor/DetailCustomizations/Private/VehicleTransmissionDataCustomization.cpp
--- a/Engine/Source/Editor/DetailCustomizations/Private/VehicleTransmissionDataCustomization.cpp
+++ b/Engine/Source/Editor/DetailCustomizations/Private/VehicleTransmissionDataCustomization.cpp
@@ -164,7 +164,7 @@ void FVehicleTransmissionDataCustomization::BuildColumnsHeaderHelper(TSharedRef<
 void FVehicleTransmissionDataCustomization::CreateGearUIDelegate(TSharedRef<IPropertyHandle> GearProperty, int32 GearIdx, IDetailChildrenBuilder& ChildrenBuilder)
 {
 	FText Label = FText::Format(LOCTEXT("TransmissionGear", "Gear {0}"), FText::AsNumber(GearIdx + 1));
-	CreateGearUIHelper(ChildrenBuilder.AddChildContent(Label.ToString()), Label, GearProperty, ForwardGear);
+	CreateGearUIHelper(ChildrenBuilder.AddChildContent(Label), Label, GearProperty, ForwardGear);
 }
 
 void FVehicleTransmissionDataCustomization::CustomizeChildren(TSharedRef<class IPropertyHandle> StructPropertyHandle, class IDetailChildrenBuilder& StructBuilder, IPropertyTypeCustomizationUtils& StructCustomizationUtils)
@@ -187,7 +187,7 @@ void FVehicleTransmissionDataCustomization::CustomizeChildren(TSharedRef<class I
 	StructPropertyHandle->GetNumChildren(NumChildren);
 
 	FName GearsSetupGroupName = TEXT("GearsSetup");
-	FString GearsSetupGroupLabel = LOCTEXT("GearSetupLabel", "Gears Setup").ToString();
+	FText GearsSetupGroupLabel = LOCTEXT("GearSetupLabel", "Gears Setup");
 	IDetailGroup * GearSetupGroup = NULL;
 	
 	for (uint32 ChildIdx = 0; ChildIdx < NumChildren; ++ChildIdx)
@@ -200,7 +200,7 @@ void FVehicleTransmissionDataCustomization::CustomizeChildren(TSharedRef<class I
 			if (GearSetupGroup == NULL)
 			{
 				GearSetupGroup = &StructBuilder.AddChildGroup(GearsSetupGroupName, GearsSetupGroupLabel);
-				BuildColumnsHeaderHelper(StructPropertyHandle, StructBuilder.AddChildContent(TEXT("GearSetup")));
+				BuildColumnsHeaderHelper(StructPropertyHandle, StructBuilder.AddChildContent(GearsSetupGroupLabel));
 			}
 
 			//determine which gear we're showing
@@ -222,8 +222,9 @@ void FVehicleTransmissionDataCustomization::CustomizeChildren(TSharedRef<class I
 				StructBuilder.AddChildCustomBuilder(GearsArrayBuilder);
 			}
 			else
-			{	
-				CreateGearUIHelper(StructBuilder.AddChildContent(PropertyName), FText::FromString(PropertyName), ChildProperty, GearType);
+			{
+				const FText PropertyNameText = FText::FromString(PropertyName);
+				CreateGearUIHelper(StructBuilder.AddChildContent(PropertyNameText), PropertyNameText, ChildProperty, GearType);
 			}
 
 
diff --git a/Engine/Source/Editor/DetailCustomizations/Private/WindowsTargetSettingsDetails.cpp b/Engine/Source/Editor/DetailCustomizations/Private/WindowsTargetSettingsDetails.cpp
--- a/Engine/Source/Editor/DetailCustomizations/Private/WindowsTargetSettingsDetails.cpp
+++ b/Engine/Source/Editor/DetailCustomizations/Private/WindowsTargetSettingsDetails.cpp
@@ -117,10 +117,10 @@ void FWindowsTargetSettingsDetails::CustomizeDetails( IDetailLayoutBuilder& Deta
 	TargetShaderFormatsDetails->CreateTargetShaderFormatsPropertyView();
 
 	// Next add the splash image customization
+	const FText EditorSplashDesc(LOCTEXT("EditorSplashLabel", "Editor Splash"));
 	IDetailCategoryBuilder& SplashCategoryBuilder = DetailBuilder.EditCategory(TEXT("Splash"));
-	FDetailWidgetRow& EditorSplashWidgetRow = SplashCategoryBuilder.AddCustomRow(TEXT("Editor Splash"));
+	FDetailWidgetRow& EditorSplashWidgetRow = SplashCategoryBuilder.AddCustomRow(EditorSplashDesc);
 
-	const FText EditorSplashDesc(LOCTEXT("EditorSplashLabel", "Editor Splash"));
 	const FString EditorSplash_TargetImagePath = GetSplashFilename(EImageScope::GameOverride, true);
 	const FString EditorSplash_DefaultImagePath = GetSplashFilename(EImageScope::Engine, true);
 
@@ -153,9 +153,8 @@ void FWindowsTargetSettingsDetails::CustomizeDetails( IDetailLayoutBuilder& Deta
 		]
 	];
 
-	FDetailWidgetRow& GameSplashWidgetRow = SplashCategoryBuilder.AddCustomRow(TEXT("Game Splash"));
-
 	const FText GameSplashDesc(LOCTEXT("GameSplashLabel", "Game Splash"));
+	FDetailWidgetRow& GameSplashWidgetRow = SplashCategoryBuilder.AddCustomRow(GameSplashDesc);
 	const FString GameSplash_TargetImagePath = GetSplashFilename(EImageScope::GameOverride, false);
 	const FString GameSplash_DefaultImagePath = GetSplashFilename(EImageScope::Engine, false);
 
@@ -189,7 +188,7 @@ void FWindowsTargetSettingsDetails::CustomizeDetails( IDetailLayoutBuilder& Deta
 	];
 
 	IDetailCategoryBuilder& IconsCategoryBuilder = DetailBuilder.EditCategory(TEXT("Icon"));	
-	FDetailWidgetRow& GameIconWidgetRow = IconsCategoryBuilder.AddCustomRow(TEXT("Game Icon"));
+	FDetailWidgetRow& GameIconWidgetRow = IconsCategoryBuilder.AddCustomRow(LOCTEXT("GameIconLabel", "Game Icon"));
 	GameIconWidgetRow
 	.NameContent()
 	[
@@ -260,10 +259,10 @@ void FTargetShaderFormatsPropertyDetails::CreateTargetShaderFormatsPropertyView(
 
 	for (const FName& ShaderFormat : ShaderFormats)
 	{
-		FDetailWidgetRow& TargetedRHIWidgetRow = TargetedRHICategoryBuilder.AddCustomRow(ShaderFormat.ToString());
-
 		FText FriendlyShaderFormatName = GetFriendlyNameFromRHIName(ShaderFormat.ToString());
 
+		FDetailWidgetRow& TargetedRHIWidgetRow = TargetedRHICategoryBuilder.AddCustomRow(FriendlyShaderFormatName);
+
 		TargetedRHIWidgetRow
 		.NameContent()
 		[
diff --git a/Engine/Source/Editor/DetailCustomizations/Private/WorldSettingsDetails.cpp b/Engine/Source/Editor/DetailCustomizations/Private/WorldSettingsDetails.cpp
--- a/Engine/Source/Editor/DetailCustomizations/Private/WorldSettingsDetails.cpp
+++ b/Engine/Source/Editor/DetailCustomizations/Private/WorldSettingsDetails.cpp
@@ -96,7 +96,7 @@ void FLightmapCustomNodeBuilder::GenerateChildContent(IDetailChildrenBuilder& Ch
 {
 	RefreshLightmapItems();
 
-	ChildrenBuilder.AddChildContent(LOCTEXT("LightMapsFilter", "Lightmaps").ToString())
+	ChildrenBuilder.AddChildContent(LOCTEXT("LightMapsFilter", "Lightmaps"))
 	.ValueContent()
 	.HAlign(HAlign_Fill)
 	[
diff --git a/Engine/Source/Editor/EnvironmentQueryEditor/Private/DetailCustomizations/EnvQueryTestDetails.cpp b/Engine/Source/Editor/EnvironmentQueryEditor/Private/DetailCustomizations/EnvQueryTestDetails.cpp
--- a/Engine/Source/Editor/EnvironmentQueryEditor/Private/DetailCustomizations/EnvQueryTestDetails.cpp
+++ b/Engine/Source/Editor/EnvironmentQueryEditor/Private/DetailCustomizations/EnvQueryTestDetails.cpp
@@ -121,14 +121,14 @@ void FEnvQueryTestDetails::CustomizeDetails( IDetailLayoutBuilder& DetailLayout
 	IDetailPropertyRow& BoolFilterRow = FilterCategory.AddProperty(DetailLayout.GetProperty(GET_MEMBER_NAME_CHECKED(UEnvQueryTest, BoolFilter)));
 	BoolFilterRow.Visibility(TAttribute<EVisibility>::Create(TAttribute<EVisibility>::FGetter::CreateSP(this, &FEnvQueryTestDetails::GetBoolFilterVisibility)));
 
-	IDetailGroup& HackToEnsureFilterCategoryIsVisible = FilterCategory.AddGroup("HackForVisibility", FString(TEXT(" ")));
+	IDetailGroup& HackToEnsureFilterCategoryIsVisible = FilterCategory.AddGroup("HackForVisibility", FText::GetEmpty());
 
 	// Scoring
 	IDetailCategoryBuilder& ScoreCategory = DetailLayout.EditCategory("Score");
 
 	//----------------------------
 	// BEGIN Scoring: Clamping
-	IDetailGroup& ClampingGroup = ScoreCategory.AddGroup("Clamping", FString(TEXT("Clamping")));
+	IDetailGroup& ClampingGroup = ScoreCategory.AddGroup("Clamping", LOCTEXT("ClampingLabel", "Clamping"));
 	
 	// Drop-downs for setting type of lower and upper bound normalization
 	IDetailPropertyRow& ClampMinTypeRow = ClampingGroup.AddPropertyRow(ClampMinTypeHandle.ToSharedRef());
@@ -160,7 +160,7 @@ void FEnvQueryTestDetails::CustomizeDetails( IDetailLayoutBuilder& DetailLayout
 	{
 		IDetailPropertyRow& FloatFilterMinForClampingRow = ClampingGroup.AddPropertyRow(FloatFilterMinHandle.ToSharedRef());
 		FloatFilterMinForClampingRow.Visibility(TAttribute<EVisibility>::Create(TAttribute<EVisibility>::FGetter::CreateSP(this, &FEnvQueryTestDetails::GetVisibilityOfFilterMinForScoreClamping)));
-		FloatFilterMinForClampingRow.ToolTip(TEXT("See Filter Thresholds under the Filter tab.  Values lower than this (before clamping) cause the item to be thrown out as invalid.  Values are normalized with this value as the minimum, so items with this value will have a normalized score of 0."));
+		FloatFilterMinForClampingRow.ToolTip(LOCTEXT("FloatFilterMinForClampingRowToolTip", "See Filter Thresholds under the Filter tab.  Values lower than this (before clamping) cause the item to be thrown out as invalid.  Values are normalized with this value as the minimum, so items with this value will have a normalized score of 0."));
 		FloatFilterMinForClampingRow.EditCondition(TAttribute<bool>(this, &FEnvQueryTestDetails::AllowWritingToFiltersFromScore), NULL);
 	}
 
@@ -199,7 +199,7 @@ void FEnvQueryTestDetails::CustomizeDetails( IDetailLayoutBuilder& DetailLayout
 		// Upper Bound for scoring when tied to filter maximum.
 		IDetailPropertyRow& FloatFilterMaxForClampingRow = ClampingGroup.AddPropertyRow(FloatFilterMaxHandle.ToSharedRef());
 		FloatFilterMaxForClampingRow.Visibility(TAttribute<EVisibility>::Create(TAttribute<EVisibility>::FGetter::CreateSP(this, &FEnvQueryTestDetails::GetVisibilityOfFilterMaxForScoreClamping)));
-		FloatFilterMaxForClampingRow.ToolTip(TEXT("See Filter Thresholds under the Filter tab.  Values higher than this (before normalization) cause the item to be thrown out as invalid.  Values are normalized with this value as the maximum, so items with this value will have a normalized score of 1."));
+		FloatFilterMaxForClampingRow.ToolTip(LOCTEXT("FloatFilterMaxForClampingRowToolTip", "See Filter Thresholds under the Filter tab.  Values higher than this (before normalization) cause the item to be thrown out as invalid.  Values are normalized with this value as the maximum, so items with this value will have a normalized score of 1."));
 		FloatFilterMaxForClampingRow.EditCondition(TAttribute<bool>(this, &FEnvQueryTestDetails::AllowWritingToFiltersFromScore), NULL);
 	}
 	// END Scoring: Clamping, continue Scoring
@@ -207,8 +207,8 @@ void FEnvQueryTestDetails::CustomizeDetails( IDetailLayoutBuilder& DetailLayout
 
 	IDetailPropertyRow& BoolScoreTestRow = ScoreCategory.AddProperty(DetailLayout.GetProperty(GET_MEMBER_NAME_CHECKED(UEnvQueryTest, BoolFilter)));
 	BoolScoreTestRow.Visibility(TAttribute<EVisibility>::Create(TAttribute<EVisibility>::FGetter::CreateSP(this, &FEnvQueryTestDetails::GetBoolFilterVisibilityForScoring)));
-	BoolScoreTestRow.DisplayName(TEXT("Bool Match"));
-	BoolScoreTestRow.ToolTip(TEXT("Boolean value to match in order to grant score of 'Weight'.  Not matching this value will not change score."));
+	BoolScoreTestRow.DisplayName(LOCTEXT("BoolMatchLabel", "Bool Match"));
+	BoolScoreTestRow.ToolTip(LOCTEXT("BoolMatchToolTip", "Boolean value to match in order to grant score of 'Weight'.  Not matching this value will not change score."));
 
 // 	IDetailPropertyRow& ScoreMirrorNormalizedScoreRow = ScoreCategory.AddProperty(DetailLayout.GetProperty(GET_MEMBER_NAME_CHECKED(UEnvQueryTest, bMirrorNormalizedScore)));
 // 	ScoreMirrorNormalizedScoreRow.Visibility(TAttribute<EVisibility>::Create(TAttribute<EVisibility>::FGetter::CreateSP(this, &FEnvQueryTestDetails::GetFloatScoreVisibility)));
@@ -254,7 +254,7 @@ void FEnvQueryTestDetails::CustomizeDetails( IDetailLayoutBuilder& DetailLayout
 
 	// scoring & filter function preview
 	IDetailCategoryBuilder& PreviewCategory = DetailLayout.EditCategory("Preview");
-	PreviewCategory.AddCustomRow("Preview").WholeRowWidget
+	PreviewCategory.AddCustomRow(LOCTEXT("Preview", "Preview")).WholeRowWidget
 	[
 		SAssignNew(PreviewWidget, STestFunctionWidget)
 	];
diff --git a/Engine/Source/Editor/EnvironmentQueryEditor/Private/DetailCustomizations/EnvTraceDataCustomization.cpp b/Engine/Source/Editor/EnvironmentQueryEditor/Private/DetailCustomizations/EnvTraceDataCustomization.cpp
--- a/Engine/Source/Editor/EnvironmentQueryEditor/Private/DetailCustomizations/EnvTraceDataCustomization.cpp
+++ b/Engine/Source/Editor/EnvironmentQueryEditor/Private/DetailCustomizations/EnvTraceDataCustomization.cpp
@@ -35,7 +35,7 @@ void FEnvTraceDataCustomization::CustomizeChildren( TSharedRef<class IPropertyHa
 {
 	if (TraceModes.Num() > 1)
 	{
-		StructBuilder.AddChildContent("TraceMode")
+		StructBuilder.AddChildContent(LOCTEXT("TraceMode", "Trace Mode"))
 		.NameContent()
 		[
 			PropTraceMode->CreatePropertyNameWidget()
diff --git a/Engine/Source/Editor/GameplayAbilitiesEditor/Private/AttributeDetails.cpp b/Engine/Source/Editor/GameplayAbilitiesEditor/Private/AttributeDetails.cpp
--- a/Engine/Source/Editor/GameplayAbilitiesEditor/Private/AttributeDetails.cpp
+++ b/Engine/Source/Editor/GameplayAbilitiesEditor/Private/AttributeDetails.cpp
@@ -202,10 +202,10 @@ void FAttributeDetails::CustomizeDetails(IDetailLayoutBuilder& DetailLayout)
 		PropertyOptions.Add(MakeShareable(new FString(Property->GetName())));
 	}
 
-	IDetailCategoryBuilder& Category = DetailLayout.EditCategory("Variable", LOCTEXT("VariableDetailsCategory", "Variable").ToString());
+	IDetailCategoryBuilder& Category = DetailLayout.EditCategory("Variable", LOCTEXT("VariableDetailsCategory", "Variable"));
 	const FSlateFontInfo DetailFontInfo = IDetailLayoutBuilder::GetDetailFont();
 	
-	Category.AddCustomRow( TEXT("Replication") )
+	Category.AddCustomRow( LOCTEXT("ReplicationLabel", "Replication") )
 		//.Visibility(TAttribute<EVisibility>(this, &FBlueprintVarActionDetails::ReplicationVisibility))
 		.NameContent()
 		[
diff --git a/Engine/Source/Editor/GameplayAbilitiesEditor/Private/GameplayEffectExecutionScopedModifierInfoDetails.cpp b/Engine/Source/Editor/GameplayAbilitiesEditor/Private/GameplayEffectExecutionScopedModifierInfoDetails.cpp
--- a/Engine/Source/Editor/GameplayAbilitiesEditor/Private/GameplayEffectExecutionScopedModifierInfoDetails.cpp
+++ b/Engine/Source/Editor/GameplayAbilitiesEditor/Private/GameplayEffectExecutionScopedModifierInfoDetails.cpp
@@ -192,7 +192,7 @@ void FGameplayEffectExecutionScopedModifierInfoDetails::CustomizeChildren(TShare
 		{
 			TSharedPtr< SComboBox< TSharedPtr<FGameplayEffectAttributeCaptureDefinition> > > BackingComboBox;
 
-			StructBuilder.AddChildContent(TEXT("CaptureDef"))
+			StructBuilder.AddChildContent(NSLOCTEXT("ScopedModifierDetails", "CaptureDefLabel", "Backing Capture Definition"))
 			.NameContent()
 			[
 				SNew(STextBlock)
diff --git a/Engine/Source/Editor/GameplayTagsEditor/Private/GameplayTagContainerCustomization.cpp b/Engine/Source/Editor/GameplayTagsEditor/Private/GameplayTagContainerCustomization.cpp
--- a/Engine/Source/Editor/GameplayTagsEditor/Private/GameplayTagContainerCustomization.cpp
+++ b/Engine/Source/Editor/GameplayTagsEditor/Private/GameplayTagContainerCustomization.cpp
@@ -123,7 +123,7 @@ FReply FGameplayTagContainerCustomization::OnEditButtonClicked()
 
 	FText Title;
 	FText AssetName;
-	FText PropertyName = FText::FromString( StructPropertyHandle->GetPropertyDisplayName() );
+	FText PropertyName = StructPropertyHandle->GetPropertyDisplayName();
 
 	if (OuterObjects.Num() > 1)
 	{
@@ -144,7 +144,7 @@ FReply FGameplayTagContainerCustomization::OnEditButtonClicked()
 		.Filter( Categories )
 		.OnTagChanged(this, &FGameplayTagContainerCustomization::RefreshTagList)
 		.ReadOnly(bReadOnly)
-		.TagContainerName( StructPropertyHandle->GetPropertyDisplayName() )
+		.TagContainerName( StructPropertyHandle->GetPropertyDisplayName().ToString() )
 		.PropertyHandle( StructPropertyHandle )
 	];
 
diff --git a/Engine/Source/Editor/GameplayTagsEditor/Private/GameplayTagCustomization.cpp b/Engine/Source/Editor/GameplayTagsEditor/Private/GameplayTagCustomization.cpp
--- a/Engine/Source/Editor/GameplayTagsEditor/Private/GameplayTagCustomization.cpp
+++ b/Engine/Source/Editor/GameplayTagsEditor/Private/GameplayTagCustomization.cpp
@@ -74,7 +74,7 @@ TSharedRef<SWidget> FGameplayTagCustomization::GetListContent()
 			SNew(SGameplayTagWidget, EditableContainers)
 			.Filter(Categories)
 			.ReadOnly(bReadOnly)
-			.TagContainerName(StructPropertyHandle->GetPropertyDisplayName())
+			.TagContainerName(StructPropertyHandle->GetPropertyDisplayName().ToString())
 			.MultiSelect(false)
 			.OnTagChanged(this, &FGameplayTagCustomization::OnTagChanged)
 		];
diff --git a/Engine/Source/Editor/IntroTutorials/Private/EditorTutorialDetailsCustomization.cpp b/Engine/Source/Editor/IntroTutorials/Private/EditorTutorialDetailsCustomization.cpp
--- a/Engine/Source/Editor/IntroTutorials/Private/EditorTutorialDetailsCustomization.cpp
+++ b/Engine/Source/Editor/IntroTutorials/Private/EditorTutorialDetailsCustomization.cpp
@@ -31,9 +31,9 @@ void FEditorTutorialDetailsCustomization::CustomizeDetails( IDetailLayoutBuilder
 	check(Objects.Num() > 0);
 	UEditorTutorial* Tutorial = CastChecked<UEditorTutorial>(Objects[0].Get());
 
-	IDetailCategoryBuilder& CategoryBuilder = DetailLayout.EditCategory(TEXT("Testing"), LOCTEXT("TestingSection", "Testing").ToString(), ECategoryPriority::Important);
+	IDetailCategoryBuilder& CategoryBuilder = DetailLayout.EditCategory(TEXT("Testing"), LOCTEXT("TestingSection", "Testing"), ECategoryPriority::Important);
 
-	CategoryBuilder.AddCustomRow(TEXT("Launch"))
+	CategoryBuilder.AddCustomRow(LOCTEXT("LaunchButtonLabel", "Launch"))
 	.WholeRowContent()
 	.VAlign(VAlign_Center)
 	.HAlign(HAlign_Left)
diff --git a/Engine/Source/Editor/Kismet/Private/BlueprintDetailsCustomization.cpp b/Engine/Source/Editor/Kismet/Private/BlueprintDetailsCustomization.cpp
--- a/Engine/Source/Editor/Kismet/Private/BlueprintDetailsCustomization.cpp
+++ b/Engine/Source/Editor/Kismet/Private/BlueprintDetailsCustomization.cpp
@@ -56,14 +56,14 @@ void FBlueprintVarActionDetails::CustomizeDetails( IDetailLayoutBuilder& DetailL
 	FName VarName =	GetVariableName();
 	const UEdGraphSchema_K2* Schema = GetDefault<UEdGraphSchema_K2>();
 	
-	IDetailCategoryBuilder& Category = DetailLayout.EditCategory("Variable", LOCTEXT("VariableDetailsCategory", "Variable").ToString());
+	IDetailCategoryBuilder& Category = DetailLayout.EditCategory("Variable", LOCTEXT("VariableDetailsCategory", "Variable"));
 	const FSlateFontInfo DetailFontInfo = IDetailLayoutBuilder::GetDetailFont();
 	
 	const FString DocLink = TEXT("Shared/Editors/BlueprintEditor/VariableDetails");
 
 	TSharedPtr<SToolTip> VarNameTooltip = IDocumentation::Get()->CreateToolTip(LOCTEXT("VarNameTooltip", "The name of the variable."), NULL, DocLink, TEXT("VariableName"));
 
-	Category.AddCustomRow( TEXT("Variable Name") )
+	Category.AddCustomRow( LOCTEXT("BlueprintVarActionDetails_VariableNameLabel", "Variable Name") )
 	.NameContent()
 	[
 		SNew(STextBlock)
@@ -85,7 +85,7 @@ void FBlueprintVarActionDetails::CustomizeDetails( IDetailLayoutBuilder& DetailL
 
 	TSharedPtr<SToolTip> VarTypeTooltip = IDocumentation::Get()->CreateToolTip(LOCTEXT("VarTypeTooltip", "The type of the variable."), NULL, DocLink, TEXT("VariableType"));
 
-	Category.AddCustomRow( TEXT("Variable Type") )
+	Category.AddCustomRow( LOCTEXT("VariableTypeLabel", "Variable Type") )
 	.NameContent()
 	[
 		SNew(STextBlock)
@@ -107,7 +107,7 @@ void FBlueprintVarActionDetails::CustomizeDetails( IDetailLayoutBuilder& DetailL
 
 	TSharedPtr<SToolTip> EditableTooltip = IDocumentation::Get()->CreateToolTip(LOCTEXT("VarEditableTooltip", "Whether this variable is publically editable on instances of this Blueprint."), NULL, DocLink, TEXT("Editable"));
 
-	Category.AddCustomRow( TEXT("Editable") )
+	Category.AddCustomRow( LOCTEXT("IsVariableEditableLabel", "Editable") )
 	.Visibility(TAttribute<EVisibility>(this, &FBlueprintVarActionDetails::ShowEditableCheckboxVisibilty))
 	.NameContent()
 	[
@@ -126,7 +126,7 @@ void FBlueprintVarActionDetails::CustomizeDetails( IDetailLayoutBuilder& DetailL
 
 	TSharedPtr<SToolTip> ToolTipTooltip = IDocumentation::Get()->CreateToolTip(LOCTEXT("VarToolTipTooltip", "Extra information about this variable, shown when cursor is over it."), NULL, DocLink, TEXT("Tooltip"));
 
-	Category.AddCustomRow( TEXT("Tooltip") )
+	Category.AddCustomRow( LOCTEXT("IsVariableToolTipLabel", "Tooltip") )
 	.Visibility(TAttribute<EVisibility>(this, &FBlueprintVarActionDetails::IsTooltipEditVisible))
 	.NameContent()
 	[
@@ -146,7 +146,7 @@ void FBlueprintVarActionDetails::CustomizeDetails( IDetailLayoutBuilder& DetailL
 
 	TSharedPtr<SToolTip> Widget3DTooltip = IDocumentation::Get()->CreateToolTip(LOCTEXT("VariableWidget3D_Tooltip", "When true, allows the user to tweak the vector variable by using a 3D transform widget in the viewport (usable when varible is public/enabled)."), NULL, DocLink, TEXT("Widget3D"));
 
-	Category.AddCustomRow( LOCTEXT("VariableWidget3D_Prompt", "Show 3D Widget").ToString() )
+	Category.AddCustomRow( LOCTEXT("VariableWidget3D_Prompt", "Show 3D Widget") )
 	.Visibility(TAttribute<EVisibility>(this, &FBlueprintVarActionDetails::Show3DWidgetVisibility))
 	.NameContent()
 	[
@@ -167,7 +167,7 @@ void FBlueprintVarActionDetails::CustomizeDetails( IDetailLayoutBuilder& DetailL
 
 	TSharedPtr<SToolTip> ExposeOnSpawnTooltip = IDocumentation::Get()->CreateToolTip(LOCTEXT("VariableExposeToSpawn_Tooltip", "Should this variable be exposed as a pin when spawning this Blueprint?"), NULL, DocLink, TEXT("ExposeOnSpawn"));
 
-	Category.AddCustomRow( TEXT("Expose on Spawn") )
+	Category.AddCustomRow( LOCTEXT("VariableExposeToSpawnLabel", "Expose on Spawn") )
 	.Visibility(TAttribute<EVisibility>(this, &FBlueprintVarActionDetails::ExposeOnSpawnVisibility))
 	.NameContent()
 	[
@@ -186,7 +186,7 @@ void FBlueprintVarActionDetails::CustomizeDetails( IDetailLayoutBuilder& DetailL
 
 	TSharedPtr<SToolTip> PrivateTooltip = IDocumentation::Get()->CreateToolTip(LOCTEXT("VariablePrivate_Tooltip", "Should this variable be private (derived blueprints cannot modify it)?"), NULL, DocLink, TEXT("Private"));
 
-	Category.AddCustomRow(TEXT(""))
+	Category.AddCustomRow(LOCTEXT("VariablePrivate", "Private"))
 	.Visibility(TAttribute<EVisibility>(this, &FBlueprintVarActionDetails::ExposePrivateVisibility))
 	.NameContent()
 	[
@@ -205,7 +205,7 @@ void FBlueprintVarActionDetails::CustomizeDetails( IDetailLayoutBuilder& DetailL
 
 	TSharedPtr<SToolTip> ExposeToMatineeTooltip = IDocumentation::Get()->CreateToolTip(LOCTEXT("VariableExposeToMatinee_Tooltip", "Should this variable be exposed for Matinee to modify?"), NULL, DocLink, TEXT("ExposeToMatinee"));
 
-	Category.AddCustomRow( TEXT("Expose to Matinee") )
+	Category.AddCustomRow( LOCTEXT("VariableExposeToMatinee", "Expose to Matinee") )
 	.Visibility(TAttribute<EVisibility>(this, &FBlueprintVarActionDetails::ExposeToMatineeVisibility))
 	.NameContent()
 	[
@@ -228,7 +228,7 @@ void FBlueprintVarActionDetails::CustomizeDetails( IDetailLayoutBuilder& DetailL
 
 	TSharedPtr<SToolTip> CategoryTooltip = IDocumentation::Get()->CreateToolTip(LOCTEXT("EditCategoryName_Tooltip", "The category of the variable; editing this will place the variable into another category or create a new one."), NULL, DocLink, TEXT("Category"));
 
-	Category.AddCustomRow( TEXT("Category") )
+	Category.AddCustomRow( LOCTEXT("CategoryLabel", "Category") )
 	.NameContent()
 	[
 		SNew(STextBlock)
@@ -279,7 +279,7 @@ void FBlueprintVarActionDetails::CustomizeDetails( IDetailLayoutBuilder& DetailL
 
 	FName UIMin = TEXT("UIMin");
 	FName UIMax = TEXT("UIMax");
-	Category.AddCustomRow( TEXT("Slider Range") )
+	Category.AddCustomRow( LOCTEXT("SliderRangeLabel", "Slider Range") )
 	.Visibility(TAttribute<EVisibility>(this, &FBlueprintVarActionDetails::RangeVisibility))
 	.NameContent()
 	[
@@ -321,7 +321,7 @@ void FBlueprintVarActionDetails::CustomizeDetails( IDetailLayoutBuilder& DetailL
 
 	FName ClampMin = TEXT("ClampMin");
 	FName ClampMax = TEXT("ClampMax");
-	Category.AddCustomRow(TEXT("Value Range"))
+	Category.AddCustomRow(LOCTEXT("ValueRangeLabel", "Value Range"))
 	.Visibility(TAttribute<EVisibility>(this, &FBlueprintVarActionDetails::RangeVisibility))
 	.NameContent()
 	[
@@ -365,7 +365,7 @@ void FBlueprintVarActionDetails::CustomizeDetails( IDetailLayoutBuilder& DetailL
 
 	TSharedPtr<SToolTip> ReplicationTooltip = IDocumentation::Get()->CreateToolTip(LOCTEXT("VariableReplicate_Tooltip", "Should this Variable be replicated over the network?"), NULL, DocLink, TEXT("Replication"));
 
-	Category.AddCustomRow( TEXT("Replication") )
+	Category.AddCustomRow( LOCTEXT("VariableReplicationLabel", "Replication") )
 	.Visibility(TAttribute<EVisibility>(this, &FBlueprintVarActionDetails::ReplicationVisibility))
 	.NameContent()
 	[
@@ -394,7 +394,7 @@ void FBlueprintVarActionDetails::CustomizeDetails( IDetailLayoutBuilder& DetailL
 				// Add Events Section property class that has valid events
 				if( FBlueprintEditor::CanClassGenerateEvents( ComponentProperty->PropertyClass ))
 				{
-					IDetailCategoryBuilder& EventCategory = DetailLayout.EditCategory("Component", LOCTEXT("ComponentDetailsCategory", "Events").ToString());
+					IDetailCategoryBuilder& EventCategory = DetailLayout.EditCategory("Component", LOCTEXT("ComponentDetailsCategory", "Events"));
 
 					FText NodeName = FText::FromName( VarName );
 					const FText AddEventLabel = FText::Format( LOCTEXT( "ScriptingEvents_AddEvent", "Add Event For {0}" ), NodeName );
@@ -402,7 +402,7 @@ void FBlueprintVarActionDetails::CustomizeDetails( IDetailLayoutBuilder& DetailL
 
 					TSharedPtr<SToolTip> EventsTooltip = IDocumentation::Get()->CreateToolTip(AddEventToolTip, NULL, DocLink, TEXT("Events"));
 
-					EventCategory.AddCustomRow( LOCTEXT("AddEventHeader", "Add Event").ToString())
+					EventCategory.AddCustomRow( LOCTEXT("AddEventHeader", "Add Event"))
 						.NameContent()
 						[
 							SNew(SHorizontalBox)
@@ -469,7 +469,7 @@ void FBlueprintVarActionDetails::CustomizeDetails( IDetailLayoutBuilder& DetailL
 		UClass* VariableClass = (VariableProperty != NULL) ? VariableProperty->GetTypedOuter<UClass>() : NULL;
 
 		FText ErrorMessage;
-		IDetailCategoryBuilder& DefaultValueCategory = DetailLayout.EditCategory(TEXT("DefaultValueCategory"), LOCTEXT("DefaultValueCategoryHeading", "Default Value").ToString());
+		IDetailCategoryBuilder& DefaultValueCategory = DetailLayout.EditCategory(TEXT("DefaultValueCategory"), LOCTEXT("DefaultValueCategoryHeading", "Default Value"));
 
 		if (VariableProperty == NULL)
 		{
@@ -508,7 +508,7 @@ void FBlueprintVarActionDetails::CustomizeDetails( IDetailLayoutBuilder& DetailL
 		// Show the error message if something went wrong
 		if (!ErrorMessage.IsEmpty())
 		{
-			DefaultValueCategory.AddCustomRow( TEXT("Error") )
+			DefaultValueCategory.AddCustomRow( ErrorMessage )
 			[
 				SNew(STextBlock)
 				.ToolTipText(ErrorMessage)
@@ -526,7 +526,7 @@ void FBlueprintVarActionDetails::CustomizeDetails( IDetailLayoutBuilder& DetailL
 
 		TSharedPtr<SToolTip> TransientTooltip = IDocumentation::Get()->CreateToolTip(LOCTEXT("VariableTransient_Tooltip", "Should this variable not serialize and be zero-filled at load?"), NULL, DocLink, TEXT("Transient"));
 
-		Category.AddCustomRow(TEXT(""), true)
+		Category.AddCustomRow(LOCTEXT("VariableTransient", "Transient"), true)
 			.Visibility(TAttribute<EVisibility>(this, &FBlueprintVarActionDetails::GetTransientVisibility))
 			.NameContent()
 		[
@@ -545,7 +545,7 @@ void FBlueprintVarActionDetails::CustomizeDetails( IDetailLayoutBuilder& DetailL
 
 		TSharedPtr<SToolTip> SaveGameTooltip = IDocumentation::Get()->CreateToolTip(LOCTEXT("VariableSaveGame_Tooltip", "Should this variable be serialized for saved games?"), NULL, DocLink, TEXT("SaveGame"));
 
-		Category.AddCustomRow(TEXT(""), true)
+		Category.AddCustomRow(LOCTEXT("VariableSaveGame", "SaveGame"), true)
 		.Visibility(TAttribute<EVisibility>(this, &FBlueprintVarActionDetails::GetSaveGameVisibility))
 		.NameContent()
 		[
@@ -564,7 +564,7 @@ void FBlueprintVarActionDetails::CustomizeDetails( IDetailLayoutBuilder& DetailL
 
 		TSharedPtr<SToolTip> PropertyFlagsTooltip = IDocumentation::Get()->CreateToolTip(LOCTEXT("DefinedPropertyFlags_Tooltip", "List of defined flags for this property"), NULL, DocLink, TEXT("PropertyFlags"));
 
-		Category.AddCustomRow(TEXT(""), true)
+		Category.AddCustomRow(LOCTEXT("DefinedPropertyFlags", "Defined Property Flags"), true)
 		.WholeRowWidget
 		[
 			SNew(STextBlock)
@@ -573,7 +573,7 @@ void FBlueprintVarActionDetails::CustomizeDetails( IDetailLayoutBuilder& DetailL
 			.Font( IDetailLayoutBuilder::GetDetailFontBold() )
 		];
 
-		Category.AddCustomRow(TEXT(""), true)
+		Category.AddCustomRow(FText::GetEmpty(), true)
 		.WholeRowWidget
 		[
 			SAssignNew(PropertyFlagWidget, SListView< TSharedPtr< FString > >)
@@ -1709,7 +1709,7 @@ void FBlueprintGraphArgumentLayout::GenerateChildContent( IDetailChildrenBuilder
 {
 	if (bHasDefaultValue)
 	{
-		ChildrenBuilder.AddChildContent( *LOCTEXT( "FunctionArgDetailsDefaultValue", "Default Value" ).ToString() )
+		ChildrenBuilder.AddChildContent( LOCTEXT( "FunctionArgDetailsDefaultValue", "Default Value" ) )
 		.NameContent()
 		[
 			SNew(STextBlock)
@@ -1726,7 +1726,7 @@ void FBlueprintGraphArgumentLayout::GenerateChildContent( IDetailChildrenBuilder
 				.Font( IDetailLayoutBuilder::GetDetailFont() )
 		];
 
-		ChildrenBuilder.AddChildContent( *LOCTEXT( "FunctionArgDetailsPassByReference", "Pass-by-Reference" ).ToString() )
+		ChildrenBuilder.AddChildContent( LOCTEXT( "FunctionArgDetailsPassByReference", "Pass-by-Reference" ) )
 		.NameContent()
 		[
 			SNew(STextBlock)
@@ -1746,7 +1746,7 @@ void FBlueprintGraphArgumentLayout::GenerateChildContent( IDetailChildrenBuilder
 	// Read only graphs can't have their pins re-organized
 	if ( !IsPinEditingReadOnly() )
 	{
-		ChildrenBuilder.AddChildContent( *LOCTEXT( "FunctionArgDetailsMoving", "Moving" ).ToString() )
+		ChildrenBuilder.AddChildContent( LOCTEXT( "FunctionArgDetailsMoving", "Moving" ) )
 			[
 				SNew(SHorizontalBox)
 				+SHorizontalBox::Slot()
@@ -2011,15 +2011,15 @@ void FBlueprintGraphActionDetails::CustomizeDetails( IDetailLayoutBuilder& Detai
 
 	if (FunctionEntryNode && FunctionEntryNode->IsEditable())
 	{
-		IDetailCategoryBuilder& Category = DetailLayout.EditCategory("Graph", LOCTEXT("FunctionDetailsGraph", "Graph").ToString());
+		IDetailCategoryBuilder& Category = DetailLayout.EditCategory("Graph", LOCTEXT("FunctionDetailsGraph", "Graph"));
 		if (bHasAGraph)
 		{
-			Category.AddCustomRow( LOCTEXT( "DefaultTooltip", "Description" ).ToString() )
+			Category.AddCustomRow( LOCTEXT( "DefaultTooltip", "Description" ) )
 			.NameContent()
 			[
 				SNew(STextBlock)
-					.Text( LOCTEXT( "DefaultTooltip", "Description" ).ToString() )
-					.ToolTipText(LOCTEXT("FunctionTooltipTooltip", "Enter a short message describing the purpose and operation of this graph").ToString())
+					.Text( LOCTEXT( "DefaultTooltip", "Description" ) )
+					.ToolTipText(LOCTEXT("FunctionTooltipTooltip", "Enter a short message describing the purpose and operation of this graph"))
 					.Font( IDetailLayoutBuilder::GetDetailFont() )
 			]
 			.ValueContent()
@@ -2040,7 +2040,7 @@ void FBlueprintGraphActionDetails::CustomizeDetails( IDetailLayoutBuilder& Detai
 				const FString DocLink = TEXT("Shared/Editors/BlueprintEditor/VariableDetails");
 				TSharedPtr<SToolTip> CategoryTooltip = IDocumentation::Get()->CreateToolTip(LOCTEXT("EditCategoryName_Tooltip", "The category of the variable; editing this will place the variable into another category or create a new one."), NULL, DocLink, TEXT("Category"));
 
-				Category.AddCustomRow( TEXT("Category") )
+				Category.AddCustomRow( LOCTEXT("CategoryLabel", "Category") )
 					.NameContent()
 					[
 						SNew(STextBlock)
@@ -2089,11 +2089,11 @@ void FBlueprintGraphActionDetails::CustomizeDetails( IDetailLayoutBuilder& Detai
 
 			if (IsAccessSpecifierVisible())
 			{
-				Category.AddCustomRow( LOCTEXT( "AccessSpecifier", "Access Specifier" ).ToString() )
+				Category.AddCustomRow( LOCTEXT( "AccessSpecifier", "Access Specifier" ) )
 				.NameContent()
 				[
 					SNew(STextBlock)
-						.Text( LOCTEXT( "AccessSpecifier", "Access Specifier" ).ToString() )
+						.Text( LOCTEXT( "AccessSpecifier", "Access Specifier" ) )
 						.Font( IDetailLayoutBuilder::GetDetailFont() )
 				]
 				.ValueContent()
@@ -2117,7 +2117,7 @@ void FBlueprintGraphActionDetails::CustomizeDetails( IDetailLayoutBuilder& Detai
 			}
 			if (GetInstanceColorVisibility())
 			{
-				Category.AddCustomRow( LOCTEXT( "InstanceColor", "Instance Color" ).ToString() )
+				Category.AddCustomRow( LOCTEXT( "InstanceColor", "Instance Color" ) )
 				.NameContent()
 				[
 					SNew(STextBlock)
@@ -2135,12 +2135,12 @@ void FBlueprintGraphActionDetails::CustomizeDetails( IDetailLayoutBuilder& Detai
 			}
 			if (IsPureFunctionVisible())
 			{
-				Category.AddCustomRow( LOCTEXT( "FunctionPure_Tooltip", "Pure" ).ToString() )
+				Category.AddCustomRow( LOCTEXT( "FunctionPure_Tooltip", "Pure" ) )
 				.NameContent()
 				[
 					SNew(STextBlock)
-						.Text( LOCTEXT( "FunctionPure_Tooltip", "Pure" ).ToString() )
-						.ToolTipText( LOCTEXT("FunctionIsPure_Tooltip", "Force this to be a pure function?").ToString() )
+						.Text( LOCTEXT( "FunctionPure_Tooltip", "Pure" ) )
+						.ToolTipText( LOCTEXT("FunctionIsPure_Tooltip", "Force this to be a pure function?") )
 						.Font( IDetailLayoutBuilder::GetDetailFont() )
 				]
 				.ValueContent()
@@ -2217,7 +2217,7 @@ void FBlueprintGraphActionDetails::CustomizeDetails( IDetailLayoutBuilder& Detai
 										FSlateIcon(),
 										FUIAction(FExecuteAction::CreateStatic( &FBlueprintGraphActionDetails::SetNetFlags, FunctionEntryNodePtr, static_cast<uint32>(FUNC_NetClient) ), CanExecuteDelegate));
 
-			Category.AddCustomRow( LOCTEXT( "FunctionReplicate", "Replicates" ).ToString() )
+			Category.AddCustomRow( LOCTEXT( "FunctionReplicate", "Replicates" ) )
 			.NameContent()
 			[
 				SNew(STextBlock)
@@ -2271,12 +2271,12 @@ void FBlueprintGraphActionDetails::CustomizeDetails( IDetailLayoutBuilder& Detai
 							.OnCheckStateChanged( this, &FBlueprintGraphActionDetails::OnIsReliableReplicationFunctionModified )
 						[
 							SNew(STextBlock)
-								.Text( LOCTEXT( "FunctionReplicateReliable", "Reliable" ).ToString() )
+								.Text( LOCTEXT( "FunctionReplicateReliable", "Reliable" ) )
 						]
 					]
 				]
 			];
-			Category.AddCustomRow( LOCTEXT( "EditorCallable", "Call In Editor" ).ToString() )
+			Category.AddCustomRow( LOCTEXT( "EditorCallable", "Call In Editor" ) )
 			.NameContent()
 			[
 				SNew(STextBlock)
@@ -2303,19 +2303,19 @@ void FBlueprintGraphActionDetails::CustomizeDetails( IDetailLayoutBuilder& Detai
 			];
 		}
 
-		IDetailCategoryBuilder& InputsCategory = DetailLayout.EditCategory("Inputs", LOCTEXT("FunctionDetailsInputs", "Inputs").ToString());
+		IDetailCategoryBuilder& InputsCategory = DetailLayout.EditCategory("Inputs", LOCTEXT("FunctionDetailsInputs", "Inputs"));
 		
 		TSharedRef<FBlueprintGraphArgumentGroupLayout> InputArgumentGroup =
 			MakeShareable(new FBlueprintGraphArgumentGroupLayout(SharedThis(this), FunctionEntryNode));
 		InputsCategory.AddCustomBuilder(InputArgumentGroup);
 
-		InputsCategory.AddCustomRow( LOCTEXT("FunctionNewInputArg", "New").ToString() )
+		InputsCategory.AddCustomRow( LOCTEXT("FunctionNewInputArg", "New") )
 		[
 			SNew(SBox)
 			.HAlign(HAlign_Right)
 			[
 				SNew(SButton)
-				.Text(LOCTEXT("FunctionNewInputArg", "New").ToString())
+				.Text(LOCTEXT("FunctionNewInputArg", "New"))
 				.OnClicked(this, &FBlueprintGraphActionDetails::OnAddNewInputClicked)
 				.Visibility(this, &FBlueprintGraphActionDetails::GetAddNewInputOutputVisibility)
 			]
@@ -2323,7 +2323,7 @@ void FBlueprintGraphActionDetails::CustomizeDetails( IDetailLayoutBuilder& Detai
 
 		if (bHasAGraph)
 		{
-			IDetailCategoryBuilder& OutputsCategory = DetailLayout.EditCategory("Outputs", LOCTEXT("FunctionDetailsOutputs", "Outputs").ToString());
+			IDetailCategoryBuilder& OutputsCategory = DetailLayout.EditCategory("Outputs", LOCTEXT("FunctionDetailsOutputs", "Outputs"));
 		
 			if (FunctionResultNode)
 			{
@@ -2332,13 +2332,13 @@ void FBlueprintGraphActionDetails::CustomizeDetails( IDetailLayoutBuilder& Detai
 				OutputsCategory.AddCustomBuilder(OutputArgumentGroup);
 			}
 
-			OutputsCategory.AddCustomRow( LOCTEXT("FunctionNewOutputArg", "New").ToString() )
+			OutputsCategory.AddCustomRow( LOCTEXT("FunctionNewOutputArg", "New") )
 			[
 				SNew(SBox)
 				.HAlign(HAlign_Right)
 				[
 					SNew(SButton)
-					.Text(LOCTEXT("FunctionNewOutputArg", "New").ToString())
+					.Text(LOCTEXT("FunctionNewOutputArg", "New"))
 					.OnClicked(this, &FBlueprintGraphActionDetails::OnAddNewOutputClicked)
 					.Visibility(this, &FBlueprintGraphActionDetails::GetAddNewInputOutputVisibility)
 				]
@@ -2349,11 +2349,11 @@ void FBlueprintGraphActionDetails::CustomizeDetails( IDetailLayoutBuilder& Detai
 	{
 		if (bHasAGraph)
 		{
-			IDetailCategoryBuilder& Category = DetailLayout.EditCategory("Graph", LOCTEXT("FunctionDetailsGraph", "Graph").ToString());
-			Category.AddCustomRow( TEXT("") )
+			IDetailCategoryBuilder& Category = DetailLayout.EditCategory("Graph", LOCTEXT("FunctionDetailsGraph", "Graph"));
+			Category.AddCustomRow( FText::GetEmpty() )
 			[
 				SNew(STextBlock)
-				.Text( LOCTEXT("GraphPresentButNotEditable", "Graph is not editable.").ToString() )
+				.Text( LOCTEXT("GraphPresentButNotEditable", "Graph is not editable.") )
 			];
 		}
 	}
@@ -2366,8 +2366,8 @@ TSharedRef<ITableRow> FBlueprintGraphActionDetails::OnGenerateReplicationComboWi
 		SNew(STableRow< TSharedPtr<FString> >, OwnerTable)
 		[
 			SNew( STextBlock )
-			.Text( InNetFlag.IsValid() ? InNetFlag.Get()->LocalizedName.ToString() : FString() )
-			.ToolTipText( InNetFlag.IsValid() ? InNetFlag.Get()->LocalizedToolTip.ToString() : FString() )
+			.Text( InNetFlag.IsValid() ? InNetFlag.Get()->LocalizedName : FText::GetEmpty() )
+			.ToolTipText( InNetFlag.IsValid() ? InNetFlag.Get()->LocalizedToolTip : FText::GetEmpty() )
 		];
 }
 
@@ -2657,8 +2657,8 @@ void FBlueprintDelegateActionDetails::CustomizeDetails( IDetailLayoutBuilder& De
 	const FSlateFontInfo DetailFontInfo = IDetailLayoutBuilder::GetDetailFont();
 
 	{
-		IDetailCategoryBuilder& Category = DetailLayout.EditCategory("Delegate", LOCTEXT("DelegateDetailsCategory", "Delegate").ToString());
-		Category.AddCustomRow( TEXT("Tooltip") )
+		IDetailCategoryBuilder& Category = DetailLayout.EditCategory("Delegate", LOCTEXT("DelegateDetailsCategory", "Delegate"));
+		Category.AddCustomRow( LOCTEXT("VariableToolTipLabel", "Tooltip") )
 		.NameContent()
 		[
 			SNew(STextBlock)
@@ -2677,7 +2677,7 @@ void FBlueprintDelegateActionDetails::CustomizeDetails( IDetailLayoutBuilder& De
 		TSharedPtr<SComboButton> NewComboButton;
 		TSharedPtr<SListView<TSharedPtr<FString>>> NewListView;
 
-		Category.AddCustomRow( TEXT("Category") )
+		Category.AddCustomRow( LOCTEXT("CategoryLabel", "Category") )
 		.NameContent()
 		[
 			SNew(STextBlock)
@@ -2724,11 +2724,11 @@ void FBlueprintDelegateActionDetails::CustomizeDetails( IDetailLayoutBuilder& De
 
 	if (UK2Node_EditablePinBase* FunctionEntryNode = FunctionEntryNodePtr.Get())
 	{
-		IDetailCategoryBuilder& InputsCategory = DetailLayout.EditCategory("DelegateInputs", LOCTEXT("DelegateDetailsInputs", "Inputs").ToString());
+		IDetailCategoryBuilder& InputsCategory = DetailLayout.EditCategory("DelegateInputs", LOCTEXT("DelegateDetailsInputs", "Inputs"));
 		TSharedRef<FBlueprintGraphArgumentGroupLayout> InputArgumentGroup = MakeShareable(new FBlueprintGraphArgumentGroupLayout(SharedThis(this), FunctionEntryNode));
 		InputsCategory.AddCustomBuilder(InputArgumentGroup);
 
-		InputsCategory.AddCustomRow( LOCTEXT("FunctionNewInputArg", "New").ToString() )
+		InputsCategory.AddCustomRow( LOCTEXT("FunctionNewInputArg", "New") )
 		[
 			SNew(SBox)
 			.HAlign(HAlign_Right)
@@ -2741,7 +2741,7 @@ void FBlueprintDelegateActionDetails::CustomizeDetails( IDetailLayoutBuilder& De
 
 		CollectAvailibleSignatures();
 
-		InputsCategory.AddCustomRow( LOCTEXT("CopySignatureFrom", "Copy signature from").ToString() )
+		InputsCategory.AddCustomRow( LOCTEXT("CopySignatureFrom", "Copy signature from") )
 		.NameContent()
 		[
 			SNew(STextBlock)
@@ -3794,7 +3794,7 @@ void FBlueprintInterfaceLayout::GenerateChildContent( IDetailChildrenBuilder& Ch
 	for (int32 i = 0; i < Interfaces.Num(); ++i)
 	{
 		TSharedPtr<SHorizontalBox> Box;
-		ChildrenBuilder.AddChildContent( *LOCTEXT( "BlueprintInterfaceValue", "Interface Value" ).ToString() )
+		ChildrenBuilder.AddChildContent( LOCTEXT( "BlueprintInterfaceValue", "Interface Value" ) )
 		[
 			SAssignNew(Box, SHorizontalBox)
 			+SHorizontalBox::Slot()
@@ -3838,7 +3838,7 @@ void FBlueprintInterfaceLayout::GenerateChildContent( IDetailChildrenBuilder& Ch
 	// Add message if no interfaces are being used
 	if (Interfaces.Num() == 0)
 	{
-		ChildrenBuilder.AddChildContent(*LOCTEXT("BlueprintInterfaceValue", "Interface Value").ToString())
+		ChildrenBuilder.AddChildContent(LOCTEXT("BlueprintInterfaceValue", "Interface Value"))
 		[
 			SNew(STextBlock)
 			.Text(LOCTEXT("NoBlueprintInterface", "No Interfaces"))
@@ -3848,7 +3848,7 @@ void FBlueprintInterfaceLayout::GenerateChildContent( IDetailChildrenBuilder& Ch
 
 	if (!bShowsInheritedInterfaces)
 	{
-		ChildrenBuilder.AddChildContent( *LOCTEXT( "BlueprintAddInterface", "Add Interface" ).ToString() )
+		ChildrenBuilder.AddChildContent( LOCTEXT( "BlueprintAddInterface", "Add Interface" ) )
 		[
 			SNew(SBox)
 			.HAlign(HAlign_Right)
@@ -4042,8 +4042,8 @@ void FBlueprintGlobalOptionsDetails::CustomizeDetails(IDetailLayoutBuilder& Deta
 		}
 
 		// Display the parent class and set up the menu for reparenting
-		IDetailCategoryBuilder& Category = DetailLayout.EditCategory("Globals", LOCTEXT("BlueprintGlobalDetailsCategory", "Globals").ToString());
-		Category.AddCustomRow( LOCTEXT("BlueprintGlobalDetailsCategory", "Globals").ToString() )
+		IDetailCategoryBuilder& Category = DetailLayout.EditCategory("Globals", LOCTEXT("BlueprintGlobalDetailsCategory", "Globals"));
+		Category.AddCustomRow( LOCTEXT("BlueprintGlobalDetailsCategory", "Globals") )
 		.NameContent()
 		[
 			SNew(STextBlock)
@@ -4072,7 +4072,7 @@ void FBlueprintGlobalOptionsDetails::CustomizeDetails(IDetailLayoutBuilder& Deta
 		if (bSupportsInterfaces)
 		{
 			// Interface details customization
-			IDetailCategoryBuilder& InterfacesCategory = DetailLayout.EditCategory("Interfaces", LOCTEXT("BlueprintInterfacesDetailsCategory", "Interfaces").ToString());
+			IDetailCategoryBuilder& InterfacesCategory = DetailLayout.EditCategory("Interfaces", LOCTEXT("BlueprintInterfacesDetailsCategory", "Interfaces"));
 		
 			TSharedRef<FBlueprintInterfaceLayout> InterfaceLayout = MakeShareable(new FBlueprintInterfaceLayout(SharedThis(this), false));
 			InterfacesCategory.AddCustomBuilder(InterfaceLayout);
@@ -4094,7 +4094,7 @@ void FBlueprintGlobalOptionsDetails::CustomizeDetails(IDetailLayoutBuilder& Deta
 		else
 		{
 			// Only display the ability to deprecate a Blueprint on non-level Blueprints.
-			Category.AddCustomRow( TEXT("Deprecate") )
+			Category.AddCustomRow( LOCTEXT("DeprecateLabel", "Deprecate") )
 				.NameContent()
 				[
 					SNew(STextBlock)
@@ -4134,13 +4134,13 @@ void FBlueprintComponentDetails::CustomizeDetails(IDetailLayoutBuilder& DetailLa
 
 	if( CachedNodePtr.IsValid() )
 	{
-		IDetailCategoryBuilder& VariableCategory = DetailLayout.EditCategory("Variable", LOCTEXT("VariableDetailsCategory", "Variable").ToString(), ECategoryPriority::Variable);
+		IDetailCategoryBuilder& VariableCategory = DetailLayout.EditCategory("Variable", LOCTEXT("VariableDetailsCategory", "Variable"), ECategoryPriority::Variable);
 
-		VariableCategory.AddCustomRow(TEXT("Variable Name"))
+		VariableCategory.AddCustomRow(LOCTEXT("BlueprintComponentDetails_VariableNameLabel", "Variable Name"))
 		.NameContent()
 		[
 			SNew(STextBlock)
-			.Text(LOCTEXT("BlueprintComponentDetails_VariableNameLabel", "Variable Name").ToString())
+			.Text(LOCTEXT("BlueprintComponentDetails_VariableNameLabel", "Variable Name"))
 			.Font(IDetailLayoutBuilder::GetDetailFont())
 		]
 		.ValueContent()
@@ -4153,11 +4153,11 @@ void FBlueprintComponentDetails::CustomizeDetails(IDetailLayoutBuilder& DetailLa
 			.Font(IDetailLayoutBuilder::GetDetailFont())
 		];
 
-		VariableCategory.AddCustomRow(TEXT("Tooltip"))
+		VariableCategory.AddCustomRow(LOCTEXT("BlueprintComponentDetails_VariableTooltipLabel", "Tooltip"))
 		.NameContent()
 		[
 			SNew(STextBlock)
-			.Text(LOCTEXT("BlueprintComponentDetails_VariableTooltipLabel", "Tooltip").ToString())
+			.Text(LOCTEXT("BlueprintComponentDetails_VariableTooltipLabel", "Tooltip"))
 			.Font(IDetailLayoutBuilder::GetDetailFont())
 		]
 		.ValueContent()
@@ -4171,11 +4171,11 @@ void FBlueprintComponentDetails::CustomizeDetails(IDetailLayoutBuilder& DetailLa
 		PopulateVariableCategories();
 		const FString CategoryTooltip = LOCTEXT("EditCategoryName_Tooltip", "The category of the variable; editing this will place the variable into another category or create a new one.").ToString();
 
-		VariableCategory.AddCustomRow( TEXT("Category") )
+		VariableCategory.AddCustomRow( LOCTEXT("BlueprintComponentDetails_VariableCategoryLabel", "Category") )
 		.NameContent()
 		[
 			SNew(STextBlock)
-			.Text(LOCTEXT("BlueprintComponentDetails_VariableCategoryLabel", "Category").ToString())
+			.Text(LOCTEXT("BlueprintComponentDetails_VariableCategoryLabel", "Category"))
 			.ToolTipText(CategoryTooltip)
 			.Font(IDetailLayoutBuilder::GetDetailFont())
 		]
@@ -4214,13 +4214,13 @@ void FBlueprintComponentDetails::CustomizeDetails(IDetailLayoutBuilder& DetailLa
 			]
 		];
 
-		IDetailCategoryBuilder& SocketsCategory = DetailLayout.EditCategory("Sockets", LOCTEXT("BlueprintComponentDetailsCategory", "Sockets").ToString(), ECategoryPriority::Important);
+		IDetailCategoryBuilder& SocketsCategory = DetailLayout.EditCategory("Sockets", LOCTEXT("BlueprintComponentDetailsCategory", "Sockets"), ECategoryPriority::Important);
 
-		SocketsCategory.AddCustomRow(LOCTEXT("BlueprintComponentDetails_Sockets", "Sockets").ToString())
+		SocketsCategory.AddCustomRow(LOCTEXT("BlueprintComponentDetails_Sockets", "Sockets"))
 		.NameContent()
 		[
 			SNew(STextBlock)
-			.Text(LOCTEXT("BlueprintComponentDetails_ParentSocket", "Parent Socket").ToString())
+			.Text(LOCTEXT("BlueprintComponentDetails_ParentSocket", "Parent Socket"))
 			.Font(IDetailLayoutBuilder::GetDetailFont())
 		]
 		.ValueContent()
@@ -4259,13 +4259,13 @@ void FBlueprintComponentDetails::CustomizeDetails(IDetailLayoutBuilder& DetailLa
 
 	if( FBlueprintEditor::CanClassGenerateEvents( CommonEventsClass ))
 	{
-		IDetailCategoryBuilder& EventCategory = DetailLayout.EditCategory("Component", LOCTEXT("ComponentDetailsCategory", "Events").ToString(), ECategoryPriority::Important);
+		IDetailCategoryBuilder& EventCategory = DetailLayout.EditCategory("Component", LOCTEXT("ComponentDetailsCategory", "Events"), ECategoryPriority::Important);
 
 		FText NodeName = ( Nodes.Num() > 1 ) ? LOCTEXT( "ScriptingEvents_SelectedComponents", "Selected Components" ) : FText::FromName( CachedNodePtr->GetVariableName() );
 		const FText AddEventLabel = FText::Format( LOCTEXT( "ScriptingEvents_AddEvent", "Add Event For {0}" ), NodeName );
 		const FText AddEventToolTip = FText::Format( LOCTEXT( "ScriptingEvents_AddOrView", "Adds or views events for {0} in the Component Blueprint" ), NodeName );
 
-		EventCategory.AddCustomRow( LOCTEXT("AddEventHeader", "Add Event").ToString())
+		EventCategory.AddCustomRow( LOCTEXT("AddEventHeader", "Add Event"))
 			.NameContent()
 			[
 				SNew(SHorizontalBox)
@@ -4662,7 +4662,7 @@ void FBlueprintGraphNodeDetails::CustomizeDetails( IDetailLayoutBuilder& DetailL
 		return;
 	}
 
-	IDetailCategoryBuilder& Category = DetailLayout.EditCategory("GraphNodeDetail", LOCTEXT("GraphNodeDetailsCategory", "Graph Node").ToString(), ECategoryPriority::Important);
+	IDetailCategoryBuilder& Category = DetailLayout.EditCategory("GraphNodeDetail", LOCTEXT("GraphNodeDetailsCategory", "Graph Node"), ECategoryPriority::Important);
 	const FSlateFontInfo DetailFontInfo = IDetailLayoutBuilder::GetDetailFont();
 	FText RowHeader;
 	FText NameContent;
@@ -4679,7 +4679,7 @@ void FBlueprintGraphNodeDetails::CustomizeDetails( IDetailLayoutBuilder& DetailL
 	}
 
 
-	Category.AddCustomRow( RowHeader.ToString() )
+	Category.AddCustomRow( RowHeader )
 	.NameContent()
 	[
 		SNew(STextBlock)
@@ -4838,14 +4838,14 @@ void FBlueprintDocumentationDetails::CustomizeDetails(IDetailLayoutBuilder& Deta
 		DocumentationLink = DocumentationNodePtr->GetDocumentationLink();
 		DocumentationExcerpt = DocumentationNodePtr->GetDocumentationExcerptName();
 
-		IDetailCategoryBuilder& DocumentationCategory = DetailLayout.EditCategory("Documentation", LOCTEXT("DocumentationDetailsCategory", "Documentation").ToString(), ECategoryPriority::Default);
+		IDetailCategoryBuilder& DocumentationCategory = DetailLayout.EditCategory("Documentation", LOCTEXT("DocumentationDetailsCategory", "Documentation"), ECategoryPriority::Default);
 
-		DocumentationCategory.AddCustomRow( TEXT( "Documentation Link" ))
+		DocumentationCategory.AddCustomRow( LOCTEXT( "DocumentationLinkLabel", "Documentation Link" ))
 		.NameContent()
 		.HAlign( HAlign_Fill )
 		[
 			SNew( STextBlock )
-			.Text( LOCTEXT( "FBlueprintDocumentationDetails_Link", "Link" ).ToString() )
+			.Text( LOCTEXT( "FBlueprintDocumentationDetails_Link", "Link" ) )
 			.ToolTipText( LOCTEXT( "FBlueprintDocumentationDetails_LinkPathTooltip", "The documentation content path" ))
 			.Font( IDetailLayoutBuilder::GetDetailFont() )
 		]
@@ -4862,12 +4862,12 @@ void FBlueprintDocumentationDetails::CustomizeDetails(IDetailLayoutBuilder& Deta
 			.Font( IDetailLayoutBuilder::GetDetailFont() )
 		];
 
-		DocumentationCategory.AddCustomRow( TEXT( "Documentation Excerpts" ))
+		DocumentationCategory.AddCustomRow( LOCTEXT( "DocumentationExcerptsLabel", "Documentation Excerpts" ))
 		.NameContent()
 		.HAlign( HAlign_Left )
 		[
 			SNew( STextBlock )
-			.Text( LOCTEXT( "FBlueprintDocumentationDetails_Excerpt", "Excerpt" ).ToString() )
+			.Text( LOCTEXT( "FBlueprintDocumentationDetails_Excerpt", "Excerpt" ) )
 			.ToolTipText( LOCTEXT( "FBlueprintDocumentationDetails_ExcerptTooltip", "The current documentation excerpt" ))
 			.Font( IDetailLayoutBuilder::GetDetailFont() )
 		]
diff --git a/Engine/Source/Editor/Kismet/Private/FormatTextDetails.cpp b/Engine/Source/Editor/Kismet/Private/FormatTextDetails.cpp
--- a/Engine/Source/Editor/Kismet/Private/FormatTextDetails.cpp
+++ b/Engine/Source/Editor/Kismet/Private/FormatTextDetails.cpp
@@ -18,9 +18,9 @@ void FFormatTextDetails::CustomizeDetails( IDetailLayoutBuilder& DetailLayout )
 		TargetNode = CastChecked<UK2Node_FormatText>(Objects[0].Get());
 		TSharedRef<IPropertyHandle> PropertyHandle = DetailLayout.GetProperty(FName("PinNames"), UK2Node_FormatText::StaticClass());
 
-		IDetailCategoryBuilder& InputsCategory = DetailLayout.EditCategory("Arguments", LOCTEXT("FormatTextDetailsArguments", "Arguments").ToString());
+		IDetailCategoryBuilder& InputsCategory = DetailLayout.EditCategory("Arguments", LOCTEXT("FormatTextDetailsArguments", "Arguments"));
 
-		InputsCategory.AddCustomRow( LOCTEXT("FunctionNewInputArg", "New").ToString() )
+		InputsCategory.AddCustomRow( LOCTEXT("FunctionNewInputArg", "New") )
 			[
 				SNew(SBox)
 				.HAlign(HAlign_Right)
diff --git a/Engine/Source/Editor/Kismet/Private/UserDefinedEnumEditor.cpp b/Engine/Source/Editor/Kismet/Private/UserDefinedEnumEditor.cpp
--- a/Engine/Source/Editor/Kismet/Private/UserDefinedEnumEditor.cpp
+++ b/Engine/Source/Editor/Kismet/Private/UserDefinedEnumEditor.cpp
@@ -150,9 +150,9 @@ void FEnumDetails::CustomizeDetails( IDetailLayoutBuilder& DetailLayout )
 		TargetEnum = CastChecked<UUserDefinedEnum>(Objects[0].Get());
 		TSharedRef<IPropertyHandle> PropertyHandle = DetailLayout.GetProperty(FName("Names"), UEnum::StaticClass());
 
-		IDetailCategoryBuilder& InputsCategory = DetailLayout.EditCategory("Enumerators", LOCTEXT("EnumDetailsEnumerators", "Enumerators").ToString());
+		IDetailCategoryBuilder& InputsCategory = DetailLayout.EditCategory("Enumerators", LOCTEXT("EnumDetailsEnumerators", "Enumerators"));
 
-		InputsCategory.AddCustomRow( LOCTEXT("FunctionNewInputArg", "New").ToString() )
+		InputsCategory.AddCustomRow( LOCTEXT("FunctionNewInputArg", "New") )
 			[
 				SNew(SBox)
 				.HAlign(HAlign_Right)
diff --git a/Engine/Source/Editor/Kismet/Private/UserDefinedStructureEditor.cpp b/Engine/Source/Editor/Kismet/Private/UserDefinedStructureEditor.cpp
--- a/Engine/Source/Editor/Kismet/Private/UserDefinedStructureEditor.cpp
+++ b/Engine/Source/Editor/Kismet/Private/UserDefinedStructureEditor.cpp
@@ -34,7 +34,7 @@ public:
 		ViewArgs.bHideSelectionTip = false;
 		ViewArgs.bShowActorLabel = false;
 
-		StructureDetailsView = PropertyModule.CreateStructureDetailView(ViewArgs, StructData, false, LOCTEXT("DefaultValues", "Default Values").ToString());
+		StructureDetailsView = PropertyModule.CreateStructureDetailView(ViewArgs, StructData, false, LOCTEXT("DefaultValues", "Default Values"));
 		StructureDetailsView->GetOnFinishedChangingPropertiesDelegate().AddSP(this, &FStructureDefaultValueView::OnFinishedChangingProperties);
 	}
 
@@ -770,7 +770,7 @@ public:
 
 	virtual void GenerateChildContent( IDetailChildrenBuilder& ChildrenBuilder ) override 
 	{
-		ChildrenBuilder.AddChildContent(*LOCTEXT("Tooltip", "Tooltip").ToString())
+		ChildrenBuilder.AddChildContent(LOCTEXT("Tooltip", "Tooltip"))
 		.NameContent()
 		[
 			SNew(STextBlock)
@@ -785,7 +785,7 @@ public:
 			.Font(IDetailLayoutBuilder::GetDetailFont())
 		];
 
-		ChildrenBuilder.AddChildContent(*LOCTEXT("EditableOnInstance", "EditableOnInstance").ToString())
+		ChildrenBuilder.AddChildContent(LOCTEXT("EditableOnInstance", "EditableOnInstance"))
 		.NameContent()
 		[
 			SNew(STextBlock)
@@ -800,7 +800,7 @@ public:
 			.IsChecked(this, &FUserDefinedStructureFieldLayout::OnGetEditableOnBPInstanceState)
 		];
 
-		ChildrenBuilder.AddChildContent(*LOCTEXT("3dWidget", "3d Widget").ToString())
+		ChildrenBuilder.AddChildContent(LOCTEXT("3dWidget", "3d Widget"))
 		.NameContent()
 		[
 			SNew(STextBlock)
@@ -843,7 +843,7 @@ void FUserDefinedStructureLayout::GenerateChildContent( IDetailChildrenBuilder&
 	const float NameWidth = 80.0f;
 	const float ContentWidth = 130.0f;
 
-	ChildrenBuilder.AddChildContent(FString())
+	ChildrenBuilder.AddChildContent(FText::GetEmpty())
 	[
 		SNew(SHorizontalBox)
 		+ SHorizontalBox::Slot()
@@ -864,13 +864,13 @@ void FUserDefinedStructureLayout::GenerateChildContent( IDetailChildrenBuilder&
 			[
 				SNew(SButton)
 				.HAlign(HAlign_Center)
-				.Text(LOCTEXT("NewStructureField", "New Variable").ToString())
+				.Text(LOCTEXT("NewStructureField", "New Variable"))
 				.OnClicked(this, &FUserDefinedStructureLayout::OnAddNewField)
 			]
 		]
 	];
 
-	ChildrenBuilder.AddChildContent(FString())
+	ChildrenBuilder.AddChildContent(FText::GetEmpty())
 	[
 		SNew(SHorizontalBox)
 		+ SHorizontalBox::Slot()
@@ -928,7 +928,7 @@ void FUserDefinedStructureDetails::CustomizeDetails(class IDetailLayoutBuilder&
 	{
 		UserDefinedStruct = CastChecked<UUserDefinedStruct>(Objects[0].Get());
 
-		IDetailCategoryBuilder& StructureCategory = DetailLayout.EditCategory("Structure", LOCTEXT("Structure", "Structure").ToString());
+		IDetailCategoryBuilder& StructureCategory = DetailLayout.EditCategory("Structure", LOCTEXT("Structure", "Structure"));
 		Layout = MakeShareable(new FUserDefinedStructureLayout(SharedThis(this)));
 		StructureCategory.AddCustomBuilder(Layout.ToSharedRef());
 	}
