Fixed double initialization of merge graph panel...

diff --git a/Engine/Source/Editor/Kismet/Private/SBlueprintDiff.cpp b/Engine/Source/Editor/Kismet/Private/SBlueprintDiff.cpp
--- a/Engine/Source/Editor/Kismet/Private/SBlueprintDiff.cpp
+++ b/Engine/Source/Editor/Kismet/Private/SBlueprintDiff.cpp
@@ -628,7 +628,7 @@ FDiffPanel::FDiffPanel()
 	LastFocusedPin = NULL;
 }
 
-TSharedRef<SWidget> FDiffPanel::InitializeDiffPanel()
+void FDiffPanel::InitializeDiffPanel()
 {
 	TSharedRef< SKismetInspector > Inspector = SNew(SKismetInspector)
 		.HideNameArea(true)
@@ -637,23 +637,6 @@ TSharedRef<SWidget> FDiffPanel::InitializeDiffPanel()
 		.IsPropertyEditingEnabledDelegate(FIsPropertyEditingEnabled::CreateStatic([] { return false; }));
 	DetailsView = Inspector;
 	MyBlueprint->SetInspector(DetailsView);
-
-	return SNew(SSplitter)
-		.Orientation(Orient_Vertical)
-		+ SSplitter::Slot()
-		.Value(0.8f)
-		[
-			SAssignNew(GraphEditorBorder, SBox)
-			.VAlign(VAlign_Fill)
-			[
-				SBlueprintDiff::DefaultEmptyPanel()
-			]
-		]
-	+ SSplitter::Slot()
-		.Value(0.2f)
-		[
-			Inspector
-		];
 }
 
 int32 GetCurrentIndex( SListView< TSharedPtr< struct FDiffSingleResult> > const& ListView, const TArray< TSharedPtr< FDiffSingleResult > >& ListViewSource )
diff --git a/Engine/Source/Editor/Kismet/Public/SBlueprintDiff.h b/Engine/Source/Editor/Kismet/Public/SBlueprintDiff.h
--- a/Engine/Source/Editor/Kismet/Public/SBlueprintDiff.h
+++ b/Engine/Source/Editor/Kismet/Public/SBlueprintDiff.h
@@ -46,7 +46,7 @@ struct KISMET_API FDiffPanel
 	FDiffPanel();
 
 	/** Initializes the panel, can be moved into constructor if diff and merge clients are made more uniform: */
-	TSharedRef<SWidget> InitializeDiffPanel();
+	void InitializeDiffPanel();
 
 	/* Generate this panel based on the specified graph */
 	void GeneratePanel(UEdGraph* Graph, UEdGraph* GraphToDiff);
