Fixed code relying on SLATE_TEXT_ARGUMENT for SErrorText

Made sure everything was using FText rather than FString.

diff --git a/Engine/Source/Editor/Kismet/Private/BlueprintEditorTabFactories.cpp b/Engine/Source/Editor/Kismet/Private/BlueprintEditorTabFactories.cpp
--- a/Engine/Source/Editor/Kismet/Private/BlueprintEditorTabFactories.cpp
+++ b/Engine/Source/Editor/Kismet/Private/BlueprintEditorTabFactories.cpp
@@ -184,7 +184,7 @@ TSharedRef<SWidget> FSCSViewportSummoner::CreateTabBody(const FWorkflowTabSpawnI
 	{
 		return SNew(SErrorText)
 			.BackgroundColor(FLinearColor::Transparent)
-			.ErrorText(LOCTEXT("SCSViewportView_Unavailable", "Viewport is not available for this Blueprint.").ToString());
+			.ErrorText(LOCTEXT("SCSViewportView_Unavailable", "Viewport is not available for this Blueprint."));
 	}
 }
 
