Creating tree view for displaying changes in the blueprint diff tool, will be extended to merge tool presently...

diff --git a/Engine/Source/Editor/GraphEditor/Public/DiffResults.h b/Engine/Source/Editor/GraphEditor/Public/DiffResults.h
--- a/Engine/Source/Editor/GraphEditor/Public/DiffResults.h
+++ b/Engine/Source/Editor/GraphEditor/Public/DiffResults.h
@@ -68,6 +68,9 @@ struct FDiffSingleResult
 
 	/** User can override color to use for display string */
 	FLinearColor DisplayColor;
+
+	/** Name of the graph this difference was created on*/
+	FName OwningGraph;
 };
 
 FORCEINLINE bool operator==( const FDiffSingleResult& LHS, const FDiffSingleResult& RHS )
diff --git a/Engine/Source/Editor/Kismet/Private/SBlueprintDiff.cpp b/Engine/Source/Editor/Kismet/Private/SBlueprintDiff.cpp
--- a/Engine/Source/Editor/Kismet/Private/SBlueprintDiff.cpp
+++ b/Engine/Source/Editor/Kismet/Private/SBlueprintDiff.cpp
@@ -242,6 +242,8 @@ struct KISMET_API FListItemGraphToDiff	: public TSharedFromThis<FListItemGraphTo
 	/*Get new(right) graph*/
 	UEdGraph* GetGraphNew()const{ return GraphNew; }
 
+	/** Source for list view */
+	TArray<TSharedPtr<struct FDiffResultItem>> DiffListSource;
 private:
 
 	/*Get icon to use by graph name*/
@@ -280,9 +282,6 @@ public:
 	/** Called when user clicks on a new graph list item */
 	void OnSelectionChanged(FSharedDiffOnGraph Item, ESelectInfo::Type SelectionType);
 
-	/** Called when user presses key within the diff view */
-	void KeyWasPressed(const FKeyEvent& InKeyEvent);
-
 private:
 
 	/** Called when user clicks button to go to next difference in graph */
@@ -303,12 +302,6 @@ private:
 
 	/** ListView of differences */
 	TSharedPtr<SListViewType> DiffList;
-
-	/** Source for list view */
-	TArray<FSharedDiffOnGraph> DiffListSource;
-
-	/** Key commands processed by this widget */
-	TSharedPtr< FUICommandList > KeyCommands;
 };
 
 TSharedRef<SWidget>	FDiffResultItem::GenerateWidget() const
@@ -337,6 +330,8 @@ FListItemGraphToDiff::FListItemGraphToDiff( class SBlueprintDiff* InDiff, class
 	{
 		InGraphNew->AddOnGraphChangedHandler( FOnGraphChanged::FDelegate::CreateRaw(this, &FListItemGraphToDiff::OnGraphChanged));
 	}
+
+	BuildDiffSourceArray();
 }
 
 FListItemGraphToDiff::~FListItemGraphToDiff()
@@ -353,8 +348,6 @@ TSharedRef<SWidget> FListItemGraphToDiff::GenerateWidget()
 	
 	FLinearColor Color = (GraphOld && GraphNew) ? FLinearColor::White : FLinearColor(0.3f,0.3f,1.f);
 
-	BuildDiffSourceArray();
-
 	const bool bHasDiffs = DiffListSource.Num() > 0;
 
 	//give it a red color to indicate that the graph contains differences
@@ -364,7 +357,7 @@ TSharedRef<SWidget> FListItemGraphToDiff::GenerateWidget()
 	}
 
 	return SNew(SHorizontalBox)
-		+SHorizontalBox::Slot()
+	+ SHorizontalBox::Slot()
 		.AutoWidth()
 		.HAlign(HAlign_Center)
 		.VAlign(VAlign_Center)
@@ -395,50 +388,10 @@ TSharedRef<SWidget> FListItemGraphToDiff::GenerateWidget()
 		];
 }
 
-//////////////////////////////////////////////////////////////////////////
-// FDiffListCommands
-
-class FDiffListCommands : public TCommands<FDiffListCommands>
-{
-public:
-	/** Constructor */
-	FDiffListCommands() 
-		: TCommands<FDiffListCommands>("DiffList", NSLOCTEXT("Diff", "Blueprint", "Blueprint Diff"), NAME_None, FEditorStyle::GetStyleSetName())
-	{
-	}
-
-	/** Go to previous difference */
-	TSharedPtr<FUICommandInfo> Previous;
-
-	/** Go to next difference */
-	TSharedPtr<FUICommandInfo> Next;
-
-	/** Initialize commands */
-	virtual void RegisterCommands() override;
-};
-
-void FDiffListCommands::RegisterCommands()
-{
-	UI_COMMAND(Previous, "Prev", "Go to previous difference", EUserInterfaceActionType::Button, FInputGesture(EKeys::F7, EModifierKey::Control));
-	UI_COMMAND(Next, "Next", "Go to next difference", EUserInterfaceActionType::Button, FInputGesture(EKeys::F7));
-}
-
-
-TSharedRef<SWidget> FListItemGraphToDiff::GenerateDiffListWidget() 
+TSharedRef<SWidget> FListItemGraphToDiff::GenerateDiffListWidget()
 {
-	BuildDiffSourceArray();
 	if(DiffListSource.Num() > 0)
 	{
-		struct SortDiff
-		{
-			bool operator () (const FSharedDiffOnGraph& A, const FSharedDiffOnGraph& B) const
-			{
-				return A->Result.Diff < B->Result.Diff;
-			}
-		};
-
-		Sort(DiffListSource.GetData(),DiffListSource.Num(), SortDiff());
-
 		TSharedPtr<SListViewType> DiffListRef;
 		TSharedRef<SHorizontalBox> Result =	SNew(SHorizontalBox)
 			+SHorizontalBox::Slot()
@@ -496,7 +449,7 @@ void FListItemGraphToDiff::OnSelectionChanged( FSharedDiffOnGraph Item, ESelectI
 {
 	if(Item.IsValid())
 	{
-		Diff->OnDiffListSelectionChanged(Item, this);
+		Diff->OnDiffListSelectionChanged(Item);
 	}
 }
 
@@ -509,7 +462,6 @@ FText FListItemGraphToDiff::GetToolTip()
 {
 	if ( GraphOld && GraphNew )
 	{
-		BuildDiffSourceArray();
 		if ( DiffListSource.Num() > 0 )
 		{
 			return LOCTEXT("ContainsDifferences", "Revisions are different");
@@ -539,11 +491,24 @@ void FListItemGraphToDiff::BuildDiffSourceArray()
 	TArray<FDiffSingleResult> FoundDiffs;
 	FGraphDiffControl::DiffGraphs(GraphOld, GraphNew, FoundDiffs);
 
+	FName GraphName = GraphOld ? GraphOld->GetFName() : GraphNew->GetFName();
+
 	DiffListSource.Empty();
 	for (auto DiffIt(FoundDiffs.CreateIterator()); DiffIt; ++DiffIt)
 	{
+		DiffIt->OwningGraph = GraphName;
 		DiffListSource.Add(FSharedDiffOnGraph(new FDiffResultItem(*DiffIt)));
 	}
+
+	struct SortDiff
+	{
+		bool operator () (const FSharedDiffOnGraph& A, const FSharedDiffOnGraph& B) const
+		{
+			return A->Result.Diff < B->Result.Diff;
+		}
+	};
+
+	Sort(DiffListSource.GetData(), DiffListSource.Num(), SortDiff());
 }
 
 void FListItemGraphToDiff::NextDiff()
@@ -614,14 +579,6 @@ void FListItemGraphToDiff::OnGraphChanged( const FEdGraphEditAction& Action )
 	Diff->OnGraphChanged(this);
 }
 
-void FListItemGraphToDiff::KeyWasPressed( const FKeyEvent& InKeyEvent )
-{
-	if ( KeyCommands.IsValid() )
-	{
-		KeyCommands->ProcessCommandBindings(InKeyEvent);
-	}
-}
-
 FDiffPanel::FDiffPanel()
 {
 	Blueprint = NULL;
@@ -710,7 +667,6 @@ void SBlueprintDiff::Construct( const FArguments& InArgs)
 		.SetDisplayName(NSLOCTEXT("SBlueprintDiff", "MyBlueprintTabTitle", "My Blueprint"))
 		.SetTooltipText(NSLOCTEXT("SBlueprintDiff", "MyBlueprintTooltipText", "Differences in the 'My Blueprints' attributes of the blueprint"));
 
-	FDiffListCommands::Register();
 	check(InArgs._BlueprintOld && InArgs._BlueprintNew);
 	PanelOld.Blueprint = InArgs._BlueprintOld;
 	PanelNew.Blueprint = InArgs._BlueprintNew;
@@ -774,15 +730,93 @@ void SBlueprintDiff::Construct( const FArguments& InArgs)
 			];
 	}
 
-	const FDiffListCommands& Commands = FDiffListCommands::Get();
-	KeyCommands = MakeShareable(new FUICommandList);
+	FToolBarBuilder ToolbarBuilder(TSharedPtr< const FUICommandList >(), FMultiBoxCustomization::None);
+	ToolbarBuilder.AddToolBarButton(
+		FUIAction(
+			FExecuteAction::CreateSP(this, &SBlueprintDiff::PrevDiff),
+			FCanExecuteAction::CreateSP( this, &SBlueprintDiff::HasPrevDiff)
+		)
+		, NAME_None
+		, LOCTEXT("PrevDiffLabel", "Prev")
+		, LOCTEXT("PrevDiffTooltip", "Go to previous difference")
+		, FSlateIcon(FEditorStyle::GetStyleSetName()
+		, "BlueprintDif.PrevDiff")
+	);
+	ToolbarBuilder.AddToolBarButton(
+		FUIAction(
+			FExecuteAction::CreateSP(this, &SBlueprintDiff::NextDiff),
+			FCanExecuteAction::CreateSP(this, &SBlueprintDiff::HasNextDiff)
+		)
+		, NAME_None
+		, LOCTEXT("NextDiffLabel", "Next")
+		, LOCTEXT("NextDiffTooltip", "Go to next difference")
+		, FSlateIcon(FEditorStyle::GetStyleSetName(
+		), "BlueprintDif.NextDiff")
+	);
+	ToolbarBuilder.AddSeparator();
+	ToolbarBuilder.AddToolBarButton(
+		FUIAction(FExecuteAction::CreateSP(this, &SBlueprintDiff::OnToggleLockView2))
+		, NAME_None
+		, LOCTEXT("LockGraphsLabel", "Lock Graphs")
+		, LOCTEXT("LockGraphsTooltip", "Lock the various panels displayed in the graph view, so that they scroll together")
+		, TAttribute<FSlateIcon>(this, &SBlueprintDiff::GetLockViewImage2)
+	);
+
+	GraphPanel = GenerateGraphPanel();
+	DefaultsPanel = GenerateDefaultsPanel();
+	ComponentsPanel = GenerateComponentsPanel();
+
+	// Unfortunately we can't perform the diff until the UI is generated, the primary reason for this is that
+	// details customizations determine what is actually editable:
+	for (const auto& Graph : Graphs)
+	{
+		TArray< TSharedPtr<FBlueprintDifferenceTreeEntry> > Children;
+		for (const auto& Difference : Graph->DiffListSource)
+		{
+			auto ChildEntry = MakeShareable(
+				new FBlueprintDifferenceTreeEntry(
+				FOnDiffEntryFocused::CreateRaw(this, &SBlueprintDiff::OnDiffListSelectionChanged, Difference)
+				, FGenerateDiffEntryWidget::CreateSP(Difference.ToSharedRef(), &FDiffResultItem::GenerateWidget)
+				, TArray< TSharedPtr<FBlueprintDifferenceTreeEntry> >()
+				)
+			);
+			Children.Push(ChildEntry);
+		}
+
+		auto Entry = MakeShareable(
+			new FBlueprintDifferenceTreeEntry(
+				FOnDiffEntryFocused::CreateRaw(this, &SBlueprintDiff::OnSelectionChanged, Graph, ESelectInfo::Direct )
+				, FGenerateDiffEntryWidget::CreateSP(Graph.ToSharedRef(), &FListItemGraphToDiff::GenerateWidget)
+				, Children)
+			);
+		MasterDifferencesList.Push(Entry);
+	}
+
+	//DefaultsPanel.DiffControl->CalculateDifferences(MasterDifferencesList);
+	//ComponentsPanel.DiffControl->CalculateDifferences(MasterDifferencesList);
 
-	KeyCommands->MapAction(Commands.Previous, FExecuteAction::CreateSP(this, &SBlueprintDiff::PrevDiff), FCanExecuteAction::CreateSP( this, &SBlueprintDiff::HasPrevDiff) );
-	KeyCommands->MapAction(Commands.Next, FExecuteAction::CreateSP(this, &SBlueprintDiff::NextDiff), FCanExecuteAction::CreateSP( this, &SBlueprintDiff::HasNextDiff) );
+	const auto RowGenerator = []( TSharedPtr< FBlueprintDifferenceTreeEntry > Entry, const TSharedRef<STableViewBase>& Owner ) -> TSharedRef< ITableRow >
+	{
+		return SNew( STableRow<TSharedPtr<FBlueprintDifferenceTreeEntry> >, Owner )
+			[
+				Entry->GenerateWidget.Execute()
+			];
+	};
 
-	FToolBarBuilder ToolbarBuilder(KeyCommands.ToSharedRef(), FMultiBoxCustomization::None);
-	ToolbarBuilder.AddToolBarButton(Commands.Previous, NAME_None, TAttribute<FText>(), TAttribute<FText>(), FSlateIcon(FEditorStyle::GetStyleSetName(), "BlueprintDif.PrevDiff"));
-	ToolbarBuilder.AddToolBarButton(Commands.Next, NAME_None, TAttribute<FText>(), TAttribute<FText>(), FSlateIcon(FEditorStyle::GetStyleSetName(), "BlueprintDif.NextDiff"));
+	const auto ChildrenAccessor = [](TSharedPtr<FBlueprintDifferenceTreeEntry> InTreeItem, TArray< TSharedPtr< FBlueprintDifferenceTreeEntry > >& OutChildren, const TArray< TSharedPtr<FBlueprintDifferenceTreeEntry> >* MasterList )
+	{
+		OutChildren = InTreeItem->Children;
+	};
+
+	const auto Selector = [](TSharedPtr<FBlueprintDifferenceTreeEntry> InTreeItem, ESelectInfo::Type Type)
+	{
+		if( InTreeItem.IsValid() )
+		{
+			InTreeItem->OnFocus.ExecuteIfBound();
+		}
+	};
+	
+	const TArray< TSharedPtr<FBlueprintDifferenceTreeEntry> >* MasterDifferencesListPtr = &MasterDifferencesList;
 
 	this->ChildSlot
 		[
@@ -824,7 +858,25 @@ void SBlueprintDiff::Construct( const FArguments& InArgs)
 				]
 				+ SVerticalBox::Slot()
 				[
-					SAssignNew(ModeContents, SBorder)
+					SNew(SSplitter)
+					+ SSplitter::Slot()
+					.Value(.2f)
+					[
+						SNew(SBorder)
+						.BorderImage(FEditorStyle::GetBrush("ToolPanel.GroupBorder"))
+						[
+							SNew(STreeView< TSharedPtr< FBlueprintDifferenceTreeEntry > >)
+							.OnGenerateRow( STreeView< TSharedPtr< FBlueprintDifferenceTreeEntry > >::FOnGenerateRow::CreateStatic(RowGenerator) )
+							.OnGetChildren( STreeView< TSharedPtr< FBlueprintDifferenceTreeEntry > >::FOnGetChildren::CreateStatic(ChildrenAccessor, MasterDifferencesListPtr) )
+							.OnSelectionChanged( STreeView< TSharedPtr< FBlueprintDifferenceTreeEntry > >::FOnSelectionChanged::CreateStatic(Selector) )
+							.TreeItemsSource(&MasterDifferencesList)
+						]
+					]
+					+ SSplitter::Slot()
+					.Value(.8f)
+					[
+						SAssignNew(ModeContents, SBox)
+					]
 				]
 			]
 			
@@ -851,12 +903,12 @@ void SBlueprintDiff::CreateGraphEntry( class UEdGraph* GraphOld, class UEdGraph*
 
 void SBlueprintDiff::OnSelectionChanged( FGraphToDiff Item, ESelectInfo::Type SelectionType) 
 {
-	if(SelectionType != ESelectInfo::OnMouseClick || !Item.IsValid())
+	if(!Item.IsValid())
 	{
 		return;
 	}
 
-	FocusOnGraphRevisions(Item->GetGraphOld(), Item->GetGraphNew(), Item.Get());
+	FocusOnGraphRevisions(Item.Get());
 
 }
 
@@ -864,7 +916,7 @@ void SBlueprintDiff::OnGraphChanged(FListItemGraphToDiff* Diff)
 {
 	if(PanelNew.GraphEditor.IsValid() && PanelNew.GraphEditor.Pin()->GetCurrentGraph() == Diff->GetGraphNew())
 	{
-		FocusOnGraphRevisions(Diff->GetGraphOld(), Diff->GetGraphNew(), Diff);
+		FocusOnGraphRevisions(Diff);
 	}
 }
 
@@ -946,19 +998,31 @@ TSharedRef<SDockTab> SBlueprintDiff::CreateMyBlueprintsViews(const FSpawnTabArgs
 	];
 }
 
-void SBlueprintDiff::FocusOnGraphRevisions( class UEdGraph* GraphOld, class UEdGraph* GraphNew, FListItemGraphToDiff* Diff )
+FListItemGraphToDiff* SBlueprintDiff::FindGraphToDiffEntry(FName ByName)
 {
-	HandleGraphChanged( GraphOld ? GraphOld->GetFName() : GraphNew->GetFName() );
+	for( const auto& Graph : Graphs )
+	{
+		FName GraphName = Graph->GetGraphOld() ? Graph->GetGraphOld()->GetFName() : Graph->GetGraphNew()->GetFName();
+		if (GraphName == ByName )
+		{
+			return Graph.Get();
+		}
+	}
+	return nullptr;
+}
 
-	ResetGraphEditors();
+void SBlueprintDiff::FocusOnGraphRevisions( FListItemGraphToDiff* Diff )
+{
+	FName GraphName = Diff->GetGraphOld() ? Diff->GetGraphOld()->GetFName() : Diff->GetGraphNew()->GetFName();
+	HandleGraphChanged( GraphName );
 
-	DiffControl = Diff->AsShared();
-	//set new diff list
-	DiffListBorder->SetContent(Diff->GenerateDiffListWidget());
+	ResetGraphEditors();
 }
 
-void SBlueprintDiff::OnDiffListSelectionChanged(const TSharedPtr<struct FDiffResultItem>& TheDiff, FListItemGraphToDiff* GraphDiffer )
+void SBlueprintDiff::OnDiffListSelectionChanged(TSharedPtr<FDiffResultItem> TheDiff )
 {
+	check( TheDiff->Result.OwningGraph != FName() );
+	FocusOnGraphRevisions( FindGraphToDiffEntry( TheDiff->Result.OwningGraph ) );
 	FDiffSingleResult Result = TheDiff->Result;
 
 	const auto SafeClearSelection = []( TWeakPtr<SGraphEditor> GraphEditor )
@@ -991,6 +1055,11 @@ void SBlueprintDiff::OnDiffListSelectionChanged(const TSharedPtr<struct FDiffRes
 	}
 }
 
+void	SBlueprintDiff::OnToggleLockView2()
+{
+	OnToggleLockView();
+}
+
 FReply	SBlueprintDiff::OnToggleLockView()
 {
 	bLockViews = !bLockViews;
@@ -998,6 +1067,11 @@ FReply	SBlueprintDiff::OnToggleLockView()
 	return FReply::Handled();
 }
 
+FSlateIcon SBlueprintDiff::GetLockViewImage2() const
+{
+	return FSlateIcon(FEditorStyle::GetStyleSetName(), bLockViews ? "GenericLock" : "GenericUnlock");
+}
+
 const FSlateBrush* SBlueprintDiff::GetLockViewImage() const
 {
 	return bLockViews ? FEditorStyle::GetBrush("GenericLock") : FEditorStyle::GetBrush("GenericUnlock");
@@ -1198,6 +1272,11 @@ FDiffPanel& SBlueprintDiff::GetDiffPanelForNode(UEdGraphNode& Node)
 
 void SBlueprintDiff::HandleGraphChanged( const FName GraphName )
 {
+	if( CurrentMode != FBlueprintEditorApplicationModes::StandardBlueprintEditorMode)
+	{
+		SetCurrentMode(FBlueprintEditorApplicationModes::StandardBlueprintEditorMode);
+	}
+	
 	TArray<UEdGraph*> GraphsOld, GraphsNew;
 	PanelOld.Blueprint->GetAllGraphs(GraphsOld);
 	PanelNew.Blueprint->GetAllGraphs(GraphsNew);
@@ -1217,7 +1296,7 @@ void SBlueprintDiff::HandleGraphChanged( const FName GraphName )
 	PanelNew.GeneratePanel(GraphNew, GraphOld);
 }
 
-TSharedRef<SWidget> SBlueprintDiff::GenerateGraphPanel()
+SBlueprintDiff::FDiffControl SBlueprintDiff::GenerateGraphPanel()
 {
 	const TSharedRef<FTabManager::FLayout> DefaultLayout = FTabManager::NewLayout("BlueprintDiff_Layout_v1")
 	->AddArea
@@ -1243,140 +1322,66 @@ TSharedRef<SWidget> SBlueprintDiff::GenerateGraphPanel()
 			.IsPropertyEditingEnabledDelegate(FIsPropertyEditingEnabled::CreateStatic([] { return false; }));
 	};
 
+	SBlueprintDiff::FDiffControl Ret;
+
 	PanelOld.DetailsView = CreateInspector(PanelOld.MyBlueprint);
 	PanelOld.MyBlueprint->SetInspector(PanelOld.DetailsView);
 	PanelNew.DetailsView = CreateInspector(PanelNew.MyBlueprint);
 	PanelNew.MyBlueprint->SetInspector(PanelNew.DetailsView);
 
-	return SNew(SBorder)
+	Ret.Widget = SNew(SBorder)
 	.BorderImage(FEditorStyle::GetBrush("ToolPanel.GroupBorder"))
 	.Content()
 	[
-		SNew(SBorder)
+		SNew(SVerticalBox)
+		+ SVerticalBox::Slot()
+		.FillHeight(1.f)
 		[
-			SNew(SVerticalBox)
-			+ SVerticalBox::Slot()
-			.FillHeight(1.f)
+			SNew(SHorizontalBox)
+			+ SHorizontalBox::Slot()
+			.FillWidth(1.f)
 			[
+				//diff window
 				SNew(SSplitter)
-				+ SSplitter::Slot()
-				.Value(0.2f)
+				.Orientation(Orient_Vertical)
+				+SSplitter::Slot()
+				.Value(.8f)
 				[
-					SNew(SBorder)
-					[
-						SNew(SSplitter)
-						.Orientation(Orient_Vertical)
-						+ SSplitter::Slot()
-						.Value(0.3f)
-						[
-							SNew(SVerticalBox)
-							+ SVerticalBox::Slot()
-							.AutoHeight()
-							.HAlign(HAlign_Left)
-							[
-								//toggle lock button
-								SNew(SButton)
-								.OnClicked(this, &SBlueprintDiff::OnToggleLockView)
-								.Content()
-								[
-									SNew(SImage)
-									.Image(this, &SBlueprintDiff::GetLockViewImage)
-									.ToolTipText(LOCTEXT("BPDifLock", "Lock the blueprint views?"))
-								]
-							]
-							+ SVerticalBox::Slot()
-								.FillHeight(1.f)
-								[
-									//graph selection 
-									SNew(SBorder)
-									[
-										SNew(SVerticalBox)
-										+ SVerticalBox::Slot()
-										.AutoHeight()
-										[
-											SNew(SBorder)
-											.BorderImage(FEditorStyle::GetBrush("PropertyWindow.CategoryBackground"))
-											.Padding(FMargin(2.0f))
-											.ForegroundColor(FEditorStyle::GetColor("PropertyWindow.CategoryForeground"))
-											.ToolTipText(LOCTEXT("BlueprintDifGraphsToolTip", "Select Graph to Diff"))
-											.HAlign(HAlign_Center)
-											[
-												SNew(STextBlock)
-												.Text(LOCTEXT("BlueprintDifGraphs", "Graphs"))
-											]
-										]
-										+ SVerticalBox::Slot()
-											.FillHeight(1.f)
-											[
-												SAssignNew(GraphsToDiff, SListViewType)
-												.ItemHeight(24)
-												.ListItemsSource(&Graphs)
-												.OnGenerateRow(this, &SBlueprintDiff::OnGenerateRow)
-												.SelectionMode(ESelectionMode::Single)
-												.OnSelectionChanged(this, &SBlueprintDiff::OnSelectionChanged)
-											]
-									]
-								]
-						]
-						+ SSplitter::Slot()
-							.Value(0.7f)
-							[
-								SNew(SVerticalBox)
-								+ SVerticalBox::Slot()
-								.FillHeight(1.f)
-								[
-									SAssignNew(DiffListBorder, SBorder)
-								]
-							]
-					]
+					// graph and my blueprint views:
+					TabControl
 				]
 				+ SSplitter::Slot()
-				.Value(0.8f)
+				.Value(.2f)
 				[
-					SNew(SHorizontalBox)
-					+ SHorizontalBox::Slot()
-					.FillWidth(1.f)
+					SNew( SSplitter )
+					+SSplitter::Slot()
 					[
-						//diff window
-						SNew(SSplitter)
-						.Orientation(Orient_Vertical)
-						+SSplitter::Slot()
-						.Value(.8f)
-						[
-							// graph and my blueprint views:
-							TabControl
-						]
-						+ SSplitter::Slot()
-						.Value(.2f)
-						[
-							SNew( SSplitter )
-							+SSplitter::Slot()
-							[
-								PanelOld.DetailsView.ToSharedRef()
-							]
-							+ SSplitter::Slot()
-							[
-								PanelNew.DetailsView.ToSharedRef()
-							]
-						]
+						PanelOld.DetailsView.ToSharedRef()
+					]
+					+ SSplitter::Slot()
+					[
+						PanelNew.DetailsView.ToSharedRef()
 					]
 				]
 			]
 		]
 	];
+
+	return Ret;
 }
 
-TSharedRef<SWidget> SBlueprintDiff::GenerateDefaultsPanel()
+SBlueprintDiff::FDiffControl SBlueprintDiff::GenerateDefaultsPanel()
 {
 	const UObject* A = DiffUtils::GetCDO(PanelOld.Blueprint);
 	const UObject* B = DiffUtils::GetCDO(PanelNew.Blueprint);
 	TArray<FPropertySoftPath> DifferingProperties;
 	DiffUtils::CompareUnrelatedObjects( A, B, DifferingProperties);
 
+	SBlueprintDiff::FDiffControl Ret;
 	auto NewDiffControl = TSharedPtr<FCDODiffControl>(new FCDODiffControl(A, B, DifferingProperties) );
 	//Splitter for left and right blueprint. Current convention is for the local (probably newer?) blueprint to be on the right:
-	DiffControl = NewDiffControl;
-	return SNew(SSplitter)
+	Ret.DiffControl = NewDiffControl;
+	Ret.Widget = SNew(SSplitter)
 		+ SSplitter::Slot()
 		.Value(0.5f)
 		[
@@ -1395,14 +1400,18 @@ TSharedRef<SWidget> SBlueprintDiff::GenerateDefaultsPanel()
 				NewDiffControl->NewDetailsWidget()
 			]
 		];
+
+	return Ret;
 }
 
-TSharedRef<SWidget> SBlueprintDiff::GenerateComponentsPanel()
+SBlueprintDiff::FDiffControl SBlueprintDiff::GenerateComponentsPanel()
 {
+	SBlueprintDiff::FDiffControl Ret;
+
 	//Splitter for left and right blueprint. Current convention is for the local (probably newer?) blueprint to be on the right:
 	auto NewDiffControl = TSharedPtr<FSCSDiffControl>(new FSCSDiffControl(PanelOld.Blueprint, PanelNew.Blueprint) );
-	DiffControl = NewDiffControl;
-	return SNew(SSplitter)
+	Ret.DiffControl = NewDiffControl;
+	Ret.Widget = SNew(SSplitter)
 		+ SSplitter::Slot()
 		.Value(0.5f)
 		[
@@ -1421,6 +1430,8 @@ TSharedRef<SWidget> SBlueprintDiff::GenerateComponentsPanel()
 				NewDiffControl->NewTreeWidget()
 			]
 		];
+
+	return Ret;
 }
 
 void SBlueprintDiff::SetCurrentMode(FName NewMode)
@@ -1430,15 +1441,18 @@ void SBlueprintDiff::SetCurrentMode(FName NewMode)
 	DiffControl = TSharedPtr< IDiffControl >();
 	if( NewMode == FBlueprintEditorApplicationModes::StandardBlueprintEditorMode)
 	{
-		ModeContents->SetContent( GenerateGraphPanel() );
+		DiffControl = GraphPanel.DiffControl;
+		ModeContents->SetContent( GraphPanel.Widget.ToSharedRef() );
 	}
 	else if( NewMode == FBlueprintEditorApplicationModes::BlueprintDefaultsMode )
 	{
-		ModeContents->SetContent( GenerateDefaultsPanel() );
+		DiffControl = DefaultsPanel.DiffControl;
+		ModeContents->SetContent( DefaultsPanel.Widget.ToSharedRef() );
 	}
 	else if (NewMode == FBlueprintEditorApplicationModes::BlueprintComponentsMode)
 	{
-		ModeContents->SetContent( GenerateComponentsPanel() );
+		DiffControl = ComponentsPanel.DiffControl;
+		ModeContents->SetContent( ComponentsPanel.Widget.ToSharedRef() );
 	}
 	else
 	{
@@ -1446,16 +1460,5 @@ void SBlueprintDiff::SetCurrentMode(FName NewMode)
 	}
 }
 
-FReply SBlueprintDiff::OnKeyDown( const FGeometry& MyGeometry, const FKeyEvent& InKeyEvent ) 
-{
-	TArray< FGraphToDiff>  Selected = GraphsToDiff->GetSelectedItems();
-	for(auto It(Selected.CreateIterator());It;++It)
-	{
-		FGraphToDiff& Diff = *It;
-		Diff->KeyWasPressed(InKeyEvent);
-	}
-	return FReply::Handled();
-}
-
 #undef LOCTEXT_NAMESPACE
 
diff --git a/Engine/Source/Editor/Kismet/Public/DiffUtils.h b/Engine/Source/Editor/Kismet/Public/DiffUtils.h
--- a/Engine/Source/Editor/Kismet/Public/DiffUtils.h
+++ b/Engine/Source/Editor/Kismet/Public/DiffUtils.h
@@ -151,3 +151,22 @@ namespace DiffUtils
 	KISMET_API TArray<FPropertyPath> ResolveAll(const UObject* Object, const TArray<FPropertySoftPath>& InSoftProperties);
 
 }
+
+DECLARE_DELEGATE(FOnDiffEntryFocused);
+DECLARE_DELEGATE_RetVal(TSharedRef<SWidget>, FGenerateDiffEntryWidget);
+
+class FBlueprintDifferenceTreeEntry
+{
+public:
+	FBlueprintDifferenceTreeEntry( FOnDiffEntryFocused InOnFocus, FGenerateDiffEntryWidget InGenerateWidget, TArray< TSharedPtr<FBlueprintDifferenceTreeEntry> > InChildren ) 
+		: OnFocus(InOnFocus)
+		, GenerateWidget(InGenerateWidget)
+		, Children(InChildren) 
+	{
+		check( InGenerateWidget.IsBound() );
+	}
+
+	FOnDiffEntryFocused OnFocus;
+	FGenerateDiffEntryWidget GenerateWidget;
+	TArray< TSharedPtr<FBlueprintDifferenceTreeEntry> > Children;
+};
diff --git a/Engine/Source/Editor/Kismet/Public/SBlueprintDiff.h b/Engine/Source/Editor/Kismet/Public/SBlueprintDiff.h
--- a/Engine/Source/Editor/Kismet/Public/SBlueprintDiff.h
+++ b/Engine/Source/Editor/Kismet/Public/SBlueprintDiff.h
@@ -135,15 +135,14 @@ protected:
 	TSharedRef<SDockTab> CreateGraphDiffViews( const FSpawnTabArgs& Args );
 	TSharedRef<SDockTab> CreateMyBlueprintsViews( const FSpawnTabArgs& Args );
 
-	/* Need to process keys for shortcuts to buttons */
-	virtual FReply OnKeyDown( const FGeometry& MyGeometry, const FKeyEvent& InKeyEvent ) override;
-
-
 	typedef TSharedPtr<struct FListItemGraphToDiff>	FGraphToDiff;
 	typedef SListView<FGraphToDiff >	SListViewType;
 
+	/** Find the FListItemGraphToDiff that displays the graph with GraphName: */
+	FListItemGraphToDiff* FindGraphToDiffEntry(FName ByName);
+
 	/** Bring these revisions of graph into focus on main display*/
-	void FocusOnGraphRevisions( class UEdGraph* GraphOld, class UEdGraph* GraphNew , struct FListItemGraphToDiff* Diff);
+	void FocusOnGraphRevisions( struct FListItemGraphToDiff* Diff);
 
 	/*Create a list item entry graph that exists in at least one of the blueprints */
 	void CreateGraphEntry(class UEdGraph* GraphOld, class UEdGraph* GraphNew);
@@ -155,18 +154,20 @@ protected:
 	void OnSelectionChanged(FGraphToDiff Item, ESelectInfo::Type SelectionType);
 
 	/** Called when user clicks on an entry in the listview of differences */
-	void OnDiffListSelectionChanged(const TSharedPtr<struct FDiffResultItem>& TheDiff, struct FListItemGraphToDiff* GraphDiffer);
+	void OnDiffListSelectionChanged(TSharedPtr<struct FDiffResultItem> TheDiff );
 		
 	/** Disable the focus on a particular pin */
 	void DisablePinDiffFocus();
 
 	/*User toggles the option to lock the views between the two blueprints */
+	void	OnToggleLockView2();
 	FReply	OnToggleLockView();
 
 	/*Reset the graph editor, called when user switches graphs to display*/
 	void ResetGraphEditors();
 
 	/*Get the image to show for the toggle lock option*/
+	FSlateIcon GetLockViewImage2() const;
 	const FSlateBrush* GetLockViewImage() const;
 
 	/* This buffer stores the currently displayed results */
@@ -178,9 +179,21 @@ protected:
 	/** Event handler that updates the graph view when user selects a new graph */
 	void HandleGraphChanged( const FName GraphName );
 
-	TSharedRef<SWidget> GenerateGraphPanel();
-	TSharedRef<SWidget> GenerateDefaultsPanel();
-	TSharedRef<SWidget> GenerateComponentsPanel();
+	struct FDiffControl
+	{
+		FDiffControl()
+		: Widget()
+		, DiffControl(NULL)
+		{
+		}
+
+		TSharedPtr<SWidget> Widget;
+		TSharedPtr< class IDiffControl > DiffControl;
+	};
+
+	FDiffControl GenerateGraphPanel();
+	FDiffControl GenerateDefaultsPanel();
+	FDiffControl GenerateComponentsPanel();
 
 	/** Accessor and event handler for toggling between diff view modes (defaults, components, graph view, interface, macro): */
 	void SetCurrentMode(FName NewMode);
@@ -194,25 +207,24 @@ protected:
 	/** If the two views should be locked */
 	bool	bLockViews;
 
-	/** Border Widget, inside is the current graphs being diffed, we can replace content to change the graph*/
-	TSharedPtr<SBorder>	DiffListBorder;
-
 	/** Contents widget that we swap when mode changes (defaults, components, etc) */
-	TSharedPtr<SBorder> ModeContents;
-
-	/** The ListView containing the graphs the user can select */
-	TSharedPtr<SListViewType>	GraphsToDiff;
+	TSharedPtr<SBox> ModeContents;
 
 	friend struct FListItemGraphToDiff;
 
-	/** Key commands processed by this widget */
-	TSharedPtr< FUICommandList > KeyCommands;
-
 	/** Helper class for highlighting diffs in different types of controls (graph view, details view, etc) */
 	TSharedPtr< class IDiffControl > DiffControl;
 
 	/** We can't use the global tab manager because we need to instance the diff control, so we have our own tab manager: */
 	TSharedPtr<FTabManager> TabManager;
+
+	/** Tree of differences collected across all panels: */
+	TArray< TSharedPtr<class FBlueprintDifferenceTreeEntry> > MasterDifferencesList;
+
+	/** Stored references to widgets used to display various parts of a blueprint: */
+	FDiffControl GraphPanel;
+	FDiffControl DefaultsPanel;
+	FDiffControl ComponentsPanel;
 };
 
 
