Adding and removing interfaces from Blueprints, as well as reparenting, will no longer drop the current details view.

diff --git a/Engine/Source/Editor/Kismet/Private/BlueprintDetailsCustomization.cpp b/Engine/Source/Editor/Kismet/Private/BlueprintDetailsCustomization.cpp
--- a/Engine/Source/Editor/Kismet/Private/BlueprintDetailsCustomization.cpp
+++ b/Engine/Source/Editor/Kismet/Private/BlueprintDetailsCustomization.cpp
@@ -3895,6 +3895,8 @@ void FBlueprintInterfaceLayout::OnRemoveInterface(FInterfaceName InterfaceName)
 	FBlueprintEditorUtils::RemoveInterface(Blueprint, InterfaceFName, bPreserveInterfaceFunctions);
 
 	RegenerateChildrenDelegate.ExecuteIfBound();
+
+	OnRefreshInDetailsView();
 }
 
 void FBlueprintInterfaceLayout::OnClassPicked(UClass* PickedClass)
@@ -3910,6 +3912,8 @@ void FBlueprintInterfaceLayout::OnClassPicked(UClass* PickedClass)
 	FBlueprintEditorUtils::ImplementNewInterface( Blueprint, PickedClass->GetFName() );
 
 	RegenerateChildrenDelegate.ExecuteIfBound();
+
+	OnRefreshInDetailsView();
 }
 
 TSharedRef<SWidget> FBlueprintInterfaceLayout::OnGetAddInterfaceMenuContent()
@@ -3938,6 +3942,16 @@ TSharedRef<SWidget> FBlueprintInterfaceLayout::OnGetAddInterfaceMenuContent()
 		];
 }
 
+void FBlueprintInterfaceLayout::OnRefreshInDetailsView()
+{
+	TSharedPtr<SKismetInspector> Inspector = GlobalOptionsDetailsPtr.Pin()->GetBlueprintEditorPtr().Pin()->GetInspector();
+	UBlueprint* Blueprint = GlobalOptionsDetailsPtr.Pin()->GetBlueprintObj();
+	check(Blueprint);
+
+	// Show details for the Blueprint instance we're editing
+	Inspector->ShowDetailsForSingleObject(Blueprint);
+}
+
 UBlueprint* FBlueprintGlobalOptionsDetails::GetBlueprintObj() const
 {
 	if(BlueprintEditorPtr.IsValid())
@@ -3992,6 +4006,11 @@ void FBlueprintGlobalOptionsDetails::OnClassPicked(UClass* PickedClass)
 	{
 		BlueprintEditorPtr.Pin()->ReparentBlueprint_NewParentChosen(PickedClass);
 	}
+
+	check(BlueprintEditorPtr.IsValid());
+	TSharedPtr<SKismetInspector> Inspector = BlueprintEditorPtr.Pin()->GetInspector();
+	// Show details for the Blueprint instance we're editing
+	Inspector->ShowDetailsForSingleObject(GetBlueprintObj());
 }
 
 bool FBlueprintGlobalOptionsDetails::CanDeprecateBlueprint() const
diff --git a/Engine/Source/Editor/Kismet/Private/BlueprintDetailsCustomization.h b/Engine/Source/Editor/Kismet/Private/BlueprintDetailsCustomization.h
--- a/Engine/Source/Editor/Kismet/Private/BlueprintDetailsCustomization.h
+++ b/Engine/Source/Editor/Kismet/Private/BlueprintDetailsCustomization.h
@@ -551,6 +551,9 @@ private:
 	/** Callback function when an interface class is picked */
 	void OnClassPicked(UClass* PickedClass);
 
+	/** Helper function to set the Blueprint back into the KismetInspector's details view */
+	void OnRefreshInDetailsView();
+
 private:
 	/** The parent graph action details customization */
 	TWeakPtr<class FBlueprintGlobalOptionsDetails> GlobalOptionsDetailsPtr;
