Blueprint's debugged world combo box doesn't update it

diff --git a/Engine/Source/Editor/Kismet/Private/SBlueprintEditorSelectedDebugObjectWidget.cpp b/Engine/Source/Editor/Kismet/Private/SBlueprintEditorSelectedDebugObjectWidget.cpp
--- a/Engine/Source/Editor/Kismet/Private/SBlueprintEditorSelectedDebugObjectWidget.cpp
+++ b/Engine/Source/Editor/Kismet/Private/SBlueprintEditorSelectedDebugObjectWidget.cpp
@@ -358,7 +358,16 @@ void SBlueprintEditorSelectedDebugObjectWidget::GenerateDebugObjectNames(bool bR
 	if (DebugWorldsComboBox.IsValid())
 	{
 		TSharedPtr<FString> CurrentWorldSelection = DebugWorldsComboBox->GetSelectedItem();
-		int32 SelectedIndex = DebugWorldNames.Find(CurrentWorldSelection);
+		int32 SelectedIndex = INDEX_NONE;
+		for (int32 WorldIdx = 0; WorldIdx < DebugWorldNames.Num(); ++WorldIdx)
+		{
+			if (DebugWorldNames[WorldIdx].IsValid() && CurrentWorldSelection.IsValid()
+				&& (*DebugWorldNames[WorldIdx] == *CurrentWorldSelection))
+			{
+				SelectedIndex = WorldIdx;
+				break;
+			}
+		}
 		if (SelectedIndex > 0 && DebugWorldNames.IsValidIndex(SelectedIndex))
 		{
 			DebugWorld = DebugWorlds[SelectedIndex].Get();
