[UE-6238] restore placeable Add Component actions to the Blueprint graph context menu/palette (regression)

change summary:
- added a UK2Node_AddComponent::CanPasteHere() override to bypass the base UK2Node_CallFunction logic; the base UK2Node method already handles the necessary checking for this node type.
- modified CanPaletteItemBePlaced() in SBlueprintSubPalette.cpp to bypass the CanFunctionBePlacedInGraph() check for UK2Node_AddComponent node types; thus "Add Component" actions are now also placeable in Actor-based graphs via drag-and-drop.

diff --git a/Engine/Source/Editor/BlueprintGraph/Classes/K2Node_AddComponent.h b/Engine/Source/Editor/BlueprintGraph/Classes/K2Node_AddComponent.h
--- a/Engine/Source/Editor/BlueprintGraph/Classes/K2Node_AddComponent.h
+++ b/Engine/Source/Editor/BlueprintGraph/Classes/K2Node_AddComponent.h
@@ -27,6 +27,7 @@ class UK2Node_AddComponent : public UK2Node_CallFunction
 	virtual FString GetDocumentationLink() const override;
 	virtual FString GetDocumentationExcerptName() const override;
 	virtual bool IsCompatibleWithGraph(UEdGraph const* Graph) const override;
+	virtual bool CanPasteHere(const UEdGraph* TargetGraph) const override;
 	// End UEdGraphNode interface
 
 	// Begin K2Node interface
diff --git a/Engine/Source/Editor/BlueprintGraph/Private/K2Node_AddComponent.cpp b/Engine/Source/Editor/BlueprintGraph/Private/K2Node_AddComponent.cpp
--- a/Engine/Source/Editor/BlueprintGraph/Private/K2Node_AddComponent.cpp
+++ b/Engine/Source/Editor/BlueprintGraph/Private/K2Node_AddComponent.cpp
@@ -422,4 +422,10 @@ void UK2Node_AddComponent::ExpandNode(class FKismetCompilerContext& CompilerCont
 	}
 }
 
+bool UK2Node_AddComponent::CanPasteHere(const UEdGraph* TargetGraph) const
+{
+	// We want to bypass the base UK2Node_CallFunction override for this node type - we will always allow the function this node represents to be placed in a compatible graph.
+	return UK2Node::CanPasteHere(TargetGraph);
+}
+
 #undef LOCTEXT_NAMESPACE
diff --git a/Engine/Source/Editor/Kismet/Private/SBlueprintSubPalette.cpp b/Engine/Source/Editor/Kismet/Private/SBlueprintSubPalette.cpp
--- a/Engine/Source/Editor/Kismet/Private/SBlueprintSubPalette.cpp
+++ b/Engine/Source/Editor/Kismet/Private/SBlueprintSubPalette.cpp
@@ -88,14 +88,19 @@ static bool CanPaletteItemBePlaced(TSharedPtr<FEdGraphSchemaAction> DropActionIn
 			}
 			else
 			{
-				uint32 AllowedFunctionTypes = UEdGraphSchema_K2::EFunctionType::FT_Pure | UEdGraphSchema_K2::EFunctionType::FT_Const | UEdGraphSchema_K2::EFunctionType::FT_Protected;
-				if(K2Schema->DoesGraphSupportImpureFunctions(HoveredGraphIn))
-				{
-					AllowedFunctionTypes |= UEdGraphSchema_K2::EFunctionType::FT_Imperative;
+				// @TODO - Consolidate this as a call to UK2Node::IsActionFilteredOut() here instead? Would need to add the ImpededReason as an 'out' param to that API first.
+				//  	   We could then also skip the additonal 'CanPasteHere' check below in that case as it would be redundant for CallFunction node types specifically.
+				if(!NodeToBePlaced->IsA<UK2Node_AddComponent>())
+				{ 
+					uint32 AllowedFunctionTypes = UEdGraphSchema_K2::EFunctionType::FT_Pure | UEdGraphSchema_K2::EFunctionType::FT_Const | UEdGraphSchema_K2::EFunctionType::FT_Protected;
+					if(K2Schema->DoesGraphSupportImpureFunctions(HoveredGraphIn))
+					{
+						AllowedFunctionTypes |= UEdGraphSchema_K2::EFunctionType::FT_Imperative;
+					}
+
+					const UClass* GeneratedClass = FBlueprintEditorUtils::FindBlueprintForGraphChecked(HoveredGraphIn)->GeneratedClass;
+					bCanBePlaced = K2Schema->CanFunctionBeUsedInGraph(GeneratedClass, Function, HoveredGraphIn, AllowedFunctionTypes, false, FFunctionTargetInfo(), &ImpededReasonOut);
 				}
-
-				const UClass* GeneratedClass = FBlueprintEditorUtils::FindBlueprintForGraphChecked(HoveredGraphIn)->GeneratedClass;
-				bCanBePlaced = K2Schema->CanFunctionBeUsedInGraph(GeneratedClass, Function, HoveredGraphIn, AllowedFunctionTypes, false, FFunctionTargetInfo(), &ImpededReasonOut);
 			}
 		}
 		else if (UK2Node_Event const* EventNode = Cast<UK2Node_Event const>(NodeToBePlaced))
