Duplicating blueprint component variables (that are autogenerated from the components mode) will no longer crash.

Editing these variable types is no longer enabled (would do nothing).

#jra UE-6136 - BLUEPRINTS: Crash when duplicating a component in the graph view.

diff --git a/Engine/Source/Editor/Kismet/Private/BlueprintDetailsCustomization.cpp b/Engine/Source/Editor/Kismet/Private/BlueprintDetailsCustomization.cpp
--- a/Engine/Source/Editor/Kismet/Private/BlueprintDetailsCustomization.cpp
+++ b/Engine/Source/Editor/Kismet/Private/BlueprintDetailsCustomization.cpp
@@ -676,15 +676,9 @@ bool FBlueprintVarActionDetails::GetVariableNameChangeEnabled() const
 		}
 		else if(IsAComponentVariable(VariableProperty) && Blueprint->SimpleConstructionScript != NULL)
 		{
-			TArray<USCS_Node*> Nodes = Blueprint->SimpleConstructionScript->GetAllNodes();
-			for (TArray<USCS_Node*>::TConstIterator NodeIt(Nodes); NodeIt; ++NodeIt)
+			if (USCS_Node* Node = Blueprint->SimpleConstructionScript->FindSCSNode(GetVariableName()))
 			{
-				USCS_Node* Node = *NodeIt;
-				if (Node->VariableName == GetVariableName())
-				{
-					bIsReadOnly = !Node->IsValidVariableNameString(Node->VariableName.ToString());
-					break;
-				}
+				bIsReadOnly = !Node->IsValidVariableNameString(Node->VariableName.ToString());
 			}
 		}
 		else if(IsALocalVariable(VariableProperty))
@@ -798,7 +792,15 @@ bool FBlueprintVarActionDetails::GetVariableTypeChangeEnabled() const
 	UProperty* VariableProperty = SelectionAsProperty();
 	if(VariableProperty)
 	{
-		return GetBlueprintObj()->SkeletonGeneratedClass->GetAuthoritativeClass() == VariableProperty->GetOwnerClass()->GetAuthoritativeClass();
+		if(GetBlueprintObj()->SkeletonGeneratedClass->GetAuthoritativeClass() != VariableProperty->GetOwnerClass()->GetAuthoritativeClass())
+		{
+			return false;
+		}
+		// If the variable belongs to this class and cannot be found in the member variable list, it is not editable (it may be a component)
+		if (FBlueprintEditorUtils::FindNewVariableIndex(GetBlueprintObj(), GetVariableName()) == INDEX_NONE)
+		{
+			return false;
+		}
 	}
 	return true;
 }
diff --git a/Engine/Source/Editor/Kismet/Private/SMyBlueprint.cpp b/Engine/Source/Editor/Kismet/Private/SMyBlueprint.cpp
--- a/Engine/Source/Editor/Kismet/Private/SMyBlueprint.cpp
+++ b/Engine/Source/Editor/Kismet/Private/SMyBlueprint.cpp
@@ -2027,7 +2027,16 @@ void SMyBlueprint::OnDuplicateAction()
 		const FScopedTransaction Transaction( LOCTEXT( "DuplicateVariable", "Duplicate Variable" ) );
 		GetBlueprintObj()->Modify();
 
-		DuplicateActionName = FBlueprintEditorUtils::DuplicateVariable(GetBlueprintObj(), nullptr, VarAction->GetVariableName());
+		if(FBlueprintEditorUtils::FindNewVariableIndex(GetBlueprintObj(), VarAction->GetVariableName()) != INDEX_NONE)
+		{
+			DuplicateActionName = FBlueprintEditorUtils::DuplicateVariable(GetBlueprintObj(), nullptr, VarAction->GetVariableName());
+		}
+		else
+		{
+			FEdGraphPinType VarPinType;
+			GetDefault<UEdGraphSchema_K2>()->ConvertPropertyToPinType(VarAction->GetProperty(), VarPinType);
+			FBlueprintEditorUtils::AddMemberVariable(GetBlueprintObj(), FBlueprintEditorUtils::FindUniqueKismetName(Blueprint, VarAction->GetVariableName().ToString()), VarPinType);
+		}
 	}
 	else if (FEdGraphSchemaAction_K2LocalVar* LocalVarAction = SelectionAsLocalVar())
 	{
diff --git a/Engine/Source/Runtime/Engine/Classes/Engine/SimpleConstructionScript.h b/Engine/Source/Runtime/Engine/Classes/Engine/SimpleConstructionScript.h
--- a/Engine/Source/Runtime/Engine/Classes/Engine/SimpleConstructionScript.h
+++ b/Engine/Source/Runtime/Engine/Classes/Engine/SimpleConstructionScript.h
@@ -58,6 +58,9 @@ class USimpleConstructionScript : public UObject
 	/** Find the parent node of this one. Returns NULL if node is not in tree or if is root */
 	ENGINE_API USCS_Node* FindParentNode(USCS_Node* InNode) const;
 
+	/** Find the SCS_Node node by name and return it if found */
+	ENGINE_API USCS_Node* FindSCSNode(FName InName);
+
 protected:
 	/** Checks the root node set for scene components and ensures that it is valid (e.g. after a removal) */
 	void ValidateSceneRootNodes();
diff --git a/Engine/Source/Runtime/Engine/Private/SimpleConstructionScript.cpp b/Engine/Source/Runtime/Engine/Private/SimpleConstructionScript.cpp
--- a/Engine/Source/Runtime/Engine/Private/SimpleConstructionScript.cpp
+++ b/Engine/Source/Runtime/Engine/Private/SimpleConstructionScript.cpp
@@ -582,6 +582,22 @@ USCS_Node* USimpleConstructionScript::FindParentNode(USCS_Node* InNode) const
 	return NULL;
 }
 
+USCS_Node* USimpleConstructionScript::FindSCSNode(FName InName)
+{
+	TArray<USCS_Node*> AllNodes = GetAllNodes();
+	USCS_Node* ReturnSCSNode = nullptr;
+
+	for( USCS_Node* SCSNode : AllNodes )
+	{
+		if (SCSNode->GetVariableName() == InName)
+		{
+			ReturnSCSNode = SCSNode;
+			break;
+		}
+	}
+	return ReturnSCSNode;
+}
+
 void USimpleConstructionScript::ValidateSceneRootNodes()
 {
 #if WITH_EDITOR
