Fix so Blueprint function local variables are editable again, broke in CL# 2376408.

diff --git a/Engine/Source/Editor/Kismet/Private/BlueprintDetailsCustomization.cpp b/Engine/Source/Editor/Kismet/Private/BlueprintDetailsCustomization.cpp
--- a/Engine/Source/Editor/Kismet/Private/BlueprintDetailsCustomization.cpp
+++ b/Engine/Source/Editor/Kismet/Private/BlueprintDetailsCustomization.cpp
@@ -790,7 +790,7 @@ void FBlueprintVarActionDetails::OnVarNameCommitted(const FText& InNewText, ETex
 bool FBlueprintVarActionDetails::GetVariableTypeChangeEnabled() const
 {
 	UProperty* VariableProperty = SelectionAsProperty();
-	if(VariableProperty)
+	if(VariableProperty && !IsALocalVariable(VariableProperty))
 	{
 		if(GetBlueprintObj()->SkeletonGeneratedClass->GetAuthoritativeClass() != VariableProperty->GetOwnerClass()->GetAuthoritativeClass())
 		{
