[UE-6331] restore input axis accessor functions to BP graph context menu/palette (regression)

change summary:
- reverted previous change to UK2Node_AddComponent; will need to be applicable to all derived UK2Node_CallFunction types
- modified UK2Node_CallFunction::CanPasteNode() to bypass the general function placement check for derived UK2Node_CallFunction types (derived types are considered placeable if compatible)
- modified CanPaletteItemBePlaced() in SBlueprintSubPalette.cpp to skip over the general function placement check for derived UK2Node_CallFunction types (and not just UK2Node_AddComponent)

diff --git a/Engine/Source/Editor/BlueprintGraph/Classes/K2Node_AddComponent.h b/Engine/Source/Editor/BlueprintGraph/Classes/K2Node_AddComponent.h
--- a/Engine/Source/Editor/BlueprintGraph/Classes/K2Node_AddComponent.h
+++ b/Engine/Source/Editor/BlueprintGraph/Classes/K2Node_AddComponent.h
@@ -27,7 +27,6 @@ class UK2Node_AddComponent : public UK2Node_CallFunction
 	virtual FString GetDocumentationLink() const override;
 	virtual FString GetDocumentationExcerptName() const override;
 	virtual bool IsCompatibleWithGraph(UEdGraph const* Graph) const override;
-	virtual bool CanPasteHere(const UEdGraph* TargetGraph) const override;
 	// End UEdGraphNode interface
 
 	// Begin K2Node interface
diff --git a/Engine/Source/Editor/BlueprintGraph/Private/K2Node_AddComponent.cpp b/Engine/Source/Editor/BlueprintGraph/Private/K2Node_AddComponent.cpp
--- a/Engine/Source/Editor/BlueprintGraph/Private/K2Node_AddComponent.cpp
+++ b/Engine/Source/Editor/BlueprintGraph/Private/K2Node_AddComponent.cpp
@@ -422,10 +422,4 @@ void UK2Node_AddComponent::ExpandNode(class FKismetCompilerContext& CompilerCont
 	}
 }
 
-bool UK2Node_AddComponent::CanPasteHere(const UEdGraph* TargetGraph) const
-{
-	// We want to bypass the base UK2Node_CallFunction override for this node type - we will always allow the function this node represents to be placed in a compatible graph.
-	return UK2Node::CanPasteHere(TargetGraph);
-}
-
 #undef LOCTEXT_NAMESPACE
diff --git a/Engine/Source/Editor/BlueprintGraph/Private/K2Node_CallFunction.cpp b/Engine/Source/Editor/BlueprintGraph/Private/K2Node_CallFunction.cpp
--- a/Engine/Source/Editor/BlueprintGraph/Private/K2Node_CallFunction.cpp
+++ b/Engine/Source/Editor/BlueprintGraph/Private/K2Node_CallFunction.cpp
@@ -1014,8 +1014,12 @@ bool UK2Node_CallFunction::CanFunctionSupportMultipleTargets(UFunction const* Fu
 
 bool UK2Node_CallFunction::CanPasteHere(const UEdGraph* TargetGraph) const
 {
+	// Basic check for graph compatibility, etc.
 	bool bCanPaste = Super::CanPasteHere(TargetGraph);
-	if(bCanPaste)
+
+	// We check function context for placability only in the base class case; derived classes are typically bound to
+	// specific functions that should always be placeable, but may not always be explicitly callable (e.g. InternalUseOnly).
+	if(bCanPaste && GetClass() == StaticClass())
 	{
 		const UEdGraphSchema_K2* K2Schema = GetDefault<UEdGraphSchema_K2>();
 		uint32 AllowedFunctionTypes = UEdGraphSchema_K2::EFunctionType::FT_Pure | UEdGraphSchema_K2::EFunctionType::FT_Const | UEdGraphSchema_K2::EFunctionType::FT_Protected;
diff --git a/Engine/Source/Editor/Kismet/Private/SBlueprintSubPalette.cpp b/Engine/Source/Editor/Kismet/Private/SBlueprintSubPalette.cpp
--- a/Engine/Source/Editor/Kismet/Private/SBlueprintSubPalette.cpp
+++ b/Engine/Source/Editor/Kismet/Private/SBlueprintSubPalette.cpp
@@ -88,9 +88,10 @@ static bool CanPaletteItemBePlaced(TSharedPtr<FEdGraphSchemaAction> DropActionIn
 			}
 			else
 			{
+				// Note: We only check function context for UK2Node_CallFunction types specifically; derivatives are typically bound to specific functions that should be placeable but may not be explicitly callable (e.g. InternalUseOnly).
 				// @TODO - Consolidate this as a call to UK2Node::IsActionFilteredOut() here instead? Would need to add the ImpededReason as an 'out' param to that API first.
 				//  	   We could then also skip the additonal 'CanPasteHere' check below in that case as it would be redundant for CallFunction node types specifically.
-				if(!NodeToBePlaced->IsA<UK2Node_AddComponent>())
+				if(NodeToBePlaced->GetClass() == UK2Node_CallFunction::StaticClass())
 				{ 
 					uint32 AllowedFunctionTypes = UEdGraphSchema_K2::EFunctionType::FT_Pure | UEdGraphSchema_K2::EFunctionType::FT_Const | UEdGraphSchema_K2::EFunctionType::FT_Protected;
 					if(K2Schema->DoesGraphSupportImpureFunctions(HoveredGraphIn))
