BP ui crash in SMyBlueprint::GetLocalVariables

diff --git a/Engine/Source/Editor/Kismet/Private/SMyBlueprint.cpp b/Engine/Source/Editor/Kismet/Private/SMyBlueprint.cpp
--- a/Engine/Source/Editor/Kismet/Private/SMyBlueprint.cpp
+++ b/Engine/Source/Editor/Kismet/Private/SMyBlueprint.cpp
@@ -666,10 +666,13 @@ void SMyBlueprint::GetLocalVariables(FGraphActionListBuilderBase& OutAllActions)
 					Category = FString();
 				}
 
-				TSharedPtr<FEdGraphSchemaAction_K2LocalVar> NewVarAction = MakeShareable(new FEdGraphSchemaAction_K2LocalVar(Category, FText::FromName(Variable.VarName), TEXT(""), 0));
 				UFunction* Func = FindField<UFunction>(GetBlueprintObj()->SkeletonGeneratedClass, EdGraph->GetFName());
-				NewVarAction->SetVariableInfo(Variable.VarName, Func);
-				OutAllActions.AddAction(NewVarAction);
+				if (Func)
+				{
+					TSharedPtr<FEdGraphSchemaAction_K2LocalVar> NewVarAction = MakeShareable(new FEdGraphSchemaAction_K2LocalVar(Category, FText::FromName(Variable.VarName), TEXT(""), 0));
+					NewVarAction->SetVariableInfo(Variable.VarName, Func);
+					OutAllActions.AddAction(NewVarAction);
+				}
 			}
 		}
 	}
