Expanding macros with illegally placeable nodes in the context of the current graph will no longer crash.

Assign and temporary variable nodes can now be pasted into any graph with a supporting schema. Drag and drop from palette and context menu will not allow these nodes still.

When dragging and dropping nodes from the palette, if the node cannot be pasted or it is not compatible with the graph, it will not be droppable.

diff --git a/Engine/Source/Editor/BlueprintGraph/Classes/K2Node_AssignmentStatement.h b/Engine/Source/Editor/BlueprintGraph/Classes/K2Node_AssignmentStatement.h
--- a/Engine/Source/Editor/BlueprintGraph/Classes/K2Node_AssignmentStatement.h
+++ b/Engine/Source/Editor/BlueprintGraph/Classes/K2Node_AssignmentStatement.h
@@ -20,6 +20,7 @@ class UK2Node_AssignmentStatement : public UK2Node
 	virtual FText GetTooltipText() const override;
 	virtual FText GetNodeTitle(ENodeTitleType::Type TitleType) const override;
 	virtual bool IsCompatibleWithGraph(UEdGraph const* TargetGraph) const override;
+	virtual bool CanPasteHere(const UEdGraph* TargetGraph) const override;
 	// End UEdGraphNode interface
 
 	// Begin UK2Node interface
diff --git a/Engine/Source/Editor/BlueprintGraph/Classes/K2Node_TemporaryVariable.h b/Engine/Source/Editor/BlueprintGraph/Classes/K2Node_TemporaryVariable.h
--- a/Engine/Source/Editor/BlueprintGraph/Classes/K2Node_TemporaryVariable.h
+++ b/Engine/Source/Editor/BlueprintGraph/Classes/K2Node_TemporaryVariable.h
@@ -26,6 +26,7 @@ class UK2Node_TemporaryVariable : public UK2Node
 	virtual FText GetNodeTitle(ENodeTitleType::Type TitleType) const override;
 	virtual FString GetDescriptiveCompiledName() const override;
 	virtual bool IsCompatibleWithGraph(UEdGraph const* TargetGraph) const override;
+	virtual bool CanPasteHere(const UEdGraph* TargetGraph) const override;
 	// End UEdGraphNode interface.
 
 	// Begin UK2Node interface.
diff --git a/Engine/Source/Editor/BlueprintGraph/Private/K2Node_AssignmentStatement.cpp b/Engine/Source/Editor/BlueprintGraph/Private/K2Node_AssignmentStatement.cpp
--- a/Engine/Source/Editor/BlueprintGraph/Private/K2Node_AssignmentStatement.cpp
+++ b/Engine/Source/Editor/BlueprintGraph/Private/K2Node_AssignmentStatement.cpp
@@ -97,6 +97,14 @@ bool UK2Node_AssignmentStatement::IsCompatibleWithGraph(UEdGraph const* TargetGr
 	return bIsCompatible;
 }
 
+bool UK2Node_AssignmentStatement::CanPasteHere(const UEdGraph* TargetGraph) const
+{
+	// These nodes can be pasted anywhere that UK2Node's are compatible with the graph
+	// Avoiding the call to IsCompatibleWithGraph because these nodes should normally only
+	// be placed in Macros, but it's nice to be able to paste Macro functionality anywhere.
+	return Super::IsCompatibleWithGraph(TargetGraph);
+}
+
 void UK2Node_AssignmentStatement::NotifyPinConnectionListChanged(UEdGraphPin* Pin)
 {
 	Super::NotifyPinConnectionListChanged(Pin);
diff --git a/Engine/Source/Editor/BlueprintGraph/Private/K2Node_TemporaryVariable.cpp b/Engine/Source/Editor/BlueprintGraph/Private/K2Node_TemporaryVariable.cpp
--- a/Engine/Source/Editor/BlueprintGraph/Private/K2Node_TemporaryVariable.cpp
+++ b/Engine/Source/Editor/BlueprintGraph/Private/K2Node_TemporaryVariable.cpp
@@ -116,6 +116,14 @@ bool UK2Node_TemporaryVariable::IsCompatibleWithGraph(UEdGraph const* TargetGrap
 	return bIsCompatible;
 }
 
+bool UK2Node_TemporaryVariable::CanPasteHere(const UEdGraph* TargetGraph) const
+{
+	// These nodes can be pasted anywhere that UK2Node's are compatible with the graph
+	// Avoiding the call to IsCompatibleWithGraph because these nodes should normally only
+	// be placed in Macros, but it's nice to be able to paste Macro functionality anywhere.
+	return Super::IsCompatibleWithGraph(TargetGraph);
+}
+
 // get variable pin
 UEdGraphPin* UK2Node_TemporaryVariable::GetVariablePin()
 {
diff --git a/Engine/Source/Editor/Kismet/Private/BlueprintEditor.cpp b/Engine/Source/Editor/Kismet/Private/BlueprintEditor.cpp
--- a/Engine/Source/Editor/Kismet/Private/BlueprintEditor.cpp
+++ b/Engine/Source/Editor/Kismet/Private/BlueprintEditor.cpp
@@ -5764,6 +5764,7 @@ void FBlueprintEditor::ExpandNode(UEdGraphNode* InNodeToExpand, UEdGraph* InSour
 		// so when expanding a collapsed graph we don't want to check the CanPasteHere function:
 		if (!bIsCollapsedGraph && !Node->CanPasteHere(DestinationGraph))
 		{
+			Node->BreakAllNodeLinks();
 			continue;
 		}
 
diff --git a/Engine/Source/Editor/Kismet/Private/SBlueprintSubPalette.cpp b/Engine/Source/Editor/Kismet/Private/SBlueprintSubPalette.cpp
--- a/Engine/Source/Editor/Kismet/Private/SBlueprintSubPalette.cpp
+++ b/Engine/Source/Editor/Kismet/Private/SBlueprintSubPalette.cpp
@@ -141,8 +141,10 @@ static bool CanPaletteItemBePlaced(TSharedPtr<FEdGraphSchemaAction> DropActionIn
 			bCanBePlaced = true;
 			ImpededReasonOut = FText::GetEmpty();
 		}
-		// as a general catch-all, if a node cannot be pasted, it probably can't be created there
-		else if (bCanBePlaced && !NodeToBePlaced->CanPasteHere(HoveredGraphIn) && !bWillFocusOnExistingNode)
+		// as a general catch-all, if a node cannot be pasted or placed in the graph, it probably can't be created there.
+		// Some nodes allow themselves to be pasted where they are generally not allowed, if either does not want the 
+		// node placed, it should not be placeable
+		else if (bCanBePlaced && (!NodeToBePlaced->CanPasteHere(HoveredGraphIn) || !NodeToBePlaced->IsCompatibleWithGraph(HoveredGraphIn)) && !bWillFocusOnExistingNode)
 		{
 			bCanBePlaced = false;
 			ImpededReasonOut = LOCTEXT("CannotPaste", "Cannot place this node in this type of graph");
