Slate - ESlateCheckBoxState has been renamed to ECheckBoxState.

diff --git a/Engine/Source/Editor/FoliageEdit/Private/SFoliageEditMeshDisplayItem.h b/Engine/Source/Editor/FoliageEdit/Private/SFoliageEditMeshDisplayItem.h
--- a/Engine/Source/Editor/FoliageEdit/Private/SFoliageEditMeshDisplayItem.h
+++ b/Engine/Source/Editor/FoliageEdit/Private/SFoliageEditMeshDisplayItem.h
@@ -69,9 +69,9 @@ private:
 	FSlateColor GetBorderColor() const;
 
 	/** Returns true if this item is selected. */
-	ESlateCheckBoxState::Type IsSelected() const;
+	ECheckBoxState IsSelected() const;
 
-	void OnSelectionChanged(ESlateCheckBoxState::Type InType);
+	void OnSelectionChanged(ECheckBoxState InType);
 
 	/** Bind commands used by this class for the toolbar. */
 	void BindCommands();
@@ -133,9 +133,9 @@ private:
 
 
 	/** Density */
-	void OnDensityReapply(ESlateCheckBoxState::Type InState);
+	void OnDensityReapply(ECheckBoxState InState);
 
-	ESlateCheckBoxState::Type IsDensityReapplyChecked() const;
+	ECheckBoxState IsDensityReapplyChecked() const;
 
 	void OnDensityChanged(float InValue);
 
@@ -161,22 +161,22 @@ private:
 	void OnGrowthRadiusChanged(float InValue);
 
 	/** Radius */
-	void OnRadiusReapply(ESlateCheckBoxState::Type InState);
+	void OnRadiusReapply(ECheckBoxState InState);
 
-	ESlateCheckBoxState::Type IsRadiusReapplyChecked() const;
+	ECheckBoxState IsRadiusReapplyChecked() const;
 
 	void OnRadiusChanged(float InValue);
 
 	float GetRadius() const;
 
 	/** Align to Normal */
-	void OnAlignToNormalReapply(ESlateCheckBoxState::Type InState);
+	void OnAlignToNormalReapply(ECheckBoxState InState);
 
-	ESlateCheckBoxState::Type IsAlignToNormalReapplyChecked() const;
+	ECheckBoxState IsAlignToNormalReapplyChecked() const;
 
-	void OnAlignToNormal(ESlateCheckBoxState::Type InState);
+	void OnAlignToNormal(ECheckBoxState InState);
 
-	ESlateCheckBoxState::Type IsAlignToNormalChecked() const;
+	ECheckBoxState IsAlignToNormalChecked() const;
 
 	EVisibility IsAlignToNormalVisible() const;
 
@@ -186,23 +186,23 @@ private:
 	float GetMaxAngle() const;
 
 	/** Random Yaw */
-	void OnRandomYawReapply(ESlateCheckBoxState::Type InState);
+	void OnRandomYawReapply(ECheckBoxState InState);
 
-	ESlateCheckBoxState::Type IsRandomYawReapplyChecked() const;
+	ECheckBoxState IsRandomYawReapplyChecked() const;
 
-	void OnRandomYaw(ESlateCheckBoxState::Type InState);
+	void OnRandomYaw(ECheckBoxState InState);
 
-	ESlateCheckBoxState::Type IsRandomYawChecked() const;
+	ECheckBoxState IsRandomYawChecked() const;
 
 	/** Uniform Scale Checkbox */
-	void OnUniformScale(ESlateCheckBoxState::Type InState);
+	void OnUniformScale(ECheckBoxState InState);
 
-	ESlateCheckBoxState::Type IsUniformScaleChecked() const;
+	ECheckBoxState IsUniformScaleChecked() const;
 
 	/** Uniform Scaling Min/Max */
-	void OnScaleUniformReapply(ESlateCheckBoxState::Type InState);
+	void OnScaleUniformReapply(ECheckBoxState InState);
 
-	ESlateCheckBoxState::Type IsScaleUniformReapplyChecked() const;
+	ECheckBoxState IsScaleUniformReapplyChecked() const;
 
 	void OnScaleUniformMinChanged(float InValue);
 
@@ -213,9 +213,9 @@ private:
 	float GetScaleUniformMax() const;
 
 	/** X-Axis Scaling Min/Max */
-	void OnScaleXReapply(ESlateCheckBoxState::Type InState);
+	void OnScaleXReapply(ECheckBoxState InState);
 
-	ESlateCheckBoxState::Type IsScaleXReapplyChecked() const;
+	ECheckBoxState IsScaleXReapplyChecked() const;
 
 	void OnScaleXMinChanged(float InValue);
 
@@ -225,14 +225,14 @@ private:
 
 	float GetScaleXMax() const;
 
-	ESlateCheckBoxState::Type IsScaleXLockedChecked() const;
+	ECheckBoxState IsScaleXLockedChecked() const;
 
-	void OnScaleXLocked(ESlateCheckBoxState::Type InState);
+	void OnScaleXLocked(ECheckBoxState InState);
 
 	/** Y-Axis Scaling Min/Max */
-	void OnScaleYReapply(ESlateCheckBoxState::Type InState);
+	void OnScaleYReapply(ECheckBoxState InState);
 
-	ESlateCheckBoxState::Type IsScaleYReapplyChecked() const;
+	ECheckBoxState IsScaleYReapplyChecked() const;
 
 	void OnScaleYMinChanged(float InValue);
 
@@ -242,14 +242,14 @@ private:
 
 	float GetScaleYMax() const;
 
-	ESlateCheckBoxState::Type IsScaleYLockedChecked() const;
+	ECheckBoxState IsScaleYLockedChecked() const;
 
-	void OnScaleYLocked(ESlateCheckBoxState::Type InState);
+	void OnScaleYLocked(ECheckBoxState InState);
 
 	/** Z-Axis Scaling Min/Max */
-	void OnScaleZReapply(ESlateCheckBoxState::Type InState);
+	void OnScaleZReapply(ECheckBoxState InState);
 
-	ESlateCheckBoxState::Type IsScaleZReapplyChecked() const;
+	ECheckBoxState IsScaleZReapplyChecked() const;
 
 	void OnScaleZMinChanged(float InValue);
 
@@ -259,14 +259,14 @@ private:
 
 	float GetScaleZMax() const;
 
-	ESlateCheckBoxState::Type IsScaleZLockedChecked() const;
+	ECheckBoxState IsScaleZLockedChecked() const;
 
-	void OnScaleZLocked(ESlateCheckBoxState::Type InState);
+	void OnScaleZLocked(ECheckBoxState InState);
 
 	/** Z-Offset Min/Max */
-	void OnZOffsetReapply(ESlateCheckBoxState::Type InState);
+	void OnZOffsetReapply(ECheckBoxState InState);
 
-	ESlateCheckBoxState::Type IsZOffsetReapplyChecked() const;
+	ECheckBoxState IsZOffsetReapplyChecked() const;
 
 	void OnZOffsetMin(float InValue);
 
@@ -277,27 +277,27 @@ private:
 	float GetZOffsetMax() const;
 
 	/** Random Pitch +/- */
-	void OnRandomPitchReapply(ESlateCheckBoxState::Type InState);
+	void OnRandomPitchReapply(ECheckBoxState InState);
 
-	ESlateCheckBoxState::Type IsRandomPitchReapplyChecked() const;
+	ECheckBoxState IsRandomPitchReapplyChecked() const;
 
 	void OnRandomPitchChanged(float InValue);
 
 	float GetRandomPitch() const;
 
 	/** Ground Slope */
-	void OnGroundSlopeReapply(ESlateCheckBoxState::Type InState);
+	void OnGroundSlopeReapply(ECheckBoxState InState);
 
-	ESlateCheckBoxState::Type IsGroundSlopeReapplyChecked() const;
+	ECheckBoxState IsGroundSlopeReapplyChecked() const;
 
 	void OnGroundSlopeChanged(float InValue);
 
 	float GetGroundSlope() const;
 
 	/** Height Min/Max */
-	void OnHeightReapply(ESlateCheckBoxState::Type InState);
+	void OnHeightReapply(ECheckBoxState InState);
 
-	ESlateCheckBoxState::Type IsHeightReapplyChecked() const;
+	ECheckBoxState IsHeightReapplyChecked() const;
 
 	void OnHeightMinChanged(float InValue);
 
@@ -321,24 +321,24 @@ private:
 	void OnSimulationPriorityChanged(float InValue);
 
 	/** Landscape Layer */
-	void OnLandscapeLayerReapply(ESlateCheckBoxState::Type InState);
+	void OnLandscapeLayerReapply(ECheckBoxState InState);
 
-	ESlateCheckBoxState::Type IsLandscapeLayerReapplyChecked() const;
+	ECheckBoxState IsLandscapeLayerReapplyChecked() const;
 
 	void OnLandscapeLayerChanged(const FText& InValue);
 
 	FText GetLandscapeLayer() const;
 
 	/** Collision with World */
-	void OnCollisionWithWorld(ESlateCheckBoxState::Type InState);
+	void OnCollisionWithWorld(ECheckBoxState InState);
 
-	ESlateCheckBoxState::Type IsCollisionWithWorldChecked() const;
+	ECheckBoxState IsCollisionWithWorldChecked() const;
 
 	EVisibility IsCollisionWithWorldVisible() const;
 
-	void OnCollisionWithWorldReapply(ESlateCheckBoxState::Type InState);
+	void OnCollisionWithWorldReapply(ECheckBoxState InState);
 
-	ESlateCheckBoxState::Type IsCollisionWithWorldReapplyChecked() const;
+	ECheckBoxState IsCollisionWithWorldReapplyChecked() const;
 
 	void OnCollisionScaleXChanged(float InValue);
 	void OnCollisionScaleYChanged(float InValue);
@@ -349,13 +349,13 @@ private:
 	float GetCollisionScaleZ() const;
 
 	/** Vertex Color Mask */
-	void OnVertexColorMaskReapply(ESlateCheckBoxState::Type InState);
+	void OnVertexColorMaskReapply(ECheckBoxState InState);
 
-	ESlateCheckBoxState::Type IsVertexColorMaskReapplyChecked() const;
+	ECheckBoxState IsVertexColorMaskReapplyChecked() const;
 
-	void OnVertexColorMask(ESlateCheckBoxState::Type InState, FoliageVertexColorMask Mask);
+	void OnVertexColorMask(ECheckBoxState InState, FoliageVertexColorMask Mask);
 
-	ESlateCheckBoxState::Type IsVertexColorMaskChecked(FoliageVertexColorMask Mask) const;
+	ECheckBoxState IsVertexColorMaskChecked(FoliageVertexColorMask Mask) const;
 
 	/** Vertex Color Threshold */
 	void OnVertexColorMaskThresholdChanged(float InValue);
@@ -364,9 +364,9 @@ private:
 
 	EVisibility IsVertexColorMaskThresholdVisible() const;
 
-	void OnVertexColorMaskInvert(ESlateCheckBoxState::Type InState);
+	void OnVertexColorMaskInvert(ECheckBoxState InState);
 
-	ESlateCheckBoxState::Type IsVertexColorMaskInvertChecked() const;
+	ECheckBoxState IsVertexColorMaskInvertChecked() const;
 
 private:
 	/** Pointer back to the foliage edit mode containing this. */
diff --git a/Engine/Source/Editor/GameProjectGeneration/Private/SNewClassDialog.cpp b/Engine/Source/Editor/GameProjectGeneration/Private/SNewClassDialog.cpp
--- a/Engine/Source/Editor/GameProjectGeneration/Private/SNewClassDialog.cpp
+++ b/Engine/Source/Editor/GameProjectGeneration/Private/SNewClassDialog.cpp
@@ -845,14 +845,14 @@ void SNewClassDialog::OnAdvancedClassSelected(UClass* Class)
 	ParentClassInfo = GameProjectUtils::FNewClassInfo(Class);
 }
 
-ESlateCheckBoxState::Type SNewClassDialog::IsFullClassTreeChecked() const
+ECheckBoxState SNewClassDialog::IsFullClassTreeChecked() const
 {
-	return bShowFullClassTree ? ESlateCheckBoxState::Checked : ESlateCheckBoxState::Unchecked;
+	return bShowFullClassTree ? ECheckBoxState::Checked : ECheckBoxState::Unchecked;
 }
 
-void SNewClassDialog::OnFullClassTreeChanged(ESlateCheckBoxState::Type NewCheckedState)
+void SNewClassDialog::OnFullClassTreeChanged(ECheckBoxState NewCheckedState)
 {
-	bShowFullClassTree = (NewCheckedState == ESlateCheckBoxState::Checked);
+	bShowFullClassTree = (NewCheckedState == ECheckBoxState::Checked);
 }
 
 EVisibility SNewClassDialog::GetBasicParentClassVisibility() const
@@ -1140,14 +1140,14 @@ FText SNewClassDialog::GetClassLocationTooltip(GameProjectUtils::EClassLocation
 	return LOCTEXT("ClassLocation_UserDefined", "Your project is either not using a Public and Private source layout, or you're explicitly creating your class outside of the Public or Private folder");
 }
 
-ESlateCheckBoxState::Type SNewClassDialog::IsClassLocationActive(GameProjectUtils::EClassLocation InLocation) const
+ECheckBoxState SNewClassDialog::IsClassLocationActive(GameProjectUtils::EClassLocation InLocation) const
 {
-	return (ClassLocation == InLocation) ? ESlateCheckBoxState::Checked : ESlateCheckBoxState::Unchecked;
+	return (ClassLocation == InLocation) ? ECheckBoxState::Checked : ECheckBoxState::Unchecked;
 }
 
-void SNewClassDialog::OnClassLocationChanged(ESlateCheckBoxState::Type InCheckedState, GameProjectUtils::EClassLocation InLocation)
+void SNewClassDialog::OnClassLocationChanged(ECheckBoxState InCheckedState, GameProjectUtils::EClassLocation InLocation)
 {
-	if(InCheckedState == ESlateCheckBoxState::Checked)
+	if(InCheckedState == ECheckBoxState::Checked)
 	{
 		const FString AbsoluteClassPath = FPaths::ConvertRelativePathToFull(NewClassPath) / ""; // Ensure trailing /
 
diff --git a/Engine/Source/Editor/GameProjectGeneration/Private/SNewClassDialog.h b/Engine/Source/Editor/GameProjectGeneration/Private/SNewClassDialog.h
--- a/Engine/Source/Editor/GameProjectGeneration/Private/SNewClassDialog.h
+++ b/Engine/Source/Editor/GameProjectGeneration/Private/SNewClassDialog.h
@@ -57,10 +57,10 @@ private:
 	void OnAdvancedClassSelected(UClass* Class);
 
 	/** Gets the check box state for the full class list */
-	ESlateCheckBoxState::Type IsFullClassTreeChecked() const;
+	ECheckBoxState IsFullClassTreeChecked() const;
 
 	/** Gets the check box state for the full class list */
-	void OnFullClassTreeChanged(ESlateCheckBoxState::Type NewCheckedState);
+	void OnFullClassTreeChanged(ECheckBoxState NewCheckedState);
 
 	/** Gets the visibility of the basic class list */
 	EVisibility GetBasicParentClassVisibility() const;
@@ -140,10 +140,10 @@ private:
 	FText GetClassLocationTooltip(GameProjectUtils::EClassLocation InLocation) const;
 
 	/** Checks to see if the given class location is active based on the current value of NewClassPath */
-	ESlateCheckBoxState::Type IsClassLocationActive(GameProjectUtils::EClassLocation InLocation) const;
+	ECheckBoxState IsClassLocationActive(GameProjectUtils::EClassLocation InLocation) const;
 
 	/** Update the value of NewClassPath so that it uses the given class location */
-	void OnClassLocationChanged(ESlateCheckBoxState::Type InCheckedState, GameProjectUtils::EClassLocation InLocation);
+	void OnClassLocationChanged(ECheckBoxState InCheckedState, GameProjectUtils::EClassLocation InLocation);
 
 	/** Checks to see if the class location can be changed (relies on us having a Public/Private folder layout) */
 	bool CanChangeClassLocation() const;
diff --git a/Engine/Source/Editor/GameProjectGeneration/Private/SNewProjectWizard.cpp b/Engine/Source/Editor/GameProjectGeneration/Private/SNewProjectWizard.cpp
--- a/Engine/Source/Editor/GameProjectGeneration/Private/SNewProjectWizard.cpp
+++ b/Engine/Source/Editor/GameProjectGeneration/Private/SNewProjectWizard.cpp
@@ -730,7 +730,7 @@ void SNewProjectWizard::Construct( const FArguments& InArgs )
 	// Initialize the current page name. Assuming the template page.
 	CurrentPageName = TemplatePageName;
 
-	HandleCategoryChanged(ESlateCheckBoxState::Checked, ActiveCategory);
+	HandleCategoryChanged(ECheckBoxState::Checked, ActiveCategory);
 
 	UpdateProjectFileValidity();
 }
@@ -1582,14 +1582,14 @@ void SNewProjectWizard::HandleCreateProjectWizardFinished( )
 	CreateAndOpenProject();
 }
 
-ESlateCheckBoxState::Type SNewProjectWizard::GetCategoryTabCheckState(FName Category) const
+ECheckBoxState SNewProjectWizard::GetCategoryTabCheckState(FName Category) const
 {
-	return Category == ActiveCategory ? ESlateCheckBoxState::Checked : ESlateCheckBoxState::Unchecked;
+	return Category == ActiveCategory ? ECheckBoxState::Checked : ECheckBoxState::Unchecked;
 }
 
-void SNewProjectWizard::HandleCategoryChanged(ESlateCheckBoxState::Type CheckState, FName Category)
+void SNewProjectWizard::HandleCategoryChanged(ECheckBoxState CheckState, FName Category)
 {
-	if (CheckState != ESlateCheckBoxState::Checked)
+	if (CheckState != ECheckBoxState::Checked)
 	{
 		return;
 	}
diff --git a/Engine/Source/Editor/GameProjectGeneration/Private/SNewProjectWizard.h b/Engine/Source/Editor/GameProjectGeneration/Private/SNewProjectWizard.h
--- a/Engine/Source/Editor/GameProjectGeneration/Private/SNewProjectWizard.h
+++ b/Engine/Source/Editor/GameProjectGeneration/Private/SNewProjectWizard.h
@@ -195,10 +195,10 @@ private:
 	void HandleTemplateListViewSelectionChanged( TSharedPtr<FTemplateItem> TemplateItem, ESelectInfo::Type SelectInfo );
 
 	/** Handle choosing a different category tab */
-	void HandleCategoryChanged(ESlateCheckBoxState::Type Checked, FName Category);
+	void HandleCategoryChanged(ECheckBoxState Checked, FName Category);
 
 	/** Get the check state for the specified categories tab */
-	ESlateCheckBoxState::Type GetCategoryTabCheckState(FName Category) const;
+	ECheckBoxState GetCategoryTabCheckState(FName Category) const;
 
 private:
 
diff --git a/Engine/Source/Editor/GameProjectGeneration/Private/SProjectBrowser.cpp b/Engine/Source/Editor/GameProjectGeneration/Private/SProjectBrowser.cpp
--- a/Engine/Source/Editor/GameProjectGeneration/Private/SProjectBrowser.cpp
+++ b/Engine/Source/Editor/GameProjectGeneration/Private/SProjectBrowser.cpp
@@ -262,7 +262,7 @@ void SProjectBrowser::Construct( const FArguments& InArgs )
 				.VAlign(VAlign_Center)
 				[
 					SNew(SCheckBox)			
-					.IsChecked(GEditor->GetGameAgnosticSettings().bLoadTheMostRecentlyLoadedProjectAtStartup ? ESlateCheckBoxState::Checked : ESlateCheckBoxState::Unchecked)
+					.IsChecked(GEditor->GetGameAgnosticSettings().bLoadTheMostRecentlyLoadedProjectAtStartup ? ECheckBoxState::Checked : ECheckBoxState::Unchecked)
 					.OnCheckStateChanged(this, &SProjectBrowser::OnAutoloadLastProjectChanged)
 					.Content()
 					[
@@ -1265,10 +1265,10 @@ void SProjectBrowser::OnFilterTextChanged(const FText& InText)
 	PopulateFilteredProjectCategories();
 }
 
-void SProjectBrowser::OnAutoloadLastProjectChanged(ESlateCheckBoxState::Type NewState)
+void SProjectBrowser::OnAutoloadLastProjectChanged(ECheckBoxState NewState)
 {
 	UEditorGameAgnosticSettings &Settings = GEditor->AccessGameAgnosticSettings();
-	Settings.bLoadTheMostRecentlyLoadedProjectAtStartup = (NewState == ESlateCheckBoxState::Checked);
+	Settings.bLoadTheMostRecentlyLoadedProjectAtStartup = (NewState == ECheckBoxState::Checked);
 
 	UProperty* AutoloadProjectProperty = FindField<UProperty>(Settings.GetClass(), "bLoadTheMostRecentlyLoadedProjectAtStartup");
 	if (AutoloadProjectProperty != NULL)
diff --git a/Engine/Source/Editor/GameProjectGeneration/Private/SProjectBrowser.h b/Engine/Source/Editor/GameProjectGeneration/Private/SProjectBrowser.h
--- a/Engine/Source/Editor/GameProjectGeneration/Private/SProjectBrowser.h
+++ b/Engine/Source/Editor/GameProjectGeneration/Private/SProjectBrowser.h
@@ -112,7 +112,7 @@ private:
 	void OnFilterTextChanged(const FText& InText);
 
 	/** Whether to autoload the last project. */
-	void OnAutoloadLastProjectChanged(ESlateCheckBoxState::Type NewState);
+	void OnAutoloadLastProjectChanged(ECheckBoxState NewState);
 
 	/** Get the visibility of the specified category */
 	EVisibility GetProjectCategoryVisibility(const TSharedRef<FProjectCategory> InCategory) const;
diff --git a/Engine/Source/Editor/GameplayTagsEditor/Private/SGameplayTagWidget.cpp b/Engine/Source/Editor/GameplayTagsEditor/Private/SGameplayTagWidget.cpp
--- a/Engine/Source/Editor/GameplayTagsEditor/Private/SGameplayTagWidget.cpp
+++ b/Engine/Source/Editor/GameplayTagsEditor/Private/SGameplayTagWidget.cpp
@@ -201,13 +201,13 @@ void SGameplayTagWidget::OnGetChildren(TSharedPtr<FGameplayTagNode> InItem, TArr
 	OutChildren += FilteredChildren;
 }
 
-void SGameplayTagWidget::OnTagCheckStatusChanged(ESlateCheckBoxState::Type NewCheckState, TSharedPtr<FGameplayTagNode> NodeChanged)
+void SGameplayTagWidget::OnTagCheckStatusChanged(ECheckBoxState NewCheckState, TSharedPtr<FGameplayTagNode> NodeChanged)
 {
-	if (NewCheckState == ESlateCheckBoxState::Checked)
+	if (NewCheckState == ECheckBoxState::Checked)
 	{
 		OnTagChecked(NodeChanged);
 	}
-	else if (NewCheckState == ESlateCheckBoxState::Unchecked)
+	else if (NewCheckState == ECheckBoxState::Unchecked)
 	{
 		OnTagUnchecked(NodeChanged);
 	}
@@ -322,7 +322,7 @@ void SGameplayTagWidget::UncheckChildren(TSharedPtr<FGameplayTagNode> NodeUnchec
 	}
 }
 
-ESlateCheckBoxState::Type SGameplayTagWidget::IsTagChecked(TSharedPtr<FGameplayTagNode> Node) const
+ECheckBoxState SGameplayTagWidget::IsTagChecked(TSharedPtr<FGameplayTagNode> Node) const
 {
 	int32 NumValidAssets = 0;
 	int32 NumAssetsTagIsAppliedTo = 0;
@@ -351,21 +351,21 @@ ESlateCheckBoxState::Type SGameplayTagWidget::IsTagChecked(TSharedPtr<FGameplayT
 		// Check if any children are tagged
 		for (auto It = Node->GetChildTagNodes().CreateConstIterator(); It; ++It)
 		{
-			if (IsTagChecked(*It) == ESlateCheckBoxState::Checked)
+			if (IsTagChecked(*It) == ECheckBoxState::Checked)
 			{
-				return ESlateCheckBoxState::Checked;
+				return ECheckBoxState::Checked;
 			}
 		}
 
-		return ESlateCheckBoxState::Unchecked;
+		return ECheckBoxState::Unchecked;
 	}
 	else if (NumAssetsTagIsAppliedTo == NumValidAssets)
 	{
-		return ESlateCheckBoxState::Checked;
+		return ECheckBoxState::Checked;
 	}
 	else
 	{
-		return ESlateCheckBoxState::Undetermined;
+		return ECheckBoxState::Undetermined;
 	}
 }
 
@@ -494,18 +494,18 @@ void SGameplayTagWidget::LoadSettings()
 
 void SGameplayTagWidget::SetDefaultTagNodeItemExpansion( TSharedPtr<FGameplayTagNode> Node )
 {
-	if (Node.IsValid() && TagTreeWidget.IsValid())
+	if ( Node.IsValid() && TagTreeWidget.IsValid() )
 	{
 		bool bExpanded = false;
 
-		if( IsTagChecked( Node ) )
+		if ( IsTagChecked(Node) == ECheckBoxState::Checked )
 		{
 			bExpanded = true;
 		}
-		TagTreeWidget->SetItemExpansion( Node, bExpanded );
+		TagTreeWidget->SetItemExpansion(Node, bExpanded);
 
 		const TArray< TSharedPtr<FGameplayTagNode> >& ChildTags = Node->GetChildTagNodes();
-		for (int32 ChildIdx = 0; ChildIdx < ChildTags.Num(); ++ChildIdx)
+		for ( int32 ChildIdx = 0; ChildIdx < ChildTags.Num(); ++ChildIdx )
 		{
 			SetDefaultTagNodeItemExpansion(ChildTags[ChildIdx]);
 		}
@@ -522,7 +522,7 @@ void SGameplayTagWidget::LoadTagNodeItemExpansion( TSharedPtr<FGameplayTagNode>
 		{
 			TagTreeWidget->SetItemExpansion( Node, bExpanded );
 		}
-		else if( IsTagChecked( Node ) ) // If we have no save data but its ticked then we probably lost our settings so we shall expand it
+		else if( IsTagChecked( Node ) == ECheckBoxState::Checked ) // If we have no save data but its ticked then we probably lost our settings so we shall expand it
 		{
 			TagTreeWidget->SetItemExpansion( Node, true );
 		}
diff --git a/Engine/Source/Editor/GameplayTagsEditor/Private/SGameplayTagWidget.h b/Engine/Source/Editor/GameplayTagsEditor/Private/SGameplayTagWidget.h
--- a/Engine/Source/Editor/GameplayTagsEditor/Private/SGameplayTagWidget.h
+++ b/Engine/Source/Editor/GameplayTagsEditor/Private/SGameplayTagWidget.h
@@ -117,7 +117,7 @@ private:
 	 * @param NewCheckState	New check box state
 	 * @param NodeChanged	Node that was checked/unchecked
 	 */
-	void OnTagCheckStatusChanged(ESlateCheckBoxState::Type NewCheckState, TSharedPtr<FGameplayTagNode> NodeChanged);
+	void OnTagCheckStatusChanged(ECheckBoxState NewCheckState, TSharedPtr<FGameplayTagNode> NodeChanged);
 
 	/**
 	 * Called via delegate to determine the checkbox state of the specified node
@@ -126,7 +126,7 @@ private:
 	 * 
 	 * @return Checkbox state of the specified node
 	 */
-	ESlateCheckBoxState::Type IsTagChecked(TSharedPtr<FGameplayTagNode> Node) const;
+	ECheckBoxState IsTagChecked(TSharedPtr<FGameplayTagNode> Node) const;
 
 	/**
 	 * Helper function called when the specified node is checked
diff --git a/Engine/Source/Editor/GeometryMode/Private/GeometryMode.cpp b/Engine/Source/Editor/GeometryMode/Private/GeometryMode.cpp
--- a/Engine/Source/Editor/GeometryMode/Private/GeometryMode.cpp
+++ b/Engine/Source/Editor/GeometryMode/Private/GeometryMode.cpp
@@ -33,9 +33,9 @@ void SGeometryModeControls::Construct(const FArguments& InArgs)
 	CreateLayout();
 }
 
-void SGeometryModeControls::OnModifierStateChanged(ESlateCheckBoxState::Type NewCheckedState, UGeomModifier* Modifier)
+void SGeometryModeControls::OnModifierStateChanged(ECheckBoxState NewCheckedState, UGeomModifier* Modifier)
 {
-	if (NewCheckedState == ESlateCheckBoxState::Checked)
+	if (NewCheckedState == ECheckBoxState::Checked)
 	{
 		GetGeometryModeTool()->SetCurrentModifier(Modifier);
 
@@ -45,11 +45,11 @@ void SGeometryModeControls::OnModifierStateChanged(ESlateCheckBoxState::Type New
 	}
 }
 
-ESlateCheckBoxState::Type SGeometryModeControls::IsModifierChecked(UGeomModifier* Modifier) const
+ECheckBoxState SGeometryModeControls::IsModifierChecked(UGeomModifier* Modifier) const
 {
 	return (GetGeometryModeTool()->GetCurrentModifier() == Modifier)
-		? ESlateCheckBoxState::Checked
-		: ESlateCheckBoxState::Unchecked;
+		? ECheckBoxState::Checked
+		: ECheckBoxState::Unchecked;
 }
 
 bool SGeometryModeControls::IsModifierEnabled(UGeomModifier* Modifier) const
diff --git a/Engine/Source/Editor/GeometryMode/Private/GeometryMode.h b/Engine/Source/Editor/GeometryMode/Private/GeometryMode.h
--- a/Engine/Source/Editor/GeometryMode/Private/GeometryMode.h
+++ b/Engine/Source/Editor/GeometryMode/Private/GeometryMode.h
@@ -20,10 +20,10 @@ public:
 
 protected:
 	/** Called when a new modifier mode is selected */
-	void OnModifierStateChanged(ESlateCheckBoxState::Type NewCheckedState, UGeomModifier* Modifier);
+	void OnModifierStateChanged(ECheckBoxState NewCheckedState, UGeomModifier* Modifier);
 
 	/** Returns the state of a modifier radio button */
-	ESlateCheckBoxState::Type IsModifierChecked(UGeomModifier* Modifier) const;
+	ECheckBoxState IsModifierChecked(UGeomModifier* Modifier) const;
 
 	/** Returns the enabled state of a modifier button */
 	bool IsModifierEnabled(UGeomModifier* Modifier) const;
diff --git a/Engine/Source/Editor/GraphEditor/Private/KismetPins/SGraphPinBool.cpp b/Engine/Source/Editor/GraphEditor/Private/KismetPins/SGraphPinBool.cpp
--- a/Engine/Source/Editor/GraphEditor/Private/KismetPins/SGraphPinBool.cpp
+++ b/Engine/Source/Editor/GraphEditor/Private/KismetPins/SGraphPinBool.cpp
@@ -18,14 +18,14 @@ TSharedRef<SWidget>	SGraphPinBool::GetDefaultValueWidget()
 		.Visibility( this, &SGraphPin::GetDefaultValueVisibility );
 }
 
-ESlateCheckBoxState::Type SGraphPinBool::IsDefaultValueChecked() const
+ECheckBoxState SGraphPinBool::IsDefaultValueChecked() const
 {
 	FString CurrentValue = GraphPinObj->GetDefaultAsString();
-	return CurrentValue.ToBool() ? ESlateCheckBoxState::Checked : ESlateCheckBoxState::Unchecked;
+	return CurrentValue.ToBool() ? ECheckBoxState::Checked : ECheckBoxState::Unchecked;
 }
 
-void SGraphPinBool::OnDefaultValueCheckBoxChanged(ESlateCheckBoxState::Type InIsChecked)
+void SGraphPinBool::OnDefaultValueCheckBoxChanged(ECheckBoxState InIsChecked)
 {
-	const FString BoolString = (InIsChecked == ESlateCheckBoxState::Checked) ? TEXT("true") : TEXT("false");
+	const FString BoolString = (InIsChecked == ECheckBoxState::Checked) ? TEXT("true") : TEXT("false");
 	GraphPinObj->GetSchema()->TrySetDefaultValue(*GraphPinObj, BoolString);
 }
\ No newline at end of file
diff --git a/Engine/Source/Editor/GraphEditor/Private/KismetPins/SGraphPinBool.h b/Engine/Source/Editor/GraphEditor/Private/KismetPins/SGraphPinBool.h
--- a/Engine/Source/Editor/GraphEditor/Private/KismetPins/SGraphPinBool.h
+++ b/Engine/Source/Editor/GraphEditor/Private/KismetPins/SGraphPinBool.h
@@ -16,8 +16,8 @@ protected:
 	// End SGraphPin interface
 
 	/** Determine if the check box should be checked or not */
-	ESlateCheckBoxState::Type IsDefaultValueChecked() const;
+	ECheckBoxState IsDefaultValueChecked() const;
 
 	/** Called when check box is changed */
-	void OnDefaultValueCheckBoxChanged( ESlateCheckBoxState::Type InIsChecked );
+	void OnDefaultValueCheckBoxChanged( ECheckBoxState InIsChecked );
 };
diff --git a/Engine/Source/Editor/GraphEditor/Private/MaterialNodes/SGraphNodeMaterialBase.cpp b/Engine/Source/Editor/GraphEditor/Private/MaterialNodes/SGraphNodeMaterialBase.cpp
--- a/Engine/Source/Editor/GraphEditor/Private/MaterialNodes/SGraphNodeMaterialBase.cpp
+++ b/Engine/Source/Editor/GraphEditor/Private/MaterialNodes/SGraphNodeMaterialBase.cpp
@@ -439,10 +439,10 @@ EVisibility SGraphNodeMaterialBase::ExpressionPreviewVisibility() const
 	return bShowPreview ? EVisibility::Visible : EVisibility::Collapsed;
 }
 
-void SGraphNodeMaterialBase::OnExpressionPreviewChanged( const ESlateCheckBoxState::Type NewCheckedState )
+void SGraphNodeMaterialBase::OnExpressionPreviewChanged( const ECheckBoxState NewCheckedState )
 {
 	UMaterialExpression* MaterialExpression = MaterialNode->MaterialExpression;
-	const bool bCollapsed = (NewCheckedState != ESlateCheckBoxState::Checked);
+	const bool bCollapsed = (NewCheckedState != ECheckBoxState::Checked);
 	if (MaterialExpression->bCollapsed != bCollapsed)
 	{
 		UMaterialGraph* MaterialGraph = CastChecked<UMaterialGraph>(MaterialNode->GetGraph());
@@ -453,9 +453,9 @@ void SGraphNodeMaterialBase::OnExpressionPreviewChanged( const ESlateCheckBoxSta
 	}
 }
 
-ESlateCheckBoxState::Type SGraphNodeMaterialBase::IsExpressionPreviewChecked() const
+ECheckBoxState SGraphNodeMaterialBase::IsExpressionPreviewChecked() const
 {
-	return MaterialNode->MaterialExpression->bCollapsed ? ESlateCheckBoxState::Unchecked : ESlateCheckBoxState::Checked;
+	return MaterialNode->MaterialExpression->bCollapsed ? ECheckBoxState::Unchecked : ECheckBoxState::Checked;
 }
 
 const FSlateBrush* SGraphNodeMaterialBase::GetExpressionPreviewArrow() const
diff --git a/Engine/Source/Editor/GraphEditor/Private/MaterialNodes/SGraphNodeMaterialBase.h b/Engine/Source/Editor/GraphEditor/Private/MaterialNodes/SGraphNodeMaterialBase.h
--- a/Engine/Source/Editor/GraphEditor/Private/MaterialNodes/SGraphNodeMaterialBase.h
+++ b/Engine/Source/Editor/GraphEditor/Private/MaterialNodes/SGraphNodeMaterialBase.h
@@ -91,10 +91,10 @@ protected:
 	EVisibility ExpressionPreviewVisibility() const;
 
 	/** Show/hide Expression Preview */
-	void OnExpressionPreviewChanged( const ESlateCheckBoxState::Type NewCheckedState );
+	void OnExpressionPreviewChanged( const ECheckBoxState NewCheckedState );
 
 	/** hidden == unchecked, shown == checked */
-	ESlateCheckBoxState::Type IsExpressionPreviewChecked() const;
+	ECheckBoxState IsExpressionPreviewChecked() const;
 
 	/** Up when shown, down when hidden */
 	const FSlateBrush* GetExpressionPreviewArrow() const;
diff --git a/Engine/Source/Editor/GraphEditor/Private/SCommentBubble.cpp b/Engine/Source/Editor/GraphEditor/Private/SCommentBubble.cpp
--- a/Engine/Source/Editor/GraphEditor/Private/SCommentBubble.cpp
+++ b/Engine/Source/Editor/GraphEditor/Private/SCommentBubble.cpp
@@ -127,7 +127,7 @@ void SCommentBubble::UpdateBubble()
 					[
 						SNew( SCheckBox )
 						.Style( &FEditorStyle::Get().GetWidgetStyle<FCheckBoxStyle>( "CommentBubblePin" ))
-						.IsChecked( GraphNode->bCommentBubblePinned ? ESlateCheckBoxState::Checked : ESlateCheckBoxState::Unchecked )
+						.IsChecked( GraphNode->bCommentBubblePinned ? ECheckBoxState::Checked : ECheckBoxState::Unchecked )
 						.OnCheckStateChanged( this, &SCommentBubble::OnPinStateToggle )
 						.ToolTipText( this, &SCommentBubble::GetScaleButtonTooltip )
 						.Cursor( EMouseCursor::Default )
@@ -141,7 +141,7 @@ void SCommentBubble::UpdateBubble()
 				[
 					SNew( SCheckBox )
 					.Style( &FEditorStyle::Get().GetWidgetStyle< FCheckBoxStyle >( "CommentBubbleButton" ))
-					.IsChecked( GraphNode->bCommentBubbleVisible ? ESlateCheckBoxState::Checked : ESlateCheckBoxState::Unchecked )
+					.IsChecked( GraphNode->bCommentBubbleVisible ? ECheckBoxState::Checked : ECheckBoxState::Unchecked )
 					.OnCheckStateChanged( this, &SCommentBubble::OnCommentBubbleToggle )
 					.ToolTipText( NSLOCTEXT( "CommentBubble", "ToggleCommentTooltip", "Toggle Comment Bubble" ))
 					.Cursor( EMouseCursor::Default )
@@ -158,7 +158,7 @@ void SCommentBubble::UpdateBubble()
 				[
 					SNew( SCheckBox )
 					.Style( &FEditorStyle::Get().GetWidgetStyle< FCheckBoxStyle >( "CommentBubbleButton" ))
-					.IsChecked( GraphNode->bCommentBubbleVisible ? ESlateCheckBoxState::Checked : ESlateCheckBoxState::Unchecked )
+					.IsChecked( GraphNode->bCommentBubbleVisible ? ECheckBoxState::Checked : ECheckBoxState::Unchecked )
 					.OnCheckStateChanged( this, &SCommentBubble::OnCommentBubbleToggle )
 					.ToolTipText( NSLOCTEXT( "CommentBubble", "ToggleCommentTooltip", "Toggle Comment Bubble" ))
 					.Cursor( EMouseCursor::Default )
@@ -245,7 +245,7 @@ void SCommentBubble::UpdateBubble()
 			[
 				SNew( SCheckBox )
 				.Style( &FEditorStyle::Get().GetWidgetStyle< FCheckBoxStyle >( "CommentTitleButton" ))
-				.IsChecked( ESlateCheckBoxState::Unchecked )
+				.IsChecked( ECheckBoxState::Unchecked )
 				.OnCheckStateChanged( this, &SCommentBubble::OnCommentBubbleToggle )
 				.ToolTipText( NSLOCTEXT( "CommentBubble", "ToggleCommentTooltip", "Toggle Comment Bubble" ))
 				.Cursor( EMouseCursor::Default )
@@ -358,7 +358,7 @@ EVisibility SCommentBubble::GetBubbleVisibility() const
 	return bIsVisible ? EVisibility::Visible : EVisibility::Hidden;
 }
 
-void SCommentBubble::OnCommentBubbleToggle( ESlateCheckBoxState::Type State )
+void SCommentBubble::OnCommentBubbleToggle( ECheckBoxState State )
 {
 	if( GraphNode )
 	{
@@ -370,12 +370,12 @@ void SCommentBubble::OnCommentBubbleToggle( ESlateCheckBoxState::Type State )
 	}
 }
 
-void SCommentBubble::OnPinStateToggle( ESlateCheckBoxState::Type State ) const
+void SCommentBubble::OnPinStateToggle( ECheckBoxState State ) const
 {
 	if( GraphNode )
 	{
 		const FScopedTransaction Transaction( NSLOCTEXT( "CommentBubble", "BubblePinned", "Comment Bubble Pin" ) );
 		GraphNode->Modify();
-		GraphNode->bCommentBubblePinned = State == ESlateCheckBoxState::Checked;
+		GraphNode->bCommentBubblePinned = State == ECheckBoxState::Checked;
 	}
 }
\ No newline at end of file
diff --git a/Engine/Source/Editor/GraphEditor/Private/SGraphNode.cpp b/Engine/Source/Editor/GraphEditor/Private/SGraphNode.cpp
--- a/Engine/Source/Editor/GraphEditor/Private/SGraphNode.cpp
+++ b/Engine/Source/Editor/GraphEditor/Private/SGraphNode.cpp
@@ -860,19 +860,19 @@ EVisibility SGraphNode::AdvancedViewArrowVisibility() const
 	return bShowAdvancedViewArrow ? EVisibility::Visible : EVisibility::Collapsed;
 }
 
-void SGraphNode::OnAdvancedViewChanged( const ESlateCheckBoxState::Type NewCheckedState )
+void SGraphNode::OnAdvancedViewChanged( const ECheckBoxState NewCheckedState )
 {
 	if(GraphNode && (ENodeAdvancedPins::NoPins != GraphNode->AdvancedPinDisplay))
 	{
-		const bool bAdvancedPinsHidden = (NewCheckedState != ESlateCheckBoxState::Checked);
+		const bool bAdvancedPinsHidden = (NewCheckedState != ECheckBoxState::Checked);
 		GraphNode->AdvancedPinDisplay = bAdvancedPinsHidden ? ENodeAdvancedPins::Hidden : ENodeAdvancedPins::Shown;
 	}
 }
 
-ESlateCheckBoxState::Type SGraphNode::IsAdvancedViewChecked() const
+ECheckBoxState SGraphNode::IsAdvancedViewChecked() const
 {
 	const bool bAdvancedPinsHidden = GraphNode && (ENodeAdvancedPins::Hidden == GraphNode->AdvancedPinDisplay);
-	return bAdvancedPinsHidden ? ESlateCheckBoxState::Unchecked : ESlateCheckBoxState::Checked;
+	return bAdvancedPinsHidden ? ECheckBoxState::Unchecked : ECheckBoxState::Checked;
 }
 
 const FSlateBrush* SGraphNode::GetAdvancedViewArrow() const
diff --git a/Engine/Source/Editor/GraphEditor/Public/SCommentBubble.h b/Engine/Source/Editor/GraphEditor/Public/SCommentBubble.h
--- a/Engine/Source/Editor/GraphEditor/Public/SCommentBubble.h
+++ b/Engine/Source/Editor/GraphEditor/Public/SCommentBubble.h
@@ -70,10 +70,10 @@ class GRAPHEDITOR_API SCommentBubble : public SCompoundWidget
 	bool IsScalingAllowed() const;
 
 	/** Called to display/hide the comment bubble */
-	void OnCommentBubbleToggle( ESlateCheckBoxState::Type State );
+	void OnCommentBubbleToggle( ECheckBoxState State );
 
 	/** Called when a node's comment bubble pinned state is changed */
-	void OnPinStateToggle( ESlateCheckBoxState::Type State ) const;
+	void OnPinStateToggle( ECheckBoxState State ) const;
 
 	/** Called to update the bubble widget layout */
 	void UpdateBubble();
diff --git a/Engine/Source/Editor/GraphEditor/Public/SGraphNode.h b/Engine/Source/Editor/GraphEditor/Public/SGraphNode.h
--- a/Engine/Source/Editor/GraphEditor/Public/SGraphNode.h
+++ b/Engine/Source/Editor/GraphEditor/Public/SGraphNode.h
@@ -271,10 +271,10 @@ protected:
 	EVisibility AdvancedViewArrowVisibility() const;
 
 	/** Show/hide advanced view */
-	void OnAdvancedViewChanged( const ESlateCheckBoxState::Type NewCheckedState );
+	void OnAdvancedViewChanged( const ECheckBoxState NewCheckedState );
 
 	/** hidden == unchecked, shown == checked */
-	ESlateCheckBoxState::Type IsAdvancedViewChecked() const;
+	ECheckBoxState IsAdvancedViewChecked() const;
 
 	/** Up when shown, down when hidden */
 	const FSlateBrush* GetAdvancedViewArrow() const;
diff --git a/Engine/Source/Editor/IntroTutorials/Private/STutorialEditableText.cpp b/Engine/Source/Editor/IntroTutorials/Private/STutorialEditableText.cpp
--- a/Engine/Source/Editor/IntroTutorials/Private/STutorialEditableText.cpp
+++ b/Engine/Source/Editor/IntroTutorials/Private/STutorialEditableText.cpp
@@ -550,14 +550,14 @@ EVisibility STutorialEditableText::GetOpenAssetVisibility() const
 	return CurrentHyperlinkType.IsValid() && CurrentHyperlinkType->Type == EHyperlinkType::Asset ? EVisibility::Visible : EVisibility::Collapsed;
 }
 
-void STutorialEditableText::HandleOpenAssetCheckStateChanged(ESlateCheckBoxState::Type InCheckState)
+void STutorialEditableText::HandleOpenAssetCheckStateChanged(ECheckBoxState InCheckState)
 {
 	bOpenAsset = !bOpenAsset;
 }
 
-ESlateCheckBoxState::Type STutorialEditableText::IsOpenAssetChecked() const
+ECheckBoxState STutorialEditableText::IsOpenAssetChecked() const
 {
-	return bOpenAsset ? ESlateCheckBoxState::Checked : ESlateCheckBoxState::Unchecked;
+	return bOpenAsset ? ECheckBoxState::Checked : ECheckBoxState::Unchecked;
 }
 
 EVisibility STutorialEditableText::GetExcerptVisibility() const
diff --git a/Engine/Source/Editor/IntroTutorials/Private/STutorialEditableText.h b/Engine/Source/Editor/IntroTutorials/Private/STutorialEditableText.h
--- a/Engine/Source/Editor/IntroTutorials/Private/STutorialEditableText.h
+++ b/Engine/Source/Editor/IntroTutorials/Private/STutorialEditableText.h
@@ -64,9 +64,9 @@ protected:
 
 	EVisibility GetOpenAssetVisibility() const;
 
-	void HandleOpenAssetCheckStateChanged(ESlateCheckBoxState::Type InCheckState);
+	void HandleOpenAssetCheckStateChanged(ECheckBoxState InCheckState);
 
-	ESlateCheckBoxState::Type IsOpenAssetChecked() const;
+	ECheckBoxState IsOpenAssetChecked() const;
 
 	EVisibility GetExcerptVisibility() const;
 
diff --git a/Engine/Source/Editor/Kismet/Private/BlueprintDetailsCustomization.cpp b/Engine/Source/Editor/Kismet/Private/BlueprintDetailsCustomization.cpp
--- a/Engine/Source/Editor/Kismet/Private/BlueprintDetailsCustomization.cpp
+++ b/Engine/Source/Editor/Kismet/Private/BlueprintDetailsCustomization.cpp
@@ -673,7 +673,7 @@ TSharedRef<ITableRow> FBlueprintVarActionDetails::OnGenerateWidgetForPropertyLis
 				.AutoWidth()
 			[
 				SNew(SCheckBox)
-					.IsChecked(true ? ESlateCheckBoxState::Checked : ESlateCheckBoxState::Unchecked)
+					.IsChecked(true ? ECheckBoxState::Checked : ECheckBoxState::Unchecked)
 					.IsEnabled(false)
 			]
 		];
@@ -1204,45 +1204,45 @@ EVisibility FBlueprintVarActionDetails::ShowEditableCheckboxVisibilty() const
 	return EVisibility::Collapsed;
 }
 
-ESlateCheckBoxState::Type FBlueprintVarActionDetails::OnEditableCheckboxState() const
+ECheckBoxState FBlueprintVarActionDetails::OnEditableCheckboxState() const
 {
 	UProperty* VariableProperty = SelectionAsProperty();
 	if (VariableProperty)
 	{
-		return VariableProperty->HasAnyPropertyFlags(CPF_DisableEditOnInstance) ? ESlateCheckBoxState::Unchecked : ESlateCheckBoxState::Checked;
+		return VariableProperty->HasAnyPropertyFlags(CPF_DisableEditOnInstance) ? ECheckBoxState::Unchecked : ECheckBoxState::Checked;
 	}
-	return ESlateCheckBoxState::Unchecked;
+	return ECheckBoxState::Unchecked;
 }
 
-void FBlueprintVarActionDetails::OnEditableChanged(ESlateCheckBoxState::Type InNewState)
+void FBlueprintVarActionDetails::OnEditableChanged(ECheckBoxState InNewState)
 {
 	FName VarName = GetVariableName();
 
 	// Toggle the flag on the blueprint's version of the variable description, based on state
-	const bool bVariableIsExposed = InNewState == ESlateCheckBoxState::Checked;
+	const bool bVariableIsExposed = InNewState == ECheckBoxState::Checked;
 
 	UBlueprint* Blueprint = MyBlueprint.Pin()->GetBlueprintObj();
 	FBlueprintEditorUtils::SetBlueprintOnlyEditableFlag(Blueprint, VarName, !bVariableIsExposed);
 }
 
-ESlateCheckBoxState::Type FBlueprintVarActionDetails::OnCreateWidgetCheckboxState() const
+ECheckBoxState FBlueprintVarActionDetails::OnCreateWidgetCheckboxState() const
 {
 	UProperty* Property = SelectionAsProperty();
 	if (Property)
 	{
 		bool bMakingWidget = FEdMode::ShouldCreateWidgetForProperty(Property);
 
-		return bMakingWidget ? ESlateCheckBoxState::Checked : ESlateCheckBoxState::Unchecked;
+		return bMakingWidget ? ECheckBoxState::Checked : ECheckBoxState::Unchecked;
 	}
-	return ESlateCheckBoxState::Unchecked;
+	return ECheckBoxState::Unchecked;
 }
 
-void FBlueprintVarActionDetails::OnCreateWidgetChanged(ESlateCheckBoxState::Type InNewState)
+void FBlueprintVarActionDetails::OnCreateWidgetChanged(ECheckBoxState InNewState)
 {
 	const FName VarName = GetVariableName();
 	if (VarName != NAME_None)
 	{
-		if (InNewState == ESlateCheckBoxState::Checked)
+		if (InNewState == ECheckBoxState::Checked)
 		{
 			FBlueprintEditorUtils::SetBlueprintVariableMetaData(GetBlueprintObj(), VarName, GetLocalVariableScope(SelectionAsProperty()), FEdMode::MD_MakeEditWidget, TEXT("true"));
 		}
@@ -1276,22 +1276,22 @@ bool FBlueprintVarActionDetails::Is3DWidgetEnabled()
 	return false;
 }
 
-ESlateCheckBoxState::Type FBlueprintVarActionDetails::OnGetExposedToSpawnCheckboxState() const
+ECheckBoxState FBlueprintVarActionDetails::OnGetExposedToSpawnCheckboxState() const
 {
 	UProperty* Property = SelectionAsProperty();
 	if (Property)
 	{
-		return (Property && Property->GetBoolMetaData(FBlueprintMetadata::MD_ExposeOnSpawn) != false) ? ESlateCheckBoxState::Checked : ESlateCheckBoxState::Unchecked;
+		return (Property && Property->GetBoolMetaData(FBlueprintMetadata::MD_ExposeOnSpawn) != false) ? ECheckBoxState::Checked : ECheckBoxState::Unchecked;
 	}
-	return ESlateCheckBoxState::Unchecked;
+	return ECheckBoxState::Unchecked;
 }
 
-void FBlueprintVarActionDetails::OnExposedToSpawnChanged(ESlateCheckBoxState::Type InNewState)
+void FBlueprintVarActionDetails::OnExposedToSpawnChanged(ECheckBoxState InNewState)
 {
 	const FName VarName = GetVariableName();
 	if (VarName != NAME_None)
 	{
-		const bool bExposeOnSpawn = (InNewState == ESlateCheckBoxState::Checked);
+		const bool bExposeOnSpawn = (InNewState == ECheckBoxState::Checked);
 		if(bExposeOnSpawn)
 		{
 			FBlueprintEditorUtils::SetBlueprintVariableMetaData(GetBlueprintObj(), VarName, NULL, FBlueprintMetadata::MD_ExposeOnSpawn, TEXT("true"));
@@ -1321,22 +1321,22 @@ EVisibility FBlueprintVarActionDetails::ExposeOnSpawnVisibility() const
 	return EVisibility::Collapsed;
 }
 
-ESlateCheckBoxState::Type FBlueprintVarActionDetails::OnGetPrivateCheckboxState() const
+ECheckBoxState FBlueprintVarActionDetails::OnGetPrivateCheckboxState() const
 {
 	UProperty* Property = SelectionAsProperty();
 	if (Property)
 	{
-		return (Property && Property->GetBoolMetaData(FBlueprintMetadata::MD_Private) != false) ? ESlateCheckBoxState::Checked : ESlateCheckBoxState::Unchecked;
+		return (Property && Property->GetBoolMetaData(FBlueprintMetadata::MD_Private) != false) ? ECheckBoxState::Checked : ECheckBoxState::Unchecked;
 	}
-	return ESlateCheckBoxState::Unchecked;
+	return ECheckBoxState::Unchecked;
 }
 
-void FBlueprintVarActionDetails::OnPrivateChanged(ESlateCheckBoxState::Type InNewState)
+void FBlueprintVarActionDetails::OnPrivateChanged(ECheckBoxState InNewState)
 {
 	const FName VarName = GetVariableName();
 	if (VarName != NAME_None)
 	{
-		const bool bExposeOnSpawn = (InNewState == ESlateCheckBoxState::Checked);
+		const bool bExposeOnSpawn = (InNewState == ECheckBoxState::Checked);
 		if(bExposeOnSpawn)
 		{
 			FBlueprintEditorUtils::SetBlueprintVariableMetaData(GetBlueprintObj(), VarName, NULL, FBlueprintMetadata::MD_Private, TEXT("true"));
@@ -1361,20 +1361,20 @@ EVisibility FBlueprintVarActionDetails::ExposePrivateVisibility() const
 	return EVisibility::Collapsed;
 }
 
-ESlateCheckBoxState::Type FBlueprintVarActionDetails::OnGetExposedToMatineeCheckboxState() const
+ECheckBoxState FBlueprintVarActionDetails::OnGetExposedToMatineeCheckboxState() const
 {
 	UProperty* Property = SelectionAsProperty();
 	if (Property)
 	{
-		return Property && Property->HasAnyPropertyFlags(CPF_Interp) ? ESlateCheckBoxState::Checked : ESlateCheckBoxState::Unchecked;
+		return Property && Property->HasAnyPropertyFlags(CPF_Interp) ? ECheckBoxState::Checked : ECheckBoxState::Unchecked;
 	}
-	return ESlateCheckBoxState::Unchecked;
+	return ECheckBoxState::Unchecked;
 }
 
-void FBlueprintVarActionDetails::OnExposedToMatineeChanged(ESlateCheckBoxState::Type InNewState)
+void FBlueprintVarActionDetails::OnExposedToMatineeChanged(ECheckBoxState InNewState)
 {
 	// Toggle the flag on the blueprint's version of the variable description, based on state
-	const bool bExposeToMatinee = (InNewState == ESlateCheckBoxState::Checked);
+	const bool bExposeToMatinee = (InNewState == ECheckBoxState::Checked);
 	
 	const FName VarName = GetVariableName();
 	if (VarName != NAME_None)
@@ -1601,22 +1601,22 @@ EVisibility FBlueprintVarActionDetails::GetTransientVisibility() const
 	return EVisibility::Collapsed;
 }
 
-ESlateCheckBoxState::Type FBlueprintVarActionDetails::OnGetTransientCheckboxState() const
+ECheckBoxState FBlueprintVarActionDetails::OnGetTransientCheckboxState() const
 {
 	UProperty* Property = SelectionAsProperty();
 	if (Property)
 	{
-		return (Property && Property->HasAnyPropertyFlags(CPF_Transient)) ? ESlateCheckBoxState::Checked : ESlateCheckBoxState::Unchecked;
+		return (Property && Property->HasAnyPropertyFlags(CPF_Transient)) ? ECheckBoxState::Checked : ECheckBoxState::Unchecked;
 	}
-	return ESlateCheckBoxState::Unchecked;
+	return ECheckBoxState::Unchecked;
 }
 
-void FBlueprintVarActionDetails::OnTransientChanged(ESlateCheckBoxState::Type InNewState)
+void FBlueprintVarActionDetails::OnTransientChanged(ECheckBoxState InNewState)
 {
 	UProperty* Property = SelectionAsProperty();
 	if (Property)
 	{
-		const bool bTransientFlag = (InNewState == ESlateCheckBoxState::Checked);
+		const bool bTransientFlag = (InNewState == ECheckBoxState::Checked);
 		FBlueprintEditorUtils::SetVariableTransientFlag(GetBlueprintObj(), Property->GetFName(), bTransientFlag);
 	}
 }
@@ -1634,22 +1634,22 @@ EVisibility FBlueprintVarActionDetails::GetSaveGameVisibility() const
 	return EVisibility::Collapsed;
 }
 
-ESlateCheckBoxState::Type FBlueprintVarActionDetails::OnGetSaveGameCheckboxState() const
+ECheckBoxState FBlueprintVarActionDetails::OnGetSaveGameCheckboxState() const
 {
 	UProperty* Property = SelectionAsProperty();
 	if (Property)
 	{
-		return (Property && Property->HasAnyPropertyFlags(CPF_SaveGame)) ? ESlateCheckBoxState::Checked : ESlateCheckBoxState::Unchecked;
+		return (Property && Property->HasAnyPropertyFlags(CPF_SaveGame)) ? ECheckBoxState::Checked : ECheckBoxState::Unchecked;
 	}
-	return ESlateCheckBoxState::Unchecked;
+	return ECheckBoxState::Unchecked;
 }
 
-void FBlueprintVarActionDetails::OnSaveGameChanged(ESlateCheckBoxState::Type InNewState)
+void FBlueprintVarActionDetails::OnSaveGameChanged(ECheckBoxState InNewState)
 {
 	UProperty* Property = SelectionAsProperty();
 	if (Property)
 	{
-		const bool bSaveGameFlag = (InNewState == ESlateCheckBoxState::Checked);
+		const bool bSaveGameFlag = (InNewState == ECheckBoxState::Checked);
 		FBlueprintEditorUtils::SetVariableSaveGameFlag(GetBlueprintObj(), Property->GetFName(), bSaveGameFlag);
 	}
 }
@@ -2020,16 +2020,16 @@ FEdGraphPinType FBlueprintGraphArgumentLayout::OnGetPinInfo() const
 	return FEdGraphPinType();
 }
 
-ESlateCheckBoxState::Type FBlueprintGraphArgumentLayout::IsRefChecked() const
+ECheckBoxState FBlueprintGraphArgumentLayout::IsRefChecked() const
 {
 	FEdGraphPinType PinType = OnGetPinInfo();
-	return PinType.bIsReference? ESlateCheckBoxState::Checked : ESlateCheckBoxState::Unchecked;
+	return PinType.bIsReference? ECheckBoxState::Checked : ECheckBoxState::Unchecked;
 }
 
-void FBlueprintGraphArgumentLayout::OnRefCheckStateChanged(ESlateCheckBoxState::Type InState)
+void FBlueprintGraphArgumentLayout::OnRefCheckStateChanged(ECheckBoxState InState)
 {
 	FEdGraphPinType PinType = OnGetPinInfo();
-	PinType.bIsReference = (InState == ESlateCheckBoxState::Checked)? true : false;
+	PinType.bIsReference = (InState == ECheckBoxState::Checked)? true : false;
 	PinInfoChanged(PinType);
 }
 
@@ -2574,9 +2574,9 @@ void FBaseBlueprintGraphActionDetails::SetRefreshDelegate(FSimpleDelegate Refres
 	((bForInputs) ? RegenerateInputsChildrenDelegate : RegenerateOutputsChildrenDelegate) = RefreshDelegate;
 }
 
-ESlateCheckBoxState::Type FBlueprintGraphActionDetails::GetIsEditorCallableEvent() const
+ECheckBoxState FBlueprintGraphActionDetails::GetIsEditorCallableEvent() const
 {
-	ESlateCheckBoxState::Type Result = ESlateCheckBoxState::Unchecked;
+	ECheckBoxState Result = ECheckBoxState::Unchecked;
 
 	if( FunctionEntryNodePtr.IsValid() )
 	{
@@ -2584,13 +2584,13 @@ ESlateCheckBoxState::Type FBlueprintGraphActionDetails::GetIsEditorCallableEvent
 
 		if( CustomEventNode && CustomEventNode->bCallInEditor )
 		{
-			Result = ESlateCheckBoxState::Checked;
+			Result = ECheckBoxState::Checked;
 		}
 	}
 	return Result;
 }
 
-void FBlueprintGraphActionDetails::OnEditorCallableEventModified( const ESlateCheckBoxState::Type NewCheckedState ) const
+void FBlueprintGraphActionDetails::OnEditorCallableEventModified( const ECheckBoxState NewCheckedState ) const
 {
 	if( FunctionEntryNodePtr.IsValid() )
 	{
@@ -2598,7 +2598,7 @@ void FBlueprintGraphActionDetails::OnEditorCallableEventModified( const ESlateCh
 		{
 			if( UBlueprint* Blueprint = FunctionEntryNodePtr->GetBlueprint() )
 			{
-				const bool bCallInEditor = NewCheckedState == ESlateCheckBoxState::Checked;
+				const bool bCallInEditor = NewCheckedState == ECheckBoxState::Checked;
 				const FText TransactionType = bCallInEditor ?	LOCTEXT( "DisableCallInEditor", "Disable Call In Editor " ) : 
 																LOCTEXT( "EnableCallInEditor", "Enable Call In Editor" );
 				const FScopedTransaction Transaction( TransactionType );
@@ -3605,13 +3605,13 @@ bool FBlueprintGraphActionDetails::IsCustomEvent() const
 	return (NULL != Cast<UK2Node_CustomEvent>(FunctionEntryNodePtr.Get()));
 }
 
-void FBlueprintGraphActionDetails::OnIsReliableReplicationFunctionModified(const ESlateCheckBoxState::Type NewCheckedState)
+void FBlueprintGraphActionDetails::OnIsReliableReplicationFunctionModified(const ECheckBoxState NewCheckedState)
 {
 	UK2Node_EditablePinBase * FunctionEntryNode = FunctionEntryNodePtr.Get();
 	UK2Node_CustomEvent* CustomEvent = Cast<UK2Node_CustomEvent>(FunctionEntryNode);
 	if( CustomEvent )
 	{
-		if (NewCheckedState == ESlateCheckBoxState::Checked)
+		if (NewCheckedState == ECheckBoxState::Checked)
 		{
 			if (UK2Node_FunctionEntry* TypedEntryNode = Cast<UK2Node_FunctionEntry>(FunctionEntryNode))
 			{
@@ -3638,22 +3638,22 @@ void FBlueprintGraphActionDetails::OnIsReliableReplicationFunctionModified(const
 	}
 }
 
-ESlateCheckBoxState::Type FBlueprintGraphActionDetails::GetIsReliableReplicatedFunction() const
+ECheckBoxState FBlueprintGraphActionDetails::GetIsReliableReplicatedFunction() const
 {
 	const UK2Node_EditablePinBase * FunctionEntryNode = FunctionEntryNodePtr.Get();
 	const UK2Node_CustomEvent* CustomEvent = Cast<const UK2Node_CustomEvent>(FunctionEntryNode);
 	if(!CustomEvent)
 	{
-		return ESlateCheckBoxState::Undetermined;
+		return ECheckBoxState::Undetermined;
 	}
 
 	uint32 const NetReliableMask = (FUNC_Net | FUNC_NetReliable);
 	if ((CustomEvent->GetNetFlags() & NetReliableMask) == NetReliableMask)
 	{
-		return ESlateCheckBoxState::Checked;
+		return ECheckBoxState::Checked;
 	}
 	
-	return ESlateCheckBoxState::Unchecked;
+	return ECheckBoxState::Unchecked;
 }
 
 bool FBlueprintGraphActionDetails::IsPureFunctionVisible() const
@@ -3672,7 +3672,7 @@ bool FBlueprintGraphActionDetails::IsPureFunctionVisible() const
 	return bSupportedType && bIsEditable;
 }
 
-void FBlueprintGraphActionDetails::OnIsPureFunctionModified( const ESlateCheckBoxState::Type NewCheckedState )
+void FBlueprintGraphActionDetails::OnIsPureFunctionModified( const ECheckBoxState NewCheckedState )
 {
 	UK2Node_EditablePinBase * FunctionEntryNode = FunctionEntryNodePtr.Get();
 	auto Function = FindFunction();
@@ -3690,15 +3690,15 @@ void FBlueprintGraphActionDetails::OnIsPureFunctionModified( const ESlateCheckBo
 	}
 }
 
-ESlateCheckBoxState::Type FBlueprintGraphActionDetails::GetIsPureFunction() const
+ECheckBoxState FBlueprintGraphActionDetails::GetIsPureFunction() const
 {
 	UK2Node_EditablePinBase * FunctionEntryNode = FunctionEntryNodePtr.Get();
 	auto EntryNode = Cast<UK2Node_FunctionEntry>(FunctionEntryNode);
 	if(!EntryNode)
 	{
-		return ESlateCheckBoxState::Undetermined;
+		return ECheckBoxState::Undetermined;
 	}
-	return (EntryNode->ExtraFlags & FUNC_BlueprintPure) ? ESlateCheckBoxState::Checked :  ESlateCheckBoxState::Unchecked;
+	return (EntryNode->ExtraFlags & FUNC_BlueprintPure) ? ECheckBoxState::Checked :  ECheckBoxState::Unchecked;
 }
 
 bool FBaseBlueprintGraphActionDetails::IsPinNameUnique(const FString& TestName) const
@@ -4118,15 +4118,15 @@ bool FBlueprintGlobalOptionsDetails::CanDeprecateBlueprint() const
 	return true;
 }
 
-void FBlueprintGlobalOptionsDetails::OnDeprecateBlueprint(ESlateCheckBoxState::Type InCheckState)
+void FBlueprintGlobalOptionsDetails::OnDeprecateBlueprint(ECheckBoxState InCheckState)
 {
-	GetBlueprintObj()->bDeprecate = InCheckState == ESlateCheckBoxState::Checked? true : false;
+	GetBlueprintObj()->bDeprecate = InCheckState == ECheckBoxState::Checked? true : false;
 	FBlueprintEditorUtils::MarkBlueprintAsStructurallyModified(GetBlueprintObj());
 }
 
-ESlateCheckBoxState::Type FBlueprintGlobalOptionsDetails::IsDeprecatedBlueprint() const
+ECheckBoxState FBlueprintGlobalOptionsDetails::IsDeprecatedBlueprint() const
 {
-	return GetBlueprintObj()->bDeprecate? ESlateCheckBoxState::Checked : ESlateCheckBoxState::Unchecked;
+	return GetBlueprintObj()->bDeprecate? ECheckBoxState::Checked : ECheckBoxState::Unchecked;
 }
 
 FText FBlueprintGlobalOptionsDetails::GetDeprecatedTooltip() const
diff --git a/Engine/Source/Editor/Kismet/Private/BlueprintDetailsCustomization.h b/Engine/Source/Editor/Kismet/Private/BlueprintDetailsCustomization.h
--- a/Engine/Source/Editor/Kismet/Private/BlueprintDetailsCustomization.h
+++ b/Engine/Source/Editor/Kismet/Private/BlueprintDetailsCustomization.h
@@ -100,24 +100,24 @@ private:
 	void OnCategorySelectionChanged( TSharedPtr<FString> ProposedSelection, ESelectInfo::Type /*SelectInfo*/ );
 	
 	EVisibility ShowEditableCheckboxVisibilty() const;
-	ESlateCheckBoxState::Type OnEditableCheckboxState() const;
-	void OnEditableChanged(ESlateCheckBoxState::Type InNewState);
+	ECheckBoxState OnEditableCheckboxState() const;
+	void OnEditableChanged(ECheckBoxState InNewState);
 
-	ESlateCheckBoxState::Type OnCreateWidgetCheckboxState() const;
-	void OnCreateWidgetChanged(ESlateCheckBoxState::Type InNewState);
+	ECheckBoxState OnCreateWidgetCheckboxState() const;
+	void OnCreateWidgetChanged(ECheckBoxState InNewState);
 	EVisibility Show3DWidgetVisibility() const;
 	bool Is3DWidgetEnabled();
 	
-	ESlateCheckBoxState::Type OnGetExposedToSpawnCheckboxState() const;
-	void OnExposedToSpawnChanged(ESlateCheckBoxState::Type InNewState);
+	ECheckBoxState OnGetExposedToSpawnCheckboxState() const;
+	void OnExposedToSpawnChanged(ECheckBoxState InNewState);
 	EVisibility ExposeOnSpawnVisibility() const;
 
-	ESlateCheckBoxState::Type OnGetPrivateCheckboxState() const;
-	void OnPrivateChanged(ESlateCheckBoxState::Type InNewState);
+	ECheckBoxState OnGetPrivateCheckboxState() const;
+	void OnPrivateChanged(ECheckBoxState InNewState);
 	EVisibility ExposePrivateVisibility() const;
 	
-	ESlateCheckBoxState::Type OnGetExposedToMatineeCheckboxState() const;
-	void OnExposedToMatineeChanged(ESlateCheckBoxState::Type InNewState);
+	ECheckBoxState OnGetExposedToMatineeCheckboxState() const;
+	void OnExposedToMatineeChanged(ECheckBoxState InNewState);
 	EVisibility ExposeToMatineeVisibility() const;
 
 	FText OnGetMetaKeyValue(FName Key) const;
@@ -130,12 +130,12 @@ private:
 	EVisibility ReplicationVisibility() const;
 
 	EVisibility GetTransientVisibility() const;
-	ESlateCheckBoxState::Type OnGetTransientCheckboxState() const;
-	void OnTransientChanged(ESlateCheckBoxState::Type InNewState);
+	ECheckBoxState OnGetTransientCheckboxState() const;
+	void OnTransientChanged(ECheckBoxState InNewState);
 
 	EVisibility GetSaveGameVisibility() const;
-	ESlateCheckBoxState::Type OnGetSaveGameCheckboxState() const;
-	void OnSaveGameChanged(ESlateCheckBoxState::Type InNewState);
+	ECheckBoxState OnGetSaveGameCheckboxState() const;
+	void OnSaveGameChanged(ECheckBoxState InNewState);
 
 	/** Refresh the property flags list */
 	void RefreshPropertyFlags();
@@ -226,7 +226,7 @@ protected:
 	bool AttemptToCreateResultNode();
 
 	/** Toggles the ability to be called in editor */
-	void OnToggleEditorCallableEvent( const ESlateCheckBoxState::Type NewCheckedState, TWeakObjectPtr<UK2Node_EditablePinBase> SelectedNode ) const;
+	void OnToggleEditorCallableEvent( const ECheckBoxState NewCheckedState, TWeakObjectPtr<UK2Node_EditablePinBase> SelectedNode ) const;
 
 protected:
 	/** Pointer to the parent */
@@ -366,10 +366,10 @@ private:
 	void OnPrePinInfoChange(const FEdGraphPinType& PinType);
 
 	/** Returns whether the "Pass-by-Reference" checkbox is checked or not */
-	ESlateCheckBoxState::Type IsRefChecked() const;
+	ECheckBoxState IsRefChecked() const;
 
 	/** Handles toggling the "Pass-by-Reference" checkbox */
-	void OnRefCheckStateChanged(ESlateCheckBoxState::Type InState);
+	void OnRefCheckStateChanged(ECheckBoxState InState);
 
 	FText OnGetArgDefaultValueText() const;
 	void OnArgDefaultValueCommitted(const FText& NewText, ETextCommit::Type InTextCommit);
@@ -455,8 +455,8 @@ private:
 	FReply ColorBlock_OnMouseButtonDown(const FGeometry& MyGeometry, const FPointerEvent& MouseEvent);
 	
 	bool IsCustomEvent() const;
-	void OnIsReliableReplicationFunctionModified(const ESlateCheckBoxState::Type NewCheckedState);
-	ESlateCheckBoxState::Type GetIsReliableReplicatedFunction() const;
+	void OnIsReliableReplicationFunctionModified(const ECheckBoxState NewCheckedState);
+	ECheckBoxState GetIsReliableReplicatedFunction() const;
 
 	struct FReplicationSpecifierLabel
 	{
@@ -474,14 +474,14 @@ private:
 	TSharedRef<ITableRow> OnGenerateReplicationComboWidget( TSharedPtr<FReplicationSpecifierLabel> InNetFlag, const TSharedRef<STableViewBase>& OwnerTable );
 	
 	bool IsPureFunctionVisible() const;
-	void OnIsPureFunctionModified(const ESlateCheckBoxState::Type NewCheckedState);
-	ESlateCheckBoxState::Type GetIsPureFunction() const;
+	void OnIsPureFunctionModified(const ECheckBoxState NewCheckedState);
+	ECheckBoxState GetIsPureFunction() const;
 
 	/** Determines if the selected event is identified as editor callable */
-	ESlateCheckBoxState::Type GetIsEditorCallableEvent() const;
+	ECheckBoxState GetIsEditorCallableEvent() const;
 
 	/** Enables/Disables selected event as editor callable  */
-	void OnEditorCallableEventModified( const ESlateCheckBoxState::Type NewCheckedState ) const;
+	void OnEditorCallableEventModified( const ECheckBoxState NewCheckedState ) const;
 
 	
 	FReply OnAddNewOutputClicked();
@@ -623,10 +623,10 @@ protected:
 	bool CanDeprecateBlueprint() const;
 
 	/** Callback when toggling the Deprecate checkbox, handles marking a Blueprint as deprecated */
-	void OnDeprecateBlueprint(ESlateCheckBoxState::Type InCheckState);
+	void OnDeprecateBlueprint(ECheckBoxState InCheckState);
 
 	/** Callback for Deprecate checkbox, returns checked if the Blueprint is deprecated */
-	ESlateCheckBoxState::Type IsDeprecatedBlueprint() const;
+	ECheckBoxState IsDeprecatedBlueprint() const;
 
 	/** Returns the tooltip explaining deprecation */
 	FText GetDeprecatedTooltip() const;
diff --git a/Engine/Source/Editor/Kismet/Private/FindInBlueprints.cpp b/Engine/Source/Editor/Kismet/Private/FindInBlueprints.cpp
--- a/Engine/Source/Editor/Kismet/Private/FindInBlueprints.cpp
+++ b/Engine/Source/Editor/Kismet/Private/FindInBlueprints.cpp
@@ -1188,14 +1188,14 @@ void SFindInBlueprints::OnSearchTextCommitted( const FText& Text, ETextCommit::T
 	}
 }
 
-void SFindInBlueprints::OnFindModeChanged(ESlateCheckBoxState::Type CheckState)
+void SFindInBlueprints::OnFindModeChanged(ECheckBoxState CheckState)
 {
-	bIsInFindWithinBlueprintMode = CheckState == ESlateCheckBoxState::Checked;
+	bIsInFindWithinBlueprintMode = CheckState == ECheckBoxState::Checked;
 }
 
-ESlateCheckBoxState::Type SFindInBlueprints::OnGetFindModeChecked() const
+ECheckBoxState SFindInBlueprints::OnGetFindModeChecked() const
 {
-	return bIsInFindWithinBlueprintMode ? ESlateCheckBoxState::Checked : ESlateCheckBoxState::Unchecked;
+	return bIsInFindWithinBlueprintMode ? ECheckBoxState::Checked : ECheckBoxState::Unchecked;
 }
 
 void SFindInBlueprints::InitiateSearch()
diff --git a/Engine/Source/Editor/Kismet/Private/SBlueprintActionMenu.cpp b/Engine/Source/Editor/Kismet/Private/SBlueprintActionMenu.cpp
--- a/Engine/Source/Editor/Kismet/Private/SBlueprintActionMenu.cpp
+++ b/Engine/Source/Editor/Kismet/Private/SBlueprintActionMenu.cpp
@@ -137,7 +137,7 @@ private:
 	 */
 	FText GetToolTipText() const
 	{
-		if (GetFavoritedState() == ESlateCheckBoxState::Checked)
+		if (GetFavoritedState() == ECheckBoxState::Checked)
 		{
 			return LOCTEXT("Unfavorite", "Click to remove this item from your favorites.");
 		}
@@ -148,14 +148,14 @@ private:
 	 * Checks on the associated action's favorite state, and returns a 
 	 * corresponding checkbox state to match.
 	 * 
-	 * @return ESlateCheckBoxState::Checked if the associated action is already favorited, ESlateCheckBoxState::Unchecked if not.
+	 * @return ECheckBoxState::Checked if the associated action is already favorited, ECheckBoxState::Unchecked if not.
 	 */
-	ESlateCheckBoxState::Type GetFavoritedState() const
+	ECheckBoxState GetFavoritedState() const
 	{
-		ESlateCheckBoxState::Type FavoriteState = ESlateCheckBoxState::Unchecked;
+		ECheckBoxState FavoriteState = ECheckBoxState::Unchecked;
 		if (GEditor->EditorUserSettings->BlueprintFavorites->IsFavorited(ActionPtr.Pin()))
 		{
-			FavoriteState = ESlateCheckBoxState::Checked;
+			FavoriteState = ECheckBoxState::Checked;
 		}
 		return FavoriteState;
 	}
@@ -166,9 +166,9 @@ private:
 	 * 
 	 * @param  InNewState	The new state that the user set the checkbox to.
 	 */
-	void OnFavoriteToggled(ESlateCheckBoxState::Type InNewState)
+	void OnFavoriteToggled(ECheckBoxState InNewState)
 	{
-		if (InNewState == ESlateCheckBoxState::Checked)
+		if (InNewState == ECheckBoxState::Checked)
 		{
 			GEditor->EditorUserSettings->BlueprintFavorites->AddFavorite(ActionPtr.Pin());
 		}
@@ -193,7 +193,7 @@ private:
 SBlueprintActionMenu::~SBlueprintActionMenu()
 {
 	OnClosedCallback.ExecuteIfBound();
-	OnCloseReasonCallback.ExecuteIfBound(bActionExecuted, ContextToggleIsChecked() == ESlateCheckBoxState::Checked, DraggedFromPins.Num() > 0);
+	OnCloseReasonCallback.ExecuteIfBound(bActionExecuted, ContextToggleIsChecked() == ECheckBoxState::Checked, DraggedFromPins.Num() > 0);
 }
 
 void SBlueprintActionMenu::Construct( const FArguments& InArgs, TSharedPtr<FBlueprintEditor> InEditor )
@@ -374,15 +374,15 @@ FText SBlueprintActionMenu::GetSearchContextDesc() const
 	}
 }
 
-void SBlueprintActionMenu::OnContextToggleChanged(ESlateCheckBoxState::Type CheckState)
+void SBlueprintActionMenu::OnContextToggleChanged(ECheckBoxState CheckState)
 {
-	EditorPtr.Pin()->GetIsContextSensitive() = CheckState == ESlateCheckBoxState::Checked;
+	EditorPtr.Pin()->GetIsContextSensitive() = CheckState == ECheckBoxState::Checked;
 	GraphActionMenu->RefreshAllActions(true, false);
 }
 
-ESlateCheckBoxState::Type SBlueprintActionMenu::ContextToggleIsChecked() const
+ECheckBoxState SBlueprintActionMenu::ContextToggleIsChecked() const
 {
-	return EditorPtr.Pin()->GetIsContextSensitive() ? ESlateCheckBoxState::Checked : ESlateCheckBoxState::Unchecked;
+	return EditorPtr.Pin()->GetIsContextSensitive() ? ECheckBoxState::Checked : ECheckBoxState::Unchecked;
 }
 
 void SBlueprintActionMenu::CollectAllActions(FGraphActionListBuilderBase& OutAllActions)
diff --git a/Engine/Source/Editor/Kismet/Private/SBlueprintActionMenu.h b/Engine/Source/Editor/Kismet/Private/SBlueprintActionMenu.h
--- a/Engine/Source/Editor/Kismet/Private/SBlueprintActionMenu.h
+++ b/Engine/Source/Editor/Kismet/Private/SBlueprintActionMenu.h
@@ -40,8 +40,8 @@ protected:
 	/** UI Callback functions */
 	EVisibility GetTypeImageVisibility() const;
 	FText GetSearchContextDesc() const;
-	void OnContextToggleChanged(ESlateCheckBoxState::Type CheckState);
-	ESlateCheckBoxState::Type ContextToggleIsChecked() const;
+	void OnContextToggleChanged(ECheckBoxState CheckState);
+	ECheckBoxState ContextToggleIsChecked() const;
 
 	void OnActionSelected( const TArray< TSharedPtr<FEdGraphSchemaAction> >& SelectedAction );
 
diff --git a/Engine/Source/Editor/Kismet/Private/SBlueprintPalette.cpp b/Engine/Source/Editor/Kismet/Private/SBlueprintPalette.cpp
--- a/Engine/Source/Editor/Kismet/Private/SBlueprintPalette.cpp
+++ b/Engine/Source/Editor/Kismet/Private/SBlueprintPalette.cpp
@@ -742,19 +742,19 @@ private:
 	 * Used by this visibility-toggle widget to see if the property represented 
 	 * by this item is visible outside of Kismet.
 	 * 
-	 * @return ESlateCheckBoxState::Checked if the property is visible, false if not (or if the property wasn't found)
+	 * @return ECheckBoxState::Checked if the property is visible, false if not (or if the property wasn't found)
 	 */
-	ESlateCheckBoxState::Type GetVisibilityToggleState() const
+	ECheckBoxState GetVisibilityToggleState() const
 	{
 		TSharedPtr<FEdGraphSchemaAction_K2Var> VarAction = StaticCastSharedPtr<FEdGraphSchemaAction_K2Var>(ActionPtr.Pin());
 		UProperty* VariableProperty = VarAction->GetProperty();
 		if (VariableProperty != NULL)
 		{
-			return VariableProperty->HasAnyPropertyFlags(CPF_DisableEditOnInstance) ? ESlateCheckBoxState::Unchecked : ESlateCheckBoxState::Checked;
+			return VariableProperty->HasAnyPropertyFlags(CPF_DisableEditOnInstance) ? ECheckBoxState::Unchecked : ECheckBoxState::Checked;
 		}
 		else
 		{
-			return ESlateCheckBoxState::Unchecked;
+			return ECheckBoxState::Unchecked;
 		}
 	}
 
@@ -765,7 +765,7 @@ private:
 	 * 
 	 * @param  InNewState	The new state that the user set the checkbox to.
 	 */
-	void OnVisibilityToggleFlipped(ESlateCheckBoxState::Type InNewState)
+	void OnVisibilityToggleFlipped(ECheckBoxState InNewState)
 	{
 		if( !BlueprintEditorPtr.IsValid() )
 		{
@@ -775,7 +775,7 @@ private:
 		TSharedPtr<FEdGraphSchemaAction_K2Var> VarAction = StaticCastSharedPtr<FEdGraphSchemaAction_K2Var>(ActionPtr.Pin());
 
 		// Toggle the flag on the blueprint's version of the variable description, based on state
-		const bool bVariableIsExposed = (InNewState == ESlateCheckBoxState::Checked);
+		const bool bVariableIsExposed = (InNewState == ECheckBoxState::Checked);
 
 		FBlueprintEditorUtils::SetBlueprintOnlyEditableFlag(BlueprintObj, VarAction->GetVariableName(), !bVariableIsExposed);
 	}
@@ -788,7 +788,7 @@ private:
 	 */
 	const FSlateBrush* GetVisibilityIcon() const
 	{
-		return GetVisibilityToggleState() == ESlateCheckBoxState::Checked ?
+		return GetVisibilityToggleState() == ECheckBoxState::Checked ?
 			FEditorStyle::GetBrush( "Kismet.VariableList.ExposeForInstance" ) :
 			FEditorStyle::GetBrush( "Kismet.VariableList.HideForInstance" );
 	}
@@ -802,7 +802,7 @@ private:
 	 */
 	FLinearColor GetVisibilityToggleColor() const 
 	{
-		if(GetVisibilityToggleState() != ESlateCheckBoxState::Checked)
+		if(GetVisibilityToggleState() != ECheckBoxState::Checked)
 		{
 			return FColor(64,64,64).ReinterpretAsLinear();
 		}
@@ -834,7 +834,7 @@ private:
 	FString GetVisibilityToggleToolTip() const
 	{
 		FString ToolTip;
-		if(GetVisibilityToggleState() != ESlateCheckBoxState::Checked)
+		if(GetVisibilityToggleState() != ECheckBoxState::Checked)
 		{
 			ToolTip =  LOCTEXT("VariablePrivacy_not_public_Tooltip", "Variable is not public and will not be editable on an instance of this Blueprint.").ToString();
 		}
diff --git a/Engine/Source/Editor/Kismet/Private/SKismetInspector.cpp b/Engine/Source/Editor/Kismet/Private/SKismetInspector.cpp
--- a/Engine/Source/Editor/Kismet/Private/SKismetInspector.cpp
+++ b/Engine/Source/Editor/Kismet/Private/SKismetInspector.cpp
@@ -187,7 +187,7 @@ FString SKismetInspector::GetContextualEditingWidgetTitle() const
 void SKismetInspector::Construct(const FArguments& InArgs)
 {
 	bShowInspectorPropertyView = true;
-	PublicViewState = ESlateCheckBoxState::Unchecked;
+	PublicViewState = ECheckBoxState::Unchecked;
 
 	Kismet2Ptr = InArgs._Kismet2;
 	bShowPublicView = InArgs._ShowPublicViewControl;
@@ -487,7 +487,7 @@ bool SKismetInspector::IsPropertyVisible( const FPropertyAndParent& PropertyAndP
 
 
 	// If we are in 'instance preview' - hide anything marked 'disabled edit on instance'
-	if ((ESlateCheckBoxState::Checked == PublicViewState) && Property.HasAnyPropertyFlags(CPF_DisableEditOnInstance))
+	if ((ECheckBoxState::Checked == PublicViewState) && Property.HasAnyPropertyFlags(CPF_DisableEditOnInstance))
 	{
 		return false;
 	}
@@ -547,12 +547,12 @@ EVisibility SKismetInspector::GetPropertyViewVisibility() const
 	return bShowInspectorPropertyView? EVisibility::Visible : EVisibility::Collapsed;
 }
 
-ESlateCheckBoxState::Type SKismetInspector::GetPublicViewCheckboxState() const
+ECheckBoxState SKismetInspector::GetPublicViewCheckboxState() const
 {
 	return PublicViewState;
 }
 
-void SKismetInspector::SetPublicViewCheckboxState( ESlateCheckBoxState::Type InIsChecked )
+void SKismetInspector::SetPublicViewCheckboxState( ECheckBoxState InIsChecked )
 {
 	PublicViewState = InIsChecked;
 
diff --git a/Engine/Source/Editor/Kismet/Private/SMyBlueprint.cpp b/Engine/Source/Editor/Kismet/Private/SMyBlueprint.cpp
--- a/Engine/Source/Editor/Kismet/Private/SMyBlueprint.cpp
+++ b/Engine/Source/Editor/Kismet/Private/SMyBlueprint.cpp
@@ -1045,14 +1045,14 @@ void SMyBlueprint::CollectAllActions(FGraphActionListBuilderBase& OutAllActions)
 	}
 }
 
-ESlateCheckBoxState::Type SMyBlueprint::OnUserVarsCheckState() const
+ECheckBoxState SMyBlueprint::OnUserVarsCheckState() const
 {
-	return !bShowUserVarsOnly ? ESlateCheckBoxState::Checked : ESlateCheckBoxState::Unchecked;
+	return !bShowUserVarsOnly ? ECheckBoxState::Checked : ECheckBoxState::Unchecked;
 }
 
-void SMyBlueprint::OnUserVarsCheckStateChanged(ESlateCheckBoxState::Type InNewState)
+void SMyBlueprint::OnUserVarsCheckStateChanged(ECheckBoxState InNewState)
 {
-	bShowUserVarsOnly = (InNewState != ESlateCheckBoxState::Checked);
+	bShowUserVarsOnly = (InNewState != ECheckBoxState::Checked);
 	Refresh();
 }
 
diff --git a/Engine/Source/Editor/Kismet/Private/SMyBlueprint.h b/Engine/Source/Editor/Kismet/Private/SMyBlueprint.h
--- a/Engine/Source/Editor/Kismet/Private/SMyBlueprint.h
+++ b/Engine/Source/Editor/Kismet/Private/SMyBlueprint.h
@@ -157,8 +157,8 @@ private:
 	FText OnGetSectionTitle( int32 InSectionID );
 
 	/** Support functions for checkbox to manage displaying user variables only */
-	ESlateCheckBoxState::Type OnUserVarsCheckState() const;
-	void OnUserVarsCheckStateChanged(ESlateCheckBoxState::Type InNewState);
+	ECheckBoxState OnUserVarsCheckState() const;
+	void OnUserVarsCheckStateChanged(ECheckBoxState InNewState);
 	
 	/** Helper function to open the selected graph */
 	void OpenGraph(FDocumentTracker::EOpenDocumentCause InCause);
diff --git a/Engine/Source/Editor/Kismet/Private/STimelineEditor.cpp b/Engine/Source/Editor/Kismet/Private/STimelineEditor.cpp
--- a/Engine/Source/Editor/Kismet/Private/STimelineEditor.cpp
+++ b/Engine/Source/Editor/Kismet/Private/STimelineEditor.cpp
@@ -1204,16 +1204,16 @@ void STimelineEditor::OnItemScrolledIntoView( TSharedPtr<FTimelineEdTrack> InTra
 	}
 }
 
-ESlateCheckBoxState::Type STimelineEditor::IsAutoPlayChecked() const
+ECheckBoxState STimelineEditor::IsAutoPlayChecked() const
 {
-	return (TimelineObj && TimelineObj->bAutoPlay) ? ESlateCheckBoxState::Checked : ESlateCheckBoxState::Unchecked;
+	return (TimelineObj && TimelineObj->bAutoPlay) ? ECheckBoxState::Checked : ECheckBoxState::Unchecked;
 }
 
-void STimelineEditor::OnAutoPlayChanged(ESlateCheckBoxState::Type NewType)
+void STimelineEditor::OnAutoPlayChanged(ECheckBoxState NewType)
 {
 	if(TimelineObj)
 	{
-		TimelineObj->bAutoPlay = (NewType == ESlateCheckBoxState::Checked) ? true : false;
+		TimelineObj->bAutoPlay = (NewType == ECheckBoxState::Checked) ? true : false;
 
 		// Refresh the node that owns this timeline template to cache play status
 		TSharedPtr<FBlueprintEditor> Kismet2 = Kismet2Ptr.Pin();
@@ -1227,16 +1227,16 @@ void STimelineEditor::OnAutoPlayChanged(ESlateCheckBoxState::Type NewType)
 }
 
 
-ESlateCheckBoxState::Type STimelineEditor::IsLoopChecked() const
+ECheckBoxState STimelineEditor::IsLoopChecked() const
 {
-	return (TimelineObj && TimelineObj->bLoop) ? ESlateCheckBoxState::Checked : ESlateCheckBoxState::Unchecked;
+	return (TimelineObj && TimelineObj->bLoop) ? ECheckBoxState::Checked : ECheckBoxState::Unchecked;
 }
 
-void STimelineEditor::OnLoopChanged(ESlateCheckBoxState::Type NewType)
+void STimelineEditor::OnLoopChanged(ECheckBoxState NewType)
 {
 	if(TimelineObj)
 	{
-		TimelineObj->bLoop = (NewType == ESlateCheckBoxState::Checked) ? true : false;
+		TimelineObj->bLoop = (NewType == ECheckBoxState::Checked) ? true : false;
 
 		// Refresh the node that owns this timeline template to cache play status
 		TSharedPtr<FBlueprintEditor> Kismet2 = Kismet2Ptr.Pin();
@@ -1249,16 +1249,16 @@ void STimelineEditor::OnLoopChanged(ESlateCheckBoxState::Type NewType)
 	}
 }
 
-ESlateCheckBoxState::Type STimelineEditor::IsReplicatedChecked() const
+ECheckBoxState STimelineEditor::IsReplicatedChecked() const
 {
-	return (TimelineObj && TimelineObj->bReplicated) ? ESlateCheckBoxState::Checked : ESlateCheckBoxState::Unchecked;
+	return (TimelineObj && TimelineObj->bReplicated) ? ECheckBoxState::Checked : ECheckBoxState::Unchecked;
 }
 
-void STimelineEditor::OnReplicatedChanged(ESlateCheckBoxState::Type NewType)
+void STimelineEditor::OnReplicatedChanged(ECheckBoxState NewType)
 {
 	if(TimelineObj)
 	{
-		TimelineObj->bReplicated = (NewType == ESlateCheckBoxState::Checked) ? true : false;
+		TimelineObj->bReplicated = (NewType == ECheckBoxState::Checked) ? true : false;
 
 		// Refresh the node that owns this timeline template to cache play status
 		TSharedPtr<FBlueprintEditor> Kismet2 = Kismet2Ptr.Pin();
@@ -1271,16 +1271,16 @@ void STimelineEditor::OnReplicatedChanged(ESlateCheckBoxState::Type NewType)
 	}
 }
 
-ESlateCheckBoxState::Type STimelineEditor::IsUseLastKeyframeChecked() const
+ECheckBoxState STimelineEditor::IsUseLastKeyframeChecked() const
 {
-	return (TimelineObj && TimelineObj->LengthMode == ETimelineLengthMode::TL_LastKeyFrame) ? ESlateCheckBoxState::Checked : ESlateCheckBoxState::Unchecked;
+	return (TimelineObj && TimelineObj->LengthMode == ETimelineLengthMode::TL_LastKeyFrame) ? ECheckBoxState::Checked : ECheckBoxState::Unchecked;
 }
 
-void STimelineEditor::OnUseLastKeyframeChanged(ESlateCheckBoxState::Type NewType)
+void STimelineEditor::OnUseLastKeyframeChanged(ECheckBoxState NewType)
 {
 	if(TimelineObj)
 	{
-		TimelineObj->LengthMode = (NewType == ESlateCheckBoxState::Checked) ? ETimelineLengthMode::TL_LastKeyFrame : ETimelineLengthMode::TL_TimelineLength;
+		TimelineObj->LengthMode = (NewType == ECheckBoxState::Checked) ? ETimelineLengthMode::TL_LastKeyFrame : ETimelineLengthMode::TL_TimelineLength;
 	}
 }
 
diff --git a/Engine/Source/Editor/Kismet/Private/STimelineEditor.h b/Engine/Source/Editor/Kismet/Private/STimelineEditor.h
--- a/Engine/Source/Editor/Kismet/Private/STimelineEditor.h
+++ b/Engine/Source/Editor/Kismet/Private/STimelineEditor.h
@@ -239,24 +239,24 @@ private:
 	FText GetTimelineName() const;
 
 	/** Get state of autoplay box*/
-	ESlateCheckBoxState::Type IsAutoPlayChecked() const;
+	ECheckBoxState IsAutoPlayChecked() const;
 	/** Handle play box being changed */
-	void OnAutoPlayChanged(ESlateCheckBoxState::Type NewType);
+	void OnAutoPlayChanged(ECheckBoxState NewType);
 
 	/** Get state of loop box */
-	ESlateCheckBoxState::Type IsLoopChecked() const;
+	ECheckBoxState IsLoopChecked() const;
 	/** Handle loop box being changed */
-	void OnLoopChanged(ESlateCheckBoxState::Type NewType);
+	void OnLoopChanged(ECheckBoxState NewType);
 
 	/** Get state of replicated box */
-	ESlateCheckBoxState::Type IsReplicatedChecked() const;
+	ECheckBoxState IsReplicatedChecked() const;
 	/** Handle loop box being changed */
-	void OnReplicatedChanged(ESlateCheckBoxState::Type NewType);
+	void OnReplicatedChanged(ECheckBoxState NewType);
 
 	/** Get state of the use last keyframe checkbox */
-	ESlateCheckBoxState::Type IsUseLastKeyframeChecked() const;
+	ECheckBoxState IsUseLastKeyframeChecked() const;
 	/** Handle toggling between use last keyframe and use length */
-	void OnUseLastKeyframeChanged(ESlateCheckBoxState::Type NewType);
+	void OnUseLastKeyframeChanged(ECheckBoxState NewType);
 
 	/** Get current length of timeline as string */
 	FText GetLengthString() const;
diff --git a/Engine/Source/Editor/Kismet/Private/UserDefinedStructureEditor.cpp b/Engine/Source/Editor/Kismet/Private/UserDefinedStructureEditor.cpp
--- a/Engine/Source/Editor/Kismet/Private/UserDefinedStructureEditor.cpp
+++ b/Engine/Source/Editor/Kismet/Private/UserDefinedStructureEditor.cpp
@@ -537,25 +537,25 @@ public:
 		}
 	}
 
-	ESlateCheckBoxState::Type OnGetEditableOnBPInstanceState() const
+	ECheckBoxState OnGetEditableOnBPInstanceState() const
 	{
 		auto StructureDetailsSP = StructureDetails.Pin();
 		if (StructureDetailsSP.IsValid())
 		{
 			if (const FStructVariableDescription* FieldDesc = StructureDetailsSP->FindStructureFieldByGuid(FieldGuid))
 			{
-				return !FieldDesc->bDontEditoOnInstance ? ESlateCheckBoxState::Checked : ESlateCheckBoxState::Unchecked;
+				return !FieldDesc->bDontEditoOnInstance ? ECheckBoxState::Checked : ECheckBoxState::Unchecked;
 			}
 		}
-		return ESlateCheckBoxState::Undetermined;
+		return ECheckBoxState::Undetermined;
 	}
 
-	void OnEditableOnBPInstanceCommitted(ESlateCheckBoxState::Type InNewState)
+	void OnEditableOnBPInstanceCommitted(ECheckBoxState InNewState)
 	{
 		auto StructureDetailsSP = StructureDetails.Pin();
 		if (StructureDetailsSP.IsValid())
 		{
-			FStructureEditorUtils::ChangeEditableOnBPInstance(StructureDetailsSP->GetUserDefinedStruct(), FieldGuid, ESlateCheckBoxState::Unchecked != InNewState);
+			FStructureEditorUtils::ChangeEditableOnBPInstance(StructureDetailsSP->GetUserDefinedStruct(), FieldGuid, ECheckBoxState::Unchecked != InNewState);
 		}
 	}
 
@@ -570,22 +570,22 @@ public:
 		return EVisibility::Collapsed;
 	}
 
-	ESlateCheckBoxState::Type OnGet3dWidgetEnabled() const
+	ECheckBoxState OnGet3dWidgetEnabled() const
 	{
 		auto StructureDetailsSP = StructureDetails.Pin();
 		if (StructureDetailsSP.IsValid())
 		{
-			return FStructureEditorUtils::Is3dWidgetEnabled(StructureDetailsSP->GetUserDefinedStruct(), FieldGuid) ? ESlateCheckBoxState::Checked : ESlateCheckBoxState::Unchecked;
+			return FStructureEditorUtils::Is3dWidgetEnabled(StructureDetailsSP->GetUserDefinedStruct(), FieldGuid) ? ECheckBoxState::Checked : ECheckBoxState::Unchecked;
 		}
-		return ESlateCheckBoxState::Undetermined;
+		return ECheckBoxState::Undetermined;
 	}
 
-	void On3dWidgetEnabledCommitted(ESlateCheckBoxState::Type InNewState)
+	void On3dWidgetEnabledCommitted(ECheckBoxState InNewState)
 	{
 		auto StructureDetailsSP = StructureDetails.Pin();
-		if (StructureDetailsSP.IsValid() && (ESlateCheckBoxState::Undetermined != InNewState))
+		if (StructureDetailsSP.IsValid() && (ECheckBoxState::Undetermined != InNewState))
 		{
-			FStructureEditorUtils::Change3dWidgetEnabled(StructureDetailsSP->GetUserDefinedStruct(), FieldGuid, ESlateCheckBoxState::Checked == InNewState);
+			FStructureEditorUtils::Change3dWidgetEnabled(StructureDetailsSP->GetUserDefinedStruct(), FieldGuid, ECheckBoxState::Checked == InNewState);
 		}
 	}
 
diff --git a/Engine/Source/Editor/Kismet/Public/FindInBlueprints.h b/Engine/Source/Editor/Kismet/Public/FindInBlueprints.h
--- a/Engine/Source/Editor/Kismet/Public/FindInBlueprints.h
+++ b/Engine/Source/Editor/Kismet/Public/FindInBlueprints.h
@@ -245,10 +245,10 @@ private:
 	void OnSearchTextCommitted(const FText& Text, ETextCommit::Type CommitType);
 
 	/** Called when the find mode checkbox is hit */
-	void OnFindModeChanged(ESlateCheckBoxState::Type CheckState);
+	void OnFindModeChanged(ECheckBoxState CheckState);
 
 	/** Called to check what the find mode is for the checkbox */
-	ESlateCheckBoxState::Type OnGetFindModeChecked() const;
+	ECheckBoxState OnGetFindModeChecked() const;
 
 	/* Get the children of a row */
 	void OnGetChildren( FSearchResult InItem, TArray< FSearchResult >& OutChildren );
diff --git a/Engine/Source/Editor/Kismet/Public/SKismetInspector.h b/Engine/Source/Editor/Kismet/Public/SKismetInspector.h
--- a/Engine/Source/Editor/Kismet/Public/SKismetInspector.h
+++ b/Engine/Source/Editor/Kismet/Public/SKismetInspector.h
@@ -80,7 +80,7 @@ protected:
 	bool bShowInspectorPropertyView;
 
 	/** State of CheckBox representing whether to show only the public variables*/
-	ESlateCheckBoxState::Type	PublicViewState;
+	ECheckBoxState	PublicViewState;
 
 	/** Property viewing widget */
 	TSharedPtr<class IDetailsView> PropertyView;
@@ -124,6 +124,6 @@ protected:
 	 */
 	FString GetContextualEditingWidgetTitle() const;
 
-	ESlateCheckBoxState::Type GetPublicViewCheckboxState() const;
-	void SetPublicViewCheckboxState(ESlateCheckBoxState::Type InIsChecked);
+	ECheckBoxState GetPublicViewCheckboxState() const;
+	void SetPublicViewCheckboxState(ECheckBoxState InIsChecked);
 };
diff --git a/Engine/Source/Editor/KismetWidgets/Private/SPinTypeSelector.cpp b/Engine/Source/Editor/KismetWidgets/Private/SPinTypeSelector.cpp
--- a/Engine/Source/Editor/KismetWidgets/Private/SPinTypeSelector.cpp
+++ b/Engine/Source/Editor/KismetWidgets/Private/SPinTypeSelector.cpp
@@ -114,15 +114,15 @@ FSlateColor SPinTypeSelector::GetTypeIconColor() const
 	return Schema->GetPinTypeColor(TargetPinType.Get());
 }
 
-ESlateCheckBoxState::Type SPinTypeSelector::IsArrayChecked() const
+ECheckBoxState SPinTypeSelector::IsArrayChecked() const
 {
-	return TargetPinType.Get().bIsArray ? ESlateCheckBoxState::Checked : ESlateCheckBoxState::Unchecked;
+	return TargetPinType.Get().bIsArray ? ECheckBoxState::Checked : ECheckBoxState::Unchecked;
 }
 
-void SPinTypeSelector::OnArrayCheckStateChanged(ESlateCheckBoxState::Type NewState)
+void SPinTypeSelector::OnArrayCheckStateChanged(ECheckBoxState NewState)
 {
 	FEdGraphPinType NewTargetPinType = TargetPinType.Get();
-	NewTargetPinType.bIsArray = (NewState == ESlateCheckBoxState::Checked)? true : false;
+	NewTargetPinType.bIsArray = (NewState == ECheckBoxState::Checked)? true : false;
 
 	OnTypeChanged.ExecuteIfBound(NewTargetPinType);
 }
@@ -373,7 +373,7 @@ const FSlateBrush* SPinTypeSelector::GetIconFromPin( const FEdGraphPinType& PinT
 {
 	const FSlateBrush* IconBrush = FEditorStyle::GetBrush(TEXT("Kismet.VariableList.TypeIcon"));
 	const UObject* PinSubObject = PinType.PinSubCategoryObject.Get();
-	if( (PinType.bIsArray || (IsArrayChecked() == ESlateCheckBoxState::Checked)) && PinType.PinCategory != Schema->PC_Exec )
+	if( (PinType.bIsArray || (IsArrayChecked() == ECheckBoxState::Checked)) && PinType.PinCategory != Schema->PC_Exec )
 	{
 		IconBrush = FEditorStyle::GetBrush(TEXT("Kismet.VariableList.ArrayTypeIcon"));
 	}
diff --git a/Engine/Source/Editor/KismetWidgets/Public/SPinTypeSelector.h b/Engine/Source/Editor/KismetWidgets/Public/SPinTypeSelector.h
--- a/Engine/Source/Editor/KismetWidgets/Public/SPinTypeSelector.h
+++ b/Engine/Source/Editor/KismetWidgets/Public/SPinTypeSelector.h
@@ -83,8 +83,8 @@ protected:
 	TAttribute<FOptionalSize>	TreeViewHeight;
 
 	/** Array checkbox support functions */
-	ESlateCheckBoxState::Type IsArrayChecked() const;
-	void OnArrayCheckStateChanged(ESlateCheckBoxState::Type NewState);
+	ECheckBoxState IsArrayChecked() const;
+	void OnArrayCheckStateChanged(ECheckBoxState NewState);
 
 	/** Array containing the unfiltered list of all supported types this pin could possibly have */
 	TArray<FPinTypeTreeItem>		TypeTreeRoot;
diff --git a/Engine/Source/Editor/LandscapeEditor/Private/LandscapeEditorDetailCustomization_MiscTools.cpp b/Engine/Source/Editor/LandscapeEditor/Private/LandscapeEditorDetailCustomization_MiscTools.cpp
--- a/Engine/Source/Editor/LandscapeEditor/Private/LandscapeEditorDetailCustomization_MiscTools.cpp
+++ b/Engine/Source/Editor/LandscapeEditor/Private/LandscapeEditorDetailCustomization_MiscTools.cpp
@@ -239,23 +239,23 @@ FReply FLandscapeEditorDetailCustomization_MiscTools::OnApplySelectedSplinesButt
 	return FReply::Handled();
 }
 
-void FLandscapeEditorDetailCustomization_MiscTools::OnbUseAutoRotateControlPointChanged(ESlateCheckBoxState::Type NewState)
+void FLandscapeEditorDetailCustomization_MiscTools::OnbUseAutoRotateControlPointChanged(ECheckBoxState NewState)
 {
 	FEdModeLandscape* LandscapeEdMode = GetEditorMode();
 	if (LandscapeEdMode)
 	{
-		LandscapeEdMode->SetbUseAutoRotateOnJoin(NewState == ESlateCheckBoxState::Checked);
+		LandscapeEdMode->SetbUseAutoRotateOnJoin(NewState == ECheckBoxState::Checked);
 	}
 }
 
-ESlateCheckBoxState::Type FLandscapeEditorDetailCustomization_MiscTools::GetbUseAutoRotateControlPoint() const
+ECheckBoxState FLandscapeEditorDetailCustomization_MiscTools::GetbUseAutoRotateControlPoint() const
 {
 	FEdModeLandscape* LandscapeEdMode = GetEditorMode();
 	if (LandscapeEdMode)
 	{
-		return LandscapeEdMode->GetbUseAutoRotateOnJoin() ? ESlateCheckBoxState::Checked : ESlateCheckBoxState::Unchecked;
+		return LandscapeEdMode->GetbUseAutoRotateOnJoin() ? ECheckBoxState::Checked : ECheckBoxState::Unchecked;
 	}
-	return ESlateCheckBoxState::Unchecked;
+	return ECheckBoxState::Unchecked;
 }
 
 FReply FLandscapeEditorDetailCustomization_MiscTools::OnApplyRampButtonClicked()
diff --git a/Engine/Source/Editor/LandscapeEditor/Private/LandscapeEditorDetailCustomization_MiscTools.h b/Engine/Source/Editor/LandscapeEditor/Private/LandscapeEditorDetailCustomization_MiscTools.h
--- a/Engine/Source/Editor/LandscapeEditor/Private/LandscapeEditorDetailCustomization_MiscTools.h
+++ b/Engine/Source/Editor/LandscapeEditor/Private/LandscapeEditorDetailCustomization_MiscTools.h
@@ -30,8 +30,8 @@ protected:
 	// Splines Tool
 	static FReply OnApplyAllSplinesButtonClicked();
 	static FReply OnApplySelectedSplinesButtonClicked();
-	void OnbUseAutoRotateControlPointChanged(ESlateCheckBoxState::Type NewState);
-	ESlateCheckBoxState::Type GetbUseAutoRotateControlPoint() const;
+	void OnbUseAutoRotateControlPointChanged(ECheckBoxState NewState);
+	ECheckBoxState GetbUseAutoRotateControlPoint() const;
 
 	// Ramp Tool
 	static FReply OnApplyRampButtonClicked();
diff --git a/Engine/Source/Editor/LandscapeEditor/Private/LandscapeEditorDetailCustomization_NewLandscape.cpp b/Engine/Source/Editor/LandscapeEditor/Private/LandscapeEditorDetailCustomization_NewLandscape.cpp
--- a/Engine/Source/Editor/LandscapeEditor/Private/LandscapeEditorDetailCustomization_NewLandscape.cpp
+++ b/Engine/Source/Editor/LandscapeEditor/Private/LandscapeEditorDetailCustomization_NewLandscape.cpp
@@ -726,22 +726,22 @@ EVisibility FLandscapeEditorDetailCustomization_NewLandscape::GetVisibilityOnlyI
 	return EVisibility::Collapsed;
 }
 
-ESlateCheckBoxState::Type FLandscapeEditorDetailCustomization_NewLandscape::NewLandscapeModeIsChecked(ENewLandscapePreviewMode::Type value) const
+ECheckBoxState FLandscapeEditorDetailCustomization_NewLandscape::NewLandscapeModeIsChecked(ENewLandscapePreviewMode::Type value) const
 {
 	FEdModeLandscape* LandscapeEdMode = GetEditorMode();
 	if (LandscapeEdMode != NULL)
 	{
 		if (LandscapeEdMode->NewLandscapePreviewMode == value)
 		{
-			return ESlateCheckBoxState::Checked;
+			return ECheckBoxState::Checked;
 		}
 	}
-	return ESlateCheckBoxState::Unchecked;
+	return ECheckBoxState::Unchecked;
 }
 
-void FLandscapeEditorDetailCustomization_NewLandscape::OnNewLandscapeModeChanged(ESlateCheckBoxState::Type NewCheckedState, ENewLandscapePreviewMode::Type value)
+void FLandscapeEditorDetailCustomization_NewLandscape::OnNewLandscapeModeChanged(ECheckBoxState NewCheckedState, ENewLandscapePreviewMode::Type value)
 {
-	if (NewCheckedState == ESlateCheckBoxState::Checked)
+	if (NewCheckedState == ECheckBoxState::Checked)
 	{
 		FEdModeLandscape* LandscapeEdMode = GetEditorMode();
 		if (LandscapeEdMode != NULL)
diff --git a/Engine/Source/Editor/LandscapeEditor/Private/LandscapeEditorDetailCustomization_NewLandscape.h b/Engine/Source/Editor/LandscapeEditor/Private/LandscapeEditorDetailCustomization_NewLandscape.h
--- a/Engine/Source/Editor/LandscapeEditor/Private/LandscapeEditorDetailCustomization_NewLandscape.h
+++ b/Engine/Source/Editor/LandscapeEditor/Private/LandscapeEditorDetailCustomization_NewLandscape.h
@@ -52,8 +52,8 @@ public:
 	FReply OnFillWorldButtonClicked();
 
 	static EVisibility GetVisibilityOnlyInNewLandscapeMode(ENewLandscapePreviewMode::Type value);
-	ESlateCheckBoxState::Type NewLandscapeModeIsChecked(ENewLandscapePreviewMode::Type value) const;
-	void OnNewLandscapeModeChanged(ESlateCheckBoxState::Type NewCheckedState, ENewLandscapePreviewMode::Type value);
+	ECheckBoxState NewLandscapeModeIsChecked(ENewLandscapePreviewMode::Type value) const;
+	void OnNewLandscapeModeChanged(ECheckBoxState NewCheckedState, ENewLandscapePreviewMode::Type value);
 
 	// Import
 	static EVisibility GetHeightmapErrorVisibility(TSharedRef<IPropertyHandle> PropertyHandle_HeightmapError);
diff --git a/Engine/Source/Editor/LandscapeEditor/Private/LandscapeEditorDetailCustomization_TargetLayers.cpp b/Engine/Source/Editor/LandscapeEditor/Private/LandscapeEditorDetailCustomization_TargetLayers.cpp
--- a/Engine/Source/Editor/LandscapeEditor/Private/LandscapeEditorDetailCustomization_TargetLayers.cpp
+++ b/Engine/Source/Editor/LandscapeEditor/Private/LandscapeEditorDetailCustomization_TargetLayers.cpp
@@ -832,19 +832,19 @@ EVisibility FLandscapeEditorCustomNodeBuilder_TargetLayers::GetDebugModeColorCha
 	return EVisibility::Collapsed;
 }
 
-ESlateCheckBoxState::Type FLandscapeEditorCustomNodeBuilder_TargetLayers::DebugModeColorChannelIsChecked(const TSharedRef<FLandscapeTargetListInfo> Target, int32 Channel)
+ECheckBoxState FLandscapeEditorCustomNodeBuilder_TargetLayers::DebugModeColorChannelIsChecked(const TSharedRef<FLandscapeTargetListInfo> Target, int32 Channel)
 {
 	if (Target->DebugColorChannel == Channel)
 	{
-		return ESlateCheckBoxState::Checked;
+		return ECheckBoxState::Checked;
 	}
 
-	return ESlateCheckBoxState::Unchecked;
+	return ECheckBoxState::Unchecked;
 }
 
-void FLandscapeEditorCustomNodeBuilder_TargetLayers::OnDebugModeColorChannelChanged(ESlateCheckBoxState::Type NewCheckedState, const TSharedRef<FLandscapeTargetListInfo> Target, int32 Channel)
+void FLandscapeEditorCustomNodeBuilder_TargetLayers::OnDebugModeColorChannelChanged(ECheckBoxState NewCheckedState, const TSharedRef<FLandscapeTargetListInfo> Target, int32 Channel)
 {
-	if (NewCheckedState == ESlateCheckBoxState::Checked)
+	if (NewCheckedState == ECheckBoxState::Checked)
 	{
 		// Enable on us and disable colour channel on other targets
 		if (ensure(Target->LayerInfoObj != NULL))
diff --git a/Engine/Source/Editor/LandscapeEditor/Private/LandscapeEditorDetailCustomization_TargetLayers.h b/Engine/Source/Editor/LandscapeEditor/Private/LandscapeEditorDetailCustomization_TargetLayers.h
--- a/Engine/Source/Editor/LandscapeEditor/Private/LandscapeEditorDetailCustomization_TargetLayers.h
+++ b/Engine/Source/Editor/LandscapeEditor/Private/LandscapeEditorDetailCustomization_TargetLayers.h
@@ -62,8 +62,8 @@ protected:
 	static FReply OnTargetLayerMakePublicClicked(const TSharedRef<FLandscapeTargetListInfo> Target);
 	static FReply OnTargetLayerDeleteClicked(const TSharedRef<FLandscapeTargetListInfo> Target);
 	static EVisibility GetDebugModeColorChannelVisibility(const TSharedRef<FLandscapeTargetListInfo> Target);
-	static ESlateCheckBoxState::Type DebugModeColorChannelIsChecked(const TSharedRef<FLandscapeTargetListInfo> Target, int32 Channel);
-	static void OnDebugModeColorChannelChanged(ESlateCheckBoxState::Type NewCheckedState, const TSharedRef<FLandscapeTargetListInfo> Target, int32 Channel);
+	static ECheckBoxState DebugModeColorChannelIsChecked(const TSharedRef<FLandscapeTargetListInfo> Target, int32 Channel);
+	static void OnDebugModeColorChannelChanged(ECheckBoxState NewCheckedState, const TSharedRef<FLandscapeTargetListInfo> Target, int32 Channel);
 };
 
 class SLandscapeEditorSelectableBorder : public SBorder
diff --git a/Engine/Source/Editor/LevelEditor/Private/HighresScreenshotUI.h b/Engine/Source/Editor/LevelEditor/Private/HighresScreenshotUI.h
--- a/Engine/Source/Editor/LevelEditor/Private/HighresScreenshotUI.h
+++ b/Engine/Source/Editor/LevelEditor/Private/HighresScreenshotUI.h
@@ -76,9 +76,9 @@ private:
 		Config.ResolutionMultiplier = FMath::RoundToFloat(FMath::Lerp(FHighResScreenshotConfig::MinResolutionMultipler, FHighResScreenshotConfig::MaxResolutionMultipler, NewValue));
 	}
 
-	void OnMaskEnabledChanged( ESlateCheckBoxState::Type NewValue )
+	void OnMaskEnabledChanged( ECheckBoxState NewValue )
 	{
-		Config.bMaskEnabled = (NewValue == ESlateCheckBoxState::Checked);
+		Config.bMaskEnabled = (NewValue == ECheckBoxState::Checked);
 		auto ConfigViewport = Config.TargetViewport.Pin();
 		if (ConfigViewport.IsValid())
 		{
@@ -87,9 +87,9 @@ private:
 		}
 	}
 
-	void OnHDREnabledChanged(ESlateCheckBoxState::Type NewValue)
+	void OnHDREnabledChanged(ECheckBoxState NewValue)
 	{
-		Config.SetHDRCapture(NewValue == ESlateCheckBoxState::Checked);
+		Config.SetHDRCapture(NewValue == ECheckBoxState::Checked);
 		auto ConfigViewport = Config.TargetViewport.Pin();
 		if (ConfigViewport.IsValid())
 		{
@@ -97,9 +97,9 @@ private:
 		}
 	}
 
-	void OnBufferVisualizationDumpEnabledChanged(ESlateCheckBoxState::Type NewValue)
+	void OnBufferVisualizationDumpEnabledChanged(ECheckBoxState NewValue)
 	{
-		bool bEnabled = (NewValue == ESlateCheckBoxState::Checked);
+		bool bEnabled = (NewValue == ECheckBoxState::Checked);
 		Config.bDumpBufferVisualizationTargets = bEnabled;
 		SetHDRUIEnableState(bEnabled);
 	}
@@ -129,19 +129,19 @@ private:
 		return Config.ResolutionMultiplierScale;
 	}
 
-	ESlateCheckBoxState::Type GetMaskEnabled() const
+	ECheckBoxState GetMaskEnabled() const
 	{
-		return Config.bMaskEnabled ? ESlateCheckBoxState::Checked : ESlateCheckBoxState::Unchecked;
+		return Config.bMaskEnabled ? ECheckBoxState::Checked : ECheckBoxState::Unchecked;
 	}
 
-	ESlateCheckBoxState::Type GetHDRCheckboxUIState() const
+	ECheckBoxState GetHDRCheckboxUIState() const
 	{
-		return Config.bCaptureHDR ? ESlateCheckBoxState::Checked : ESlateCheckBoxState::Unchecked;
+		return Config.bCaptureHDR ? ECheckBoxState::Checked : ECheckBoxState::Unchecked;
 	}
 
-	ESlateCheckBoxState::Type GetBufferVisualizationDumpEnabled() const
+	ECheckBoxState GetBufferVisualizationDumpEnabled() const
 	{
-		return Config.bDumpBufferVisualizationTargets ? ESlateCheckBoxState::Checked : ESlateCheckBoxState::Unchecked;
+		return Config.bDumpBufferVisualizationTargets ? ECheckBoxState::Checked : ECheckBoxState::Unchecked;
 	}
 
 	bool IsCaptureRegionEditingAvailable() const
diff --git a/Engine/Source/Editor/LevelEditor/Private/LevelEditorActions.cpp b/Engine/Source/Editor/LevelEditor/Private/LevelEditorActions.cpp
--- a/Engine/Source/Editor/LevelEditor/Private/LevelEditorActions.cpp
+++ b/Engine/Source/Editor/LevelEditor/Private/LevelEditorActions.cpp
@@ -752,26 +752,26 @@ bool FLevelEditorActionCallbacks::IsLightingDensityRenderGrayscaleChecked()
 	return GEngine->bRenderLightMapDensityGrayscale;
 }
 
-void FLevelEditorActionCallbacks::SetLightingResolutionStaticMeshes( ESlateCheckBoxState::Type NewCheckedState )
+void FLevelEditorActionCallbacks::SetLightingResolutionStaticMeshes( ECheckBoxState NewCheckedState )
 {
 	FLightmapResRatioAdjustSettings& Settings = FLightmapResRatioAdjustSettings::Get();
-	Settings.bStaticMeshes = ( NewCheckedState == ESlateCheckBoxState::Checked );
+	Settings.bStaticMeshes = ( NewCheckedState == ECheckBoxState::Checked );
 }
 
-ESlateCheckBoxState::Type FLevelEditorActionCallbacks::IsLightingResolutionStaticMeshesChecked()
+ECheckBoxState FLevelEditorActionCallbacks::IsLightingResolutionStaticMeshesChecked()
 {
-	return ( FLightmapResRatioAdjustSettings::Get().bStaticMeshes ? ESlateCheckBoxState::Checked : ESlateCheckBoxState::Unchecked );
+	return ( FLightmapResRatioAdjustSettings::Get().bStaticMeshes ? ECheckBoxState::Checked : ECheckBoxState::Unchecked );
 }
 
-void FLevelEditorActionCallbacks::SetLightingResolutionBSPSurfaces( ESlateCheckBoxState::Type NewCheckedState )
+void FLevelEditorActionCallbacks::SetLightingResolutionBSPSurfaces( ECheckBoxState NewCheckedState )
 {
 	FLightmapResRatioAdjustSettings& Settings = FLightmapResRatioAdjustSettings::Get();
-	Settings.bBSPSurfaces = ( NewCheckedState == ESlateCheckBoxState::Checked );
+	Settings.bBSPSurfaces = ( NewCheckedState == ECheckBoxState::Checked );
 }
 
-ESlateCheckBoxState::Type FLevelEditorActionCallbacks::IsLightingResolutionBSPSurfacesChecked()
+ECheckBoxState FLevelEditorActionCallbacks::IsLightingResolutionBSPSurfacesChecked()
 {
-	return ( FLightmapResRatioAdjustSettings::Get().bBSPSurfaces ? ESlateCheckBoxState::Checked : ESlateCheckBoxState::Unchecked );
+	return ( FLightmapResRatioAdjustSettings::Get().bBSPSurfaces ? ECheckBoxState::Checked : ECheckBoxState::Unchecked );
 }
 
 void FLevelEditorActionCallbacks::SetLightingResolutionLevel( FLightmapResRatioAdjustSettings::AdjustLevels NewLevel )
diff --git a/Engine/Source/Editor/LevelEditor/Private/SLevelEditorBuildAndSubmit.cpp b/Engine/Source/Editor/LevelEditor/Private/SLevelEditorBuildAndSubmit.cpp
--- a/Engine/Source/Editor/LevelEditor/Private/SLevelEditorBuildAndSubmit.cpp
+++ b/Engine/Source/Editor/LevelEditor/Private/SLevelEditorBuildAndSubmit.cpp
@@ -51,7 +51,7 @@ public:
 				[
 					SNew(SCheckBox)
 					.OnCheckStateChanged(this, &SPackageItem::OnCheckStateChanged)
-					.IsChecked(Item->Selected ? ESlateCheckBoxState::Checked : ESlateCheckBoxState::Unchecked)
+					.IsChecked(Item->Selected ? ECheckBoxState::Checked : ECheckBoxState::Unchecked)
 					[
 						SNew(STextBlock)
 						.Font(FEditorStyle::GetFontStyle("BuildAndSubmit.NormalFont"))
@@ -113,9 +113,9 @@ public:
 
 private:
 	/** User clicked the check box on this item - set the item struct's selected flag to match */
-	void OnCheckStateChanged( const ESlateCheckBoxState::Type NewCheckedState )
+	void OnCheckStateChanged( const ECheckBoxState NewCheckedState )
 	{
-		Item->Selected = (NewCheckedState == ESlateCheckBoxState::Checked);
+		Item->Selected = (NewCheckedState == ECheckBoxState::Checked);
 	}
 
 	/** the item that this row represents */
@@ -242,7 +242,7 @@ void SLevelEditorBuildAndSubmit::Construct( const FArguments& InArgs, const TSha
 				.Padding(2)
 				[
 					SAssignNew(AddFilesToSCBox, SCheckBox)
-					.IsChecked(ESlateCheckBoxState::Checked)
+					.IsChecked(ECheckBoxState::Checked)
 					[
 						SNew(STextBlock)
 						.Text(LOCTEXT("AddFilesButtonLabel", "Add Files to Source Control if Necessary"))
@@ -432,7 +432,7 @@ void SLevelEditorBuildAndSubmit::OnShowHideExtraPackagesSection(bool bIsExpanded
 	}
 }
 
-void SLevelEditorBuildAndSubmit::OnShowPackagesNotInSCBoxChanged(ESlateCheckBoxState::Type InNewState)
+void SLevelEditorBuildAndSubmit::OnShowPackagesNotInSCBoxChanged(ECheckBoxState InNewState)
 {
 	if (bIsExtraPackagesSectionExpanded)
 	{
diff --git a/Engine/Source/Editor/LevelEditor/Private/SLevelEditorBuildAndSubmit.h b/Engine/Source/Editor/LevelEditor/Private/SLevelEditorBuildAndSubmit.h
--- a/Engine/Source/Editor/LevelEditor/Private/SLevelEditorBuildAndSubmit.h
+++ b/Engine/Source/Editor/LevelEditor/Private/SLevelEditorBuildAndSubmit.h
@@ -66,7 +66,7 @@ private:
 	 *
 	 * @param InNewState	The state of the check box
 	 */
-	void OnShowPackagesNotInSCBoxChanged(ESlateCheckBoxState::Type InNewState);
+	void OnShowPackagesNotInSCBoxChanged(ECheckBoxState InNewState);
 
 	/** Level editor that we're associated with */
 	TWeakPtr<class ILevelEditor> LevelEditor;
