Blueprints are marked dirty when they are re-parented.

diff --git a/Engine/Source/Editor/Kismet/Private/BlueprintEditor.cpp b/Engine/Source/Editor/Kismet/Private/BlueprintEditor.cpp
--- a/Engine/Source/Editor/Kismet/Private/BlueprintEditor.cpp
+++ b/Engine/Source/Editor/Kismet/Private/BlueprintEditor.cpp
@@ -2264,6 +2264,7 @@ void FBlueprintEditor::ReparentBlueprint_NewParentChosen(UClass* ChosenClass)
 
 			// Ensure that the Blueprint is up-to-date (valid SCS etc.) before compiling
 			EnsureBlueprintIsUpToDate(BlueprintObj);
+			FBlueprintEditorUtils::MarkBlueprintAsModified(BlueprintObj);
 
 			Compile();
 
