Fix crash when adding local variables to overriden functions.

diff --git a/Engine/Source/Editor/Kismet/Private/BlueprintDetailsCustomization.cpp b/Engine/Source/Editor/Kismet/Private/BlueprintDetailsCustomization.cpp
--- a/Engine/Source/Editor/Kismet/Private/BlueprintDetailsCustomization.cpp
+++ b/Engine/Source/Editor/Kismet/Private/BlueprintDetailsCustomization.cpp
@@ -535,14 +535,15 @@ void FBlueprintVarActionDetails::CustomizeDetails( IDetailLayoutBuilder& DetailL
 
 				TSharedPtr<FStructOnScope> StructData = MakeShareable(new FStructOnScope((UFunction*)StructScope));
 				UEdGraph* Graph = FBlueprintEditorUtils::FindScopeGraph(GetBlueprintObj(), (UFunction*)StructScope);
-				TWeakObjectPtr<UK2Node_EditablePinBase> EntryNode;
-				TWeakObjectPtr<UK2Node_EditablePinBase> ResultNode;
-				FBlueprintEditorUtils::GetEntryAndResultNodes(Graph, EntryNode, ResultNode);
 
-				UK2Node_FunctionEntry* FuncEntry = Cast<UK2Node_FunctionEntry>(EntryNode.Get());
+				// Find the function entry nodes in the current graph
+				TArray<UK2Node_FunctionEntry*> EntryNodes;
+				Graph->GetNodesOfClass(EntryNodes);
 
-				UClass* ActorClass = FindObject<UClass>(ANY_PACKAGE, TEXT("Actor"));
-				UClass* BPActorClass = FindObject<UClass>(ANY_PACKAGE, TEXT("BP_Actor_C"));
+				// There should always be an entry node in the function graph
+				check(EntryNodes.Num() > 0);
+
+				UK2Node_FunctionEntry* FuncEntry = EntryNodes[0];
 				for(auto& LocalVar : FuncEntry->LocalVariables)
 				{
 					if(LocalVar.VarName == VariableProperty->GetFName()) //Property->GetFName())
@@ -563,6 +564,7 @@ void FBlueprintVarActionDetails::CustomizeDetails( IDetailLayoutBuilder& DetailL
 
 					if(DetailsView.IsValid())
 					{
+						TWeakObjectPtr<UK2Node_EditablePinBase> EntryNode = FuncEntry;
 						DetailsView->OnFinishedChangingProperties().AddSP(this, &FBlueprintVarActionDetails::OnFinishedChangingProperties, StructData, EntryNode);
 					}
 				}
