Diff and merge tool improvements, introduced tree view to display changes, removed the list view that used to display changes made in the blueprint's various graphs

diff --git a/Engine/Source/Developer/Merge/Private/BlueprintMergeData.h b/Engine/Source/Developer/Merge/Private/BlueprintMergeData.h
--- a/Engine/Source/Developer/Merge/Private/BlueprintMergeData.h
+++ b/Engine/Source/Developer/Merge/Private/BlueprintMergeData.h
@@ -28,10 +28,11 @@ struct FBlueprintMergeData
 	)
 		: OwningEditor(		InOwningEditor		)
 		, BlueprintLocal(	InBlueprintLocal	)
+		, RevisionLocal(	FRevisionInfo::InvalidRevision())
 		, BlueprintBase(	InBlueprintBase		)
 		, RevisionBase(		InRevisionBase		)
 		, BlueprintRemote(	InBlueprintRemote	)
-		, RevisionRemote(		InRevisionRemote)
+		, RevisionRemote(	InRevisionRemote	)
 	{
 	}
 
@@ -59,11 +60,5 @@ namespace EMergeParticipant
 	};
 }
 
-class IMergeControl
-{
-public:
-	virtual void HighlightNextConflict() = 0;
-	virtual void HighlightPrevConflict() = 0;
-	virtual bool HasNextConflict() const = 0;
-	virtual bool HasPrevConflict() const = 0;
-};
+DECLARE_DELEGATE(FOnMergeNodeSelected);
+
diff --git a/Engine/Source/Developer/Merge/Private/Merge.cpp b/Engine/Source/Developer/Merge/Private/Merge.cpp
--- a/Engine/Source/Developer/Merge/Private/Merge.cpp
+++ b/Engine/Source/Developer/Merge/Private/Merge.cpp
@@ -69,9 +69,9 @@ static TSharedPtr<SWidget> GenerateMergeTabContents(TSharedRef<FBlueprintEditor>
 	FBlueprintMergeData Data(Editor
 		, LocalBlueprint
 		, BaseBlueprint
-		, RemoteRevInfo
+		, BaseRevInfo
 		, RemoteBlueprint
-		, BaseRevInfo);
+		, RemoteRevInfo);
 
 	return SNew(SBlueprintMerge, Data)
 		.bForcePickAssets(bForceAssetPicker)
diff --git a/Engine/Source/Developer/Merge/Private/SBlueprintMerge.cpp b/Engine/Source/Developer/Merge/Private/SBlueprintMerge.cpp
--- a/Engine/Source/Developer/Merge/Private/SBlueprintMerge.cpp
+++ b/Engine/Source/Developer/Merge/Private/SBlueprintMerge.cpp
@@ -4,7 +4,6 @@
 
 #include "BlueprintEditor.h"
 #include "BlueprintEditorModes.h"
-#include "IDiffControl.h"
 #include "ISourceControlModule.h"
 #include "SBlueprintEditorToolbar.h"
 #include "SBlueprintMerge.h"
@@ -20,8 +19,6 @@
 
 SBlueprintMerge::SBlueprintMerge()
 	: Data()
-	, CurrentDiffControl( nullptr )
-	, CurrentMergeControl( nullptr )
 {
 }
 
@@ -169,7 +166,36 @@ void SBlueprintMerge::Construct(const FArguments InArgs, const FBlueprintMergeDa
 		, FSlateIcon(FEditorStyle::GetStyleSetName(), "BlueprintMerge.Cancel")
 	);
 
-	ChildSlot [
+	DifferencesTreeView = DiffTreeView::CreateTreeView(&MasterDifferencesList);
+
+	TSharedRef<SWidget> Overlay = SNew(SHorizontalBox);
+	if( IsActivelyMerging() )
+	{
+		const auto TextBlock = []( FText Text ) -> TSharedRef<SWidget>
+		{
+			return SNew(STextBlock)
+				.Visibility(EVisibility::HitTestInvisible)
+				.TextStyle(FEditorStyle::Get(), "GraphPreview.CornerText")
+				.Text(Text);
+		};
+
+		Overlay = SNew(SHorizontalBox)
+		+ SHorizontalBox::Slot()
+		[
+			TextBlock(DiffViewUtils::GetPanelLabel(InData.BlueprintRemote, InData.RevisionRemote, LOCTEXT("RemoteLabel", "SOURCE (REMOTE)")))
+		]
+		+ SHorizontalBox::Slot()
+		[
+			TextBlock(DiffViewUtils::GetPanelLabel(InData.BlueprintBase, InData.RevisionBase, LOCTEXT("BaseLabel", "BASE")))
+		]
+		+ SHorizontalBox::Slot()
+		[
+			TextBlock(DiffViewUtils::GetPanelLabel(InData.BlueprintLocal, InData.RevisionLocal, LOCTEXT("LocalLabel", "TARGET (LOCAL)")))
+		];
+	}
+
+	ChildSlot 
+	[
 		SNew( SVerticalBox )
 		+ SVerticalBox::Slot()
 		.AutoHeight()
@@ -177,14 +203,37 @@ void SBlueprintMerge::Construct(const FArguments InArgs, const FBlueprintMergeDa
 			SNew( SHorizontalBox )
 			+ SHorizontalBox::Slot()
 			.AutoWidth()
+			.Padding(2.f)
 			[
 				ToolbarBuilder.MakeWidget()
 			]
 		]
 		+ SVerticalBox::Slot()
 		[
-			SAssignNew(MainView, SBorder)
-				.BorderImage(FEditorStyle::GetBrush("NoBorder"))
+			SNew(SSplitter)
+			+ SSplitter::Slot()
+			.Value(.2f)
+			[
+				SNew(SBorder)
+				.BorderImage(FEditorStyle::GetBrush("ToolPanel.GroupBorder"))
+				[
+					DifferencesTreeView.ToSharedRef()
+				]
+			]
+			+ SSplitter::Slot()
+			.Value(.8f)
+			[
+				SNew(SOverlay)
+				+SOverlay::Slot()
+				[
+					SAssignNew(MainView, SBox)
+				]
+				+ SOverlay::Slot()
+				.VAlign(VAlign_Bottom)
+				[
+					Overlay
+				]
+			]
 		]
 	];
 
@@ -197,9 +246,7 @@ void SBlueprintMerge::Construct(const FArguments InArgs, const FBlueprintMergeDa
 		bIsPickingAssets = true;
 
 		auto AssetPickerView = SNew(SMergeAssetPickerView, InData).OnAssetChanged(this, &SBlueprintMerge::OnMergeAssetSelected);
-		AssetPickerControl.Widget = AssetPickerView;
-		AssetPickerControl.DiffControl  = &AssetPickerView.Get();
-		AssetPickerControl.MergeControl = &AssetPickerView.Get();
+		AssetPickerControl = AssetPickerView;
 	}
 
 	Data.OwningEditor.Pin()->OnModeSet().AddSP( this, &SBlueprintMerge::OnModeChanged );
@@ -213,42 +260,42 @@ UBlueprint* SBlueprintMerge::GetTargetBlueprint()
 
 void SBlueprintMerge::NextDiff()
 {
-	CurrentDiffControl->NextDiff();
+	DiffTreeView::HighlightNextDifference(DifferencesTreeView.ToSharedRef(), RealDifferences, MasterDifferencesList);
 }
 
 void SBlueprintMerge::PrevDiff()
 {
-	CurrentDiffControl->PrevDiff();
+	 DiffTreeView::HighlightPrevDifference( DifferencesTreeView.ToSharedRef(), RealDifferences, MasterDifferencesList );
 }
 
 bool SBlueprintMerge::HasNextDiff() const
 {
-	return CurrentDiffControl && CurrentDiffControl->HasNextDifference();
+	return DiffTreeView::HasNextDifference( DifferencesTreeView.ToSharedRef(), RealDifferences );
 }
 
 bool SBlueprintMerge::HasPrevDiff() const
 {
-	return CurrentDiffControl && CurrentDiffControl->HasPrevDifference();
+	return DiffTreeView::HasPrevDifference(DifferencesTreeView.ToSharedRef(), RealDifferences );
 }
 
 void SBlueprintMerge::NextConflict()
 {
-	CurrentMergeControl->HighlightNextConflict();
+	DiffTreeView::HighlightNextDifference(DifferencesTreeView.ToSharedRef(), MergeConflicts, MasterDifferencesList);
 }
 
 void SBlueprintMerge::PrevConflict()
 {
-	CurrentMergeControl->HighlightPrevConflict();
+	DiffTreeView::HighlightPrevDifference(DifferencesTreeView.ToSharedRef(), MergeConflicts, MasterDifferencesList);
 }
 
 bool SBlueprintMerge::HasNextConflict() const
 {
-	return CurrentMergeControl && CurrentMergeControl->HasNextConflict();
+	return DiffTreeView::HasNextDifference(DifferencesTreeView.ToSharedRef(), MergeConflicts);
 }
 
 bool SBlueprintMerge::HasPrevConflict() const
 {
-	return CurrentMergeControl && CurrentMergeControl->HasPrevConflict();
+	return DiffTreeView::HasPrevDifference(DifferencesTreeView.ToSharedRef(), MergeConflicts);
 }
 
 void SBlueprintMerge::OnStartMerge()
@@ -280,9 +327,6 @@ void SBlueprintMerge::OnStartMerge()
 		{
 			UBlueprint* LocalBlueprint = GetTargetBlueprint();
 			// we use this to suppress blueprint compilation during StaticDuplicateObject()
-			// @TODO: we don't remember exactly why this was needed, but I saw the 
-			//        the blueprint become unplacable and the file become unloadable
-			//        without this, but don't know exactly why
 			TGuardValue<bool> GuardDuplicatingReadOnly(LocalBlueprint->bDuplicatingReadOnly, true);
 
 			UPackage* TransientPackage = GetTransientPackage();
@@ -301,21 +345,42 @@ void SBlueprintMerge::OnStartMerge()
 			LocalBackupPath.Empty();
 		}
 
-		auto GraphView = SNew(SMergeGraphView, Data);
-		GraphControl.Widget = GraphView;
-		GraphControl.DiffControl = &GraphView.Get();
-		GraphControl.MergeControl = &GraphView.Get();
-
-		auto TreeView = SNew(SMergeTreeView, Data);
-		TreeControl.Widget = TreeView;
-		TreeControl.DiffControl = &TreeView.Get();
-		TreeControl.MergeControl = &TreeView.Get();
+		const auto DetailsSelected = [](TWeakPtr<FBlueprintEditor> Editor)
+		{
+			Editor.Pin()->SetCurrentMode(FBlueprintEditorApplicationModes::BlueprintDefaultsMode);
+		};
+		auto DetailsView = SNew(SMergeDetailsView
+			, Data
+			, FOnMergeNodeSelected::CreateStatic(DetailsSelected, Data.OwningEditor)
+			, MasterDifferencesList
+			, RealDifferences
+			, MergeConflicts);
+		DetailsControl = DetailsView;
+
+		const auto ComponentSelected = [](TWeakPtr<FBlueprintEditor> Editor)
+		{
+			Editor.Pin()->SetCurrentMode(FBlueprintEditorApplicationModes::BlueprintComponentsMode);
+		};
+		auto TreeView = SNew(SMergeTreeView
+			, Data
+			, FOnMergeNodeSelected::CreateStatic(ComponentSelected, Data.OwningEditor)
+			, MasterDifferencesList
+			, RealDifferences
+			, MergeConflicts);
+		TreeControl = TreeView;
+
+		const auto GraphSelected = [](TWeakPtr<FBlueprintEditor> Editor)
+		{
+			Editor.Pin()->SetCurrentMode(FBlueprintEditorApplicationModes::StandardBlueprintEditorMode);
+		};
+		auto GraphView = SNew(SMergeGraphView
+							, Data
+							, FOnMergeNodeSelected::CreateStatic(GraphSelected, Data.OwningEditor)
+							, MasterDifferencesList
+							, RealDifferences
+							, MergeConflicts);
+		GraphControl = GraphView;
 
-		auto DetailsView = SNew(SMergeDetailsView, Data);
-		DetailsControl.Widget = DetailsView;
-		DetailsControl.DiffControl = &DetailsView.Get();
-		DetailsControl.MergeControl = &DetailsView.Get();
-	
 		bIsPickingAssets = false;
 		OnModeChanged(Data.OwningEditor.Pin()->GetCurrentMode());
 	}
@@ -350,29 +415,21 @@ void SBlueprintMerge::OnModeChanged(FName NewMode)
 {
 	if (!IsActivelyMerging())
 	{
-		MainView->SetContent(AssetPickerControl.Widget.ToSharedRef());
-		CurrentDiffControl = AssetPickerControl.DiffControl;
-		CurrentMergeControl = AssetPickerControl.MergeControl;
+		MainView->SetContent(AssetPickerControl.ToSharedRef());
 	}
 	else if (NewMode == FBlueprintEditorApplicationModes::StandardBlueprintEditorMode ||
 		NewMode == FBlueprintEditorApplicationModes::BlueprintMacroMode)
 	{
-		MainView->SetContent( GraphControl.Widget.ToSharedRef() );
-		CurrentDiffControl = GraphControl.DiffControl;
-		CurrentMergeControl = GraphControl.MergeControl;
+		MainView->SetContent( GraphControl.ToSharedRef() );
 	}
 	else if (NewMode == FBlueprintEditorApplicationModes::BlueprintComponentsMode)
 	{
-		MainView->SetContent(TreeControl.Widget.ToSharedRef());
-		CurrentDiffControl = TreeControl.DiffControl;
-		CurrentMergeControl = TreeControl.MergeControl;
+		MainView->SetContent(TreeControl.ToSharedRef());
 	}
 	else if (NewMode == FBlueprintEditorApplicationModes::BlueprintDefaultsMode ||
 			NewMode == FBlueprintEditorApplicationModes::BlueprintInterfaceMode)
 	{
-		MainView->SetContent(DetailsControl.Widget.ToSharedRef());
-		CurrentDiffControl = DetailsControl.DiffControl;
-		CurrentMergeControl = DetailsControl.MergeControl;
+		MainView->SetContent(DetailsControl.ToSharedRef());
 	}
 	else
 	{
diff --git a/Engine/Source/Developer/Merge/Private/SBlueprintMerge.h b/Engine/Source/Developer/Merge/Private/SBlueprintMerge.h
--- a/Engine/Source/Developer/Merge/Private/SBlueprintMerge.h
+++ b/Engine/Source/Developer/Merge/Private/SBlueprintMerge.h
@@ -75,7 +75,7 @@ private:
 	FBlueprintMergeData		Data; 
 	FString					BackupSubDir;
 
-	TSharedPtr<SBorder>		MainView;
+	TSharedPtr<SBox>		MainView;
 
 	/** 
 	 * We track the package name-paths for the remote, base, & local assets (so
@@ -93,31 +93,25 @@ private:
 	 */
 	FString	LocalBackupPath;
 
-
-	struct FMergeControl
-	{
-		FMergeControl()
-			: Widget()
-			, DiffControl(NULL)
-			, MergeControl(NULL)
-		{
-		}
-
-		TSharedPtr<SWidget> Widget;
-		class IDiffControl* DiffControl;
-		class IMergeControl* MergeControl;
-	};
-
-	FMergeControl GraphControl;
-	FMergeControl TreeControl;
-	FMergeControl DetailsControl;
-	FMergeControl AssetPickerControl;
+	TSharedPtr<SWidget> GraphControl;
+	TSharedPtr<SWidget> TreeControl;
+	TSharedPtr<SWidget> DetailsControl;
+	TSharedPtr<SWidget> AssetPickerControl;
 	
 	bool bIsPickingAssets;
 	FOnMergeResolved OnMergeResolved;
 
-	class IDiffControl* CurrentDiffControl;
-	class IMergeControl* CurrentMergeControl;
+	/** Treeview to display all differences collected across all panels: */
+	TSharedPtr< STreeView< TSharedPtr< class FBlueprintDifferenceTreeEntry > > > DifferencesTreeView;
+
+	/** List of differences collected across all panels: */
+	TArray< TSharedPtr<class FBlueprintDifferenceTreeEntry> > MasterDifferencesList;
+
+	/** List of all differences, cached so that we can iterate only the differences and not labels, etc: */
+	TArray< TSharedPtr<class FBlueprintDifferenceTreeEntry> > RealDifferences;
+
+	/** List of all merge conflicts: */
+	TArray< TSharedPtr<class FBlueprintDifferenceTreeEntry> > MergeConflicts;
 
 	// This has to be allocated here because SListView cannot own the list
 	// that it is displaying. It also seems like the display list *has*
diff --git a/Engine/Source/Developer/Merge/Private/SMergeAssetPickerView.h b/Engine/Source/Developer/Merge/Private/SMergeAssetPickerView.h
--- a/Engine/Source/Developer/Merge/Private/SMergeAssetPickerView.h
+++ b/Engine/Source/Developer/Merge/Private/SMergeAssetPickerView.h
@@ -1,7 +1,6 @@
 // Copyright 1998-2015 Epic Games, Inc. All Rights Reserved.
 #pragma once
 
-#include "IDiffControl.h"
 #include "BlueprintMergeData.h" // for IMergeControl
 #include "MergeUtils.h" // for EMergeAssetId
 
@@ -22,8 +21,6 @@ struct FAssetRevisionInfo
  ******************************************************************************/
 
 class SMergeAssetPickerView : public SCompoundWidget
-                      , public IDiffControl
-					  , public IMergeControl
 {
 	DECLARE_DELEGATE_TwoParams(FOnMergeAssetChanged, EMergeAssetId::Type, const FAssetRevisionInfo&)
 
@@ -35,18 +32,6 @@ public:
 	void Construct(const FArguments InArgs, const FBlueprintMergeData& InData);
 
 private:
-	/** Implementation of IDiffControl: */
-	void NextDiff() override {}
-	void PrevDiff() override {}
-	bool HasNextDifference() const override { return false; }
-	bool HasPrevDifference() const override { return false; }
-
-	/** Implementation of IMergeControl */
-	void HighlightNextConflict() override {}
-	void HighlightPrevConflict() override {}
-	bool HasNextConflict() const override { return false; }
-	bool HasPrevConflict() const override { return false; }
-
 	void HandleAssetChange(const UObject* NewAsset, EMergeAssetId::Type PanelId);
 	void HandleRevisionChange(const FRevisionInfo& NewRevision, EMergeAssetId::Type PanelId);
 
diff --git a/Engine/Source/Developer/Merge/Private/SMergeDetailsView.cpp b/Engine/Source/Developer/Merge/Private/SMergeDetailsView.cpp
--- a/Engine/Source/Developer/Merge/Private/SMergeDetailsView.cpp
+++ b/Engine/Source/Developer/Merge/Private/SMergeDetailsView.cpp
@@ -7,7 +7,12 @@
 #include "SBlueprintDiff.h"
 #include "SMergeDetailsView.h"
 
-void SMergeDetailsView::Construct(const FArguments InArgs, const FBlueprintMergeData& InData )
+void SMergeDetailsView::Construct(const FArguments InArgs
+	, const FBlueprintMergeData& InData
+	, FOnMergeNodeSelected SelectionCallback
+	, TArray< TSharedPtr<FBlueprintDifferenceTreeEntry> >& OutTreeEntries
+	, TArray< TSharedPtr<FBlueprintDifferenceTreeEntry> >& OutRealDifferences
+	, TArray< TSharedPtr<FBlueprintDifferenceTreeEntry> >& OutConflicts)
 {
 	Data = InData;
 	CurrentMergeConflict = INDEX_NONE;
@@ -16,10 +21,23 @@ void SMergeDetailsView::Construct(const FArguments InArgs, const FBlueprintMerge
 	const UObject* BaseCDO = DiffUtils::GetCDO(InData.BlueprintBase);
 	const UObject* LocalCDO = DiffUtils::GetCDO(InData.BlueprintLocal);
 
+	TArray<FSingleObjectDiffEntry> RemoteDifferences;
+	DiffUtils::CompareUnrelatedObjects(BaseCDO, RemoteCDO, RemoteDifferences);
+	TArray<FSingleObjectDiffEntry> LocalDifferences;
+	DiffUtils::CompareUnrelatedObjects(BaseCDO, LocalCDO, LocalDifferences);
+
+	const auto GetPropertyNames = []( const TArray<FSingleObjectDiffEntry>& FromDifferences, TArray<FPropertySoftPath>& OutNames )
+	{
+		for( const auto& Entry : FromDifferences )
+		{
+			OutNames.Push(Entry.Identifier);
+		}
+	};
+
 	TArray<FPropertySoftPath> RemoteDifferingProperties;
-	DiffUtils::CompareUnrelatedObjects(BaseCDO, RemoteCDO, RemoteDifferingProperties);
+	GetPropertyNames( RemoteDifferences, RemoteDifferingProperties );
 	TArray<FPropertySoftPath> LocalDifferingProperties;
-	DiffUtils::CompareUnrelatedObjects(BaseCDO, LocalCDO, LocalDifferingProperties);
+	GetPropertyNames(LocalDifferences, LocalDifferingProperties);
 
 	FPropertySoftPathSet RemoteDifferingPropertiesSet(RemoteDifferingProperties);
 	FPropertySoftPathSet LocalDifferingPropertiesSet(LocalDifferingProperties);
@@ -69,27 +87,163 @@ void SMergeDetailsView::Construct(const FArguments InArgs, const FBlueprintMerge
 	TArray<FPropertySoftPath> RemoteVisibleProperties = GetRemoteDetails().GetDisplayedProperties();
 	TArray<FPropertySoftPath> BaseVisibleProperties = GetBaseDetails().GetDisplayedProperties();
 	TArray<FPropertySoftPath> LocalVisibleProperties = GetLocalDetails().GetDisplayedProperties();
-	if (BaseCDO)
+	int IterRemote = 0;
+	int IterBase = 0;
+	int IterLocal = 0;
+
+	bool bDoneRemote = IterRemote == RemoteVisibleProperties.Num();
+	bool bDoneBase = IterBase == BaseVisibleProperties.Num();
+	bool bDoneLocal = IterLocal == LocalVisibleProperties.Num();
+
+	bool bAnyConflict = false;
+
+	struct FDiffPair
+	{
+		FPropertySoftPath Identifier;
+		FText Label;
+		bool bConflicted;
+	};
+
+	const auto ProcessPotentialDifference = []( const FPropertySoftPath& PropertyIdentifier
+											, TArray<FSingleObjectDiffEntry> const& RemoteDifferences
+											, TArray<FSingleObjectDiffEntry> const& LocalDifferences
+											, TArray< FDiffPair >& OutProcessedDifferences
+											, bool& bOutAnyConflict )
 	{
-		for (const auto& Property : BaseVisibleProperties)
+		const FText RemoteLabel = NSLOCTEXT("SMergeDetailsView", "RemoteLabel", "Remote");
+		const FText LocalLabel = NSLOCTEXT("SMergeDetailsView", "LocalLabel", "Local");
+
+		const auto FindDiffering = [](TArray< FSingleObjectDiffEntry > const& InDifferences, const FPropertySoftPath& PropertyIdentifier) -> const FSingleObjectDiffEntry*
+		{
+			for (const auto& Difference : InDifferences)
+			{
+				if (Difference.Identifier == PropertyIdentifier)
+				{
+					return &Difference;
+				}
+			}
+			return nullptr;
+		};
+
+		const FSingleObjectDiffEntry* RemoteDiffering = FindDiffering(RemoteDifferences, PropertyIdentifier);
+		const FSingleObjectDiffEntry* LocalDiffering = FindDiffering(LocalDifferences, PropertyIdentifier);
+		if (RemoteDiffering && LocalDiffering)
+		{
+			// conflicting change:
+			bOutAnyConflict = true;
+
+			FText Label1 = DiffViewUtils::PropertyDiffMessage(*RemoteDiffering, RemoteLabel);
+			FText Label2 = DiffViewUtils::PropertyDiffMessage(*LocalDiffering, LocalLabel);
+			FText FinalLabel = FText::Format(NSLOCTEXT("SMergeDetailsView", "PropertyConflict", "Conflict: {0} and {1}"), Label1, Label2);
+
+			FDiffPair Difference = { RemoteDiffering->Identifier, FinalLabel, true };
+			OutProcessedDifferences.Push(Difference);
+		}
+		else if (RemoteDiffering)
 		{
-			AddPropertiesOrdered(Property, BaseDifferingPropertiesSet, DifferingProperties);
+			FText Label = DiffViewUtils::PropertyDiffMessage(*RemoteDiffering, RemoteLabel);
+			FDiffPair Difference = { RemoteDiffering->Identifier,  Label, false };
+			OutProcessedDifferences.Push(Difference);
 		}
-	}
-	if (RemoteCDO)
+		else if (LocalDiffering)
+		{
+			FText Label = DiffViewUtils::PropertyDiffMessage(*LocalDiffering, LocalLabel);
+			FDiffPair Difference = { LocalDiffering->Identifier, Label, false };
+			OutProcessedDifferences.Push(Difference);
+		}
+	};
+
+	TArray< FDiffPair > OrderedDifferences;
+	while( !bDoneRemote || !bDoneBase || !bDoneLocal )
 	{
-		for (const auto& Property : RemoteVisibleProperties)
+		bool bLocalMatchesBase = !bDoneLocal && !bDoneBase && BaseVisibleProperties[IterBase] == LocalVisibleProperties[IterLocal];
+		bool bRemoteMatchesBase = !bDoneRemote && !bDoneBase && BaseVisibleProperties[IterBase] == RemoteVisibleProperties[IterRemote];
+
+		if( bRemoteMatchesBase && bLocalMatchesBase ||
+			(bDoneLocal && bRemoteMatchesBase ) ||
+			(bDoneRemote && bLocalMatchesBase ) )
+		{
+			ProcessPotentialDifference(BaseVisibleProperties[IterBase], RemoteDifferences, LocalDifferences, OrderedDifferences, bAnyConflict );
+			AddPropertiesOrdered(BaseVisibleProperties[IterBase], BaseDifferingPropertiesSet, DifferingProperties);
+
+			if (!bDoneLocal)
+			{
+				++IterLocal;
+			}
+			if (!bDoneRemote)
+			{
+				++IterRemote;
+			}
+			if(!bDoneBase )
+			{
+				++IterBase;
+			}
+		}
+		else if ( !bDoneRemote && !bRemoteMatchesBase )
 		{
-			AddPropertiesOrdered(Property, RemoteDifferingPropertiesSet, DifferingProperties);
+			ProcessPotentialDifference(RemoteVisibleProperties[IterRemote], RemoteDifferences, LocalDifferences, OrderedDifferences, bAnyConflict);
+			AddPropertiesOrdered(RemoteVisibleProperties[IterRemote], RemoteDifferingPropertiesSet, DifferingProperties);
+			++IterRemote;
 		}
+		else if( !bDoneLocal && !bLocalMatchesBase )
+		{
+			ProcessPotentialDifference(LocalVisibleProperties[IterLocal], RemoteDifferences, LocalDifferences, OrderedDifferences, bAnyConflict);
+			AddPropertiesOrdered(LocalVisibleProperties[IterLocal], LocalDifferingPropertiesSet, DifferingProperties);
+			++IterLocal;
+		}
+
+		bDoneRemote = IterRemote == RemoteVisibleProperties.Num();
+		bDoneBase = IterBase == BaseVisibleProperties.Num();
+		bDoneLocal = IterLocal == LocalVisibleProperties.Num();
+	}
+
+	TArray< TSharedPtr<FBlueprintDifferenceTreeEntry> > Children;
+	if( OrderedDifferences.Num() == 0 )
+	{
+		Children.Push( FBlueprintDifferenceTreeEntry::NoDifferencesEntry() );
 	}
-	if (LocalCDO)
+	else
 	{
-		for (const auto& Property : LocalVisibleProperties)
+		const auto CreateDetailsMergeWidget = [](FDiffPair Entry) -> TSharedRef<SWidget>
 		{
-			AddPropertiesOrdered(Property, LocalDifferingPropertiesSet, DifferingProperties);
+			return SNew(STextBlock)
+				.Text(Entry.Label)
+				.ColorAndOpacity(Entry.bConflicted ? DiffViewUtils::Conflicting() : DiffViewUtils::Differs());
+		};
+
+		const auto FocusDetailsDifferenceEntry = []( FPropertySoftPath PropertyIdentifier, SMergeDetailsView* Parent, FOnMergeNodeSelected SelectionCallback )
+		{
+			SelectionCallback.ExecuteIfBound();
+			Parent->HighlightDifference(PropertyIdentifier);
+		};
+
+		for( const auto Difference : OrderedDifferences )
+		{
+			auto Entry = TSharedPtr<FBlueprintDifferenceTreeEntry>(
+				new FBlueprintDifferenceTreeEntry(
+					FOnDiffEntryFocused::CreateStatic(FocusDetailsDifferenceEntry, Difference.Identifier, this, SelectionCallback)
+					, FGenerateDiffEntryWidget::CreateStatic(CreateDetailsMergeWidget, Difference)
+					, TArray< TSharedPtr<FBlueprintDifferenceTreeEntry> >()
+				)
+			);
+			Children.Push(Entry);
+			OutRealDifferences.Push(Entry);
+			if( Difference.bConflicted )
+			{
+				OutConflicts.Push(Entry);
+			}
 		}
 	}
+
+	const auto ForwardSelection = [](FOnMergeNodeSelected SelectionCallback)
+	{
+		// This allows the owning control to focus the correct tab (or do whatever else it likes):
+		SelectionCallback.ExecuteIfBound();
+	};
+
+	TSharedPtr<FBlueprintDifferenceTreeEntry> Category = FBlueprintDifferenceTreeEntry::CreateDefaultsCategoryEntryForMerge(FOnDiffEntryFocused::CreateStatic(ForwardSelection, SelectionCallback), Children, RemoteDifferences.Num() != 0, LocalDifferences.Num() != 0, bAnyConflict);
+	OutTreeEntries.Push(Category);
+
 	CurrentDifference = -1;
 
 	ChildSlot[
@@ -109,90 +263,14 @@ void SMergeDetailsView::Construct(const FArguments InArgs, const FBlueprintMerge
 	];
 }
 
-void SMergeDetailsView::NextDiff()
-{
-	if (DifferingProperties.Num() == 0)
-	{
-		return;
-	}
-
-	CurrentDifference = (CurrentDifference + 1) % DifferingProperties.Num();
-
-	HighlightCurrentDifference();
-}
-
-void SMergeDetailsView::PrevDiff()
-{
-	if (DifferingProperties.Num() == 0)
-	{
-		return;
-	}
-
-	--CurrentDifference;
-	if (CurrentDifference < 0)
-	{
-		CurrentDifference = DifferingProperties.Num() - 1;
-	}
-
-	HighlightCurrentDifference();
-}
-
-void SMergeDetailsView::HighlightCurrentDifference()
-{
-	for( auto& DetailDiff : DetailsViews )
-	{
-		DetailDiff.HighlightProperty(DifferingProperties[CurrentDifference]);
-	}
-}
-
-bool SMergeDetailsView::HasNextDifference() const
-{
-	return DifferingProperties.IsValidIndex( CurrentDifference + 1 );
-}
-
-bool SMergeDetailsView::HasPrevDifference() const
-{
-	return DifferingProperties.IsValidIndex(CurrentDifference - 1);
-}
-
-void SMergeDetailsView::HighlightNextConflict()
-{
-	if (CurrentMergeConflict + 1 < MergeConflicts.Num())
-	{
-		++CurrentMergeConflict;
-	}
-	for (auto& DetailDiff : DetailsViews)
-	{
-		DetailDiff.HighlightProperty(MergeConflicts[CurrentMergeConflict]);
-	}
-}
-
-void SMergeDetailsView::HighlightPrevConflict()
+void SMergeDetailsView::HighlightDifference(FPropertySoftPath Path)
 {
-	if (CurrentMergeConflict - 1 >= 0)
-	{
-		--CurrentMergeConflict;
-	}
 	for (auto& DetailDiff : DetailsViews)
 	{
-		DetailDiff.HighlightProperty(MergeConflicts[CurrentMergeConflict]);
+		DetailDiff.HighlightProperty(Path);
 	}
 }
 
-bool SMergeDetailsView::HasNextConflict() const
-{
-	// return true if we have one conflict so that users can reselect the conflict if they desire. If we 
-	// return false when we already have selected this one and only conflict then there will be no way
-	// to reselect it if the user wants to.
-	return MergeConflicts.Num() != 0 && (MergeConflicts.Num() == 1 || CurrentMergeConflict + 1 < MergeConflicts.Num());
-}
-
-bool SMergeDetailsView::HasPrevConflict() const
-{
-	// note in HasNextConflict applies here as well.
-	return MergeConflicts.Num() == 1 || CurrentMergeConflict > 0;
-}
-
 FDetailsDiff& SMergeDetailsView::GetRemoteDetails()
 {
 	return DetailsViews[EMergeParticipant::Remote];
diff --git a/Engine/Source/Developer/Merge/Private/SMergeDetailsView.h b/Engine/Source/Developer/Merge/Private/SMergeDetailsView.h
--- a/Engine/Source/Developer/Merge/Private/SMergeDetailsView.h
+++ b/Engine/Source/Developer/Merge/Private/SMergeDetailsView.h
@@ -2,12 +2,9 @@
 #pragma once
 
 #include "BlueprintMergeData.h"
-#include "IDiffControl.h"
 #include "DetailsDiff.h"
 
 class SMergeDetailsView : public SCompoundWidget
-						, public IDiffControl
-						, public IMergeControl
 {
 public:
 	virtual ~SMergeDetailsView() {}
@@ -16,21 +13,14 @@ public:
 	{}
 	SLATE_END_ARGS()
 
-	void Construct(const FArguments InArgs, const FBlueprintMergeData& InData );
+	void Construct(const FArguments InArgs
+					, const FBlueprintMergeData& InData
+					, FOnMergeNodeSelected SelectionCallback
+					, TArray< TSharedPtr<FBlueprintDifferenceTreeEntry> >& OutTreeEntries
+					, TArray< TSharedPtr<FBlueprintDifferenceTreeEntry> >& OutRealDifferences
+					, TArray< TSharedPtr<FBlueprintDifferenceTreeEntry> >& OutConflicts);
 private:
-	/** Implementation of IDiffControl: */
-	void NextDiff() override;
-	void PrevDiff() override;
-	bool HasNextDifference() const override;
-	bool HasPrevDifference() const override;
-
-	/** Implementation of IMergeControl: */
-	void HighlightNextConflict() override;
-	void HighlightPrevConflict() override;
-	bool HasNextConflict() const override;
-	bool HasPrevConflict() const override;
-
-	void HighlightCurrentDifference();
+	void HighlightDifference(FPropertySoftPath Path);
 	FDetailsDiff& GetRemoteDetails();
 	FDetailsDiff& GetBaseDetails();
 	FDetailsDiff& GetLocalDetails();
diff --git a/Engine/Source/Developer/Merge/Private/SMergeGraphView.cpp b/Engine/Source/Developer/Merge/Private/SMergeGraphView.cpp
--- a/Engine/Source/Developer/Merge/Private/SMergeGraphView.cpp
+++ b/Engine/Source/Developer/Merge/Private/SMergeGraphView.cpp
@@ -12,36 +12,6 @@
 const FName MergeMyBluerpintTabId = FName(TEXT("MergeMyBluerpintTab"));
 const FName MergeGraphTabId = FName(TEXT("MergeGraphTab"));
 
-struct FGraphMergeConflict
-{
-	FGraphMergeConflict(FName InGraphName, const FDiffSingleResult& InRemoteDifference, const FDiffSingleResult& InLocalDifference )
-		: GraphName(InGraphName)
-		, RemoteDifference(InRemoteDifference)
-		, LocalDifference(InLocalDifference)
-	{
-	}
-
-	FName GraphName;
-
-	// These are the diff results that we have judged to be conflicting:
-	FDiffSingleResult RemoteDifference;
-	FDiffSingleResult LocalDifference;
-};
-
-// This cludge is used to move arrays of structures into arrays of 
-// shared ptrs of those structures. It is used because SListView
-// does not support arrays of values:
-template< typename T >
-TArray< TSharedPtr<T> > ConcreteToShared(const TArray<T>& Values)
-{
-	TArray< TSharedPtr<T> > Ret;
-	for (const auto& Value : Values)
-	{
-		Ret.Add(TSharedPtr<T>(new T(Value)));
-	}
-	return Ret;
-}
-
 struct FBlueprintRevPair
 {
 	FBlueprintRevPair(const UBlueprint* InBlueprint, const FRevisionInfo& InRevData)
@@ -52,23 +22,7 @@ struct FBlueprintRevPair
 
 	const UBlueprint* Blueprint;
 	const FRevisionInfo& RevData;
-};
-
-struct FMergeGraphRowEntry
-{
-	FName GraphName;
-
-	// These lists contain shared ptrs because they are displayed by a
-	// SListView, and it currently does not support lists of values.
-	TArray< TSharedPtr<FDiffSingleResult> > RemoteDifferences;
-	TArray< TSharedPtr<FDiffSingleResult> > LocalDifferences;
-
-	bool bExistsInRemote;
-	bool bExistsInBase;
-	bool bExistsInLocal;
-	bool bDiffersInRemote;
-	bool bDiffersInLocal;
-};
+};;
 
 static UEdGraph* FindGraphByName(UBlueprint const& FromBlueprint, const FName& GraphName)
 {
@@ -83,16 +37,38 @@ static UEdGraph* FindGraphByName(UBlueprint const& FromBlueprint, const FName& G
 	return Ret;
 }
 
-static const UEdGraphNode* FindNodeInUnrelatedGraph(UEdGraph* Graph, UEdGraphNode* TargetNode)
+struct FMergeGraphRowEntry
 {
-	check(TargetNode->GetGraph() != Graph); // purpose of this function is to find a node in a different version of its parent graph
+	FText Label;
 
-	TArray< FGraphDiffControl::FNodeMatch > PriorMatches;
-	FGraphDiffControl::FNodeMatch NodeMatch = FGraphDiffControl::FindNodeMatch(Graph, TargetNode, PriorMatches);
-	return NodeMatch.NewNode;
-}
+	FName GraphName;
 
-static TArray< TSharedPtr<FMergeGraphRowEntry> > GenerateDiffListItems(const FBlueprintRevPair& RemoteBlueprint, const FBlueprintRevPair& BaseBlueprint, const FBlueprintRevPair& LocalBlueprint,  TArray< TSharedPtr<FGraphMergeConflict> >& OutMergeConflicts )
+	UEdGraphNode* LocalNode;
+	UEdGraphNode* BaseNode;
+	UEdGraphNode* RemoteNode;
+
+	UEdGraphPin* LocalPin;
+	UEdGraphPin* BasePin;
+	UEdGraphPin* RemotePin;
+
+	FLinearColor DisplayColor;
+};
+
+struct FMergeGraphEntry
+{
+	FName GraphName;
+
+	TArray<FMergeGraphRowEntry> Changes;
+	bool bAnyConflics;
+	bool bRemoteDifferences;
+	bool bLocalDifferences;
+
+	bool bExistsInRemote;
+	bool bExistsInBase;
+	bool bExistsInLocal;
+};
+
+static TArray< FMergeGraphEntry > GenerateDiffListItems(const FBlueprintRevPair& RemoteBlueprint, const FBlueprintRevPair& BaseBlueprint, const FBlueprintRevPair& LocalBlueprint )
 {
 	// Index all the graphs by name, we use the name of the graph as the 
 	// basis of comparison between the various versions of the blueprint.
@@ -119,12 +95,21 @@ static TArray< TSharedPtr<FMergeGraphRowEntry> > GenerateDiffListItems(const FBl
 		ToMap(GraphsLocal, LocalGraphMap);
 	}
 
-	TArray< TSharedPtr<FMergeGraphRowEntry> > Ret;
+	TArray< FMergeGraphEntry > Ret;
 	{
 		const auto GenerateDifferences = [](UEdGraph* GraphNew, UEdGraph** GraphOld)
 		{
 			TArray<FDiffSingleResult> Results;
 			FGraphDiffControl::DiffGraphs(GraphOld ? *GraphOld : nullptr, GraphNew, Results);
+			struct SortDiff
+			{
+				bool operator () (const FDiffResultItem& A, const FDiffResultItem& B) const
+				{
+					return A.Result.Diff < B.Result.Diff;
+				}
+			};
+
+			Sort(Results.GetData(), Results.Num(), SortDiff());
 			return Results;
 		};
 
@@ -134,11 +119,18 @@ static TArray< TSharedPtr<FMergeGraphRowEntry> > GenerateDiffListItems(const FBl
 			TArray< FDiffSingleResult > LocalDifferences;
 			bool bExistsInRemote, bExistsInBase, bExistsInLocal;
 
+			FMergeGraphEntry GraphEntry;
+			GraphEntry.GraphName = GraphName;
 			{
 				UEdGraph** RemoteGraph = RemoteGraphMap.Find(GraphName);
 				UEdGraph** BaseGraph = BaseGraphMap.Find(GraphName);
 				UEdGraph** LocalGraph = LocalGraphMap.Find(GraphName);
 
+				GraphEntry.bAnyConflics = false;
+				GraphEntry.bExistsInRemote = RemoteGraph != nullptr;
+				GraphEntry.bExistsInBase = BaseGraph != nullptr;
+				GraphEntry.bExistsInLocal = LocalGraph != nullptr;
+
 				if (RemoteGraph)
 				{
 					RemoteDifferences = GenerateDifferences(*RemoteGraph, BaseGraph);
@@ -152,23 +144,25 @@ static TArray< TSharedPtr<FMergeGraphRowEntry> > GenerateDiffListItems(const FBl
 				// 'join' the local differences and remote differences by noting changes
 				// that affected the same common base:
 				{
+					TMap< const FDiffSingleResult*, const FDiffSingleResult*> ConflictMap;
+
 					for (const auto& RemoteDifference : RemoteDifferences)
 					{
 						const FDiffSingleResult* ConflictingDifference = nullptr;
 
 						for (const auto& LocalDifference : LocalDifferences)
 						{
-							if (RemoteDifference.Node1 == LocalDifference.Node1 )
+							if (RemoteDifference.Node1 == LocalDifference.Node1)
 							{
-								if( RemoteDifference.Diff == EDiffType::NODE_REMOVED || 
+								if (RemoteDifference.Diff == EDiffType::NODE_REMOVED ||
 									LocalDifference.Diff == EDiffType::NODE_REMOVED ||
-									RemoteDifference.Pin1 == LocalDifference.Pin1 )
+									RemoteDifference.Pin1 == LocalDifference.Pin1)
 								{
 									ConflictingDifference = &LocalDifference;
 									break;
 								}
 							}
-							else if( RemoteDifference.Pin1 != nullptr && (RemoteDifference.Pin1 == LocalDifference.Pin1 ) )
+							else if (RemoteDifference.Pin1 != nullptr && (RemoteDifference.Pin1 == LocalDifference.Pin1))
 							{
 								// it's possible the users made the same change to the same pin, but given the wide
 								// variety of changes that can be made to a pin it is difficult to identify the change 
@@ -187,29 +181,77 @@ static TArray< TSharedPtr<FMergeGraphRowEntry> > GenerateDiffListItems(const FBl
 								continue;
 							}
 
-							OutMergeConflicts.Push(MakeShareable( new FGraphMergeConflict( GraphName, RemoteDifference, *ConflictingDifference)) );
+							ConflictMap[&RemoteDifference] = ConflictingDifference;
+							ConflictMap[ConflictingDifference] = &RemoteDifference;
 						}
-						// else: no conflict, we want to be able to automatically apply this remote change
-						// to the local revision:
 					}
+
+
+					for( const auto& Difference : RemoteDifferences )
+					{
+						FText Label;
+
+						const FDiffSingleResult** ConflictingDifference = ConflictMap.Find(&Difference);
+
+						if( ConflictingDifference )
+						{
+							Label = FText::Format( NSLOCTEXT("SMergeGraphView", "ConflictIdentifier", "CONFLICT: {0} conflicts with {1}" ), FText::FromString((*ConflictingDifference)->DisplayString), FText::FromString(Difference.DisplayString) );
+						}
+						else
+						{
+							Label = FText::FromString( Difference.DisplayString );
+						}
+
+						FMergeGraphRowEntry NewEntry = {
+							Label
+							, Difference.OwningGraph
+							, ConflictingDifference ? (*ConflictingDifference)->Node2 : nullptr /*UEdGraphNode* LocalNode*/
+							, Difference.Node1 /*UEdGraphNode* BaseNode*/
+							, Difference.Node2 /*UEdGraphNode* RemoteNode*/
+							, ConflictingDifference ? (*ConflictingDifference)->Pin2 : nullptr /*UEdGraphPin* LocalPin*/
+							, Difference.Pin1 /*UEdGraphPin* BasePin*/
+							, Difference.Pin2 /*UEdGraphPin* RemotePin*/
+							, Difference.DisplayColor
+						};
+
+						GraphEntry.bAnyConflics |= NewEntry.LocalNode && NewEntry.RemoteNode;
+						GraphEntry.Changes.Push( NewEntry );
+					}
+
+					for (const auto& Difference : LocalDifferences)
+					{
+						FText Label;
+
+						const FDiffSingleResult** ConflictingDifference = ConflictMap.Find( &Difference );
+
+						if (!ConflictingDifference)
+						{
+							FMergeGraphRowEntry NewEntry = {
+								FText::FromString(Difference.DisplayString)
+								, Difference.OwningGraph
+								, Difference.Node2 /*UEdGraphNode* LocalNode*/
+								, Difference.Node1 /*UEdGraphNode* BaseNode*/
+								, nullptr
+								, Difference.Pin2 /*UEdGraphPin* LocalPin*/
+								, Difference.Pin1 /*UEdGraphPin* BasePin*/
+								, nullptr
+								, Difference.DisplayColor
+							};
+
+							GraphEntry.Changes.Push(NewEntry);
+						}
+					}
+
+					GraphEntry.bLocalDifferences = LocalDifferences.Num() != 0;
+					GraphEntry.bRemoteDifferences = RemoteDifferences.Num() != 0;
 				}
 
+				Ret.Add(GraphEntry);
+
 				bExistsInRemote = RemoteGraph != nullptr;
 				bExistsInBase = BaseGraph != nullptr;
 				bExistsInLocal = LocalGraph != nullptr;
 			}
-
-			FMergeGraphRowEntry NewEntry = {
-				GraphName,
-				ConcreteToShared(RemoteDifferences),
-				ConcreteToShared(LocalDifferences),
-				bExistsInRemote,
-				bExistsInBase,
-				bExistsInLocal,
-				RemoteDifferences.Num() != 0,
-				LocalDifferences.Num() != 0
-			};
-			Ret.Add(TSharedPtr<FMergeGraphRowEntry>(new FMergeGraphRowEntry(NewEntry)));
 		}
 	}
 
@@ -262,13 +304,13 @@ FDiffPanel& GetDiffPanelForNode(const UEdGraphNode& Node, TArray< FDiffPanel >&
 	return Default;
 }
 
-// Some color constants used by the merge view:
-const FLinearColor MergeBlue = FLinearColor(0.3f, 0.3f, 1.f);
-const FLinearColor MergeRed = FLinearColor(1.0f, 0.2f, 0.3f);
-const FLinearColor MergeYellow = FLinearColor::Yellow;
-const FLinearColor MergeWhite = FLinearColor::White;
-
-void SMergeGraphView::Construct(const FArguments InArgs, const FBlueprintMergeData& InData )
+void SMergeGraphView::Construct(const FArguments InArgs
+	, const FBlueprintMergeData& InData
+	, FOnMergeNodeSelected SelectionCallback
+	, TArray< TSharedPtr<FBlueprintDifferenceTreeEntry> >& OutTreeEntries
+	, TArray< TSharedPtr<FBlueprintDifferenceTreeEntry> >& OutRealDifferences
+	, TArray< TSharedPtr<FBlueprintDifferenceTreeEntry> >& OutConflicts
+	)
 {
 	const TSharedRef<SDockTab> MajorTab = SNew(SDockTab)
 		.TabRole(ETabRole::MajorTab);
@@ -285,13 +327,9 @@ void SMergeGraphView::Construct(const FArguments InArgs, const FBlueprintMergeDa
 		.SetDisplayName(LOCTEXT("MergeMyBlueprintTabTitle", "My Blueprint"))
 		.SetTooltipText(LOCTEXT("MergeMyBlueprintTooltipText", "Differences in the 'My Blueprints' attributes of the blueprint"));
 
-	CurrentMergeConflict = INDEX_NONE;
-
 	Data = InData;
 	bViewsAreLocked = true;
-	RemoteDiffResultsListData = nullptr;
-	LocalDiffResultsListData = nullptr;
-
+	
 	TArray<FBlueprintRevPair> BlueprintsForDisplay;
 	// EMergeParticipant::Remote
 	BlueprintsForDisplay.Add(FBlueprintRevPair(InData.BlueprintRemote, InData.RevisionRemote));
@@ -337,398 +375,166 @@ void SMergeGraphView::Construct(const FArguments InArgs, const FBlueprintMergeDa
 		];
 	}
 
-	DifferencesFromBase = GenerateDiffListItems(BlueprintsForDisplay[EMergeParticipant::Remote]
-												, BlueprintsForDisplay[EMergeParticipant::Base]
-												, BlueprintsForDisplay[EMergeParticipant::Local]
-												, MergeConflicts);
+	Differences = TSharedPtr< TArray< FMergeGraphEntry > >( new TArray< FMergeGraphEntry >( 
+																GenerateDiffListItems(BlueprintsForDisplay[EMergeParticipant::Remote]
+																					, BlueprintsForDisplay[EMergeParticipant::Base]
+																					, BlueprintsForDisplay[EMergeParticipant::Local]) ) );
 
-	// This is the function we'll use to generate a row in the control that lists all the available graphs:
-	const auto RowGenerator = [](TSharedPtr<FMergeGraphRowEntry> ParamItem, const TSharedRef<STableViewBase>& OwnerTable) -> TSharedRef<ITableRow>
+	for( const auto& Difference : *Differences)
 	{
-		// blue indicates added, red indicates changed, yellow indicates removed, white indicates no change:
-		const auto ComputeColorForRevision = [](bool bExistsInBase, bool bExistsInThisRevision, bool bHasChangesInThisRevision) -> FLinearColor
+		TArray< TSharedPtr<FBlueprintDifferenceTreeEntry> > Children;
+		for( const auto& Change : Difference.Changes )
 		{
-			if (bExistsInBase)
+			const auto ChangeWidget = [](FText Label, FLinearColor Color) -> TSharedRef<SWidget>
 			{
-				if (!bExistsInThisRevision)
-				{
-					return MergeYellow;
-				}
-				else if (bHasChangesInThisRevision)
-				{
-					return MergeRed;
-				}
+				return SNew(STextBlock)
+						.Text(Label)
+						.ColorAndOpacity(Color);
+			};
+
+			const auto SelectGraphNode = [](FOnMergeNodeSelected SelectionCallback, FMergeGraphRowEntry DiffEntry, SMergeGraphView* Parent)
+			{
+				SelectionCallback.ExecuteIfBound();
+				Parent->HighlightEntry( DiffEntry );
+			};
+
+			TSharedPtr<FBlueprintDifferenceTreeEntry> Entry = TSharedPtr<FBlueprintDifferenceTreeEntry>(new FBlueprintDifferenceTreeEntry(
+				FOnDiffEntryFocused::CreateStatic(SelectGraphNode, SelectionCallback, Change, this)
+				, FGenerateDiffEntryWidget::CreateStatic(ChangeWidget, Change.Label, Change.DisplayColor )
+				, TArray< TSharedPtr<FBlueprintDifferenceTreeEntry> >()
+			));
+			Children.Push(Entry);
+			OutRealDifferences.Push(Entry);
+			if( Change.LocalNode && Change.RemoteNode )
+			{
+				OutConflicts.Push(Entry);
 			}
-			if (!bExistsInBase)
+		}
+
+		const auto Widget = []( const FMergeGraphEntry* Difference ) -> TSharedRef<SWidget>
+		{
+			// blue indicates added, red indicates changed, yellow indicates removed, white indicates no change:
+			const auto ComputeColor = [](const bool bAnyConflicts, const bool bAnyDifferences) -> FLinearColor
 			{
-				if (bExistsInThisRevision)
+				if( bAnyConflicts )
 				{
-					return MergeBlue;
+					return DiffViewUtils::Conflicting();
 				}
-			}
+				else if( bAnyDifferences )
+				{
+					return DiffViewUtils::Differs();
+				}
+				return DiffViewUtils::Identical();
+			};
 
-			return MergeWhite;
-		};
+			const auto Box = [](bool bIsPresent, FLinearColor Color) -> SHorizontalBox::FSlot&
+			{
+				return SHorizontalBox::Slot()
+					.AutoWidth()
+					.HAlign(HAlign_Right)
+					.VAlign(VAlign_Center)
+					.MaxWidth(8.0f)
+					[
+						SNew(SImage)
+						.ColorAndOpacity(Color)
+						.Image(bIsPresent ? FEditorStyle::GetBrush("BlueprintDif.HasGraph") : FEditorStyle::GetBrush("BlueprintDif.MissingGraph"))
+					];
+			};
 
-		FLinearColor RemoteColor = ComputeColorForRevision(ParamItem->bExistsInBase, ParamItem->bExistsInRemote, ParamItem->bDiffersInRemote);
-		FLinearColor BaseColor = ComputeColorForRevision(ParamItem->bExistsInBase, ParamItem->bExistsInBase, false);
-		FLinearColor LocalColor = ComputeColorForRevision(ParamItem->bExistsInBase, ParamItem->bExistsInLocal, ParamItem->bDiffersInLocal);
+			FLinearColor RemoteColor = ComputeColor(Difference->bAnyConflics, Difference->bRemoteDifferences);
+			FLinearColor BaseColor = ComputeColor(Difference->bAnyConflics, false);
+			FLinearColor LocalColor = ComputeColor(Difference->bAnyConflics, Difference->bLocalDifferences);
+			FLinearColor TextColor = ComputeColor(Difference->bAnyConflics, Difference->bLocalDifferences || Difference->bRemoteDifferences);
 
-		const auto Box = [](bool bIsPresent, FLinearColor Color) -> SHorizontalBox::FSlot&
-		{
-			return SHorizontalBox::Slot()
-				.AutoWidth()
-				.HAlign(HAlign_Right)
-				.VAlign(VAlign_Center)
-				.MaxWidth(8.0f)
-				[
-					SNew(SImage)
-					.ColorAndOpacity(Color)
-					.Image(bIsPresent ? FEditorStyle::GetBrush("BlueprintDif.HasGraph") : FEditorStyle::GetBrush("BlueprintDif.MissingGraph"))
-				];
-		};
-
-		return SNew(STableRow< TSharedPtr<FMergeGraphRowEntry> >, OwnerTable)
-			.ToolTipText(LOCTEXT("MergeGraphsDifferentText", "@todo doc"))
-			.Content()
-			[
-				SNew(SHorizontalBox)
+			return SNew(SHorizontalBox)
 				+ SHorizontalBox::Slot()
 				[
 					SNew(STextBlock)
-					.ColorAndOpacity(FLinearColor::White)
-					.Text(ParamItem->GraphName.GetPlainNameString())
+					.ColorAndOpacity( TextColor )
+					.Text(FText::FromString(Difference->GraphName.GetPlainNameString()))
 				]
-				+ Box(ParamItem->bExistsInRemote, RemoteColor)
-				+ Box(ParamItem->bExistsInBase, BaseColor)
-				+ Box(ParamItem->bExistsInLocal, LocalColor)
-			];
-	};
+				+ DiffViewUtils::Box(Difference->bExistsInRemote, RemoteColor)
+				+ DiffViewUtils::Box(Difference->bExistsInBase, BaseColor)
+				+ DiffViewUtils::Box(Difference->bExistsInLocal, LocalColor);
+		};
 
-	const auto DifferencesListGenerator = []( FText Title, TSharedPtr<SBox>& OutResultsContainer )
-	{
-		return SNew(SVerticalBox)
-			+ SVerticalBox::Slot()
-			.AutoHeight()
-			.HAlign(HAlign_Left)
-			[
-				SNew(STextBlock)
-				.Text(Title)
-			]
-		+ SVerticalBox::Slot()
-			.AutoHeight()
-			[
-				// this creates a line between the header and the list itself:
-				SNew(SBorder)
-				.Padding(FEditorStyle::GetMargin(TEXT("Menu.Separator.Padding")))
-				.BorderImage(FEditorStyle::GetBrush(TEXT("Menu.Separator")))
-			]
-		+ SVerticalBox::Slot()
-			[
-				SAssignNew(OutResultsContainer, SBox)
-			];
-	};
+		const auto FocusGraph = [](FOnMergeNodeSelected SelectionCallback, SMergeGraphView* Parent, FName GraphName)
+		{
+			SelectionCallback.ExecuteIfBound();
+			Parent->FocusGraph( GraphName );
+		};
+
+		if( Children.Num() == 0 )
+		{
+			Children.Push(FBlueprintDifferenceTreeEntry::NoDifferencesEntry());
+		}
+
+		OutTreeEntries.Push(
+			TSharedPtr<FBlueprintDifferenceTreeEntry>(new FBlueprintDifferenceTreeEntry(
+				FOnDiffEntryFocused::CreateStatic(FocusGraph, SelectionCallback, this, Difference.GraphName)
+				, FGenerateDiffEntryWidget::CreateStatic(Widget, &Difference)
+				, Children
+			))
+		);
+	}
 
 	ChildSlot
+	[
+		SNew(SSplitter)
+		.Orientation(Orient_Horizontal)
+		+ SSplitter::Slot()
+		.Value(0.9f)
 		[
 			SNew(SSplitter)
-			.Orientation(Orient_Horizontal)
+			.Orientation(Orient_Vertical)
 			+ SSplitter::Slot()
-			.Value(0.1f)
+			.Value(.8f)
 			[
-				// ToolPanel.GroupBorder is currently what gives this control its dark grey appearance:
-				SNew(SBorder)
-				.BorderImage(FEditorStyle::GetBrush("NoBorder"))
-				.Padding(FMargin(1.0f))
-				[
-					SNew(SSplitter)
-					.Orientation(Orient_Vertical)
-					+ SSplitter::Slot()
-					.Value(.3f)
-					[
-						SNew(SBorder)
-						.BorderImage(FEditorStyle::GetBrush("ToolPanel.GroupBorder"))
-						[
-							SNew(SVerticalBox)
-							+ SVerticalBox::Slot()
-							.AutoHeight()
-							.HAlign(HAlign_Left)
-							[
-								SNew(SButton)
-								.OnClicked(this, &SMergeGraphView::OnToggleLockView)
-								.Content()
-								[
-									SNew(SImage)
-									.Image(this, &SMergeGraphView::GetLockViewImage)
-									.ToolTipText(LOCTEXT("BPDifLock", "Lock the blueprint views?"))
-								]
-							]
-							+ SVerticalBox::Slot()
-								[
-									SNew(SListView< TSharedPtr<FMergeGraphRowEntry> >)
-									.ItemHeight(24)
-									.ListItemsSource(&DifferencesFromBase)
-									.OnGenerateRow(SListView< TSharedPtr<FMergeGraphRowEntry> >::FOnGenerateRow::CreateStatic(RowGenerator))
-									.SelectionMode(ESelectionMode::Single)
-									.OnSelectionChanged(this, &SMergeGraphView::OnGraphListSelectionChanged)
-								]
-						]
-					]
-					+ SSplitter::Slot()
-					.Value(.35f)
-					[
-						SNew(SBorder)
-						.BorderImage(FEditorStyle::GetBrush("ToolPanel.GroupBorder"))
-						[
-							DifferencesListGenerator(LOCTEXT("RemoteChangesLabel", "Remote Changes"), RemoteDiffResultsWidget)
-						]
-					]
-					+ SSplitter::Slot()
-					.Value(.35f)
-					[
-						SNew(SBorder)
-						.BorderImage(FEditorStyle::GetBrush("ToolPanel.GroupBorder"))
-						[
-							DifferencesListGenerator(LOCTEXT("LocalChangesLabel", "Local Changes"), LocalDiffResultsWidget)
-						]
-					]
-				]
+				GraphPanelContainer
 			]
 			+ SSplitter::Slot()
-				.Value(0.9f)
-				[
-					SNew(SSplitter)
-					.Orientation(Orient_Vertical)
-					+ SSplitter::Slot()
-					.Value(.8f)
-					[
-						GraphPanelContainer
-					]
-					+ SSplitter::Slot()
-					.Value(.2f)
-					[
-						DetailsPanelContainer
-					]
-				]
-		];
-}
-
-void SMergeGraphView::NextDiff()
-{
-	check(	RemoteDiffResultList.Pin().IsValid() && RemoteDiffResultsListData || 
-			LocalDiffResultList.Pin().IsValid() && LocalDiffResultsListData );
-	if( RemoteDiffResultsListData && DiffWidgetUtils::HasNextDifference( *RemoteDiffResultList.Pin(), *RemoteDiffResultsListData) )
-	{
-		DiffWidgetUtils::SelectNextRow( *RemoteDiffResultList.Pin(), *RemoteDiffResultsListData );
-	}
-	else
-	{
-		DiffWidgetUtils::SelectNextRow(*LocalDiffResultList.Pin(), *LocalDiffResultsListData);
-	}
-}
-
-void SMergeGraphView::PrevDiff()
-{
-	check(	RemoteDiffResultList.Pin().IsValid() && RemoteDiffResultsListData ||
-			LocalDiffResultList.Pin().IsValid() && LocalDiffResultsListData);
-	if (LocalDiffResultsListData && DiffWidgetUtils::HasPrevDifference(*LocalDiffResultList.Pin(), *LocalDiffResultsListData))
-	{
-		DiffWidgetUtils::SelectPrevRow(*LocalDiffResultList.Pin(), *LocalDiffResultsListData);
-	}
-	else
-	{
-		DiffWidgetUtils::SelectPrevRow(*RemoteDiffResultList.Pin(), *RemoteDiffResultsListData);
-	}
-}
-
-bool SMergeGraphView::HasNextDifference() const
-{
-	if (HasNoDifferences())
-	{
-		return false;
-	}
-
-	check(	RemoteDiffResultList.Pin().IsValid() && RemoteDiffResultsListData ||
-			LocalDiffResultList.Pin().IsValid() && LocalDiffResultsListData);
-	return	(RemoteDiffResultsListData && DiffWidgetUtils::HasNextDifference(*RemoteDiffResultList.Pin(), *RemoteDiffResultsListData) ) ||
-			(LocalDiffResultsListData && DiffWidgetUtils::HasNextDifference(*LocalDiffResultList.Pin(), *LocalDiffResultsListData) ) ;
-}
-
-bool SMergeGraphView::HasPrevDifference() const
-{
-	if (HasNoDifferences())
-	{
-		return false;
-	}
-
-	check(	RemoteDiffResultList.Pin().IsValid() && RemoteDiffResultsListData ||
-			LocalDiffResultList.Pin().IsValid() && LocalDiffResultsListData);
-	return DiffWidgetUtils::HasPrevDifference(*RemoteDiffResultList.Pin(), *RemoteDiffResultsListData) 
-		|| DiffWidgetUtils::HasPrevDifference(*LocalDiffResultList.Pin(), *LocalDiffResultsListData);
-}
-
-void SMergeGraphView::HighlightNextConflict()
-{
-	if( CurrentMergeConflict + 1 < MergeConflicts.Num())
-	{
-		++CurrentMergeConflict;
-	}
-	HighlightConflict(*MergeConflicts[CurrentMergeConflict]);
-}
-
-void SMergeGraphView::HighlightPrevConflict()
-{
-	if (CurrentMergeConflict - 1 >= 0 )
-	{
-		--CurrentMergeConflict;
-	}
-	HighlightConflict(*MergeConflicts[CurrentMergeConflict]);
-}
-
-bool SMergeGraphView::HasNextConflict() const 
-{
-	// return true if we have one conflict so that users can reselect the conflict if they desire. If we 
-	// return false when we already have selected this one and only conflict then there will be no way
-	// to reselect it if the user wants to.
-	return  MergeConflicts.Num() != 0 && (MergeConflicts.Num() == 1 || CurrentMergeConflict + 1 < MergeConflicts.Num());
-}
-
-bool SMergeGraphView::HasPrevConflict() const
-{
-	// note in HasNextConflict applies here as well.
-	return MergeConflicts.Num() == 1 || CurrentMergeConflict > 0;
-}
-
-void SMergeGraphView::HighlightConflict(const FGraphMergeConflict& Conflict)
-{
-	for( const auto& MergeGraphEntry : DifferencesFromBase )
-	{
-		if( MergeGraphEntry->GraphName == Conflict.GraphName )
-		{
-			OnGraphListSelectionChanged(MergeGraphEntry, ESelectInfo::Direct);
-			break;
-		}
-	}
-
-	// highlight the change made to the remote graph:
-	if( UEdGraphPin* Pin = Conflict.RemoteDifference.Pin2 )
-	{
-		// then look for the diff panel and focus on the change:
-		GetDiffPanelForNode(*Pin->GetOwningNode(), DiffPanels).FocusDiff(*Pin);
-	}
-	else if( UEdGraphNode* Node = Conflict.RemoteDifference.Node2 )
-	{
-		GetDiffPanelForNode(*Node, DiffPanels).FocusDiff(*Node);
-	}
-	
-	// and the change made to the local graph:
-	if (UEdGraphPin* Pin = Conflict.LocalDifference.Pin2)
-	{
-		GetDiffPanelForNode(*Pin->GetOwningNode(), DiffPanels).FocusDiff(*Pin);
-	}
-	else if (UEdGraphNode* Node = Conflict.LocalDifference.Node2)
-	{
-		GetDiffPanelForNode(*Node, DiffPanels).FocusDiff(*Node);
-	}
-}
-
-bool SMergeGraphView::HasNoDifferences() const
-{
-	auto RemoteDiffResultListPtr = RemoteDiffResultList.Pin();
-	auto LocalDiffResultListPtr = LocalDiffResultList.Pin();
-	return	(!RemoteDiffResultListPtr.IsValid() || RemoteDiffResultListPtr->GetNumItemsBeingObserved() == 0) &&
-			(!LocalDiffResultListPtr.IsValid() || LocalDiffResultListPtr->GetNumItemsBeingObserved() == 0);
+			.Value(.2f)
+			[
+				DetailsPanelContainer
+			]
+		]
+	];
 }
 
-void SMergeGraphView::OnGraphListSelectionChanged(TSharedPtr<FMergeGraphRowEntry> Item, ESelectInfo::Type SelectionType)
+void SMergeGraphView::FocusGraph(FName GraphName)
 {
-	if (( SelectionType != ESelectInfo::OnMouseClick &&
-		  SelectionType != ESelectInfo::Direct) ||
-		!Item.IsValid())
-	{
-		return;
-	}
-
-	FName GraphName = Item->GraphName;
-
 	UEdGraph* GraphRemote = FindGraphByName(*GetRemotePanel().Blueprint, GraphName);
 	UEdGraph* GraphBase = FindGraphByName(*GetBasePanel().Blueprint, GraphName);
 	UEdGraph* GraphLocal = FindGraphByName(*GetLocalPanel().Blueprint, GraphName);
 
-	GetBasePanel().GeneratePanel(GraphBase, nullptr );
-	GetRemotePanel().GeneratePanel(GraphRemote, GraphBase );
-	GetLocalPanel().GeneratePanel(GraphLocal, GraphBase );
+	GetBasePanel().GeneratePanel(GraphBase, nullptr);
+	GetRemotePanel().GeneratePanel(GraphRemote, GraphBase);
+	GetLocalPanel().GeneratePanel(GraphLocal, GraphBase);
 
 	LockViews(DiffPanels, bViewsAreLocked);
-
-	RemoteDiffResultsListData = &(Item->RemoteDifferences);
-	LocalDiffResultsListData = &(Item->LocalDifferences);
-
-	const auto SetupDiffList = [this]( TArray< TSharedPtr<FDiffSingleResult> > *Differences, TWeakPtr < SListView< TSharedPtr< struct FDiffSingleResult> > >& OutWidget)
-	{
-		return SAssignNew(OutWidget, SListView< TSharedPtr<FDiffSingleResult> >)
-			.ItemHeight(24)
-			.ListItemsSource(Differences)
-			.OnGenerateRow(SListView< TSharedPtr<FDiffSingleResult> >::FOnGenerateRow::CreateStatic(
-			[](TSharedPtr<FDiffSingleResult> ParamItem, const TSharedRef<STableViewBase>& OwnerTable) -> TSharedRef<ITableRow>
-		{
-			FDiffResultItem WidgetGenerator(*ParamItem);
-			return SNew(STableRow< TSharedPtr<FDiffSingleResult> >, OwnerTable)
-				.Content()
-				[
-					WidgetGenerator.GenerateWidget()
-				];
-		}
-		))
-			.SelectionMode(ESelectionMode::Single)
-			.OnSelectionChanged(this, &SMergeGraphView::OnDiffListSelectionChanged);
-	};
-
-	RemoteDiffResultsWidget->SetContent(
-		SetupDiffList( &Item->RemoteDifferences, RemoteDiffResultList)
-	);
-
-	LocalDiffResultsWidget->SetContent(
-		SetupDiffList(&Item->LocalDifferences, LocalDiffResultList)
-	);
 }
 
-void SMergeGraphView::OnDiffListSelectionChanged(TSharedPtr<FDiffSingleResult> Item, ESelectInfo::Type SelectionType)
+void SMergeGraphView::HighlightEntry(const struct FMergeGraphRowEntry& Conflict)
 {
-	if (!Item.IsValid())
-	{
-		return;
-	}
+	FocusGraph(Conflict.GraphName);
 
-	// Clear graph selection:
-	for (auto& Panel : DiffPanels)
+	const auto FocusPinOrNode = [this]( UEdGraphPin* Pin, UEdGraphNode* Node )
 	{
-		auto GraphEdtiorPtr = Panel.GraphEditor.Pin();
-		if (GraphEdtiorPtr.IsValid())
+		if (Pin)
 		{
-			GraphEdtiorPtr->ClearSelectionSet();
+			// then look for the diff panel and focus on the change:
+			GetDiffPanelForNode(*Pin->GetOwningNode(), DiffPanels).FocusDiff(*Pin);
 		}
-	}
-
-	//focus the graph onto the diff that was clicked on
-	if (Item->Pin1)
-	{
-		GetDiffPanelForNode(*Item->Pin1->GetOwningNode(), DiffPanels).FocusDiff(*Item->Pin1);
-
-		if (Item->Pin2)
+		else if (Node)
 		{
-			GetDiffPanelForNode(*Item->Pin2->GetOwningNode(), DiffPanels).FocusDiff(*Item->Pin2);
+			GetDiffPanelForNode(*Node, DiffPanels).FocusDiff(*Node);
 		}
-	}
-	else if (Item->Node1)
-	{
-		GetDiffPanelForNode(*Item->Node1, DiffPanels).FocusDiff(*Item->Node1);
+	};
 
-		if (Item->Node2)
-		{
-			GetDiffPanelForNode(*Item->Node2, DiffPanels).FocusDiff(*Item->Node2);
-		}
-	}
+	// highlight the change made to the remote graph:
+	FocusPinOrNode(Conflict.RemotePin, Conflict.RemoteNode);
+	FocusPinOrNode(Conflict.LocalPin, Conflict.LocalNode);
+	FocusPinOrNode(Conflict.BasePin, Conflict.BaseNode);
 }
 
 TSharedRef<SDockTab> SMergeGraphView::CreateGraphDiffViews(const FSpawnTabArgs& Args)
diff --git a/Engine/Source/Developer/Merge/Private/SMergeGraphView.h b/Engine/Source/Developer/Merge/Private/SMergeGraphView.h
--- a/Engine/Source/Developer/Merge/Private/SMergeGraphView.h
+++ b/Engine/Source/Developer/Merge/Private/SMergeGraphView.h
@@ -2,36 +2,26 @@
 #pragma once
 
 #include "BlueprintMergeData.h"
-#include "IDiffControl.h"
 #include "SBlueprintDiff.h"
 
 class SMergeGraphView : public SCompoundWidget
-						, public IDiffControl
-						, public IMergeControl
 {
 public:
 	SLATE_BEGIN_ARGS(SMergeGraphView){}
 	SLATE_END_ARGS()
 
-	void Construct(const FArguments InArgs, const FBlueprintMergeData& InData );
+	void Construct(	const FArguments InArgs
+					, const FBlueprintMergeData& InData
+					, FOnMergeNodeSelected SelectionCallback
+					, TArray< TSharedPtr<FBlueprintDifferenceTreeEntry> >& OutTreeEntries
+					, TArray< TSharedPtr<FBlueprintDifferenceTreeEntry> >& OutRealDifferences
+					, TArray< TSharedPtr<FBlueprintDifferenceTreeEntry> >& OutConflicts);
 private:
-	/** Implementation of IDiffControl: */
-	void NextDiff() override;
-	void PrevDiff() override;
-	bool HasNextDifference() const override;
-	bool HasPrevDifference() const override;
-
-	/** Implementation of IMergeControl */
-	void HighlightNextConflict() override;
-	void HighlightPrevConflict() override;
-	bool HasNextConflict() const override;
-	bool HasPrevConflict() const override;
-
 	/** Helper functions and event handlers: */
-	void HighlightConflict(const struct FGraphMergeConflict& Conflict);
+	void FocusGraph(FName GraphName);
+	void HighlightEntry(const struct FMergeGraphRowEntry& Conflict);
 	bool HasNoDifferences() const;
-	void OnGraphListSelectionChanged(TSharedPtr<struct FMergeGraphRowEntry> Item, ESelectInfo::Type SelectionType);
-	void OnDiffListSelectionChanged(TSharedPtr<struct FDiffSingleResult> Item, ESelectInfo::Type SelectionType);
+	
 	TSharedRef<SDockTab> CreateGraphDiffViews(const FSpawnTabArgs& Args);
 	TSharedRef<SDockTab> CreateMyBlueprintsViews(const FSpawnTabArgs& Args);
 
@@ -45,33 +35,7 @@ private:
 	TArray< FDiffPanel > DiffPanels;
 	FBlueprintMergeData Data;
 
-	TArray< TSharedPtr< struct FMergeGraphRowEntry> > DifferencesFromBase;
-
-	TSharedPtr<class SBox> RemoteDiffResultsWidget;
-	TSharedPtr<class SBox> LocalDiffResultsWidget;
-
-	TWeakPtr < SListView< TSharedPtr< struct FDiffSingleResult> > > RemoteDiffResultList;
-	TWeakPtr < SListView< TSharedPtr< struct FDiffSingleResult> > > LocalDiffResultList;
-
-	TArray< TSharedPtr<struct FGraphMergeConflict> > MergeConflicts;
-	int CurrentMergeConflict;
-
-	/** 
-	 * Unfortunately this innocuous member requires some explanation:
-	 * 	This list is owned by the FMergeGraphRowEntry and displayed by DiffResultList, but we have
-	 * 	to keep track of it because SListView doesn't expose the list it's displaying, and we
-	 * 	don't have any need to interact with the ListView that displays the DifferencesFromBase
-	 * 	so we can't access the current FMergeGraphRowEntry. Basically we have three options:
-	 * 	
-	 * 	1. Cache the list view when a new FMergeGraphRowEntry is selected
-	 * 	2. Cache a reference to the ListView that displays DifferencesFromBase, solely so we can
-	 * 		access the currently selected item and dig out the diffs associated with that graph
-	 * 	3. Expose the list being displayed by the listview (requires changing slate)
-	 * 
-	 * 	I have chosen option 1.
-	 */
-	TArray< TSharedPtr< struct FDiffSingleResult> > const* RemoteDiffResultsListData;
-	TArray< TSharedPtr< struct FDiffSingleResult> > const* LocalDiffResultsListData;
+	TSharedPtr< TArray< struct FMergeGraphEntry > > Differences;
 
 	bool bViewsAreLocked;
 
diff --git a/Engine/Source/Developer/Merge/Private/SMergeTreeView.cpp b/Engine/Source/Developer/Merge/Private/SMergeTreeView.cpp
--- a/Engine/Source/Developer/Merge/Private/SMergeTreeView.cpp
+++ b/Engine/Source/Developer/Merge/Private/SMergeTreeView.cpp
@@ -4,7 +4,12 @@
 
 #include "SMergeTreeView.h"
 
-void SMergeTreeView::Construct(const FArguments InArgs, const FBlueprintMergeData& InData)
+void SMergeTreeView::Construct(const FArguments InArgs
+	, const FBlueprintMergeData& InData
+	, FOnMergeNodeSelected SelectionCallback
+	, TArray< TSharedPtr<FBlueprintDifferenceTreeEntry> >& OutTreeEntries
+	, TArray< TSharedPtr<FBlueprintDifferenceTreeEntry> >& OutRealDifferences
+	, TArray< TSharedPtr<FBlueprintDifferenceTreeEntry> >& OutConflicts)
 {
 	Data = InData;
 	CurrentDifference = -1;
@@ -42,33 +47,13 @@ void SMergeTreeView::Construct(const FArguments InArgs, const FBlueprintMergeDat
 	DifferingProperties = RemoteDifferingProperties;
 	DifferingProperties.Entries.Append( LocalDifferingProperties.Entries );
 
-	// Remove duplicates and detect conflicts, there is opportunity for performance improvement here, but I doubt this represents a 
-	// significant cost:
-	for( int32 Iter = 0; Iter != DifferingProperties.Entries.Num(); ++Iter )
+	struct FSCSDiffPair
 	{
-		const FSCSDiffEntry& IterEntry = DifferingProperties.Entries[Iter];
-		for( int32 Jter = Iter + 1; Jter < DifferingProperties.Entries.Num(); ++Jter )
-		{
-			const FSCSDiffEntry& JterEntry = DifferingProperties.Entries[Jter];
-			// If the change is to a specific property, then it's only a conflict if the other change
-			// is also to that same specific property. All 'whole object' (non property specific changes)
-			// are conflicting:
-			if( IterEntry.TreeIdentifier == JterEntry.TreeIdentifier &&
-				( ( IterEntry.PropertyIdentifier == JterEntry.PropertyIdentifier &&
-				  IterEntry.PropertyIdentifier != FPropertySoftPath() ) 
-				||
-				( ( IterEntry.PropertyIdentifier == FPropertySoftPath() ||
-					JterEntry.PropertyIdentifier == FPropertySoftPath() ) ) ))
-			{
-				// for now just flag it as a conflict:
-				MergeConflicts.Entries.Push( IterEntry );
-
-				// and remove the redundant difference:
-				DifferingProperties.Entries.RemoveAtSwap(Jter--);
-			}
-		}
-	}
-
+		const FSCSDiffEntry* Local;
+		const FSCSDiffEntry* Remote;
+	};
+	TArray < FSCSDiffPair > ConflictingDifferences;
+	
 	/* This predicate sorts the list of differing properties so that those that are 'earlier' in the tree appear first.
 		For example, if we get the following two trees back:
 			
@@ -80,7 +65,7 @@ void SMergeTreeView::Construct(const FArguments InArgs, const FBlueprintMergeDat
 			D added at position (4, 2, 1)
 
 			the resulting list will be
-			C, B, aand D last: */
+			[C, B, D]: */
 	const auto SortTreePredicate = []( const FSCSDiffEntry& A, const FSCSDiffEntry& B )
 	{
 		int32 Idx = 0;
@@ -119,9 +104,139 @@ void SMergeTreeView::Construct(const FArguments InArgs, const FBlueprintMergeDat
 		return A.DiffType < B.DiffType;
 	};
 
+	RemoteDifferingProperties.Entries.Sort(SortTreePredicate);
+	LocalDifferingProperties.Entries.Sort(SortTreePredicate);
+
+	const FText RemoteLabel = NSLOCTEXT("SMergeTreeView", "RemoteLabel", "Remote");
+	const FText LocalLabel = NSLOCTEXT("SMergeTreeView", "LocalLabel", "Local");
+
+	struct FSCSMergeEntry
+	{
+		FText Label;
+		FSCSIdentifier Identifier;
+		FPropertySoftPath PropertyIdentifier;
+		bool bConflicted;
+	};
+
+	TArray<FSCSMergeEntry> Entries;
+	bool bAnyConflict = false;
+	int RemoteIter = 0;
+	int LocalIter = 0;
+	while(	RemoteIter != RemoteDifferingProperties.Entries.Num() ||
+			LocalIter != LocalDifferingProperties.Entries.Num() )
+	{
+		if (RemoteIter != RemoteDifferingProperties.Entries.Num() &&
+			LocalIter != LocalDifferingProperties.Entries.Num())
+		{
+			// check for conflicts:
+			const FSCSDiffEntry& Remote = RemoteDifferingProperties.Entries[RemoteIter];
+			const FSCSDiffEntry& Local = LocalDifferingProperties.Entries[LocalIter];
+			if( Remote.TreeIdentifier == Local.TreeIdentifier)
+			{
+				bool bConflicting = true;
+
+				if( Remote.DiffType == ETreeDiffType::NODE_PROPERTY_CHANGED &&
+					Local.DiffType == ETreeDiffType::NODE_PROPERTY_CHANGED )
+				{
+					// conflict only if property changed is the same:
+					bConflicting = Remote.PropertyDiff.Identifier == Local.PropertyDiff.Identifier;
+				}
+
+				if( bConflicting )
+				{
+					bAnyConflict = true;
+
+					// create a tree entry that describes both the local and remote change..
+					Entries.Push( {FText::Format( NSLOCTEXT("SMergeTreeView", "ConflictIdentifier", "CONFLICT: {0} conflicts with {1}"), DiffViewUtils::SCSDiffMessage( Remote, RemoteLabel ), DiffViewUtils::SCSDiffMessage( Local, LocalLabel ) ),
+									Remote.TreeIdentifier,
+									Remote.DiffType == ETreeDiffType::NODE_PROPERTY_CHANGED ? Remote.PropertyDiff.Identifier : Local.PropertyDiff.Identifier,
+									true });
+
+					++RemoteIter;
+					++LocalIter;
+					continue;
+				}
+			}
+		}
+
+		// no possibility of conflict, advance the entry that has a 'lower' tree identifier, keeping in mind that tree identifier
+		// may be equal, and that in that case we need to use the property identifier as a tiebreaker:
+		const FSCSDiffEntry* Remote = RemoteIter != RemoteDifferingProperties.Entries.Num() ? &RemoteDifferingProperties.Entries[RemoteIter] : nullptr;
+		const FSCSDiffEntry* Local = LocalIter != LocalDifferingProperties.Entries.Num() ? &LocalDifferingProperties.Entries[LocalIter] : nullptr;
+
+		if( Local && ( !Remote || SortTreePredicate( *Local, *Remote ) ) )
+		{
+			Entries.Push( {	DiffViewUtils::SCSDiffMessage(*Local, LocalLabel),
+							Local->TreeIdentifier,
+							Local->PropertyDiff.Identifier,
+							false } );
+
+			++LocalIter;
+		}
+		else
+		{
+			Entries.Push( { DiffViewUtils::SCSDiffMessage(*Remote, RemoteLabel),
+							Remote->TreeIdentifier,
+							Remote->PropertyDiff.Identifier,
+							false });
+
+			++RemoteIter;
+		}
+	}
+
+	const auto CreateSCSMergeWidget = [](FSCSMergeEntry Entry) -> TSharedRef<SWidget>
+	{
+		return SNew(STextBlock)
+			.Text(Entry.Label)
+			.ColorAndOpacity(Entry.bConflicted ? DiffViewUtils::Conflicting() : DiffViewUtils::Differs());
+	};
+
+	const auto FocusSCSDifferenceEntry = [](FSCSMergeEntry Entry, SMergeTreeView* Parent, FOnMergeNodeSelected SelectionCallback)
+	{
+		SelectionCallback.ExecuteIfBound();
+		Parent->HighlightDifference(Entry.Identifier, Entry.PropertyIdentifier);
+	};
+	
+	TArray< TSharedPtr<FBlueprintDifferenceTreeEntry> > Children;
+	for ( const auto& Difference : Entries )
+	{
+		auto Entry = TSharedPtr<FBlueprintDifferenceTreeEntry>(
+			new FBlueprintDifferenceTreeEntry(
+				FOnDiffEntryFocused::CreateStatic(FocusSCSDifferenceEntry, Difference, this, SelectionCallback)
+				, FGenerateDiffEntryWidget::CreateStatic(CreateSCSMergeWidget, Difference)
+				, TArray< TSharedPtr<FBlueprintDifferenceTreeEntry> >()
+			)
+		);
+		Children.Push(Entry);
+		OutRealDifferences.Push(Entry);
+		if( Difference.bConflicted )
+		{
+			OutConflicts.Push(Entry);
+		}
+	}
 
 	DifferingProperties.Entries.Sort(SortTreePredicate);
 
+	const auto ForwardSelection = [](FOnMergeNodeSelected SelectionCallback)
+	{
+		// This allows the owning control to focus the correct tab (or do whatever else it likes):
+		SelectionCallback.ExecuteIfBound();
+	};
+
+	const bool bHasDiffferences = Children.Num() != 0;
+	if( !bHasDiffferences )
+	{
+		Children.Push( FBlueprintDifferenceTreeEntry::NoDifferencesEntry() );
+	}
+
+	TSharedPtr<FBlueprintDifferenceTreeEntry> Category = FBlueprintDifferenceTreeEntry::CreateComponentsCategoryEntryForMerge(
+		FOnDiffEntryFocused::CreateStatic(ForwardSelection, SelectionCallback)
+		, Children
+		, RemoteDifferingProperties.Entries.Num() != 0
+		, LocalDifferingProperties.Entries.Num() != 0
+		, bAnyConflict);
+	OutTreeEntries.Push(Category);
+
 	ChildSlot[
 		SNew(SSplitter)
 			+ SSplitter::Slot()
@@ -139,75 +254,11 @@ void SMergeTreeView::Construct(const FArguments InArgs, const FBlueprintMergeDat
 	];
 }
 
-void SMergeTreeView::NextDiff()
-{
-	++CurrentDifference;
-	HighlightCurrentDifference();
-}
-
-void SMergeTreeView::PrevDiff()
-{
-	--CurrentDifference;
-	HighlightCurrentDifference();
-}
-
-bool SMergeTreeView::HasNextDifference() const
-{
-	return DifferingProperties.Entries.IsValidIndex(CurrentDifference + 1);
-}
-
-bool SMergeTreeView::HasPrevDifference() const
-{
-	return DifferingProperties.Entries.IsValidIndex(CurrentDifference - 1);
-}
-
-void SMergeTreeView::HighlightNextConflict()
-{
-	if (CurrentMergeConflict + 1 < MergeConflicts.Entries.Num())
-	{
-		++CurrentMergeConflict;
-	}
-	HighlightCurrentConflict();
-}
-
-void SMergeTreeView::HighlightPrevConflict()
-{
-	if (CurrentMergeConflict - 1 >= 0)
-	{
-		--CurrentMergeConflict;
-	}
-	HighlightCurrentConflict();
-}
-
-bool SMergeTreeView::HasNextConflict() const
-{
-	// return true if we have one conflict so that users can reselect the conflict if they desire. If we 
-	// return false when we already have selected this one and only conflict then there will be no way
-	// to reselect it if the user wants to.
-	return MergeConflicts.Entries.Num() != 0 && (MergeConflicts.Entries.Num() == 1 || CurrentMergeConflict + 1 < MergeConflicts.Entries.Num());
-}
-
-bool SMergeTreeView::HasPrevConflict() const
-{
-	// note in HasNextConflict applies here as well.
-	return MergeConflicts.Entries.Num() == 1 || CurrentMergeConflict > 0;
-}
-
-void SMergeTreeView::HighlightCurrentConflict()
-{
-	const FSCSDiffEntry& Conflict = MergeConflicts.Entries[CurrentMergeConflict];
-	for (auto& View : SCSViews)
-	{
-		View.HighlightProperty(Conflict.TreeIdentifier.Name, Conflict.PropertyIdentifier);
-	}
-}
-
-void SMergeTreeView::HighlightCurrentDifference()
+void SMergeTreeView::HighlightDifference(FSCSIdentifier TreeIdentifier, FPropertySoftPath Property)
 {
-	const FSCSDiffEntry& Difference = DifferingProperties.Entries[CurrentDifference];
 	for (auto& View : SCSViews)
 	{
-		View.HighlightProperty(Difference.TreeIdentifier.Name, Difference.PropertyIdentifier);
+		View.HighlightProperty(TreeIdentifier.Name, Property);
 	}
 }
 
diff --git a/Engine/Source/Developer/Merge/Private/SMergeTreeView.h b/Engine/Source/Developer/Merge/Private/SMergeTreeView.h
--- a/Engine/Source/Developer/Merge/Private/SMergeTreeView.h
+++ b/Engine/Source/Developer/Merge/Private/SMergeTreeView.h
@@ -2,12 +2,9 @@
 #pragma once
 
 #include "BlueprintMergeData.h"
-#include "IDiffControl.h"
 #include "SCSDiff.h"
 
 class SMergeTreeView	: public SCompoundWidget
-						, public IDiffControl
-						, public IMergeControl
 {
 public:
 	virtual ~SMergeTreeView() {}
@@ -16,22 +13,15 @@ public:
 	{}
 	SLATE_END_ARGS()
 
-	void Construct(const FArguments InArgs, const FBlueprintMergeData& InData);
+	void Construct(	const FArguments InArgs
+					, const FBlueprintMergeData& InData
+					, FOnMergeNodeSelected SelectionCallback
+					, TArray< TSharedPtr<FBlueprintDifferenceTreeEntry> >& OutTreeEntries
+					, TArray< TSharedPtr<FBlueprintDifferenceTreeEntry> >& OutRealDifferences
+					, TArray< TSharedPtr<FBlueprintDifferenceTreeEntry> >& OutConflicts);
 private:
-	/** Implementation of IDiffControl: */
-	void NextDiff() override;
-	void PrevDiff() override;
-	bool HasNextDifference() const override;
-	bool HasPrevDifference() const override;
-
-	/** Implementation of IMergeControl: */
-	void HighlightNextConflict() override;
-	void HighlightPrevConflict() override;
-	bool HasNextConflict() const override;
-	bool HasPrevConflict() const override;
-
-	void HighlightCurrentConflict();
-	void HighlightCurrentDifference();
+	void HighlightDifference( FSCSIdentifier VarName, FPropertySoftPath Property );
+
 	FSCSDiff& GetRemoteView();
 	FSCSDiff& GetBaseView();
 	FSCSDiff& GetLocalView();
diff --git a/Engine/Source/Editor/EditorStyle/Private/SlateEditorStyle.cpp b/Engine/Source/Editor/EditorStyle/Private/SlateEditorStyle.cpp
--- a/Engine/Source/Editor/EditorStyle/Private/SlateEditorStyle.cpp
+++ b/Engine/Source/Editor/EditorStyle/Private/SlateEditorStyle.cpp
@@ -4639,12 +4639,18 @@ void FSlateEditorStyle::FStyle::SetupPersonaStyle()
 
 	}
 
-	//Blueprint dif
+	//Blueprint Diff
 	{
 		Set( "BlueprintDif.HasGraph", new IMAGE_BRUSH( "/Icons/blueprint_Dif_has_graph_8x", Icon8x8 )  );
 		Set( "BlueprintDif.MissingGraph", new IMAGE_BRUSH( "/Icons/blueprint_Dif_missing_graph_8x", Icon8x8 )  );
 		Set( "BlueprintDif.NextDiff", new IMAGE_BRUSH( "/Icons/diff_next_40x", Icon16x16 )  );
 		Set( "BlueprintDif.PrevDiff", new IMAGE_BRUSH( "/Icons/diff_prev_40x", Icon16x16 )  );
+
+		Set( "BlueprintDif.ItalicText", 
+			FTextBlockStyle(NormalText)
+			.SetFont(TTF_FONT("Fonts/Roboto-Italic", 10))
+			.SetColorAndOpacity(FLinearColor(.7f, .7f, .7f))
+		);
 	}
 
 	//Blueprint Merge
diff --git a/Engine/Source/Editor/GraphEditor/Private/GraphDiffControl.cpp b/Engine/Source/Editor/GraphEditor/Private/GraphDiffControl.cpp
--- a/Engine/Source/Editor/GraphEditor/Private/GraphDiffControl.cpp
+++ b/Engine/Source/Editor/GraphEditor/Private/GraphDiffControl.cpp
@@ -586,6 +586,13 @@ bool FGraphDiffControl::DiffGraphs(UEdGraph* const LhsGraph, UEdGraph* const Rhs
 		}
 	}
 
+	// storing the graph name for all diff entries:
+	FName GraphName = LhsGraph ? LhsGraph->GetFName() : LhsGraph->GetFName();
+	for( auto& Entry : DiffsOut )
+	{
+		Entry.OwningGraph = GraphName;
+	}
+
 	return bFoundDifferences;
 }
 
diff --git a/Engine/Source/Editor/Kismet/Private/DiffUtils.cpp b/Engine/Source/Editor/Kismet/Private/DiffUtils.cpp
--- a/Engine/Source/Editor/Kismet/Private/DiffUtils.cpp
+++ b/Engine/Source/Editor/Kismet/Private/DiffUtils.cpp
@@ -4,9 +4,10 @@
 
 #include "DiffUtils.h"
 #include "EditorCategoryUtils.h"
-#include "ObjectEditorUtils.h"
 #include "Engine/SCS_Node.h"
 #include "Engine/SimpleConstructionScript.h"
+#include "IAssetTypeActions.h"
+#include "ObjectEditorUtils.h"
 
 static const UProperty* Resolve( const UStruct* Class, FName PropertyName )
 {
@@ -74,16 +75,24 @@ const UObject* DiffUtils::GetCDO(const UBlueprint* ForBlueprint)
 	return ForBlueprint->GeneratedClass->ClassDefaultObject;
 }
 
-void DiffUtils::CompareUnrelatedObjects(const UObject* A, UObject const* B, TArray<FPropertySoftPath>& OutDifferingProperties)
+void DiffUtils::CompareUnrelatedObjects(const UObject* A, const UObject* B, TArray<FSingleObjectDiffEntry>& OutDifferingProperties)
 {
 	FPropertySoftPathSet PropertiesInA = GetPropertyNameSet(A);
 	FPropertySoftPathSet PropertiesInB = GetPropertyNameSet(B);
 
 	// any properties in A that aren't in B are differing:
-	OutDifferingProperties.Append(PropertiesInA.Difference(PropertiesInB).Array());
+	auto AddedToA = PropertiesInA.Difference(PropertiesInB).Array();
+	for( const auto& Entry : AddedToA )
+	{
+		OutDifferingProperties.Push(FSingleObjectDiffEntry( Entry, EPropertyDiffType::PropertyAddedToA ));
+	}
 
 	// and the converse:
-	OutDifferingProperties.Append(PropertiesInB.Difference(PropertiesInA).Array());
+	auto AddedToB = PropertiesInB.Difference(PropertiesInA).Array();
+	for (const auto& Entry : AddedToB)
+	{
+		OutDifferingProperties.Push(FSingleObjectDiffEntry( Entry, EPropertyDiffType::PropertyAddedToB ));
+	}
 
 	// for properties in common, dig out the uproperties and determine if they're identical:
 	if (A && B)
@@ -97,7 +106,7 @@ void DiffUtils::CompareUnrelatedObjects(const UObject* A, UObject const* B, TArr
 			check(AProp != FResolvedProperty() && BProp != FResolvedProperty());
 			if (!DiffUtils::Identical(AProp, BProp))
 			{
-				OutDifferingProperties.Add(PropertyName);
+				OutDifferingProperties.Push(FSingleObjectDiffEntry( PropertyName, EPropertyDiffType::PropertyValueChanged ));
 			}
 		}
 	}
@@ -123,8 +132,8 @@ void DiffUtils::CompareUnrelatedSCS(const UBlueprint* Old, const TArray< FSCSRes
 
 		if (NewEntry != nullptr)
 		{
-			// @todo doc: did properties change?
-			TArray<FPropertySoftPath> DifferingProperties;
+			// did a property change?
+			TArray<FSingleObjectDiffEntry> DifferingProperties;
 			DiffUtils::CompareUnrelatedObjects(OldNode.Object, NewEntry->Object, DifferingProperties);
 			for (const auto& Property : DifferingProperties)
 			{
@@ -133,14 +142,18 @@ void DiffUtils::CompareUnrelatedSCS(const UBlueprint* Old, const TArray< FSCSRes
 			}
 
 			// did it move?
-
+			if( NewEntry->Identifier.TreeLocation != OldNode.Identifier.TreeLocation )
+			{
+				FSCSDiffEntry Diff = { OldNode.Identifier, ETreeDiffType::NODE_MOVED, FSingleObjectDiffEntry() };
+				OutDifferingEntries.Entries.Push(Diff);
+			}
 
 			// no change! Do nothing.
 		}
 		else
 		{
 			// not found in the new data, must have been deleted:
-			FSCSDiffEntry Entry = { OldNode.Identifier, ETreeDiffType::NODE_REMOVED, FPropertySoftPath() };
+			FSCSDiffEntry Entry = { OldNode.Identifier, ETreeDiffType::NODE_REMOVED, FSingleObjectDiffEntry() };
 			OutDifferingEntries.Entries.Push( Entry );
 		}
 	}
@@ -151,12 +164,10 @@ void DiffUtils::CompareUnrelatedSCS(const UBlueprint* Old, const TArray< FSCSRes
 
 		if (OldEntry == nullptr)
 		{
-			FSCSDiffEntry Entry = { NewNode.Identifier, ETreeDiffType::NODE_ADDED, FPropertySoftPath() };
+			FSCSDiffEntry Entry = { NewNode.Identifier, ETreeDiffType::NODE_ADDED, FSingleObjectDiffEntry() };
 			OutDifferingEntries.Entries.Push( Entry );
 		}
 	}
-
-
 }
 
 bool DiffUtils::Identical(const FResolvedProperty& AProp, const FResolvedProperty& BProp)
@@ -211,3 +222,339 @@ TArray<FPropertyPath> DiffUtils::ResolveAll(const UObject* Object, const TArray<
 	}
 	return Ret;
 }
+
+TArray<FPropertyPath> DiffUtils::ResolveAll(const UObject* Object, const TArray<FSingleObjectDiffEntry>& InDifferences)
+{
+	TArray< FPropertyPath > Ret;
+	for (const auto& Difference : InDifferences)
+	{
+		Ret.Push(Difference.Identifier.ResolvePath(Object));
+	}
+	return Ret;
+}
+
+TSharedPtr<FBlueprintDifferenceTreeEntry> FBlueprintDifferenceTreeEntry::NoDifferencesEntry()
+{
+	// This just generates a widget that tells the user that no differences were detected. Without this
+	// the treeview displaying differences is confusing when no differences are present because it is not obvious
+	// that the control is a treeview (a treeview with no children looks like a listview).
+	const auto GenerateWidget = []() -> TSharedRef<SWidget>
+	{
+		return SNew(STextBlock)
+			.ColorAndOpacity(FLinearColor::FLinearColor(.7f, .7f, .7f))
+			.TextStyle(FEditorStyle::Get(), TEXT("BlueprintDif.ItalicText"))
+			.Text(NSLOCTEXT("FBlueprintDifferenceTreeEntry", "NoDifferencesLabel", "No differences detected..."));
+	};
+
+	return TSharedPtr<FBlueprintDifferenceTreeEntry>(new FBlueprintDifferenceTreeEntry(
+		FOnDiffEntryFocused()
+		, FGenerateDiffEntryWidget::CreateStatic(GenerateWidget)
+		, TArray< TSharedPtr<FBlueprintDifferenceTreeEntry> >()
+	) );
+}
+
+TSharedPtr<FBlueprintDifferenceTreeEntry> FBlueprintDifferenceTreeEntry::CreateDefaultsCategoryEntry(FOnDiffEntryFocused FocusCallback, const TArray< TSharedPtr<FBlueprintDifferenceTreeEntry> >& Children, bool bHasDifferences)
+{
+	const auto CreateDefaultsRootEntry = [](FLinearColor Color) -> TSharedRef<SWidget>
+	{
+		return SNew(STextBlock)
+			.ToolTipText(NSLOCTEXT("FBlueprintDifferenceTreeEntry", "DefaultsTooltip", "The list of changes made in the Defaults panel"))
+			.ColorAndOpacity(Color)
+			.Text(NSLOCTEXT("FBlueprintDifferenceTreeEntry", "DefaultsLabel", "Defaults"));
+	};
+
+	return TSharedPtr<FBlueprintDifferenceTreeEntry>(new FBlueprintDifferenceTreeEntry(
+		FocusCallback
+		, FGenerateDiffEntryWidget::CreateStatic(CreateDefaultsRootEntry, DiffViewUtils::LookupColor(bHasDifferences) )
+		, Children
+	));
+}
+
+TSharedPtr<FBlueprintDifferenceTreeEntry> FBlueprintDifferenceTreeEntry::CreateDefaultsCategoryEntryForMerge(FOnDiffEntryFocused FocusCallback, const TArray< TSharedPtr<FBlueprintDifferenceTreeEntry> >& Children, bool bHasRemoteDifferences, bool bHasLocalDifferences, bool bHasConflicts)
+{
+	const auto CreateDefaultsRootEntry = [](bool bHasRemoteDifferences, bool bHasLocalDifferences, bool bHasConflicts) -> TSharedRef<SWidget>
+	{
+		const FLinearColor BaseColor = DiffViewUtils::LookupColor(bHasRemoteDifferences || bHasLocalDifferences, bHasConflicts);
+		return SNew(SHorizontalBox)
+			+ SHorizontalBox::Slot()
+			[
+				SNew(STextBlock)
+				.ToolTipText(NSLOCTEXT("FBlueprintDifferenceTreeEntry", "DefaultsTooltip", "The list of changes made in the Defaults panel"))
+				.ColorAndOpacity(BaseColor)
+				.Text(NSLOCTEXT("FBlueprintDifferenceTreeEntry", "DefaultsLabel", "Defaults"))
+			]
+			+ DiffViewUtils::Box(true, DiffViewUtils::LookupColor(bHasRemoteDifferences, bHasConflicts))
+			+ DiffViewUtils::Box(true, BaseColor)
+			+ DiffViewUtils::Box(true, DiffViewUtils::LookupColor(bHasLocalDifferences, bHasConflicts));
+	};
+
+	return TSharedPtr<FBlueprintDifferenceTreeEntry>(new FBlueprintDifferenceTreeEntry(
+		FocusCallback
+		, FGenerateDiffEntryWidget::CreateStatic(CreateDefaultsRootEntry, bHasRemoteDifferences, bHasLocalDifferences, bHasConflicts)
+		, Children
+		));
+}
+
+TSharedPtr<FBlueprintDifferenceTreeEntry> FBlueprintDifferenceTreeEntry::CreateComponentsCategoryEntry(FOnDiffEntryFocused FocusCallback, const TArray< TSharedPtr<FBlueprintDifferenceTreeEntry> >& Children, bool bHasDifferences)
+{
+	const auto CreateComponentsRootEntry = [](FLinearColor Color) -> TSharedRef<SWidget>
+	{
+		return SNew(STextBlock)
+			.ToolTipText(NSLOCTEXT("FBlueprintDifferenceTreeEntry", "SCSTooltip", "The list of changes made in the Components panel"))
+			.ColorAndOpacity(Color)
+			.Text(NSLOCTEXT("FBlueprintDifferenceTreeEntry", "SCSLabel", "Components"));
+	};
+
+	return TSharedPtr<FBlueprintDifferenceTreeEntry>(new FBlueprintDifferenceTreeEntry(
+		FocusCallback
+		, FGenerateDiffEntryWidget::CreateStatic(CreateComponentsRootEntry, DiffViewUtils::LookupColor(bHasDifferences))
+		, Children
+		));
+}
+
+TSharedPtr<FBlueprintDifferenceTreeEntry> FBlueprintDifferenceTreeEntry::CreateComponentsCategoryEntryForMerge(FOnDiffEntryFocused FocusCallback, const TArray< TSharedPtr<FBlueprintDifferenceTreeEntry> >& Children, bool bHasRemoteDifferences, bool bHasLocalDifferences, bool bHasConflicts)
+{
+	const auto CreateComponentsRootEntry = [](bool bHasRemoteDifferences, bool bHasLocalDifferences, bool bHasConflicts) -> TSharedRef<SWidget>
+	{
+		const FLinearColor BaseColor = DiffViewUtils::LookupColor(bHasRemoteDifferences || bHasLocalDifferences, bHasConflicts);
+		return  SNew(SHorizontalBox)
+			+ SHorizontalBox::Slot()
+			[
+				SNew(STextBlock)
+				.ToolTipText(NSLOCTEXT("FBlueprintDifferenceTreeEntry", "SCSTooltip", "The list of changes made in the Components panel"))
+				.ColorAndOpacity(BaseColor)
+				.Text(NSLOCTEXT("FBlueprintDifferenceTreeEntry", "SCSLabel", "Components"))
+			]
+			+ DiffViewUtils::Box(true, DiffViewUtils::LookupColor(bHasRemoteDifferences, bHasConflicts))
+			+ DiffViewUtils::Box(true, BaseColor)
+			+ DiffViewUtils::Box(true, DiffViewUtils::LookupColor(bHasLocalDifferences, bHasConflicts));
+	};
+
+	return TSharedPtr<FBlueprintDifferenceTreeEntry>(new FBlueprintDifferenceTreeEntry(
+		FocusCallback
+		, FGenerateDiffEntryWidget::CreateStatic(CreateComponentsRootEntry, bHasRemoteDifferences, bHasLocalDifferences, bHasConflicts)
+		, Children
+	));
+}
+
+TSharedRef< STreeView<TSharedPtr< FBlueprintDifferenceTreeEntry > > > DiffTreeView::CreateTreeView(TArray< TSharedPtr<FBlueprintDifferenceTreeEntry> >* DifferencesList)
+{
+	const auto RowGenerator = [](TSharedPtr< FBlueprintDifferenceTreeEntry > Entry, const TSharedRef<STableViewBase>& Owner) -> TSharedRef< ITableRow >
+	{
+		return SNew(STableRow<TSharedPtr<FBlueprintDifferenceTreeEntry> >, Owner)
+			[
+				Entry->GenerateWidget.Execute()
+			];
+	};
+
+	const auto ChildrenAccessor = [](TSharedPtr<FBlueprintDifferenceTreeEntry> InTreeItem, TArray< TSharedPtr< FBlueprintDifferenceTreeEntry > >& OutChildren, TArray< TSharedPtr<FBlueprintDifferenceTreeEntry> >* MasterList)
+	{
+		OutChildren = InTreeItem->Children;
+	};
+
+	const auto Selector = [](TSharedPtr<FBlueprintDifferenceTreeEntry> InTreeItem, ESelectInfo::Type Type)
+	{
+		if (InTreeItem.IsValid())
+		{
+			InTreeItem->OnFocus.ExecuteIfBound();
+		}
+	};
+
+	return SNew(STreeView< TSharedPtr< FBlueprintDifferenceTreeEntry > >)
+		.OnGenerateRow(STreeView< TSharedPtr< FBlueprintDifferenceTreeEntry > >::FOnGenerateRow::CreateStatic(RowGenerator))
+		.OnGetChildren(STreeView< TSharedPtr< FBlueprintDifferenceTreeEntry > >::FOnGetChildren::CreateStatic(ChildrenAccessor, DifferencesList))
+		.OnSelectionChanged(STreeView< TSharedPtr< FBlueprintDifferenceTreeEntry > >::FOnSelectionChanged::CreateStatic(Selector))
+		.TreeItemsSource(DifferencesList);
+}
+
+int32 DiffTreeView::CurrentDifference(TSharedRef< STreeView<TSharedPtr< FBlueprintDifferenceTreeEntry > > > TreeView, const TArray< TSharedPtr<class FBlueprintDifferenceTreeEntry> >& Differences)
+{
+	auto SelectedItems = TreeView->GetSelectedItems();
+	if (SelectedItems.Num() == 0)
+	{
+		return INDEX_NONE;
+	}
+
+	for (int32 Iter = 0; Iter < SelectedItems.Num(); ++Iter)
+	{
+		int32 Index = Differences.Find(SelectedItems[Iter]);
+		if (Index != INDEX_NONE)
+		{
+			return Index;
+		}
+	}
+
+	return INDEX_NONE;
+}
+
+void DiffTreeView::HighlightNextDifference(TSharedRef< STreeView<TSharedPtr< FBlueprintDifferenceTreeEntry > > > TreeView, const TArray< TSharedPtr<class FBlueprintDifferenceTreeEntry> >& Differences, const TArray< TSharedPtr<class FBlueprintDifferenceTreeEntry> >& RootDifferences)
+{
+	int32 CurrentIndex = CurrentDifference(TreeView, Differences);
+
+	auto Next = Differences[CurrentIndex + 1];
+	// we have to manually expand our parent:
+	for (auto& Test : RootDifferences)
+	{
+		if (Test->Children.Contains(Next))
+		{
+			TreeView->SetItemExpansion(Test, true);
+			break;
+		}
+	}
+
+	TreeView->SetSelection(Next);
+	TreeView->RequestScrollIntoView(Next);
+}
+
+void DiffTreeView::HighlightPrevDifference(TSharedRef< STreeView<TSharedPtr< FBlueprintDifferenceTreeEntry > > > TreeView, const TArray< TSharedPtr<class FBlueprintDifferenceTreeEntry> >& Differences, const TArray< TSharedPtr<class FBlueprintDifferenceTreeEntry> >& RootDifferences)
+{
+	int32 CurrentIndex = CurrentDifference(TreeView, Differences);
+
+	auto Prev = Differences[CurrentIndex - 1];
+	// we have to manually expand our parent:
+	for (auto& Test : RootDifferences)
+	{
+		if (Test->Children.Contains(Prev))
+		{
+			TreeView->SetItemExpansion(Test, true);
+			break;
+		}
+	}
+
+	TreeView->SetSelection(Prev);
+	TreeView->RequestScrollIntoView(Prev);
+}
+
+bool DiffTreeView::HasNextDifference(TSharedRef< STreeView<TSharedPtr< FBlueprintDifferenceTreeEntry > > > TreeView, const TArray< TSharedPtr<class FBlueprintDifferenceTreeEntry> >& Differences)
+{
+	int32 CurrentIndex = CurrentDifference(TreeView, Differences);
+	return Differences.IsValidIndex(CurrentIndex + 1);
+}
+
+bool DiffTreeView::HasPrevDifference(TSharedRef< STreeView<TSharedPtr< FBlueprintDifferenceTreeEntry > > > TreeView, const TArray< TSharedPtr<class FBlueprintDifferenceTreeEntry> >& Differences)
+{
+	int32 CurrentIndex = CurrentDifference(TreeView, Differences);
+	return Differences.IsValidIndex(CurrentIndex - 1);
+}
+
+FLinearColor DiffViewUtils::LookupColor(bool bDiffers, bool bConflicts)
+{
+	if( bConflicts )
+	{
+		return DiffViewUtils::Conflicting();
+	}
+	else if( bDiffers )
+	{
+		return DiffViewUtils::Differs();
+	}
+	else
+	{
+		return DiffViewUtils::Identical();
+	}
+}
+
+FLinearColor DiffViewUtils::Differs()
+{
+	// yellow color
+	return FLinearColor(0.85f,0.71f,0.25f);
+}
+
+FLinearColor DiffViewUtils::Identical()
+{
+	return FLinearColor::White;
+}
+
+FLinearColor DiffViewUtils::Missing()
+{
+	// blue color
+	return FLinearColor(0.3f,0.3f,1.f);
+}
+
+FLinearColor DiffViewUtils::Conflicting()
+{
+	// red color
+	return FLinearColor(1.0f,0.2f,0.3f);
+}
+
+FText DiffViewUtils::PropertyDiffMessage(FSingleObjectDiffEntry Difference, FText ObjectName)
+{
+	FText Message;
+	FString PropertyName = Difference.Identifier.LastPropertyName().GetPlainNameString();
+	switch (Difference.DiffType)
+	{
+	case EPropertyDiffType::PropertyAddedToA:
+		Message = FText::Format(NSLOCTEXT("DiffViewUtils", "PropertyValueChange", "{0} removed from {1}"), FText::FromString(PropertyName), ObjectName);
+		break;
+	case EPropertyDiffType::PropertyAddedToB:
+		Message = FText::Format(NSLOCTEXT("DiffViewUtils", "PropertyValueChange", "{0} added to {1}"), FText::FromString(PropertyName), ObjectName);
+		break;
+	case EPropertyDiffType::PropertyValueChanged:
+		Message = FText::Format(NSLOCTEXT("DiffViewUtils", "PropertyValueChange", "{0} changed value in {1}"), FText::FromString(PropertyName), ObjectName);
+		break;
+	}
+	return Message;
+}
+
+FText DiffViewUtils::SCSDiffMessage(const FSCSDiffEntry& Difference, FText ObjectName)
+{
+	const FText NodeName = FText::FromName(Difference.TreeIdentifier.Name);
+	FText Text;
+	switch (Difference.DiffType)
+	{
+	case ETreeDiffType::NODE_ADDED:
+		Text = FText::Format(NSLOCTEXT("DiffViewUtils", "NodeAdded", "Added Node {0} to {1}"), NodeName, ObjectName);
+		break;
+	case ETreeDiffType::NODE_REMOVED:
+		Text = FText::Format(NSLOCTEXT("DiffViewUtils", "NodeAdded", "Removed Node {0} from {1}"), NodeName, ObjectName);
+		break;
+	case ETreeDiffType::NODE_PROPERTY_CHANGED:
+		Text = FText::Format(NSLOCTEXT("DiffViewUtils", "NodeAdded", "{0} on {2}"), DiffViewUtils::PropertyDiffMessage(Difference.PropertyDiff, NodeName), ObjectName);
+		break;
+	case ETreeDiffType::NODE_MOVED:
+		Text = FText::Format(NSLOCTEXT("DiffViewUtils", "NodeAdded", "Moved Node {0} in {1}"), NodeName, ObjectName);
+		break;
+	}
+	return Text;
+}
+
+FText DiffViewUtils::GetPanelLabel(const UBlueprint* Blueprint, const FRevisionInfo& Revision, FText Label )
+{
+	if( Revision.Revision >= 0 )
+	{
+		FText RevisionData = FText::Format(NSLOCTEXT("DiffViewUtils", "RevisionData", "Revision {0} - CL {1} - {2}")
+			, FText::AsNumber(Revision.Revision)
+			, FText::AsNumber(Revision.Changelist)
+			, FText::FromString(Revision.Date.ToString(TEXT("%m/%d/%Y"))));
+		return FText::Format( NSLOCTEXT("DiffViewUtils", "RevisionLabel", "{0}\n{1}\n{2}")
+			, Label
+			, FText::FromString( Blueprint->GetName() )
+			, RevisionData );
+	}
+	else
+	{
+		if( Blueprint )
+		{
+			return FText::Format( NSLOCTEXT("DiffViewUtils", "RevisionLabel", "{0}\n{1}\n{2}")
+				, Label
+				, FText::FromString( Blueprint->GetName() )
+				, NSLOCTEXT("DiffViewUtils", "LocalRevisionLabel", "Local Revision" ));
+		}
+
+		return NSLOCTEXT("DiffViewUtils", "NoBlueprint", "None" );
+	}
+}
+
+SHorizontalBox::FSlot& DiffViewUtils::Box(bool bIsPresent, FLinearColor Color)
+{
+	return SHorizontalBox::Slot()
+		.AutoWidth()
+		.HAlign(HAlign_Right)
+		.VAlign(VAlign_Center)
+		.MaxWidth(8.0f)
+		[
+			SNew(SImage)
+			.ColorAndOpacity(Color)
+			.Image(bIsPresent ? FEditorStyle::GetBrush("BlueprintDif.HasGraph") : FEditorStyle::GetBrush("BlueprintDif.MissingGraph"))
+		];
+};
+
diff --git a/Engine/Source/Editor/Kismet/Private/SBlueprintDiff.cpp b/Engine/Source/Editor/Kismet/Private/SBlueprintDiff.cpp
--- a/Engine/Source/Editor/Kismet/Private/SBlueprintDiff.cpp
+++ b/Engine/Source/Editor/Kismet/Private/SBlueprintDiff.cpp
@@ -14,7 +14,6 @@
 #include "Editor/UnrealEd/Public/EdGraphUtilities.h"
 #include "GraphDiffControl.h"
 #include "GraphEditor.h"
-#include "IDiffControl.h"
 #include "Kismet2/BlueprintEditorUtils.h"
 #include "SBlueprintDiff.h"
 #include "SDockTab.h"
@@ -27,11 +26,21 @@
 
 #define LOCTEXT_NAMESPACE "SBlueprintDif"
 
+class IDiffControl
+{
+public:
+	virtual ~IDiffControl() {}
+};
+
+FText RightRevision = LOCTEXT("OlderRevisionIdentifier", "Right Revision");
+
 typedef TMap< FName, const UProperty* > FNamePropertyMap;
 
 const FName DiffMyBluerpintTabId = FName(TEXT("DiffMyBluerpintTab"));
 const FName DiffGraphTabId = FName(TEXT("DiffGraphTab"));
 
+DECLARE_DELEGATE(FOnSCSDiffControlChanged);
+
 // Each difference in the tree will either be a tree node that is added in one Blueprint 
 // or a tree node and an FName of a property that has been added or edited in one Blueprint
 class FSCSDiffControl	: public TSharedFromThis< FSCSDiffControl >
@@ -41,6 +50,9 @@ public:
 	FSCSDiffControl(
 					const UBlueprint* InOldBlueprint
 					, const UBlueprint* InNewBlueprint
+					, TArray< TSharedPtr<FBlueprintDifferenceTreeEntry> >& OutTreeEntries
+					, TArray< TSharedPtr<FBlueprintDifferenceTreeEntry> >& OutRealDifferences
+					, FOnSCSDiffControlChanged SelectionCallback
 					);
 
 
@@ -49,15 +61,7 @@ public:
 
 	virtual ~FSCSDiffControl() {}
 private:
-	void NextDiff() override;
-	void PrevDiff() override;
-	bool HasNextDifference() const override;
-	bool HasPrevDifference() const override;
-
-	void HighlightCurrentDifference();
-
 	FSCSDiffRoot DifferingProperties;
-	int CurrentDifference;
 
 	FSCSDiff OldSCS;
 	FSCSDiff NewSCS;
@@ -66,127 +70,122 @@ private:
 FSCSDiffControl::FSCSDiffControl(
 		const UBlueprint* InOldBlueprint
 		, const UBlueprint* InNewBlueprint
+		, TArray< TSharedPtr<FBlueprintDifferenceTreeEntry> >& OutTreeEntries
+		, TArray< TSharedPtr<FBlueprintDifferenceTreeEntry> >& OutRealDifferences
+		, FOnSCSDiffControlChanged SelectionCallback
 	)
 	: DifferingProperties()
-	, CurrentDifference(-1)
 	, OldSCS( InOldBlueprint )
 	, NewSCS( InNewBlueprint )
 {
 	TArray< FSCSResolvedIdentifier > OldHierarchy = OldSCS.GetDisplayedHierarchy();
 	TArray< FSCSResolvedIdentifier > NewHierarchy = NewSCS.GetDisplayedHierarchy();
 	DiffUtils::CompareUnrelatedSCS(InOldBlueprint, OldHierarchy, InNewBlueprint, NewHierarchy, DifferingProperties);
-}
 
-void FSCSDiffControl::NextDiff()
-{
-	++CurrentDifference;
-	HighlightCurrentDifference();
-}
+	const auto FocusSCSDifferenceEntry = [](FSCSDiffEntry Entry, FOnSCSDiffControlChanged SelectionCallback, FSCSDiffControl* Owner)
+	{
+		SelectionCallback.ExecuteIfBound();
+		Owner->OldSCS.HighlightProperty(Entry.TreeIdentifier.Name, FPropertyPath());
+		Owner->NewSCS.HighlightProperty(Entry.TreeIdentifier.Name, FPropertyPath());
+	};
 
-void FSCSDiffControl::PrevDiff()
-{
-	--CurrentDifference;
-	HighlightCurrentDifference();
-}
+	const auto CreateSCSDifferenceWidget = [](FSCSDiffEntry Entry, FText ObjectName) -> TSharedRef<SWidget>
+	{
+		return SNew(STextBlock)
+			.Text(DiffViewUtils::SCSDiffMessage(Entry,ObjectName))
+			.ColorAndOpacity(DiffViewUtils::Differs());
+	};
 
-bool FSCSDiffControl::HasNextDifference() const
-{
-	return DifferingProperties.Entries.IsValidIndex(CurrentDifference + 1);
-}
+	TArray< TSharedPtr<FBlueprintDifferenceTreeEntry> > Children;
+	for( auto Difference : DifferingProperties.Entries )
+	{
+		auto Entry = TSharedPtr<FBlueprintDifferenceTreeEntry>(
+			new FBlueprintDifferenceTreeEntry(
+				FOnDiffEntryFocused::CreateStatic(FocusSCSDifferenceEntry, Difference, SelectionCallback, this)
+				, FGenerateDiffEntryWidget::CreateStatic(CreateSCSDifferenceWidget, Difference, RightRevision)
+				, TArray< TSharedPtr<FBlueprintDifferenceTreeEntry> >()
+			)
+		);
+		Children.Push(Entry);
+		OutRealDifferences.Push(Entry);
+	}
 
-bool FSCSDiffControl::HasPrevDifference() const
-{
-	return DifferingProperties.Entries.IsValidIndex(CurrentDifference - 1);
-}
+	const bool bHasDiffferences = Children.Num() != 0;
+	if (!bHasDiffferences)
+	{
+		// make one child informing the user that there are no differences:
+		Children.Push(FBlueprintDifferenceTreeEntry::NoDifferencesEntry());
+	}
 
-void FSCSDiffControl::HighlightCurrentDifference()
-{
-	OldSCS.HighlightProperty(DifferingProperties.Entries[CurrentDifference].TreeIdentifier.Name, FPropertyPath());
-	NewSCS.HighlightProperty(DifferingProperties.Entries[CurrentDifference].TreeIdentifier.Name, FPropertyPath());
+	const auto ForwardSelection = [](FOnSCSDiffControlChanged SelectionCallback)
+	{
+		// This allows the owning control to focus the correct tab (or do whatever else it likes):
+		SelectionCallback.ExecuteIfBound();
+	};
+
+	OutTreeEntries.Push(FBlueprintDifferenceTreeEntry::CreateComponentsCategoryEntry(
+		FOnDiffEntryFocused::CreateStatic(ForwardSelection, SelectionCallback),
+		Children,
+		bHasDiffferences
+	));
 }
 
+DECLARE_DELEGATE(FOnCDODiffControlChanged);
+
 class FCDODiffControl	: public TSharedFromThis<FCDODiffControl>
 						, public IDiffControl
 {
 public:
 	FCDODiffControl( const UObject* InOldCDO
 					, const UObject* InNewCDO
-					, const TArray<FPropertySoftPath>& InDifferingProperties );
+					, TArray< FSingleObjectDiffEntry > const& InDifferences
+					, TArray< TSharedPtr<FBlueprintDifferenceTreeEntry> >& OutTreeEntries
+					, TArray< TSharedPtr<FBlueprintDifferenceTreeEntry> >& OutRealDifferences
+					, FOnCDODiffControlChanged SelectionCallback );
 
 	TSharedRef<SWidget> OldDetailsWidget() { return OldDetails.DetailsWidget(); }
 	TSharedRef<SWidget> NewDetailsWidget() { return NewDetails.DetailsWidget(); }
 
 	virtual ~FCDODiffControl() {}
-private:
-	void NextDiff() override;
-	void PrevDiff() override;
-	bool HasNextDifference() const override;
-	bool HasPrevDifference() const override;
-
-	void HighlightCurrentDifference();
 
-	void UpdateDisplayedDifferences();
+private:
+	void HighlightDifference(const FPropertySoftPath& PropertyName);
+	void UpdateDisplayedDifferences(TArray< FSingleObjectDiffEntry > const& InDifferences, TArray< TSharedPtr<FBlueprintDifferenceTreeEntry> >& OutTreeEntries, FOnCDODiffControlChanged SelectionCallback);
 
 	FDetailsDiff OldDetails;
 	FDetailsDiff NewDetails;
 
-	TArray<FPropertySoftPath> DisplayedDifferingProperties;
-	FPropertySoftPathSet DifferingProperties;
 	int CurrentDifference;
 };
 
 FCDODiffControl::FCDODiffControl( 
 		const UObject* InOldCDO
 		, const UObject* InNewCDO
-		, const TArray<FPropertySoftPath>& InDifferingProperties)
-	: OldDetails(InOldCDO, DiffUtils::ResolveAll( InOldCDO, InDifferingProperties), FDetailsDiff::FOnDisplayedPropertiesChanged::CreateRaw(this, &FCDODiffControl::UpdateDisplayedDifferences) )
-	, NewDetails(InNewCDO, DiffUtils::ResolveAll( InNewCDO, InDifferingProperties), FDetailsDiff::FOnDisplayedPropertiesChanged::CreateRaw(this, &FCDODiffControl::UpdateDisplayedDifferences) )
-	, DisplayedDifferingProperties()
-	, DifferingProperties(InDifferingProperties)
-	, CurrentDifference( -1 )
-{
-	UpdateDisplayedDifferences();
-}
-
-void FCDODiffControl::NextDiff()
-{
-	++CurrentDifference;
-	HighlightCurrentDifference();
-}
-
-void FCDODiffControl::PrevDiff()
-{
-	--CurrentDifference;
-	HighlightCurrentDifference();
-}
-
-bool FCDODiffControl::HasNextDifference() const
-{
-	return DisplayedDifferingProperties.IsValidIndex(CurrentDifference + 1);
-}
-
-bool FCDODiffControl::HasPrevDifference() const
-{
-	return DisplayedDifferingProperties.IsValidIndex(CurrentDifference - 1);
-}
-
-void FCDODiffControl::HighlightCurrentDifference()
-{
-	const FPropertySoftPath& PropertyName = DisplayedDifferingProperties[CurrentDifference];
-
-	OldDetails.HighlightProperty(PropertyName);
-	NewDetails.HighlightProperty(PropertyName);
-}
-
-void FCDODiffControl::UpdateDisplayedDifferences()
-{
-	DisplayedDifferingProperties.Empty();
-	CurrentDifference = 0;
+		, TArray< FSingleObjectDiffEntry > const& InDifferences
+		, TArray< TSharedPtr<FBlueprintDifferenceTreeEntry> >& OutTreeEntries
+		, TArray< TSharedPtr<FBlueprintDifferenceTreeEntry> >& OutRealDifferences
+		, FOnCDODiffControlChanged SelectionCallback )
+	: OldDetails(InOldCDO, DiffUtils::ResolveAll( InOldCDO, InDifferences), FDetailsDiff::FOnDisplayedPropertiesChanged() )
+	, NewDetails(InNewCDO, DiffUtils::ResolveAll( InNewCDO, InDifferences), FDetailsDiff::FOnDisplayedPropertiesChanged() )
+{// OrderedProperties will contain differences in the order they are displayed:
+	TArray< const FSingleObjectDiffEntry* > OrderedProperties;
 
 	// create differing properties list based on what is displayed by the old properties..
 	TArray<FPropertySoftPath> OldProperties = OldDetails.GetDisplayedProperties();
 	TArray<FPropertySoftPath> NewProperties = NewDetails.GetDisplayedProperties();
 
+	const auto FindDiffering = [](TArray< FSingleObjectDiffEntry > const& InDifferences, const FPropertySoftPath& PropertyIdentifer) -> const FSingleObjectDiffEntry*
+	{
+		for (const auto& Difference : InDifferences)
+		{
+			if (Difference.Identifier == PropertyIdentifer)
+			{
+				return &Difference;
+			}
+		}
+		return nullptr;
+	};
+
 	// zip the two sets of properties, zip iterators are not trivial to write in c++,
 	// so this procedural stuff will have to do:
 	int IterOld = 0;
@@ -196,36 +195,90 @@ void FCDODiffControl::UpdateDisplayedDifferences()
 		if (IterOld != OldProperties.Num() && IterNew != NewProperties.Num()
 			&& OldProperties[IterOld] == NewProperties[IterNew])
 		{
-			if( DifferingProperties.Contains(OldProperties[IterOld]) )
+			if (const FSingleObjectDiffEntry* Differing = FindDiffering(InDifferences, OldProperties[IterOld]))
 			{
-				DisplayedDifferingProperties.Push(OldProperties[IterOld]);
+				OrderedProperties.Push(Differing);
 			}
 			++IterOld;
 			++IterNew;
 		}
 		else if (IterOld != OldProperties.Num())
 		{
-			if (DifferingProperties.Contains(OldProperties[IterOld]))
+			if (const FSingleObjectDiffEntry* Differing = FindDiffering(InDifferences, OldProperties[IterOld]))
 			{
-				DisplayedDifferingProperties.Push(OldProperties[IterOld]);
+				OrderedProperties.Push(Differing);
 			}
 			++IterOld;
 		}
 		else
 		{
 			check(IterNew != NewProperties.Num());
-			if (DifferingProperties.Contains(NewProperties[IterNew]))
+			if (const FSingleObjectDiffEntry* Differing = FindDiffering(InDifferences, NewProperties[IterNew]))
 			{
-				DisplayedDifferingProperties.Push(NewProperties[IterNew]);
+				OrderedProperties.Push(Differing);
 			}
 			++IterNew;
 		}
 	}
+
+	const auto CreateCDODifferenceWidget = [](FSingleObjectDiffEntry DiffEntry, FText ObjectName) -> TSharedRef<SWidget>
+	{
+		return SNew(STextBlock)
+			.Text(DiffViewUtils::PropertyDiffMessage(DiffEntry, ObjectName))
+			.ColorAndOpacity(DiffViewUtils::Differs());
+	};
+
+	const auto FocusDetailsDifferenceEntry = [](FPropertySoftPath Identifier, FCDODiffControl* Control, FOnCDODiffControlChanged SelectionCallback)
+	{
+		// This allows the owning control to focus the correct tab (or do whatever else it likes):
+		SelectionCallback.ExecuteIfBound();
+		Control->HighlightDifference(Identifier);
+	};
+
+	TArray< TSharedPtr<FBlueprintDifferenceTreeEntry> > Children;
+
+	for (auto Difference : OrderedProperties)
+	{
+		auto Entry = TSharedPtr<FBlueprintDifferenceTreeEntry>(
+			new FBlueprintDifferenceTreeEntry(
+			FOnDiffEntryFocused::CreateStatic(FocusDetailsDifferenceEntry, Difference->Identifier, this, SelectionCallback)
+			, FGenerateDiffEntryWidget::CreateStatic(CreateCDODifferenceWidget, *Difference, RightRevision)
+			, TArray< TSharedPtr<FBlueprintDifferenceTreeEntry> >()
+			)
+		);
+		Children.Push(Entry);
+		OutRealDifferences.Push(Entry);
+	}
+
+	const bool bHasDiffferences = Children.Num() != 0;
+	FLinearColor Color = bHasDiffferences ? DiffViewUtils::Differs() : DiffViewUtils::Identical();
+	if (!bHasDiffferences)
+	{
+		// make one child informing the user that there are no differences:
+		Children.Push(FBlueprintDifferenceTreeEntry::NoDifferencesEntry() );
+	}
+
+	const auto ForwardSelection = [](FOnCDODiffControlChanged SelectionCallback)
+	{
+		// This allows the owning control to focus the correct tab (or do whatever else it likes):
+		SelectionCallback.ExecuteIfBound();
+	};
+
+	OutTreeEntries.Push(FBlueprintDifferenceTreeEntry::CreateDefaultsCategoryEntry(
+		FOnDiffEntryFocused::CreateStatic(ForwardSelection, SelectionCallback)
+		, Children
+		, bHasDiffferences
+	));
+}
+
+void FCDODiffControl::HighlightDifference(const FPropertySoftPath& PropertyName)
+{
+	OldDetails.HighlightProperty(PropertyName);
+	NewDetails.HighlightProperty(PropertyName);
 }
 
 /*List item that entry for a graph*/
 struct KISMET_API FListItemGraphToDiff	: public TSharedFromThis<FListItemGraphToDiff>
-										, public IDiffControl
 {
 	FListItemGraphToDiff(class SBlueprintDiff* Diff, class UEdGraph* GraphOld, class UEdGraph* GraphNew, const FRevisionInfo& RevisionOld, const FRevisionInfo& RevisionNew);
 	virtual ~FListItemGraphToDiff();
@@ -245,10 +298,6 @@ struct KISMET_API FListItemGraphToDiff	: public TSharedFromThis<FListItemGraphTo
 	/** Source for list view */
 	TArray<TSharedPtr<struct FDiffResultItem>> DiffListSource;
 private:
-
-	/*Get icon to use by graph name*/
-	static const FSlateBrush* GetIconForGraph(UEdGraph* Graph);
-
 	/*Diff widget*/
 	class SBlueprintDiff* Diff;
 
@@ -283,17 +332,6 @@ public:
 	void OnSelectionChanged(FSharedDiffOnGraph Item, ESelectInfo::Type SelectionType);
 
 private:
-
-	/** Called when user clicks button to go to next difference in graph */
-	void NextDiff() override;
-
-	/** Called when user clicks button to go to prev difference in graph */
-	void PrevDiff() override;
-
-	/** Called when determining whether the prev/nextbuttons should be enabled */
-	bool HasNextDifference() const override;
-	bool HasPrevDifference() const override;
-
 	/** Get Index of the current diff that is selected */
 	int32 GetCurrentDiffIndex() const;
 
@@ -350,42 +388,20 @@ TSharedRef<SWidget> FListItemGraphToDiff::GenerateWidget()
 
 	const bool bHasDiffs = DiffListSource.Num() > 0;
 
-	//give it a red color to indicate that the graph contains differences
 	if(bHasDiffs)
 	{
-		Color =  FLinearColor(1.0f,0.2f,0.3f);
+		Color = DiffViewUtils::Differs();
 	}
 
 	return SNew(SHorizontalBox)
 	+ SHorizontalBox::Slot()
-		.AutoWidth()
-		.HAlign(HAlign_Center)
-		.VAlign(VAlign_Center)
-		.MaxWidth(8.0f)
-		[
-			SNew(SImage)
-			.ColorAndOpacity(Color)
-			.Image(GetIconForGraph(GraphOld))
-
-		]
-	+SHorizontalBox::Slot()
-		.FillWidth(1.0f)
-		[
-			SNew(STextBlock)
-			.ColorAndOpacity(Color)
-			.Text(Graph->GetName())
-
-		]
-	+SHorizontalBox::Slot()
-		.AutoWidth()
-		.HAlign(HAlign_Right)
-		.VAlign(VAlign_Center)
-		.MaxWidth(8.0f)
-		[
-			SNew(SImage)
-			.ColorAndOpacity(Color)
-			.Image(GetIconForGraph(GraphNew))
-		];
+	[
+		SNew(STextBlock)
+		.ColorAndOpacity(Color)
+		.Text(Graph->GetName())
+	]
+	+ DiffViewUtils::Box( GraphOld != nullptr, Color )
+	+ DiffViewUtils::Box( GraphNew != nullptr, Color );
 }
 
 TSharedRef<SWidget> FListItemGraphToDiff::GenerateDiffListWidget()
@@ -453,11 +469,6 @@ void FListItemGraphToDiff::OnSelectionChanged( FSharedDiffOnGraph Item, ESelectI
 	}
 }
 
-const FSlateBrush* FListItemGraphToDiff::GetIconForGraph( UEdGraph* Graph ) 
-{
-	return Graph? FEditorStyle::GetBrush("BlueprintDif.HasGraph") : FEditorStyle::GetBrush("BlueprintDif.MissingGraph");
-}
-
 FText FListItemGraphToDiff::GetToolTip() 
 {
 	if ( GraphOld && GraphNew )
@@ -491,12 +502,9 @@ void FListItemGraphToDiff::BuildDiffSourceArray()
 	TArray<FDiffSingleResult> FoundDiffs;
 	FGraphDiffControl::DiffGraphs(GraphOld, GraphNew, FoundDiffs);
 
-	FName GraphName = GraphOld ? GraphOld->GetFName() : GraphNew->GetFName();
-
 	DiffListSource.Empty();
 	for (auto DiffIt(FoundDiffs.CreateIterator()); DiffIt; ++DiffIt)
 	{
-		DiffIt->OwningGraph = GraphName;
 		DiffListSource.Add(FSharedDiffOnGraph(new FDiffResultItem(*DiffIt)));
 	}
 
@@ -511,49 +519,6 @@ void FListItemGraphToDiff::BuildDiffSourceArray()
 	Sort(DiffListSource.GetData(), DiffListSource.Num(), SortDiff());
 }
 
-void FListItemGraphToDiff::NextDiff()
-{
-	int32 Index = (GetCurrentDiffIndex() + 1) % DiffListSource.Num();
-
-	DiffList->SetSelection(DiffListSource[Index]);
-}
-
-void FListItemGraphToDiff::PrevDiff()
-{
-	int32 Index = GetCurrentDiffIndex();
-	if(Index == 0)
-	{
-		Index = DiffListSource.Num() - 1;
-	}
-	else
-	{
-		Index = (Index - 1) % DiffListSource.Num();
-	}
-
-	DiffList->SetSelection(DiffListSource[Index]);
-}
-
-bool FListItemGraphToDiff::HasNextDifference() const
-{
-	if (!DiffList.IsValid())
-	{
-		return false;
-	}
-	int32 Index = GetCurrentDiffIndex();
-	return DiffListSource.IsValidIndex(Index + 1);
-}
-
-bool FListItemGraphToDiff::HasPrevDifference() const
-{
-	if (!DiffList.IsValid())
-	{
-		return false;
-	}
-
-	int32 Index = GetCurrentDiffIndex();
-	return DiffListSource.IsValidIndex(Index - 1);
-}
-
 int32 FListItemGraphToDiff::GetCurrentDiffIndex() const
 {
 	if ( DiffList.IsValid() )
@@ -647,9 +612,6 @@ bool DiffWidgetUtils::HasPrevDifference(SListView< TSharedPtr< struct FDiffSingl
 	return ListViewSource.IsValidIndex(CurrentIndex - 1);
 }
 
-//////////////////////////////////////////////////////////////////////////
-// SBlueprintDiff
-
 BEGIN_SLATE_FUNCTION_BUILD_OPTIMIZATION
 void SBlueprintDiff::Construct( const FArguments& InArgs)
 {
@@ -757,15 +719,14 @@ void SBlueprintDiff::Construct( const FArguments& InArgs)
 	ToolbarBuilder.AddToolBarButton(
 		FUIAction(FExecuteAction::CreateSP(this, &SBlueprintDiff::OnToggleLockView2))
 		, NAME_None
-		, LOCTEXT("LockGraphsLabel", "Lock Graphs")
-		, LOCTEXT("LockGraphsTooltip", "Lock the various panels displayed in the graph view, so that they scroll together")
-		, TAttribute<FSlateIcon>(this, &SBlueprintDiff::GetLockViewImage2)
+		, LOCTEXT("LockGraphsLabel", "Lock/Unlock")
+		, LOCTEXT("LockGraphsTooltip", "Force all graph views to change together, or allow independent scrolling/zooming")
+		, TAttribute<FSlateIcon>(this, &SBlueprintDiff::GetLockViewImage)
 	);
 
 	GraphPanel = GenerateGraphPanel();
 	DefaultsPanel = GenerateDefaultsPanel();
 	ComponentsPanel = GenerateComponentsPanel();
-
 	// Unfortunately we can't perform the diff until the UI is generated, the primary reason for this is that
 	// details customizations determine what is actually editable:
 	for (const auto& Graph : Graphs)
@@ -773,7 +734,7 @@ void SBlueprintDiff::Construct( const FArguments& InArgs)
 		TArray< TSharedPtr<FBlueprintDifferenceTreeEntry> > Children;
 		for (const auto& Difference : Graph->DiffListSource)
 		{
-			auto ChildEntry = MakeShareable(
+			auto ChildEntry = TSharedPtr<FBlueprintDifferenceTreeEntry>(
 				new FBlueprintDifferenceTreeEntry(
 				FOnDiffEntryFocused::CreateRaw(this, &SBlueprintDiff::OnDiffListSelectionChanged, Difference)
 				, FGenerateDiffEntryWidget::CreateSP(Difference.ToSharedRef(), &FDiffResultItem::GenerateWidget)
@@ -781,9 +742,17 @@ void SBlueprintDiff::Construct( const FArguments& InArgs)
 				)
 			);
 			Children.Push(ChildEntry);
+			RealDifferences.Push(ChildEntry);
+			(void)ChildEntry;
+		}
+
+		if (Children.Num() == 0)
+		{
+			// make one child informing the user that there are no differences:
+			Children.Push(FBlueprintDifferenceTreeEntry::NoDifferencesEntry());
 		}
 
-		auto Entry = MakeShareable(
+		auto Entry = TSharedPtr<FBlueprintDifferenceTreeEntry>(
 			new FBlueprintDifferenceTreeEntry(
 				FOnDiffEntryFocused::CreateRaw(this, &SBlueprintDiff::OnSelectionChanged, Graph, ESelectInfo::Direct )
 				, FGenerateDiffEntryWidget::CreateSP(Graph.ToSharedRef(), &FListItemGraphToDiff::GenerateWidget)
@@ -792,31 +761,25 @@ void SBlueprintDiff::Construct( const FArguments& InArgs)
 		MasterDifferencesList.Push(Entry);
 	}
 
-	//DefaultsPanel.DiffControl->CalculateDifferences(MasterDifferencesList);
-	//ComponentsPanel.DiffControl->CalculateDifferences(MasterDifferencesList);
-
-	const auto RowGenerator = []( TSharedPtr< FBlueprintDifferenceTreeEntry > Entry, const TSharedRef<STableViewBase>& Owner ) -> TSharedRef< ITableRow >
-	{
-		return SNew( STableRow<TSharedPtr<FBlueprintDifferenceTreeEntry> >, Owner )
-			[
-				Entry->GenerateWidget.Execute()
-			];
-	};
+	DifferencesTreeView = DiffTreeView::CreateTreeView(&MasterDifferencesList);
 
-	const auto ChildrenAccessor = [](TSharedPtr<FBlueprintDifferenceTreeEntry> InTreeItem, TArray< TSharedPtr< FBlueprintDifferenceTreeEntry > >& OutChildren, const TArray< TSharedPtr<FBlueprintDifferenceTreeEntry> >* MasterList )
+	const auto TextBlock = [](FText Text) -> TSharedRef<SWidget>
 	{
-		OutChildren = InTreeItem->Children;
+		return SNew(STextBlock)
+			.Visibility(EVisibility::HitTestInvisible)
+			.TextStyle(FEditorStyle::Get(), "GraphPreview.CornerText")
+			.Text(Text);
 	};
 
-	const auto Selector = [](TSharedPtr<FBlueprintDifferenceTreeEntry> InTreeItem, ESelectInfo::Type Type)
-	{
-		if( InTreeItem.IsValid() )
-		{
-			InTreeItem->OnFocus.ExecuteIfBound();
-		}
-	};
-	
-	const TArray< TSharedPtr<FBlueprintDifferenceTreeEntry> >* MasterDifferencesListPtr = &MasterDifferencesList;
+	TSharedRef<SWidget> Overlay = SNew(SHorizontalBox)
+		+ SHorizontalBox::Slot()
+		[
+			TextBlock(DiffViewUtils::GetPanelLabel(PanelOld.Blueprint, PanelOld.RevisionInfo, FText()))
+		]
+	+ SHorizontalBox::Slot()
+		[
+			TextBlock(DiffViewUtils::GetPanelLabel(PanelNew.Blueprint, PanelNew.RevisionInfo, FText()))
+		];
 
 	this->ChildSlot
 		[
@@ -865,17 +828,22 @@ void SBlueprintDiff::Construct( const FArguments& InArgs)
 						SNew(SBorder)
 						.BorderImage(FEditorStyle::GetBrush("ToolPanel.GroupBorder"))
 						[
-							SNew(STreeView< TSharedPtr< FBlueprintDifferenceTreeEntry > >)
-							.OnGenerateRow( STreeView< TSharedPtr< FBlueprintDifferenceTreeEntry > >::FOnGenerateRow::CreateStatic(RowGenerator) )
-							.OnGetChildren( STreeView< TSharedPtr< FBlueprintDifferenceTreeEntry > >::FOnGetChildren::CreateStatic(ChildrenAccessor, MasterDifferencesListPtr) )
-							.OnSelectionChanged( STreeView< TSharedPtr< FBlueprintDifferenceTreeEntry > >::FOnSelectionChanged::CreateStatic(Selector) )
-							.TreeItemsSource(&MasterDifferencesList)
+							DifferencesTreeView.ToSharedRef()
 						]
 					]
 					+ SSplitter::Slot()
 					.Value(.8f)
 					[
-						SAssignNew(ModeContents, SBox)
+						SNew(SOverlay)
+						+ SOverlay::Slot()
+						[
+							SAssignNew(ModeContents, SBox)
+						]
+						+ SOverlay::Slot()
+						.VAlign(VAlign_Bottom)
+						[
+							Overlay
+						]
 					]
 				]
 			]
@@ -934,28 +902,22 @@ TSharedRef<SWidget> SBlueprintDiff::DefaultEmptyPanel()
 
 void SBlueprintDiff::NextDiff()
 {
-	if( DiffControl.IsValid() )
-	{
-		DiffControl->NextDiff();
-	}
+	DiffTreeView::HighlightNextDifference(DifferencesTreeView.ToSharedRef(), RealDifferences, MasterDifferencesList);
 }
 
 void SBlueprintDiff::PrevDiff()
 {
-	if (DiffControl.IsValid())
-	{
-		DiffControl->PrevDiff();
-	}
+	DiffTreeView::HighlightPrevDifference(DifferencesTreeView.ToSharedRef(), RealDifferences, MasterDifferencesList);
 }
 
 bool SBlueprintDiff::HasNextDiff() const
 {
-	return DiffControl.IsValid() && DiffControl->HasNextDifference();
+	return DiffTreeView::HasNextDifference(DifferencesTreeView.ToSharedRef(), RealDifferences);
 }
 
 bool SBlueprintDiff::HasPrevDiff() const
 {
-	return DiffControl.IsValid() && DiffControl->HasPrevDifference();
+	return DiffTreeView::HasPrevDifference(DifferencesTreeView.ToSharedRef(), RealDifferences);
 }
 
 TSharedRef<SDockTab> SBlueprintDiff::CreateGraphDiffViews(const FSpawnTabArgs& Args)
@@ -1067,16 +1029,11 @@ FReply	SBlueprintDiff::OnToggleLockView()
 	return FReply::Handled();
 }
 
-FSlateIcon SBlueprintDiff::GetLockViewImage2() const
+FSlateIcon SBlueprintDiff::GetLockViewImage() const
 {
 	return FSlateIcon(FEditorStyle::GetStyleSetName(), bLockViews ? "GenericLock" : "GenericUnlock");
 }
 
-const FSlateBrush* SBlueprintDiff::GetLockViewImage() const
-{
-	return bLockViews ? FEditorStyle::GetBrush("GenericLock") : FEditorStyle::GetBrush("GenericUnlock");
-}
-
 void SBlueprintDiff::ResetGraphEditors()
 {
 	if(PanelOld.GraphEditor.IsValid() && PanelNew.GraphEditor.IsValid())
@@ -1096,6 +1053,11 @@ void SBlueprintDiff::ResetGraphEditors()
 
 void FDiffPanel::GeneratePanel(UEdGraph* Graph, UEdGraph* GraphToDiff )
 {
+	if( GraphEditor.IsValid() && GraphEditor.Pin()->GetCurrentGraph() == Graph )
+	{
+		return;
+	}
+
 	if( LastFocusedPin )
 	{
 		LastFocusedPin->bIsDiffing = false;
@@ -1120,12 +1082,9 @@ void FDiffPanel::GeneratePanel(UEdGraph* Graph, UEdGraph* GraphToDiff )
 			InEvents.OnSelectionChanged = SGraphEditor::FOnSelectionChanged::CreateStatic(SelectionChangedHandler, DetailsView);
 		}
 
-		FGraphAppearanceInfo AppearanceInfo;
-		AppearanceInfo.CornerText = LOCTEXT("AppearanceCornerText_BlueprintDif", "DIFF").ToString();
-
 		if ( !GraphEditorCommands.IsValid() )
 		{
-			GraphEditorCommands = MakeShareable( new FUICommandList() );
+			GraphEditorCommands = TSharedPtr<FUICommandList>( new FUICommandList() );
 
 			GraphEditorCommands->MapAction( FGenericCommands::Get().Copy,
 				FExecuteAction::CreateRaw( this, &FDiffPanel::CopySelectedNodes ),
@@ -1142,11 +1101,6 @@ void FDiffPanel::GeneratePanel(UEdGraph* Graph, UEdGraph* GraphToDiff )
 			.GraphToDiff(GraphToDiff)
 			.IsEditable(false)
 			.TitleBarEnabledOnly(false)
-			.TitleBar(SNew(SBorder).HAlign(HAlign_Center)
-						[
-							SNew(STextBlock).Text(GetTitle())
-						])
-			.Appearance(AppearanceInfo)
 			.GraphEvents(InEvents);
 
 		GraphEditor = Editor;
@@ -1182,39 +1136,6 @@ void FDiffPanel::CopySelectedNodes()
 	FPlatformMisc::ClipboardCopy(*ExportedText);
 }
 
-
-FString FDiffPanel::GetTitle() const
-{
-	FString Title = LOCTEXT("CurrentRevision", "Current Revision").ToString();
-
-	// if this isn't the current working version being displayed
-	if (RevisionInfo.Revision >= 0)
-	{
-		FString DateString = RevisionInfo.Date.ToString(TEXT("%m/%d/%Y"));
-		if (bShowAssetName)
-		{
-			FString AssetName = Blueprint->GetName();
-
-			FText LocalizedFormat = LOCTEXT("NamedRevisionDiff", "%s - Revision %i, CL %i, %s");
-			Title = FString::Printf(*LocalizedFormat.ToString(), *AssetName, RevisionInfo.Revision, RevisionInfo.Changelist, *DateString);
-		}
-		else
-		{
-			FText LocalizedFormat = LOCTEXT("PreviousRevisionDif", "Revision %i, CL %i, %s");
-			Title = FString::Printf(*LocalizedFormat.ToString(), RevisionInfo.Revision, RevisionInfo.Changelist, *DateString);
-		}
-	}
-	else if (bShowAssetName)
-	{
-		FString AssetName = Blueprint->GetName();
-
-		FText LocalizedFormat = LOCTEXT("NamedCurrentRevision", "%s - Current Revision");
-		Title = FString::Printf(*LocalizedFormat.ToString(), *AssetName);
-	}
-
-	return Title;
-}
-
 bool FDiffPanel::CanCopyNodes() const
 {
 	// If any of the nodes can be duplicated then we should allow copying
@@ -1272,10 +1193,7 @@ FDiffPanel& SBlueprintDiff::GetDiffPanelForNode(UEdGraphNode& Node)
 
 void SBlueprintDiff::HandleGraphChanged( const FName GraphName )
 {
-	if( CurrentMode != FBlueprintEditorApplicationModes::StandardBlueprintEditorMode)
-	{
-		SetCurrentMode(FBlueprintEditorApplicationModes::StandardBlueprintEditorMode);
-	}
+	SetCurrentMode(FBlueprintEditorApplicationModes::StandardBlueprintEditorMode);
 	
 	TArray<UEdGraph*> GraphsOld, GraphsNew;
 	PanelOld.Blueprint->GetAllGraphs(GraphsOld);
@@ -1374,11 +1292,12 @@ SBlueprintDiff::FDiffControl SBlueprintDiff::GenerateDefaultsPanel()
 {
 	const UObject* A = DiffUtils::GetCDO(PanelOld.Blueprint);
 	const UObject* B = DiffUtils::GetCDO(PanelNew.Blueprint);
-	TArray<FPropertySoftPath> DifferingProperties;
+	TArray<FSingleObjectDiffEntry> DifferingProperties;
 	DiffUtils::CompareUnrelatedObjects( A, B, DifferingProperties);
 
+	auto NewDiffControl = TSharedPtr<FCDODiffControl>(new FCDODiffControl(A, B, DifferingProperties, MasterDifferencesList, RealDifferences, FOnCDODiffControlChanged::CreateRaw(this, &SBlueprintDiff::SetCurrentMode, FBlueprintEditorApplicationModes::BlueprintDefaultsMode)));
+
 	SBlueprintDiff::FDiffControl Ret;
-	auto NewDiffControl = TSharedPtr<FCDODiffControl>(new FCDODiffControl(A, B, DifferingProperties) );
 	//Splitter for left and right blueprint. Current convention is for the local (probably newer?) blueprint to be on the right:
 	Ret.DiffControl = NewDiffControl;
 	Ret.Widget = SNew(SSplitter)
@@ -1409,7 +1328,7 @@ SBlueprintDiff::FDiffControl SBlueprintDiff::GenerateComponentsPanel()
 	SBlueprintDiff::FDiffControl Ret;
 
 	//Splitter for left and right blueprint. Current convention is for the local (probably newer?) blueprint to be on the right:
-	auto NewDiffControl = TSharedPtr<FSCSDiffControl>(new FSCSDiffControl(PanelOld.Blueprint, PanelNew.Blueprint) );
+	auto NewDiffControl = TSharedPtr<FSCSDiffControl>(new FSCSDiffControl(PanelOld.Blueprint, PanelNew.Blueprint, MasterDifferencesList, RealDifferences, FOnSCSDiffControlChanged::CreateRaw(this, &SBlueprintDiff::SetCurrentMode, FBlueprintEditorApplicationModes::BlueprintComponentsMode) ) );
 	Ret.DiffControl = NewDiffControl;
 	Ret.Widget = SNew(SSplitter)
 		+ SSplitter::Slot()
@@ -1436,6 +1355,11 @@ SBlueprintDiff::FDiffControl SBlueprintDiff::GenerateComponentsPanel()
 
 void SBlueprintDiff::SetCurrentMode(FName NewMode)
 {
+	if( CurrentMode == NewMode )
+	{
+		return;
+	}
+
 	CurrentMode = NewMode;
 
 	DiffControl = TSharedPtr< IDiffControl >();
diff --git a/Engine/Source/Editor/Kismet/Public/DiffUtils.h b/Engine/Source/Editor/Kismet/Public/DiffUtils.h
--- a/Engine/Source/Editor/Kismet/Public/DiffUtils.h
+++ b/Engine/Source/Editor/Kismet/Public/DiffUtils.h
@@ -2,6 +2,7 @@
 #pragma once
 
 #include "PropertyPath.h"
+#include "IAssetTypeActions.h"
 
 struct FResolvedProperty
 {
@@ -114,11 +115,40 @@ FORCEINLINE uint32 GetTypeHash( FPropertySoftPath const& Path )
 // Trying to restrict us to this typedef because I'm a little skeptical about hashing FPropertySoftPath safely
 typedef TSet< FPropertySoftPath > FPropertySoftPathSet;
 
+namespace EPropertyDiffType
+{
+	enum Type
+	{
+		Invalid,
+
+		PropertyAddedToA,
+		PropertyAddedToB,
+		PropertyValueChanged,
+	};
+}
+
+struct FSingleObjectDiffEntry
+{
+	FSingleObjectDiffEntry()
+		: Identifier()
+		, DiffType(EPropertyDiffType::Invalid)
+	{
+	}
+
+	FSingleObjectDiffEntry(const FPropertySoftPath& InIdentifier, EPropertyDiffType::Type InDiffType )
+		: Identifier(InIdentifier)
+		, DiffType(InDiffType)
+	{
+	}
+
+	FPropertySoftPath Identifier;
+	EPropertyDiffType::Type DiffType;
+};
+
 namespace ETreeDiffType
 {
 	enum Type
 	{
-		UNKNOWN,
 		NODE_ADDED,
 		NODE_REMOVED,
 		NODE_PROPERTY_CHANGED,
@@ -129,9 +159,16 @@ namespace ETreeDiffType
 
 struct FSCSDiffEntry
 {
+	FSCSDiffEntry( const FSCSIdentifier& InIdentifier, ETreeDiffType::Type InDiffType, const FSingleObjectDiffEntry& InPropertyDiff )
+		: TreeIdentifier(InIdentifier)
+		, DiffType(InDiffType)
+		, PropertyDiff(InPropertyDiff)
+	{
+	}
+
 	FSCSIdentifier TreeIdentifier;
 	ETreeDiffType::Type DiffType;
-	FPropertySoftPath PropertyIdentifier;
+	FSingleObjectDiffEntry PropertyDiff;
 };
 
 struct FSCSDiffRoot
@@ -143,13 +180,13 @@ struct FSCSDiffRoot
 namespace DiffUtils
 {
 	KISMET_API const UObject* GetCDO(const UBlueprint* ForBlueprint);
-	KISMET_API void CompareUnrelatedObjects(const UObject* A, const UObject* B, TArray<FPropertySoftPath>& OutDifferingProperties);
+	KISMET_API void CompareUnrelatedObjects(const UObject* A, const UObject* B, TArray<FSingleObjectDiffEntry>& OutDifferingProperties);
 	KISMET_API void CompareUnrelatedSCS(const UBlueprint* Old, const TArray< FSCSResolvedIdentifier >& OldHierarchy, const UBlueprint* New, const TArray< FSCSResolvedIdentifier >& NewHierarchy, FSCSDiffRoot& OutDifferingEntries );
 	KISMET_API bool Identical(const FResolvedProperty& AProp, const FResolvedProperty& BProp);
 	TArray<FPropertySoftPath> GetVisiblePropertiesInOrderDeclared(const UObject* ForObj, const TArray<FName>& Scope = TArray<FName>());
 
 	KISMET_API TArray<FPropertyPath> ResolveAll(const UObject* Object, const TArray<FPropertySoftPath>& InSoftProperties);
-
+	KISMET_API TArray<FPropertyPath> ResolveAll(const UObject* Object, const TArray<FSingleObjectDiffEntry>& InDifferences);
 }
 
 DECLARE_DELEGATE(FOnDiffEntryFocused);
@@ -166,7 +203,40 @@ public:
 		check( InGenerateWidget.IsBound() );
 	}
 
+	/** The FBlueprintDifferenceTreeEntry used to display a message to the user explaining that there are no differences: */
+	KISMET_API static TSharedPtr<FBlueprintDifferenceTreeEntry> NoDifferencesEntry();
+	/** The FBlueprintDifferenceTreeEntry used to label the defaults category: */
+	KISMET_API static TSharedPtr<FBlueprintDifferenceTreeEntry> CreateDefaultsCategoryEntry(FOnDiffEntryFocused FocusCallback, const TArray< TSharedPtr<FBlueprintDifferenceTreeEntry> >& Children, bool bHasDifferences );
+	KISMET_API static TSharedPtr<FBlueprintDifferenceTreeEntry> CreateDefaultsCategoryEntryForMerge(FOnDiffEntryFocused FocusCallback, const TArray< TSharedPtr<FBlueprintDifferenceTreeEntry> >& Children, bool bHasRemoteDifferences, bool bHasLocalDifferences, bool bHasConflicts );
+	KISMET_API static TSharedPtr<FBlueprintDifferenceTreeEntry> CreateComponentsCategoryEntry(FOnDiffEntryFocused FocusCallback, const TArray< TSharedPtr<FBlueprintDifferenceTreeEntry> >& Children, bool bHasDifferences);
+	KISMET_API static TSharedPtr<FBlueprintDifferenceTreeEntry> CreateComponentsCategoryEntryForMerge(FOnDiffEntryFocused FocusCallback, const TArray< TSharedPtr<FBlueprintDifferenceTreeEntry> >& Children, bool bHasRemoteDifferences, bool bHasLocalDifferences, bool bHasConflicts);
+	
 	FOnDiffEntryFocused OnFocus;
 	FGenerateDiffEntryWidget GenerateWidget;
 	TArray< TSharedPtr<FBlueprintDifferenceTreeEntry> > Children;
 };
+
+namespace DiffTreeView
+{
+	KISMET_API TSharedRef< STreeView<TSharedPtr< FBlueprintDifferenceTreeEntry > > > CreateTreeView(TArray< TSharedPtr<FBlueprintDifferenceTreeEntry> >* DifferencesList);
+	KISMET_API int32 CurrentDifference( TSharedRef< STreeView<TSharedPtr< FBlueprintDifferenceTreeEntry > > > TreeView, const TArray< TSharedPtr<class FBlueprintDifferenceTreeEntry> >& Differences );
+	KISMET_API void HighlightNextDifference(TSharedRef< STreeView<TSharedPtr< FBlueprintDifferenceTreeEntry > > > TreeView, const TArray< TSharedPtr<class FBlueprintDifferenceTreeEntry> >& Differences, const TArray< TSharedPtr<class FBlueprintDifferenceTreeEntry> >& RootDifferences);
+	KISMET_API void HighlightPrevDifference(TSharedRef< STreeView<TSharedPtr< FBlueprintDifferenceTreeEntry > > > TreeView, const TArray< TSharedPtr<class FBlueprintDifferenceTreeEntry> >& Differences, const TArray< TSharedPtr<class FBlueprintDifferenceTreeEntry> >& RootDifferences);
+	KISMET_API bool HasNextDifference(TSharedRef< STreeView<TSharedPtr< FBlueprintDifferenceTreeEntry > > > TreeView, const TArray< TSharedPtr<class FBlueprintDifferenceTreeEntry> >& Differences);
+	KISMET_API bool HasPrevDifference(TSharedRef< STreeView<TSharedPtr< FBlueprintDifferenceTreeEntry > > > TreeView, const TArray< TSharedPtr<class FBlueprintDifferenceTreeEntry> >& Differences);
+}
+
+namespace DiffViewUtils
+{
+	KISMET_API FLinearColor LookupColor( bool bDiffers, bool bConflicts = false );
+	KISMET_API FLinearColor Differs();
+	KISMET_API FLinearColor Identical();
+	KISMET_API FLinearColor Missing();
+	KISMET_API FLinearColor Conflicting();
+
+	KISMET_API FText PropertyDiffMessage(FSingleObjectDiffEntry Difference, FText ObjectName);
+	KISMET_API FText SCSDiffMessage(const FSCSDiffEntry& Difference, FText ObjectName);
+	KISMET_API FText GetPanelLabel( const UBlueprint* Blueprint, const FRevisionInfo& Revision, FText Label );
+
+	KISMET_API SHorizontalBox::FSlot& Box(bool bIsPresent, FLinearColor Color);
+}
diff --git a/Engine/Source/Editor/Kismet/Public/IDiffControl.h b/Engine/Source/Editor/Kismet/Public/IDiffControl.h
deleted file mode 100644
diff --git a/Engine/Source/Editor/Kismet/Public/SBlueprintDiff.h b/Engine/Source/Editor/Kismet/Public/SBlueprintDiff.h
--- a/Engine/Source/Editor/Kismet/Public/SBlueprintDiff.h
+++ b/Engine/Source/Editor/Kismet/Public/SBlueprintDiff.h
@@ -54,9 +54,6 @@ struct KISMET_API FDiffPanel
 	/* Generate the 'MyBlueprint' widget, which is private to this module */
 	TSharedRef<SWidget> GenerateMyBlueprintPanel();
 
-	/*Get the title to show at the top of the panel*/
-	FString GetTitle() const;
-
 	/* Called when user hits keyboard shortcut to copy nodes*/
 	void CopySelectedNodes();
 
@@ -167,8 +164,7 @@ protected:
 	void ResetGraphEditors();
 
 	/*Get the image to show for the toggle lock option*/
-	FSlateIcon GetLockViewImage2() const;
-	const FSlateBrush* GetLockViewImage() const;
+	FSlateIcon GetLockViewImage() const;
 
 	/* This buffer stores the currently displayed results */
 	TArray< FGraphToDiff> Graphs;
@@ -221,6 +217,12 @@ protected:
 	/** Tree of differences collected across all panels: */
 	TArray< TSharedPtr<class FBlueprintDifferenceTreeEntry> > MasterDifferencesList;
 
+	/** List of all differences, cached so that we can iterate only the differences and not labels, etc: */
+	TArray< TSharedPtr<class FBlueprintDifferenceTreeEntry> > RealDifferences;
+
+	/** Tree view that displays the differences, cached for the buttons that iterate the differences: */
+	TSharedPtr< STreeView< TSharedPtr< FBlueprintDifferenceTreeEntry > > > DifferencesTreeView;
+
 	/** Stored references to widgets used to display various parts of a blueprint: */
 	FDiffControl GraphPanel;
 	FDiffControl DefaultsPanel;
diff --git a/Engine/Source/Runtime/Engine/Classes/Engine/Blueprint.h b/Engine/Source/Runtime/Engine/Classes/Engine/Blueprint.h
--- a/Engine/Source/Runtime/Engine/Classes/Engine/Blueprint.h
+++ b/Engine/Source/Runtime/Engine/Classes/Engine/Blueprint.h
@@ -358,6 +358,14 @@ class ENGINE_API UBlueprint : public UBlueprintCore
 	 *
 	 * This flag needs to be copied on duplication (because it's the duplicated object that we're disabling on PostDuplicate),
 	 * but we don't *need* to serialize it for permanent objects.
+	 *
+	 * Without setting this flag a blueprint will be marked dirty when it is duplicated and if saved while in this dirty
+	 * state you will not be able to open the blueprint. More specifically, UClass::Rename (called by DestroyGeneratedClass)
+	 * sets a dirty flag on the package. Once saved the package will fail to open because some unnamed objects are present in
+	 * the pacakge.
+	 *
+	 * This flag can be used to avoid the package being marked as dirty in the first place. Ideally PostDuplicateObject
+	 * would not rename classes that are still in use by the original object.
 	 */
 	UPROPERTY()
 	mutable bool bDuplicatingReadOnly;
