In the MyBlueprint window, dragging a function item will immediately display a mouse decorator.

diff --git a/Engine/Source/Editor/Kismet/Private/BPFunctionDragDropAction.cpp b/Engine/Source/Editor/Kismet/Private/BPFunctionDragDropAction.cpp
--- a/Engine/Source/Editor/Kismet/Private/BPFunctionDragDropAction.cpp
+++ b/Engine/Source/Editor/Kismet/Private/BPFunctionDragDropAction.cpp
@@ -164,11 +164,12 @@ bool FKismetDragDropAction::ActionWillShowExistingNode() const
 
 //------------------------------------------------------------------------------
 TSharedRef<FKismetFunctionDragDropAction> FKismetFunctionDragDropAction::New(
-	FName                   InFunctionName, 
-	UClass*                 InOwningClass, 
-	FMemberReference const& InCallOnMember, 
-	FNodeCreationAnalytic   AnalyticCallback, 
-	FCanBeDroppedDelegate   CanBeDroppedDelegate /* = FCanBeDroppedDelegate() */)
+	TSharedPtr<FEdGraphSchemaAction>	InActionNode,
+	FName								InFunctionName, 
+	UClass*								InOwningClass, 
+	FMemberReference const&				InCallOnMember, 
+	FNodeCreationAnalytic				AnalyticCallback, 
+	FCanBeDroppedDelegate				CanBeDroppedDelegate /* = FCanBeDroppedDelegate() */)
 {
 	TSharedRef<FKismetFunctionDragDropAction> Operation = MakeShareable(new FKismetFunctionDragDropAction);
 	Operation->FunctionName     = InFunctionName;
@@ -176,6 +177,7 @@ TSharedRef<FKismetFunctionDragDropAction> FKismetFunctionDragDropAction::New(
 	Operation->CallOnMember     = InCallOnMember;
 	Operation->AnalyticCallback = AnalyticCallback;
 	Operation->CanBeDroppedDelegate = CanBeDroppedDelegate;
+	Operation->ActionNode = InActionNode;
 
 	if (!CanBeDroppedDelegate.IsBound())
 	{
@@ -196,18 +198,6 @@ FKismetFunctionDragDropAction::FKismetFunctionDragDropAction()
 //------------------------------------------------------------------------------
 void FKismetFunctionDragDropAction::HoverTargetChanged()
 {
-	FGraphActionListBuilderBase::ActionGroup DropActionSet(TSharedPtr<FEdGraphSchemaAction>(NULL));
-	GetDropAction(DropActionSet);
-
-	if (DropActionSet.Actions.Num() > 0)
-	{
-		ActionNode = DropActionSet.Actions[0];
-	}
-	else 
-	{
-		ActionNode = NULL;
-	}
-	
 	FKismetDragDropAction::HoverTargetChanged();
 }
 
@@ -216,8 +206,9 @@ FReply FKismetFunctionDragDropAction::DroppedOnPanel(TSharedRef<SWidget> const&
 {	
 	FReply Reply = FReply::Unhandled();
 
+	// The ActionNode set during construction points to the Graph, this is suitable for displaying the mouse decorator but needs to be more complete based on the current graph
 	FGraphActionListBuilderBase::ActionGroup DropActionSet(TSharedPtr<FEdGraphSchemaAction>(NULL));
-	GetDropAction(DropActionSet);
+	GetDropAction(Graph, DropActionSet);
 
 	if (DropActionSet.Actions.Num() > 0)
 	{
@@ -254,9 +245,9 @@ UFunction const* FKismetFunctionDragDropAction::GetFunctionProperty() const
 }
 
 //------------------------------------------------------------------------------
-void FKismetFunctionDragDropAction::GetDropAction(FGraphActionListBuilderBase::ActionGroup& DropActionOut) const
+void FKismetFunctionDragDropAction::GetDropAction(UEdGraph& Graph, FGraphActionListBuilderBase::ActionGroup& DropActionOut) const
 {
-	if (UEdGraph const* const HoveredGraph = GetHoveredGraph())
+	if (UEdGraph const* const HoveredGraph = &Graph)
 	{
 		if (UBlueprint* DropOnBlueprint = FBlueprintEditorUtils::FindBlueprintForGraph(HoveredGraph))
 		{
diff --git a/Engine/Source/Editor/Kismet/Private/BPFunctionDragDropAction.h b/Engine/Source/Editor/Kismet/Private/BPFunctionDragDropAction.h
--- a/Engine/Source/Editor/Kismet/Private/BPFunctionDragDropAction.h
+++ b/Engine/Source/Editor/Kismet/Private/BPFunctionDragDropAction.h
@@ -57,7 +57,7 @@ public:
 	virtual FReply DroppedOnPanel( const TSharedRef< class SWidget >& Panel, FVector2D ScreenPosition, FVector2D GraphPosition, UEdGraph& Graph) override;
 	// End of FGraphEditorDragDropAction
 
-	static TSharedRef<FKismetFunctionDragDropAction> New(FName InFunctionName, UClass* InOwningClass, const FMemberReference& InCallOnMember, FNodeCreationAnalytic AnalyticCallback, FCanBeDroppedDelegate CanBeDroppedDelegate = FCanBeDroppedDelegate());
+	static TSharedRef<FKismetFunctionDragDropAction> New(TSharedPtr<FEdGraphSchemaAction> InActionNode, FName InFunctionName, UClass* InOwningClass, const FMemberReference& InCallOnMember, FNodeCreationAnalytic AnalyticCallback, FCanBeDroppedDelegate CanBeDroppedDelegate = FCanBeDroppedDelegate());
 
 protected:
 	/** Name of function being dragged */
@@ -71,7 +71,7 @@ protected:
 	UFunction const* GetFunctionProperty() const;
 
 	/** Constructs an action to execute, placing a function call node for the associated function */
-	void GetDropAction(FGraphActionListBuilderBase::ActionGroup& DropActionOut) const;
+	void GetDropAction(UEdGraph& Graph, FGraphActionListBuilderBase::ActionGroup& DropActionOut) const;
 };
 
 /*******************************************************************************
diff --git a/Engine/Source/Editor/Kismet/Private/SMyBlueprint.cpp b/Engine/Source/Editor/Kismet/Private/SMyBlueprint.cpp
--- a/Engine/Source/Editor/Kismet/Private/SMyBlueprint.cpp
+++ b/Engine/Source/Editor/Kismet/Private/SMyBlueprint.cpp
@@ -1087,7 +1087,7 @@ FReply SMyBlueprint::OnActionDragged( const TArray< TSharedPtr<FEdGraphSchemaAct
 
 				if (bIsBlueprintCallableFunction)
 				{
-					return FReply::Handled().BeginDragDrop(FKismetFunctionDragDropAction::New(FuncAction->FuncName, BlueprintEditorPtr.Pin()->GetBlueprintObj()->SkeletonGeneratedClass, FMemberReference(), AnalyticsDelegate));
+					return FReply::Handled().BeginDragDrop(FKismetFunctionDragDropAction::New(InAction, FuncAction->FuncName, BlueprintEditorPtr.Pin()->GetBlueprintObj()->SkeletonGeneratedClass, FMemberReference(), AnalyticsDelegate));
 				}
 			}
 			else if (FuncAction->GraphType == EEdGraphSchemaAction_K2Graph::Macro)
