Blueprints - Making a base async node based on the LatentOnlineCall node.  From now on to create any async action just create a UObject that extends from  UBllueprintAsyncActionBase and a new node entry will be added that spawns that class and activates it.

diff --git a/Engine/Source/Editor/Kismet/Private/Nodes/K2Node_AsyncAction.cpp b/Engine/Source/Editor/Kismet/Private/Nodes/K2Node_AsyncAction.cpp
new file mode 100644
--- /dev/null
+++ b/Engine/Source/Editor/Kismet/Private/Nodes/K2Node_AsyncAction.cpp
@@ -0,0 +1,122 @@
+// Copyright 1998-2015 Epic Games, Inc. All Rights Reserved.
+
+#include "BlueprintEditorPrivatePCH.h"
+
+#include "Kismet/BlueprintAsyncActionBase.h"
+#include "BlueprintFunctionNodeSpawner.h"
+#include "BlueprintActionDatabaseRegistrar.h"
+
+#include "Nodes/K2Node_AsyncAction.h"
+
+#define LOCTEXT_NAMESPACE "K2Node"
+
+UK2Node_AsyncAction::UK2Node_AsyncAction(const FObjectInitializer& ObjectInitializer)
+	: Super(ObjectInitializer)
+{
+	ProxyActivateFunctionName = GET_FUNCTION_NAME_CHECKED(UBlueprintAsyncActionBase, Activate);
+}
+
+void UK2Node_AsyncAction::GetMenuEntries(FGraphContextMenuBuilder& ContextMenuBuilder) const
+{
+	const UEdGraphSchema_K2* K2Schema = GetDefault<UEdGraphSchema_K2>();
+	EGraphType GraphType = K2Schema->GetGraphType(ContextMenuBuilder.CurrentGraph);
+	const bool bAllowLatentFuncs = (GraphType == GT_Ubergraph || GraphType == GT_Macro);
+
+
+	if (bAllowLatentFuncs)
+	{
+		for (TObjectIterator<UClass> ClassIt; ClassIt; ++ClassIt)
+		{
+			UClass* TestClass = *ClassIt;
+			if ( TestClass->IsChildOf(UBlueprintAsyncActionBase::StaticClass()) && !TestClass->HasAnyClassFlags(CLASS_Abstract | CLASS_Deprecated) )
+			{
+				for (TFieldIterator<UFunction> FuncIt(TestClass, EFieldIteratorFlags::ExcludeSuper); FuncIt; ++FuncIt)
+				{
+					// See if the function is a static factory method for online proxies
+					UFunction* CandidateFunction = *FuncIt;
+
+					UObjectProperty* ReturnProperty = Cast<UObjectProperty>(CandidateFunction->GetReturnProperty());
+					const bool bValidReturnType = ( ReturnProperty != nullptr ) && ( ReturnProperty->PropertyClass != nullptr ) && ( ReturnProperty->PropertyClass->IsChildOf(UBlueprintAsyncActionBase::StaticClass()) );
+
+					if (CandidateFunction->HasAllFunctionFlags(FUNC_Static) && bValidReturnType)
+					{
+						// Create a node template for this factory method
+						UK2Node_AsyncAction* NodeTemplate = NewObject<UK2Node_AsyncAction>(ContextMenuBuilder.OwnerOfTemporaries);
+						NodeTemplate->ProxyFactoryFunctionName = CandidateFunction->GetFName();
+						NodeTemplate->ProxyFactoryClass = TestClass;
+						NodeTemplate->ProxyClass = ReturnProperty->PropertyClass;
+
+						CreateDefaultMenuEntry(NodeTemplate, ContextMenuBuilder);
+					}
+				}
+			}
+		}
+	}
+}
+
+void UK2Node_AsyncAction::GetMenuActions(FBlueprintActionDatabaseRegistrar& ActionRegistrar) const
+{
+	// these nested loops are combing over the same classes/functions the
+	// FBlueprintActionDatabase does; ideally we save on perf and fold this in
+	// with FBlueprintActionDatabase, but we want to keep the modules separate
+	for (TObjectIterator<UClass> ClassIt; ClassIt; ++ClassIt)
+	{
+		UClass* Class = *ClassIt;
+		if ( !Class->IsChildOf<UBlueprintAsyncActionBase>() || Class->HasAnyClassFlags(CLASS_Abstract) )
+		{
+			continue;
+		}
+		
+		for (TFieldIterator<UFunction> FuncIt(Class, EFieldIteratorFlags::ExcludeSuper); FuncIt; ++FuncIt)
+		{
+			UFunction* Function = *FuncIt;
+			if (!Function->HasAnyFunctionFlags(FUNC_Static))
+			{
+				continue;
+			}
+
+			// to keep from needlessly instantiating a UBlueprintNodeSpawner, first   
+			// check to make sure that the registrar is looking for actions of this type
+			// (could be regenerating actions for a specific asset, and therefore the 
+			// registrar would only accept actions corresponding to that asset)
+			if (!ActionRegistrar.IsOpenForRegistration(Function))
+			{
+				continue;
+			}
+			
+			UObjectProperty* ReturnProperty = Cast<UObjectProperty>(Function->GetReturnProperty());
+			// see if the function is a static factory method for online proxies
+			bool const bIsProxyFactoryMethod = ( ReturnProperty != nullptr ) && ReturnProperty->PropertyClass->IsChildOf<UBlueprintAsyncActionBase>();
+			
+			if (bIsProxyFactoryMethod)
+			{
+				UBlueprintNodeSpawner* NodeSpawner = UBlueprintFunctionNodeSpawner::Create(Function);
+				check(NodeSpawner != nullptr);
+				NodeSpawner->NodeClass = GetClass();
+				
+				auto CustomizeAcyncNodeLambda = [](UEdGraphNode* NewNode, bool bIsTemplateNode, TWeakObjectPtr<UFunction> FunctionPtr)
+				{
+					UK2Node_AsyncAction* AsyncTaskNode = CastChecked<UK2Node_AsyncAction>(NewNode);
+					if (FunctionPtr.IsValid())
+					{
+						UFunction* Function = FunctionPtr.Get();
+						UObjectProperty* ReturnProperty = CastChecked<UObjectProperty>(Function->GetReturnProperty());
+						
+						AsyncTaskNode->ProxyFactoryFunctionName = Function->GetFName();
+						AsyncTaskNode->ProxyFactoryClass        = Function->GetOuterUClass();
+						AsyncTaskNode->ProxyClass               = ReturnProperty->PropertyClass;
+					}
+				};
+				
+				TWeakObjectPtr<UFunction> FunctionPtr = Function;
+				NodeSpawner->CustomizeNodeDelegate = UBlueprintNodeSpawner::FCustomizeNodeDelegate::CreateStatic(CustomizeAcyncNodeLambda, FunctionPtr);
+				
+				// @TODO: since this can't be folded into FBlueprintActionDatabase, we
+				//        need a way to associate these spawners with a certain class
+				ActionRegistrar.AddBlueprintAction(Function, NodeSpawner);
+			}
+		}
+	}
+}
+
+#undef LOCTEXT_NAMESPACE
diff --git a/Engine/Source/Editor/Kismet/Public/Nodes/K2Node_AsyncAction.h b/Engine/Source/Editor/Kismet/Public/Nodes/K2Node_AsyncAction.h
new file mode 100644
--- /dev/null
+++ b/Engine/Source/Editor/Kismet/Public/Nodes/K2Node_AsyncAction.h
@@ -0,0 +1,16 @@
+// Copyright 1998-2015 Epic Games, Inc. All Rights Reserved.
+
+#pragma once
+
+#include "K2Node_AsyncAction.generated.h"
+
+UCLASS()
+class KISMET_API UK2Node_AsyncAction : public UK2Node_BaseAsyncTask
+{
+	GENERATED_UCLASS_BODY()
+	
+	// UEdGraphNode interface
+	virtual void GetMenuEntries(FGraphContextMenuBuilder& ContextMenuBuilder) const override;
+	virtual void GetMenuActions(FBlueprintActionDatabaseRegistrar& ActionRegistrar) const override;
+	// End of UEdGraphNode interface
+};
diff --git a/Engine/Source/Runtime/Engine/Classes/Kismet/BlueprintAsyncActionBase.h b/Engine/Source/Runtime/Engine/Classes/Kismet/BlueprintAsyncActionBase.h
new file mode 100644
--- /dev/null
+++ b/Engine/Source/Runtime/Engine/Classes/Kismet/BlueprintAsyncActionBase.h
@@ -0,0 +1,14 @@
+// Copyright 1998-2015 Epic Games, Inc. All Rights Reserved.
+#pragma once
+
+#include "BlueprintAsyncActionBase.generated.h"
+
+UCLASS()
+class ENGINE_API UBlueprintAsyncActionBase : public UObject
+{
+	GENERATED_UCLASS_BODY()
+
+	// Called to trigger the action once the delegates have been bound
+	UFUNCTION(BlueprintCallable, meta=(BlueprintInternalUseOnly="true"))
+	virtual void Activate();
+};
diff --git a/Engine/Source/Runtime/Engine/Private/Kismet/BlueprintAsyncActionBase.cpp b/Engine/Source/Runtime/Engine/Private/Kismet/BlueprintAsyncActionBase.cpp
new file mode 100644
--- /dev/null
+++ b/Engine/Source/Runtime/Engine/Private/Kismet/BlueprintAsyncActionBase.cpp
@@ -0,0 +1,16 @@
+// Copyright 1998-2015 Epic Games, Inc. All Rights Reserved.
+
+#include "EnginePrivate.h"
+#include "Kismet/BlueprintAsyncActionBase.h"
+
+//////////////////////////////////////////////////////////////////////////
+// UBlueprintAsyncActionBase
+
+UBlueprintAsyncActionBase::UBlueprintAsyncActionBase(const FObjectInitializer& ObjectInitializer)
+	: Super(ObjectInitializer)
+{
+}
+
+void UBlueprintAsyncActionBase::Activate()
+{
+}
