This updates the blutility function detection and enables the UI for GlobalEditorUtilityBase based BP's in addition to providing support for parameters with CallInEditor event blutilities.
#UE-3840 - BP: Blutilities UI broke in 4.4.
#Branch UE4
#Proj Blutility, DetailCustomizations, KismetCompiler, UnrealEd, CoreUObject

#Added FUNC_Blutility function flag to enable better blutility identification.
#Change Updated kismet compiler to tag blutility events as FUNC_Blutility
#Change Updated FBlueprintEditorUtils::CreateFunctionGraph to tag the function flags with FUNC_Blutility if the blueprint is based on a blutility.
#Change CallInEditor code previously detected through metadata was modified to use the new FUNC_Blutility flags.
#Change Modified the blutilities dialog to detect any UFunctions with FUNC_Blutility flags rather than FUNC_Exec which had previously been broken.
#Change Added a test in the blutility module to enable blueprint blutility detection and exposed a function to check this in BPEditorUtils

diff --git a/Engine/Source/Editor/Blutility/Private/BlutilityDetailsPanel.cpp b/Engine/Source/Editor/Blutility/Private/BlutilityDetailsPanel.cpp
--- a/Engine/Source/Editor/Blutility/Private/BlutilityDetailsPanel.cpp
+++ b/Engine/Source/Editor/Blutility/Private/BlutilityDetailsPanel.cpp
@@ -92,7 +92,8 @@ void FEditorUtilityInstanceDetails::CustomizeDetails(IDetailLayoutBuilder& Detai
 		{
 			UFunction* Function = *FuncIt;
 
-			const bool bCanExecute = (Function->NumParms == 0) && Function->HasAllFunctionFlags(FUNC_Exec);
+			const bool bCallInEditorFunc = Function->GetBoolMetaData( TEXT("CallInEditor") );
+			const bool bCanExecute = (Function->NumParms == 0) && bCallInEditorFunc;
 
 			if (bCanExecute)
 			{
diff --git a/Engine/Source/Editor/Blutility/Private/BlutilityModule.cpp b/Engine/Source/Editor/Blutility/Private/BlutilityModule.cpp
--- a/Engine/Source/Editor/Blutility/Private/BlutilityModule.cpp
+++ b/Engine/Source/Editor/Blutility/Private/BlutilityModule.cpp
@@ -74,6 +74,17 @@ public:
 		}
 	}
 
+	virtual bool IsBlutility( const UBlueprint* Blueprint ) const override
+	{
+		const UClass* BPClass = Blueprint ? Blueprint->GetClass() : nullptr;
+
+		if( BPClass && BPClass->IsChildOf( UEditorUtilityBlueprint::StaticClass() ))
+		{
+			return true;
+		}
+		return false;
+	}
+
 protected:
 	static TSharedRef<SDockTab> SpawnBlutilityShelfTab(const FSpawnTabArgs& Args)
 	{
diff --git a/Engine/Source/Editor/Blutility/Public/IBlutilityModule.h b/Engine/Source/Editor/Blutility/Public/IBlutilityModule.h
--- a/Engine/Source/Editor/Blutility/Public/IBlutilityModule.h
+++ b/Engine/Source/Editor/Blutility/Public/IBlutilityModule.h
@@ -11,5 +11,9 @@
 class IBlutilityModule : public IModuleInterface
 {
 public:
+
+	/** Returns if the blueprint is blutility based */
+	virtual bool IsBlutility( const UBlueprint* Blueprint ) const = 0;
+
 };
 
diff --git a/Engine/Source/Editor/DetailCustomizations/Private/ActorDetails.cpp b/Engine/Source/Editor/DetailCustomizations/Private/ActorDetails.cpp
--- a/Engine/Source/Editor/DetailCustomizations/Private/ActorDetails.cpp
+++ b/Engine/Source/Editor/DetailCustomizations/Private/ActorDetails.cpp
@@ -672,6 +672,17 @@ void FActorDetails::AddBlutilityCategory( IDetailLayoutBuilder& DetailBuilder, c
 				]
 			]
 		];
+		// Iterate over actor properties, adding property rows that are editable and may feed into the blutility function.
+		TArray<UObject*> ActorArray;
+		ActorArray.Add( ActorPtr.Get() );
+		for( TFieldIterator<UProperty> PropertyIt(ActorPtr->GetClass()); PropertyIt; ++PropertyIt )
+		{
+			const bool bBlueprintProperty = PropertyIt->GetOuter()->GetClass() == UBlueprintGeneratedClass::StaticClass();
+			if( bBlueprintProperty && !PropertyIt->HasAllPropertyFlags( CPF_DisableEditOnInstance ))
+			{
+				BlutilitiesCategory.AddExternalProperty( ActorArray, PropertyIt->GetFName(), EPropertyLocation::Advanced );
+			}
+		}
 	}
 }
 
@@ -1493,7 +1504,7 @@ bool FActorDetails::DoesActorHaveBlutiltyFunctions() const
 
 		if( ActorClass )
 		{
-			for (TFieldIterator<UFunction> FunctionIter(ActorClass, EFieldIteratorFlags::IncludeSuper); FunctionIter; ++FunctionIter)
+			for( TFieldIterator<UFunction> FunctionIter(ActorClass, EFieldIteratorFlags::IncludeSuper); FunctionIter; ++FunctionIter )
 			{
 				if( FunctionIter->GetBoolMetaData( FBlueprintMetadata::MD_CallInEditor ))
 				{
@@ -1552,12 +1563,7 @@ FReply FActorDetails::CallBlutilityFunction()
 	if( ActorWeakPtr.IsValid() && ActiveBlutilityFunction.IsValid() )
 	{
 		AActor* Actor = ActorWeakPtr.Get();
-		UFunction* Function = ActiveBlutilityFunction.Get();
-
-		if( Function->GetOuter() == Actor->GetClass() )
-		{
-			Actor->ProcessEvent( Function, NULL );
-		}
+		Actor->ProcessEvent( ActiveBlutilityFunction.Get(), NULL );
 	}
 
 	return FReply::Handled();
diff --git a/Engine/Source/Editor/Kismet/Private/BlueprintDetailsCustomization.cpp b/Engine/Source/Editor/Kismet/Private/BlueprintDetailsCustomization.cpp
--- a/Engine/Source/Editor/Kismet/Private/BlueprintDetailsCustomization.cpp
+++ b/Engine/Source/Editor/Kismet/Private/BlueprintDetailsCustomization.cpp
@@ -2372,6 +2372,10 @@ void FBlueprintGraphActionDetails::CustomizeDetails( IDetailLayoutBuilder& Detai
 					]
 				]
 			];
+		}
+		const bool bShowCallInEditor = IsCustomEvent() || FBlueprintEditorUtils::IsBlutility( GetBlueprintObj() );
+		if( bShowCallInEditor )
+		{
 			Category.AddCustomRow( LOCTEXT( "EditorCallable", "Call In Editor" ) )
 			.NameContent()
 			[
@@ -2582,11 +2586,19 @@ ECheckBoxState FBlueprintGraphActionDetails::GetIsEditorCallableEvent() const
 
 	if( FunctionEntryNodePtr.IsValid() )
 	{
-		UK2Node_CustomEvent* CustomEventNode = Cast<UK2Node_CustomEvent>(FunctionEntryNodePtr.Get());
-
-		if( CustomEventNode && CustomEventNode->bCallInEditor )
+		if( UK2Node_CustomEvent* CustomEventNode = Cast<UK2Node_CustomEvent>(FunctionEntryNodePtr.Get()))
+		{
+			if( CustomEventNode->bCallInEditor  )
+			{
+				Result = ECheckBoxState::Checked;
+			}
+		}
+		else if( UK2Node_FunctionEntry* EntryPoint = Cast<UK2Node_FunctionEntry>(FunctionEntryNodePtr.Get()) )
 		{
-			Result = ECheckBoxState::Checked;
+			if( EntryPoint->MetaData.bCallInEditor )
+			{
+				Result = ECheckBoxState::Checked;
+			}
 		}
 	}
 	return Result;
@@ -2596,18 +2608,25 @@ void FBlueprintGraphActionDetails::OnEditorCallableEventModified( const ECheckBo
 {
 	if( FunctionEntryNodePtr.IsValid() )
 	{
+		const bool bCallInEditor = NewCheckedState == ECheckBoxState::Checked;
+		const FText TransactionType = bCallInEditor ?	LOCTEXT( "DisableCallInEditor", "Disable Call In Editor " ) : 
+														LOCTEXT( "EnableCallInEditor", "Enable Call In Editor" );
+
 		if( UK2Node_CustomEvent* CustomEventNode = Cast<UK2Node_CustomEvent>(FunctionEntryNodePtr.Get()) )
 		{
 			if( UBlueprint* Blueprint = FunctionEntryNodePtr->GetBlueprint() )
 			{
-				const bool bCallInEditor = NewCheckedState == ECheckBoxState::Checked;
-				const FText TransactionType = bCallInEditor ?	LOCTEXT( "DisableCallInEditor", "Disable Call In Editor " ) : 
-																LOCTEXT( "EnableCallInEditor", "Enable Call In Editor" );
 				const FScopedTransaction Transaction( TransactionType );
 				CustomEventNode->bCallInEditor = bCallInEditor;
 				FBlueprintEditorUtils::MarkBlueprintAsModified( CustomEventNode->GetBlueprint() );
 			}
 		}
+		else if( UK2Node_FunctionEntry* EntryPoint = Cast<UK2Node_FunctionEntry>(FunctionEntryNodePtr.Get()) )
+		{
+			const FScopedTransaction Transaction( TransactionType );
+			EntryPoint->MetaData.bCallInEditor = bCallInEditor;
+			FBlueprintEditorUtils::MarkBlueprintAsModified( EntryPoint->GetBlueprint() );
+		}
 	}
 }
 
diff --git a/Engine/Source/Editor/KismetCompiler/Private/KismetCompiler.cpp b/Engine/Source/Editor/KismetCompiler/Private/KismetCompiler.cpp
--- a/Engine/Source/Editor/KismetCompiler/Private/KismetCompiler.cpp
+++ b/Engine/Source/Editor/KismetCompiler/Private/KismetCompiler.cpp
@@ -1186,6 +1186,7 @@ void FKismetCompilerContext::PrecompileFunction(FKismetFunctionContext& Context)
 		{
 			Context.Function->SetMetaData(FBlueprintMetadata::MD_FunctionCategory, *FunctionMetaData.Category);
 		}
+		// Set as blutility function
 		if( FunctionMetaData.bCallInEditor )
 		{
 			Context.Function->SetMetaData(FBlueprintMetadata::MD_CallInEditor, TEXT( "true" ));
diff --git a/Engine/Source/Editor/UnrealEd/Private/Kismet2/BlueprintEditorUtils.cpp b/Engine/Source/Editor/UnrealEd/Private/Kismet2/BlueprintEditorUtils.cpp
--- a/Engine/Source/Editor/UnrealEd/Private/Kismet2/BlueprintEditorUtils.cpp
+++ b/Engine/Source/Editor/UnrealEd/Private/Kismet2/BlueprintEditorUtils.cpp
@@ -44,6 +44,7 @@
 
 #include "SNotificationList.h"
 #include "NotificationManager.h"
+#include "Editor/Blutility/Public/IBlutilityModule.h"
 
 #define LOCTEXT_NAMESPACE "Blueprint"
 
@@ -2392,6 +2393,12 @@ bool FBlueprintEditorUtils::IsBlueprintConst(const UBlueprint* Blueprint)
 	return Blueprint->BlueprintType == BPTYPE_Const;
 }
 
+bool FBlueprintEditorUtils::IsBlutility(const UBlueprint* Blueprint)
+{
+	IBlutilityModule& BlutilityModule = FModuleManager::GetModuleChecked<IBlutilityModule>("Blutility");
+	return BlutilityModule.IsBlutility( Blueprint );
+}
+
 bool FBlueprintEditorUtils::IsActorBased(const UBlueprint* Blueprint)
 {
 	return Blueprint->ParentClass && Blueprint->ParentClass->IsChildOf(AActor::StaticClass());
@@ -6888,6 +6895,23 @@ void FBlueprintEditorUtils::GetEntryAndResultNodes(const UEdGraph* InGraph, TWea
 	}
 }
 
+FKismetUserDeclaredFunctionMetadata* FBlueprintEditorUtils::GetGraphFunctionMetaData( UEdGraph* InGraph )
+{
+	FKismetUserDeclaredFunctionMetadata* FunctionMetaData = nullptr;
+
+	if( InGraph )
+	{
+		TArray<UK2Node_FunctionEntry*> EntryNodes;
+		InGraph->GetNodesOfClass( EntryNodes );
+
+		if( EntryNodes.Num() )
+		{
+			FunctionMetaData = &EntryNodes[ 0 ]->MetaData;
+		}
+	}
+	return FunctionMetaData;
+}
+
 FText FBlueprintEditorUtils::GetGraphDescription(const UEdGraph* InGraph)
 {
 	TWeakObjectPtr<class UK2Node_EditablePinBase> EntryNode;
diff --git a/Engine/Source/Editor/UnrealEd/Public/Kismet2/BlueprintEditorUtils.h b/Engine/Source/Editor/UnrealEd/Public/Kismet2/BlueprintEditorUtils.h
--- a/Engine/Source/Editor/UnrealEd/Public/Kismet2/BlueprintEditorUtils.h
+++ b/Engine/Source/Editor/UnrealEd/Public/Kismet2/BlueprintEditorUtils.h
@@ -7,6 +7,7 @@
 #include "EdGraphSchema_K2.h"
 
 class USCS_Node;
+struct FKismetUserDeclaredFunctionMetadata;
 
 /** 
   * Flags describing how to handle graph removal
@@ -193,6 +194,13 @@ public:
 				{
 					ExtraFunctionFlags |= FUNC_Static;
 				}
+				if( IsBlutility( Blueprint ))
+				{
+					if( FKismetUserDeclaredFunctionMetadata* MetaData = GetGraphFunctionMetaData( Graph ))
+					{
+						MetaData->bCallInEditor = true;
+					}
+				}
 				K2Schema->AddExtraFunctionFlags(Graph, ExtraFunctionFlags);
 				K2Schema->MarkFunctionEntryAsEditable(Graph, true);
 			}
@@ -339,6 +347,9 @@ public:
 	/** Returns whether or not the blueprint is const during execution */
 	static bool IsBlueprintConst(const UBlueprint* Blueprint);
 
+	/** Returns whether or not the blueprint is a blutility */
+	static bool IsBlutility(const UBlueprint* Blueprint);
+
 	/**
 	 * Whether or not this is an actor-based blueprint, and supports features like the uber-graph, components, etc
 	 *
@@ -1034,6 +1045,14 @@ public:
 	 */
 	static void GetEntryAndResultNodes(const UEdGraph* InGraph, TWeakObjectPtr<class UK2Node_EditablePinBase>& OutEntryNode, TWeakObjectPtr<class UK2Node_EditablePinBase>& OutResultNode);
 
+	/**
+	 * Returns the function meta data block for the graph entry node.
+	 *
+	 * @param InGraph			The graph to search through
+	 * @return					If valid a pointer to the user declared function meta data structure otherwise nullptr.
+	 */
+	static FKismetUserDeclaredFunctionMetadata* GetGraphFunctionMetaData( UEdGraph* InGraph );
+
 	/**
 	 * Returns the description of the graph from the metadata
 	 *
