Fixed malformed display string (UE-7124)

diff --git a/Engine/Source/Editor/Kismet/Private/DiffUtils.cpp b/Engine/Source/Editor/Kismet/Private/DiffUtils.cpp
--- a/Engine/Source/Editor/Kismet/Private/DiffUtils.cpp
+++ b/Engine/Source/Editor/Kismet/Private/DiffUtils.cpp
@@ -528,7 +528,7 @@ FText DiffViewUtils::SCSDiffMessage(const FSCSDiffEntry& Difference, FText Objec
 		Text = FText::Format(NSLOCTEXT("DiffViewUtils", "NodeAdded", "Removed Node {0} from {1}"), NodeName, ObjectName);
 		break;
 	case ETreeDiffType::NODE_PROPERTY_CHANGED:
-		Text = FText::Format(NSLOCTEXT("DiffViewUtils", "NodeAdded", "{0} on {2}"), DiffViewUtils::PropertyDiffMessage(Difference.PropertyDiff, NodeName), ObjectName);
+		Text = FText::Format(NSLOCTEXT("DiffViewUtils", "NodeAdded", "{0} on {1}"), DiffViewUtils::PropertyDiffMessage(Difference.PropertyDiff, NodeName), ObjectName);
 		break;
 	case ETreeDiffType::NODE_MOVED:
 		Text = FText::Format(NSLOCTEXT("DiffViewUtils", "NodeAdded", "Moved Node {0} in {1}"), NodeName, ObjectName);
