BP - Adding experimental application mode that displays BP editor in a single screen.

diff --git a/Engine/Source/Editor/Kismet/Private/BlueprintEditor.cpp b/Engine/Source/Editor/Kismet/Private/BlueprintEditor.cpp
--- a/Engine/Source/Editor/Kismet/Private/BlueprintEditor.cpp
+++ b/Engine/Source/Editor/Kismet/Private/BlueprintEditor.cpp
@@ -1429,27 +1429,37 @@ void FBlueprintEditor::RegisterApplicationModes(const TArray<UBlueprint*>& InBlu
 		}
 		else
 		{
-			AddApplicationMode(
-				FBlueprintEditorApplicationModes::StandardBlueprintEditorMode,
-				MakeShareable(new FBlueprintEditorApplicationMode(SharedThis(this), FBlueprintEditorApplicationModes::StandardBlueprintEditorMode, FBlueprintEditorApplicationModes::GetLocalizedMode)));
-			AddApplicationMode(
-				FBlueprintEditorApplicationModes::BlueprintDefaultsMode,
-				MakeShareable(new FBlueprintDefaultsApplicationMode(SharedThis(this))));
-			AddApplicationMode(
-				FBlueprintEditorApplicationModes::BlueprintComponentsMode,
-				MakeShareable(new FBlueprintComponentsApplicationMode(SharedThis(this))));
-
-			if ( bShouldOpenInDefaultsMode )
-			{
-				SetCurrentMode(FBlueprintEditorApplicationModes::BlueprintDefaultsMode);
-			}
-			else if ( bShouldOpenInComponentsMode && CanAccessComponentsMode() )
+			if ( GetDefault<UEditorExperimentalSettings>()->bInWorldBPEditing )
 			{
-				SetCurrentMode(FBlueprintEditorApplicationModes::BlueprintComponentsMode);
+				AddApplicationMode(
+					FBlueprintEditorApplicationModes::StandardBlueprintEditorMode,
+					MakeShareable(new FBlueprintEditorUnifiedMode(SharedThis(this), FBlueprintEditorApplicationModes::StandardBlueprintEditorMode, FBlueprintEditorApplicationModes::GetLocalizedMode)));
+				SetCurrentMode(FBlueprintEditorApplicationModes::StandardBlueprintEditorMode);
 			}
 			else
 			{
-				SetCurrentMode(FBlueprintEditorApplicationModes::StandardBlueprintEditorMode);
+				AddApplicationMode(
+					FBlueprintEditorApplicationModes::StandardBlueprintEditorMode,
+					MakeShareable(new FBlueprintEditorApplicationMode(SharedThis(this), FBlueprintEditorApplicationModes::StandardBlueprintEditorMode, FBlueprintEditorApplicationModes::GetLocalizedMode)));
+				AddApplicationMode(
+					FBlueprintEditorApplicationModes::BlueprintDefaultsMode,
+					MakeShareable(new FBlueprintDefaultsApplicationMode(SharedThis(this))));
+				AddApplicationMode(
+					FBlueprintEditorApplicationModes::BlueprintComponentsMode,
+					MakeShareable(new FBlueprintComponentsApplicationMode(SharedThis(this))));
+
+				if ( bShouldOpenInDefaultsMode )
+				{
+					SetCurrentMode(FBlueprintEditorApplicationModes::BlueprintDefaultsMode);
+				}
+				else if ( bShouldOpenInComponentsMode && CanAccessComponentsMode() )
+				{
+					SetCurrentMode(FBlueprintEditorApplicationModes::BlueprintComponentsMode);
+				}
+				else
+				{
+					SetCurrentMode(FBlueprintEditorApplicationModes::StandardBlueprintEditorMode);
+				}
 			}
 		}
 	}
diff --git a/Engine/Source/Editor/Kismet/Private/BlueprintEditorModes.cpp b/Engine/Source/Editor/Kismet/Private/BlueprintEditorModes.cpp
--- a/Engine/Source/Editor/Kismet/Private/BlueprintEditorModes.cpp
+++ b/Engine/Source/Editor/Kismet/Private/BlueprintEditorModes.cpp
@@ -519,4 +519,133 @@ void FBlueprintMacroApplicationMode::RegisterTabFactories(TSharedPtr<FTabManager
 	FApplicationMode::RegisterTabFactories(InTabManager);
 }
 
+FBlueprintEditorUnifiedMode::FBlueprintEditorUnifiedMode(TSharedPtr<class FBlueprintEditor> InBlueprintEditor, FName InModeName, FText(*GetLocalizedMode)( const FName ), const bool bRegisterViewport, const bool bRegisterDefaultsTab)
+	: FApplicationMode(InModeName, GetLocalizedMode)
+{
+	MyBlueprintEditor = InBlueprintEditor;
+
+	// Create the tab factories
+	BlueprintEditorTabFactories.RegisterFactory(MakeShareable(new FDebugInfoSummoner(InBlueprintEditor)));
+	BlueprintEditorTabFactories.RegisterFactory(MakeShareable(new FPaletteSummoner(InBlueprintEditor)));
+	BlueprintEditorTabFactories.RegisterFactory(MakeShareable(new FMyBlueprintSummoner(InBlueprintEditor)));
+	BlueprintEditorTabFactories.RegisterFactory(MakeShareable(new FCompilerResultsSummoner(InBlueprintEditor)));
+	BlueprintEditorTabFactories.RegisterFactory(MakeShareable(new FFindResultsSummoner(InBlueprintEditor)));
+	
+	if( bRegisterViewport )
+	{
+		BlueprintEditorTabFactories.RegisterFactory(MakeShareable(new FConstructionScriptEditorSummoner(InBlueprintEditor)));
+		BlueprintEditorTabFactories.RegisterFactory(MakeShareable(new FSCSViewportSummoner(InBlueprintEditor)));
+	}
+	if( bRegisterDefaultsTab )
+	{
+		BlueprintEditorTabFactories.RegisterFactory(MakeShareable(new FDefaultsEditorSummoner(InBlueprintEditor)));
+	}
+	CoreTabFactories.RegisterFactory(MakeShareable(new FSelectionDetailsSummoner(InBlueprintEditor)));
+
+	TabLayout = FTabManager::NewLayout( "Standalone_BlueprintUnifiedEditor_Layout_v1_1" )
+		->AddArea
+		(
+			FTabManager::NewPrimaryArea() ->SetOrientation(Orient_Vertical)
+			->Split
+			(
+				FTabManager::NewSplitter() ->SetOrientation(Orient_Horizontal)
+				->Split
+				(
+					FTabManager::NewSplitter() ->SetOrientation( Orient_Vertical )
+					->SetSizeCoefficient(0.15f)
+					->Split
+					(
+						FTabManager::NewStack()
+						->SetSizeCoefficient(0.25f)
+						->AddTab( FBlueprintEditorTabs::ConstructionScriptEditorID, ETabState::OpenedTab )
+					)
+					->Split
+					(
+						FTabManager::NewStack()->SetSizeCoefficient(0.75f)
+						->AddTab( FBlueprintEditorTabs::MyBlueprintID, ETabState::OpenedTab )
+					)
+				)
+				->Split
+				(
+					FTabManager::NewSplitter()->SetOrientation( Orient_Vertical )
+					->SetSizeCoefficient(0.70f)
+					->Split
+					(
+						FTabManager::NewStack()
+						->SetSizeCoefficient(0.186721f)
+						->SetHideTabWell(true)
+						->AddTab(InBlueprintEditor->GetToolbarTabId(), ETabState::OpenedTab)
+					)
+					->Split
+					(
+						FTabManager::NewStack()
+						->SetSizeCoefficient( 0.80f )
+						->AddTab(FBlueprintEditorTabs::SCSViewportID, ETabState::OpenedTab)
+						->AddTab( "Document", ETabState::ClosedTab )
+					)
+					->Split
+					(
+						FTabManager::NewStack()
+						->SetSizeCoefficient( 0.20f )
+						->AddTab( FBlueprintEditorTabs::CompilerResultsID, ETabState::ClosedTab )
+						->AddTab( FBlueprintEditorTabs::FindResultsID, ETabState::ClosedTab )
+					)
+				)
+				->Split
+				(
+					FTabManager::NewSplitter() ->SetOrientation( Orient_Vertical )
+					->SetSizeCoefficient(0.15f)
+					->Split
+					(
+						FTabManager::NewStack()
+						->AddTab( FBlueprintEditorTabs::DetailsID, ETabState::OpenedTab )
+					)
+				)
+			)
+		);
+	
+	// setup toolbar
+	//@TODO: Keep this in sync with AnimBlueprintMode.cpp
+	//InBlueprintEditor->GetToolbarBuilder()->AddBlueprintEditorModesToolbar(ToolbarExtender);
+	InBlueprintEditor->GetToolbarBuilder()->AddCompileToolbar(ToolbarExtender);
+	InBlueprintEditor->GetToolbarBuilder()->AddScriptingToolbar(ToolbarExtender);
+	InBlueprintEditor->GetToolbarBuilder()->AddBlueprintGlobalOptionsToolbar(ToolbarExtender);
+	InBlueprintEditor->GetToolbarBuilder()->AddDebuggingToolbar(ToolbarExtender);
+	InBlueprintEditor->GetToolbarBuilder()->AddComponentsToolbar(ToolbarExtender);
+}
+
+void FBlueprintEditorUnifiedMode::RegisterTabFactories(TSharedPtr<FTabManager> InTabManager)
+{
+	TSharedPtr<FBlueprintEditor> BP = MyBlueprintEditor.Pin();
+	
+	BP->RegisterToolbarTab(InTabManager.ToSharedRef());
+
+	// Mode-specific setup
+	BP->PushTabFactories(CoreTabFactories);
+	BP->PushTabFactories(BlueprintEditorOnlyTabFactories);
+	BP->PushTabFactories(BlueprintEditorTabFactories);
+
+	FApplicationMode::RegisterTabFactories(InTabManager);
+}
+
+void FBlueprintEditorUnifiedMode::PreDeactivateMode()
+{
+	FApplicationMode::PreDeactivateMode();
+
+	TSharedPtr<FBlueprintEditor> BP = MyBlueprintEditor.Pin();
+	
+	BP->SaveEditedObjectState();
+	BP->GetMyBlueprintWidget()->ClearGraphActionMenuSelection();
+}
+
+void FBlueprintEditorUnifiedMode::PostActivateMode()
+{
+	// Reopen any documents that were open when the blueprint was last saved
+	TSharedPtr<FBlueprintEditor> BP = MyBlueprintEditor.Pin();
+	BP->RestoreEditedObjectState();
+	BP->SetupViewForBlueprintEditingMode();
+
+	FApplicationMode::PostActivateMode();
+}
+
 #undef LOCTEXT_NAMESPACE
diff --git a/Engine/Source/Editor/Kismet/Public/BlueprintEditorModes.h b/Engine/Source/Editor/Kismet/Public/BlueprintEditorModes.h
--- a/Engine/Source/Editor/Kismet/Public/BlueprintEditorModes.h
+++ b/Engine/Source/Editor/Kismet/Public/BlueprintEditorModes.h
@@ -127,4 +127,27 @@ protected:
 
 	// Set of spawnable tabs in this mode
 	FWorkflowAllowedTabSet BlueprintMacroTabFactories;
-};
\ No newline at end of file
+};
+
+class KISMET_API FBlueprintEditorUnifiedMode : public FApplicationMode
+{
+public:
+	FBlueprintEditorUnifiedMode(TSharedPtr<class FBlueprintEditor> InBlueprintEditor, FName InModeName, FText(*GetLocalizedMode)( const FName ), const bool bRegisterViewport = true, const bool bRegisterDefaultsTab = true);
+
+	virtual void RegisterTabFactories(TSharedPtr<FTabManager> InTabManager) override;
+	virtual void PreDeactivateMode() override;
+	virtual void PostActivateMode() override;
+public:
+
+protected:
+	TWeakPtr<FBlueprintEditor> MyBlueprintEditor;
+
+	// Set of spawnable tabs in blueprint editing mode
+	FWorkflowAllowedTabSet BlueprintEditorTabFactories;
+
+	// Set of spawnable tabs useful in derived classes, even without a blueprint
+	FWorkflowAllowedTabSet CoreTabFactories;
+
+	// Set of spawnable tabs only usable in blueprint editing mode (not useful in Persona, etc...)
+	FWorkflowAllowedTabSet BlueprintEditorOnlyTabFactories;
+};
