BP - More work on the unified BP view, fixing some crashes and bad layout in other editors like the level BP script editor.

diff --git a/Engine/Source/Editor/Kismet/Private/BlueprintEditor.cpp b/Engine/Source/Editor/Kismet/Private/BlueprintEditor.cpp
--- a/Engine/Source/Editor/Kismet/Private/BlueprintEditor.cpp
+++ b/Engine/Source/Editor/Kismet/Private/BlueprintEditor.cpp
@@ -1433,7 +1433,7 @@ void FBlueprintEditor::RegisterApplicationModes(const TArray<UBlueprint*>& InBlu
 			{
 				AddApplicationMode(
 					FBlueprintEditorApplicationModes::StandardBlueprintEditorMode,
-					MakeShareable(new FBlueprintEditorUnifiedMode(SharedThis(this), FBlueprintEditorApplicationModes::StandardBlueprintEditorMode, FBlueprintEditorApplicationModes::GetLocalizedMode)));
+					MakeShareable(new FBlueprintEditorUnifiedMode(SharedThis(this), FBlueprintEditorApplicationModes::StandardBlueprintEditorMode, FBlueprintEditorApplicationModes::GetLocalizedMode, bShouldOpenInComponentsMode && CanAccessComponentsMode())));
 				SetCurrentMode(FBlueprintEditorApplicationModes::StandardBlueprintEditorMode);
 			}
 			else
diff --git a/Engine/Source/Editor/Kismet/Private/BlueprintEditorModes.cpp b/Engine/Source/Editor/Kismet/Private/BlueprintEditorModes.cpp
--- a/Engine/Source/Editor/Kismet/Private/BlueprintEditorModes.cpp
+++ b/Engine/Source/Editor/Kismet/Private/BlueprintEditorModes.cpp
@@ -41,6 +41,69 @@ const FName FBlueprintEditorApplicationModes::BlueprintComponentsMode( TEXT("Com
 const FName FBlueprintEditorApplicationModes::BlueprintInterfaceMode( TEXT("InterfaceName") );
 const FName FBlueprintEditorApplicationModes::BlueprintMacroMode( TEXT("MacroName") );
 
+TSharedPtr<FTabManager::FLayout> GetDefaltEditorLayout(TSharedPtr<class FBlueprintEditor> InBlueprintEditor)
+{
+	return FTabManager::NewLayout( "Standalone_BlueprintEditor_Layout_v5" )
+	->AddArea
+	(
+		FTabManager::NewPrimaryArea() ->SetOrientation(Orient_Vertical)
+		->Split
+		(
+			FTabManager::NewStack()
+			->SetSizeCoefficient( 0.186721f )
+			->SetHideTabWell(true)
+			->AddTab( InBlueprintEditor->GetToolbarTabId(), ETabState::OpenedTab )
+		)
+		->Split
+		(
+			FTabManager::NewSplitter() ->SetOrientation(Orient_Horizontal)
+			->Split
+			(
+				FTabManager::NewSplitter() ->SetOrientation( Orient_Vertical )
+				->SetSizeCoefficient(0.15f)
+				->Split
+				(
+					FTabManager::NewStack() ->SetSizeCoefficient(0.5f)
+					->AddTab( FBlueprintEditorTabs::MyBlueprintID, ETabState::OpenedTab )
+				)
+				->Split
+				(
+					FTabManager::NewStack() ->SetSizeCoefficient(0.5f)
+					->AddTab( FBlueprintEditorTabs::DetailsID, ETabState::OpenedTab )
+				)
+			)
+			->Split
+			(
+				FTabManager::NewSplitter() ->SetOrientation( Orient_Vertical )
+				->SetSizeCoefficient(0.70f)
+				->Split
+				(
+					FTabManager::NewStack()
+					->SetSizeCoefficient( 0.80f )
+					->AddTab( "Document", ETabState::ClosedTab )
+				)
+				->Split
+				(
+					FTabManager::NewStack()
+					->SetSizeCoefficient( 0.20f )
+					->AddTab( FBlueprintEditorTabs::CompilerResultsID, ETabState::ClosedTab )
+					->AddTab( FBlueprintEditorTabs::FindResultsID, ETabState::ClosedTab )
+				)
+			)
+			->Split
+			(
+				FTabManager::NewSplitter() ->SetOrientation( Orient_Vertical )
+				->SetSizeCoefficient(0.15f)
+				->Split
+				(
+					FTabManager::NewStack()
+					->AddTab( FBlueprintEditorTabs::PaletteID, ETabState::ClosedTab )
+				)
+			)
+		)
+	);
+}
+
 FBlueprintEditorApplicationMode::FBlueprintEditorApplicationMode(TSharedPtr<class FBlueprintEditor> InBlueprintEditor, FName InModeName, FText(*GetLocalizedMode)(const FName), const bool bRegisterViewport, const bool bRegisterDefaultsTab)
 	: FApplicationMode(InModeName, GetLocalizedMode)
 {
@@ -63,65 +126,7 @@ FBlueprintEditorApplicationMode::FBlueprintEditorApplicationMode(TSharedPtr<clas
 	}
 	CoreTabFactories.RegisterFactory(MakeShareable(new FSelectionDetailsSummoner(InBlueprintEditor)));
 
-	TabLayout = FTabManager::NewLayout( "Standalone_BlueprintEditor_Layout_v5" )
-		->AddArea
-		(
-			FTabManager::NewPrimaryArea() ->SetOrientation(Orient_Vertical)
-			->Split
-			(
-				FTabManager::NewStack()
-				->SetSizeCoefficient( 0.186721f )
-				->SetHideTabWell(true)
-				->AddTab( InBlueprintEditor->GetToolbarTabId(), ETabState::OpenedTab )
-			)
-			->Split
-			(
-				FTabManager::NewSplitter() ->SetOrientation(Orient_Horizontal)
-				->Split
-				(
-					FTabManager::NewSplitter() ->SetOrientation( Orient_Vertical )
-					->SetSizeCoefficient(0.15f)
-					->Split
-					(
-						FTabManager::NewStack() ->SetSizeCoefficient(0.5f)
-						->AddTab( FBlueprintEditorTabs::MyBlueprintID, ETabState::OpenedTab )
-					)
-					->Split
-					(
-						FTabManager::NewStack() ->SetSizeCoefficient(0.5f)
-						->AddTab( FBlueprintEditorTabs::DetailsID, ETabState::OpenedTab )
-					)
-				)
-				->Split
-				(
-					FTabManager::NewSplitter() ->SetOrientation( Orient_Vertical )
-					->SetSizeCoefficient(0.70f)
-					->Split
-					(
-						FTabManager::NewStack()
-						->SetSizeCoefficient( 0.80f )
-						->AddTab( "Document", ETabState::ClosedTab )
-					)
-					->Split
-					(
-						FTabManager::NewStack()
-						->SetSizeCoefficient( 0.20f )
-						->AddTab( FBlueprintEditorTabs::CompilerResultsID, ETabState::ClosedTab )
-						->AddTab( FBlueprintEditorTabs::FindResultsID, ETabState::ClosedTab )
-					)
-				)
-				->Split
-				(
-					FTabManager::NewSplitter() ->SetOrientation( Orient_Vertical )
-					->SetSizeCoefficient(0.15f)
-					->Split
-					(
-						FTabManager::NewStack()
-						->AddTab( FBlueprintEditorTabs::PaletteID, ETabState::ClosedTab )
-					)
-				)
-			)
-		);
+	TabLayout = GetDefaltEditorLayout(InBlueprintEditor);
 	
 	// setup toolbar
 	//@TODO: Keep this in sync with AnimBlueprintMode.cpp
@@ -542,7 +547,9 @@ FBlueprintEditorUnifiedMode::FBlueprintEditorUnifiedMode(TSharedPtr<class FBluep
 	}
 	CoreTabFactories.RegisterFactory(MakeShareable(new FSelectionDetailsSummoner(InBlueprintEditor)));
 
-	TabLayout = FTabManager::NewLayout( "Standalone_BlueprintUnifiedEditor_Layout_v1_1" )
+	if ( bRegisterViewport )
+	{
+		TabLayout = FTabManager::NewLayout( "Standalone_BlueprintUnifiedEditor_Layout_v1_1" )
 		->AddArea
 		(
 			FTabManager::NewPrimaryArea() ->SetOrientation(Orient_Vertical)
@@ -603,6 +610,11 @@ FBlueprintEditorUnifiedMode::FBlueprintEditorUnifiedMode(TSharedPtr<class FBluep
 				)
 			)
 		);
+	}
+	else
+	{
+		TabLayout = GetDefaltEditorLayout(InBlueprintEditor);
+	}
 	
 	// setup toolbar
 	//@TODO: Keep this in sync with AnimBlueprintMode.cpp
@@ -611,7 +623,11 @@ FBlueprintEditorUnifiedMode::FBlueprintEditorUnifiedMode(TSharedPtr<class FBluep
 	InBlueprintEditor->GetToolbarBuilder()->AddScriptingToolbar(ToolbarExtender);
 	InBlueprintEditor->GetToolbarBuilder()->AddBlueprintGlobalOptionsToolbar(ToolbarExtender);
 	InBlueprintEditor->GetToolbarBuilder()->AddDebuggingToolbar(ToolbarExtender);
-	InBlueprintEditor->GetToolbarBuilder()->AddComponentsToolbar(ToolbarExtender);
+
+	if ( bRegisterViewport )
+	{
+		InBlueprintEditor->GetToolbarBuilder()->AddComponentsToolbar(ToolbarExtender);
+	}
 }
 
 void FBlueprintEditorUnifiedMode::RegisterTabFactories(TSharedPtr<FTabManager> InTabManager)
diff --git a/Engine/Source/Editor/Kismet/Private/SBlueprintEditorToolbar.cpp b/Engine/Source/Editor/Kismet/Private/SBlueprintEditorToolbar.cpp
--- a/Engine/Source/Editor/Kismet/Private/SBlueprintEditorToolbar.cpp
+++ b/Engine/Source/Editor/Kismet/Private/SBlueprintEditorToolbar.cpp
@@ -400,13 +400,16 @@ static void BlueprintEditorToolbarImpl::MakeCompileDeveloperSubMenu(FMenuBuilder
 
 void FBlueprintEditorToolbar::AddBlueprintEditorModesToolbar(TSharedPtr<FExtender> Extender)
 {
-	TSharedPtr<FBlueprintEditor> BlueprintEditorPtr = BlueprintEditor.Pin();
+	if ( !GetDefault<UEditorExperimentalSettings>()->bInWorldBPEditing )
+	{
+		TSharedPtr<FBlueprintEditor> BlueprintEditorPtr = BlueprintEditor.Pin();
 
-	Extender->AddToolBarExtension(
-		"Asset",
-		EExtensionHook::After,
-		BlueprintEditorPtr->GetToolkitCommands(),
-		FToolBarExtensionDelegate::CreateSP( this, &FBlueprintEditorToolbar::FillBlueprintEditorModesToolbar ) );
+		Extender->AddToolBarExtension(
+			"Asset",
+			EExtensionHook::After,
+			BlueprintEditorPtr->GetToolkitCommands(),
+			FToolBarExtensionDelegate::CreateSP(this, &FBlueprintEditorToolbar::FillBlueprintEditorModesToolbar));
+	}
 }
 
 void FBlueprintEditorToolbar::AddBlueprintGlobalOptionsToolbar(TSharedPtr<FExtender> Extender)
