BP - Fixing a bug when picking the right layout for a BP for the unified view.

diff --git a/Engine/Source/Editor/Kismet/Private/BlueprintEditor.cpp b/Engine/Source/Editor/Kismet/Private/BlueprintEditor.cpp
--- a/Engine/Source/Editor/Kismet/Private/BlueprintEditor.cpp
+++ b/Engine/Source/Editor/Kismet/Private/BlueprintEditor.cpp
@@ -1433,7 +1433,7 @@ void FBlueprintEditor::RegisterApplicationModes(const TArray<UBlueprint*>& InBlu
 			{
 				AddApplicationMode(
 					FBlueprintEditorApplicationModes::StandardBlueprintEditorMode,
-					MakeShareable(new FBlueprintEditorUnifiedMode(SharedThis(this), FBlueprintEditorApplicationModes::StandardBlueprintEditorMode, FBlueprintEditorApplicationModes::GetLocalizedMode, bShouldOpenInComponentsMode && CanAccessComponentsMode())));
+					MakeShareable(new FBlueprintEditorUnifiedMode(SharedThis(this), FBlueprintEditorApplicationModes::StandardBlueprintEditorMode, FBlueprintEditorApplicationModes::GetLocalizedMode, CanAccessComponentsMode())));
 				SetCurrentMode(FBlueprintEditorApplicationModes::StandardBlueprintEditorMode);
 			}
 			else
