BP - Adding a new experimental setting for the unified Blueprint editor, updating references.  The simulation button now activates the viewport upon simuluation started for the unified mode.

diff --git a/Engine/Source/Editor/Kismet/Private/BlueprintEditor.cpp b/Engine/Source/Editor/Kismet/Private/BlueprintEditor.cpp
--- a/Engine/Source/Editor/Kismet/Private/BlueprintEditor.cpp
+++ b/Engine/Source/Editor/Kismet/Private/BlueprintEditor.cpp
@@ -1429,7 +1429,7 @@ void FBlueprintEditor::RegisterApplicationModes(const TArray<UBlueprint*>& InBlu
 		}
 		else
 		{
-			if ( GetDefault<UEditorExperimentalSettings>()->bInWorldBPEditing )
+			if ( GetDefault<UEditorExperimentalSettings>()->bUnifiedBlueprintEditor )
 			{
 				AddApplicationMode(
 					FBlueprintEditorApplicationModes::StandardBlueprintEditorMode,
@@ -1898,7 +1898,7 @@ void FBlueprintEditor::CreateDefaultTabContents(const TArray<UBlueprint*>& InBlu
 	{
 		SCSEditor = SNew(SSCSEditor, SharedThis(this), InBlueprint->SimpleConstructionScript, InBlueprint);
 		
-		SCSViewport = SAssignNew(SCSViewport, SSCSEditorViewport) .BlueprintEditor(SharedThis(this));
+		SCSViewport = SAssignNew(SCSViewport, SSCSEditorViewport).BlueprintEditor(SharedThis(this));
 	}
 }
 
diff --git a/Engine/Source/Editor/Kismet/Private/SBlueprintEditorToolbar.cpp b/Engine/Source/Editor/Kismet/Private/SBlueprintEditorToolbar.cpp
--- a/Engine/Source/Editor/Kismet/Private/SBlueprintEditorToolbar.cpp
+++ b/Engine/Source/Editor/Kismet/Private/SBlueprintEditorToolbar.cpp
@@ -400,7 +400,7 @@ static void BlueprintEditorToolbarImpl::MakeCompileDeveloperSubMenu(FMenuBuilder
 
 void FBlueprintEditorToolbar::AddBlueprintEditorModesToolbar(TSharedPtr<FExtender> Extender)
 {
-	if ( !GetDefault<UEditorExperimentalSettings>()->bInWorldBPEditing )
+	if ( !GetDefault<UEditorExperimentalSettings>()->bUnifiedBlueprintEditor )
 	{
 		TSharedPtr<FBlueprintEditor> BlueprintEditorPtr = BlueprintEditor.Pin();
 
diff --git a/Engine/Source/Editor/Kismet/Private/SSCSEditorViewport.cpp b/Engine/Source/Editor/Kismet/Private/SSCSEditorViewport.cpp
--- a/Engine/Source/Editor/Kismet/Private/SSCSEditorViewport.cpp
+++ b/Engine/Source/Editor/Kismet/Private/SSCSEditorViewport.cpp
@@ -11,6 +11,7 @@
 #include "Editor/UnrealEd/Public/STransformViewportToolbar.h"
 #include "EditorViewportCommands.h"
 #include "SEditorViewportToolBarMenu.h"
+#include "BlueprintEditorTabs.h"
 
 /*-----------------------------------------------------------------------------
    SSCSEditorViewportToolBar
@@ -251,7 +252,7 @@ void SSCSEditorViewport::BindCommands()
 
 	BlueprintEditorPtr.Pin()->GetToolkitCommands()->MapAction(
 		FBlueprintEditorCommands::Get().EnableSimulation,
-		FExecuteAction::CreateSP(ViewportClient.Get(), &FSCSEditorViewportClient::ToggleIsSimulateEnabled),
+		FExecuteAction::CreateSP(this, &SSCSEditorViewport::ToggleIsSimulateEnabled),
 		FCanExecuteAction(),
 		FIsActionChecked::CreateSP(ViewportClient.Get(), &FSCSEditorViewportClient::GetIsSimulateEnabled));
 
@@ -278,6 +279,21 @@ void SSCSEditorViewport::Invalidate()
 	ViewportClient->Invalidate();
 }
 
+void SSCSEditorViewport::ToggleIsSimulateEnabled()
+{
+	// If we're in unified BP editing mode, make the viewport visible when simulation is starting.
+	if ( GetDefault<UEditorExperimentalSettings>()->bUnifiedBlueprintEditor )
+	{
+		// Only trigger the switch if the simulation is starting.
+		if ( !ViewportClient->GetIsSimulateEnabled() )
+		{
+			BlueprintEditorPtr.Pin()->GetTabManager()->InvokeTab(FBlueprintEditorTabs::SCSViewportID);
+		}
+	}
+
+	ViewportClient->ToggleIsSimulateEnabled();
+}
+
 void SSCSEditorViewport::EnablePreview(bool bEnable)
 {
 	if(bEnable)
diff --git a/Engine/Source/Editor/Kismet/Private/SSCSEditorViewport.h b/Engine/Source/Editor/Kismet/Private/SSCSEditorViewport.h
--- a/Engine/Source/Editor/Kismet/Private/SSCSEditorViewport.h
+++ b/Engine/Source/Editor/Kismet/Private/SSCSEditorViewport.h
@@ -84,6 +84,9 @@ protected:
 	bool IsVisible() const;
 	EVisibility GetWidgetVisibility() const;
 
+	/** Called when the simulation toggle command is fired */
+	void ToggleIsSimulateEnabled();
+
 	/** SEditorViewport interface */
 	virtual TSharedRef<class FEditorViewportClient> MakeEditorViewportClient() override;
 	virtual TSharedPtr<class SWidget> MakeViewportToolbar() override;
diff --git a/Engine/Source/Editor/UnrealEd/Classes/Settings/EditorExperimentalSettings.h b/Engine/Source/Editor/UnrealEd/Classes/Settings/EditorExperimentalSettings.h
--- a/Engine/Source/Editor/UnrealEd/Classes/Settings/EditorExperimentalSettings.h
+++ b/Engine/Source/Editor/UnrealEd/Classes/Settings/EditorExperimentalSettings.h
@@ -29,10 +29,14 @@ public:
 	UPROPERTY(EditAnywhere, config, Category=Tools, meta=(DisplayName="Editor Utility Blueprints (Blutility)"))
 	bool bEnableEditorUtilityBlueprints;
 
-	/** Enable Single Layout BP Editor and In World BP Editing (WIP). */
-	UPROPERTY(EditAnywhere, config, Category = Tools, meta = (DisplayName = "Single Layout BP Editor and In World BP Editing"))
+	/** Enable In World BP Editing (WIP). */
+	UPROPERTY(EditAnywhere, config, Category = Tools, meta = (DisplayName = "In World Blueprint Editing"))
 	bool bInWorldBPEditing;
 
+	/** Enable Single Layout BP Editor. */
+	UPROPERTY(EditAnywhere, config, Category = Tools, meta = ( DisplayName = "Single Layout Blueprint Editor" ))
+	bool bUnifiedBlueprintEditor;
+
 	/** The Messaging Debugger provides a visual utility for debugging the messaging system. */
 	UPROPERTY(EditAnywhere, config, Category=Tools, meta=(DisplayName="Messaging Debugger"))
 	bool bMessagingDebugger;
