Don't allow references to Actors in other sublevels to be added to the levelscript

diff --git a/Engine/Source/Editor/Kismet/Private/BlueprintActionMenuUtils.cpp b/Engine/Source/Editor/Kismet/Private/BlueprintActionMenuUtils.cpp
--- a/Engine/Source/Editor/Kismet/Private/BlueprintActionMenuUtils.cpp
+++ b/Engine/Source/Editor/Kismet/Private/BlueprintActionMenuUtils.cpp
@@ -20,6 +20,7 @@
 #include "BlueprintEditorSettings.h"
 #include "Engine/Selection.h"
 #include "Engine/LevelScriptActor.h"
+#include "Engine/LevelScriptBlueprint.h"
 
 #define LOCTEXT_NAMESPACE "BlueprintActionMenuUtils"
 
@@ -428,6 +429,28 @@ void FBlueprintActionMenuUtils::MakeContextMenu(FBlueprintActionContext const& C
 	// only want bound actions for this menu section
 	LevelActorsFilter.AddRejectionTest(FBlueprintActionFilter::FRejectionTestDelegate::CreateStatic(IsUnBoundSpawner));
 
+	const UBlueprintEditorSettings* BlueprintSettings = GetDefault<UBlueprintEditorSettings>();
+	bool const bAddTargetContext  = bIsContextSensitive && BlueprintSettings->bUseTargetContextForNodeMenu;
+	bool bCanOperateOnLevelActors = bIsContextSensitive;
+	bool bCanHaveActorComponents  = bIsContextSensitive;
+	// determine if we can operate on certain object selections (level actors, 
+	// components, etc.)
+	for (UBlueprint* Blueprint : Context.Blueprints)
+	{
+		UClass* BlueprintClass = Blueprint->SkeletonGeneratedClass;
+		if (BlueprintClass != nullptr)
+		{
+			bCanOperateOnLevelActors &= BlueprintClass->IsChildOf<ALevelScriptActor>();
+			if (bAddTargetContext)
+			{
+				MainMenuFilter.TargetClasses.Add(BlueprintClass);
+			}
+		}
+		bCanHaveActorComponents &= FBlueprintEditorUtils::DoesSupportComponents(Blueprint);
+	}
+
+	UEdGraphSchema_K2 const* K2Schema = GetDefault<UEdGraphSchema_K2>();
+
 	// make sure the bound menu sections have the proper OwnerClasses specified
 	for (UObject* Selection : Context.SelectedObjects)
 	{
@@ -438,11 +461,24 @@ void FBlueprintActionMenuUtils::MakeContextMenu(FBlueprintActionContext const& C
 		else if (AActor* LevelActor = Cast<AActor>(Selection))
 		{
 			ComponentsFilter.Context.SelectedObjects.Remove(Selection);
-			if (!LevelActor->NeedsLoadForClient() && !LevelActor->NeedsLoadForServer())
+			if (!bCanOperateOnLevelActors || (!LevelActor->NeedsLoadForClient() && !LevelActor->NeedsLoadForServer()))
 			{
 				// don't want to let the level script operate on actors that won't be loaded in game
 				LevelActorsFilter.Context.SelectedObjects.Remove(Selection);
 			}
+			else
+			{
+				// Make sure every blueprint is in the same level as this actor
+				for (UBlueprint* Blueprint : Context.Blueprints)
+				{
+					ULevelScriptBlueprint* LevelScriptBlueprint = Cast<ULevelScriptBlueprint>(Blueprint);
+					if (LevelScriptBlueprint && !K2Schema->IsActorValidForLevelScriptRefs(LevelActor, LevelScriptBlueprint))
+					{
+						LevelActorsFilter.Context.SelectedObjects.Remove(Selection);
+						break;
+					}
+				}
+			}
 		}
 		else
 		{
@@ -457,30 +493,24 @@ void FBlueprintActionMenuUtils::MakeContextMenu(FBlueprintActionContext const& C
 	{
 		AActor* LevelActor = Cast<AActor>(*LvlActorIt);
 		// don't want to let the level script operate on actors that won't be loaded in game
-		if (LevelActor->NeedsLoadForClient() || LevelActor->NeedsLoadForServer())
-		{
-			LevelActorsFilter.Context.SelectedObjects.AddUnique(LevelActor);
-		}
-	}
-
-	const UBlueprintEditorSettings* BlueprintSettings = GetDefault<UBlueprintEditorSettings>();
-	bool const bAddTargetContext  = bIsContextSensitive && BlueprintSettings->bUseTargetContextForNodeMenu;
-	bool bCanOperateOnLevelActors = bIsContextSensitive;
-	bool bCanHaveActorComponents  = bIsContextSensitive;
-	// determine if we can operate on certain object selections (level actors, 
-	// components, etc.)
-	for (UBlueprint* Blueprint : Context.Blueprints)
-	{
-		UClass* BlueprintClass = Blueprint->SkeletonGeneratedClass;
-		if (BlueprintClass != nullptr)
+		if (bCanOperateOnLevelActors && (LevelActor->NeedsLoadForClient() || LevelActor->NeedsLoadForServer()))
 		{
-			bCanOperateOnLevelActors &= BlueprintClass->IsChildOf<ALevelScriptActor>();
-			if (bAddTargetContext)
+			bool bAddActor = true;
+			// Make sure every blueprint is in the same level as this actor
+			for (UBlueprint* Blueprint : Context.Blueprints)
 			{
-				MainMenuFilter.TargetClasses.Add(BlueprintClass);
+				ULevelScriptBlueprint* LevelScriptBlueprint = Cast<ULevelScriptBlueprint>(Blueprint);
+				if (LevelScriptBlueprint && !K2Schema->IsActorValidForLevelScriptRefs(LevelActor, LevelScriptBlueprint))
+				{
+					bAddActor = false;
+					break;
+				}
+			}
+			if (bAddActor)
+			{
+				LevelActorsFilter.Context.SelectedObjects.AddUnique(LevelActor);
 			}
 		}
-		bCanHaveActorComponents &= FBlueprintEditorUtils::DoesSupportComponents(Blueprint);
 	}
 
 	if(bCanHaveActorComponents)
