Removing code from Find-in-Blueprints that supported legacy reads from the DDC.

diff --git a/Engine/Source/Editor/Kismet/Private/FindInBlueprintManager.cpp b/Engine/Source/Editor/Kismet/Private/FindInBlueprintManager.cpp
--- a/Engine/Source/Editor/Kismet/Private/FindInBlueprintManager.cpp
+++ b/Engine/Source/Editor/Kismet/Private/FindInBlueprintManager.cpp
@@ -6,8 +6,6 @@
 #include "FindInBlueprints.h"
 #include "Editor/UnrealEd/Public/Kismet2/BlueprintEditorUtils.h"
 #include "AssetRegistryModule.h"
-#include "DerivedDataCacheInterface.h"
-#include "DerivedDataPluginInterface.h"
 
 #include "JsonUtilities.h"
 #include "SNotificationList.h"
@@ -289,49 +287,6 @@ namespace BlueprintSearchMetaDataHelpers
 
 	typedef TJsonFindInBlueprintStringReader<TCHAR> SearchMetaDataReader;
 
-	/** Json Writer used for serializing FText's in the correct format for Find-in-Blueprints */
-	template <class CharType = TCHAR>
-	class TJsonFindInBlueprintStringReaderDDC : public FJsonStringReader
-	{
-	public:
-		static TSharedRef< TJsonFindInBlueprintStringReaderDDC< TCHAR > > Create( const FString& InJsonString  )
-		{
-			return MakeShareable( new TJsonFindInBlueprintStringReaderDDC(InJsonString) );
-		}
-
-		TJsonFindInBlueprintStringReaderDDC( const FString& InJsonString )
-			: FJsonStringReader(InJsonString)
-		{
-
-		}
-
-		FORCEINLINE virtual const FString& GetIdentifier() const override
-		{
-			if(this->Identifier.Len())
-			{
-				// Remove the const, non-ideal, but this code is to bridge between the old DDC method and the new Asset Registry method of story FiB data.
-				FString& IdentifierNonConst = const_cast<FString&>(this->Identifier);
-				IdentifierNonConst = FFindInBlueprintSearchManager::ConvertHexStringToFText(this->Identifier).ToString();
-			}
-			return this->Identifier;
-		}
-
-		FORCEINLINE virtual  const FString& GetValueAsString() const override
-		{ 
-			check( this->CurrentToken == EJsonToken::String );
-			if(this->StringValue.Len())
-			{
-				// Remove the const, non-ideal, but this code is to bridge between the old DDC method and the new Asset Registry method of story FiB data.
-				FString& StringValueNonConst = const_cast<FString&>(this->StringValue);
-				StringValueNonConst = FFindInBlueprintSearchManager::ConvertHexStringToFText(this->StringValue).ToString();
-			}
-			return this->StringValue;
-		}
-	};
-
-	typedef TJsonFindInBlueprintStringReaderDDC<TCHAR> SearchMetaDataReaderDDC;
-
-
 	/**
 	 * Checks if Json value is searchable, eliminating data that not considered useful to search for
 	 *
@@ -691,25 +646,6 @@ void FFindInBlueprintSearchManager::OnAssetAdded(const FAssetData& InAssetData)
 				// Since the asset was not loaded, pull out the searchable data stored in the asset
 				NewSearchData.Value = *FiBSearchData;
 			}
-			else
-			{
-				// We will look for DDC information about this asset
-				FString SearchGuid;
-				const FString* SearchGuidPtr = InAssetData.TagsAndValues.Find("SearchGuid");
-				if(SearchGuidPtr)
-				{
-					SearchGuid = *SearchGuidPtr;
-
-					FDerivedDataCacheInterface& DDC = GetDerivedDataCacheRef();
-					if(InAssetData.IsAssetLoaded())
-					{
-						NewSearchData.Blueprint = Cast<UBlueprint>(InAssetData.GetAsset());
-					}
-					NewSearchData.BlueprintPath = InAssetData.ObjectPath.ToString();
-					NewSearchData.DDCRetrievalID = DDC.GetAsynchronous(*SearchGuid);
-					NewSearchData.bMarkedForDeletion = false;
-				}
-			}
 
 			int32 ArrayIndex = SearchArray.Add(NewSearchData);
 
@@ -942,19 +878,6 @@ bool FFindInBlueprintSearchManager::ContinueSearchQuery(const FStreamSearch* InS
 			else
 			{
  				OutSearchData = SearchArray[SearchIdx++];
-				if(OutSearchData.Value.IsEmpty())
-				{
-					if(OutSearchData.Blueprint.IsValid())
-					{
-						RetrieveDDCData(SearchArray[SearchIdx - 1], *OutSearchData.Blueprint->GetPathName());
-					}
-					else
-					{
-						// Attempt to retrieve the data from the DDC, if there is none still, it
-						RetrieveDDCData(SearchArray[SearchIdx - 1], *SearchArray[SearchIdx - 1].BlueprintPath);
-					}
-					OutSearchData = SearchArray[SearchIdx - 1];
-				}
 				return true;
 			}
 
@@ -1001,39 +924,6 @@ float FFindInBlueprintSearchManager::GetPercentComplete(const FStreamSearch* InS
 	return ReturnPercent;
 }
 
-void FFindInBlueprintSearchManager::RetrieveDDCData(FSearchData& InOutSearchData, FName InBlueprintPath)
-{
-	FDerivedDataCacheInterface& DDC = GetDerivedDataCacheRef();
-
-	// The data is valid, but we may still need to wait on the DDC.
-	if(InOutSearchData.DDCRetrievalID != INDEX_NONE)
-	{
-		// Allow only one thread to query the DDC at a single time, this is important in the case that two threads are waiting for the same DDC information, once the information is pulled, it is deleted and is no longer accessible.
-		FScopeLock ScopeLock(&SafeModifyCacheCriticalSection);
-
-		// Check to see if the data was pulled by another thread while we waited
-		if(InOutSearchData.DDCRetrievalID != INDEX_NONE)
-		{
-			// Wait for completion of the DDC retrieval, we are on a separate thread so this will not stall the editor
-			DDC.WaitAsynchronousCompletion(InOutSearchData.DDCRetrievalID);
-
-			TArray<uint8> DerivedData;
-
-			if( DDC.GetAsynchronousResults(InOutSearchData.DDCRetrievalID, DerivedData) ) 
-			{
-				check(DerivedData.Num());
-
-				FMemoryReader DDCAr(DerivedData, /*bIsPersistent=*/ true);
-				DDCAr << InOutSearchData.BlueprintPath;
-				DDCAr << InOutSearchData.Value;
-				DDCAr.Close();
-			}
-
-			InOutSearchData.DDCRetrievalID = INDEX_NONE;
-		}
-	}
-}
-
 FString FFindInBlueprintSearchManager::QuerySingleBlueprint(UBlueprint* InBlueprint, bool bInRebuildSearchData/* = true*/)
 {
 	if(bInRebuildSearchData)
@@ -1213,16 +1103,6 @@ FString FFindInBlueprintSearchManager::ConvertFTextToHexString(FText InValue)
 
 TSharedPtr< FJsonObject > FFindInBlueprintSearchManager::ConvertJsonStringToObject(FString InJsonString)
 {
-	if(InJsonString.Len() && InJsonString[0] == '{')
-	{
-		// This is a DDC version Json string, use old method of deserialization
-		TSharedPtr< FJsonObject > JsonObject = NULL;
-		TSharedRef< TJsonReader<> > Reader = BlueprintSearchMetaDataHelpers::SearchMetaDataReaderDDC::Create( InJsonString );
-		FJsonSerializer::Deserialize( Reader, JsonObject );
-
-		return JsonObject;
-	}
-
 	/** The searchable data is more complicated than a Json string, the Json being the main searchable body that is parsed. Below is a diagram of the full data:
 	 *  | int32 "Size" | TMap "Lookup Table" | Json String |
 	 *
diff --git a/Engine/Source/Editor/Kismet/Public/FindInBlueprintManager.h b/Engine/Source/Editor/Kismet/Public/FindInBlueprintManager.h
--- a/Engine/Source/Editor/Kismet/Public/FindInBlueprintManager.h
+++ b/Engine/Source/Editor/Kismet/Public/FindInBlueprintManager.h
@@ -81,13 +81,9 @@ struct FSearchData
 	/** Cached to determine if the Blueprint is seen as no longer valid, allows it to be cleared out next save to disk */
 	bool bMarkedForDeletion;
 
-	/** Retrieval ID from the DDC, set to INDEX_None when retrieval is not necessary */
-	int32 DDCRetrievalID;
-
 	FSearchData()
 		: Blueprint(nullptr)
 		, bMarkedForDeletion(false)
-		, DDCRetrievalID(INDEX_NONE)
 	{
 
 	}
@@ -206,9 +202,6 @@ private:
 	/** Builds the cache from all available Blueprint assets that the asset registry has discovered at the time of this function. Occurs on startup */
 	void BuildCache();
 
-	/** Retrieves Find-in-Blueprint data from the DDC */
-	void RetrieveDDCData(FSearchData& InOutSearchData, FName InBlueprintPath);
-
 protected:
 	/** Maps the Blueprint paths to their index in the SearchArray */
 	TMap<FString, int> SearchMap;
