Disable ability to make references to level Actors in anonymous blueprints until we can do it via variables.

diff --git a/Engine/Source/Editor/BlueprintGraph/Private/EdGraphSchema_K2.cpp b/Engine/Source/Editor/BlueprintGraph/Private/EdGraphSchema_K2.cpp
--- a/Engine/Source/Editor/BlueprintGraph/Private/EdGraphSchema_K2.cpp
+++ b/Engine/Source/Editor/BlueprintGraph/Private/EdGraphSchema_K2.cpp
@@ -878,8 +878,9 @@ void UEdGraphSchema_K2::ListFunctionsMatchingSignatureAsDelegates(FGraphContextM
 bool UEdGraphSchema_K2::IsActorValidForLevelScriptRefs(const AActor* TestActor, const UBlueprint* Blueprint) const
 {
 	check(Blueprint);
-	
+
 	return TestActor
+		&& FBlueprintEditorUtils::IsLevelScriptBlueprint(Blueprint)
 		&& (TestActor->GetLevel() == FBlueprintEditorUtils::GetLevelFromBlueprint(Blueprint))
 		&& FKismetEditorUtilities::IsActorValidForLevelScript(TestActor);
 }
diff --git a/Engine/Source/Editor/Kismet/Private/BlueprintActionMenuUtils.cpp b/Engine/Source/Editor/Kismet/Private/BlueprintActionMenuUtils.cpp
--- a/Engine/Source/Editor/Kismet/Private/BlueprintActionMenuUtils.cpp
+++ b/Engine/Source/Editor/Kismet/Private/BlueprintActionMenuUtils.cpp
@@ -440,7 +440,7 @@ void FBlueprintActionMenuUtils::MakeContextMenu(FBlueprintActionContext const& C
 		UClass* BlueprintClass = Blueprint->SkeletonGeneratedClass;
 		if (BlueprintClass != nullptr)
 		{
-			bCanOperateOnLevelActors &= (BlueprintClass->IsChildOf<ALevelScriptActor>() || FBlueprintEditorUtils::IsAnonymousBlueprintClass(BlueprintClass));
+			bCanOperateOnLevelActors &= BlueprintClass->IsChildOf<ALevelScriptActor>();
 			if (bAddTargetContext)
 			{
 				MainMenuFilter.TargetClasses.Add(BlueprintClass);
diff --git a/Engine/Source/Editor/Kismet/Private/BlueprintEditor.cpp b/Engine/Source/Editor/Kismet/Private/BlueprintEditor.cpp
--- a/Engine/Source/Editor/Kismet/Private/BlueprintEditor.cpp
+++ b/Engine/Source/Editor/Kismet/Private/BlueprintEditor.cpp
@@ -2502,7 +2502,7 @@ void FBlueprintEditor::OnGraphEditorDropActor(const TArray< TWeakObjectPtr<AActo
 	// We need to check that the dropped actor is in the right sublevel for the reference
 	ULevel* BlueprintLevel = FBlueprintEditorUtils::GetLevelFromBlueprint(GetBlueprintObj());
 
-	if (BlueprintLevel)
+	if (BlueprintLevel && FBlueprintEditorUtils::IsLevelScriptBlueprint(GetBlueprintObj()))
 	{
 		FVector2D NodeLocation = DropLocation;
 		for (int32 i = 0; i < Actors.Num(); i++)
