[UE-7375] Fix for crash when diffing a Blueprint Interface or a blueprint that has a new graph in it

diff --git a/Engine/Source/Editor/GraphEditor/Private/GraphDiffControl.cpp b/Engine/Source/Editor/GraphEditor/Private/GraphDiffControl.cpp
--- a/Engine/Source/Editor/GraphEditor/Private/GraphDiffControl.cpp
+++ b/Engine/Source/Editor/GraphEditor/Private/GraphDiffControl.cpp
@@ -586,7 +586,7 @@ bool FGraphDiffControl::DiffGraphs(UEdGraph* const LhsGraph, UEdGraph* const Rhs
 	}
 
 	// storing the graph name for all diff entries:
-	FName GraphName = LhsGraph ? LhsGraph->GetFName() : LhsGraph->GetFName();
+	FName GraphName = LhsGraph ? LhsGraph->GetFName() : RhsGraph->GetFName();
 	for( auto& Entry : DiffsOut )
 	{
 		Entry.OwningGraph = GraphName;
diff --git a/Engine/Source/Editor/Kismet/Private/SCSDiff.cpp b/Engine/Source/Editor/Kismet/Private/SCSDiff.cpp
--- a/Engine/Source/Editor/Kismet/Private/SCSDiff.cpp
+++ b/Engine/Source/Editor/Kismet/Private/SCSDiff.cpp
@@ -9,6 +9,12 @@
 
 FSCSDiff::FSCSDiff(const UBlueprint* InBlueprint)
 {
+	if (!FBlueprintEditorUtils::SupportsConstructionScript(InBlueprint) || InBlueprint->SimpleConstructionScript == NULL)
+	{
+		ContainerWidget = SNew(SBox);
+		return;
+	}
+
 	TSharedRef<SKismetInspector> Inspector = SNew(SKismetInspector)
 		.HideNameArea(true)
 		.ViewIdentifier(FName("BlueprintInspector"))
@@ -28,8 +34,11 @@ FSCSDiff::FSCSDiff(const UBlueprint* InBlueprint)
 
 void FSCSDiff::HighlightProperty(FName VarName, FPropertySoftPath Property)
 {
-	check( VarName != FName() );
-	SCSEditor->HighlightTreeNode( VarName, FPropertyPath() );
+	if (SCSEditor.IsValid())
+	{
+		check(VarName != FName());
+		SCSEditor->HighlightTreeNode(VarName, FPropertyPath());
+	}
 }
 
 TSharedRef< SWidget > FSCSDiff::TreeWidget()
@@ -55,11 +64,14 @@ TArray< FSCSResolvedIdentifier > FSCSDiff::GetDisplayedHierarchy() const
 {
 	TArray< FSCSResolvedIdentifier > Ret;
 
-	for (int32 Iter = 0; Iter != SCSEditor->RootNodes.Num(); ++Iter)
+	if( SCSEditor.IsValid() )
 	{
-		TArray< int32 > TreeAddress;
-		TreeAddress.Push(Iter);
-		GetDisplayedHierarchyRecursive(TreeAddress, *SCSEditor->RootNodes[Iter], Ret);
+		for (int32 Iter = 0; Iter != SCSEditor->RootNodes.Num(); ++Iter)
+		{
+			TArray< int32 > TreeAddress;
+			TreeAddress.Push(Iter);
+			GetDisplayedHierarchyRecursive(TreeAddress, *SCSEditor->RootNodes[Iter], Ret);
+		}
 	}
 
 	return Ret;
