Editor - Adjusting some of the styles and padding in the editor to thin things out some.

diff --git a/Engine/Source/Editor/Kismet/Private/SKismetInspector.cpp b/Engine/Source/Editor/Kismet/Private/SKismetInspector.cpp
--- a/Engine/Source/Editor/Kismet/Private/SKismetInspector.cpp
+++ b/Engine/Source/Editor/Kismet/Private/SKismetInspector.cpp
@@ -73,31 +73,31 @@ TSharedRef<SWidget> SKismetInspector::MakeContextualEditingWidget(struct FKismet
 	if (SelectionInfo.ObjectsForPropertyEditing.Num())
 	{
 		ContextualEditingWidget->AddSlot()
-			.FillHeight( 0.9f )
-			.VAlign( VAlign_Top )
+		.FillHeight( 0.9f )
+		.VAlign( VAlign_Top )
+		[
+			SNew( SBox )
+			.Visibility(this, &SKismetInspector::GetPropertyViewVisibility)
 			[
-				SNew( SBorder )
-				.Visibility(this, &SKismetInspector::GetPropertyViewVisibility)
-				.BorderImage( FEditorStyle::GetBrush("NoBorder") )
-				[
-					PropertyView.ToSharedRef()
-				]
-			];
+				PropertyView.ToSharedRef()
+			]
+		];
 
 		if (bShowPublicView)
 		{
 			ContextualEditingWidget->AddSlot()
-				.AutoHeight()
-				.VAlign( VAlign_Top )
+			.AutoHeight()
+			.VAlign( VAlign_Top )
+			[
+				SNew(SCheckBox)
+				.ToolTipText(LOCTEXT("TogglePublicView", "Toggle Public View"))
+				.IsChecked(this, &SKismetInspector::GetPublicViewCheckboxState)
+				.OnCheckStateChanged( this, &SKismetInspector::SetPublicViewCheckboxState)
 				[
-					SNew(SCheckBox)
-					.ToolTipText(LOCTEXT("TogglePublicView", "Toggle Public View"))
-					.IsChecked(this, &SKismetInspector::GetPublicViewCheckboxState)
-					.OnCheckStateChanged( this, &SKismetInspector::SetPublicViewCheckboxState)
-					[
-						SNew(STextBlock) .Text(LOCTEXT("PublicViewCheckboxLabel", "Public View"))
-					]
-				];
+					SNew(STextBlock)
+					.Text(LOCTEXT("PublicViewCheckboxLabel", "Public View"))
+				]
+			];
 		}
 	}
 
diff --git a/Engine/Source/Editor/LevelEditor/Private/SLevelEditor.cpp b/Engine/Source/Editor/LevelEditor/Private/SLevelEditor.cpp
--- a/Engine/Source/Editor/LevelEditor/Private/SLevelEditor.cpp
+++ b/Engine/Source/Editor/LevelEditor/Private/SLevelEditor.cpp
@@ -709,7 +709,6 @@ TSharedRef<SDockTab> SLevelEditor::SpawnLevelEditorTab( const FSpawnTabArgs& Arg
 			.Icon( FEditorStyle::GetBrush( "LevelEditor.Tabs.Outliner" ) )
 			.Label( Label )
 			.ToolTip( IDocumentation::Get()->CreateToolTip( Label, nullptr, "Shared/LevelEditor", "SceneOutlinerTab" ) )
-			.ContentPadding( 5 )
 			[
 				SNew(SBorder)
 				.Padding(4)
diff --git a/Engine/Source/Editor/PropertyEditor/Private/DetailAdvancedDropdownNode.cpp b/Engine/Source/Editor/PropertyEditor/Private/DetailAdvancedDropdownNode.cpp
--- a/Engine/Source/Editor/PropertyEditor/Private/DetailAdvancedDropdownNode.cpp
+++ b/Engine/Source/Editor/PropertyEditor/Private/DetailAdvancedDropdownNode.cpp
@@ -33,7 +33,7 @@ public:
 		TSharedPtr<SWidget> Content;
 		if( bIsTopNode )
 		{
-			Content=
+			Content =
 				SNew( SBorder )
 				.BorderImage( FEditorStyle::GetBrush("DetailsView.CategoryMiddle") )
 				.Padding( FMargin( 0.0f, 3.0f, 0.0f, 0.0f ) )
@@ -118,15 +118,15 @@ public:
 			}
 			else
 			{
-				SplitterArea = 
-					SNew( SSpacer )
-					.Size( FVector2D(1.0f, 3.0f) );
+				SplitterArea = SNew(SSpacer);
 			}
 
-			Content=
+			SplitterArea = SNew(SSpacer);
+
+			Content =
 				SNew( SBorder )
 				.BorderImage( FEditorStyle::GetBrush("DetailsView.CategoryBottom") )
-				.Padding( FMargin( 0.0f, 0.0f, 0.0f, 0.0f ) )
+				.Padding( FMargin( 0.0f, 0.0f, 0.0f, 2.0f ) )
 				[
 					SplitterArea.ToSharedRef()	
 				];
diff --git a/Engine/Source/Editor/PropertyEditor/Private/DetailCategoryGroupNode.cpp b/Engine/Source/Editor/PropertyEditor/Private/DetailCategoryGroupNode.cpp
--- a/Engine/Source/Editor/PropertyEditor/Private/DetailCategoryGroupNode.cpp
+++ b/Engine/Source/Editor/PropertyEditor/Private/DetailCategoryGroupNode.cpp
@@ -12,7 +12,7 @@ void SDetailCategoryTableRow::Construct( const FArguments& InArgs, TSharedRef<ID
 	TSharedPtr<SHorizontalBox> MyContent;
 
 	ChildSlot
-	.Padding( 0.0f, 5.0f, 0.0f, 0.0f )
+	.Padding( 0.0f, 2.0f, 0.0f, 0.0f )
 	[	
 		SNew( SBorder )
 		.BorderImage( this, &SDetailCategoryTableRow::GetBackgroundImage )
diff --git a/Engine/Source/Runtime/Slate/Private/Framework/Docking/SDockingTabStack.cpp b/Engine/Source/Runtime/Slate/Private/Framework/Docking/SDockingTabStack.cpp
--- a/Engine/Source/Runtime/Slate/Private/Framework/Docking/SDockingTabStack.cpp
+++ b/Engine/Source/Runtime/Slate/Private/Framework/Docking/SDockingTabStack.cpp
@@ -142,7 +142,7 @@ void SDockingTabStack::Construct( const FArguments& InArgs, const TSharedRef<FTa
 				.AutoHeight()
 				[
 					SNew(SImage)
-						.Image(this, &SDockingTabStack::GetTabWellBrush)
+					.Image(this, &SDockingTabStack::GetTabWellBrush)
 				]
 			]
 		]
@@ -769,7 +769,7 @@ FMargin SDockingTabStack::GetContentPadding() const
 	TSharedPtr<SDockTab> ForegroundTab = TabWell->GetForegroundTab();
 	return (ForegroundTab.IsValid())
 		? ForegroundTab->GetContentPadding()
-		: FMargin(5);
+		: FMargin(2);
 }
 
 EVisibility SDockingTabStack::GetTabWellVisibility() const
diff --git a/Engine/Source/Runtime/Slate/Private/Widgets/Views/SExpanderArrow.cpp b/Engine/Source/Runtime/Slate/Private/Widgets/Views/SExpanderArrow.cpp
--- a/Engine/Source/Runtime/Slate/Private/Widgets/Views/SExpanderArrow.cpp
+++ b/Engine/Source/Runtime/Slate/Private/Widgets/Views/SExpanderArrow.cpp
@@ -7,6 +7,7 @@ void SExpanderArrow::Construct( const FArguments& InArgs, const TSharedPtr<class
 {
 	OwnerRowPtr = TableRow;
 	IndentAmount = InArgs._IndentAmount;
+	BaseIndentLevel = InArgs._BaseIndentLevel;
 
 	this->ChildSlot
 	.Padding( TAttribute<FMargin>( this, &SExpanderArrow::GetExpanderPadding ) )
@@ -55,9 +56,9 @@ EVisibility SExpanderArrow::GetExpanderVisibility() const
 /** @return the margin corresponding to how far this item is indented */
 FMargin SExpanderArrow::GetExpanderPadding() const
 {
-	const int32 NestingDepth = OwnerRowPtr.Pin()->GetIndentLevel();
+	const int32 NestingDepth = FMath::Max(0, OwnerRowPtr.Pin()->GetIndentLevel() - BaseIndentLevel.Get());
 	const float Indent = IndentAmount.Get(10.f);
-	return FMargin( NestingDepth*Indent, 0,0,0 );
+	return FMargin( NestingDepth * Indent, 0,0,0 );
 }
 
 /** @return the name of an image that should be shown as the expander arrow */
diff --git a/Engine/Source/Runtime/Slate/Public/Widgets/Docking/SDockTab.h b/Engine/Source/Runtime/Slate/Public/Widgets/Docking/SDockTab.h
--- a/Engine/Source/Runtime/Slate/Public/Widgets/Docking/SDockTab.h
+++ b/Engine/Source/Runtime/Slate/Public/Widgets/Docking/SDockTab.h
@@ -48,7 +48,7 @@ public:
 		: _Content()
 		, _TabWellContentLeft()
 		, _TabWellContentRight()
-		, _ContentPadding( FMargin( 5 ) )
+		, _ContentPadding( FMargin( 2 ) )
 		, _TabRole(ETabRole::PanelTab)
 		, _Label()
 		, _Icon( FStyleDefaults::GetNoBrush() )
diff --git a/Engine/Source/Runtime/Slate/Public/Widgets/Views/SExpanderArrow.h b/Engine/Source/Runtime/Slate/Public/Widgets/Views/SExpanderArrow.h
--- a/Engine/Source/Runtime/Slate/Public/Widgets/Views/SExpanderArrow.h
+++ b/Engine/Source/Runtime/Slate/Public/Widgets/Views/SExpanderArrow.h
@@ -13,8 +13,11 @@ class SLATE_API SExpanderArrow : public SCompoundWidget
 public:
 
 	SLATE_BEGIN_ARGS( SExpanderArrow )
+		: _IndentAmount(10)
+		, _BaseIndentLevel(0)
 	{ }
 		SLATE_ATTRIBUTE(float, IndentAmount)
+		SLATE_ATTRIBUTE(int32, BaseIndentLevel)
 	SLATE_END_ARGS()
 
 	void Construct( const FArguments& InArgs, const TSharedPtr<class ITableRow>& TableRow );
@@ -40,4 +43,7 @@ protected:
 
 	/** The amount of space to indent at each level */
 	TAttribute<float> IndentAmount;
+
+	/** The level in the tree that begins the indention amount */
+	TAttribute<int32> BaseIndentLevel;
 };
diff --git a/Engine/Source/Runtime/SlateCore/Private/Styling/CoreStyle.cpp b/Engine/Source/Runtime/SlateCore/Private/Styling/CoreStyle.cpp
--- a/Engine/Source/Runtime/SlateCore/Private/Styling/CoreStyle.cpp
+++ b/Engine/Source/Runtime/SlateCore/Private/Styling/CoreStyle.cpp
@@ -972,7 +972,7 @@ TSharedRef<ISlateStyle> FCoreStyle::Create( const FName& InStyleSetName )
 			.SetForegroundBrush( BOX_BRUSH( "/Docking/Tab_Foreground", 4/16.0f ) )
 			.SetHoveredBrush( BOX_BRUSH( "/Docking/Tab_Hovered", 4/16.0f ) )
 			.SetContentAreaBrush( BOX_BRUSH( "/Docking/TabContentArea", FMargin(4/16.0f) ) )
-			.SetTabWellBrush( IMAGE_BRUSH( "/Docking/TabWellSeparator", FVector2D(16,4) ) )
+			.SetTabWellBrush( FSlateNoResource() )
 			.SetTabPadding( FMargin(5, 2, 5, 2) )
 			.SetOverlapWidth( -1.0f )
 			.SetFlashColor( TabFlashColor )
@@ -987,7 +987,7 @@ TSharedRef<ISlateStyle> FCoreStyle::Create( const FName& InStyleSetName )
 			.SetForegroundBrush( BOX_BRUSH( "/Docking/AppTab_Foreground", FMargin(24.0f/64.0f, 4/32.0f) ) )
 			.SetHoveredBrush( BOX_BRUSH( "/Docking/AppTab_Hovered", FMargin(24.0f/64.0f, 4/32.0f) ) )
 			.SetContentAreaBrush( BOX_BRUSH( "/Docking/AppTabContentArea", FMargin(4/16.0f) ) )
-			.SetTabWellBrush( IMAGE_BRUSH( "/Docking/AppTabWellSeparator", FVector2D(16,2) ) )
+			.SetTabWellBrush( FSlateNoResource() )
 			.SetTabPadding( FMargin(17, 4, 15, 4) )
 			.SetOverlapWidth( 21.0f )
 			.SetFlashColor( TabFlashColor )
