[UE-7192] break out SCS components tree from BP editor so it can be used in other areas

diff --git a/Engine/Source/Editor/Kismet/Private/BlueprintDetailsCustomization.cpp b/Engine/Source/Editor/Kismet/Private/BlueprintDetailsCustomization.cpp
--- a/Engine/Source/Editor/Kismet/Private/BlueprintDetailsCustomization.cpp
+++ b/Engine/Source/Editor/Kismet/Private/BlueprintDetailsCustomization.cpp
@@ -392,7 +392,7 @@ void FBlueprintVarActionDetails::CustomizeDetails( IDetailLayoutBuilder& DetailL
 			if (ComponentProperty->PropertyClass && ComponentProperty->PropertyClass->IsChildOf(UActorComponent::StaticClass()))
 			{
 				// Add Events Section property class that has valid events
-				if( FBlueprintEditor::CanClassGenerateEvents( ComponentProperty->PropertyClass ))
+				if( FBlueprintEditorUtils::CanClassGenerateEvents( ComponentProperty->PropertyClass ))
 				{
 					IDetailCategoryBuilder& EventCategory = DetailLayout.EditCategory("Component", LOCTEXT("ComponentDetailsCategory", "Events"));
 
@@ -1582,7 +1582,8 @@ TSharedRef<SWidget> FBlueprintVarActionDetails::BuildEventsMenuForVariable() con
 		{
 			TSharedRef<SSCSEditor> Editor =  BlueprintEditorPtr.Pin()->GetSCSEditor();
 			FMenuBuilder MenuBuilder( true, NULL );
-			Editor->BuildMenuEventsSection( MenuBuilder, BlueprintEditorPtr.Pin(), ComponentProperty->PropertyClass, 
+			Editor->BuildMenuEventsSection( MenuBuilder, BlueprintEditorPtr.Pin()->GetBlueprintObj(), ComponentProperty->PropertyClass, 
+											FCanExecuteAction::CreateSP(BlueprintEditorPtr.Pin().Get(), &FBlueprintEditor::InEditingMode),
 											FGetSelectedObjectsDelegate::CreateSP(MyBlueprintPtr.Get(), &SMyBlueprint::GetSelectedItemsForContextMenu));
 			return MenuBuilder.MakeWidget();
 		}
@@ -4391,7 +4392,7 @@ void FBlueprintComponentDetails::CustomizeDetails(IDetailLayoutBuilder& DetailLa
 	// Add Events Section if we have a common base class that has valid events
 	UClass* CommonEventsClass = FindCommonBaseClassFromSelected();
 
-	if( FBlueprintEditor::CanClassGenerateEvents( CommonEventsClass ))
+	if( FBlueprintEditorUtils::CanClassGenerateEvents( CommonEventsClass ))
 	{
 		IDetailCategoryBuilder& EventCategory = DetailLayout.EditCategory("Component", LOCTEXT("ComponentDetailsCategory", "Events"), ECategoryPriority::Important);
 
@@ -4474,7 +4475,8 @@ TSharedRef<SWidget> FBlueprintComponentDetails::BuildEventsMenuForComponents() c
 			if( CommonComponentClass )
 			{
 				FMenuBuilder MenuBuilder( true, NULL );
-				Editor->BuildMenuEventsSection( MenuBuilder, BlueprintEditorPtr.Pin(), CommonComponentClass, 
+				Editor->BuildMenuEventsSection( MenuBuilder, BlueprintEditorPtr.Pin()->GetBlueprintObj(), CommonComponentClass, 
+												FCanExecuteAction::CreateSP(BlueprintEditorPtr.Pin().Get(), &FBlueprintEditor::InEditingMode),
 												FGetSelectedObjectsDelegate::CreateSP(Editor.Get(), &SSCSEditor::GetSelectedItemsForContextMenu));
 				return MenuBuilder.MakeWidget();
 			}
diff --git a/Engine/Source/Editor/Kismet/Private/BlueprintEditor.cpp b/Engine/Source/Editor/Kismet/Private/BlueprintEditor.cpp
--- a/Engine/Source/Editor/Kismet/Private/BlueprintEditor.cpp
+++ b/Engine/Source/Editor/Kismet/Private/BlueprintEditor.cpp
@@ -659,6 +659,46 @@ void FBlueprintEditor::UpdateSCSPreview(bool bUpdateNow)
 	}
 }
 
+void FBlueprintEditor::OnSCSEditorTreeViewSelectionChanged(const TArray<FSCSEditorTreeNodePtrType>& SelectedNodes)
+{
+	if (SCSViewport.IsValid())
+	{
+		SCSViewport->OnComponentSelectionChanged();
+	}
+}
+
+void FBlueprintEditor::OnSCSEditorUpdateSelectionFromNodes(const TArray<FSCSEditorTreeNodePtrType>& SelectedNodes)
+{
+	// Convert the selection set to an array of UObject* pointers
+	FText InspectorTitle = FText::GetEmpty();
+	TArray<UObject*> InspectorObjects;
+	InspectorObjects.Empty(SelectedNodes.Num());
+	for (auto NodeIt = SelectedNodes.CreateConstIterator(); NodeIt; ++NodeIt)
+	{
+		auto NodePtr = *NodeIt;
+		if(NodePtr.IsValid() && NodePtr->CanEditDefaults())
+		{
+			InspectorTitle = FText::FromString(NodePtr->GetDisplayString());
+			InspectorObjects.Add(NodePtr->GetComponentTemplate());
+		}
+	}
+
+	// Update the details panel
+	if (Inspector.IsValid())
+	{
+		SKismetInspector::FShowDetailsOptions Options(InspectorTitle, true);
+		Inspector->ShowDetailsForObjects(InspectorObjects, Options);
+	}
+}
+
+void FBlueprintEditor::OnSCSEditorHighlightPropertyInDetailsView(const FPropertyPath& InPropertyPath)
+{
+	if(Inspector.IsValid())
+	{
+		Inspector->GetPropertyView()->HighlightProperty(InPropertyPath);
+	}
+}
+
 /** Create new tab for the supplied graph - don't call this directly.*/
 TSharedRef<SGraphEditor> FBlueprintEditor::CreateGraphEditorWidget(TSharedRef<FTabInfo> InTabInfo, UEdGraph* InGraph)
 {
@@ -1118,8 +1158,14 @@ void FBlueprintEditor::EnsureBlueprintIsUpToDate(UBlueprint* BlueprintObj)
 			BlueprintObj->SimpleConstructionScript->SetFlags(RF_Transactional);
 
 			// Recreate (or create) any widgets that depend on the SCS
-			SCSEditor = SNew(SSCSEditor, SharedThis(this), BlueprintObj->SimpleConstructionScript, BlueprintObj);
-			SCSViewport = SAssignNew(SCSViewport, SSCSEditorViewport) .BlueprintEditor(SharedThis(this));
+			SCSEditor = SAssignNew(SCSEditor, SSCSEditor, BlueprintObj->SimpleConstructionScript)
+				.InEditingMode(this, &FBlueprintEditor::InEditingMode)
+				.PreviewActor(this, &FBlueprintEditor::GetPreviewActor)
+				.OnTreeViewSelectionChanged(this, &FBlueprintEditor::OnSCSEditorTreeViewSelectionChanged)
+				.OnUpdateSelectionFromNodes(this, &FBlueprintEditor::OnSCSEditorUpdateSelectionFromNodes)
+				.OnHighlightPropertyInDetailsView(this, &FBlueprintEditor::OnSCSEditorHighlightPropertyInDetailsView);
+			SCSViewport = SAssignNew(SCSViewport, SSCSEditorViewport)
+				.BlueprintEditor(SharedThis(this));
 		}
 
 		// If we should have a UCS but don't yet, make it
@@ -1896,9 +1942,14 @@ void FBlueprintEditor::CreateDefaultTabContents(const TArray<UBlueprint*>& InBlu
 		InBlueprint->ParentClass->IsChildOf(AActor::StaticClass()) && 
 		InBlueprint->SimpleConstructionScript )
 	{
-		SCSEditor = SNew(SSCSEditor, SharedThis(this), InBlueprint->SimpleConstructionScript, InBlueprint);
-		
-		SCSViewport = SAssignNew(SCSViewport, SSCSEditorViewport).BlueprintEditor(SharedThis(this));
+		SCSEditor = SAssignNew(SCSEditor, SSCSEditor, InBlueprint->SimpleConstructionScript)
+			.InEditingMode(this, &FBlueprintEditor::InEditingMode)
+			.PreviewActor(this, &FBlueprintEditor::GetPreviewActor)
+			.OnTreeViewSelectionChanged(this, &FBlueprintEditor::OnSCSEditorTreeViewSelectionChanged)
+			.OnUpdateSelectionFromNodes(this, &FBlueprintEditor::OnSCSEditorUpdateSelectionFromNodes)
+			.OnHighlightPropertyInDetailsView(this, &FBlueprintEditor::OnSCSEditorHighlightPropertyInDetailsView);
+		SCSViewport = SAssignNew(SCSViewport, SSCSEditorViewport)
+			.BlueprintEditor(SharedThis(this));
 	}
 }
 
@@ -6861,22 +6912,6 @@ void FBlueprintEditor::ToolkitBroughtToFront()
 	}
 }
 
-bool FBlueprintEditor::CanClassGenerateEvents( UClass* InClass )
-{
-	if( InClass )
-	{
-		for( TFieldIterator<UMulticastDelegateProperty> PropertyIt( InClass, EFieldIteratorFlags::IncludeSuper ); PropertyIt; ++PropertyIt )
-		{
-			UProperty* Property = *PropertyIt;
-			if( !Property->HasAnyPropertyFlags( CPF_Parm ) && Property->HasAllPropertyFlags( CPF_BlueprintAssignable ))
-			{
-				return true;
-			}
-		}
-	}
-	return false;
-}
-
 void FBlueprintEditor::OnNodeSpawnedByKeymap()
 {
 	UpdateNodeCreationStats( ENodeCreateAction::Keymap );
diff --git a/Engine/Source/Editor/Kismet/Private/SCSDiff.cpp b/Engine/Source/Editor/Kismet/Private/SCSDiff.cpp
--- a/Engine/Source/Editor/Kismet/Private/SCSDiff.cpp
+++ b/Engine/Source/Editor/Kismet/Private/SCSDiff.cpp
@@ -4,6 +4,7 @@
 #include "SCSDiff.h"
 #include "SKismetInspector.h"
 #include "SSCSEditor.h"
+#include "IDetailsView.h"
 
 #include <vector>
 
@@ -15,7 +16,7 @@ FSCSDiff::FSCSDiff(const UBlueprint* InBlueprint)
 		return;
 	}
 
-	TSharedRef<SKismetInspector> Inspector = SNew(SKismetInspector)
+	Inspector = SNew(SKismetInspector)
 		.HideNameArea(true)
 		.ViewIdentifier(FName("BlueprintInspector"))
 		.IsPropertyEditingEnabledDelegate(FIsPropertyEditingEnabled::CreateStatic([] { return false; }));
@@ -24,11 +25,15 @@ FSCSDiff::FSCSDiff(const UBlueprint* InBlueprint)
 		.Orientation(Orient_Vertical)
 		+ SSplitter::Slot()
 		[
-			SAssignNew(SCSEditor, SSCSEditor, TSharedPtr<FBlueprintEditor>(), InBlueprint->SimpleConstructionScript, const_cast<UBlueprint*>(InBlueprint), Inspector)
+			SAssignNew(SCSEditor, SSCSEditor, InBlueprint->SimpleConstructionScript)
+				.InEditingMode(false)
+				.HideComponentClassCombo(true)
+				.OnUpdateSelectionFromNodes(SSCSEditor::FOnUpdateSelectionFromNodes::CreateRaw(this, &FSCSDiff::OnSCSEditorUpdateSelectionFromNodes))
+				.OnHighlightPropertyInDetailsView(SSCSEditor::FOnHighlightPropertyInDetailsView::CreateRaw(this, &FSCSDiff::OnSCSEditorHighlightPropertyInDetailsView))
 		]
 		+ SSplitter::Slot()
 		[
-			Inspector
+			Inspector.ToSharedRef()
 		];
 }
 
@@ -77,3 +82,32 @@ TArray< FSCSResolvedIdentifier > FSCSDiff::GetDisplayedHierarchy() const
 	return Ret;
 }
 
+void FSCSDiff::OnSCSEditorUpdateSelectionFromNodes(const TArray<FSCSEditorTreeNodePtrType>& SelectedNodes)
+{
+	FText InspectorTitle = FText::GetEmpty();
+	TArray<UObject*> InspectorObjects;
+	InspectorObjects.Empty(SelectedNodes.Num());
+	for (auto NodeIt = SelectedNodes.CreateConstIterator(); NodeIt; ++NodeIt)
+	{
+		auto NodePtr = *NodeIt;
+		if(NodePtr.IsValid() && NodePtr->CanEditDefaults())
+		{
+			InspectorTitle = FText::FromString(NodePtr->GetDisplayString());
+			InspectorObjects.Add(NodePtr->GetComponentTemplate());
+		}
+	}
+
+	if( Inspector.IsValid() )
+	{
+		SKismetInspector::FShowDetailsOptions Options(InspectorTitle, true);
+		Inspector->ShowDetailsForObjects(InspectorObjects, Options);
+	}
+}
+
+void FSCSDiff::OnSCSEditorHighlightPropertyInDetailsView(const FPropertyPath& InPropertyPath)
+{
+	if( Inspector.IsValid() )
+	{
+		Inspector->GetPropertyView()->HighlightProperty(InPropertyPath);
+	}
+}
\ No newline at end of file
diff --git a/Engine/Source/Editor/Kismet/Private/SMyBlueprint.cpp b/Engine/Source/Editor/Kismet/Private/SMyBlueprint.cpp
--- a/Engine/Source/Editor/Kismet/Private/SMyBlueprint.cpp
+++ b/Engine/Source/Editor/Kismet/Private/SMyBlueprint.cpp
@@ -1506,7 +1506,7 @@ TSharedPtr<SWidget> SMyBlueprint::OnContextMenuOpening()
 			if ( ComponentProperty && ComponentProperty->PropertyClass &&
 				 ComponentProperty->PropertyClass->IsChildOf( UActorComponent::StaticClass() ) )
 			{
-				if( FBlueprintEditor::CanClassGenerateEvents( ComponentProperty->PropertyClass ))
+				if( FBlueprintEditorUtils::CanClassGenerateEvents( ComponentProperty->PropertyClass ))
 				{
 					TSharedPtr<FBlueprintEditor> BlueprintEditor(BlueprintEditorPtr.Pin());
 
@@ -1514,8 +1514,10 @@ TSharedPtr<SWidget> SMyBlueprint::OnContextMenuOpening()
 					// of events appropriate to the component.
 					MenuBuilder.AddSubMenu(	LOCTEXT("AddEventSubMenu", "Add Event"), 
 											LOCTEXT("AddEventSubMenu_ToolTip", "Add Event"), 
-											FNewMenuDelegate::CreateStatic(	&SSCSEditor::BuildMenuEventsSection, BlueprintEditor, ComponentProperty->PropertyClass, 
-											FGetSelectedObjectsDelegate::CreateSP(this, &SMyBlueprint::GetSelectedItemsForContextMenu)));
+											FNewMenuDelegate::CreateStatic(	&SSCSEditor::BuildMenuEventsSection,
+												BlueprintEditor->GetBlueprintObj(), ComponentProperty->PropertyClass, 
+												FCanExecuteAction::CreateSP(BlueprintEditor.Get(), &FBlueprintEditor::InEditingMode),
+												FGetSelectedObjectsDelegate::CreateSP(this, &SMyBlueprint::GetSelectedItemsForContextMenu)));
 				}
 			}
 		}
diff --git a/Engine/Source/Editor/Kismet/Private/SSCSEditor.cpp b/Engine/Source/Editor/Kismet/Private/SSCSEditor.cpp
--- a/Engine/Source/Editor/Kismet/Private/SSCSEditor.cpp
+++ b/Engine/Source/Editor/Kismet/Private/SSCSEditor.cpp
@@ -1448,17 +1448,15 @@ void SSCS_RowWidget::OnAttachToDropAction(const TArray<FSCSEditorTreeNodePtrType
 	TSharedPtr<SSCSEditor> SCSEditorPtr = SCSEditor.Pin();
 	check(SCSEditorPtr.IsValid());
 
-	TSharedPtr<FBlueprintEditor> Kismet2Ptr = SCSEditorPtr->Kismet2Ptr.Pin();
-	check(Kismet2Ptr.IsValid());
-
-	AActor* PreviewActor = Kismet2Ptr->GetPreviewActor();
-	check(PreviewActor);
-
 	// Get the current Blueprint context
 	check(SCSEditorPtr->SCS);
 	UBlueprint* Blueprint = SCSEditorPtr->SCS->GetBlueprint();
 	check(Blueprint);
 
+	// Get the current Actor context
+	AActor* PreviewActor = SCSEditorPtr->PreviewActor.Get();
+	check(PreviewActor != nullptr);
+
 	bool bRegenerateTreeNodes = false;
 	const FScopedTransaction TransactionContext(DroppedNodePtrs.Num() > 1 ? LOCTEXT("AttachComponents", "Attach Components") : LOCTEXT("AttachComponent", "Attach Component"));
 
@@ -1576,11 +1574,9 @@ void SSCS_RowWidget::OnDetachFromDropAction(const TArray<FSCSEditorTreeNodePtrTy
 	TSharedPtr<SSCSEditor> SCSEditorPtr = SCSEditor.Pin();
 	check(SCSEditorPtr.IsValid());
 
-	TSharedPtr<FBlueprintEditor> Kismet2Ptr = SCSEditorPtr->Kismet2Ptr.Pin();
-	check(Kismet2Ptr.IsValid());
-
-	AActor* PreviewActor = Kismet2Ptr->GetPreviewActor();
-	check(PreviewActor);
+	// Get the current preview Actor context
+	AActor* PreviewActor = SCSEditorPtr->PreviewActor.Get();
+	check(PreviewActor != nullptr);
 
 	const FScopedTransaction TransactionContext(DroppedNodePtrs.Num() > 1 ? LOCTEXT("DetachComponents", "Detach Components") : LOCTEXT("DetachComponent", "Detach Component"));
 
@@ -1910,12 +1906,16 @@ void SSCS_RowWidget::OnNameTextCommit(const FText& InNewName, ETextCommit::Type
 // SSCSEditor
 
 
-void SSCSEditor::Construct( const FArguments& InArgs, TSharedPtr<FBlueprintEditor> InKismet2, USimpleConstructionScript* InSCS, UBlueprint* InBlueprint, TSharedPtr<SKismetInspector> Inspector )
+void SSCSEditor::Construct( const FArguments& InArgs, USimpleConstructionScript* InSCS )
 {
-	Kismet2Ptr = InKismet2;
-	KismetInspectorPtr = Inspector.IsValid() ? Inspector : InKismet2->GetInspector();
 	SCS = InSCS;
-	Blueprint = InBlueprint;
+	Blueprint = InSCS->GetBlueprint();
+
+	PreviewActor = InArgs._PreviewActor;
+	bInEditingMode = InArgs._InEditingMode;
+	OnTreeViewSelectionChanged = InArgs._OnTreeViewSelectionChanged;
+	OnUpdateSelectionFromNodes = InArgs._OnUpdateSelectionFromNodes;
+	OnHighlightPropertyInDetailsView = InArgs._OnHighlightPropertyInDetailsView;
 
 	CommandList = MakeShareable( new FUICommandList );
 	CommandList->MapAction( FGenericCommands::Get().Cut,
@@ -1988,7 +1988,11 @@ void SSCSEditor::Construct( const FArguments& InArgs, TSharedPtr<FBlueprintEdito
 		);
 
 	TSharedPtr<SWidget> Contents;
-	if( InKismet2.IsValid() )
+	if( InArgs._HideComponentClassCombo.Get() )
+	{
+		Contents = SCSTreeWidget;
+	}
+	else
 	{
 		Contents = SNew(SVerticalBox)
 
@@ -2005,7 +2009,6 @@ void SSCSEditor::Construct( const FArguments& InArgs, TSharedPtr<FBlueprintEdito
 				.OnComponentClassSelected(this, &SSCSEditor::PerformComboAddClass)
 			]
 		]
-
 		// Tree
 		+ SVerticalBox::Slot()
 		.Padding(0.f, 0.f, 0.f, 2.f)
@@ -2013,10 +2016,6 @@ void SSCSEditor::Construct( const FArguments& InArgs, TSharedPtr<FBlueprintEdito
 			SCSTreeWidget.ToSharedRef()
 		];
 	}
-	else
-	{
-		Contents = SCSTreeWidget;
-	}
 
 	this->ChildSlot
 	[
@@ -2112,18 +2111,18 @@ TSharedPtr< SWidget > SSCSEditor::CreateContextMenu()
 					}
 				}
 
-				TSharedPtr<FBlueprintEditor> PinnedBlueprintEditor = Kismet2Ptr.Pin();
-				if ( PinnedBlueprintEditor.IsValid() && SelectionClasses.Num() )
+				if ( SelectionClasses.Num() )
 				{
 					// Find the common base class of all selected classes
 					UClass* SelectedClass = UClass::FindCommonBase( SelectionClasses );
 					// Build an event submenu if we can generate events
-					if( FBlueprintEditor::CanClassGenerateEvents( SelectedClass ))
+					if( FBlueprintEditorUtils::CanClassGenerateEvents( SelectedClass ))
 					{
 						MenuBuilder.AddSubMenu(	LOCTEXT("AddEventSubMenu", "Add Event"), 
 							LOCTEXT("ActtionsSubMenu_ToolTip", "Add Event"), 
-							FNewMenuDelegate::CreateStatic( &SSCSEditor::BuildMenuEventsSection, PinnedBlueprintEditor, SelectedClass, 
-							FGetSelectedObjectsDelegate::CreateSP(this, &SSCSEditor::GetSelectedItemsForContextMenu)));
+							FNewMenuDelegate::CreateStatic( &SSCSEditor::BuildMenuEventsSection,
+								Blueprint, SelectedClass, FCanExecuteAction::CreateSP(this, &SSCSEditor::InEditingMode),
+								FGetSelectedObjectsDelegate::CreateSP(this, &SSCSEditor::GetSelectedItemsForContextMenu)));
 					}
 				}
 			}
@@ -2139,12 +2138,11 @@ TSharedPtr< SWidget > SSCSEditor::CreateContextMenu()
 	return TSharedPtr<SWidget>();
 }
 
-void SSCSEditor::BuildMenuEventsSection(FMenuBuilder& Menu, TSharedPtr<FBlueprintEditor> BlueprintEditor, UClass* SelectedClass, FGetSelectedObjectsDelegate GetSelectedObjectsDelegate)
+void SSCSEditor::BuildMenuEventsSection(FMenuBuilder& Menu, UBlueprint* Blueprint, UClass* SelectedClass, FCanExecuteAction CanExecuteActionDelegate, FGetSelectedObjectsDelegate GetSelectedObjectsDelegate)
 {
 	// Get Selected Nodes
 	TArray<FComponentEventConstructionData> SelectedNodes;
 	GetSelectedObjectsDelegate.ExecuteIfBound( SelectedNodes );
-	UBlueprint* Blueprint = BlueprintEditor->GetBlueprintObj();
 
 	struct FMenuEntry
 	{
@@ -2178,8 +2176,7 @@ void SSCSEditor::BuildMenuEventsSection(FMenuBuilder& Menu, TSharedPtr<FBlueprin
 						FMenuEntry NewEntry;
 						NewEntry.Label = ( SelectedNodes.Num() > 1 ) ?	FText::Format( LOCTEXT("ViewEvent_ToolTipFor", "{0} for {1}"), FText::FromName( EventName ), FText::FromName( VariableName )) : 
 																		FText::Format( LOCTEXT("ViewEvent_ToolTip", "{0}"), FText::FromName( EventName ));
-						NewEntry.UIAction =	FUIAction(FExecuteAction::CreateStatic( &SSCSEditor::ViewEvent, Blueprint, EventName, *NodeIter ),
-													  FCanExecuteAction::CreateSP( BlueprintEditor.Get(), &FBlueprintEditor::InEditingMode ));
+						NewEntry.UIAction =	FUIAction(FExecuteAction::CreateStatic( &SSCSEditor::ViewEvent, Blueprint, EventName, *NodeIter ), CanExecuteActionDelegate);
 						NodeActions.Add( NewEntry );
 						ComponentEventViewEntries++;
 					}
@@ -2190,8 +2187,7 @@ void SSCSEditor::BuildMenuEventsSection(FMenuBuilder& Menu, TSharedPtr<FBlueprin
 			// Create menu Add entry
 				FMenuEntry NewEntry;
 				NewEntry.Label = FText::Format( LOCTEXT("AddEvent_ToolTip", "Add {0}" ), FText::FromName( EventName ));
-				NewEntry.UIAction =	FUIAction(FExecuteAction::CreateStatic( &SSCSEditor::CreateEventsForSelection, BlueprintEditor, EventName, GetSelectedObjectsDelegate),
-											  FCanExecuteAction::CreateSP(BlueprintEditor.Get(), &FBlueprintEditor::InEditingMode));
+				NewEntry.UIAction =	FUIAction(FExecuteAction::CreateStatic( &SSCSEditor::CreateEventsForSelection, Blueprint, EventName, GetSelectedObjectsDelegate), CanExecuteActionDelegate);
 				Actions.Add( NewEntry );
 		}
 	}
@@ -2211,11 +2207,10 @@ void SSCSEditor::BuildMenuEventsSection(FMenuBuilder& Menu, TSharedPtr<FBlueprin
 	Menu.EndSection();
 }
 
-void SSCSEditor::CreateEventsForSelection(TSharedPtr<FBlueprintEditor> BlueprintEditor, FName EventName, FGetSelectedObjectsDelegate GetSelectedObjectsDelegate)
+void SSCSEditor::CreateEventsForSelection(UBlueprint* Blueprint, FName EventName, FGetSelectedObjectsDelegate GetSelectedObjectsDelegate)
 {	
 	if (EventName != NAME_None)
 	{
-		UBlueprint* Blueprint = BlueprintEditor->GetBlueprintObj();
 		TArray<FComponentEventConstructionData> SelectedNodes;
 		GetSelectedObjectsDelegate.ExecuteIfBound(SelectedNodes);
 
@@ -2514,9 +2509,8 @@ void SSCSEditor::HighlightTreeNode(FName TreeNodeName, const class FPropertyPath
 
 			if (Property != FPropertyPath())
 			{
-				auto KismetInspectorSPtr = KismetInspectorPtr.Pin();
-				check(KismetInspectorSPtr.IsValid());
-				KismetInspectorSPtr->GetPropertyView()->HighlightProperty(Property);
+				// Invoke the delegate to highlight the property
+				OnHighlightPropertyInDetailsView.ExecuteIfBound(Property);
 			}
 
 			return;
@@ -2534,8 +2528,6 @@ void SSCSEditor::HighlightTreeNode(const USCS_Node* Node, FName Property)
 	SelectNode( TreeNode, false );
 	if( Property != FName() )
 	{
-		auto KismetInspectorSPtr = KismetInspectorPtr.Pin();
-		check( KismetInspectorSPtr.IsValid() );
 		UActorComponent* Component = TreeNode->GetComponentTemplate();
 		UProperty* CurrentProp = FindField<UProperty>(Component->GetClass(), Property);
 		FPropertyPath Path;
@@ -2544,7 +2536,9 @@ void SSCSEditor::HighlightTreeNode(const USCS_Node* Node, FName Property)
 			FPropertyInfo NewInfo = { CurrentProp, -1 };
 			Path.ExtendPath(NewInfo);
 		}
-		KismetInspectorSPtr->GetPropertyView()->HighlightProperty( Path );
+
+		// Invoke the delegate to highlight the property
+		OnHighlightPropertyInDetailsView.ExecuteIfBound( Path );
 	}
 }
 
@@ -2674,11 +2668,9 @@ void SSCSEditor::ClearSelection()
 {
 	check(SCSTreeWidget.IsValid());
 	SCSTreeWidget->ClearSelection();
-	
-	if( KismetInspectorPtr.IsValid() )
-	{
-		KismetInspectorPtr.Pin()->ShowDetailsForObjects(TArray<UObject*>());
-	}
+
+	// Invoke the selection changed delegate with an empty selection set
+	OnTreeViewSelectionChanged.ExecuteIfBound(TArray<FSCSEditorTreeNodePtrType>());
 }
 
 
@@ -2712,7 +2704,7 @@ void SSCSEditor::SaveSCSNode( USCS_Node* Node )
 
 bool SSCSEditor::InEditingMode() const
 {
-	return Kismet2Ptr.IsValid() && NULL == GEditor->PlayWorld;
+	return bInEditingMode.Get() && nullptr == GEditor->PlayWorld;
 }
 
 UActorComponent* SSCSEditor::AddNewComponent( UClass* NewComponentClass, UObject* Asset  )
@@ -3071,26 +3063,8 @@ void SSCSEditor::RemoveComponentNode(FSCSEditorTreeNodePtrType InNodePtr)
 
 void SSCSEditor::UpdateSelectionFromNodes(const TArray<FSCSEditorTreeNodePtrType> &SelectedNodes )
 {
-	// Convert the selection set to an array of UObject* pointers
-	FText InspectorTitle = FText::GetEmpty();
-	TArray<UObject*> InspectorObjects;
-	InspectorObjects.Empty(SelectedNodes.Num());
-	for (auto NodeIt = SelectedNodes.CreateConstIterator(); NodeIt; ++NodeIt)
-	{
-		auto NodePtr = *NodeIt;
-		if(NodePtr.IsValid() && NodePtr->CanEditDefaults())
-		{
-			InspectorTitle = FText::FromString(NodePtr->GetDisplayString());
-			InspectorObjects.Add(NodePtr->GetComponentTemplate());
-		}
-	}
-
-	// Update the details panel
-	if (KismetInspectorPtr.IsValid())
-	{
-		SKismetInspector::FShowDetailsOptions Options(InspectorTitle, true);
-		KismetInspectorPtr.Pin()->ShowDetailsForObjects(InspectorObjects, Options);
-	}
+	// Invoke the delegate to update any selection details
+	OnUpdateSelectionFromNodes.ExecuteIfBound(SelectedNodes);
 }
 
 void SSCSEditor::RefreshSelectionDetails()
@@ -3100,16 +3074,14 @@ void SSCSEditor::RefreshSelectionDetails()
 
 void SSCSEditor::OnTreeSelectionChanged(FSCSEditorTreeNodePtrType, ESelectInfo::Type /*SelectInfo*/)
 {
-	UpdateSelectionFromNodes(SCSTreeWidget->GetSelectedItems());
+	// Get the selected tree nodes
+	const TArray<FSCSEditorTreeNodePtrType>& SelectedNodes = SCSTreeWidget->GetSelectedItems();
 
-	if( Kismet2Ptr.IsValid() )
-	{
-		TSharedPtr<class SSCSEditorViewport> ViewportPtr = Kismet2Ptr.Pin()->GetSCSViewport();
-		if (ViewportPtr.IsValid())
-		{
-			ViewportPtr->OnComponentSelectionChanged();
-		}
-	}
+	// Update the details view selection
+	UpdateSelectionFromNodes(SelectedNodes);
+
+	// Pass the tree selection change event along to the delegate
+	OnTreeViewSelectionChanged.ExecuteIfBound(SelectedNodes);
 
 	// Update the selection visualization
 	AActor* EditorActorInstance = SCS->GetComponentEditorActorInstance();
diff --git a/Engine/Source/Editor/Kismet/Private/SSCSEditor.h b/Engine/Source/Editor/Kismet/Private/SSCSEditor.h
--- a/Engine/Source/Editor/Kismet/Private/SSCSEditor.h
+++ b/Engine/Source/Editor/Kismet/Private/SSCSEditor.h
@@ -373,10 +373,29 @@ typedef SSCSEditorDragDropTree SSCSTreeType;
 class SSCSEditor : public SCompoundWidget
 {
 public:
-	SLATE_BEGIN_ARGS( SSCSEditor ){}
+	DECLARE_DELEGATE_OneParam(FOnTreeViewSelectionChanged, const TArray<FSCSEditorTreeNodePtrType>&);
+	DECLARE_DELEGATE_OneParam(FOnUpdateSelectionFromNodes, const TArray<FSCSEditorTreeNodePtrType>&);
+	DECLARE_DELEGATE_OneParam(FOnHighlightPropertyInDetailsView, const class FPropertyPath&);
+
+	SLATE_BEGIN_ARGS( SSCSEditor )
+		:_InEditingMode(true)
+		,_PreviewActor(nullptr)
+		,_HideComponentClassCombo(false)
+		,_OnTreeViewSelectionChanged()
+		,_OnUpdateSelectionFromNodes()
+		,_OnHighlightPropertyInDetailsView()
+		{}
+
+		SLATE_ATTRIBUTE(bool, InEditingMode)
+		SLATE_ATTRIBUTE(class AActor*, PreviewActor)
+		SLATE_ATTRIBUTE(bool, HideComponentClassCombo)
+		SLATE_EVENT(FOnTreeViewSelectionChanged, OnTreeViewSelectionChanged)
+		SLATE_EVENT(FOnUpdateSelectionFromNodes, OnUpdateSelectionFromNodes)
+		SLATE_EVENT(FOnHighlightPropertyInDetailsView, OnHighlightPropertyInDetailsView)
+
 	SLATE_END_ARGS()
 
-	void Construct(const FArguments& InArgs, TSharedPtr<FBlueprintEditor> InKismet2, USimpleConstructionScript* InSCS, UBlueprint* InBlueprint, TSharedPtr<class SKismetInspector> Inspector = TSharedPtr<class SKismetInspector>() );
+	void Construct(const FArguments& InArgs, USimpleConstructionScript* InSCS);
 
 	/** Override OnKeyDown */
 	virtual FReply OnKeyDown( const FGeometry& MyGeometry, const FKeyEvent& InKeyEvent );
@@ -432,7 +451,7 @@ public:
 	/** Called when selection in the tree changes */
 	void OnTreeSelectionChanged(FSCSEditorTreeNodePtrType InSelectedNodePtr, ESelectInfo::Type SelectInfo);
 
-	/** Update the Kismet2 inspector from the passed in nodes */
+	/** Update any associated selection (e.g. details view) from the passed in nodes */
 	void UpdateSelectionFromNodes(const TArray<FSCSEditorTreeNodePtrType> &SelectedNodes );
 
 	/** Refresh the tree control to reflect changes in the SCS */
@@ -450,11 +469,12 @@ public:
 	/**
 	 * Fills out an events section in ui.
 	 * @param Menu								the menu to add the events section into
-	 * @param BlueprintEditor					the active blueprint editor containing the selected objects
+	 * @param Blueprint							the active blueprint context being edited
 	 * @param SelectedClass						the common component class to build the events list from
+	 * @param CanExecuteActionDelegate			the delegate to query whether or not to execute the UI action
 	 * @param GetSelectedObjectsDelegate		the delegate to fill the currently select variables / components
 	 */
-	static void BuildMenuEventsSection( FMenuBuilder& Menu, TSharedPtr<FBlueprintEditor> BlueprintEditor, UClass* SelectedClass, FGetSelectedObjectsDelegate GetSelectedObjectsDelegate );
+	static void BuildMenuEventsSection( FMenuBuilder& Menu, UBlueprint* Blueprint, UClass* SelectedClass, FCanExecuteAction CanExecuteActionDelegate, FGetSelectedObjectsDelegate GetSelectedObjectsDelegate );
 
 	/**
 	 * Given an actor component, attempts to find an associated tree node.
@@ -527,12 +547,12 @@ protected:
 
 	/**
 	 * Function to create events for the current selection
-	 * @param BlueprintEditor				the active blueprint editor
+	 * @param Blueprint						the active blueprint context
 	 * @param EventName						the event to add
 	 * @param GetSelectedObjectsDelegate	the delegate to gather information about current selection
 	 * @param NodeIndex						an index to a specified node to add event for or < 0 for all selected nodes.
 	 */
-	static void CreateEventsForSelection(TSharedPtr<FBlueprintEditor> BlueprintEditor, FName EventName, FGetSelectedObjectsDelegate GetSelectedObjectsDelegate);
+	static void CreateEventsForSelection(UBlueprint* Blueprint, FName EventName, FGetSelectedObjectsDelegate GetSelectedObjectsDelegate);
 
 	/**
 	 * Function to construct an event for a node
@@ -578,12 +598,6 @@ public:
 	/** Pointer to blueprint we were created with */
 	UBlueprint* Blueprint;
 
-	/** Pointer back to owning Kismet 2 tool */
-	TWeakPtr<FBlueprintEditor> Kismet2Ptr;
-
-	/** Pointer back to companion inspector tool */
-	TWeakPtr<class SKismetInspector> KismetInspectorPtr;
-
 	/** Tree widget */
 	TSharedPtr<SSCSTreeType> SCSTreeWidget;
 
@@ -601,6 +615,21 @@ public:
 
 	/** Whether or not the deferred rename request was flagged as transactional */
 	bool bIsDeferredRenameRequestTransactional;
+
+	/** Attribute to control whether or not to allow editing. */
+	TAttribute<bool> bInEditingMode;
+
+	/** Attribute that provides access to the "preview" Actor context. */
+	TAttribute<AActor*> PreviewActor;
+
+	/** Delegate to invoke on tree view selection change. */
+	FOnTreeViewSelectionChanged OnTreeViewSelectionChanged;
+
+	/** Delegate to invoke when details should be updated. */
+	FOnUpdateSelectionFromNodes OnUpdateSelectionFromNodes;
+
+	/** Delegate to invoke when the given property should be highlighted in the details view (e.g. diff). */
+	FOnHighlightPropertyInDetailsView OnHighlightPropertyInDetailsView;
 private:
 
 	/** Flag to enable/disable component editing */
diff --git a/Engine/Source/Editor/Kismet/Public/BlueprintEditor.h b/Engine/Source/Editor/Kismet/Public/BlueprintEditor.h
--- a/Engine/Source/Editor/Kismet/Public/BlueprintEditor.h
+++ b/Engine/Source/Editor/Kismet/Public/BlueprintEditor.h
@@ -358,6 +358,15 @@ public:
 	/** Refresh the preview viewport to reflect changes in the SCS */
 	void UpdateSCSPreview(bool bUpdateNow = false);
 
+	/** Delegate invoked when the tree view selection is changed in the SCS editor widget */
+	void OnSCSEditorTreeViewSelectionChanged(const TArray< TSharedPtr<class FSCSEditorTreeNode> >& SelectedNodes);
+
+	/** Delegate to update the Inspector (details) panel from the given set of selected nodes */
+	void OnSCSEditorUpdateSelectionFromNodes(const TArray< TSharedPtr<class FSCSEditorTreeNode> >& SelectedNodes);
+
+	/** Delegate invoked when the given property should be highlighted in the Inspector (details) panel */
+	void OnSCSEditorHighlightPropertyInDetailsView(const class FPropertyPath& InPropertyPath);
+
 	/** Pin visibility accessors */
 	void SetPinVisibility(SGraphEditor::EPinVisibility Visibility);
 	bool GetPinVisibility(SGraphEditor::EPinVisibility Visibility) const { return PinVisibility == Visibility; }
@@ -435,12 +444,6 @@ public:
 	 */
 	void DoPromoteToVariable( UBlueprint* InBlueprint, UEdGraphPin* InTargetPin );		
 
-	/**
-	 * Checks for events in the argument class
-	 * @param InClass	The class to check for events.
-	 */
-	static bool CanClassGenerateEvents( UClass* InClass );
-
 	/** Called when node is spawned by keymap */
 	void OnNodeSpawnedByKeymap();
 
diff --git a/Engine/Source/Editor/Kismet/Public/SCSDiff.h b/Engine/Source/Editor/Kismet/Public/SCSDiff.h
--- a/Engine/Source/Editor/Kismet/Public/SCSDiff.h
+++ b/Engine/Source/Editor/Kismet/Public/SCSDiff.h
@@ -12,7 +12,13 @@ public:
 	TSharedRef< SWidget > TreeWidget();
 
 	TArray< FSCSResolvedIdentifier > GetDisplayedHierarchy() const;
+
+protected:
+	void OnSCSEditorUpdateSelectionFromNodes(const TArray< TSharedPtr<class FSCSEditorTreeNode> >& SelectedNodes);
+	void OnSCSEditorHighlightPropertyInDetailsView(const class FPropertyPath& InPropertyPath);
+
 private:
 	TSharedPtr< class SWidget > ContainerWidget;
 	TSharedPtr< class SSCSEditor > SCSEditor;
+	TSharedPtr< class SKismetInspector > Inspector;
 };
