CIS fix (non-unity)

diff --git a/Engine/Source/Editor/Kismet/Private/BlueprintEditor.cpp b/Engine/Source/Editor/Kismet/Private/BlueprintEditor.cpp
--- a/Engine/Source/Editor/Kismet/Private/BlueprintEditor.cpp
+++ b/Engine/Source/Editor/Kismet/Private/BlueprintEditor.cpp
@@ -8,6 +8,7 @@
 #include "Editor/UnrealEd/Public/Kismet2/BlueprintEditorUtils.h"
 #include "GraphEditor.h"
 #include "BlueprintUtilities.h"
+#include "Editor/PropertyEditor/Public/IDetailsView.h"
 #include "Editor/KismetCompiler/Public/KismetCompilerModule.h"
 #include "Editor/ClassViewer/Public/ClassViewerModule.h"
 #include "Editor/UnrealEd/Public/FileHelpers.h"
diff --git a/Engine/Source/Editor/Kismet/Private/SCSDiff.cpp b/Engine/Source/Editor/Kismet/Private/SCSDiff.cpp
--- a/Engine/Source/Editor/Kismet/Private/SCSDiff.cpp
+++ b/Engine/Source/Editor/Kismet/Private/SCSDiff.cpp
@@ -5,6 +5,7 @@
 #include "SKismetInspector.h"
 #include "SSCSEditor.h"
 #include "IDetailsView.h"
+#include "BlueprintEditorUtils.h"
 
 #include <vector>
 
