[UE-5938] blueprints - fix compile error on existing event nodes that override class methods also declared as 'const'

diff --git a/Engine/Source/Editor/Kismet/Private/SMyBlueprint.cpp b/Engine/Source/Editor/Kismet/Private/SMyBlueprint.cpp
--- a/Engine/Source/Editor/Kismet/Private/SMyBlueprint.cpp
+++ b/Engine/Source/Editor/Kismet/Private/SMyBlueprint.cpp
@@ -928,7 +928,7 @@ void SMyBlueprint::CollectAllActions(FGraphActionListBuilderBase& OutAllActions)
 		const UFunction* Function = *FunctionIt;
 		const FName FunctionName = Function->GetFName();
 
-		if (UEdGraphSchema_K2::CanKismetOverrideFunction(Function) && !ImplementedFunctions.Contains(FunctionName) && !UEdGraphSchema_K2::FunctionCanBePlacedAsEvent(Function))
+		if (UEdGraphSchema_K2::CanKismetOverrideFunction(Function) && !ImplementedFunctions.Contains(FunctionName) && !UEdGraphSchema_K2::FunctionCanBePlacedAsEvent(Function) && !FBlueprintEditorUtils::FindOverrideForFunction(Blueprint, CastChecked<UClass>(Function->GetOuter()), Function->GetFName()))
 		{
 			FString FunctionTooltip = UK2Node_CallFunction::GetDefaultTooltipForFunction(Function);
 			FText FunctionDesc = FText::FromString(Function->GetMetaData(TEXT("FriendlyName")));
diff --git a/Engine/Source/Editor/KismetCompiler/Private/KismetCompiler.cpp b/Engine/Source/Editor/KismetCompiler/Private/KismetCompiler.cpp
--- a/Engine/Source/Editor/KismetCompiler/Private/KismetCompiler.cpp
+++ b/Engine/Source/Editor/KismetCompiler/Private/KismetCompiler.cpp
@@ -2548,7 +2548,7 @@ void FKismetCompilerContext::VerifyValidOverrideEvent(const UEdGraph* Graph)
 					{
 						MessageLog.Warning(*EventNode->GetDeprecationMessage(), EventNode);
 					}
-					else
+					else if(!Function->HasAllFunctionFlags(FUNC_Const))	// ...allow legacy event nodes that override methods declared as 'const' to pass.
 					{
 						MessageLog.Error(TEXT("The function in node @@ cannot be overridden and/or placed as event"), EventNode);
 					}
