Replaced the SComponentsTree with an SCSEditor in SActorDetails
- Currently only supports blueprints

diff --git a/Engine/Source/Editor/Kismet/Private/SSCSEditor.cpp b/Engine/Source/Editor/Kismet/Private/SSCSEditor.cpp
--- a/Engine/Source/Editor/Kismet/Private/SSCSEditor.cpp
+++ b/Engine/Source/Editor/Kismet/Private/SSCSEditor.cpp
@@ -2038,6 +2038,11 @@ void SSCSEditor::Construct( const FArguments& InArgs, USimpleConstructionScript*
 	}
 }
 
+void SSCSEditor::SetSCS(USimpleConstructionScript* InSCS)
+{
+
+}
+
 FReply SSCSEditor::OnKeyDown( const FGeometry& MyGeometry, const FKeyEvent& InKeyEvent )
 {
 	if ( CommandList->ProcessCommandBindings( InKeyEvent ) )
diff --git a/Engine/Source/Editor/Kismet/Private/SSCSEditor.h b/Engine/Source/Editor/Kismet/Private/SSCSEditor.h
deleted file mode 100644
diff --git a/Engine/Source/Editor/Kismet/Public/SSCSEditor.h b/Engine/Source/Editor/Kismet/Public/SSCSEditor.h
new file mode 100644
--- /dev/null
+++ b/Engine/Source/Editor/Kismet/Public/SSCSEditor.h
@@ -0,0 +1,643 @@
+// Copyright 1998-2015 Epic Games, Inc. All Rights Reserved.
+
+#pragma once
+
+#include "BlueprintEditor.h"
+
+class SSCSEditor;
+class USCS_Node;
+class USimpleConstructionScript;
+
+// SCS tree node pointer type
+typedef TSharedPtr<class FSCSEditorTreeNode> FSCSEditorTreeNodePtrType;
+
+/**
+ * FSCSEditorTreeNodeBase
+ *
+ * Wrapper class for component template nodes displayed in the SCS editor tree widget.
+ */
+class KISMET_API FSCSEditorTreeNode : public TSharedFromThis<FSCSEditorTreeNode>
+{
+public:
+	/** 
+	 * Constructs an empty tree node.
+	 */
+	FSCSEditorTreeNode();
+
+	/**
+	 * Constructs a wrapper around a component template contained within an SCS tree node.
+	 *
+	 * @param InSCSNode The SCS tree node represented by this object.
+	 * @param bInIsInherited Whether or not the SCS tree node is inherited from a parent Blueprint class.
+	 */
+	FSCSEditorTreeNode(class USCS_Node* InSCSNode, bool bInIsInherited = false);
+
+	/**
+	 * Constructs a wrapper around a component template not contained within an SCS tree node (e.g. "native" components).
+	 *
+	 * @param InComponentTemplate The component template represented by this object.
+	 */
+	FSCSEditorTreeNode(UActorComponent* InComponentTemplate);
+
+	/** 
+	 * @return The name of the variable represented by this node.
+	 */
+	FName GetVariableName() const;
+	/**
+	 * @return The string to be used in the tree display.
+	 */
+	FString GetDisplayString() const;
+	/** 
+     * @return The SCS node that is represented by this object, or NULL if there is no SCS node associated with the component template.
+	 */
+	class USCS_Node* GetSCSNode() const { return SCSNodePtr.Get(); }
+	/**
+	 * @return The component template represented by this object.
+	 */
+	UActorComponent* GetComponentTemplate() const { return ComponentTemplatePtr.Get(); }
+	/**
+	 * Finds the component instance represented by this node contained within a given Actor instance.
+	 *
+	 * @param InActor The Actor instance to use as the container object for finding the component instance.
+	 * @param bIsPreviewActor Whether or not the given actor represents the "preview" actor in the SCS editor viewport.
+	 * @return The component instance represented by this node and contained within the given Actor instance, or NULL if not found.
+	 */
+	UActorComponent* FindComponentInstanceInActor(const AActor* InActor, bool bIsPreviewActor = false) const;
+	/** 
+	 * @return This object's parent node (or an invalid reference if no parent is assigned).
+	 */
+	FSCSEditorTreeNodePtrType GetParent() const { return ParentNodePtr; }
+	/**
+	 * @return The set of nodes which are parented to this node (read-only).
+	 */
+	const TArray<FSCSEditorTreeNodePtrType>& GetChildren() const { return Children; }
+	/**
+	 * @return The Blueprint to which this node belongs.
+	 */
+	UBlueprint* GetBlueprint() const;
+	/** 
+	 * @return Whether or not this object represents a root component.
+	 */
+	bool IsRoot() const;
+	/**
+	 * @return Whether or not this node is a direct child of the given node.
+	 */
+	bool IsDirectlyAttachedTo(FSCSEditorTreeNodePtrType InNodePtr) const { return ParentNodePtr == InNodePtr; } 
+	/**
+	 * @return Whether or not this node is a child (direct or indirect) of the given node.
+	 */
+	bool IsAttachedTo(FSCSEditorTreeNodePtrType InNodePtr) const;
+	/**
+	 * @return Whether or not this object represents a "native" component template (i.e. one that is not found in the SCS tree).
+	 */
+	bool IsNative() const { return GetSCSNode() == NULL && GetComponentTemplate() != NULL; }
+	/** 
+	 * @return Whether or not this object represents a component template that was inherited from a parent Blueprint class.
+	 */
+	bool IsInherited() const { return bIsInherited; }
+	/** 
+	 * @return Whether or not this object represents the default SCS scene root component.
+	 */
+	bool IsDefaultSceneRoot() const;
+	/** 
+	 * @return Whether or not this object represents a node that can be deleted from the SCS tree.
+	 */
+	bool CanDelete() const { return !IsNative() && !IsInherited() && !IsDefaultSceneRoot(); }
+	/** 
+	 * @return Whether or not this object represents a node that can be renamed from the SCS tree.
+	 */
+	bool CanRename() const { return !IsNative() && !IsInherited() && !IsDefaultSceneRoot(); }
+	/** 
+	 * @return Whether or not this object represents a node that can be reparented to other nodes based on its context.
+	 */
+	bool CanReparent() const { return !IsNative() && !IsInherited() && !IsDefaultSceneRoot() && Cast<USceneComponent>(GetComponentTemplate()) != NULL; }
+	/** 
+	 * @return Whether or not we can edit default properties for the component template represented by this object.
+	 */
+	bool CanEditDefaults() const;
+
+	/**
+	 * Finds the closest ancestor node in the given node set.
+	 * 
+	 * @param InNodes The given node set.
+	 * @return One of the nodes from the set, or an invalid node reference if the set does not contain any ancestor nodes.
+	 */
+	FSCSEditorTreeNodePtrType FindClosestParent(TArray<FSCSEditorTreeNodePtrType> InNodes);
+
+	/** 
+	 * Adds the given node as a child node.
+	 *
+	 * @param InChildNodePtr The node to add as a child node.
+	 */
+	void AddChild(FSCSEditorTreeNodePtrType InChildNodePtr);
+
+	/** 
+	 * Adds a child node for the given SCS node.
+	 *
+	 * @param InSCSNode The SCS node to for which to create a child node.
+	 * @param bIsInherited Whether or not the given SCS node is inherited from a parent.
+	 * @return A reference to the new child node or the existing child node if a match is found.
+	 */
+	FSCSEditorTreeNodePtrType AddChild(USCS_Node* InSCSNode, bool bIsInherited);
+
+	/** 
+	 * Adds a child node for the given component template.
+	 *
+	 * @param InComponentTemplate The component template for which to create a child node.
+	 * @return A reference to the new child node or the existing child node if a match is found.
+	 */
+	FSCSEditorTreeNodePtrType AddChild(UActorComponent* InComponentTemplate);
+
+	/** 
+	 * Attempts to find a reference to the child node that matches the given SCS node.
+	 *
+	 * @param InSCSNode The SCS node to match.
+	 * @param bRecursiveSearch Whether or not to recursively search child nodes (default == false).
+	 * @param OutDepth If non-NULL, the depth of the child node will be returned in this parameter on success (default == NULL).
+	 * @return The child node that matches the given SCS node, or an invalid node reference if no match was found.
+	 */
+	FSCSEditorTreeNodePtrType FindChild(const USCS_Node* InSCSNode, bool bRecursiveSearch = false, uint32* OutDepth = NULL) const;
+
+	/** 
+	 * Attempts to find a reference to the child node that matches the given component template.
+	 *
+	 * @param InComponentTemplate The component template instance to match.
+	 * @param bRecursiveSearch Whether or not to recursively search child nodes (default == false).
+	 * @param OutDepth If non-NULL, the depth of the child node will be returned in this parameter on success (default == NULL).
+	 * @return The child node with a component template that matches the given component template instance, or an invalid node reference if no match was found.
+	 */
+	FSCSEditorTreeNodePtrType FindChild(const UActorComponent* InComponentTemplate, bool bRecursiveSearch = false, uint32* OutDepth = NULL) const;
+
+	/** 
+	 * Removes the given node from the list of child nodes.
+	 *
+	 * @param InChildNodePtr The child node to remove.
+	 */
+	void RemoveChild(FSCSEditorTreeNodePtrType InChildNodePtr);
+
+	/** Delegate for when the context menu requests a rename */
+	DECLARE_DELEGATE(FOnRenameRequested);
+
+	/** Requests a rename on the component */
+	void OnRequestRename(bool bTransactional);
+
+	/** Renames the component */
+	void OnCompleteRename(const FText& InNewName);
+
+	/** Sets up the delegate for renaming a component */
+	void SetRenameRequestedDelegate(FOnRenameRequested InRenameRequested) { RenameRequestedDelegate = InRenameRequested; }
+
+private:
+	bool bIsInherited;
+	TWeakObjectPtr<class USCS_Node> SCSNodePtr;
+	TWeakObjectPtr<UActorComponent> ComponentTemplatePtr;
+	FSCSEditorTreeNodePtrType ParentNodePtr;
+	TArray<FSCSEditorTreeNodePtrType> Children;
+
+	/** Handles rename requests */
+	bool bNonTransactionalRename;
+	FOnRenameRequested RenameRequestedDelegate;
+};
+
+//////////////////////////////////////////////////////////////////////////
+// SSCS_RowWidget
+
+class SSCS_RowWidget : public SMultiColumnTableRow<FSCSEditorTreeNodePtrType>
+{
+public:
+	SLATE_BEGIN_ARGS( SSCS_RowWidget ){}
+	SLATE_END_ARGS()
+
+	void Construct( const FArguments& InArgs, TSharedPtr<SSCSEditor> InSCSEditor, FSCSEditorTreeNodePtrType InNodePtr, TSharedPtr<STableViewBase> InOwnerTableView  );
+	virtual ~SSCS_RowWidget();
+
+	// SMultiColumnTableRow<T> interface
+	virtual TSharedRef<SWidget> GenerateWidgetForColumn( const FName& ColumnName ) override;
+	// End of SMultiColumnTableRow<T>
+
+	// SWidget interface
+	virtual FReply OnMouseButtonDown(const FGeometry& MyGeometry, const FPointerEvent& MouseEvent) override;
+	virtual FReply OnDragDetected( const FGeometry& MyGeometry, const FPointerEvent& MouseEvent ) override;
+	virtual void OnDragEnter( const FGeometry& MyGeometry, const FDragDropEvent& DragDropEvent ) override;
+	virtual void OnDragLeave( const FDragDropEvent& DragDropEvent ) override;
+	virtual FReply OnDrop( const FGeometry& MyGeometry, const FDragDropEvent& DragDropEvent ) override;
+	// End of SWidget interface
+
+	/** Get the blueprint we are editing */
+	UBlueprint* GetBlueprint() const;
+
+	FText GetNameLabel() const;
+	FText GetTooltipText() const;
+	FSlateColor GetColorTint() const;
+	FString GetDocumentationLink() const;
+	FString GetDocumentationExcerptName() const;
+	
+	FText GetAssetName() const;
+	FText GetAssetNameTooltip() const;
+	EVisibility GetAssetVisibility() const;
+
+	EVisibility GetRootLabelVisibility() const;
+
+	/* Get the node used by the row Widget */
+	FSCSEditorTreeNodePtrType GetNode() const
+	{
+		return NodePtr;
+	}
+
+private:
+	/** Verifies the name of the component when changing it */
+	bool OnNameTextVerifyChanged(const FText& InNewText, FText& OutErrorMessage);
+
+	/** Commits the new name of the component */
+	void OnNameTextCommit(const FText& InNewName, ETextCommit::Type InTextCommit);
+
+	/** Builds a context menu popup for dropping a child node onto the scene root node */
+	TSharedPtr<SWidget> BuildSceneRootDropActionMenu(FSCSEditorTreeNodePtrType DroppedNodePtr);
+
+	/** Handler for attaching a single node to this node */
+	void OnAttachToDropAction(FSCSEditorTreeNodePtrType DroppedNodePtr)
+	{
+		TArray<FSCSEditorTreeNodePtrType> DroppedNodePtrs;
+		DroppedNodePtrs.Add(DroppedNodePtr);
+		OnAttachToDropAction(DroppedNodePtrs);
+	}
+
+	/** Handler for attaching one or more nodes to this node */
+	void OnAttachToDropAction(const TArray<FSCSEditorTreeNodePtrType>& DroppedNodePtrs);
+
+	/** Handler for detaching one or more nodes from the current parent and reattaching to the existing scene root node */
+	void OnDetachFromDropAction(const TArray<FSCSEditorTreeNodePtrType>& DroppedNodePtrs);
+
+	/** Handler for making the given node the new scene root node */
+	void OnMakeNewRootDropAction(FSCSEditorTreeNodePtrType DroppedNodePtr);
+	
+	/** Tasks to perform after handling a drop action */
+	void PostDragDropAction(bool bRegenerateTreeNodes);
+
+	/**
+	 * Retrieves an image brush signifying the specified component's mobility (could sometimes be NULL).
+	 * 
+	 * @param   SCSNode		A weak-pointer to the component node that you want a mobility icon for
+	 * @returns A pointer to the FSlateBrush to use (NULL for Static and Non-SceneComponents)
+	 */
+	FSlateBrush const* GetMobilityIconImage(TWeakObjectPtr<USCS_Node> SCSNode) const;
+
+	/**
+	 * Retrieves tooltip text describing the specified component's mobility.
+	 * 
+	 * @param   SCSNode		A weak-pointer to the component node that you want tooltip text for
+	 * @returns An FText object containing a description of the component's mobility
+	 */
+	FText GetMobilityToolTipText(TWeakObjectPtr<USCS_Node> SCSNode) const;
+public:
+	/** Pointer back to owning SCSEditor 2 tool */
+	TWeakPtr<SSCSEditor> SCSEditor;
+
+protected:
+	/** Pointer to node we represent */
+	FSCSEditorTreeNodePtrType NodePtr;
+};
+
+
+
+
+
+//////////////////////////////////////////////////////////////////////////
+// SSCSEditorDragDropTree - implements STreeView for our specific node type and adds drag/drop functionality
+class SSCSEditorDragDropTree : public STreeView<FSCSEditorTreeNodePtrType>
+{
+public:
+	SLATE_BEGIN_ARGS( SSCSEditorDragDropTree )
+		: _SCSEditor( NULL )
+		, _OnGenerateRow()
+		, _OnGetChildren()
+		, _TreeItemsSource( static_cast< TArray<FSCSEditorTreeNodePtrType>* >(NULL) ) //@todo Slate Syntax: Initializing from NULL without a cast
+		, _ItemHeight(16)
+		, _OnContextMenuOpening()
+		, _OnMouseButtonDoubleClick()
+		, _OnSelectionChanged()
+		, _OnExpansionChanged()
+		, _SelectionMode(ESelectionMode::Multi)
+		, _ClearSelectionOnClick(true)
+		, _ExternalScrollbar()
+		{}
+
+		SLATE_ARGUMENT( SSCSEditor*, SCSEditor )
+
+		SLATE_EVENT( FOnGenerateRow, OnGenerateRow )
+
+		SLATE_EVENT( FOnItemScrolledIntoView, OnItemScrolledIntoView )
+
+		SLATE_EVENT( FOnGetChildren, OnGetChildren )
+
+		SLATE_ARGUMENT( TArray<FSCSEditorTreeNodePtrType>* , TreeItemsSource )
+
+		SLATE_ATTRIBUTE( float, ItemHeight )
+
+		SLATE_EVENT( FOnContextMenuOpening, OnContextMenuOpening )
+
+		SLATE_EVENT( FOnMouseButtonDoubleClick, OnMouseButtonDoubleClick )
+
+		SLATE_EVENT( FOnSelectionChanged, OnSelectionChanged )
+
+		SLATE_EVENT( FOnExpansionChanged, OnExpansionChanged )
+
+		SLATE_ATTRIBUTE( ESelectionMode::Type, SelectionMode )
+
+		SLATE_ARGUMENT( TSharedPtr<SHeaderRow>, HeaderRow )
+
+		SLATE_ARGUMENT ( bool, ClearSelectionOnClick )
+
+		SLATE_ARGUMENT( TSharedPtr<SScrollBar>, ExternalScrollbar )
+
+	SLATE_END_ARGS()
+	/** Object construction - mostly defers to the base STreeView */
+	void Construct( const FArguments& InArgs );
+
+	/** Handler for drag over operations */
+	FReply OnDragOver( const FGeometry& MyGeometry, const FDragDropEvent& DragDropEvent );
+
+	/** Handler for drop operations */
+	FReply OnDrop( const FGeometry& MyGeometry, const FDragDropEvent& DragDropEvent );
+
+private:
+	/** Pointer to the SSCSEditor that owns this widget */
+	SSCSEditor* SCSEditor;
+};
+
+//////////////////////////////////////////////////////////////////////////
+// SSCSEditor
+
+typedef SSCSEditorDragDropTree SSCSTreeType;
+
+class KISMET_API SSCSEditor : public SCompoundWidget
+{
+public:
+	DECLARE_DELEGATE_OneParam(FOnTreeViewSelectionChanged, const TArray<FSCSEditorTreeNodePtrType>&);
+	DECLARE_DELEGATE_OneParam(FOnUpdateSelectionFromNodes, const TArray<FSCSEditorTreeNodePtrType>&);
+	DECLARE_DELEGATE_OneParam(FOnHighlightPropertyInDetailsView, const class FPropertyPath&);
+
+	SLATE_BEGIN_ARGS( SSCSEditor )
+		:_InEditingMode(true)
+		,_PreviewActor(nullptr)
+		,_HideComponentClassCombo(false)
+		,_OnTreeViewSelectionChanged()
+		,_OnUpdateSelectionFromNodes()
+		,_OnHighlightPropertyInDetailsView()
+		{}
+
+		SLATE_ATTRIBUTE(bool, InEditingMode)
+		SLATE_ATTRIBUTE(class AActor*, PreviewActor)
+		SLATE_ATTRIBUTE(bool, HideComponentClassCombo)
+		SLATE_EVENT(FOnTreeViewSelectionChanged, OnTreeViewSelectionChanged)
+		SLATE_EVENT(FOnUpdateSelectionFromNodes, OnUpdateSelectionFromNodes)
+		SLATE_EVENT(FOnHighlightPropertyInDetailsView, OnHighlightPropertyInDetailsView)
+
+	SLATE_END_ARGS()
+
+	void Construct(const FArguments& InArgs, USimpleConstructionScript* InSCS);
+
+	/** Override OnKeyDown */
+	virtual FReply OnKeyDown( const FGeometry& MyGeometry, const FKeyEvent& InKeyEvent );
+
+	/** Used by tree control - make a widget for a table row from a node */
+	TSharedRef<ITableRow> MakeTableRowWidget(FSCSEditorTreeNodePtrType InNodePtr, const TSharedRef<STableViewBase>& OwnerTable);
+
+	/** Used by tree control - get children for a specified node */
+	void OnGetChildrenForTree(FSCSEditorTreeNodePtrType InNodePtr, TArray<FSCSEditorTreeNodePtrType>& OutChildren);
+
+	/** Returns true if in editing mode */
+	bool InEditingMode() const;
+
+	/**
+	 * Sets the Simple Construction Script this tree is visualizing
+	 * @param InSCS The SCS to edit
+	 */
+	void SetSCS(USimpleConstructionScript* InSCS);
+
+	/** Adds a component to the SCS Table
+	   @param NewComponentClass	(In) The class to add
+	   @param Asset				(In) Optional asset to assign to the component
+	   @return The reference of the newly created ActorComponent */
+	UActorComponent* AddNewComponent(UClass* NewComponentClass, UObject* Asset);
+
+	/** Adds a new Node to the component Table
+	   @param NewNode	(In) The node to add
+	   @param Asset		(In) Optional asset to assign to the component
+	   @param bMarkBlueprintModified (In) Whether or not to mark the Blueprint as structurally modified
+	   @param bSetFocusToNewItem (In) Select the new item and activate the inline rename widget (default is true)
+	   @return The reference of the newly created ActorComponent */
+	UActorComponent* AddNewNode(USCS_Node* NewNode, UObject* Asset, bool bMarkBlueprintModified, bool bSetFocusToNewItem = true);
+	
+	/** Returns true if the specified component is currently selected */
+	bool IsComponentSelected(const UPrimitiveComponent* PrimComponent) const;
+
+	/** Assigns a selection override delegate to the specified component */
+	void SetSelectionOverride(UPrimitiveComponent* PrimComponent) const;
+
+	/** Cut selected node(s) */
+	void CutSelectedNodes();
+	bool CanCutNodes() const;
+
+	/** Copy selected node(s) */
+	void CopySelectedNodes();
+	bool CanCopyNodes() const;
+
+	/** Pastes previously copied node(s) */
+	void PasteNodes();
+	bool CanPasteNodes() const;
+
+	/** Removes existing selected component nodes from the SCS */
+	void OnDeleteNodes();
+	bool CanDeleteNodes() const;
+
+	/** Removes an existing component node from the tree */
+	void RemoveComponentNode(FSCSEditorTreeNodePtrType InNodePtr);
+
+	/** Called when selection in the tree changes */
+	void OnTreeSelectionChanged(FSCSEditorTreeNodePtrType InSelectedNodePtr, ESelectInfo::Type SelectInfo);
+
+	/** Update any associated selection (e.g. details view) from the passed in nodes */
+	void UpdateSelectionFromNodes(const TArray<FSCSEditorTreeNodePtrType> &SelectedNodes );
+
+	/** Refresh the tree control to reflect changes in the SCS */
+	void UpdateTree(bool bRegenerateTreeNodes = true);
+
+	/** Forces the details panel to refresh on the same objects */
+	void RefreshSelectionDetails();
+
+	/** Clears the current selection */
+	void ClearSelection();
+
+	/** Get the currently selected tree nodes */
+	TArray<FSCSEditorTreeNodePtrType> GetSelectedNodes() const;
+
+	/**
+	 * Fills out an events section in ui.
+	 * @param Menu								the menu to add the events section into
+	 * @param Blueprint							the active blueprint context being edited
+	 * @param SelectedClass						the common component class to build the events list from
+	 * @param CanExecuteActionDelegate			the delegate to query whether or not to execute the UI action
+	 * @param GetSelectedObjectsDelegate		the delegate to fill the currently select variables / components
+	 */
+	static void BuildMenuEventsSection( FMenuBuilder& Menu, UBlueprint* Blueprint, UClass* SelectedClass, FCanExecuteAction CanExecuteActionDelegate, FGetSelectedObjectsDelegate GetSelectedObjectsDelegate );
+
+	/**
+	 * Given an actor component, attempts to find an associated tree node.
+	 *
+	 * @param ActorComponent The component associated with the node.
+	 * @param bIncludeAttachedComponents Whether or not to include components attached to each node in the search (default is true).
+	 * @return A shared pointer to a tree node. The pointer will be invalid if no match could be found.
+	 */
+	FSCSEditorTreeNodePtrType GetNodeFromActorComponent(const UActorComponent* ActorComponent, bool bIncludeAttachedComponents = true) const;
+
+	/** Select the given tree node */
+	void SelectNode(FSCSEditorTreeNodePtrType InNodeToSelect, bool IsCntrlDown);
+
+	/**
+	 * Highlight a tree node and, optionally, a property with in it
+	 *
+	 * @param TreeNodeName		Name of the treenode to be highlighted
+	 * @param Property	The name of the property to be highlighted in the details view
+	 * @return True if the node was found in this Editor, otherwise false
+	 */
+	void HighlightTreeNode( FName TreeNodeName, const class FPropertyPath& Property );
+	/**
+	 * Highlight a tree node and, optionally, a property with in it
+	 *
+	 * @param Node		 A Reference to the Node SCS_Node to be highlighted
+	 * @param Property	The name of the property to be highlighted in the details view
+	 */
+	void HighlightTreeNode( const USCS_Node* Node, FName Property );
+
+	/**
+	 * Function to save current state of SimpleConstructionScript and nodes associated with it.
+	 *
+	 * @param: SimpleContructionScript object reference.
+	 */
+	static void SaveSCSCurrentState( USimpleConstructionScript* SCSObj );
+
+	/**
+	 * Function to save the current state of SCS_Node and its children
+	 *
+	 * @param: Reference of the SCS_Node to be saved
+	 */
+	static void SaveSCSNode(USCS_Node* Node);
+
+	/** Is this node still used by the Simple Construction Script */
+	bool IsNodeInSimpleConstructionScript(USCS_Node* Node) const;
+
+	/**
+	 * Fills the supplied array with the currently selected objects
+	 * @param OutSelectedItems The array to fill.
+	 */
+	void GetSelectedItemsForContextMenu(TArray<FComponentEventConstructionData>& OutSelectedItems) const;
+
+protected:
+	FString GetSelectedClassText() const;
+
+	/** Add a component from the selection in the combo box */
+	void PerformComboAddClass(TSubclassOf<UActorComponent> ComponentClass);
+
+	/** Called to display context menu when right clicking on the widget */
+	TSharedPtr< SWidget > CreateContextMenu();
+
+	/**
+	 * Requests a rename on the selected component
+	 * @param bTransactional Whether or not the rename should be transactional (i.e. undoable)
+	 */
+	void OnRenameComponent(bool bTransactional);
+
+	/** Checks to see if renaming is allowed on the selected component */
+	bool CanRenameComponent() const;
+
+	/**
+	 * Function to create events for the current selection
+	 * @param Blueprint						the active blueprint context
+	 * @param EventName						the event to add
+	 * @param GetSelectedObjectsDelegate	the delegate to gather information about current selection
+	 * @param NodeIndex						an index to a specified node to add event for or < 0 for all selected nodes.
+	 */
+	static void CreateEventsForSelection(UBlueprint* Blueprint, FName EventName, FGetSelectedObjectsDelegate GetSelectedObjectsDelegate);
+
+	/**
+	 * Function to construct an event for a node
+	 * @param Blueprint						the nodes blueprint
+	 * @param EventName						the event to add
+	 * @param EventData						the event data structure describing the node
+	 */
+	static void ConstructEvent(UBlueprint* Blueprint, const FName EventName, const FComponentEventConstructionData EventData);
+
+	/**
+	 * Function to view an event for a node
+	 * @param Blueprint						the nodes blueprint
+	 * @param EventName						the event to view
+	 * @param EventData						the event data structure describing the node
+	 */
+	static void ViewEvent(UBlueprint* Blueprint, const FName EventName, const FComponentEventConstructionData EventData);
+
+	/** Callbacks to duplicate the selected component */
+	bool CanDuplicateComponent() const;
+	void OnDuplicateComponent();
+
+	/** Helper method to add a tree node for the given SCS node */
+	FSCSEditorTreeNodePtrType AddTreeNode(USCS_Node* InSCSNode, FSCSEditorTreeNodePtrType InParentNodePtr);
+
+	/** Helper method to add a tree node for the given scene component */
+	FSCSEditorTreeNodePtrType AddTreeNode(USceneComponent* InSceneComponent);
+
+	/** Helper method to recursively find a tree node for the given SCS node starting at the given tree node */
+	FSCSEditorTreeNodePtrType FindTreeNode(const USCS_Node* InSCSNode, FSCSEditorTreeNodePtrType InStartNodePtr = FSCSEditorTreeNodePtrType()) const;
+
+	/** Helper method to recursively find a tree node for the given scene component starting at the given tree node */
+	FSCSEditorTreeNodePtrType FindTreeNode(const UActorComponent* InComponent, FSCSEditorTreeNodePtrType InStartNodePtr = FSCSEditorTreeNodePtrType()) const;
+
+	/** Callback when a component item is scrolled into view */
+	void OnItemScrolledIntoView( FSCSEditorTreeNodePtrType InItem, const TSharedPtr<ITableRow>& InWidget);
+
+	/** Returns the set of expandable nodes that are currently collapsed in the UI */
+	void GetCollapsedNodes(const FSCSEditorTreeNodePtrType& InNodePtr, TSet<FSCSEditorTreeNodePtrType>& OutCollapsedNodes) const;
+public:
+	/** Pointer to the script that we are editing */
+	USimpleConstructionScript* SCS;
+
+	/** Pointer to blueprint we were created with */
+	UBlueprint* Blueprint;
+
+	/** Tree widget */
+	TSharedPtr<SSCSTreeType> SCSTreeWidget;
+
+	/** Root set of tree */
+	TArray<FSCSEditorTreeNodePtrType> RootNodes;
+
+	/** The node that represents the root component in the scene hierarchy */
+	FSCSEditorTreeNodePtrType SceneRootNodePtr;
+
+	/** Command list for handling actions in the SSCSEditor */
+	TSharedPtr< FUICommandList > CommandList;
+
+	/** Weak reference to an SCS node that has been requested to be renamed */
+	TWeakObjectPtr<USCS_Node> DeferredRenameRequest;
+
+	/** Whether or not the deferred rename request was flagged as transactional */
+	bool bIsDeferredRenameRequestTransactional;
+
+	/** Attribute to control whether or not to allow editing. */
+	TAttribute<bool> bInEditingMode;
+
+	/** Attribute that provides access to the "preview" Actor context. */
+	TAttribute<AActor*> PreviewActor;
+
+	/** Delegate to invoke on tree view selection change. */
+	FOnTreeViewSelectionChanged OnTreeViewSelectionChanged;
+
+	/** Delegate to invoke when details should be updated. */
+	FOnUpdateSelectionFromNodes OnUpdateSelectionFromNodes;
+
+	/** Delegate to invoke when the given property should be highlighted in the details view (e.g. diff). */
+	FOnHighlightPropertyInDetailsView OnHighlightPropertyInDetailsView;
+private:
+
+	/** Flag to enable/disable component editing */
+	bool	bEnableComponentEditing;
+};
diff --git a/Engine/Source/Editor/LevelEditor/LevelEditor.Build.cs b/Engine/Source/Editor/LevelEditor/LevelEditor.Build.cs
--- a/Engine/Source/Editor/LevelEditor/LevelEditor.Build.cs
+++ b/Engine/Source/Editor/LevelEditor/LevelEditor.Build.cs
@@ -9,7 +9,6 @@ public class LevelEditor : ModuleRules
 		PrivateIncludePathModuleNames.AddRange(
 			new string[] {
 				"AssetTools",
-				"Kismet",
 				"MainFrame",
                 "PlacementMode",
 				"ReferenceViewer",
@@ -54,7 +53,8 @@ public class LevelEditor : ModuleRules
 				"EngineSettings",
 				"PropertyEditor",
 				"WebBrowser",
-                "Persona"
+                "Persona",
+                "Kismet",
 			}
 		);
 
diff --git a/Engine/Source/Editor/LevelEditor/Private/SLevelEditor.cpp b/Engine/Source/Editor/LevelEditor/Private/SLevelEditor.cpp
--- a/Engine/Source/Editor/LevelEditor/Private/SLevelEditor.cpp
+++ b/Engine/Source/Editor/LevelEditor/Private/SLevelEditor.cpp
@@ -39,6 +39,8 @@
 #include "TutorialMetaData.h"
 #include "SDockTab.h"
 #include "ComponentsTree.h"
+#include "SSCSEditor.h"
+#include "ComponentEditorUtils.h"
 
 static const FName LevelEditorBuildAndSubmitTab("LevelEditorBuildAndSubmit");
 static const FName LevelEditorStatsViewerTab("LevelEditorStatsViewer");
@@ -513,6 +515,13 @@ public:
 			}
 		};
 
+		bSelectionGuard = false;
+
+		// Event subscriptions
+		USelection::SelectionChangedEvent.AddRaw(this, &SActorDetails::OnEditorSelectionChanged);
+		GEditor->OnComponentInstanceTransformed().AddRaw(this, &SActorDetails::OnComponentInstanceTransformed);
+		GEditor->OnPostTransformComponents().AddRaw(this, &SActorDetails::OnEditorPostTransformComponents);
+
 		FPropertyEditorModule& PropPlugin = FModuleManager::LoadModuleChecked<FPropertyEditorModule>("PropertyEditor");
 		const FDetailsViewArgs DetailsViewArgs( true, true, true, false, false, GUnrealEd, false, TabIdentifier );
 		DetailsView = PropPlugin.CreateDetailView(DetailsViewArgs);
@@ -523,11 +532,7 @@ public:
 		DetailsView->SetGenericLayoutDetailsDelegate( FOnGetDetailCustomizationInstance::CreateStatic( &FLevelEditorGenericDetails::MakeInstance ) );
 
 		SAssignNew(ComponentsBox, SScrollBox)
-		+SScrollBox::Slot()
-		[
-			SAssignNew(ComponentsTree, SComponentsTree, DetailsView)
-			.Visibility(EVisibility::Collapsed)
-		];
+			.Visibility(EVisibility::Collapsed);
 
 		ChildSlot
 		[
@@ -539,6 +544,18 @@ public:
 			]
 		];
 
+		DetailsSplitter->AddSlot(0)
+		.Value(.2f)
+		[
+			ComponentsBox.ToSharedRef()
+		];
+	}
+
+	~SActorDetails()
+	{
+		USelection::SelectionChangedEvent.RemoveAll(this);
+		GEditor->OnComponentInstanceTransformed().RemoveAll(this);
+		GEditor->OnPostTransformComponents().RemoveAll(this);
 	}
 
 	void SetObjects( const TArray<UObject*>& InObjects )
@@ -546,39 +563,210 @@ public:
 		if(!DetailsView->IsLocked() )
 		{
 			DetailsView->SetObjects( InObjects );
-				
+			
+			bool bShowingComponents = false;
+			
 			if( GetDefault<UEditorExperimentalSettings>()->bInWorldBPEditing )
 			{
-				bool bTreeVisible = ComponentsTree->GetVisibility().IsVisible();
-				ComponentsTree->SetObjects(InObjects);
-				if( bTreeVisible != ComponentsTree->GetVisibility().IsVisible())
+				auto Actor = GetSelectedActor();
+				if (Actor)
 				{
-					if( bTreeVisible )
-					{
-						DetailsSplitter->RemoveAt(0);
-					}
-					else
+					auto Blueprint = Cast<UBlueprint>(Actor->GetClass()->ClassGeneratedBy);
+					if (Blueprint && Blueprint->SimpleConstructionScript)
 					{
-						DetailsSplitter->AddSlot(0)
-						.Value(.2f)
+						bShowingComponents = true;
+
+						ComponentsBox->ClearChildren();
+						ComponentsBox->AddSlot()
 						[
-							ComponentsBox.ToSharedRef()
+							SAssignNew(SCSEditor, SSCSEditor, Blueprint->SimpleConstructionScript)
+							.InEditingMode(true)
+							.PreviewActor(this, &SActorDetails::GetSelectedActor)												// Get the instance of the actor in the world
+							.OnTreeViewSelectionChanged(this, &SActorDetails::OnSCSEditorTreeViewSelectionChanged)				// A selection has been made in the tree view, so inform the level editor
+							//.OnUpdateSelectionFromNodes(this, &SLevelEditor::OnSCSEditorUpdateSelectionFromNodes)				// Unsure, don't think it's needed
+							//.OnHighlightPropertyInDetailsView(this, &SLevelEditor::OnSCSEditorHighlightPropertyInDetailsView)	// Also unsure and don't think it's needed
 						];
 					}
 				}
 			}
-			else
+
+			ComponentsBox->SetVisibility(bShowingComponents ? EVisibility::Visible : EVisibility::Collapsed);
+		}
+	}
+
+	const TSharedPtr<SSCSEditor>& GetSCSEditor() const
+	{
+		return SCSEditor;
+	}
+
+private:
+
+	void OnComponentInstanceTransformed(USceneComponent& Component, const FVector& InDeltaDrag, const FRotator& InDeltaRot, const FVector& InDeltaScale)
+	{
+		// Locate and update the template component the corresponds to the transformed component instance
+		auto SCSTreeNode = SCSEditor->GetNodeFromActorComponent(&Component);
+		if (SCSTreeNode.IsValid())
+		{
+			auto ComponentTemplate = Cast<USceneComponent>(SCSTreeNode->GetComponentTemplate());
+			check(ComponentTemplate);
+
+			GEditor->ApplyDeltaToComponent(
+				ComponentTemplate,
+				true,
+				&InDeltaDrag,
+				&InDeltaRot,
+				&InDeltaScale,
+				GEditor->GetPivotLocation());
+
+			auto PreviewActor = GetSelectedActor();
+			check(PreviewActor);
+
+			auto Blueprint = Cast<UBlueprint>(PreviewActor->GetClass()->ClassGeneratedBy);
+			if (Blueprint != nullptr)
 			{
-				ComponentsTree->SetVisibility( EVisibility::Collapsed );
+				//@todo this doesn't include any of the "PropagateChange" logic/calls that are in FSCSEditorViewportClient::InputWidgetDelta()
+
+				// Like PostEditMove(), but we only need to re-run construction scripts
+				if (Blueprint->bRunConstructionScriptOnDrag)
+				{
+					PreviewActor->RerunConstructionScripts();
+				}
+
+				// Update the component selection to the new reconstructed component
+				USelection* EditorComponentSelection = GEditor->GetSelectedComponents();
+				EditorComponentSelection->Modify();
+				EditorComponentSelection->Deselect(&Component);
+				
+				UActorComponent* ReconstructedComponent = SCSTreeNode->FindComponentInstanceInActor(PreviewActor);
+				check(ReconstructedComponent);
+				EditorComponentSelection->Select(ReconstructedComponent);
+				DetailsView->SetObject(ReconstructedComponent);
+			}
+		}
+	}
+
+	void OnEditorPostTransformComponents()
+	{
+		// Update the component selection with the reconstructed components
+		if (!bSelectionGuard)
+		{
+			// Enable the selection guard to prevent OnEditorSelectionChanged() from altering the contents of the SCSTreeWidget
+			TGuardValue<bool> SelectionGuard(bSelectionGuard, true);
+
+			// Update the editor's component selection
+			TArray<UObject*> Objects;
+			USelection* SelectedComponents = GEditor->GetSelectedComponents();
+
+			SelectedComponents->Modify();
+			SelectedComponents->BeginBatchSelectOperation();
+			SelectedComponents->DeselectAll();
+
+			auto SelectedActor = GetSelectedActor();
+			check(SelectedActor);
+
+			auto SelectedNodes = SCSEditor->GetSelectedNodes();
+			for (auto& SelectedNode : SelectedNodes)
+			{
+				if (SelectedNode.IsValid())
+				{
+					UActorComponent* Component = SelectedNode->FindComponentInstanceInActor(SelectedActor);
+					if (Component)
+					{
+						Objects.Add(Component);
+						SelectedComponents->Select(Component);
+					}
+				}
+			}
+			
+			DetailsView->SetObjects(Objects, false);
+			SelectedComponents->EndBatchSelectOperation();
+		}
+	}
+
+	void OnEditorSelectionChanged(UObject* Object)
+	{
+		if (!bSelectionGuard && SCSEditor.IsValid())
+		{
+			TArray<UObject*> Objects;
+			auto Selection = Cast<USelection>(Object);
+			if (Selection == GEditor->GetSelectedComponents() || Selection == GEditor->GetSelectedActors())
+			{
+				// Enable the selection guard to prevent OnTreeSelectionChanged() from altering the editor's component selection
+				TGuardValue<bool> SelectionGuard(bSelectionGuard, true);
+
+				// Update the tree selection to match the level editor component selection
+				auto& SCSTreeWidget = SCSEditor->SCSTreeWidget;
+				SCSTreeWidget->ClearSelection();
+				for (FSelectionIterator It(GEditor->GetSelectedComponentIterator()); It; ++It)
+				{
+					UActorComponent* Component = CastChecked<UActorComponent>(*It);
+					if (Component)
+					{
+						SCSTreeWidget->SetItemSelection(SCSEditor->GetNodeFromActorComponent(Component), true);
+						Objects.Add(Component);
+					}
+				}
+
+				DetailsView->SetObjects(Objects);
+			}
+		}
+	}
+	
+	AActor* GetSelectedActor() const
+	{
+		//@todo this won't work w/ multi-select
+		return Cast<AActor>(*GEditor->GetSelectedActorIterator());
+	}
+
+	void OnSCSEditorTreeViewSelectionChanged(const TArray<FSCSEditorTreeNodePtrType>& SelectedNodes)
+	{
+		if (!bSelectionGuard/* && Actor != nullptr*/)
+		{
+			auto Actor = GetSelectedActor();
+			if (Actor)
+			{
+				// Enable the selection guard to prevent OnEditorSelectionChanged() from altering the contents of the SCSTreeWidget
+				TGuardValue<bool> SelectionGuard(bSelectionGuard, true);
+
+				// Update the editor's component selection
+				TArray<UObject*> Objects;
+				USelection* SelectedComponents = GEditor->GetSelectedComponents();
+
+				SelectedComponents->BeginBatchSelectOperation();
+				SelectedComponents->DeselectAll();
+
+				for (auto& SelectedNode : SelectedNodes)
+				{
+					if (SelectedNode.IsValid())
+					{
+						UActorComponent* Component = SelectedNode->FindComponentInstanceInActor(Actor);
+						if (Component)
+						{
+							Objects.Add(Component);
+							SelectedComponents->Select(Component);
+						}
+					}
+				}
+				DetailsView->SetObjects(Objects, false);
+
+				SelectedComponents->EndBatchSelectOperation();
+
+				GUnrealEd->SetActorSelectionFlags(Actor);
+				GUnrealEd->UpdatePivotLocationForSelection(true);
+				GEditor->RedrawLevelEditingViewports();
 			}
 		}
 	}
 
-private:
 	TSharedPtr<SSplitter> DetailsSplitter;
 	TSharedPtr<IDetailsView> DetailsView;
 	TSharedPtr<SScrollBox> ComponentsBox;
-	TSharedPtr<SComponentsTree> ComponentsTree;
+	
+	/** SCS editor */
+	TSharedPtr<SSCSEditor> SCSEditor;
+
+	bool bTreeVisible;
+	bool bSelectionGuard;
 };
 
 
diff --git a/Engine/Source/Editor/UnrealEd/Classes/Editor/EditorEngine.h b/Engine/Source/Editor/UnrealEd/Classes/Editor/EditorEngine.h
--- a/Engine/Source/Editor/UnrealEd/Classes/Editor/EditorEngine.h
+++ b/Engine/Source/Editor/UnrealEd/Classes/Editor/EditorEngine.h
@@ -585,29 +585,57 @@ public:
 	DECLARE_EVENT_OneParam( UEditorEngine, FOnBeginTransformObject, UObject& );
 	FOnBeginTransformObject& OnBeginObjectMovement() { return OnBeginObjectTransformEvent; }
 
-	/** Editor-only event triggered before after actor or component has moved, rotated or scaled by an editor system */
+	/** Editor-only event triggered after actor or component has moved, rotated or scaled by an editor system */
 	DECLARE_EVENT_OneParam( UEditorEngine, FOnEndTransformObject, UObject& );
 	FOnEndTransformObject& OnEndObjectMovement() { return OnEndObjectTransformEvent; }
 
+	/** Editor-only event triggered when a component instance in the world is transformed */
+	DECLARE_EVENT_FourParams(UEditorEngine, FOnComponentInstanceTransformed, USceneComponent&, const FVector&, const FRotator&, const FVector&);
+	FOnComponentInstanceTransformed& OnComponentInstanceTransformed() { return OnComponentInstanceTransformedEvent; }
+
+	/** Editor-only event triggered when transformation of a component in the world finishes */
+	DECLARE_EVENT(UEditorEngine, FOnPostTransformComponents);
+	FOnPostTransformComponents& OnPostTransformComponents() { return OnPostTransformComponentsEvent; }
+
 	/** Delegate broadcast by the engine every tick when PIE/SIE is active, to check to see whether we need to
 		be able to capture state for simulating actor (for Sequencer recording features).  The single bool parameter
 		should be set to true if recording features are needed. */
 	DECLARE_EVENT_OneParam( UEditorEngine, FGetActorRecordingState, bool& /* bIsRecordingActive */ );
 	FGetActorRecordingState& GetActorRecordingState() { return GetActorRecordingStateEvent; }
 
-	/** 
-	 * Called before an actor or component is about to be translated, rotated, or scaled by the editor 
-	 * 
-	 * @param Object	The actor or component that will be moved
-	 */
-	void BroadcastBeginObjectMovement( UObject& Object ) const { OnBeginObjectTransformEvent.Broadcast( Object ); }
+	/**
+	* Called before an actor or component is about to be translated, rotated, or scaled by the editor
+	*
+	* @param Object	The actor or component that will be moved
+	*/
+	void BroadcastBeginObjectMovement(UObject& Object) const { OnBeginObjectTransformEvent.Broadcast(Object); }
 
 	/**
-	 * Called when an actor or component has been translated, rotated, or scaled by the editor
-	 *
-	 * @param Object	The actor or component that moved
-	 */
-	void BroadcastEndObjectMovement( UObject& Object ) const { OnEndObjectTransformEvent.Broadcast( Object ); }
+	* Called when an actor or component has been translated, rotated, or scaled by the editor
+	*
+	* @param Object	The actor or component that moved
+	*/
+	void BroadcastEndObjectMovement(UObject& Object) const { OnEndObjectTransformEvent.Broadcast(Object); }
+
+	//@todo dhertzka - this will fail when multiple components are selected; should just send the transform data and update each component based on GetSelectedComponents
+	/**
+	* Called when a component instance in the world has been translated, rotated, or scaled
+	*
+	* @param Component	The component instance the world that was transformed
+	* @param InDeltaDrag
+	* @param InDeltaRot
+	* @param InDeltaScale
+	*/
+	void BroadcastComponentInstanceTransformed(USceneComponent& Component, const FVector& InDeltaDrag, const FRotator& InDeltaRot, const FVector& InDeltaScale) const
+	{
+		OnComponentInstanceTransformedEvent.Broadcast(Component, InDeltaDrag, InDeltaRot, InDeltaScale);
+	}
+
+	/**
+	* Called after transformation of any component(s) in the world is complete
+	*/
+	void BroadcastPostTransformComponents() { OnPostTransformComponentsEvent.Broadcast(); }
+
 
 	/**	Broadcasts that an object has been reimported. THIS SHOULD NOT BE PUBLIC */
 	void BroadcastObjectReimported(UObject* InObject);
@@ -2598,6 +2626,12 @@ private:
 	/** Delegate broadcast when an actor or component has been moved, rotated, or scaled */
 	FOnEndTransformObject OnEndObjectTransformEvent;
 
+	/** Delegate broadcast when a component instance in the world has been moved, rotated, or scaled */
+	FOnComponentInstanceTransformed OnComponentInstanceTransformedEvent;
+
+	/** Delegate broadcast when transformation of any component(s) in the world is complete */
+	FOnPostTransformComponents OnPostTransformComponentsEvent;
+
 	/** Delegate broadcast by the engine every tick when PIE/SIE is active, to check to see whether we need to
 		be able to capture state for simulating actor (for Sequencer recording features) */
 	FGetActorRecordingState GetActorRecordingStateEvent;
diff --git a/Engine/Source/Editor/UnrealEd/Private/LevelEditorViewport.cpp b/Engine/Source/Editor/UnrealEd/Private/LevelEditorViewport.cpp
--- a/Engine/Source/Editor/UnrealEd/Private/LevelEditorViewport.cpp
+++ b/Engine/Source/Editor/UnrealEd/Private/LevelEditorViewport.cpp
@@ -2532,30 +2532,38 @@ void FLevelEditorViewportClient::TrackingStopped()
 	// Don't do this if AddDelta was never called.
 	if( bDidAnythingActuallyChange && MouseDeltaTracker->HasReceivedDelta() )
 	{
-		for ( FSelectionIterator It( GEditor->GetSelectedActorIterator() ) ; It ; ++It )
+		// If components are selected, we'll need to notify after rerunning the construction script on the selected actor(s)
+		const bool bComponentsAreSelected = GEditor->GetSelectedComponentCount() > 0;
+
+		for (FSelectionIterator It(GEditor->GetSelectedActorIterator()); It; ++It)
 		{
 			AActor* Actor = static_cast<AActor*>( *It );
-			checkSlow( Actor->IsA(AActor::StaticClass()) );
+			checkSlow(Actor->IsA(AActor::StaticClass()));
 
 			// Verify that the actor is in the same world as the viewport before moving it.
-			if(GEditor->PlayWorld)
+			if (GEditor->PlayWorld)
 			{
-				if(bIsSimulateInEditorViewport)
+				if (bIsSimulateInEditorViewport)
 				{
 					// If the Actor's outer (level) outer (world) is not the PlayWorld then it cannot be moved in this viewport.
-					if( !(GEditor->PlayWorld == Actor->GetOuter()->GetOuter()) )
+					if (!( GEditor->PlayWorld == Actor->GetOuter()->GetOuter() ))
 					{
 						continue;
 					}
 				}
-				else if( !(GEditor->EditorWorld == Actor->GetOuter()->GetOuter()) )
+				else if (!( GEditor->EditorWorld == Actor->GetOuter()->GetOuter() ))
 				{
 					continue;
 				}
 			}
 
-			Actor->PostEditMove( true );
-			GEditor->BroadcastEndObjectMovement( *Actor );
+			Actor->PostEditMove(true);
+			GEditor->BroadcastEndObjectMovement(*Actor);
+		}
+
+		if (bComponentsAreSelected)
+		{
+			GEditor->BroadcastPostTransformComponents();
 		}
 
 		if (!bPivotMovedIndependently)
@@ -3040,6 +3048,8 @@ void FLevelEditorViewportClient::ApplyDeltaToComponent(USceneComponent* InCompon
 		&InDeltaRot,
 		&ModifiedDeltaScale,
 		GEditor->GetPivotLocation());
+
+	GEditor->BroadcastComponentInstanceTransformed(*InComponent, InDeltaDrag, InDeltaRot, ModifiedDeltaScale);
 }
 
 /** Helper function for ModifyScale - Convert the active Dragging Axis to per-axis flags */
