[UE-7430] SSCSEditor now binds to an Actor instance rather than an explicit SCS (to better support shared usage)

change summary:
- access to SCS and Blueprint contexts within the SSCSEditor widget are now obtained through the USCS_Node reference (when valid)
- new delegates added to handle USCS_Node creation (and potentially BPGC creation prior to instancing a new node)
- updating the tree view no longer requires a valid SCS/BPGC context; thus native class component hierarchies can now be "viewed" using the SSCSEditor widget
- replaced the 'PreviewActor' attribute with a more general 'ActorContext' attribute; whether or not it's the "preview" actor is now based on the type of world it's registered in
- fixed up mode activation in the BP editor to ensure that the tree view is updated after the "preview" Actor instance is created (in order to ensure that the tree view is populated)

diff --git a/Engine/Source/Editor/Kismet/Private/BlueprintDetailsCustomization.cpp b/Engine/Source/Editor/Kismet/Private/BlueprintDetailsCustomization.cpp
--- a/Engine/Source/Editor/Kismet/Private/BlueprintDetailsCustomization.cpp
+++ b/Engine/Source/Editor/Kismet/Private/BlueprintDetailsCustomization.cpp
@@ -4691,7 +4691,6 @@ bool FBlueprintComponentDetails::IsNodeAttachable() const
 			check(BlueprintEditorPtr.IsValid());
 			TSharedPtr<SSCSEditor> Editor = BlueprintEditorPtr.Pin()->GetSCSEditor();
 			check( Editor.IsValid() );
-			check( Editor->SCS );
 
 			USCS_Node* SCS_Node = CachedNodePtr->GetSCSNode();
 			if (SCS_Node != NULL && Editor->IsNodeInSimpleConstructionScript(SCS_Node))
@@ -4732,7 +4731,6 @@ void FBlueprintComponentDetails::OnBrowseSocket()
 	{
 		TSharedPtr<SSCSEditor> Editor = BlueprintEditorPtr.Pin()->GetSCSEditor();
 		check( Editor.IsValid() );
-		check( Editor->SCS );
 
 		FSCSEditorTreeNodePtrType ParentFNode = CachedNodePtr->GetParent();
 
diff --git a/Engine/Source/Editor/Kismet/Private/BlueprintEditor.cpp b/Engine/Source/Editor/Kismet/Private/BlueprintEditor.cpp
--- a/Engine/Source/Editor/Kismet/Private/BlueprintEditor.cpp
+++ b/Engine/Source/Editor/Kismet/Private/BlueprintEditor.cpp
@@ -660,12 +660,54 @@ void FBlueprintEditor::UpdateSCSPreview(bool bUpdateNow)
 	}
 }
 
+USCS_Node* FBlueprintEditor::OnSCSEditorAddNewComponent(UClass* InComponentClass)
+{
+	check(InComponentClass != nullptr);
+
+	UBlueprint* Blueprint = GetBlueprintObj();
+	check(Blueprint != nullptr && Blueprint->SimpleConstructionScript != nullptr);
+
+	return Blueprint->SimpleConstructionScript->CreateNode(InComponentClass);
+}
+
+USCS_Node* FBlueprintEditor::OnSCSEditorAddExistingComponent(UActorComponent* InComponentInstance)
+{
+	check(InComponentInstance != nullptr);
+
+	UBlueprint* Blueprint = GetBlueprintObj();
+	check(Blueprint != nullptr && Blueprint->SimpleConstructionScript != nullptr);
+
+	if(InComponentInstance->GetOuter() == GetTransientPackage())
+	{
+		// Relocate the instance from the transient package to the BPGC and assign it a unique object name
+		InComponentInstance->Rename(NULL, Blueprint->GeneratedClass, REN_DontCreateRedirectors|REN_DoNotDirty);
+	}
+
+	return Blueprint->SimpleConstructionScript->CreateNode(InComponentInstance);
+}
+
 void FBlueprintEditor::OnSCSEditorTreeViewSelectionChanged(const TArray<FSCSEditorTreeNodePtrType>& SelectedNodes)
 {
 	if (SCSViewport.IsValid())
 	{
 		SCSViewport->OnComponentSelectionChanged();
 	}
+
+	UBlueprint* Blueprint = GetBlueprintObj();
+	check(Blueprint != nullptr && Blueprint->SimpleConstructionScript != nullptr);
+
+	// Update the selection visualization
+	AActor* EditorActorInstance = Blueprint->SimpleConstructionScript->GetComponentEditorActorInstance();
+	if (EditorActorInstance != NULL)
+	{
+		TInlineComponentArray<UPrimitiveComponent*> PrimitiveComponents;
+		EditorActorInstance->GetComponents(PrimitiveComponents);
+
+		for (int32 Idx = 0; Idx < PrimitiveComponents.Num(); ++Idx)
+		{
+			PrimitiveComponents[Idx]->PushSelectionToProxy();
+		}
+	}
 }
 
 void FBlueprintEditor::OnSCSEditorUpdateSelectionFromNodes(const TArray<FSCSEditorTreeNodePtrType>& SelectedNodes)
@@ -705,14 +747,6 @@ void FBlueprintEditor::OnSCSEditorUpdateSelectionFromNodes(const TArray<FSCSEdit
 	}
 }
 
-void FBlueprintEditor::OnSCSEditorHighlightPropertyInDetailsView(const FPropertyPath& InPropertyPath)
-{
-	if(Inspector.IsValid())
-	{
-		Inspector->GetPropertyView()->HighlightProperty(InPropertyPath);
-	}
-}
-
 /** Create new tab for the supplied graph - don't call this directly.*/
 TSharedRef<SGraphEditor> FBlueprintEditor::CreateGraphEditorWidget(TSharedRef<FTabInfo> InTabInfo, UEdGraph* InGraph)
 {
@@ -1172,12 +1206,13 @@ void FBlueprintEditor::EnsureBlueprintIsUpToDate(UBlueprint* BlueprintObj)
 			BlueprintObj->SimpleConstructionScript->SetFlags(RF_Transactional);
 
 			// Recreate (or create) any widgets that depend on the SCS
-			SCSEditor = SAssignNew(SCSEditor, SSCSEditor, BlueprintObj->SimpleConstructionScript)
+			SCSEditor = SAssignNew(SCSEditor, SSCSEditor)
 				.InEditingMode(this, &FBlueprintEditor::InEditingMode)
-				.PreviewActor(this, &FBlueprintEditor::GetPreviewActor)
+				.ActorContext(this, &FBlueprintEditor::GetPreviewActor)
+				.OnAddNewComponent(this, &FBlueprintEditor::OnSCSEditorAddNewComponent)
+				.OnAddExistingComponent(this, &FBlueprintEditor::OnSCSEditorAddExistingComponent)
 				.OnTreeViewSelectionChanged(this, &FBlueprintEditor::OnSCSEditorTreeViewSelectionChanged)
-				.OnUpdateSelectionFromNodes(this, &FBlueprintEditor::OnSCSEditorUpdateSelectionFromNodes)
-				.OnHighlightPropertyInDetailsView(this, &FBlueprintEditor::OnSCSEditorHighlightPropertyInDetailsView);
+				.OnUpdateSelectionFromNodes(this, &FBlueprintEditor::OnSCSEditorUpdateSelectionFromNodes);
 			SCSViewport = SAssignNew(SCSViewport, SSCSEditorViewport)
 				.BlueprintEditor(SharedThis(this));
 		}
@@ -1956,12 +1991,13 @@ void FBlueprintEditor::CreateDefaultTabContents(const TArray<UBlueprint*>& InBlu
 		InBlueprint->ParentClass->IsChildOf(AActor::StaticClass()) && 
 		InBlueprint->SimpleConstructionScript )
 	{
-		SCSEditor = SAssignNew(SCSEditor, SSCSEditor, InBlueprint->SimpleConstructionScript)
+		SCSEditor = SAssignNew(SCSEditor, SSCSEditor)
 			.InEditingMode(this, &FBlueprintEditor::InEditingMode)
-			.PreviewActor(this, &FBlueprintEditor::GetPreviewActor)
+			.ActorContext(this, &FBlueprintEditor::GetPreviewActor)
+			.OnAddNewComponent(this, &FBlueprintEditor::OnSCSEditorAddNewComponent)
+			.OnAddExistingComponent(this, &FBlueprintEditor::OnSCSEditorAddExistingComponent)
 			.OnTreeViewSelectionChanged(this, &FBlueprintEditor::OnSCSEditorTreeViewSelectionChanged)
-			.OnUpdateSelectionFromNodes(this, &FBlueprintEditor::OnSCSEditorUpdateSelectionFromNodes)
-			.OnHighlightPropertyInDetailsView(this, &FBlueprintEditor::OnSCSEditorHighlightPropertyInDetailsView);
+			.OnUpdateSelectionFromNodes(this, &FBlueprintEditor::OnSCSEditorUpdateSelectionFromNodes);
 		SCSViewport = SAssignNew(SCSViewport, SSCSEditorViewport)
 			.BlueprintEditor(SharedThis(this));
 	}
@@ -6876,7 +6912,7 @@ void FBlueprintEditor::OnFindInstancesCustomEvent()
 
 AActor* FBlueprintEditor::GetPreviewActor() const
 {
-	return SCSViewport->GetPreviewActor();
+	return SCSViewport.IsValid() ? SCSViewport->GetPreviewActor() : nullptr;
 }
 
 FReply FBlueprintEditor::OnSpawnGraphNodeByShortcut(FInputGesture InGesture, const FVector2D& InPosition, UEdGraph* InGraph)
diff --git a/Engine/Source/Editor/Kismet/Private/BlueprintEditorModes.cpp b/Engine/Source/Editor/Kismet/Private/BlueprintEditorModes.cpp
--- a/Engine/Source/Editor/Kismet/Private/BlueprintEditorModes.cpp
+++ b/Engine/Source/Editor/Kismet/Private/BlueprintEditorModes.cpp
@@ -325,12 +325,13 @@ void FBlueprintComponentsApplicationMode::PostActivateMode()
 	TSharedPtr<FBlueprintEditor> BP = MyBlueprintEditor.Pin();
 	if (BP.IsValid())
 	{
-		auto SCSEditor = BP->GetSCSEditor();
-		SCSEditor->UpdateTree();
 		BP->EnableSCSPreview(true);
-		BP->UpdateSCSPreview();
+		BP->UpdateSCSPreview(true);
 		BP->GetInspector()->EnableComponentDetailsCustomization(true);
 
+		auto SCSEditor = BP->GetSCSEditor();
+		SCSEditor->UpdateTree();
+
 		// Reselect the cached components
 		TArray<TSharedPtr<FSCSEditorTreeNode>> Selection;
 		for (auto Component : CachedComponentSelection)
diff --git a/Engine/Source/Editor/Kismet/Private/SCSDiff.cpp b/Engine/Source/Editor/Kismet/Private/SCSDiff.cpp
--- a/Engine/Source/Editor/Kismet/Private/SCSDiff.cpp
+++ b/Engine/Source/Editor/Kismet/Private/SCSDiff.cpp
@@ -26,8 +26,9 @@ FSCSDiff::FSCSDiff(const UBlueprint* InBlueprint)
 		.Orientation(Orient_Vertical)
 		+ SSplitter::Slot()
 		[
-			SAssignNew(SCSEditor, SSCSEditor, InBlueprint->SimpleConstructionScript)
+			SAssignNew(SCSEditor, SSCSEditor)
 				.InEditingMode(false)
+				.ActorContext(InBlueprint->GeneratedClass->GetDefaultObject<AActor>())
 				.HideComponentClassCombo(true)
 				.OnUpdateSelectionFromNodes(SSCSEditor::FOnUpdateSelectionFromNodes::CreateRaw(this, &FSCSDiff::OnSCSEditorUpdateSelectionFromNodes))
 				.OnHighlightPropertyInDetailsView(SSCSEditor::FOnHighlightPropertyInDetailsView::CreateRaw(this, &FSCSDiff::OnSCSEditorHighlightPropertyInDetailsView))
diff --git a/Engine/Source/Editor/Kismet/Private/SCSEditorViewportClient.cpp b/Engine/Source/Editor/Kismet/Private/SCSEditorViewportClient.cpp
--- a/Engine/Source/Editor/Kismet/Private/SCSEditorViewportClient.cpp
+++ b/Engine/Source/Editor/Kismet/Private/SCSEditorViewportClient.cpp
@@ -234,7 +234,7 @@ void FSCSEditorViewportClient::Draw(const FSceneView* View, FPrimitiveDrawInterf
 			{
 				FSCSEditorTreeNodePtrType SelectedNode = SelectedNodes[SelectionIndex];
 
-				UActorComponent* Comp = SelectedNode->FindComponentInstanceInActor(PreviewActor, true);
+				UActorComponent* Comp = SelectedNode->FindComponentInstanceInActor(PreviewActor);
 				if(Comp != NULL && Comp->IsRegistered())
 				{
 					// Try and find a visualizer
@@ -261,7 +261,7 @@ void FSCSEditorViewportClient::DrawCanvas( FViewport& InViewport, FSceneView& Vi
 			{
 				FSCSEditorTreeNodePtrType SelectedNode = SelectedNodes[SelectionIndex];
 
-				UActorComponent* Comp = Cast<USceneComponent>(SelectedNode->FindComponentInstanceInActor(PreviewActor, true));
+				UActorComponent* Comp = Cast<USceneComponent>(SelectedNode->FindComponentInstanceInActor(PreviewActor));
 				if (Comp != NULL && Comp->IsRegistered())
 				{
 					// Try and find a visualizer
@@ -282,7 +282,7 @@ void FSCSEditorViewportClient::DrawCanvas( FViewport& InViewport, FSceneView& Vi
 		auto SelectedNodes = BlueprintEditorPtr.Pin()->GetSelectedSCSEditorTreeNodes();
 		if(bIsManipulating && SelectedNodes.Num() > 0)
 		{
-			USceneComponent* SceneComp = Cast<USceneComponent>(SelectedNodes[0]->FindComponentInstanceInActor(PreviewActor, true));
+			USceneComponent* SceneComp = Cast<USceneComponent>(SelectedNodes[0]->FindComponentInstanceInActor(PreviewActor));
 			if(SceneComp)
 			{
 				const FVector WidgetLocation = GetWidgetLocation();
@@ -392,7 +392,7 @@ bool FSCSEditorViewportClient::InputWidgetDelta( FViewport* Viewport, EAxisList:
 
 					if(bCanEdit)
 					{
-						USceneComponent* SceneComp = Cast<USceneComponent>(SelectedNodePtr->FindComponentInstanceInActor(PreviewActor, true));
+						USceneComponent* SceneComp = Cast<USceneComponent>(SelectedNodePtr->FindComponentInstanceInActor(PreviewActor));
 						USceneComponent* SelectedTemplate = Cast<USceneComponent>(SelectedNodePtr->GetComponentTemplate());
 						if(SceneComp != NULL && SelectedTemplate != NULL)
 						{
@@ -498,7 +498,7 @@ bool FSCSEditorViewportClient::InputWidgetDelta( FViewport* Viewport, EAxisList:
 										if(ArchetypeInstance != NULL)
 										{
 											const bool bIsProcessingPreviewActor = (ArchetypeInstance == PreviewActor);
-											SceneComp = Cast<USceneComponent>(SelectedNodePtr->FindComponentInstanceInActor(ArchetypeInstance, bIsProcessingPreviewActor));
+											SceneComp = Cast<USceneComponent>(SelectedNodePtr->FindComponentInstanceInActor(ArchetypeInstance));
 											if(!bIsProcessingPreviewActor && SceneComp != nullptr && !UpdatedComponents.Contains(SceneComp))
 											{
 												FComponentEditorUtils::PropagateTransformPropertyChange(SceneComp, SceneComp->RelativeLocation, OldRelativeLocation, SelectedTemplate->RelativeLocation, UpdatedComponents);
@@ -574,7 +574,7 @@ FWidget::EWidgetMode FSCSEditorViewportClient::GetWidgetMode() const
 			for ( int32 CurrentNodeIndex=0; CurrentNodeIndex < SelectedNodes.Num(); CurrentNodeIndex++ )
 			{
 				FSCSEditorTreeNodePtrType CurrentNodePtr = SelectedNodes[CurrentNodeIndex];
-				if ( CurrentNodePtr.IsValid() && !RootNodes.Contains( CurrentNodePtr ) && !CurrentNodePtr->IsRoot() && CurrentNodePtr->CanEditDefaults() && CurrentNodePtr->FindComponentInstanceInActor(PreviewActor, true))
+				if ( CurrentNodePtr.IsValid() && !RootNodes.Contains( CurrentNodePtr ) && !CurrentNodePtr->IsRoot() && CurrentNodePtr->CanEditDefaults() && CurrentNodePtr->FindComponentInstanceInActor(PreviewActor))
 				{
 					// a non-NULL, non-root item is selected, draw the widget
 					ReturnWidgetMode = WidgetMode;
@@ -609,7 +609,7 @@ FVector FSCSEditorViewportClient::GetWidgetLocation() const
 		if(SelectedNodes.Num() > 0)
 		{
 			// Use the last selected item for the widget location
-			USceneComponent* SceneComp = Cast<USceneComponent>(SelectedNodes.Last().Get()->FindComponentInstanceInActor(PreviewActor, true));
+			USceneComponent* SceneComp = Cast<USceneComponent>(SelectedNodes.Last().Get()->FindComponentInstanceInActor(PreviewActor));
 			if( SceneComp )
 			{
 				TSharedPtr<ISCSEditorCustomization> Customization = BlueprintEditorPtr.Pin()->CustomizeSCSEditor(SceneComp);
@@ -642,7 +642,7 @@ FMatrix FSCSEditorViewportClient::GetWidgetCoordSystem() const
 			if(SelectedNodes.Num() > 0)
 			{
 				const auto SelectedNode = SelectedNodes.Last();
-				USceneComponent* SceneComp = SelectedNode.IsValid() ? Cast<USceneComponent>(SelectedNode->FindComponentInstanceInActor(PreviewActor, true)) : NULL;
+				USceneComponent* SceneComp = SelectedNode.IsValid() ? Cast<USceneComponent>(SelectedNode->FindComponentInstanceInActor(PreviewActor)) : NULL;
 				if( SceneComp )
 				{
 					TSharedPtr<ISCSEditorCustomization> Customization = BlueprintEditor->CustomizeSCSEditor(SceneComp);
@@ -746,7 +746,7 @@ void FSCSEditorViewportClient::FocusViewportToSelection()
 		if(SelectedNodes.Num() > 0)
 		{
 			// Use the last selected item for the widget location
-			USceneComponent* SceneComp = Cast<USceneComponent>(SelectedNodes.Last()->FindComponentInstanceInActor(PreviewActor, true));
+			USceneComponent* SceneComp = Cast<USceneComponent>(SelectedNodes.Last()->FindComponentInstanceInActor(PreviewActor));
 			if( SceneComp )
 			{
 				FocusViewportOnBox( SceneComp->Bounds.GetBox() );
diff --git a/Engine/Source/Editor/Kismet/Private/SSCSEditor.cpp b/Engine/Source/Editor/Kismet/Private/SSCSEditor.cpp
--- a/Engine/Source/Editor/Kismet/Private/SSCSEditor.cpp
+++ b/Engine/Source/Editor/Kismet/Private/SSCSEditor.cpp
@@ -301,7 +301,7 @@ FString FSCSEditorTreeNode::GetDisplayString() const
 	}
 }
 
-UActorComponent* FSCSEditorTreeNode::FindComponentInstanceInActor(const AActor* InActor, bool bIsPreviewActor) const
+UActorComponent* FSCSEditorTreeNode::FindComponentInstanceInActor(const AActor* InActor) const
 {
 	USCS_Node* SCS_Node = GetSCSNode();
 	UActorComponent* ComponentTemplate = GetComponentTemplate();
@@ -314,13 +314,14 @@ UActorComponent* FSCSEditorTreeNode::FindComponentInstanceInActor(const AActor*
 			FName VariableName = SCS_Node->GetVariableName();
 			if(VariableName != NAME_None)
 			{
+				UWorld* World = InActor->GetWorld();
 				UObjectPropertyBase* Property = FindField<UObjectPropertyBase>(InActor->GetClass(), VariableName);
 				if(Property != NULL)
 				{
 					// Return the component instance that's stored in the property with the given variable name
 					ComponentInstance = Cast<UActorComponent>(Property->GetObjectPropertyValue_InContainer(InActor));
 				}
-				else if(bIsPreviewActor)
+				else if(World != nullptr && World->WorldType == EWorldType::Preview)
 				{
 					// If this is the preview actor, return the cached component instance that's being used for the preview actor prior to recompiling the Blueprint
 					ComponentInstance = SCS_Node->EditorComponentInstance;
@@ -1085,6 +1086,7 @@ EVisibility SSCS_RowWidget::GetRootLabelVisibility() const
 
 TSharedPtr<SWidget> SSCS_RowWidget::BuildSceneRootDropActionMenu(FSCSEditorTreeNodePtrType DroppedNodePtr)
 {
+	check(SCSEditor.IsValid());
 	FMenuBuilder MenuBuilder(true, SCSEditor.Pin()->CommandList);
 
 	MenuBuilder.BeginSection("SceneRootNodeDropActions", LOCTEXT("SceneRootNodeDropActionContextMenu", "Drop Actions"));
@@ -1092,9 +1094,8 @@ TSharedPtr<SWidget> SSCS_RowWidget::BuildSceneRootDropActionMenu(FSCSEditorTreeN
 		const FText DroppedVariableNameText = FText::FromName( DroppedNodePtr->GetVariableName() );
 		const FText NodeVariableNameText = FText::FromName( NodePtr->GetVariableName() );
 
-		check(SCSEditor.IsValid());
-		check(SCSEditor.Pin()->SCS);
-		const bool bDroppedInSameBlueprint = DroppedNodePtr->GetBlueprint() == SCSEditor.Pin()->SCS->GetBlueprint();
+		check(NodePtr.IsValid());
+		const bool bDroppedInSameBlueprint = DroppedNodePtr->GetBlueprint() == NodePtr->GetBlueprint();
 
 		MenuBuilder.AddMenuEntry(
 			LOCTEXT("DropActionLabel_AttachToRootNode", "Attach"),
@@ -1280,7 +1281,7 @@ void SSCS_RowWidget::OnDragEnter( const FGeometry& MyGeometry, const FDragDropEv
 					DragRowOp->CurrentHoverText = LOCTEXT("DropActionToolTip_AttachToOrMakeNewRoot", "Drop here to see available actions.");
 					DragRowOp->PendingDropAction = FSCSRowDragDropOp::DropAction_AttachToOrMakeNewRoot;
 				}
-				else if(DraggedNodePtr->GetBlueprint() != SCSEditor.Pin()->SCS->GetBlueprint())
+				else if(DraggedNodePtr->GetBlueprint() != NodePtr->GetBlueprint())
 				{
 					if(bCanMakeNewRoot)
 					{
@@ -1356,7 +1357,7 @@ void SSCS_RowWidget::OnDragEnter( const FGeometry& MyGeometry, const FDragDropEv
 			else if(HoveredTemplate->CanAttachAsChild(DraggedTemplate, NAME_None))
 			{
 				// Attach the dragged node(s) to this node
-				if(DraggedNodePtr->GetBlueprint() != SCSEditor.Pin()->SCS->GetBlueprint())
+				if(DraggedNodePtr->GetBlueprint() != NodePtr->GetBlueprint())
 				{
 					if(DragRowOp->SourceNodes.Num() > 1)
 					{
@@ -1480,17 +1481,16 @@ void SSCS_RowWidget::OnAttachToDropAction(const TArray<FSCSEditorTreeNodePtrType
 	check(NodePtr.IsValid());
 	check(DroppedNodePtrs.Num() > 0);
 
-	TSharedPtr<SSCSEditor> SCSEditorPtr = SCSEditor.Pin();
-	check(SCSEditorPtr.IsValid());
-
 	// Get the current Blueprint context
-	check(SCSEditorPtr->SCS);
-	UBlueprint* Blueprint = SCSEditorPtr->SCS->GetBlueprint();
+	UBlueprint* Blueprint = NodePtr->GetBlueprint();
 	check(Blueprint);
 
+	TSharedPtr<SSCSEditor> SCSEditorPtr = SCSEditor.Pin();
+	check(SCSEditorPtr.IsValid());
+
 	// Get the current Actor context
-	AActor* PreviewActor = SCSEditorPtr->PreviewActor.Get();
-	check(PreviewActor != nullptr);
+	AActor* Actor = SCSEditorPtr->ActorContext.Get();
+	check(Actor != nullptr);
 
 	bool bRegenerateTreeNodes = false;
 	const FScopedTransaction TransactionContext(DroppedNodePtrs.Num() > 1 ? LOCTEXT("AttachComponents", "Attach Components") : LOCTEXT("AttachComponent", "Attach Component"));
@@ -1549,12 +1549,12 @@ void SSCS_RowWidget::OnAttachToDropAction(const TArray<FSCSEditorTreeNodePtrType
 						SCS_Node->AttachToName = NAME_None;
 					}
 
-					// Attempt to locate a matching instance of the component template in the preview scene
-					USceneComponent* PreviewSceneComponent = Cast<USceneComponent>(DroppedNodePtr->FindComponentInstanceInActor(PreviewActor, true));
-					if(PreviewSceneComponent)
+					// Attempt to locate a matching registered instance of the component template in the Actor context that's being edited
+					USceneComponent* InstancedSceneComponent = Cast<USceneComponent>(DroppedNodePtr->FindComponentInstanceInActor(Actor));
+					if(InstancedSceneComponent && InstancedSceneComponent->IsRegistered())
 					{
 						// If we find a match, save off the world position
-						FTransform ComponentToWorld = PreviewSceneComponent->GetComponentToWorld();
+						FTransform ComponentToWorld = InstancedSceneComponent->GetComponentToWorld();
 						SceneComponentTemplate->RelativeLocation = ComponentToWorld.GetTranslation();
 						SceneComponentTemplate->RelativeRotation = ComponentToWorld.Rotator();
 						SceneComponentTemplate->RelativeScale3D = ComponentToWorld.GetScale3D();
@@ -1565,13 +1565,13 @@ void SSCS_RowWidget::OnAttachToDropAction(const TArray<FSCSEditorTreeNodePtrType
 			// Attach the dropped node to the given node
 			NodePtr->AddChild(DroppedNodePtr);
 
-			// Attempt to locate a matching instance of the parent component template in the preview scene
-			USceneComponent* PreviewSceneParentComponent = Cast<USceneComponent>(NodePtr->FindComponentInstanceInActor(PreviewActor, true));
-			if(SceneComponentTemplate && PreviewSceneParentComponent)
+			// Attempt to locate a matching instance of the parent component template in the Actor context that's being edited
+			USceneComponent* ParentSceneComponent = Cast<USceneComponent>(NodePtr->FindComponentInstanceInActor(Actor));
+			if(SceneComponentTemplate && ParentSceneComponent && ParentSceneComponent->IsRegistered())
 			{
-				// If we find a match, calculate its new position relative to the scene root component instance in the preview scene
+				// If we find a match, calculate its new position relative to the scene root component instance in its current scene
 				FTransform ComponentToWorld(SceneComponentTemplate->RelativeRotation, SceneComponentTemplate->RelativeLocation, SceneComponentTemplate->RelativeScale3D);
-				FTransform ParentToWorld = PreviewSceneParentComponent->GetSocketTransform(SceneComponentTemplate->AttachSocketName);
+				FTransform ParentToWorld = ParentSceneComponent->GetSocketTransform(SceneComponentTemplate->AttachSocketName);
 				FTransform RelativeTM = ComponentToWorld.GetRelativeTransform(ParentToWorld);
 
 				// Store new relative location value (if not set to absolute)
@@ -1610,8 +1610,8 @@ void SSCS_RowWidget::OnDetachFromDropAction(const TArray<FSCSEditorTreeNodePtrTy
 	check(SCSEditorPtr.IsValid());
 
 	// Get the current preview Actor context
-	AActor* PreviewActor = SCSEditorPtr->PreviewActor.Get();
-	check(PreviewActor != nullptr);
+	AActor* Actor = SCSEditorPtr->ActorContext.Get();
+	check(Actor != nullptr);
 
 	const FScopedTransaction TransactionContext(DroppedNodePtrs.Num() > 1 ? LOCTEXT("DetachComponents", "Detach Components") : LOCTEXT("DetachComponent", "Detach Component"));
 
@@ -1622,7 +1622,7 @@ void SSCS_RowWidget::OnDetachFromDropAction(const TArray<FSCSEditorTreeNodePtrTy
 		// Detach the node from its parent
 		NodePtr->RemoveChild(DroppedNodePtr);
 
-		// If the associated component template is a scene component, maintain its preview world position
+		// If the associated component template is a scene component, maintain its current world position
 		USceneComponent* SceneComponentTemplate = Cast<USceneComponent>(DroppedNodePtr->GetComponentTemplate());
 		if(SceneComponentTemplate)
 		{
@@ -1638,12 +1638,12 @@ void SSCS_RowWidget::OnDetachFromDropAction(const TArray<FSCSEditorTreeNodePtrTy
 				SCS_Node->AttachToName = NAME_None;
 			}
 
-			// Attempt to locate a matching instance of the component template in the preview scene
-			USceneComponent* PreviewSceneComponent = Cast<USceneComponent>(DroppedNodePtr->FindComponentInstanceInActor(PreviewActor, true));
-			if(PreviewSceneComponent)
+			// Attempt to locate a matching instance of the component template in the Actor context that's being edited
+			USceneComponent* InstancedSceneComponent = Cast<USceneComponent>(DroppedNodePtr->FindComponentInstanceInActor(Actor));
+			if(InstancedSceneComponent && InstancedSceneComponent->IsRegistered())
 			{
 				// If we find a match, save off the world position
-				FTransform ComponentToWorld = PreviewSceneComponent->GetComponentToWorld();
+				FTransform ComponentToWorld = InstancedSceneComponent->GetComponentToWorld();
 				SceneComponentTemplate->RelativeLocation = ComponentToWorld.GetTranslation();
 				SceneComponentTemplate->RelativeRotation = ComponentToWorld.Rotator();
 				SceneComponentTemplate->RelativeScale3D = ComponentToWorld.GetScale3D();
@@ -1654,13 +1654,13 @@ void SSCS_RowWidget::OnDetachFromDropAction(const TArray<FSCSEditorTreeNodePtrTy
 		check(SCSEditorPtr->SceneRootNodePtr.IsValid());
 		SCSEditorPtr->SceneRootNodePtr->AddChild(DroppedNodePtr);
 
-		// Attempt to locate a matching instance of the scene root component template in the preview scene
-		USceneComponent* PreviewSceneRootComponent = Cast<USceneComponent>(SCSEditorPtr->SceneRootNodePtr->FindComponentInstanceInActor(PreviewActor, true));
-		if(SceneComponentTemplate && PreviewSceneRootComponent)
+		// Attempt to locate a matching instance of the scene root component template in the Actor context that's being edited
+		USceneComponent* InstancedSceneRootComponent = Cast<USceneComponent>(SCSEditorPtr->SceneRootNodePtr->FindComponentInstanceInActor(Actor));
+		if(SceneComponentTemplate && InstancedSceneRootComponent && InstancedSceneRootComponent->IsRegistered())
 		{
 			// If we find a match, calculate its new position relative to the scene root component instance in the preview scene
 			FTransform ComponentToWorld(SceneComponentTemplate->RelativeRotation, SceneComponentTemplate->RelativeLocation, SceneComponentTemplate->RelativeScale3D);
-			FTransform ParentToWorld = PreviewSceneRootComponent->GetSocketTransform(SceneComponentTemplate->AttachSocketName);
+			FTransform ParentToWorld = InstancedSceneRootComponent->GetSocketTransform(SceneComponentTemplate->AttachSocketName);
 			FTransform RelativeTM = ComponentToWorld.GetRelativeTransform(ParentToWorld);
 
 			// Store new relative location value (if not set to absolute)
@@ -1691,20 +1691,16 @@ void SSCS_RowWidget::OnMakeNewRootDropAction(FSCSEditorTreeNodePtrType DroppedNo
 	TSharedPtr<SSCSEditor> SCSEditorPtr = SCSEditor.Pin();
 	check(SCSEditorPtr.IsValid());
 
-	// Get the current SCS context
-	USimpleConstructionScript* SCS = SCSEditorPtr->SCS;
-	check(SCS != NULL);
-
-	// Get the current Blueprint context
-	UBlueprint* Blueprint = SCS->GetBlueprint();
-	check(Blueprint != NULL);
-
 	// Get the current scene root node
 	FSCSEditorTreeNodePtrType& SceneRootNodePtr = SCSEditorPtr->SceneRootNodePtr;
 
 	check(NodePtr.IsValid() && NodePtr == SceneRootNodePtr);
 	check(DroppedNodePtr.IsValid());
 
+	// Get the current Blueprint context
+	UBlueprint* Blueprint = NodePtr->GetBlueprint();
+	check(Blueprint != NULL && Blueprint->SimpleConstructionScript != nullptr);
+
 	// Create a transaction record
 	const FScopedTransaction TransactionContext(LOCTEXT("MakeNewSceneRoot", "Make New Scene Root"));
 
@@ -1764,7 +1760,7 @@ void SSCS_RowWidget::OnMakeNewRootDropAction(FSCSEditorTreeNodePtrType DroppedNo
 		check(SceneRootNodePtr->CanReparent());
 
 		// Remove the current scene root node from the SCS context
-		SCS->RemoveNode(SceneRootNodePtr->GetSCSNode());
+		Blueprint->SimpleConstructionScript->RemoveNode(SceneRootNodePtr->GetSCSNode());
 	}
 
 	// Save old root node
@@ -1778,7 +1774,7 @@ void SSCS_RowWidget::OnMakeNewRootDropAction(FSCSEditorTreeNodePtrType DroppedNo
 	SceneRootNodePtr = DroppedNodePtr;
 
 	// Add dropped node to the SCS context
-	SCS->AddNode(SceneRootNodePtr->GetSCSNode());
+	Blueprint->SimpleConstructionScript->AddNode(SceneRootNodePtr->GetSCSNode());
 
 	// Set old root as child of new root
 	if(OldSceneRootNodePtr.IsValid())
@@ -1798,8 +1794,14 @@ void SSCS_RowWidget::PostDragDropAction(bool bRegenerateTreeNodes)
 
 		PinnedEditor->RefreshSelectionDetails();
 
-		check(PinnedEditor->SCS != NULL);
-		FBlueprintEditorUtils::PostEditChangeBlueprintActors(PinnedEditor->SCS->GetBlueprint());
+		if(NodePtr.IsValid())
+		{
+			UBlueprint* Blueprint = NodePtr->GetBlueprint();
+			if(Blueprint != nullptr)
+			{
+				FBlueprintEditorUtils::PostEditChangeBlueprintActors(Blueprint);
+			}
+		}
 	}
 }
 
@@ -1896,7 +1898,8 @@ FString SSCS_RowWidget::GetDocumentationExcerptName() const
 
 UBlueprint* SSCS_RowWidget::GetBlueprint() const
 {
-	return SCSEditor.Pin()->Blueprint;
+	check(NodePtr.IsValid());
+	return NodePtr->GetBlueprint();
 }
 
 bool SSCS_RowWidget::OnNameTextVerifyChanged(const FText& InNewText, FText& OutErrorMessage)
@@ -1941,13 +1944,12 @@ void SSCS_RowWidget::OnNameTextCommit(const FText& InNewName, ETextCommit::Type
 // SSCSEditor
 
 
-void SSCSEditor::Construct( const FArguments& InArgs, USimpleConstructionScript* InSCS )
+void SSCSEditor::Construct( const FArguments& InArgs )
 {
-	SCS = InSCS;
-	Blueprint = InSCS->GetBlueprint();
-
-	PreviewActor = InArgs._PreviewActor;
+	ActorContext = InArgs._ActorContext;
 	bInEditingMode = InArgs._InEditingMode;
+	OnAddNewComponent = InArgs._OnAddNewComponent;
+	OnAddExistingComponent = InArgs._OnAddExistingComponent;
 	OnTreeViewSelectionChanged = InArgs._OnTreeViewSelectionChanged;
 	OnUpdateSelectionFromNodes = InArgs._OnUpdateSelectionFromNodes;
 	OnHighlightPropertyInDetailsView = InArgs._OnHighlightPropertyInDetailsView;
@@ -2141,10 +2143,20 @@ TSharedPtr< SWidget > SSCSEditor::CreateContextMenu()
 				MenuBuilder.AddMenuEntry( FGenericCommands::Get().Rename );
 
 				// Collect the classes of all selected objects
+				UBlueprint* Blueprint = nullptr;
 				TArray<UClass*> SelectionClasses;
 				for( auto NodeIter = SelectedNodes.CreateConstIterator(); NodeIter; ++NodeIter )
 				{
 					auto TreeNode = *NodeIter;
+
+					// All nodes should belong to the same Blueprint
+					check(Blueprint == nullptr || Blueprint == TreeNode->GetBlueprint());
+
+					if(Blueprint == nullptr)
+					{
+						Blueprint = TreeNode->GetBlueprint();
+					}
+
 					if( TreeNode->GetComponentTemplate() )
 					{
 						SelectionClasses.Add( TreeNode->GetComponentTemplate()->GetClass() );
@@ -2601,13 +2613,20 @@ void SSCSEditor::UpdateTree(bool bRegenerateTreeNodes)
 		// Reset the scene root node
 		SceneRootNodePtr.Reset();
 
-		// Get the Blueprint class default object as well as the inheritance stack
+		// Get the class default object
 		AActor* CDO = NULL;
 		TArray<UBlueprint*> ParentBPStack;
-		if(Blueprint->GeneratedClass != NULL)
+		AActor* Actor = ActorContext.Get();
+		if(Actor != nullptr)
 		{
-			CDO = Blueprint->GeneratedClass->GetDefaultObject<AActor>();
-			UBlueprint::GetBlueprintHierarchyFromClass(Blueprint->GeneratedClass, ParentBPStack);
+			UClass* ActorClass = Actor->GetClass();
+			if(ActorClass != nullptr)
+			{
+				CDO = ActorClass->GetDefaultObject<AActor>();
+
+				// If it's a Blueprint-generated class, also get the inheritance stack
+				UBlueprint::GetBlueprintHierarchyFromClass(ActorClass, ParentBPStack);
+			}
 		}
 
 		if(CDO != NULL)
@@ -2625,7 +2644,7 @@ void SSCSEditor::UpdateTree(bool bRegenerateTreeNodes)
 				}
 			}
 
-		// Add the native base class SceneComponent hierarchy
+			// Add the native base class SceneComponent hierarchy
 			for(auto CompIter = Components.CreateIterator(); CompIter; ++CompIter)
 			{
 				USceneComponent* SceneComp = Cast<USceneComponent>(*CompIter);
@@ -2649,7 +2668,7 @@ void SSCSEditor::UpdateTree(bool bRegenerateTreeNodes)
 
 					if(SCS_Node->ParentComponentOrVariableName != NAME_None)
 					{
-						USceneComponent* ParentComponent = SCS_Node->GetParentComponentTemplate(Blueprint);
+						USceneComponent* ParentComponent = SCS_Node->GetParentComponentTemplate();
 						if(ParentComponent != NULL)
 						{
 							FSCSEditorTreeNodePtrType ParentNodePtr = FindTreeNode(ParentComponent);
@@ -2751,17 +2770,26 @@ UActorComponent* SSCSEditor::AddNewComponent( UClass* NewComponentClass, UObject
 {
 	const FScopedTransaction Transaction( LOCTEXT("AddComponent", "Add Component") );
 
-	// Create a new template and add it to the blueprint
-	Blueprint->Modify();
+	check(OnAddNewComponent.IsBound());
+	USCS_Node* NewNode = OnAddNewComponent.Execute(NewComponentClass);
+
+	USimpleConstructionScript* SCS = NewNode->GetSCS();
+	check(SCS != nullptr);
+
+	UBlueprint* Blueprint = SCS->GetBlueprint();
+	check(Blueprint != nullptr);
 
-	SaveSCSCurrentState( Blueprint->SimpleConstructionScript );
+	// Save current state
+	Blueprint->Modify();
+	SaveSCSCurrentState(SCS);
 
-	USCS_Node* NewNode = SCS->CreateNode(NewComponentClass);
 	return AddNewNode(NewNode, Asset, true);
 }
 
 UActorComponent* SSCSEditor::AddNewNode(USCS_Node* NewNode,  UObject* Asset, bool bMarkBlueprintModified, bool bSetFocusToNewItem)
 {
+	check(NewNode != nullptr);
+
 	if(Asset)
 	{
 		FComponentAssetBrokerage::AssignAssetToComponent(NewNode->ComponentTemplate, Asset);
@@ -2769,6 +2797,12 @@ UActorComponent* SSCSEditor::AddNewNode(USCS_Node* NewNode,  UObject* Asset, boo
 
 	FSCSEditorTreeNodePtrType NewNodePtr;
 
+	USimpleConstructionScript* SCS = NewNode->GetSCS();
+	check(SCS != nullptr);
+
+	UBlueprint* Blueprint = SCS->GetBlueprint();
+	check(Blueprint != nullptr);
+
 	// Reset the scene root node if it's set to the default one that's managed by the SCS
 	if(SceneRootNodePtr.IsValid() && SceneRootNodePtr->GetSCSNode() == SCS->GetDefaultSceneRootNode())
 	{
@@ -2956,12 +2990,6 @@ void SSCSEditor::PasteNodes()
 {
 	const FScopedTransaction Transaction(LOCTEXT("PasteComponents", "Paste Component(s)"));
 
-	// Mark the Blueprint as about to be modified
-	Blueprint->Modify();
-
-	// Mark the SCS and all SCS nodes as about to be modified
-	SaveSCSCurrentState(Blueprint->SimpleConstructionScript);
-
 	// Get the text from the clipboard
 	FString TextToImport;
 	FPlatformMisc::ClipboardPaste(TextToImport);
@@ -2973,6 +3001,7 @@ void SSCSEditor::PasteNodes()
 	SCSTreeWidget->ClearSelection();
 
 	// Create a new tree node for each new (pasted) component
+	UBlueprint* Blueprint = nullptr;
 	TMap<FName, FSCSEditorTreeNodePtrType> NewNodeMap;
 	for(auto NewObjectIt = Factory->NewObjectMap.CreateIterator(); NewObjectIt; ++NewObjectIt)
 	{
@@ -2980,11 +3009,28 @@ void SSCSEditor::PasteNodes()
 		UActorComponent* NewActorComponent = NewObjectIt->Value;
 		check(NewActorComponent);
 
-		// Relocate the instance from the transient package to the BPGC and assign it a unique object name
-		NewActorComponent->Rename(NULL, Blueprint->GeneratedClass, REN_DontCreateRedirectors|REN_DoNotDirty);
+		// Ask the delegate to create a new SCS node to wrap the component
+		check(OnAddExistingComponent.IsBound());
+		USCS_Node* SCS_Node = OnAddExistingComponent.Execute(NewActorComponent);
+
+		check(SCS_Node != nullptr);
+		USimpleConstructionScript* SCS = SCS_Node->GetSCS();
+		check(SCS != nullptr && (Blueprint == nullptr || Blueprint == SCS->GetBlueprint()));
+		
+		if(Blueprint == nullptr)
+		{
+			Blueprint = SCS->GetBlueprint();
+			check(Blueprint != nullptr);
+
+			// Mark the Blueprint as about to be modified
+			Blueprint->Modify();
+
+			// Mark the SCS and all SCS nodes as about to be modified
+			SaveSCSCurrentState(SCS);
+		}
 
 		// Create a new SCS node to contain the new component and add it to the tree
-		NewActorComponent = AddNewNode(SCS->CreateNode(NewActorComponent), NULL, false, false);
+		NewActorComponent = AddNewNode(SCS_Node, NULL, false, false);
 		if(NewActorComponent)
 		{
 			// Locate the node that corresponds to the new component template
@@ -3042,47 +3088,62 @@ void SSCSEditor::OnDeleteNodes()
 {
 	const FScopedTransaction Transaction( LOCTEXT("RemoveComponent", "Remove Component") );
 
-	// Get the current render info for the blueprint. If this is NULL then the blueprint is not currently visualizable (no visible primitive components)
-	FThumbnailRenderingInfo* RenderInfo = GUnrealEd->GetThumbnailManager()->GetRenderingInfo( Blueprint );
-
-	// Saving objects for restoring purpose.
-	Blueprint->Modify();
-	SSCSEditor::SaveSCSCurrentState( Blueprint->SimpleConstructionScript );
-
 	// Remove node from SCS
+	UBlueprint* Blueprint = nullptr;
+	FThumbnailRenderingInfo* RenderInfo = nullptr;
 	TArray<FSCSEditorTreeNodePtrType> SelectedNodes = SCSTreeWidget->GetSelectedItems();
 	for (int32 i = 0; i < SelectedNodes.Num(); ++i)
 	{
 		auto Node = SelectedNodes[i];
 
+		USCS_Node* SCS_Node = Node->GetSCSNode();
+		if(SCS_Node != nullptr)
+		{
+			USimpleConstructionScript* SCS = SCS_Node->GetSCS();
+			check(SCS != nullptr && (Blueprint == nullptr || Blueprint == SCS->GetBlueprint()));
+
+			if(Blueprint == nullptr)
+			{
+				Blueprint = SCS->GetBlueprint();
+				check(Blueprint != nullptr);
+
+				// Get the current render info for the blueprint. If this is NULL then the blueprint is not currently visualizable (no visible primitive components)
+				FThumbnailRenderingInfo* RenderInfo = GUnrealEd->GetThumbnailManager()->GetRenderingInfo( Blueprint );
+
+				// Saving objects for restoring purpose.
+				Blueprint->Modify();
+				SaveSCSCurrentState( SCS );
+			}
+		}
+
 		RemoveComponentNode(Node);
 	}
 
-	// Will call UpdateTree as part of OnBlueprintChanged handling
-	FBlueprintEditorUtils::MarkBlueprintAsStructurallyModified(Blueprint);
+	if(Blueprint != nullptr)
+	{
+		// Will call UpdateTree as part of OnBlueprintChanged handling
+		FBlueprintEditorUtils::MarkBlueprintAsStructurallyModified(Blueprint);
+
+		// If we had a thumbnail before we deleted any components, check to see if we should clear it
+		// If we deleted the final visualizable primitive from the blueprint, GetRenderingInfo should return NULL
+		FThumbnailRenderingInfo* NewRenderInfo = GUnrealEd->GetThumbnailManager()->GetRenderingInfo( Blueprint );
+		if ( RenderInfo && !NewRenderInfo )
+		{
+			// We removed the last visible primitive component, clear the thumbnail
+			const FString BPFullName = FString::Printf(TEXT("%s %s"), *Blueprint->GetClass()->GetName(), *Blueprint->GetPathName());
+			UPackage* BPPackage = Blueprint->GetOutermost();
+			ThumbnailTools::CacheEmptyThumbnail( BPFullName, BPPackage );
+		}
+	}
 
 	// Do this AFTER marking the Blueprint as modified
 	UpdateSelectionFromNodes(SCSTreeWidget->GetSelectedItems());
-
-	// If we had a thumbnail before we deleted any components, check to see if we should clear it
-	// If we deleted the final visualizable primitive from the blueprint, GetRenderingInfo should return NULL
-	FThumbnailRenderingInfo* NewRenderInfo = GUnrealEd->GetThumbnailManager()->GetRenderingInfo( Blueprint );
-	if ( RenderInfo && !NewRenderInfo )
-	{
-		// We removed the last visible primitive component, clear the thumbnail
-		const FString BPFullName = FString::Printf(TEXT("%s %s"), *Blueprint->GetClass()->GetName(), *Blueprint->GetPathName());
-		UPackage* BPPackage = Blueprint->GetOutermost();
-		ThumbnailTools::CacheEmptyThumbnail( BPFullName, BPPackage );
-	}
 }
 
 void SSCSEditor::RemoveComponentNode(FSCSEditorTreeNodePtrType InNodePtr)
 {
 	check(InNodePtr.IsValid());
 
-	// Remove any instances of variable accessors from the blueprint graphs
-	FBlueprintEditorUtils::RemoveVariableNodes(Blueprint, InNodePtr->GetVariableName());
-
 	// Clear selection if current
 	if(SCSTreeWidget->GetSelectedItems().Contains(InNodePtr))
 	{
@@ -3092,8 +3153,17 @@ void SSCSEditor::RemoveComponentNode(FSCSEditorTreeNodePtrType InNodePtr)
 	USCS_Node* SCS_Node = InNodePtr->GetSCSNode();
 	if(SCS_Node != NULL)
 	{
+		USimpleConstructionScript* SCS = SCS_Node->GetSCS();
+		check(SCS != nullptr);
+
+		// Remove any instances of variable accessors from the blueprint graphs
+		UBlueprint* Blueprint = SCS->GetBlueprint();
+		if(Blueprint != nullptr)
+		{
+			FBlueprintEditorUtils::RemoveVariableNodes(Blueprint, InNodePtr->GetVariableName());
+		}
+
 		// Remove node from SCS tree
-		check(SCS_Node->GetSCS() == SCS);
 		SCS->RemoveNodeAndPromoteChildren(SCS_Node);
 
 		// Clear the delegate
@@ -3122,19 +3192,6 @@ void SSCSEditor::OnTreeSelectionChanged(FSCSEditorTreeNodePtrType, ESelectInfo::
 
 	// Pass the tree selection change event along to the delegate
 	OnTreeViewSelectionChanged.ExecuteIfBound(SelectedNodes);
-
-	// Update the selection visualization
-	AActor* EditorActorInstance = SCS->GetComponentEditorActorInstance();
-	if (EditorActorInstance != NULL)
-	{
-		TInlineComponentArray<UPrimitiveComponent*> PrimitiveComponents;
-		EditorActorInstance->GetComponents(PrimitiveComponents);
-
-		for (int32 Idx = 0; Idx < PrimitiveComponents.Num(); ++Idx)
-		{
-			PrimitiveComponents[Idx]->PushSelectionToProxy();
-		}
-	}
 }
 
 bool SSCSEditor::IsNodeInSimpleConstructionScript( USCS_Node* Node ) const
@@ -3169,7 +3226,7 @@ FSCSEditorTreeNodePtrType SSCSEditor::AddTreeNode(USCS_Node* InSCSNode, FSCSEdit
 	
 	// Determine whether or not the given node is inherited from a parent Blueprint
 	USimpleConstructionScript* NodeSCS = InSCSNode->GetSCS();
-	const bool bIsInherited = NodeSCS != Blueprint->SimpleConstructionScript;
+	const bool bIsInherited = NodeSCS != nullptr && InParentNodePtr.IsValid() && NodeSCS->GetBlueprint() != InParentNodePtr->GetBlueprint();
 
 	if(InSCSNode->ComponentTemplate->IsA(USceneComponent::StaticClass()))
 	{
diff --git a/Engine/Source/Editor/Kismet/Public/BlueprintEditor.h b/Engine/Source/Editor/Kismet/Public/BlueprintEditor.h
--- a/Engine/Source/Editor/Kismet/Public/BlueprintEditor.h
+++ b/Engine/Source/Editor/Kismet/Public/BlueprintEditor.h
@@ -358,15 +358,18 @@ public:
 	/** Refresh the preview viewport to reflect changes in the SCS */
 	void UpdateSCSPreview(bool bUpdateNow = false);
 
+	/** Delegate invoked when the SCS editor widget needs to create a new component node for a component class type */
+	USCS_Node* OnSCSEditorAddNewComponent(UClass* InComponentClass);
+
+	/** Delegate invoked when the SCS editor widget needs to create a new component node for an existing component instance */
+	USCS_Node* OnSCSEditorAddExistingComponent(UActorComponent* InComponentInstance);
+
 	/** Delegate invoked when the tree view selection is changed in the SCS editor widget */
 	void OnSCSEditorTreeViewSelectionChanged(const TArray< TSharedPtr<class FSCSEditorTreeNode> >& SelectedNodes);
 
 	/** Delegate to update the Inspector (details) panel from the given set of selected nodes */
 	void OnSCSEditorUpdateSelectionFromNodes(const TArray< TSharedPtr<class FSCSEditorTreeNode> >& SelectedNodes);
 
-	/** Delegate invoked when the given property should be highlighted in the Inspector (details) panel */
-	void OnSCSEditorHighlightPropertyInDetailsView(const class FPropertyPath& InPropertyPath);
-
 	/** Pin visibility accessors */
 	void SetPinVisibility(SGraphEditor::EPinVisibility Visibility);
 	bool GetPinVisibility(SGraphEditor::EPinVisibility Visibility) const { return PinVisibility == Visibility; }
diff --git a/Engine/Source/Editor/Kismet/Public/SSCSEditor.h b/Engine/Source/Editor/Kismet/Public/SSCSEditor.h
--- a/Engine/Source/Editor/Kismet/Public/SSCSEditor.h
+++ b/Engine/Source/Editor/Kismet/Public/SSCSEditor.h
@@ -59,10 +59,9 @@ public:
 	 * Finds the component instance represented by this node contained within a given Actor instance.
 	 *
 	 * @param InActor The Actor instance to use as the container object for finding the component instance.
-	 * @param bIsPreviewActor Whether or not the given actor represents the "preview" actor in the SCS editor viewport.
 	 * @return The component instance represented by this node and contained within the given Actor instance, or NULL if not found.
 	 */
-	UActorComponent* FindComponentInstanceInActor(const AActor* InActor, bool bIsPreviewActor = false) const;
+	UActorComponent* FindComponentInstanceInActor(const AActor* InActor) const;
 	/** 
 	 * @return This object's parent node (or an invalid reference if no parent is assigned).
 	 */
@@ -376,29 +375,35 @@ typedef SSCSEditorDragDropTree SSCSTreeType;
 class KISMET_API SSCSEditor : public SCompoundWidget
 {
 public:
+	DECLARE_DELEGATE_RetVal_OneParam(class USCS_Node*, FOnAddNewComponent, class UClass*);
+	DECLARE_DELEGATE_RetVal_OneParam(class USCS_Node*, FOnAddExistingComponent, class UActorComponent*);
 	DECLARE_DELEGATE_OneParam(FOnTreeViewSelectionChanged, const TArray<FSCSEditorTreeNodePtrType>&);
 	DECLARE_DELEGATE_OneParam(FOnUpdateSelectionFromNodes, const TArray<FSCSEditorTreeNodePtrType>&);
 	DECLARE_DELEGATE_OneParam(FOnHighlightPropertyInDetailsView, const class FPropertyPath&);
 
 	SLATE_BEGIN_ARGS( SSCSEditor )
 		:_InEditingMode(true)
-		,_PreviewActor(nullptr)
+		,_ActorContext(nullptr)
 		,_HideComponentClassCombo(false)
+		,_OnAddNewComponent()
+		,_OnAddExistingComponent()
 		,_OnTreeViewSelectionChanged()
 		,_OnUpdateSelectionFromNodes()
 		,_OnHighlightPropertyInDetailsView()
 		{}
 
 		SLATE_ATTRIBUTE(bool, InEditingMode)
-		SLATE_ATTRIBUTE(class AActor*, PreviewActor)
+		SLATE_ATTRIBUTE(class AActor*, ActorContext)
 		SLATE_ATTRIBUTE(bool, HideComponentClassCombo)
+		SLATE_EVENT(FOnAddNewComponent, OnAddNewComponent)
+		SLATE_EVENT(FOnAddExistingComponent, OnAddExistingComponent)
 		SLATE_EVENT(FOnTreeViewSelectionChanged, OnTreeViewSelectionChanged)
 		SLATE_EVENT(FOnUpdateSelectionFromNodes, OnUpdateSelectionFromNodes)
 		SLATE_EVENT(FOnHighlightPropertyInDetailsView, OnHighlightPropertyInDetailsView)
 
 	SLATE_END_ARGS()
 
-	void Construct(const FArguments& InArgs, USimpleConstructionScript* InSCS);
+	void Construct(const FArguments& InArgs);
 
 	/** Override OnKeyDown */
 	virtual FReply OnKeyDown( const FGeometry& MyGeometry, const FKeyEvent& InKeyEvent );
@@ -601,12 +606,6 @@ protected:
 	/** Returns the set of expandable nodes that are currently collapsed in the UI */
 	void GetCollapsedNodes(const FSCSEditorTreeNodePtrType& InNodePtr, TSet<FSCSEditorTreeNodePtrType>& OutCollapsedNodes) const;
 public:
-	/** Pointer to the script that we are editing */
-	USimpleConstructionScript* SCS;
-
-	/** Pointer to blueprint we were created with */
-	UBlueprint* Blueprint;
-
 	/** Tree widget */
 	TSharedPtr<SSCSTreeType> SCSTreeWidget;
 
@@ -625,11 +624,17 @@ public:
 	/** Whether or not the deferred rename request was flagged as transactional */
 	bool bIsDeferredRenameRequestTransactional;
 
+	/** Attribute that provides access to the Actor context for which we are viewing/editing the SCS. */
+	TAttribute<class AActor*> ActorContext;
+
 	/** Attribute to control whether or not to allow editing. */
 	TAttribute<bool> bInEditingMode;
 
-	/** Attribute that provides access to the "preview" Actor context. */
-	TAttribute<AActor*> PreviewActor;
+	/** Delegate to invoke when a new SCS node needs to be created for a component class type. */
+	FOnAddNewComponent OnAddNewComponent;
+
+	/** Delegate to invoke when a new SCS node needs to be created for an existing component instance. */
+	FOnAddExistingComponent OnAddExistingComponent;
 
 	/** Delegate to invoke on tree view selection change. */
 	FOnTreeViewSelectionChanged OnTreeViewSelectionChanged;
diff --git a/Engine/Source/Editor/LevelEditor/Private/SLevelEditor.cpp b/Engine/Source/Editor/LevelEditor/Private/SLevelEditor.cpp
--- a/Engine/Source/Editor/LevelEditor/Private/SLevelEditor.cpp
+++ b/Engine/Source/Editor/LevelEditor/Private/SLevelEditor.cpp
@@ -579,9 +579,9 @@ public:
 						ComponentsBox->ClearChildren();
 						ComponentsBox->AddSlot()
 						[
-							SAssignNew(SCSEditor, SSCSEditor, Blueprint->SimpleConstructionScript)
+							SAssignNew(SCSEditor, SSCSEditor)
 							.InEditingMode(true)
-							.PreviewActor(this, &SActorDetails::GetSelectedActor)												// Get the instance of the actor in the world
+							.ActorContext(this, &SActorDetails::GetSelectedActor)												// Get the instance of the actor in the world
 							.OnTreeViewSelectionChanged(this, &SActorDetails::OnSCSEditorTreeViewSelectionChanged)				// A selection has been made in the tree view, so inform the level editor
 							//.OnUpdateSelectionFromNodes(this, &SLevelEditor::OnSCSEditorUpdateSelectionFromNodes)				// Unsure, don't think it's needed
 							//.OnHighlightPropertyInDetailsView(this, &SLevelEditor::OnSCSEditorHighlightPropertyInDetailsView)	// Also unsure and don't think it's needed
diff --git a/Engine/Source/Editor/UnrealEd/Private/Kismet2/ComponentEditorUtils.cpp b/Engine/Source/Editor/UnrealEd/Private/Kismet2/ComponentEditorUtils.cpp
--- a/Engine/Source/Editor/UnrealEd/Private/Kismet2/ComponentEditorUtils.cpp
+++ b/Engine/Source/Editor/UnrealEd/Private/Kismet2/ComponentEditorUtils.cpp
@@ -41,7 +41,7 @@ USceneComponent* FComponentEditorUtils::GetSceneComponent( UObject* Object, UObj
 							USCS_Node* SCS_Node = SCSNodes[ SCSNodeIndex ];
 							if( SCS_Node && SCS_Node->ComponentTemplate == SubObject )
 							{
-								return SCS_Node->GetParentComponentTemplate( Blueprint );
+								return SCS_Node->GetParentComponentTemplate();
 							}
 						}
 					}
diff --git a/Engine/Source/Editor/UnrealEd/Private/ThumbnailHelpers.cpp b/Engine/Source/Editor/UnrealEd/Private/ThumbnailHelpers.cpp
--- a/Engine/Source/Editor/UnrealEd/Private/ThumbnailHelpers.cpp
+++ b/Engine/Source/Editor/UnrealEd/Private/ThumbnailHelpers.cpp
@@ -1178,7 +1178,7 @@ TArray<UPrimitiveComponent*> FBlueprintThumbnailScene::GetPooledVisualizableComp
 							// Check to see if the root node has set an explicit parent
 							if(RootNode->ParentComponentOrVariableName != NAME_None)
 							{
-								USceneComponent* ParentComponentTemplate = RootNode->GetParentComponentTemplate(Blueprint);
+								USceneComponent* ParentComponentTemplate = RootNode->GetParentComponentTemplate();
 								if(ParentComponentTemplate != NULL)
 								{
 									if(NativeInstanceMap.Contains(ParentComponentTemplate))
diff --git a/Engine/Source/Runtime/Engine/Classes/Engine/SCS_Node.h b/Engine/Source/Runtime/Engine/Classes/Engine/SCS_Node.h
--- a/Engine/Source/Runtime/Engine/Classes/Engine/SCS_Node.h
+++ b/Engine/Source/Runtime/Engine/Classes/Engine/SCS_Node.h
@@ -150,8 +150,8 @@ class USCS_Node : public UObject
 	/** Sets parent component attributes based on the given component instance */
 	ENGINE_API void SetParent(USceneComponent* InParentComponent);
 
-	/** Finds and returns the parent component template through the given Blueprint */
-	ENGINE_API USceneComponent* GetParentComponentTemplate(UBlueprint* InBlueprint) const;
+	/** Finds and returns the parent component template */
+	ENGINE_API USceneComponent* GetParentComponentTemplate() const;
 
 	/** Test whether or not the given string is a valid variable name string */
 	ENGINE_API bool IsValidVariableNameString(const FString& InString);
diff --git a/Engine/Source/Runtime/Engine/Private/SCS_Node.cpp b/Engine/Source/Runtime/Engine/Private/SCS_Node.cpp
--- a/Engine/Source/Runtime/Engine/Private/SCS_Node.cpp
+++ b/Engine/Source/Runtime/Engine/Private/SCS_Node.cpp
@@ -317,18 +317,22 @@ void USCS_Node::SetParent(USceneComponent* InParentComponent)
 	}
 }
 
-USceneComponent* USCS_Node::GetParentComponentTemplate(UBlueprint* InBlueprint) const
+USceneComponent* USCS_Node::GetParentComponentTemplate() const
 {
 	USceneComponent* ParentComponentTemplate = NULL;
 	if(ParentComponentOrVariableName != NAME_None)
 	{
-		check(InBlueprint != NULL && InBlueprint->GeneratedClass != NULL);
+		USimpleConstructionScript* SCS = GetSCS();
+		check(SCS != nullptr);
+
+		UBlueprint* Blueprint = SCS->GetBlueprint();
+		check(Blueprint != nullptr && Blueprint->GeneratedClass != nullptr);
 
 		// If the parent component template is found in the 'Components' array of the CDO (i.e. native)
 		if(bIsParentComponentNative)
 		{
 			// Access the Blueprint CDO
-			AActor* CDO = InBlueprint->GeneratedClass->GetDefaultObject<AActor>();
+			AActor* CDO = Blueprint->GeneratedClass->GetDefaultObject<AActor>();
 			if(CDO != NULL)
 			{
 				// Find the component template in the CDO that matches the specified name
@@ -352,7 +356,7 @@ USceneComponent* USCS_Node::GetParentComponentTemplate(UBlueprint* InBlueprint)
 		{
 			// Get the Blueprint hierarchy
 			TArray<UBlueprint*> ParentBPStack;
-			UBlueprint::GetBlueprintHierarchyFromClass(InBlueprint->GeneratedClass, ParentBPStack);
+			UBlueprint::GetBlueprintHierarchyFromClass(Blueprint->GeneratedClass, ParentBPStack);
 
 			// Find the parent Blueprint in the hierarchy
 			for(int32 StackIndex = ParentBPStack.Num() - 1; StackIndex > 0; --StackIndex)
