Fix empty components tree view in single-screen BP editor mode.

diff --git a/Engine/Source/Editor/Kismet/Private/BlueprintEditorModes.cpp b/Engine/Source/Editor/Kismet/Private/BlueprintEditorModes.cpp
--- a/Engine/Source/Editor/Kismet/Private/BlueprintEditorModes.cpp
+++ b/Engine/Source/Editor/Kismet/Private/BlueprintEditorModes.cpp
@@ -663,6 +663,20 @@ void FBlueprintEditorUnifiedMode::PostActivateMode()
 	BP->RestoreEditedObjectState();
 	BP->SetupViewForBlueprintEditingMode();
 
+	TSharedPtr<SSCSEditorViewport> PreviewViewportPtr = BP->GetSCSViewport();
+	if(PreviewViewportPtr.IsValid())
+	{
+		// Ensure that the preview actor is instanced
+		PreviewViewportPtr->RequestRefresh(true, true);
+
+		TSharedPtr<SSCSEditor> SCSEditorPtr = BP->GetSCSEditor();
+		if(SCSEditorPtr.IsValid())
+		{
+			// Update the components tree view
+			SCSEditorPtr->UpdateTree();
+		}
+	}
+
 	FApplicationMode::PostActivateMode();
 }
 
