[UE-7311] Exclude global Actor class details in Components mode in the Blueprint editor.

- This applies when the CDO is selected to the Details view in order to allow the user to edit properties on native component template objects.
- Note: Does not currently fix in experimental single-screen mode (other component-specific customizations are currently missing in that mode as well).

diff --git a/Engine/Source/Editor/Kismet/Private/SKismetInspector.cpp b/Engine/Source/Editor/Kismet/Private/SKismetInspector.cpp
--- a/Engine/Source/Editor/Kismet/Private/SKismetInspector.cpp
+++ b/Engine/Source/Editor/Kismet/Private/SKismetInspector.cpp
@@ -271,13 +271,31 @@ void SKismetInspector::Construct(const FArguments& InArgs)
 
 void SKismetInspector::EnableComponentDetailsCustomization(bool bEnable)
 {
+	// An "empty" instanced customization that's intended to override any registered global details customization for
+	// the AActor class type. This will be applied -only- when the CDO is selected to the Details view in Components mode.
+	class FActorDetailsOverrideCustomization : public IDetailCustomization
+	{
+	public:
+		/** IDetailCustomization interface */
+		virtual void CustomizeDetails(IDetailLayoutBuilder& DetailLayout) override {}
+
+		static TSharedRef<class IDetailCustomization> MakeInstance()
+		{
+			return MakeShareable(new FActorDetailsOverrideCustomization());
+		}
+	};
+
 	if (bEnable)
 	{
+		FOnGetDetailCustomizationInstance ActorOverrideDetails = FOnGetDetailCustomizationInstance::CreateStatic(&FActorDetailsOverrideCustomization::MakeInstance);
+		PropertyView->RegisterInstancedCustomPropertyLayout(AActor::StaticClass(), ActorOverrideDetails);
+
 		FOnGetDetailCustomizationInstance LayoutComponentDetails = FOnGetDetailCustomizationInstance::CreateStatic(&FBlueprintComponentDetails::MakeInstance, Kismet2Ptr);
 		PropertyView->RegisterInstancedCustomPropertyLayout(UActorComponent::StaticClass(), LayoutComponentDetails);
 	}
 	else
 	{
+		PropertyView->UnregisterInstancedCustomPropertyLayout(AActor::StaticClass());
 		PropertyView->UnregisterInstancedCustomPropertyLayout(UActorComponent::StaticClass());
 	}
 }
