Limiting scope of optimization even more. Not sure if we need to update thumbnail after compiling in some cases, so I'm just doing it for now.

diff --git a/Engine/Source/Editor/Kismet/Private/BlueprintEditor.cpp b/Engine/Source/Editor/Kismet/Private/BlueprintEditor.cpp
--- a/Engine/Source/Editor/Kismet/Private/BlueprintEditor.cpp
+++ b/Engine/Source/Editor/Kismet/Private/BlueprintEditor.cpp
@@ -587,17 +587,20 @@ void FBlueprintEditor::AnalyticsTrackCompileEvent( UBlueprint* Blueprint, int32
 	}
 }
 
-void FBlueprintEditor::RefreshEditors()
+void FBlueprintEditor::RefreshEditors(ERefreshBlueprintEditorReason::Type Reason)
 {
-	DocumentManager->CleanInvalidTabs();
+	if( Reason != ERefreshBlueprintEditorReason::BlueprintCompiled )
+	{
+		DocumentManager->CleanInvalidTabs();
 
-	DocumentManager->RefreshAllTabs();
+		DocumentManager->RefreshAllTabs();
 
-	// The workflow manager only tracks document tabs.
-	FocusInspectorOnGraphSelection(GetSelectedNodes(), true);
-	if (MyBlueprintWidget.IsValid())
-	{
-		MyBlueprintWidget->Refresh();
+		// The workflow manager only tracks document tabs.
+		FocusInspectorOnGraphSelection(GetSelectedNodes(), true);
+		if (MyBlueprintWidget.IsValid())
+		{
+			MyBlueprintWidget->Refresh();
+		}
 	}
 
 	if(SCSEditor.IsValid())
@@ -2713,10 +2716,8 @@ void FBlueprintEditor::OnBlueprintChangedImpl(UBlueprint* InBlueprint, bool bIsJ
 	if (InBlueprint)
 	{
 		// Refresh the graphs
-		if( !bIsJustBeingCompiled )
-		{
-			RefreshEditors();
-		}
+		ERefreshBlueprintEditorReason::Type Reason = bIsJustBeingCompiled ? ERefreshBlueprintEditorReason::BlueprintCompiled : ERefreshBlueprintEditorReason::UnknownReason;
+		RefreshEditors(Reason);
 
 		// Notify that the blueprint has been changed (update Content browser, etc)
 		InBlueprint->PostEditChange();
diff --git a/Engine/Source/Editor/Kismet/Public/BlueprintEditor.h b/Engine/Source/Editor/Kismet/Public/BlueprintEditor.h
--- a/Engine/Source/Editor/Kismet/Public/BlueprintEditor.h
+++ b/Engine/Source/Editor/Kismet/Public/BlueprintEditor.h
@@ -152,7 +152,7 @@ public:
 	// End of FGCObject interface
 
 	// IBlueprintEditor interface
-	virtual void RefreshEditors() override;
+	virtual void RefreshEditors(ERefreshBlueprintEditorReason::Type Reason = ERefreshBlueprintEditorReason::UnknownReason) override;
 	virtual void JumpToHyperlink(const UObject* ObjectReference, bool bRequestRename = false) override;
 	virtual void SummonSearchUI(bool bSetFindWithinBlueprint, FString NewSearchTerms = FString(), bool bSelectFirstResult = false) override;
 	virtual TArray<TSharedPtr<class FSCSEditorTreeNode> > GetSelectedSCSEditorTreeNodes() const override;
diff --git a/Engine/Source/Editor/Kismet/Public/BlueprintEditorModule.h b/Engine/Source/Editor/Kismet/Public/BlueprintEditorModule.h
--- a/Engine/Source/Editor/Kismet/Public/BlueprintEditorModule.h
+++ b/Engine/Source/Editor/Kismet/Public/BlueprintEditorModule.h
@@ -11,6 +11,16 @@ extern const FName BlueprintEditorAppName;
 class FBlueprintEditor;
 class UUserDefinedEnum;
 
+/** Describes the reason for Refreshing the editor */
+namespace ERefreshBlueprintEditorReason
+{
+	enum Type
+	{
+		BlueprintCompiled,
+		UnknownReason
+	};
+}
+
 /**
  * Enum editor public interface
  */
@@ -36,7 +46,7 @@ public:
 	/** Invokes the search UI and sets the mode and search terms optionally */
 	virtual void SummonSearchUI(bool bSetFindWithinBlueprint, FString NewSearchTerms = FString(), bool bSelectFirstResult = false) = 0;
 
-	virtual void RefreshEditors() = 0;
+	virtual void RefreshEditors(ERefreshBlueprintEditorReason::Type Reason = ERefreshBlueprintEditorReason::UnknownReason) = 0;
 
 	virtual bool CanPasteNodes() const= 0;
 
