Even more restrictive; found a bug: when you add a new variable and hit compile we need to refresh the SMyBlueprint stuff so that the value of the variable can be displayed. Witout this refrsh you have to reselect the variable.

diff --git a/Engine/Source/Editor/Kismet/Private/BlueprintEditor.cpp b/Engine/Source/Editor/Kismet/Private/BlueprintEditor.cpp
--- a/Engine/Source/Editor/Kismet/Private/BlueprintEditor.cpp
+++ b/Engine/Source/Editor/Kismet/Private/BlueprintEditor.cpp
@@ -597,10 +597,11 @@ void FBlueprintEditor::RefreshEditors(ERefreshBlueprintEditorReason::Type Reason
 
 		// The workflow manager only tracks document tabs.
 		FocusInspectorOnGraphSelection(GetSelectedNodes(), true);
-		if (MyBlueprintWidget.IsValid())
-		{
-			MyBlueprintWidget->Refresh();
-		}
+	}
+
+	if (MyBlueprintWidget.IsValid())
+	{
+		MyBlueprintWidget->Refresh();
 	}
 
 	if(SCSEditor.IsValid())
