Fix a couple of regressions in the Blueprint components editor.

- Inherited component node flagging was not functioning.
- Component attachment was comparing against an incorrect Blueprint context in some cases.

diff --git a/Engine/Source/Editor/Kismet/Private/SSCSEditor.cpp b/Engine/Source/Editor/Kismet/Private/SSCSEditor.cpp
--- a/Engine/Source/Editor/Kismet/Private/SSCSEditor.cpp
+++ b/Engine/Source/Editor/Kismet/Private/SSCSEditor.cpp
@@ -1095,7 +1095,7 @@ TSharedPtr<SWidget> SSCS_RowWidget::BuildSceneRootDropActionMenu(FSCSEditorTreeN
 		const FText NodeVariableNameText = FText::FromName( NodePtr->GetVariableName() );
 
 		check(NodePtr.IsValid());
-		const bool bDroppedInSameBlueprint = DroppedNodePtr->GetBlueprint() == NodePtr->GetBlueprint();
+		const bool bDroppedInSameBlueprint = DroppedNodePtr->GetBlueprint() == GetBlueprint();
 
 		MenuBuilder.AddMenuEntry(
 			LOCTEXT("DropActionLabel_AttachToRootNode", "Attach"),
@@ -1281,7 +1281,7 @@ void SSCS_RowWidget::OnDragEnter( const FGeometry& MyGeometry, const FDragDropEv
 					DragRowOp->CurrentHoverText = LOCTEXT("DropActionToolTip_AttachToOrMakeNewRoot", "Drop here to see available actions.");
 					DragRowOp->PendingDropAction = FSCSRowDragDropOp::DropAction_AttachToOrMakeNewRoot;
 				}
-				else if(DraggedNodePtr->GetBlueprint() != NodePtr->GetBlueprint())
+				else if(DraggedNodePtr->GetBlueprint() != GetBlueprint())
 				{
 					if(bCanMakeNewRoot)
 					{
@@ -1357,7 +1357,7 @@ void SSCS_RowWidget::OnDragEnter( const FGeometry& MyGeometry, const FDragDropEv
 			else if(HoveredTemplate->CanAttachAsChild(DraggedTemplate, NAME_None))
 			{
 				// Attach the dragged node(s) to this node
-				if(DraggedNodePtr->GetBlueprint() != NodePtr->GetBlueprint())
+				if(DraggedNodePtr->GetBlueprint() != GetBlueprint())
 				{
 					if(DragRowOp->SourceNodes.Num() > 1)
 					{
@@ -1482,7 +1482,7 @@ void SSCS_RowWidget::OnAttachToDropAction(const TArray<FSCSEditorTreeNodePtrType
 	check(DroppedNodePtrs.Num() > 0);
 
 	// Get the current Blueprint context
-	UBlueprint* Blueprint = NodePtr->GetBlueprint();
+	UBlueprint* Blueprint = GetBlueprint();
 	check(Blueprint);
 
 	TSharedPtr<SSCSEditor> SCSEditorPtr = SCSEditor.Pin();
@@ -1698,7 +1698,7 @@ void SSCS_RowWidget::OnMakeNewRootDropAction(FSCSEditorTreeNodePtrType DroppedNo
 	check(DroppedNodePtr.IsValid());
 
 	// Get the current Blueprint context
-	UBlueprint* Blueprint = NodePtr->GetBlueprint();
+	UBlueprint* Blueprint = GetBlueprint();
 	check(Blueprint != NULL && Blueprint->SimpleConstructionScript != nullptr);
 
 	// Create a transaction record
@@ -1796,7 +1796,7 @@ void SSCS_RowWidget::PostDragDropAction(bool bRegenerateTreeNodes)
 
 		if(NodePtr.IsValid())
 		{
-			UBlueprint* Blueprint = NodePtr->GetBlueprint();
+			UBlueprint* Blueprint = GetBlueprint();
 			if(Blueprint != nullptr)
 			{
 				FBlueprintEditorUtils::PostEditChangeBlueprintActors(Blueprint);
@@ -1898,8 +1898,8 @@ FString SSCS_RowWidget::GetDocumentationExcerptName() const
 
 UBlueprint* SSCS_RowWidget::GetBlueprint() const
 {
-	check(NodePtr.IsValid());
-	return NodePtr->GetBlueprint();
+	check(SCSEditor.IsValid());
+	return SCSEditor.Pin()->GetBlueprint();
 }
 
 bool SSCS_RowWidget::OnNameTextVerifyChanged(const FText& InNewText, FText& OutErrorMessage)
@@ -2075,6 +2075,20 @@ void SSCSEditor::Construct( const FArguments& InArgs )
 	}
 }
 
+UBlueprint* SSCSEditor::GetBlueprint() const
+{
+	AActor* Actor = ActorContext.Get();
+	if(Actor != nullptr)
+	{
+		UClass* ActorClass = Actor->GetClass();
+		check(ActorClass != nullptr);
+
+		return Cast<UBlueprint>(ActorClass->ClassGeneratedBy);
+	}
+
+	return nullptr;
+}
+
 FReply SSCSEditor::OnKeyDown( const FGeometry& MyGeometry, const FKeyEvent& InKeyEvent )
 {
 	if ( CommandList->ProcessCommandBindings( InKeyEvent ) )
@@ -2138,20 +2152,10 @@ TSharedPtr< SWidget > SSCSEditor::CreateContextMenu()
 				MenuBuilder.AddMenuEntry( FGenericCommands::Get().Rename );
 
 				// Collect the classes of all selected objects
-				UBlueprint* Blueprint = nullptr;
 				TArray<UClass*> SelectionClasses;
 				for( auto NodeIter = SelectedNodes.CreateConstIterator(); NodeIter; ++NodeIter )
 				{
 					auto TreeNode = *NodeIter;
-
-					// All nodes should belong to the same Blueprint
-					check(Blueprint == nullptr || Blueprint == TreeNode->GetBlueprint());
-
-					if(Blueprint == nullptr)
-					{
-						Blueprint = TreeNode->GetBlueprint();
-					}
-
 					if( TreeNode->GetComponentTemplate() )
 					{
 						SelectionClasses.Add( TreeNode->GetComponentTemplate()->GetClass() );
@@ -2168,7 +2172,7 @@ TSharedPtr< SWidget > SSCSEditor::CreateContextMenu()
 						MenuBuilder.AddSubMenu(	LOCTEXT("AddEventSubMenu", "Add Event"), 
 							LOCTEXT("ActtionsSubMenu_ToolTip", "Add Event"), 
 							FNewMenuDelegate::CreateStatic( &SSCSEditor::BuildMenuEventsSection,
-								Blueprint, SelectedClass, FCanExecuteAction::CreateSP(this, &SSCSEditor::InEditingMode),
+								GetBlueprint(), SelectedClass, FCanExecuteAction::CreateSP(this, &SSCSEditor::InEditingMode),
 								FGetSelectedObjectsDelegate::CreateSP(this, &SSCSEditor::GetSelectedItemsForContextMenu)));
 					}
 				}
@@ -2669,13 +2673,13 @@ void SSCSEditor::UpdateTree(bool bRegenerateTreeNodes)
 							FSCSEditorTreeNodePtrType ParentNodePtr = FindTreeNode(ParentComponent);
 							if(ParentNodePtr.IsValid())
 							{
-								AddTreeNode(SCS_Node, ParentNodePtr);
+								AddTreeNode(SCS_Node, ParentNodePtr, StackIndex > 0);
 							}
 						}
 					}
 					else
 					{
-						AddTreeNode(SCS_Node, SceneRootNodePtr);
+						AddTreeNode(SCS_Node, SceneRootNodePtr, StackIndex > 0);
 					}
 				}
 			}
@@ -2805,7 +2809,7 @@ UActorComponent* SSCSEditor::AddNewNode(USCS_Node* NewNode,  UObject* Asset, boo
 	}
 
 	// Add the new node to the editor tree
-	NewNodePtr = AddTreeNode(NewNode, SceneRootNodePtr);
+	NewNodePtr = AddTreeNode(NewNode, SceneRootNodePtr, false);
 
 	// Potentially adjust variable names for any child blueprints
 	if(NewNode->VariableName != NAME_None)
@@ -3202,7 +3206,7 @@ bool SSCSEditor::IsNodeInSimpleConstructionScript( USCS_Node* Node ) const
 	return false;
 }
 
-FSCSEditorTreeNodePtrType SSCSEditor::AddTreeNode(USCS_Node* InSCSNode, FSCSEditorTreeNodePtrType InParentNodePtr)
+FSCSEditorTreeNodePtrType SSCSEditor::AddTreeNode(USCS_Node* InSCSNode, FSCSEditorTreeNodePtrType InParentNodePtr, const bool bIsInherited)
 {
 	FSCSEditorTreeNodePtrType NewNodePtr;
 
@@ -3221,7 +3225,6 @@ FSCSEditorTreeNodePtrType SSCSEditor::AddTreeNode(USCS_Node* InSCSNode, FSCSEdit
 	
 	// Determine whether or not the given node is inherited from a parent Blueprint
 	USimpleConstructionScript* NodeSCS = InSCSNode->GetSCS();
-	const bool bIsInherited = NodeSCS != nullptr && InParentNodePtr.IsValid() && NodeSCS->GetBlueprint() != InParentNodePtr->GetBlueprint();
 
 	if(InSCSNode->ComponentTemplate->IsA(USceneComponent::StaticClass()))
 	{
@@ -3303,7 +3306,7 @@ FSCSEditorTreeNodePtrType SSCSEditor::AddTreeNode(USCS_Node* InSCSNode, FSCSEdit
 	// Recursively add the given SCS node's child nodes
 	for(int32 NodeIndex = 0; NodeIndex < InSCSNode->ChildNodes.Num(); ++NodeIndex)
 	{
-		AddTreeNode(InSCSNode->ChildNodes[NodeIndex], NewNodePtr);
+		AddTreeNode(InSCSNode->ChildNodes[NodeIndex], NewNodePtr, bIsInherited);
 	}
 
 	return NewNodePtr;
diff --git a/Engine/Source/Editor/Kismet/Public/SSCSEditor.h b/Engine/Source/Editor/Kismet/Public/SSCSEditor.h
--- a/Engine/Source/Editor/Kismet/Public/SSCSEditor.h
+++ b/Engine/Source/Editor/Kismet/Public/SSCSEditor.h
@@ -535,6 +535,9 @@ public:
 	 */
 	void GetSelectedItemsForContextMenu(TArray<FComponentEventConstructionData>& OutSelectedItems) const;
 
+	/** Provides access to the Blueprint context that's being edited */
+	class UBlueprint* GetBlueprint() const;
+
 protected:
 	FString GetSelectedClassText() const;
 
@@ -583,7 +586,7 @@ protected:
 	void OnDuplicateComponent();
 
 	/** Helper method to add a tree node for the given SCS node */
-	FSCSEditorTreeNodePtrType AddTreeNode(USCS_Node* InSCSNode, FSCSEditorTreeNodePtrType InParentNodePtr);
+	FSCSEditorTreeNodePtrType AddTreeNode(USCS_Node* InSCSNode, FSCSEditorTreeNodePtrType InParentNodePtr, const bool bIsInherited);
 
 	/** Helper method to add a tree node for the given scene component */
 	FSCSEditorTreeNodePtrType AddTreeNode(USceneComponent* InSceneComponent);
