Started to split the SCS editor into a couple different "modes" to help support IWCE.

diff --git a/Engine/Source/Editor/Kismet/Private/BlueprintEditor.cpp b/Engine/Source/Editor/Kismet/Private/BlueprintEditor.cpp
--- a/Engine/Source/Editor/Kismet/Private/BlueprintEditor.cpp
+++ b/Engine/Source/Editor/Kismet/Private/BlueprintEditor.cpp
@@ -665,32 +665,6 @@ void FBlueprintEditor::UpdateSCSPreview(bool bUpdateNow)
 	}
 }
 
-USCS_Node* FBlueprintEditor::OnSCSEditorAddNewComponent(UClass* InComponentClass)
-{
-	check(InComponentClass != nullptr);
-
-	UBlueprint* Blueprint = GetBlueprintObj();
-	check(Blueprint != nullptr && Blueprint->SimpleConstructionScript != nullptr);
-
-	return Blueprint->SimpleConstructionScript->CreateNode(InComponentClass);
-}
-
-USCS_Node* FBlueprintEditor::OnSCSEditorAddExistingComponent(UActorComponent* InComponentInstance)
-{
-	check(InComponentInstance != nullptr);
-
-	UBlueprint* Blueprint = GetBlueprintObj();
-	check(Blueprint != nullptr && Blueprint->SimpleConstructionScript != nullptr);
-
-	if(InComponentInstance->GetOuter() == GetTransientPackage())
-	{
-		// Relocate the instance from the transient package to the BPGC and assign it a unique object name
-		InComponentInstance->Rename(NULL, Blueprint->GeneratedClass, REN_DontCreateRedirectors|REN_DoNotDirty);
-	}
-
-	return Blueprint->SimpleConstructionScript->CreateNode(InComponentInstance);
-}
-
 void FBlueprintEditor::OnSCSEditorTreeViewSelectionChanged(const TArray<FSCSEditorTreeNodePtrType>& SelectedNodes)
 {
 	if (SCSViewport.IsValid())
@@ -1215,10 +1189,8 @@ void FBlueprintEditor::EnsureBlueprintIsUpToDate(UBlueprint* BlueprintObj)
 
 			// Recreate (or create) any widgets that depend on the SCS
 			SCSEditor = SAssignNew(SCSEditor, SSCSEditor)
-				.InEditingMode(this, &FBlueprintEditor::InEditingMode)
-				.ActorContext(this, &FBlueprintEditor::GetPreviewActor)
-				.OnAddNewComponent(this, &FBlueprintEditor::OnSCSEditorAddNewComponent)
-				.OnAddExistingComponent(this, &FBlueprintEditor::OnSCSEditorAddExistingComponent)
+				.ActorContext(BlueprintObj->GeneratedClass->GetDefaultObject<AActor>())
+				.AllowEditing(this, &FBlueprintEditor::InEditingMode)
 				.OnTreeViewSelectionChanged(this, &FBlueprintEditor::OnSCSEditorTreeViewSelectionChanged)
 				.OnUpdateSelectionFromNodes(this, &FBlueprintEditor::OnSCSEditorUpdateSelectionFromNodes);
 			SCSViewport = SAssignNew(SCSViewport, SSCSEditorViewport)
@@ -2001,10 +1973,8 @@ void FBlueprintEditor::CreateDefaultTabContents(const TArray<UBlueprint*>& InBlu
 		InBlueprint->SimpleConstructionScript )
 	{
 		SCSEditor = SAssignNew(SCSEditor, SSCSEditor)
-			.InEditingMode(this, &FBlueprintEditor::InEditingMode)
-			.ActorContext(this, &FBlueprintEditor::GetPreviewActor)
-			.OnAddNewComponent(this, &FBlueprintEditor::OnSCSEditorAddNewComponent)
-			.OnAddExistingComponent(this, &FBlueprintEditor::OnSCSEditorAddExistingComponent)
+			.ActorContext(InBlueprint->GeneratedClass->GetDefaultObject<AActor>())
+			.AllowEditing(this, &FBlueprintEditor::InEditingMode)
 			.OnTreeViewSelectionChanged(this, &FBlueprintEditor::OnSCSEditorTreeViewSelectionChanged)
 			.OnUpdateSelectionFromNodes(this, &FBlueprintEditor::OnSCSEditorUpdateSelectionFromNodes);
 		SCSViewport = SAssignNew(SCSViewport, SSCSEditorViewport)
diff --git a/Engine/Source/Editor/Kismet/Private/BlueprintEditorModes.cpp b/Engine/Source/Editor/Kismet/Private/BlueprintEditorModes.cpp
--- a/Engine/Source/Editor/Kismet/Private/BlueprintEditorModes.cpp
+++ b/Engine/Source/Editor/Kismet/Private/BlueprintEditorModes.cpp
@@ -325,12 +325,11 @@ void FBlueprintComponentsApplicationMode::PostActivateMode()
 	TSharedPtr<FBlueprintEditor> BP = MyBlueprintEditor.Pin();
 	if (BP.IsValid())
 	{
-		BP->EnableSCSPreview(true);
-		BP->UpdateSCSPreview(true);
-		BP->GetInspector()->EnableComponentDetailsCustomization(true);
-
 		auto SCSEditor = BP->GetSCSEditor();
 		SCSEditor->UpdateTree();
+		BP->EnableSCSPreview(true);
+		BP->UpdateSCSPreview();
+		BP->GetInspector()->EnableComponentDetailsCustomization(true);
 
 		// Reselect the cached components
 		TArray<TSharedPtr<FSCSEditorTreeNode>> Selection;
@@ -663,20 +662,6 @@ void FBlueprintEditorUnifiedMode::PostActivateMode()
 	BP->RestoreEditedObjectState();
 	BP->SetupViewForBlueprintEditingMode();
 
-	TSharedPtr<SSCSEditorViewport> PreviewViewportPtr = BP->GetSCSViewport();
-	if(PreviewViewportPtr.IsValid())
-	{
-		// Ensure that the preview actor is instanced
-		PreviewViewportPtr->RequestRefresh(true, true);
-
-		TSharedPtr<SSCSEditor> SCSEditorPtr = BP->GetSCSEditor();
-		if(SCSEditorPtr.IsValid())
-		{
-			// Update the components tree view
-			SCSEditorPtr->UpdateTree();
-		}
-	}
-
 	FApplicationMode::PostActivateMode();
 }
 
diff --git a/Engine/Source/Editor/Kismet/Private/SCSDiff.cpp b/Engine/Source/Editor/Kismet/Private/SCSDiff.cpp
--- a/Engine/Source/Editor/Kismet/Private/SCSDiff.cpp
+++ b/Engine/Source/Editor/Kismet/Private/SCSDiff.cpp
@@ -27,8 +27,8 @@ FSCSDiff::FSCSDiff(const UBlueprint* InBlueprint)
 		+ SSplitter::Slot()
 		[
 			SAssignNew(SCSEditor, SSCSEditor)
-				.InEditingMode(false)
 				.ActorContext(InBlueprint->GeneratedClass->GetDefaultObject<AActor>())
+				.AllowEditing(false)
 				.HideComponentClassCombo(true)
 				.OnUpdateSelectionFromNodes(SSCSEditor::FOnUpdateSelectionFromNodes::CreateRaw(this, &FSCSDiff::OnSCSEditorUpdateSelectionFromNodes))
 				.OnHighlightPropertyInDetailsView(SSCSEditor::FOnHighlightPropertyInDetailsView::CreateRaw(this, &FSCSDiff::OnSCSEditorHighlightPropertyInDetailsView))
diff --git a/Engine/Source/Editor/Kismet/Private/SSCSEditor.cpp b/Engine/Source/Editor/Kismet/Private/SSCSEditor.cpp
--- a/Engine/Source/Editor/Kismet/Private/SSCSEditor.cpp
+++ b/Engine/Source/Editor/Kismet/Private/SSCSEditor.cpp
@@ -879,7 +879,7 @@ TSharedRef<SWidget> SSCS_RowWidget::GenerateWidgetForColumn( const FName& Column
 				.OnVerifyTextChanged( this, &SSCS_RowWidget::OnNameTextVerifyChanged )
 				.OnTextCommitted( this, &SSCS_RowWidget::OnNameTextCommit )
 				.IsSelected( this, &SSCS_RowWidget::IsSelectedExclusively )
-				.IsReadOnly( !NodePtr->CanRename() || (SCSEditor.IsValid() && !SCSEditor.Pin()->InEditingMode()) );
+				.IsReadOnly( !NodePtr->CanRename() || (SCSEditor.IsValid() && !SCSEditor.Pin()->IsEditingAllowed()) );
 
 		NodePtr->SetRenameRequestedDelegate(FSCSEditorTreeNode::FOnRenameRequested::CreateSP(InlineWidget.Get(), &SInlineEditableTextBlock::EnterEditingMode));
 
@@ -1139,7 +1139,7 @@ FReply SSCS_RowWidget::OnDragDetected( const FGeometry& MyGeometry, const FPoint
 	auto SCSEditorPtr = SCSEditor.Pin();
 	if ( MouseEvent.IsMouseButtonDown( EKeys::LeftMouseButton )
 			&& SCSEditorPtr.IsValid()
-			&& SCSEditorPtr->InEditingMode() ) //can only drag when editing
+			&& SCSEditorPtr->IsEditingAllowed() ) //can only drag when editing
 	{
 		TSharedRef<FSCSRowDragDropOp> Operation = FSCSRowDragDropOp::New();
 		Operation->CurrentHoverText = FText::GetEmpty();
@@ -1946,10 +1946,9 @@ void SSCS_RowWidget::OnNameTextCommit(const FText& InNewName, ETextCommit::Type
 
 void SSCSEditor::Construct( const FArguments& InArgs )
 {
+	EditorMode = InArgs._EditorMode;
 	ActorContext = InArgs._ActorContext;
-	bInEditingMode = InArgs._InEditingMode;
-	OnAddNewComponent = InArgs._OnAddNewComponent;
-	OnAddExistingComponent = InArgs._OnAddExistingComponent;
+	AllowEditing = InArgs._AllowEditing;
 	OnTreeViewSelectionChanged = InArgs._OnTreeViewSelectionChanged;
 	OnUpdateSelectionFromNodes = InArgs._OnUpdateSelectionFromNodes;
 	OnHighlightPropertyInDetailsView = InArgs._OnHighlightPropertyInDetailsView;
@@ -2172,7 +2171,7 @@ TSharedPtr< SWidget > SSCSEditor::CreateContextMenu()
 						MenuBuilder.AddSubMenu(	LOCTEXT("AddEventSubMenu", "Add Event"), 
 							LOCTEXT("ActtionsSubMenu_ToolTip", "Add Event"), 
 							FNewMenuDelegate::CreateStatic( &SSCSEditor::BuildMenuEventsSection,
-								GetBlueprint(), SelectedClass, FCanExecuteAction::CreateSP(this, &SSCSEditor::InEditingMode),
+								GetBlueprint(), SelectedClass, FCanExecuteAction::CreateSP(this, &SSCSEditor::IsEditingAllowed),
 								FGetSelectedObjectsDelegate::CreateSP(this, &SSCSEditor::GetSelectedItemsForContextMenu)));
 					}
 				}
@@ -2303,7 +2302,7 @@ void SSCSEditor::ViewEvent(UBlueprint* Blueprint, const FName EventName, const F
 
 bool SSCSEditor::CanDuplicateComponent() const
 {
-	if(!InEditingMode())
+	if(!IsEditingAllowed())
 	{
 		return false;
 	}
@@ -2612,74 +2611,109 @@ void SSCSEditor::UpdateTree(bool bRegenerateTreeNodes)
 		// Reset the scene root node
 		SceneRootNodePtr.Reset();
 
-		// Get the class default object
-		AActor* CDO = NULL;
-		TArray<UBlueprint*> ParentBPStack;
-		AActor* Actor = ActorContext.Get();
-		if(Actor != nullptr)
+		// Build the tree data source according to what mode we're in
+		if(EditorMode.Get() == EEditorMode::BlueprintSCS)
 		{
-			UClass* ActorClass = Actor->GetClass();
-			if(ActorClass != nullptr)
+			// Get the class default object
+			AActor* CDO = NULL;
+			TArray<UBlueprint*> ParentBPStack;
+			AActor* Actor = ActorContext.Get();
+			if(Actor != nullptr)
 			{
-				CDO = ActorClass->GetDefaultObject<AActor>();
+				UClass* ActorClass = Actor->GetClass();
+				if(ActorClass != nullptr)
+				{
+					CDO = ActorClass->GetDefaultObject<AActor>();
 
-				// If it's a Blueprint-generated class, also get the inheritance stack
-				UBlueprint::GetBlueprintHierarchyFromClass(ActorClass, ParentBPStack);
+					// If it's a Blueprint-generated class, also get the inheritance stack
+					UBlueprint::GetBlueprintHierarchyFromClass(ActorClass, ParentBPStack);
+				}
 			}
-		}
-
-		if(CDO != NULL)
-		{
-			// Add native ActorComponent nodes to the root set first
-			TInlineComponentArray<UActorComponent*> Components;
-			CDO->GetComponents(Components);
 
-			for(auto CompIter = Components.CreateIterator(); CompIter; ++CompIter)
+			if(CDO != NULL)
 			{
-				UActorComponent* ActorComp = *CompIter;
-				if (!ActorComp->IsA<USceneComponent>())
+				// Add native ActorComponent nodes to the root set first
+				TInlineComponentArray<UActorComponent*> Components;
+				CDO->GetComponents(Components);
+
+				for(auto CompIter = Components.CreateIterator(); CompIter; ++CompIter)
 				{
-					RootNodes.Add(MakeShareable(new FSCSEditorTreeNode(ActorComp)));
+					UActorComponent* ActorComp = *CompIter;
+					if (!ActorComp->IsA<USceneComponent>())
+					{
+						RootNodes.Add(MakeShareable(new FSCSEditorTreeNode(ActorComp)));
+					}
 				}
-			}
 
-			// Add the native base class SceneComponent hierarchy
-			for(auto CompIter = Components.CreateIterator(); CompIter; ++CompIter)
-			{
-				USceneComponent* SceneComp = Cast<USceneComponent>(*CompIter);
-				if(SceneComp != NULL)
+				// Add the native base class SceneComponent hierarchy
+				for(auto CompIter = Components.CreateIterator(); CompIter; ++CompIter)
 				{
-					AddTreeNode(SceneComp);
+					USceneComponent* SceneComp = Cast<USceneComponent>(*CompIter);
+					if(SceneComp != NULL)
+					{
+						AddTreeNode(SceneComp);
+					}
 				}
 			}
-		}
 
-		// Add the full SCS tree node hierarchy (including SCS nodes inherited from parent blueprints)
-		for(int32 StackIndex = ParentBPStack.Num() - 1; StackIndex >= 0; --StackIndex)
-		{
-			if(ParentBPStack[StackIndex]->SimpleConstructionScript != NULL)
+			// Add the full SCS tree node hierarchy (including SCS nodes inherited from parent blueprints)
+			for(int32 StackIndex = ParentBPStack.Num() - 1; StackIndex >= 0; --StackIndex)
 			{
-				const TArray<USCS_Node*>& SCS_RootNodes = ParentBPStack[StackIndex]->SimpleConstructionScript->GetRootNodes();
-				for(int32 NodeIndex = 0; NodeIndex < SCS_RootNodes.Num(); ++NodeIndex)
+				if(ParentBPStack[StackIndex]->SimpleConstructionScript != NULL)
 				{
-					USCS_Node* SCS_Node = SCS_RootNodes[NodeIndex];
-					check(SCS_Node != NULL);
-
-					if(SCS_Node->ParentComponentOrVariableName != NAME_None)
+					const TArray<USCS_Node*>& SCS_RootNodes = ParentBPStack[StackIndex]->SimpleConstructionScript->GetRootNodes();
+					for(int32 NodeIndex = 0; NodeIndex < SCS_RootNodes.Num(); ++NodeIndex)
 					{
-						USceneComponent* ParentComponent = SCS_Node->GetParentComponentTemplate();
-						if(ParentComponent != NULL)
+						USCS_Node* SCS_Node = SCS_RootNodes[NodeIndex];
+						check(SCS_Node != NULL);
+
+						if(SCS_Node->ParentComponentOrVariableName != NAME_None)
 						{
-							FSCSEditorTreeNodePtrType ParentNodePtr = FindTreeNode(ParentComponent);
-							if(ParentNodePtr.IsValid())
+							USceneComponent* ParentComponent = SCS_Node->GetParentComponentTemplate();
+							if(ParentComponent != NULL)
 							{
-								AddTreeNode(SCS_Node, ParentNodePtr, StackIndex > 0);
+								FSCSEditorTreeNodePtrType ParentNodePtr = FindTreeNode(ParentComponent);
+								if(ParentNodePtr.IsValid())
+								{
+									AddTreeNode(SCS_Node, ParentNodePtr, StackIndex > 0);
+								}
 							}
 						}
+						else
+						{
+							AddTreeNode(SCS_Node, SceneRootNodePtr, StackIndex > 0);
+						}
 					}
-					else
+				}
+			}
+		}
+		else    // EEditorMode::ActorInstance
+		{
+			// Get the actor instance that we're editing
+			AActor* ActorInstance = ActorContext.Get();
+			if(ActorInstance != nullptr)
+			{
+				// Get the full set of instanced components
+				TInlineComponentArray<UActorComponent*> Components;
+				ActorInstance->GetComponents(Components);
+
+				// Add all non-scene component instances to the root set first
+				for(auto CompIter = Components.CreateIterator(); CompIter; ++CompIter)
+				{
+					UActorComponent* ActorComp = *CompIter;
+					if (!ActorComp->IsA<USceneComponent>() && !ActorComp->IsEditorOnly())
 					{
-						AddTreeNode(SCS_Node, SceneRootNodePtr, StackIndex > 0);
+						RootNodes.Add(MakeShareable(new FSCSEditorTreeNode(ActorComp)));
+					}
+				}
+
+				// Now add the instanced scene component hierarchy
+				for(auto CompIter = Components.CreateIterator(); CompIter; ++CompIter)
+				{
+					USceneComponent* SceneComp = Cast<USceneComponent>(*CompIter);
+					if(SceneComp != nullptr && !SceneComp->IsEditorOnly())
+					{
+						AddTreeNode(SceneComp);
 					}
 				}
 			}
@@ -2760,29 +2794,41 @@ void SSCSEditor::SaveSCSNode( USCS_Node* Node )
 	}
 }
 
-bool SSCSEditor::InEditingMode() const
+bool SSCSEditor::IsEditingAllowed() const
 {
-	return bInEditingMode.Get() && nullptr == GEditor->PlayWorld;
+	return AllowEditing.Get() && nullptr == GEditor->PlayWorld;
 }
 
 UActorComponent* SSCSEditor::AddNewComponent( UClass* NewComponentClass, UObject* Asset  )
 {
 	const FScopedTransaction Transaction( LOCTEXT("AddComponent", "Add Component") );
 
-	check(OnAddNewComponent.IsBound());
-	USCS_Node* NewNode = OnAddNewComponent.Execute(NewComponentClass);
+	if(EditorMode.Get() == EEditorMode::BlueprintSCS)
+	{
+		UBlueprint* Blueprint = GetBlueprint();
+		check(Blueprint != nullptr && Blueprint->SimpleConstructionScript != nullptr);
+		
+		Blueprint->Modify();
+		SaveSCSCurrentState(Blueprint->SimpleConstructionScript);
 
-	USimpleConstructionScript* SCS = NewNode->GetSCS();
-	check(SCS != nullptr);
+		return AddNewNode(Blueprint->SimpleConstructionScript->CreateNode(NewComponentClass), Asset, true);
+	}
+	else    // EEditorMode::ActorInstance
+	{
+		AActor* ActorInstance = ActorContext.Get();
+		check(ActorInstance != nullptr);
+		
+		ActorInstance->Modify();
 
-	UBlueprint* Blueprint = SCS->GetBlueprint();
-	check(Blueprint != nullptr);
+		// Create new component
+		UActorComponent* NewComponentInstance = ConstructObject<UActorComponent>(NewComponentClass, ActorInstance, NAME_None, RF_Transactional);
+		check(NewComponentInstance != nullptr);
 
-	// Save current state
-	Blueprint->Modify();
-	SaveSCSCurrentState(SCS);
+		// Add to SerializedComponents array so it gets saved
+		ActorInstance->SerializedComponents.Add(NewComponentInstance);
 
-	return AddNewNode(NewNode, Asset, true);
+		return AddNewNode(NewComponentInstance, Asset, true);
+	}
 }
 
 UActorComponent* SSCSEditor::AddNewNode(USCS_Node* NewNode,  UObject* Asset, bool bMarkBlueprintModified, bool bSetFocusToNewItem)
@@ -2796,14 +2842,11 @@ UActorComponent* SSCSEditor::AddNewNode(USCS_Node* NewNode,  UObject* Asset, boo
 
 	FSCSEditorTreeNodePtrType NewNodePtr;
 
-	USimpleConstructionScript* SCS = NewNode->GetSCS();
-	check(SCS != nullptr);
-
-	UBlueprint* Blueprint = SCS->GetBlueprint();
-	check(Blueprint != nullptr);
+	UBlueprint* Blueprint = GetBlueprint();
+	check(Blueprint != nullptr && Blueprint->SimpleConstructionScript != nullptr);
 
 	// Reset the scene root node if it's set to the default one that's managed by the SCS
-	if(SceneRootNodePtr.IsValid() && SceneRootNodePtr->GetSCSNode() == SCS->GetDefaultSceneRootNode())
+	if(SceneRootNodePtr.IsValid() && SceneRootNodePtr->GetSCSNode() == Blueprint->SimpleConstructionScript->GetDefaultSceneRootNode())
 	{
 		SceneRootNodePtr.Reset();
 	}
@@ -2837,6 +2880,55 @@ UActorComponent* SSCSEditor::AddNewNode(USCS_Node* NewNode,  UObject* Asset, boo
 	return NewNode->ComponentTemplate;
 }
 
+UActorComponent* SSCSEditor::AddNewNode(UActorComponent* NewNode,  UObject* Asset, bool bSetFocusToNewItem)
+{
+	check(NewNode != nullptr);
+
+	if(Asset)
+	{
+		FComponentAssetBrokerage::AssignAssetToComponent(NewNode, Asset);
+	}
+
+	NewNode->RegisterComponent();
+
+	FSCSEditorTreeNodePtrType NewNodePtr;
+
+	// Add the new node to the editor tree
+	USceneComponent* NewSceneComponent = Cast<USceneComponent>(NewNode);
+	if(NewSceneComponent != nullptr)
+	{
+		NewNodePtr = AddTreeNode(NewSceneComponent);
+	}
+	else
+	{
+		NewNodePtr = MakeShareable(new FSCSEditorTreeNode(NewNode));
+
+		// Ensure that the root node ordering is what we assume it to be
+		check(!SceneRootNodePtr.IsValid() || (RootNodes.Num() > 0 && RootNodes[RootNodes.Num() - 1] == SceneRootNodePtr));
+
+		// Add the node to the end of the list of non-scene component nodes (just before the scene component hierarchy)
+		if(RootNodes.Num() > 0)
+		{
+			RootNodes.Insert(NewNodePtr, RootNodes.Num() - 1);
+		}
+		else
+		{
+			RootNodes.Add(NewNodePtr);
+		}
+	}
+
+	if(bSetFocusToNewItem)
+	{
+		// Select and request a rename on the new component
+		SCSTreeWidget->SetSelection(NewNodePtr);
+		OnRenameComponent(false);
+	}
+
+	UpdateTree();
+
+	return NewNode;
+}
+
 bool SSCSEditor::IsComponentSelected(const UPrimitiveComponent* PrimComponent) const
 {
 	check(PrimComponent != NULL);
@@ -2971,7 +3063,7 @@ void SSCSEditor::CopySelectedNodes()
 
 bool SSCSEditor::CanPasteNodes() const
 {
-	if(!InEditingMode())
+	if(!IsEditingAllowed())
 	{
 		return false;
 	}
@@ -2999,8 +3091,26 @@ void SSCSEditor::PasteNodes()
 	// Clear the current selection
 	SCSTreeWidget->ClearSelection();
 
-	// Create a new tree node for each new (pasted) component
+	// Get the object that's being edited
 	UBlueprint* Blueprint = nullptr;
+	AActor* ActorInstance = nullptr;
+	if(EditorMode.Get() == EEditorMode::BlueprintSCS)
+	{
+		Blueprint = GetBlueprint();
+		check(Blueprint != nullptr && Blueprint->SimpleConstructionScript != nullptr);
+
+		Blueprint->Modify();
+		SaveSCSCurrentState(Blueprint->SimpleConstructionScript);
+	}
+	else    // EEditorMode::ActorInstance
+	{
+		ActorInstance = ActorContext.Get();
+		check(ActorInstance != nullptr);
+
+		ActorInstance->Modify();
+	}
+
+	// Create a new tree node for each new (pasted) component
 	TMap<FName, FSCSEditorTreeNodePtrType> NewNodeMap;
 	for(auto NewObjectIt = Factory->NewObjectMap.CreateIterator(); NewObjectIt; ++NewObjectIt)
 	{
@@ -3008,31 +3118,29 @@ void SSCSEditor::PasteNodes()
 		UActorComponent* NewActorComponent = NewObjectIt->Value;
 		check(NewActorComponent);
 
-		// Ask the delegate to create a new SCS node to wrap the component
-		check(OnAddExistingComponent.IsBound());
-		USCS_Node* SCS_Node = OnAddExistingComponent.Execute(NewActorComponent);
+		if(Blueprint != nullptr)
+		{
+			// Relocate the instance from the transient package to the BPGC and assign it a unique object name
+			NewActorComponent->Rename(NULL, Blueprint->GeneratedClass, REN_DontCreateRedirectors|REN_DoNotDirty);
 
-		check(SCS_Node != nullptr);
-		USimpleConstructionScript* SCS = SCS_Node->GetSCS();
-		check(SCS != nullptr && (Blueprint == nullptr || Blueprint == SCS->GetBlueprint()));
-		
-		if(Blueprint == nullptr)
+			// Create a new SCS node to contain the new component and add it to the tree
+			NewActorComponent = AddNewNode(Blueprint->SimpleConstructionScript->CreateNode(NewActorComponent), NULL, false, false);
+		}
+		else
 		{
-			Blueprint = SCS->GetBlueprint();
-			check(Blueprint != nullptr);
+			// Relocate the instance from the transient package to the Actor and assign it a unique object name
+			NewActorComponent->Rename(NULL, ActorInstance, REN_DontCreateRedirectors|REN_DoNotDirty);
 
-			// Mark the Blueprint as about to be modified
-			Blueprint->Modify();
+			// Add to SerializedComponents array so it gets saved
+			ActorInstance->SerializedComponents.Add(NewActorComponent);
 
-			// Mark the SCS and all SCS nodes as about to be modified
-			SaveSCSCurrentState(SCS);
+			// Create a new node to contain the new component instance and add it to the tree
+			NewActorComponent = AddNewNode(NewActorComponent, NULL, false);
 		}
 
-		// Create a new SCS node to contain the new component and add it to the tree
-		NewActorComponent = AddNewNode(SCS_Node, NULL, false, false);
 		if(NewActorComponent)
 		{
-			// Locate the node that corresponds to the new component template
+			// Locate the node that corresponds to the new component template or instance
 			FSCSEditorTreeNodePtrType NewNodePtr = FindTreeNode(NewActorComponent);
 			if(NewNodePtr.IsValid())
 			{
@@ -3064,13 +3172,20 @@ void SSCSEditor::PasteNodes()
 		}
 	}
 
-	// Modify the Blueprint generated class structure (this will also call UpdateTree() as a result)
-	FBlueprintEditorUtils::MarkBlueprintAsStructurallyModified(Blueprint);
+	if(Blueprint != nullptr)
+	{
+		// Modify the Blueprint generated class structure (this will also call UpdateTree() as a result)
+		FBlueprintEditorUtils::MarkBlueprintAsStructurallyModified(Blueprint);
+	}
+	else
+	{
+		UpdateTree();
+	}
 }
 
 bool SSCSEditor::CanDeleteNodes() const
 {
-	if(!InEditingMode())
+	if(!IsEditingAllowed())
 	{
 		return false;
 	}
@@ -3458,7 +3573,7 @@ void SSCSEditor::OnRenameComponent(bool bTransactional)
 
 bool SSCSEditor::CanRenameComponent() const
 {
-	return InEditingMode() && SCSTreeWidget->GetSelectedItems().Num() == 1 && SCSTreeWidget->GetSelectedItems()[0]->CanRename();
+	return IsEditingAllowed() && SCSTreeWidget->GetSelectedItems().Num() == 1 && SCSTreeWidget->GetSelectedItems()[0]->CanRename();
 }
 
 void SSCSEditor::GetCollapsedNodes(const FSCSEditorTreeNodePtrType& InNodePtr, TSet<FSCSEditorTreeNodePtrType>& OutCollapsedNodes) const
diff --git a/Engine/Source/Editor/Kismet/Public/BlueprintEditor.h b/Engine/Source/Editor/Kismet/Public/BlueprintEditor.h
--- a/Engine/Source/Editor/Kismet/Public/BlueprintEditor.h
+++ b/Engine/Source/Editor/Kismet/Public/BlueprintEditor.h
@@ -358,12 +358,6 @@ public:
 	/** Refresh the preview viewport to reflect changes in the SCS */
 	void UpdateSCSPreview(bool bUpdateNow = false);
 
-	/** Delegate invoked when the SCS editor widget needs to create a new component node for a component class type */
-	USCS_Node* OnSCSEditorAddNewComponent(UClass* InComponentClass);
-
-	/** Delegate invoked when the SCS editor widget needs to create a new component node for an existing component instance */
-	USCS_Node* OnSCSEditorAddExistingComponent(UActorComponent* InComponentInstance);
-
 	/** Delegate invoked when the tree view selection is changed in the SCS editor widget */
 	void OnSCSEditorTreeViewSelectionChanged(const TArray< TSharedPtr<class FSCSEditorTreeNode> >& SelectedNodes);
 
diff --git a/Engine/Source/Editor/Kismet/Public/SSCSEditor.h b/Engine/Source/Editor/Kismet/Public/SSCSEditor.h
--- a/Engine/Source/Editor/Kismet/Public/SSCSEditor.h
+++ b/Engine/Source/Editor/Kismet/Public/SSCSEditor.h
@@ -375,6 +375,15 @@ typedef SSCSEditorDragDropTree SSCSTreeType;
 class KISMET_API SSCSEditor : public SCompoundWidget
 {
 public:
+	/* Editor mode */
+	enum EEditorMode
+	{
+		/* View/edit the SCS in a BGPC */
+		BlueprintSCS,
+		/* View/edit the Actor instance */
+		ActorInstance
+	};
+
 	DECLARE_DELEGATE_RetVal_OneParam(class USCS_Node*, FOnAddNewComponent, class UClass*);
 	DECLARE_DELEGATE_RetVal_OneParam(class USCS_Node*, FOnAddExistingComponent, class UActorComponent*);
 	DECLARE_DELEGATE_OneParam(FOnTreeViewSelectionChanged, const TArray<FSCSEditorTreeNodePtrType>&);
@@ -382,21 +391,19 @@ public:
 	DECLARE_DELEGATE_OneParam(FOnHighlightPropertyInDetailsView, const class FPropertyPath&);
 
 	SLATE_BEGIN_ARGS( SSCSEditor )
-		:_InEditingMode(true)
+		:_EditorMode(BlueprintSCS)
 		,_ActorContext(nullptr)
+		,_AllowEditing(true)
 		,_HideComponentClassCombo(false)
-		,_OnAddNewComponent()
-		,_OnAddExistingComponent()
 		,_OnTreeViewSelectionChanged()
 		,_OnUpdateSelectionFromNodes()
 		,_OnHighlightPropertyInDetailsView()
 		{}
 
-		SLATE_ATTRIBUTE(bool, InEditingMode)
+		SLATE_ATTRIBUTE(EEditorMode, EditorMode)
 		SLATE_ATTRIBUTE(class AActor*, ActorContext)
+		SLATE_ATTRIBUTE(bool, AllowEditing)
 		SLATE_ATTRIBUTE(bool, HideComponentClassCombo)
-		SLATE_EVENT(FOnAddNewComponent, OnAddNewComponent)
-		SLATE_EVENT(FOnAddExistingComponent, OnAddExistingComponent)
 		SLATE_EVENT(FOnTreeViewSelectionChanged, OnTreeViewSelectionChanged)
 		SLATE_EVENT(FOnUpdateSelectionFromNodes, OnUpdateSelectionFromNodes)
 		SLATE_EVENT(FOnHighlightPropertyInDetailsView, OnHighlightPropertyInDetailsView)
@@ -414,8 +421,8 @@ public:
 	/** Used by tree control - get children for a specified node */
 	void OnGetChildrenForTree(FSCSEditorTreeNodePtrType InNodePtr, TArray<FSCSEditorTreeNodePtrType>& OutChildren);
 
-	/** Returns true if in editing mode */
-	bool InEditingMode() const;
+	/** Returns true if editing is allowed */
+	bool IsEditingAllowed() const;
 
 	/** Adds a component to the SCS Table
 	   @param NewComponentClass	(In) The class to add
@@ -423,13 +430,20 @@ public:
 	   @return The reference of the newly created ActorComponent */
 	UActorComponent* AddNewComponent(UClass* NewComponentClass, UObject* Asset);
 
-	/** Adds a new Node to the component Table
-	   @param NewNode	(In) The node to add
+	/** Adds a new SCS Node to the component Table
+	   @param NewNode	(In) The SCS node to add
 	   @param Asset		(In) Optional asset to assign to the component
 	   @param bMarkBlueprintModified (In) Whether or not to mark the Blueprint as structurally modified
 	   @param bSetFocusToNewItem (In) Select the new item and activate the inline rename widget (default is true)
 	   @return The reference of the newly created ActorComponent */
 	UActorComponent* AddNewNode(USCS_Node* NewNode, UObject* Asset, bool bMarkBlueprintModified, bool bSetFocusToNewItem = true);
+
+	/** Adds a new component instance node to the component Table
+		@param NewNode	(In) The component instance node to add
+		@param Asset	(In) Optional asset to assign to the component
+		@param bSetFocusToNewItem (In) Select the new item and activate the inline rename widget (default is true)
+		@return The reference of the newly created ActorComponent */
+	UActorComponent* AddNewNode(UActorComponent* NewNode, UObject* Asset, bool bSetFocusToNewItem = true);
 	
 	/** Returns true if the specified component is currently selected */
 	bool IsComponentSelected(const UPrimitiveComponent* PrimComponent) const;
@@ -621,17 +635,14 @@ public:
 	/** Whether or not the deferred rename request was flagged as transactional */
 	bool bIsDeferredRenameRequestTransactional;
 
+	/** Attribute to indicate which editor mode we're in. */
+	TAttribute<EEditorMode> EditorMode;
+
 	/** Attribute that provides access to the Actor context for which we are viewing/editing the SCS. */
 	TAttribute<class AActor*> ActorContext;
 
-	/** Attribute to control whether or not to allow editing. */
-	TAttribute<bool> bInEditingMode;
-
-	/** Delegate to invoke when a new SCS node needs to be created for a component class type. */
-	FOnAddNewComponent OnAddNewComponent;
-
-	/** Delegate to invoke when a new SCS node needs to be created for an existing component instance. */
-	FOnAddExistingComponent OnAddExistingComponent;
+	/** Attribute to indicate whether or not editing is allowed. */
+	TAttribute<bool> AllowEditing;
 
 	/** Delegate to invoke on tree view selection change. */
 	FOnTreeViewSelectionChanged OnTreeViewSelectionChanged;
diff --git a/Engine/Source/Editor/LevelEditor/Private/SLevelEditor.cpp b/Engine/Source/Editor/LevelEditor/Private/SLevelEditor.cpp
--- a/Engine/Source/Editor/LevelEditor/Private/SLevelEditor.cpp
+++ b/Engine/Source/Editor/LevelEditor/Private/SLevelEditor.cpp
@@ -602,7 +602,7 @@ public:
 					ComponentsBox->SetContent
 					(
 						SAssignNew(SCSEditor, SSCSEditor)
-						.InEditingMode(true)
+						.EditorMode(SSCSEditor::EEditorMode::ActorInstance)
 						.ActorContext(this, &SActorDetails::GetSelectedActor)												// Get the instance of the actor in the world
 						.OnTreeViewSelectionChanged(this, &SActorDetails::OnSCSEditorTreeViewSelectionChanged)				// A selection has been made in the tree view, so inform the level editor
 						//.OnUpdateSelectionFromNodes(this, &SLevelEditor::OnSCSEditorUpdateSelectionFromNodes)				// Unsure, don't think it's needed
