Fix crash during Blueprint compile in editor. (regression)

Also restored the 'PreviewActor' attribute in order to support both a CDO and an instanced Actor in BP editor mode.

diff --git a/Engine/Source/Editor/Kismet/Private/BlueprintEditor.cpp b/Engine/Source/Editor/Kismet/Private/BlueprintEditor.cpp
--- a/Engine/Source/Editor/Kismet/Private/BlueprintEditor.cpp
+++ b/Engine/Source/Editor/Kismet/Private/BlueprintEditor.cpp
@@ -665,6 +665,18 @@ void FBlueprintEditor::UpdateSCSPreview(bool bUpdateNow)
 	}
 }
 
+AActor* FBlueprintEditor::GetSCSEditorActorContext() const
+{
+	// Return the current CDO that was last generated for the class
+	UBlueprint* Blueprint = GetBlueprintObj();
+	if(Blueprint != nullptr && Blueprint->GeneratedClass != nullptr)
+	{
+		return Blueprint->GeneratedClass->GetDefaultObject<AActor>();
+	}
+	
+	return nullptr;
+}
+
 void FBlueprintEditor::OnSCSEditorTreeViewSelectionChanged(const TArray<FSCSEditorTreeNodePtrType>& SelectedNodes)
 {
 	if (SCSViewport.IsValid())
@@ -1189,7 +1201,8 @@ void FBlueprintEditor::EnsureBlueprintIsUpToDate(UBlueprint* BlueprintObj)
 
 			// Recreate (or create) any widgets that depend on the SCS
 			SCSEditor = SAssignNew(SCSEditor, SSCSEditor)
-				.ActorContext(BlueprintObj->GeneratedClass->GetDefaultObject<AActor>())
+				.ActorContext(this, &FBlueprintEditor::GetSCSEditorActorContext)
+				.PreviewActor(this, &FBlueprintEditor::GetPreviewActor)
 				.AllowEditing(this, &FBlueprintEditor::InEditingMode)
 				.OnTreeViewSelectionChanged(this, &FBlueprintEditor::OnSCSEditorTreeViewSelectionChanged)
 				.OnUpdateSelectionFromNodes(this, &FBlueprintEditor::OnSCSEditorUpdateSelectionFromNodes);
@@ -1973,7 +1986,8 @@ void FBlueprintEditor::CreateDefaultTabContents(const TArray<UBlueprint*>& InBlu
 		InBlueprint->SimpleConstructionScript )
 	{
 		SCSEditor = SAssignNew(SCSEditor, SSCSEditor)
-			.ActorContext(InBlueprint->GeneratedClass->GetDefaultObject<AActor>())
+			.ActorContext(this, &FBlueprintEditor::GetSCSEditorActorContext)
+			.PreviewActor(this, &FBlueprintEditor::GetPreviewActor)
 			.AllowEditing(this, &FBlueprintEditor::InEditingMode)
 			.OnTreeViewSelectionChanged(this, &FBlueprintEditor::OnSCSEditorTreeViewSelectionChanged)
 			.OnUpdateSelectionFromNodes(this, &FBlueprintEditor::OnSCSEditorUpdateSelectionFromNodes);
diff --git a/Engine/Source/Editor/Kismet/Private/SSCSEditor.cpp b/Engine/Source/Editor/Kismet/Private/SSCSEditor.cpp
--- a/Engine/Source/Editor/Kismet/Private/SSCSEditor.cpp
+++ b/Engine/Source/Editor/Kismet/Private/SSCSEditor.cpp
@@ -1488,9 +1488,13 @@ void SSCS_RowWidget::OnAttachToDropAction(const TArray<FSCSEditorTreeNodePtrType
 	TSharedPtr<SSCSEditor> SCSEditorPtr = SCSEditor.Pin();
 	check(SCSEditorPtr.IsValid());
 
-	// Get the current Actor context
-	AActor* Actor = SCSEditorPtr->ActorContext.Get();
-	check(Actor != nullptr);
+	// Get the current "instanced" Actor context
+	AActor* Actor = SCSEditorPtr->PreviewActor.Get();
+	if(Actor == nullptr)
+	{
+		Actor = SCSEditorPtr->ActorContext.Get();
+		check(Actor != nullptr);
+	}
 
 	bool bRegenerateTreeNodes = false;
 	const FScopedTransaction TransactionContext(DroppedNodePtrs.Num() > 1 ? LOCTEXT("AttachComponents", "Attach Components") : LOCTEXT("AttachComponent", "Attach Component"));
@@ -1609,9 +1613,13 @@ void SSCS_RowWidget::OnDetachFromDropAction(const TArray<FSCSEditorTreeNodePtrTy
 	TSharedPtr<SSCSEditor> SCSEditorPtr = SCSEditor.Pin();
 	check(SCSEditorPtr.IsValid());
 
-	// Get the current preview Actor context
-	AActor* Actor = SCSEditorPtr->ActorContext.Get();
-	check(Actor != nullptr);
+	// Get the current "instanced" Actor context
+	AActor* Actor = SCSEditorPtr->PreviewActor.Get();
+	if(Actor == nullptr)
+	{
+		Actor = SCSEditorPtr->ActorContext.Get();
+		check(Actor != nullptr);
+	}
 
 	const FScopedTransaction TransactionContext(DroppedNodePtrs.Num() > 1 ? LOCTEXT("DetachComponents", "Detach Components") : LOCTEXT("DetachComponent", "Detach Component"));
 
@@ -1949,6 +1957,7 @@ void SSCSEditor::Construct( const FArguments& InArgs )
 	EditorMode = InArgs._EditorMode;
 	ActorContext = InArgs._ActorContext;
 	AllowEditing = InArgs._AllowEditing;
+	PreviewActor = InArgs._PreviewActor;
 	OnTreeViewSelectionChanged = InArgs._OnTreeViewSelectionChanged;
 	OnUpdateSelectionFromNodes = InArgs._OnUpdateSelectionFromNodes;
 	OnHighlightPropertyInDetailsView = InArgs._OnHighlightPropertyInDetailsView;
diff --git a/Engine/Source/Editor/Kismet/Public/BlueprintEditor.h b/Engine/Source/Editor/Kismet/Public/BlueprintEditor.h
--- a/Engine/Source/Editor/Kismet/Public/BlueprintEditor.h
+++ b/Engine/Source/Editor/Kismet/Public/BlueprintEditor.h
@@ -357,6 +357,9 @@ public:
 
 	/** Refresh the preview viewport to reflect changes in the SCS */
 	void UpdateSCSPreview(bool bUpdateNow = false);
+	
+	/** Delegate invoked when the SCS editor needs to obtain the Actor context for editing */
+	AActor* GetSCSEditorActorContext() const;
 
 	/** Delegate invoked when the tree view selection is changed in the SCS editor widget */
 	void OnSCSEditorTreeViewSelectionChanged(const TArray< TSharedPtr<class FSCSEditorTreeNode> >& SelectedNodes);
diff --git a/Engine/Source/Editor/Kismet/Public/SSCSEditor.h b/Engine/Source/Editor/Kismet/Public/SSCSEditor.h
--- a/Engine/Source/Editor/Kismet/Public/SSCSEditor.h
+++ b/Engine/Source/Editor/Kismet/Public/SSCSEditor.h
@@ -393,6 +393,7 @@ public:
 	SLATE_BEGIN_ARGS( SSCSEditor )
 		:_EditorMode(BlueprintSCS)
 		,_ActorContext(nullptr)
+		,_PreviewActor(nullptr)
 		,_AllowEditing(true)
 		,_HideComponentClassCombo(false)
 		,_OnTreeViewSelectionChanged()
@@ -402,6 +403,7 @@ public:
 
 		SLATE_ATTRIBUTE(EEditorMode, EditorMode)
 		SLATE_ATTRIBUTE(class AActor*, ActorContext)
+		SLATE_ATTRIBUTE(class AActor*, PreviewActor)
 		SLATE_ATTRIBUTE(bool, AllowEditing)
 		SLATE_ATTRIBUTE(bool, HideComponentClassCombo)
 		SLATE_EVENT(FOnTreeViewSelectionChanged, OnTreeViewSelectionChanged)
@@ -641,6 +643,9 @@ public:
 	/** Attribute that provides access to the Actor context for which we are viewing/editing the SCS. */
 	TAttribute<class AActor*> ActorContext;
 
+	/** Attribute that provides access to a "preview" Actor context (may not be same as the Actor context that's being edited. */
+	TAttribute<class AActor*> PreviewActor;
+
 	/** Attribute to indicate whether or not editing is allowed. */
 	TAttribute<bool> AllowEditing;
 
