Renaming the blueprints editor button 'Blueprints props' to 'Options'.
#UE-7418 - Change text on "Blueprint Props" button to "Blueprint Properties"
#Branch UE4
#Proj Kismet

#info I tried Blueprint Properties as the button label but it's just too big and makes the button oversized. The terminology elsewhere appears to be Blueprint Options so I went with Options.

diff --git a/Engine/Source/Editor/Kismet/Private/SBlueprintEditorToolbar.cpp b/Engine/Source/Editor/Kismet/Private/SBlueprintEditorToolbar.cpp
--- a/Engine/Source/Editor/Kismet/Private/SBlueprintEditorToolbar.cpp
+++ b/Engine/Source/Editor/Kismet/Private/SBlueprintEditorToolbar.cpp
@@ -343,7 +343,8 @@ void FFullBlueprintEditorCommands::RegisterCommands()
 	UI_COMMAND(SwitchToScriptingMode, "Graph", "Switches to Graph Editing Mode", EUserInterfaceActionType::ToggleButton, FInputGesture());
 	UI_COMMAND(SwitchToBlueprintDefaultsMode, "Defaults", "Switches to Blueprint Defaults Mode", EUserInterfaceActionType::ToggleButton, FInputGesture());
 	UI_COMMAND(SwitchToComponentsMode, "Components", "Switches to Components Mode", EUserInterfaceActionType::ToggleButton, FInputGesture());
-	UI_COMMAND(EditGlobalOptions, "Blueprint Props", "Edit Blueprint Options", EUserInterfaceActionType::Button, FInputGesture());
+	// Renaming this to Options because Blueprint props was requested to be changed and Blueprint properties was too long and looks odd.
+	UI_COMMAND(EditGlobalOptions, "Options", "Edit Blueprint Options", EUserInterfaceActionType::Button, FInputGesture());
 
 	UI_COMMAND(JumpToErrorNode, "Jump to Error Node", "When enabled, then the Blueprint will snap focus to nodes producing an error during compilation", EUserInterfaceActionType::ToggleButton, FInputGesture());
 }
