Can no longer duplicate Timeline variables in the My Blueprint window (nor any other non-allowable types that might appear).

diff --git a/Engine/Source/Editor/Kismet/Private/SMyBlueprint.cpp b/Engine/Source/Editor/Kismet/Private/SMyBlueprint.cpp
--- a/Engine/Source/Editor/Kismet/Private/SMyBlueprint.cpp
+++ b/Engine/Source/Editor/Kismet/Private/SMyBlueprint.cpp
@@ -1997,7 +1997,20 @@ bool SMyBlueprint::CanDuplicateAction() const
 			return GraphAction->EdGraph->GetSchema()->CanDuplicateGraph(GraphAction->EdGraph);
 		}
 	}
-	else if(SelectionAsVar() || SelectionAsLocalVar())
+	else if(FEdGraphSchemaAction_K2Var* VarAction = SelectionAsVar())
+	{
+		// if the property is not an allowable Blueprint variable type, do not allow the variable to be duplicated.
+		// Some actions (timelines) exist as variables but cannot be used in a user-defined variable.
+		const UObjectPropertyBase* ObjectProperty = Cast<const UObjectPropertyBase>(VarAction->GetProperty());
+		if (ObjectProperty &&
+			ObjectProperty->PropertyClass &&
+			!UEdGraphSchema_K2::IsAllowableBlueprintVariableType(ObjectProperty->PropertyClass))
+		{
+			return false;
+		}
+		return true;
+	}
+	else if(SelectionAsLocalVar())
 	{
 		return true;
 	}
