More SSCSEditor work to support shared usage w/ the Level editor.

Change summary:
- Modified FSCSEditorTreeNode to support an "instanced" concept; all nodes in IWCE are "native," "instanced" denotes that they are IWCE nodes, "user instanced" denotes that they were created by the user in IWCE. I'm checking the already-existing "bCreatedByConstructionScript" flag on the component instance itself to determine if the instance was created by a BP generated class and not by C++ or the user.
- Alternate tooltip/doc paths for IWCE.
- First-pass node attachment/detachment in IWCE via drag/drop. Can also drag assets into the tree view now from the CB.
- Fixed a regression that caused non-native node attachments to native nodes to no longer appear in the tree view for child Blueprints.

Remaining TODO items:
- Node deletion
- Node rename support
- Fix up Blueprint-instanced nodes to be refreshed after reconstruction of the selected Actor instance.

diff --git a/Engine/Source/Editor/Kismet/Private/SSCSEditor.cpp b/Engine/Source/Editor/Kismet/Private/SSCSEditor.cpp
--- a/Engine/Source/Editor/Kismet/Private/SSCSEditor.cpp
+++ b/Engine/Source/Editor/Kismet/Private/SSCSEditor.cpp
@@ -205,18 +205,49 @@ FSCSEditorTreeNode::FSCSEditorTreeNode()
 
 FSCSEditorTreeNode::FSCSEditorTreeNode(USCS_Node* InSCSNode, bool bInIsInherited)
 	:bIsInherited(bInIsInherited)
+	,bIsInstanced(false)
 	,SCSNodePtr(InSCSNode)
 	,ComponentTemplatePtr(InSCSNode != NULL ? InSCSNode->ComponentTemplate : NULL)
 	,bNonTransactionalRename(false)
+	,bWasInstancedFromNativeClass(false)
 {
 }
 
 FSCSEditorTreeNode::FSCSEditorTreeNode(UActorComponent* InComponentTemplate)
 	:bIsInherited(false)
+	,bIsInstanced(false)
 	,SCSNodePtr(NULL)
 	,ComponentTemplatePtr(InComponentTemplate)
 	,bNonTransactionalRename(false)
+	,bWasInstancedFromNativeClass(false)
 {
+	check(InComponentTemplate != nullptr);
+	AActor* Owner = InComponentTemplate->GetOwner();
+	if(Owner != nullptr && !Owner->HasAllFlags(RF_ClassDefaultObject))
+	{
+		bIsInstanced = true;
+		
+		UClass* OwnerClass = Owner->GetActorClass();
+		if(OwnerClass != nullptr)
+		{
+			AActor* CDO = OwnerClass->GetDefaultObject<AActor>();
+			if(CDO != nullptr)
+			{
+				// Iterate over the Components array and attempt to find a component with a matching name
+				TInlineComponentArray<UActorComponent*> Components;
+				CDO->GetComponents(Components);
+
+				for(auto It = Components.CreateConstIterator(); It && !bWasInstancedFromNativeClass; ++It)
+				{
+					UActorComponent* ComponentTemplate = *It;
+					if(ComponentTemplate->GetFName() == InComponentTemplate->GetFName())
+					{
+						bWasInstancedFromNativeClass = true;
+					}
+				}
+			}
+		}
+	}
 }
 
 FName FSCSEditorTreeNode::GetVariableName() const
@@ -446,6 +477,11 @@ bool FSCSEditorTreeNode::CanEditDefaults() const
 		// Evaluate to TRUE for non-native nodes if it represents a valid SCS node and it is not inherited from a parent Blueprint
 		bCanEdit = SCS_Node != NULL && !IsInherited();
 	}
+	else if(bIsInstanced)
+	{
+		// Evaluate to TRUE for all instanced components except for those instanced from the Blueprint-generated class (i.e. during SCS or UCS)
+		bCanEdit = !ComponentTemplate->bCreatedByConstructionScript;
+	}
 	else if(ComponentTemplate != NULL)
 	{
 		// Evaluate to TRUE for native nodes if it is bound to a member variable and that variable has either EditDefaultsOnly or EditAnywhere flags set
@@ -573,6 +609,17 @@ void FSCSEditorTreeNode::AddChild(FSCSEditorTreeNodePtrType InChildNodePtr)
 			}
 		}
 	}
+	else if(bIsInstanced)
+	{
+		USceneComponent* ChildInstance = Cast<USceneComponent>(InChildNodePtr->GetComponentTemplate());
+		check(ChildInstance != nullptr);
+
+		USceneComponent* ParentInstance = Cast<USceneComponent>(GetComponentTemplate());
+		check(ParentInstance != nullptr);
+
+		// Handle attachment at the instance level
+		ChildInstance->AttachTo(ParentInstance);
+	}
 }
 
 FSCSEditorTreeNodePtrType FSCSEditorTreeNode::AddChild(USCS_Node* InSCSNode, bool bInIsInherited)
@@ -681,6 +728,14 @@ void FSCSEditorTreeNode::RemoveChild(FSCSEditorTreeNodePtrType InChildNodePtr)
 			SCS->RemoveNode(SCS_ChildNode);
 		}
 	}
+	else if(bIsInstanced)
+	{
+		USceneComponent* ChildInstance = Cast<USceneComponent>(InChildNodePtr->GetComponentTemplate());
+		check(ChildInstance != nullptr);
+
+		// Handle detachment at the instance level
+		ChildInstance->DetachFromParent();
+	}
 }
 
 void FSCSEditorTreeNode::OnRequestRename(bool bTransactional)
@@ -1060,13 +1115,33 @@ EVisibility SSCS_RowWidget::GetAssetVisibility() const
 
 FSlateColor SSCS_RowWidget::GetColorTint() const
 {
-	if(NodePtr->IsNative())
+	if(SCSEditor.Pin()->EditorMode.Get() == SSCSEditor::EEditorMode::BlueprintSCS)
 	{
-		return FLinearColor(0.08f,0.15f,0.6f);
+		if(NodePtr->IsNative())
+		{
+			return FLinearColor(0.08f,0.15f,0.6f);
+		}
+		else if(NodePtr->IsInherited())
+		{
+			return FLinearColor(0.08f,0.35f,0.6f);
+		}
 	}
-	else if(NodePtr->IsInherited())
+	else    // SSCSEditor::EEditorMode::ActorInstance
 	{
-		return FLinearColor(0.08f,0.35f,0.6f);
+		if(!NodePtr->IsUserInstanced())
+		{
+			UActorComponent* InstancedComponent = NodePtr->GetComponentTemplate();
+			check(InstancedComponent != nullptr);
+
+			if(InstancedComponent->bCreatedByConstructionScript)
+			{
+				return FLinearColor(0.08f,0.35f,0.6f);
+			}
+			else
+			{
+				return FLinearColor(0.08f,0.15f,0.6f);
+			}
+		}
 	}
 
 	return FSlateColor::UseForeground();
@@ -1095,7 +1170,11 @@ TSharedPtr<SWidget> SSCS_RowWidget::BuildSceneRootDropActionMenu(FSCSEditorTreeN
 		const FText NodeVariableNameText = FText::FromName( NodePtr->GetVariableName() );
 
 		check(NodePtr.IsValid());
-		const bool bDroppedInSameBlueprint = DroppedNodePtr->GetBlueprint() == GetBlueprint();
+		bool bDroppedInSameBlueprint = true;
+		if(SCSEditor.Pin()->EditorMode.Get() == SSCSEditor::EEditorMode::BlueprintSCS)
+		{
+			bDroppedInSameBlueprint = DroppedNodePtr->GetBlueprint() == GetBlueprint();
+		}
 
 		MenuBuilder.AddMenuEntry(
 			LOCTEXT("DropActionLabel_AttachToRootNode", "Attach"),
@@ -1281,7 +1360,7 @@ void SSCS_RowWidget::OnDragEnter( const FGeometry& MyGeometry, const FDragDropEv
 					DragRowOp->CurrentHoverText = LOCTEXT("DropActionToolTip_AttachToOrMakeNewRoot", "Drop here to see available actions.");
 					DragRowOp->PendingDropAction = FSCSRowDragDropOp::DropAction_AttachToOrMakeNewRoot;
 				}
-				else if(DraggedNodePtr->GetBlueprint() != GetBlueprint())
+				else if(SCSEditor.Pin()->EditorMode.Get() == SSCSEditor::EEditorMode::BlueprintSCS && DraggedNodePtr->GetBlueprint() != GetBlueprint())
 				{
 					if(bCanMakeNewRoot)
 					{
@@ -1481,123 +1560,138 @@ void SSCS_RowWidget::OnAttachToDropAction(const TArray<FSCSEditorTreeNodePtrType
 	check(NodePtr.IsValid());
 	check(DroppedNodePtrs.Num() > 0);
 
-	// Get the current Blueprint context
-	UBlueprint* Blueprint = GetBlueprint();
-	check(Blueprint);
-
 	TSharedPtr<SSCSEditor> SCSEditorPtr = SCSEditor.Pin();
 	check(SCSEditorPtr.IsValid());
 
-	// Get the current "instanced" Actor context
-	AActor* Actor = SCSEditorPtr->PreviewActor.Get();
-	if(Actor == nullptr)
-	{
-		Actor = SCSEditorPtr->ActorContext.Get();
-		check(Actor != nullptr);
-	}
-
 	bool bRegenerateTreeNodes = false;
 	const FScopedTransaction TransactionContext(DroppedNodePtrs.Num() > 1 ? LOCTEXT("AttachComponents", "Attach Components") : LOCTEXT("AttachComponent", "Attach Component"));
 
-	for(const auto& DroppedNodePtr : DroppedNodePtrs)
+	if(SCSEditorPtr->EditorMode.Get() == SSCSEditor::EEditorMode::BlueprintSCS)
 	{
-		// Clone the component if it's being dropped into a different SCS
-		if(DroppedNodePtr->GetBlueprint() != Blueprint)
+		// Get the current Blueprint context
+		UBlueprint* Blueprint = GetBlueprint();
+		check(Blueprint);
+
+		// Get the current "preview" Actor instance
+		AActor* PreviewActor = SCSEditorPtr->PreviewActor.Get();
+		check(PreviewActor);
+
+		for(const auto& DroppedNodePtr : DroppedNodePtrs)
 		{
-			bRegenerateTreeNodes = true;
+			// Clone the component if it's being dropped into a different SCS
+			if(DroppedNodePtr->GetBlueprint() != Blueprint)
+			{
+				bRegenerateTreeNodes = true;
 
-			check(DroppedNodePtr.IsValid());
-			UActorComponent* ComponentTemplate = DroppedNodePtr->GetComponentTemplate();
-			check(ComponentTemplate);
+				check(DroppedNodePtr.IsValid());
+				UActorComponent* ComponentTemplate = DroppedNodePtr->GetComponentTemplate();
+				check(ComponentTemplate);
 
-			// Note: This will mark the Blueprint as structurally modified
-			UActorComponent* ClonedComponent = SCSEditorPtr->AddNewComponent(ComponentTemplate->GetClass(), NULL);
-			check(ClonedComponent);
+				// Note: This will mark the Blueprint as structurally modified
+				UActorComponent* ClonedComponent = SCSEditorPtr->AddNewComponent(ComponentTemplate->GetClass(), NULL);
+				check(ClonedComponent);
 
-			//Serialize object properties using write/read operations.
-			TArray<uint8> SavedProperties;
-			FObjectWriter Writer(ComponentTemplate, SavedProperties);
-			FObjectReader(ClonedComponent, SavedProperties);
+				//Serialize object properties using write/read operations.
+				TArray<uint8> SavedProperties;
+				FObjectWriter Writer(ComponentTemplate, SavedProperties);
+				FObjectReader(ClonedComponent, SavedProperties);
 
-			// Attach the copied node to the target node (this will also detach it from the root if necessary)
-			FSCSEditorTreeNodePtrType NewNodePtr = SCSEditorPtr->GetNodeFromActorComponent(ClonedComponent);
-			if(NewNodePtr.IsValid())
-			{
-				NodePtr->AddChild(NewNodePtr);
+				// Attach the copied node to the target node (this will also detach it from the root if necessary)
+				FSCSEditorTreeNodePtrType NewNodePtr = SCSEditorPtr->GetNodeFromActorComponent(ClonedComponent);
+				if(NewNodePtr.IsValid())
+				{
+					NodePtr->AddChild(NewNodePtr);
+				}
 			}
-		}
-		else
-		{
-			// Get the associated component template if it is a scene component, so we can adjust the transform
-			USceneComponent* SceneComponentTemplate = Cast<USceneComponent>(DroppedNodePtr->GetComponentTemplate());
-
-			// Check for a valid parent node
-			FSCSEditorTreeNodePtrType ParentNodePtr = DroppedNodePtr->GetParent();
-			if(ParentNodePtr.IsValid())
+			else
 			{
-				// Detach the dropped node from its parent
-				ParentNodePtr->RemoveChild(DroppedNodePtr);
+				// Get the associated component template if it is a scene component, so we can adjust the transform
+				USceneComponent* SceneComponentTemplate = Cast<USceneComponent>(DroppedNodePtr->GetComponentTemplate());
 
-				// If the associated component template is a scene component, maintain its preview world position
-				if(SceneComponentTemplate)
+				// Check for a valid parent node
+				FSCSEditorTreeNodePtrType ParentNodePtr = DroppedNodePtr->GetParent();
+				if(ParentNodePtr.IsValid())
 				{
-					// Save current state
-					SceneComponentTemplate->Modify();
+					// Detach the dropped node from its parent
+					ParentNodePtr->RemoveChild(DroppedNodePtr);
 
-					// Reset the attach socket name
-					SceneComponentTemplate->AttachSocketName = NAME_None;
-					USCS_Node* SCS_Node = DroppedNodePtr->GetSCSNode();
-					if(SCS_Node)
+					// If the associated component template is a scene component, maintain its preview world position
+					if(SceneComponentTemplate)
 					{
-						SCS_Node->Modify();
-						SCS_Node->AttachToName = NAME_None;
-					}
+						// Save current state
+						SceneComponentTemplate->Modify();
 
-					// Attempt to locate a matching registered instance of the component template in the Actor context that's being edited
-					USceneComponent* InstancedSceneComponent = Cast<USceneComponent>(DroppedNodePtr->FindComponentInstanceInActor(Actor));
-					if(InstancedSceneComponent && InstancedSceneComponent->IsRegistered())
-					{
-						// If we find a match, save off the world position
-						FTransform ComponentToWorld = InstancedSceneComponent->GetComponentToWorld();
-						SceneComponentTemplate->RelativeLocation = ComponentToWorld.GetTranslation();
-						SceneComponentTemplate->RelativeRotation = ComponentToWorld.Rotator();
-						SceneComponentTemplate->RelativeScale3D = ComponentToWorld.GetScale3D();
+						// Reset the attach socket name
+						SceneComponentTemplate->AttachSocketName = NAME_None;
+						USCS_Node* SCS_Node = DroppedNodePtr->GetSCSNode();
+						if(SCS_Node)
+						{
+							SCS_Node->Modify();
+							SCS_Node->AttachToName = NAME_None;
+						}
+
+						// Attempt to locate a matching registered instance of the component template in the Actor context that's being edited
+						USceneComponent* InstancedSceneComponent = Cast<USceneComponent>(DroppedNodePtr->FindComponentInstanceInActor(PreviewActor));
+						if(InstancedSceneComponent && InstancedSceneComponent->IsRegistered())
+						{
+							// If we find a match, save off the world position
+							FTransform ComponentToWorld = InstancedSceneComponent->GetComponentToWorld();
+							SceneComponentTemplate->RelativeLocation = ComponentToWorld.GetTranslation();
+							SceneComponentTemplate->RelativeRotation = ComponentToWorld.Rotator();
+							SceneComponentTemplate->RelativeScale3D = ComponentToWorld.GetScale3D();
+						}
 					}
 				}
-			}
 
-			// Attach the dropped node to the given node
-			NodePtr->AddChild(DroppedNodePtr);
-
-			// Attempt to locate a matching instance of the parent component template in the Actor context that's being edited
-			USceneComponent* ParentSceneComponent = Cast<USceneComponent>(NodePtr->FindComponentInstanceInActor(Actor));
-			if(SceneComponentTemplate && ParentSceneComponent && ParentSceneComponent->IsRegistered())
-			{
-				// If we find a match, calculate its new position relative to the scene root component instance in its current scene
-				FTransform ComponentToWorld(SceneComponentTemplate->RelativeRotation, SceneComponentTemplate->RelativeLocation, SceneComponentTemplate->RelativeScale3D);
-				FTransform ParentToWorld = ParentSceneComponent->GetSocketTransform(SceneComponentTemplate->AttachSocketName);
-				FTransform RelativeTM = ComponentToWorld.GetRelativeTransform(ParentToWorld);
+				// Attach the dropped node to the given node
+				NodePtr->AddChild(DroppedNodePtr);
 
-				// Store new relative location value (if not set to absolute)
-				if(!SceneComponentTemplate->bAbsoluteLocation)
+				// Attempt to locate a matching instance of the parent component template in the Actor context that's being edited
+				USceneComponent* ParentSceneComponent = Cast<USceneComponent>(NodePtr->FindComponentInstanceInActor(PreviewActor));
+				if(SceneComponentTemplate && ParentSceneComponent && ParentSceneComponent->IsRegistered())
 				{
-					SceneComponentTemplate->RelativeLocation = RelativeTM.GetTranslation();
-				}
+					// If we find a match, calculate its new position relative to the scene root component instance in its current scene
+					FTransform ComponentToWorld(SceneComponentTemplate->RelativeRotation, SceneComponentTemplate->RelativeLocation, SceneComponentTemplate->RelativeScale3D);
+					FTransform ParentToWorld = ParentSceneComponent->GetSocketTransform(SceneComponentTemplate->AttachSocketName);
+					FTransform RelativeTM = ComponentToWorld.GetRelativeTransform(ParentToWorld);
 
-				// Store new relative rotation value (if not set to absolute)
-				if(!SceneComponentTemplate->bAbsoluteRotation)
-				{
-					SceneComponentTemplate->RelativeRotation = RelativeTM.Rotator();
-				}
+					// Store new relative location value (if not set to absolute)
+					if(!SceneComponentTemplate->bAbsoluteLocation)
+					{
+						SceneComponentTemplate->RelativeLocation = RelativeTM.GetTranslation();
+					}
 
-				// Store new relative scale value (if not set to absolute)
-				if(!SceneComponentTemplate->bAbsoluteScale)
-				{
-					SceneComponentTemplate->RelativeScale3D = RelativeTM.GetScale3D();
+					// Store new relative rotation value (if not set to absolute)
+					if(!SceneComponentTemplate->bAbsoluteRotation)
+					{
+						SceneComponentTemplate->RelativeRotation = RelativeTM.Rotator();
+					}
+
+					// Store new relative scale value (if not set to absolute)
+					if(!SceneComponentTemplate->bAbsoluteScale)
+					{
+						SceneComponentTemplate->RelativeScale3D = RelativeTM.GetScale3D();
+					}
 				}
 			}
 		}
 	}
+	else    // SSCSEditor::EEditorMode::ActorInstance
+	{
+		for(const auto& DroppedNodePtr : DroppedNodePtrs)
+		{
+			// Check for a valid parent node
+			FSCSEditorTreeNodePtrType ParentNodePtr = DroppedNodePtr->GetParent();
+			if(ParentNodePtr.IsValid())
+			{
+				// Detach the dropped node from its parent
+				ParentNodePtr->RemoveChild(DroppedNodePtr);
+			}
+
+			// Attach the dropped node to the given node
+			NodePtr->AddChild(DroppedNodePtr);
+		}
+	}
 
 	check(SCSEditorPtr->SCSTreeWidget.IsValid());
 	SCSEditorPtr->SCSTreeWidget->SetItemExpansion(NodePtr, true);
@@ -1613,81 +1707,94 @@ void SSCS_RowWidget::OnDetachFromDropAction(const TArray<FSCSEditorTreeNodePtrTy
 	TSharedPtr<SSCSEditor> SCSEditorPtr = SCSEditor.Pin();
 	check(SCSEditorPtr.IsValid());
 
-	// Get the current "instanced" Actor context
-	AActor* Actor = SCSEditorPtr->PreviewActor.Get();
-	if(Actor == nullptr)
-	{
-		Actor = SCSEditorPtr->ActorContext.Get();
-		check(Actor != nullptr);
-	}
-
 	const FScopedTransaction TransactionContext(DroppedNodePtrs.Num() > 1 ? LOCTEXT("DetachComponents", "Detach Components") : LOCTEXT("DetachComponent", "Detach Component"));
 
-	for(const auto& DroppedNodePtr : DroppedNodePtrs)
+	if(SCSEditorPtr->EditorMode.Get() == SSCSEditor::EEditorMode::BlueprintSCS)
 	{
-		check(DroppedNodePtr.IsValid());
+		// Get the current "preview" Actor instance
+		AActor* PreviewActor = SCSEditorPtr->PreviewActor.Get();
+		check(PreviewActor);
 
-		// Detach the node from its parent
-		NodePtr->RemoveChild(DroppedNodePtr);
-
-		// If the associated component template is a scene component, maintain its current world position
-		USceneComponent* SceneComponentTemplate = Cast<USceneComponent>(DroppedNodePtr->GetComponentTemplate());
-		if(SceneComponentTemplate)
+		for(const auto& DroppedNodePtr : DroppedNodePtrs)
 		{
-			// Save current state
-			SceneComponentTemplate->Modify();
+			check(DroppedNodePtr.IsValid());
 
-			// Reset the attach socket name
-			SceneComponentTemplate->AttachSocketName = NAME_None;
-			USCS_Node* SCS_Node = DroppedNodePtr->GetSCSNode();
-			if(SCS_Node)
+			// Detach the node from its parent
+			NodePtr->RemoveChild(DroppedNodePtr);
+
+			// If the associated component template is a scene component, maintain its current world position
+			USceneComponent* SceneComponentTemplate = Cast<USceneComponent>(DroppedNodePtr->GetComponentTemplate());
+			if(SceneComponentTemplate)
 			{
-				SCS_Node->Modify();
-				SCS_Node->AttachToName = NAME_None;
+				// Save current state
+				SceneComponentTemplate->Modify();
+
+				// Reset the attach socket name
+				SceneComponentTemplate->AttachSocketName = NAME_None;
+				USCS_Node* SCS_Node = DroppedNodePtr->GetSCSNode();
+				if(SCS_Node)
+				{
+					SCS_Node->Modify();
+					SCS_Node->AttachToName = NAME_None;
+				}
+
+				// Attempt to locate a matching instance of the component template in the Actor context that's being edited
+				USceneComponent* InstancedSceneComponent = Cast<USceneComponent>(DroppedNodePtr->FindComponentInstanceInActor(PreviewActor));
+				if(InstancedSceneComponent && InstancedSceneComponent->IsRegistered())
+				{
+					// If we find a match, save off the world position
+					FTransform ComponentToWorld = InstancedSceneComponent->GetComponentToWorld();
+					SceneComponentTemplate->RelativeLocation = ComponentToWorld.GetTranslation();
+					SceneComponentTemplate->RelativeRotation = ComponentToWorld.Rotator();
+					SceneComponentTemplate->RelativeScale3D = ComponentToWorld.GetScale3D();
+				}
 			}
 
-			// Attempt to locate a matching instance of the component template in the Actor context that's being edited
-			USceneComponent* InstancedSceneComponent = Cast<USceneComponent>(DroppedNodePtr->FindComponentInstanceInActor(Actor));
-			if(InstancedSceneComponent && InstancedSceneComponent->IsRegistered())
+			// Attach the dropped node to the current scene root node
+			check(SCSEditorPtr->SceneRootNodePtr.IsValid());
+			SCSEditorPtr->SceneRootNodePtr->AddChild(DroppedNodePtr);
+
+			// Attempt to locate a matching instance of the scene root component template in the Actor context that's being edited
+			USceneComponent* InstancedSceneRootComponent = Cast<USceneComponent>(SCSEditorPtr->SceneRootNodePtr->FindComponentInstanceInActor(PreviewActor));
+			if(SceneComponentTemplate && InstancedSceneRootComponent && InstancedSceneRootComponent->IsRegistered())
 			{
-				// If we find a match, save off the world position
-				FTransform ComponentToWorld = InstancedSceneComponent->GetComponentToWorld();
-				SceneComponentTemplate->RelativeLocation = ComponentToWorld.GetTranslation();
-				SceneComponentTemplate->RelativeRotation = ComponentToWorld.Rotator();
-				SceneComponentTemplate->RelativeScale3D = ComponentToWorld.GetScale3D();
-			}
-		}
+				// If we find a match, calculate its new position relative to the scene root component instance in the preview scene
+				FTransform ComponentToWorld(SceneComponentTemplate->RelativeRotation, SceneComponentTemplate->RelativeLocation, SceneComponentTemplate->RelativeScale3D);
+				FTransform ParentToWorld = InstancedSceneRootComponent->GetSocketTransform(SceneComponentTemplate->AttachSocketName);
+				FTransform RelativeTM = ComponentToWorld.GetRelativeTransform(ParentToWorld);
 
-		// Attach the dropped node to the current scene root node
-		check(SCSEditorPtr->SceneRootNodePtr.IsValid());
-		SCSEditorPtr->SceneRootNodePtr->AddChild(DroppedNodePtr);
+				// Store new relative location value (if not set to absolute)
+				if(!SceneComponentTemplate->bAbsoluteLocation)
+				{
+					SceneComponentTemplate->RelativeLocation = RelativeTM.GetTranslation();
+				}
 
-		// Attempt to locate a matching instance of the scene root component template in the Actor context that's being edited
-		USceneComponent* InstancedSceneRootComponent = Cast<USceneComponent>(SCSEditorPtr->SceneRootNodePtr->FindComponentInstanceInActor(Actor));
-		if(SceneComponentTemplate && InstancedSceneRootComponent && InstancedSceneRootComponent->IsRegistered())
-		{
-			// If we find a match, calculate its new position relative to the scene root component instance in the preview scene
-			FTransform ComponentToWorld(SceneComponentTemplate->RelativeRotation, SceneComponentTemplate->RelativeLocation, SceneComponentTemplate->RelativeScale3D);
-			FTransform ParentToWorld = InstancedSceneRootComponent->GetSocketTransform(SceneComponentTemplate->AttachSocketName);
-			FTransform RelativeTM = ComponentToWorld.GetRelativeTransform(ParentToWorld);
+				// Store new relative rotation value (if not set to absolute)
+				if(!SceneComponentTemplate->bAbsoluteRotation)
+				{
+					SceneComponentTemplate->RelativeRotation = RelativeTM.Rotator();
+				}
 
-			// Store new relative location value (if not set to absolute)
-			if(!SceneComponentTemplate->bAbsoluteLocation)
-			{
-				SceneComponentTemplate->RelativeLocation = RelativeTM.GetTranslation();
+				// Store new relative scale value (if not set to absolute)
+				if(!SceneComponentTemplate->bAbsoluteScale)
+				{
+					SceneComponentTemplate->RelativeScale3D = RelativeTM.GetScale3D();
+				}
 			}
+		}
+	}
+	else    // SSCSEditor::EEditorMode::ActorInstance
+	{
+		for(const auto& DroppedNodePtr : DroppedNodePtrs)
+		{
+			check(DroppedNodePtr.IsValid());
 
-			// Store new relative rotation value (if not set to absolute)
-			if(!SceneComponentTemplate->bAbsoluteRotation)
-			{
-				SceneComponentTemplate->RelativeRotation = RelativeTM.Rotator();
-			}
+			// Detach the node from its parent
+			NodePtr->RemoveChild(DroppedNodePtr);
 
-			// Store new relative scale value (if not set to absolute)
-			if(!SceneComponentTemplate->bAbsoluteScale)
-			{
-				SceneComponentTemplate->RelativeScale3D = RelativeTM.GetScale3D();
-			}
+			// Attach the dropped node to the current scene root node
+			check(SCSEditorPtr->SceneRootNodePtr.IsValid());
+			SCSEditorPtr->SceneRootNodePtr->AddChild(DroppedNodePtr);
 		}
 	}
 	
@@ -1705,89 +1812,96 @@ void SSCS_RowWidget::OnMakeNewRootDropAction(FSCSEditorTreeNodePtrType DroppedNo
 	check(NodePtr.IsValid() && NodePtr == SceneRootNodePtr);
 	check(DroppedNodePtr.IsValid());
 
-	// Get the current Blueprint context
-	UBlueprint* Blueprint = GetBlueprint();
-	check(Blueprint != NULL && Blueprint->SimpleConstructionScript != nullptr);
-
 	// Create a transaction record
 	const FScopedTransaction TransactionContext(LOCTEXT("MakeNewSceneRoot", "Make New Scene Root"));
 
-	// Remember whether or not we're replacing the default scene root
-	bool bWasDefaultSceneRoot = SceneRootNodePtr.IsValid() && SceneRootNodePtr->IsDefaultSceneRoot();
-
-	// Clone the component if it's being dropped into a different SCS
-	if(DroppedNodePtr->GetBlueprint() != Blueprint)
+	if(SCSEditorPtr->EditorMode.Get() == SSCSEditor::EEditorMode::BlueprintSCS)
 	{
-		UActorComponent* ComponentTemplate = DroppedNodePtr->GetComponentTemplate();
-		check(ComponentTemplate);
+		// Get the current Blueprint context
+		UBlueprint* Blueprint = GetBlueprint();
+		check(Blueprint != NULL && Blueprint->SimpleConstructionScript != nullptr);
 
-		// Note: This will mark the Blueprint as structurally modified
-		UActorComponent* ClonedComponent = SCSEditorPtr->AddNewComponent(ComponentTemplate->GetClass(), NULL);
-		check(ClonedComponent);
+		// Remember whether or not we're replacing the default scene root
+		bool bWasDefaultSceneRoot = SceneRootNodePtr.IsValid() && SceneRootNodePtr->IsDefaultSceneRoot();
 
-		//Serialize object properties using write/read operations.
-		TArray<uint8> SavedProperties;
-		FObjectWriter Writer(ComponentTemplate, SavedProperties);
-		FObjectReader(ClonedComponent, SavedProperties);
+		// Clone the component if it's being dropped into a different SCS
+		if(DroppedNodePtr->GetBlueprint() != Blueprint)
+		{
+			UActorComponent* ComponentTemplate = DroppedNodePtr->GetComponentTemplate();
+			check(ComponentTemplate);
 
-		DroppedNodePtr = SCSEditorPtr->GetNodeFromActorComponent(ClonedComponent);
-		check(DroppedNodePtr.IsValid());
-	}
+			// Note: This will mark the Blueprint as structurally modified
+			UActorComponent* ClonedComponent = SCSEditorPtr->AddNewComponent(ComponentTemplate->GetClass(), NULL);
+			check(ClonedComponent);
 
-	if(DroppedNodePtr->GetParent().IsValid()
-		&& DroppedNodePtr->GetBlueprint() == Blueprint)
-	{
-		// Remove the dropped node from its existing parent
-		DroppedNodePtr->GetParent()->RemoveChild(DroppedNodePtr);
+			//Serialize object properties using write/read operations.
+			TArray<uint8> SavedProperties;
+			FObjectWriter Writer(ComponentTemplate, SavedProperties);
+			FObjectReader(ClonedComponent, SavedProperties);
+
+			DroppedNodePtr = SCSEditorPtr->GetNodeFromActorComponent(ClonedComponent);
+			check(DroppedNodePtr.IsValid());
+		}
 
-		// If the associated component template is a scene component, reset its transform since it will now become the root
-		USceneComponent* SceneComponentTemplate = Cast<USceneComponent>(DroppedNodePtr->GetComponentTemplate());
-		if(SceneComponentTemplate)
+		if(DroppedNodePtr->GetParent().IsValid()
+			&& DroppedNodePtr->GetBlueprint() == Blueprint)
 		{
-			// Save current state
-			SceneComponentTemplate->Modify();
+			// Remove the dropped node from its existing parent
+			DroppedNodePtr->GetParent()->RemoveChild(DroppedNodePtr);
 
-			// Reset the attach socket name
-			SceneComponentTemplate->AttachSocketName = NAME_None;
-			USCS_Node* SCS_Node = DroppedNodePtr->GetSCSNode();
-			if(SCS_Node)
+			// If the associated component template is a scene component, reset its transform since it will now become the root
+			USceneComponent* SceneComponentTemplate = Cast<USceneComponent>(DroppedNodePtr->GetComponentTemplate());
+			if(SceneComponentTemplate)
 			{
-				SCS_Node->Modify();
-				SCS_Node->AttachToName = NAME_None;
-			}
+				// Save current state
+				SceneComponentTemplate->Modify();
+
+				// Reset the attach socket name
+				SceneComponentTemplate->AttachSocketName = NAME_None;
+				USCS_Node* SCS_Node = DroppedNodePtr->GetSCSNode();
+				if(SCS_Node)
+				{
+					SCS_Node->Modify();
+					SCS_Node->AttachToName = NAME_None;
+				}
 
-			// Reset the relative transform
-			SceneComponentTemplate->SetRelativeLocation(FVector::ZeroVector);
-			SceneComponentTemplate->SetRelativeRotation(FRotator::ZeroRotator);
-			SceneComponentTemplate->SetRelativeScale3D(FVector(1.f));
+				// Reset the relative transform
+				SceneComponentTemplate->SetRelativeLocation(FVector::ZeroVector);
+				SceneComponentTemplate->SetRelativeRotation(FRotator::ZeroRotator);
+				SceneComponentTemplate->SetRelativeScale3D(FVector(1.f));
+			}
 		}
-	}
 
-	if(!bWasDefaultSceneRoot)
-	{
-		check(SceneRootNodePtr->CanReparent());
+		if(!bWasDefaultSceneRoot)
+		{
+			check(SceneRootNodePtr->CanReparent());
 
-		// Remove the current scene root node from the SCS context
-		Blueprint->SimpleConstructionScript->RemoveNode(SceneRootNodePtr->GetSCSNode());
-	}
+			// Remove the current scene root node from the SCS context
+			Blueprint->SimpleConstructionScript->RemoveNode(SceneRootNodePtr->GetSCSNode());
+		}
 
-	// Save old root node
-	FSCSEditorTreeNodePtrType OldSceneRootNodePtr;
-	if(!bWasDefaultSceneRoot)
-	{
-		OldSceneRootNodePtr = SceneRootNodePtr;
-	}
+		// Save old root node
+		FSCSEditorTreeNodePtrType OldSceneRootNodePtr;
+		if(!bWasDefaultSceneRoot)
+		{
+			OldSceneRootNodePtr = SceneRootNodePtr;
+		}
 
-	// Set node we are dropping as new root
-	SceneRootNodePtr = DroppedNodePtr;
+		// Set node we are dropping as new root
+		SceneRootNodePtr = DroppedNodePtr;
 
-	// Add dropped node to the SCS context
-	Blueprint->SimpleConstructionScript->AddNode(SceneRootNodePtr->GetSCSNode());
+		// Add dropped node to the SCS context
+		Blueprint->SimpleConstructionScript->AddNode(SceneRootNodePtr->GetSCSNode());
 
-	// Set old root as child of new root
-	if(OldSceneRootNodePtr.IsValid())
+		// Set old root as child of new root
+		if(OldSceneRootNodePtr.IsValid())
+		{
+			SceneRootNodePtr->AddChild(OldSceneRootNodePtr);
+		}
+	}
+	else    // SSCSEditor::EEditorMode::ActorInstance
 	{
-		SceneRootNodePtr->AddChild(OldSceneRootNodePtr);
+
 	}
 
 	PostDragDropAction(true);
@@ -1802,12 +1916,15 @@ void SSCS_RowWidget::PostDragDropAction(bool bRegenerateTreeNodes)
 
 		PinnedEditor->RefreshSelectionDetails();
 
-		if(NodePtr.IsValid())
+		if(PinnedEditor->EditorMode.Get() == SSCSEditor::EEditorMode::BlueprintSCS)
 		{
-			UBlueprint* Blueprint = GetBlueprint();
-			if(Blueprint != nullptr)
+			if(NodePtr.IsValid())
 			{
-				FBlueprintEditorUtils::PostEditChangeBlueprintActors(Blueprint);
+				UBlueprint* Blueprint = GetBlueprint();
+				if(Blueprint != nullptr)
+				{
+					FBlueprintEditorUtils::PostEditChangeBlueprintActors(Blueprint);
+				}
 			}
 		}
 	}
@@ -1840,12 +1957,24 @@ FText SSCS_RowWidget::GetTooltipText() const
 
 		if ( NodePtr->IsNative() )
 		{
-			if(NodePtr->GetComponentTemplate() != NULL)
+			if(NodePtr->IsInstanced())
 			{
-				return FText::Format( LOCTEXT("NativeClassToolTip", "Native {ClassName}"), Args );
+				if ( NodePtr->GetComponentTemplate() != NULL )
+				{
+					return FText::Format( LOCTEXT("RegularToolTip", "{ClassName}"), Args );
+				}
+
+				return FText::Format( LOCTEXT("MissingRegularComponentToolTip", "MISSING!! {NodeName}"), Args );
 			}
+			else
+			{
+				if(NodePtr->GetComponentTemplate() != NULL)
+				{
+					return FText::Format( LOCTEXT("NativeClassToolTip", "Native {ClassName}"), Args );
+				}
 
-			return FText::Format( LOCTEXT("MissingNativeComponentToolTip", "MISSING!! Native {NodeName}"), Args );
+				return FText::Format( LOCTEXT("MissingNativeComponentToolTip", "MISSING!! Native {NodeName}"), Args );
+			}
 		}
 		else
 		{
@@ -1875,10 +2004,17 @@ FString SSCS_RowWidget::GetDocumentationLink() const
 {
 	check(SCSEditor.IsValid());
 
-	if( NodePtr == SCSEditor.Pin()->SceneRootNodePtr
-		|| NodePtr->IsNative() || NodePtr->IsInherited())
+	if(SCSEditor.Pin()->EditorMode.Get() == SSCSEditor::EEditorMode::BlueprintSCS)
+	{
+		if( NodePtr == SCSEditor.Pin()->SceneRootNodePtr
+			|| NodePtr->IsNative() || NodePtr->IsInherited())
+		{
+			return TEXT("Shared/Editors/BlueprintEditor/ComponentsMode");
+		}
+	}
+	else    // SSCSEditor::EEditorMode::ActorInstance
 	{
-		return TEXT("Shared/Editors/BlueprintEditor/ComponentsMode");
+		// @TODO - Actor instance mode
 	}
 
 	return TEXT("");
@@ -1888,17 +2024,24 @@ FString SSCS_RowWidget::GetDocumentationExcerptName() const
 {
 	check(SCSEditor.IsValid());
 
-	if( NodePtr == SCSEditor.Pin()->SceneRootNodePtr)
-	{
-		return TEXT("RootComponent");
-	}
-	else if(NodePtr->IsNative())
+	if(SCSEditor.Pin()->EditorMode.Get() == SSCSEditor::EEditorMode::BlueprintSCS)
 	{
-		return TEXT("NativeComponents");
+		if( NodePtr == SCSEditor.Pin()->SceneRootNodePtr)
+		{
+			return TEXT("RootComponent");
+		}
+		else if(NodePtr->IsNative())
+		{
+			return TEXT("NativeComponents");
+		}
+		else if(NodePtr->IsInherited())
+		{
+			return TEXT("InheritedComponents");
+		}
 	}
-	else if(NodePtr->IsInherited())
+	else    // SSCSEditor::EEditorMode::ActorInstance
 	{
-		return TEXT("InheritedComponents");
+		// @TODO - Actor instance mode
 	}
 
 	return TEXT("");
@@ -2681,7 +2824,7 @@ void SSCSEditor::UpdateTree(bool bRegenerateTreeNodes)
 
 						if(SCS_Node->ParentComponentOrVariableName != NAME_None)
 						{
-							USceneComponent* ParentComponent = SCS_Node->GetParentComponentTemplate();
+							USceneComponent* ParentComponent = SCS_Node->GetParentComponentTemplate(ParentBPStack[0]);
 							if(ParentComponent != NULL)
 							{
 								FSCSEditorTreeNodePtrType ParentNodePtr = FindTreeNode(ParentComponent);
@@ -3214,39 +3357,39 @@ void SSCSEditor::OnDeleteNodes()
 {
 	const FScopedTransaction Transaction( LOCTEXT("RemoveComponent", "Remove Component") );
 
-	// Remove node from SCS
-	UBlueprint* Blueprint = nullptr;
-	FThumbnailRenderingInfo* RenderInfo = nullptr;
-	TArray<FSCSEditorTreeNodePtrType> SelectedNodes = SCSTreeWidget->GetSelectedItems();
-	for (int32 i = 0; i < SelectedNodes.Num(); ++i)
+	if(EditorMode.Get() == SSCSEditor::EEditorMode::BlueprintSCS)
 	{
-		auto Node = SelectedNodes[i];
-
-		USCS_Node* SCS_Node = Node->GetSCSNode();
-		if(SCS_Node != nullptr)
+		// Remove node from SCS
+		UBlueprint* Blueprint = GetBlueprint();
+		FThumbnailRenderingInfo* RenderInfo = nullptr;
+		TArray<FSCSEditorTreeNodePtrType> SelectedNodes = SCSTreeWidget->GetSelectedItems();
+		for (int32 i = 0; i < SelectedNodes.Num(); ++i)
 		{
-			USimpleConstructionScript* SCS = SCS_Node->GetSCS();
-			check(SCS != nullptr && (Blueprint == nullptr || Blueprint == SCS->GetBlueprint()));
+			auto Node = SelectedNodes[i];
 
-			if(Blueprint == nullptr)
+			USCS_Node* SCS_Node = Node->GetSCSNode();
+			if(SCS_Node != nullptr)
 			{
-				Blueprint = SCS->GetBlueprint();
-				check(Blueprint != nullptr);
+				USimpleConstructionScript* SCS = SCS_Node->GetSCS();
+				check(SCS != nullptr && Blueprint == SCS->GetBlueprint());
 
-				// Get the current render info for the blueprint. If this is NULL then the blueprint is not currently visualizable (no visible primitive components)
-				FThumbnailRenderingInfo* RenderInfo = GUnrealEd->GetThumbnailManager()->GetRenderingInfo( Blueprint );
+				if(Blueprint == nullptr)
+				{
+					Blueprint = SCS->GetBlueprint();
+					check(Blueprint != nullptr);
 
-				// Saving objects for restoring purpose.
-				Blueprint->Modify();
-				SaveSCSCurrentState( SCS );
+					// Get the current render info for the blueprint. If this is NULL then the blueprint is not currently visualizable (no visible primitive components)
+					FThumbnailRenderingInfo* RenderInfo = GUnrealEd->GetThumbnailManager()->GetRenderingInfo( Blueprint );
+
+					// Saving objects for restoring purpose.
+					Blueprint->Modify();
+					SaveSCSCurrentState( SCS );
+				}
 			}
-		}
 
-		RemoveComponentNode(Node);
-	}
+			RemoveComponentNode(Node);
+		}
 
-	if(Blueprint != nullptr)
-	{
 		// Will call UpdateTree as part of OnBlueprintChanged handling
 		FBlueprintEditorUtils::MarkBlueprintAsStructurallyModified(Blueprint);
 
@@ -3261,6 +3404,10 @@ void SSCSEditor::OnDeleteNodes()
 			ThumbnailTools::CacheEmptyThumbnail( BPFullName, BPPackage );
 		}
 	}
+	else    // SSCSEditor::EEditorMode::ActorInstance
+	{
+		// @TODO - Actor instance mode
+	}
 
 	// Do this AFTER marking the Blueprint as modified
 	UpdateSelectionFromNodes(SCSTreeWidget->GetSelectedItems());
@@ -3276,24 +3423,31 @@ void SSCSEditor::RemoveComponentNode(FSCSEditorTreeNodePtrType InNodePtr)
 		SCSTreeWidget->ClearSelection();
 	}
 
-	USCS_Node* SCS_Node = InNodePtr->GetSCSNode();
-	if(SCS_Node != NULL)
+	if(EditorMode.Get() == SSCSEditor::EEditorMode::BlueprintSCS)
 	{
-		USimpleConstructionScript* SCS = SCS_Node->GetSCS();
-		check(SCS != nullptr);
-
-		// Remove any instances of variable accessors from the blueprint graphs
-		UBlueprint* Blueprint = SCS->GetBlueprint();
-		if(Blueprint != nullptr)
+		USCS_Node* SCS_Node = InNodePtr->GetSCSNode();
+		if(SCS_Node != NULL)
 		{
-			FBlueprintEditorUtils::RemoveVariableNodes(Blueprint, InNodePtr->GetVariableName());
-		}
+			USimpleConstructionScript* SCS = SCS_Node->GetSCS();
+			check(SCS != nullptr);
 
-		// Remove node from SCS tree
-		SCS->RemoveNodeAndPromoteChildren(SCS_Node);
+			// Remove any instances of variable accessors from the blueprint graphs
+			UBlueprint* Blueprint = SCS->GetBlueprint();
+			if(Blueprint != nullptr)
+			{
+				FBlueprintEditorUtils::RemoveVariableNodes(Blueprint, InNodePtr->GetVariableName());
+			}
+
+			// Remove node from SCS tree
+			SCS->RemoveNodeAndPromoteChildren(SCS_Node);
 
-		// Clear the delegate
-		SCS_Node->SetOnNameChanged(FSCSNodeNameChanged());
+			// Clear the delegate
+			SCS_Node->SetOnNameChanged(FSCSNodeNameChanged());
+		}
+	}
+	else    // SSCSEditor::EEditorMode::ActorInstance
+	{
+		// @TODO - Actor instance mode
 	}
 }
 
diff --git a/Engine/Source/Editor/Kismet/Public/SSCSEditor.h b/Engine/Source/Editor/Kismet/Public/SSCSEditor.h
--- a/Engine/Source/Editor/Kismet/Public/SSCSEditor.h
+++ b/Engine/Source/Editor/Kismet/Public/SSCSEditor.h
@@ -91,9 +91,17 @@ public:
 	 */
 	bool IsNative() const { return GetSCSNode() == NULL && GetComponentTemplate() != NULL; }
 	/** 
-	 * @return Whether or not this object represents a component template that was inherited from a parent Blueprint class.
+	 * @return Whether or not this object represents a component template that was inherited from a parent Blueprint class SCS.
 	 */
 	bool IsInherited() const { return bIsInherited; }
+	/**
+	 * @return Whether or not this object represents a component instance rather than a template.
+	 */
+	bool IsInstanced() const { return bIsInstanced; }
+	/**
+	 * @return Whether or not this object represents a component instance that was created by the user and not by a native or Blueprint-generated class.
+	 */
+	bool IsUserInstanced() const { return bIsInstanced && !bWasInstancedFromNativeClass && !GetComponentTemplate()->bCreatedByConstructionScript; }
 	/** 
 	 * @return Whether or not this object represents the default SCS scene root component.
 	 */
@@ -101,7 +109,7 @@ public:
 	/** 
 	 * @return Whether or not this object represents a node that can be deleted from the SCS tree.
 	 */
-	bool CanDelete() const { return !IsNative() && !IsInherited() && !IsDefaultSceneRoot(); }
+	bool CanDelete() const { return IsUserInstanced() || (!IsNative() && !IsInherited() && !IsDefaultSceneRoot()); }
 	/** 
 	 * @return Whether or not this object represents a node that can be renamed from the SCS tree.
 	 */
@@ -109,7 +117,7 @@ public:
 	/** 
 	 * @return Whether or not this object represents a node that can be reparented to other nodes based on its context.
 	 */
-	bool CanReparent() const { return !IsNative() && !IsInherited() && !IsDefaultSceneRoot() && Cast<USceneComponent>(GetComponentTemplate()) != NULL; }
+	bool CanReparent() const { return (IsUserInstanced() || (!IsNative() && !IsInherited() && !IsDefaultSceneRoot())) && Cast<USceneComponent>(GetComponentTemplate()) != NULL; }
 	/** 
 	 * @return Whether or not we can edit default properties for the component template represented by this object.
 	 */
@@ -191,6 +199,7 @@ public:
 
 private:
 	bool bIsInherited;
+	bool bIsInstanced;
 	TWeakObjectPtr<class USCS_Node> SCSNodePtr;
 	TWeakObjectPtr<UActorComponent> ComponentTemplatePtr;
 	FSCSEditorTreeNodePtrType ParentNodePtr;
@@ -199,6 +208,8 @@ private:
 	/** Handles rename requests */
 	bool bNonTransactionalRename;
 	FOnRenameRequested RenameRequestedDelegate;
+
+	bool bWasInstancedFromNativeClass;
 };
 
 //////////////////////////////////////////////////////////////////////////
diff --git a/Engine/Source/Editor/UnrealEd/Private/Kismet2/ComponentEditorUtils.cpp b/Engine/Source/Editor/UnrealEd/Private/Kismet2/ComponentEditorUtils.cpp
--- a/Engine/Source/Editor/UnrealEd/Private/Kismet2/ComponentEditorUtils.cpp
+++ b/Engine/Source/Editor/UnrealEd/Private/Kismet2/ComponentEditorUtils.cpp
@@ -41,7 +41,7 @@ USceneComponent* FComponentEditorUtils::GetSceneComponent( UObject* Object, UObj
 							USCS_Node* SCS_Node = SCSNodes[ SCSNodeIndex ];
 							if( SCS_Node && SCS_Node->ComponentTemplate == SubObject )
 							{
-								return SCS_Node->GetParentComponentTemplate();
+								return SCS_Node->GetParentComponentTemplate(Blueprint);
 							}
 						}
 					}
diff --git a/Engine/Source/Editor/UnrealEd/Private/ThumbnailHelpers.cpp b/Engine/Source/Editor/UnrealEd/Private/ThumbnailHelpers.cpp
--- a/Engine/Source/Editor/UnrealEd/Private/ThumbnailHelpers.cpp
+++ b/Engine/Source/Editor/UnrealEd/Private/ThumbnailHelpers.cpp
@@ -1178,7 +1178,7 @@ TArray<UPrimitiveComponent*> FBlueprintThumbnailScene::GetPooledVisualizableComp
 							// Check to see if the root node has set an explicit parent
 							if(RootNode->ParentComponentOrVariableName != NAME_None)
 							{
-								USceneComponent* ParentComponentTemplate = RootNode->GetParentComponentTemplate();
+								USceneComponent* ParentComponentTemplate = RootNode->GetParentComponentTemplate(Blueprint);
 								if(ParentComponentTemplate != NULL)
 								{
 									if(NativeInstanceMap.Contains(ParentComponentTemplate))
diff --git a/Engine/Source/Runtime/Engine/Classes/Engine/SCS_Node.h b/Engine/Source/Runtime/Engine/Classes/Engine/SCS_Node.h
--- a/Engine/Source/Runtime/Engine/Classes/Engine/SCS_Node.h
+++ b/Engine/Source/Runtime/Engine/Classes/Engine/SCS_Node.h
@@ -150,8 +150,8 @@ class USCS_Node : public UObject
 	/** Sets parent component attributes based on the given component instance */
 	ENGINE_API void SetParent(USceneComponent* InParentComponent);
 
-	/** Finds and returns the parent component template */
-	ENGINE_API USceneComponent* GetParentComponentTemplate() const;
+	/** Finds and returns the parent component template through the given Blueprint */
+	ENGINE_API USceneComponent* GetParentComponentTemplate(UBlueprint* InBlueprint) const;
 
 	/** Test whether or not the given string is a valid variable name string */
 	ENGINE_API bool IsValidVariableNameString(const FString& InString);
diff --git a/Engine/Source/Runtime/Engine/Private/SCS_Node.cpp b/Engine/Source/Runtime/Engine/Private/SCS_Node.cpp
--- a/Engine/Source/Runtime/Engine/Private/SCS_Node.cpp
+++ b/Engine/Source/Runtime/Engine/Private/SCS_Node.cpp
@@ -317,22 +317,18 @@ void USCS_Node::SetParent(USceneComponent* InParentComponent)
 	}
 }
 
-USceneComponent* USCS_Node::GetParentComponentTemplate() const
+USceneComponent* USCS_Node::GetParentComponentTemplate(UBlueprint* InBlueprint) const
 {
 	USceneComponent* ParentComponentTemplate = NULL;
 	if(ParentComponentOrVariableName != NAME_None)
 	{
-		USimpleConstructionScript* SCS = GetSCS();
-		check(SCS != nullptr);
-
-		UBlueprint* Blueprint = SCS->GetBlueprint();
-		check(Blueprint != nullptr && Blueprint->GeneratedClass != nullptr);
+		check(InBlueprint != NULL && InBlueprint->GeneratedClass != NULL);
 
 		// If the parent component template is found in the 'Components' array of the CDO (i.e. native)
 		if(bIsParentComponentNative)
 		{
 			// Access the Blueprint CDO
-			AActor* CDO = Blueprint->GeneratedClass->GetDefaultObject<AActor>();
+			AActor* CDO = InBlueprint->GeneratedClass->GetDefaultObject<AActor>();
 			if(CDO != NULL)
 			{
 				// Find the component template in the CDO that matches the specified name
@@ -356,7 +352,7 @@ USceneComponent* USCS_Node::GetParentComponentTemplate() const
 		{
 			// Get the Blueprint hierarchy
 			TArray<UBlueprint*> ParentBPStack;
-			UBlueprint::GetBlueprintHierarchyFromClass(Blueprint->GeneratedClass, ParentBPStack);
+			UBlueprint::GetBlueprintHierarchyFromClass(InBlueprint->GeneratedClass, ParentBPStack);
 
 			// Find the parent Blueprint in the hierarchy
 			for(int32 StackIndex = ParentBPStack.Num() - 1; StackIndex > 0; --StackIndex)
