Fix IWCE crashes related to stale component references in the SCSEditor tree view node for Blueprint-instanced component nodes.

diff --git a/Engine/Source/Editor/Kismet/Private/SSCSEditor.cpp b/Engine/Source/Editor/Kismet/Private/SSCSEditor.cpp
--- a/Engine/Source/Editor/Kismet/Private/SSCSEditor.cpp
+++ b/Engine/Source/Editor/Kismet/Private/SSCSEditor.cpp
@@ -199,6 +199,7 @@ public:
 
 FSCSEditorTreeNode::FSCSEditorTreeNode()
 	:bIsInherited(false)
+	,bIsInstanced(false)
 	,bNonTransactionalRename(false)
 {
 }
@@ -209,7 +210,6 @@ FSCSEditorTreeNode::FSCSEditorTreeNode(USCS_Node* InSCSNode, bool bInIsInherited
 	,SCSNodePtr(InSCSNode)
 	,ComponentTemplatePtr(InSCSNode != NULL ? InSCSNode->ComponentTemplate : NULL)
 	,bNonTransactionalRename(false)
-	,bWasInstancedFromNativeClass(false)
 {
 }
 
@@ -219,13 +219,17 @@ FSCSEditorTreeNode::FSCSEditorTreeNode(UActorComponent* InComponentTemplate)
 	,SCSNodePtr(NULL)
 	,ComponentTemplatePtr(InComponentTemplate)
 	,bNonTransactionalRename(false)
-	,bWasInstancedFromNativeClass(false)
 {
 	check(InComponentTemplate != nullptr);
 	AActor* Owner = InComponentTemplate->GetOwner();
 	if(Owner != nullptr && !Owner->HasAllFlags(RF_ClassDefaultObject))
 	{
 		bIsInstanced = true;
+		bWasInstancedFromNativeClass = false;
+		InstancedComponentName = InComponentTemplate->GetFName();
+
+		ComponentTemplatePtr.Reset();
+		InstancedComponentOwnerPtr = Owner;
 		
 		UClass* OwnerClass = Owner->GetActorClass();
 		if(OwnerClass != nullptr)
@@ -250,6 +254,26 @@ FSCSEditorTreeNode::FSCSEditorTreeNode(UActorComponent* InComponentTemplate)
 	}
 }
 
+UActorComponent* FSCSEditorTreeNode::GetComponentTemplate() const
+{
+	if(bIsInstanced && InstancedComponentOwnerPtr.IsValid())
+	{
+		TInlineComponentArray<UActorComponent*> Components;
+		InstancedComponentOwnerPtr.Get()->GetComponents(Components);
+
+		for(auto It = Components.CreateConstIterator(); It; ++It)
+		{
+			UActorComponent* ComponentInstance = *It;
+			if(ComponentInstance->GetFName() == InstancedComponentName)
+			{
+				return ComponentInstance;
+			}
+		}
+	}
+	
+	return ComponentTemplatePtr.Get();
+}
+
 FName FSCSEditorTreeNode::GetVariableName() const
 {
 	FName VariableName = NAME_None;
@@ -466,6 +490,17 @@ bool FSCSEditorTreeNode::IsDefaultSceneRoot() const
 	return false;
 }
 
+bool FSCSEditorTreeNode::IsUserInstanced() const
+{
+	if(bIsInstanced && !bWasInstancedFromNativeClass)
+	{
+		UActorComponent* ComponentInstance = GetComponentTemplate();
+		return ComponentInstance != nullptr && !ComponentInstance->bCreatedByConstructionScript;
+	}
+
+	return false;
+}
+
 bool FSCSEditorTreeNode::CanEditDefaults() const
 {
 	bool bCanEdit = false;
diff --git a/Engine/Source/Editor/Kismet/Public/SSCSEditor.h b/Engine/Source/Editor/Kismet/Public/SSCSEditor.h
--- a/Engine/Source/Editor/Kismet/Public/SSCSEditor.h
+++ b/Engine/Source/Editor/Kismet/Public/SSCSEditor.h
@@ -52,9 +52,9 @@ public:
 	 */
 	class USCS_Node* GetSCSNode() const { return SCSNodePtr.Get(); }
 	/**
-	 * @return The component template represented by this object.
+	 * @return The component template or instance represented by this object.
 	 */
-	UActorComponent* GetComponentTemplate() const { return ComponentTemplatePtr.Get(); }
+	UActorComponent* GetComponentTemplate() const;
 	/**
 	 * Finds the component instance represented by this node contained within a given Actor instance.
 	 *
@@ -101,7 +101,7 @@ public:
 	/**
 	 * @return Whether or not this object represents a component instance that was created by the user and not by a native or Blueprint-generated class.
 	 */
-	bool IsUserInstanced() const { return bIsInstanced && !bWasInstancedFromNativeClass && !GetComponentTemplate()->bCreatedByConstructionScript; }
+	bool IsUserInstanced() const;
 	/** 
 	 * @return Whether or not this object represents the default SCS scene root component.
 	 */
@@ -210,6 +210,8 @@ private:
 	FOnRenameRequested RenameRequestedDelegate;
 
 	bool bWasInstancedFromNativeClass;
+	FName InstancedComponentName;
+	TWeakObjectPtr<AActor> InstancedComponentOwnerPtr;
 };
 
 //////////////////////////////////////////////////////////////////////////
