More SSCSEditor work to support shared usage w/ the Level editor.

Change summary:
- Adds support for removing user-instanced component nodes
- Adds support for renaming user-instanced component nodes with some basic name validation

diff --git a/Engine/Source/Editor/Kismet/Private/BlueprintDetailsCustomization.cpp b/Engine/Source/Editor/Kismet/Private/BlueprintDetailsCustomization.cpp
--- a/Engine/Source/Editor/Kismet/Private/BlueprintDetailsCustomization.cpp
+++ b/Engine/Source/Editor/Kismet/Private/BlueprintDetailsCustomization.cpp
@@ -25,6 +25,7 @@
 #include "PropertyCustomizationHelpers.h"
 
 #include "Editor/UnrealEd/Public/Kismet2/BlueprintEditorUtils.h"
+#include "Editor/UnrealEd/Public/Kismet2/ComponentEditorUtils.h"
 
 #include "BlueprintDetailsCustomization.h"
 #include "ObjectEditorUtils.h"
@@ -734,7 +735,7 @@ bool FBlueprintVarActionDetails::GetVariableNameChangeEnabled() const
 		{
 			if (USCS_Node* Node = Blueprint->SimpleConstructionScript->FindSCSNode(GetVariableName()))
 			{
-				bIsReadOnly = !Node->IsValidVariableNameString(Node->VariableName.ToString());
+				bIsReadOnly = !FComponentEditorUtils::IsValidVariableNameString(Node->ComponentTemplate, Node->VariableName.ToString());
 			}
 		}
 		else if(IsALocalVariable(VariableProperty))
@@ -765,7 +766,7 @@ void FBlueprintVarActionDetails::OnVarNameChanged(const FText& InNewText)
 		for (TArray<USCS_Node*>::TConstIterator NodeIt(Nodes); NodeIt; ++NodeIt)
 		{
 			USCS_Node* Node = *NodeIt;
-			if (Node->VariableName == GetVariableName() && !Node->IsValidVariableNameString(InNewText.ToString()))
+			if (Node->VariableName == GetVariableName() && !FComponentEditorUtils::IsValidVariableNameString(Node->ComponentTemplate, InNewText.ToString()))
 			{
 				VarNameEditableTextBox->SetError(LOCTEXT("ComponentVariableRenameFailed_NotValid", "This name is reserved for engine use."));
 				return;
@@ -4499,7 +4500,7 @@ void FBlueprintComponentDetails::OnVariableTextChanged(const FText& InNewText)
 	bIsVariableNameInvalid = true;
 
 	USCS_Node* SCS_Node = CachedNodePtr->GetSCSNode();
-	if(SCS_Node != NULL && !InNewText.IsEmpty() && !SCS_Node->IsValidVariableNameString(InNewText.ToString()))
+	if(SCS_Node != NULL && !InNewText.IsEmpty() && !FComponentEditorUtils::IsValidVariableNameString(SCS_Node->ComponentTemplate, InNewText.ToString()))
 	{
 		VariableNameEditableTextBox->SetError(LOCTEXT("ComponentVariableRenameFailed_NotValid", "This name is reserved for engine use."));
 		return;
diff --git a/Engine/Source/Editor/Kismet/Private/SBlueprintPalette.cpp b/Engine/Source/Editor/Kismet/Private/SBlueprintPalette.cpp
--- a/Engine/Source/Editor/Kismet/Private/SBlueprintPalette.cpp
+++ b/Engine/Source/Editor/Kismet/Private/SBlueprintPalette.cpp
@@ -34,6 +34,7 @@
 #include "BlueprintEditorSettings.h" // for bShowActionMenuItemSignatures
 #include "SInlineEditableTextBlock.h"
 #include "Engine/SimpleConstructionScript.h"
+#include "Editor/UnrealEd/Public/Kismet2/ComponentEditorUtils.h"
 
 #define LOCTEXT_NAMESPACE "BlueprintPalette"
 
@@ -1111,7 +1112,7 @@ bool SBlueprintPaletteItem::OnNameTextVerifyChanged(const FText& InNewText, FTex
 		for (TArray<USCS_Node*>::TConstIterator NodeIt(Nodes); NodeIt; ++NodeIt)
 		{
 			USCS_Node* Node = *NodeIt;
-			if (Node->VariableName == OriginalName && !Node->IsValidVariableNameString(InNewText.ToString()))
+			if (Node->VariableName == OriginalName && !FComponentEditorUtils::IsValidVariableNameString(Node->ComponentTemplate, InNewText.ToString()))
 			{
 				OutErrorMessage = LOCTEXT("RenameFailed_NotValid", "This name is reserved for engine use.");
 				return false;
diff --git a/Engine/Source/Editor/Kismet/Private/SSCSEditor.cpp b/Engine/Source/Editor/Kismet/Private/SSCSEditor.cpp
--- a/Engine/Source/Editor/Kismet/Private/SSCSEditor.cpp
+++ b/Engine/Source/Editor/Kismet/Private/SSCSEditor.cpp
@@ -4,6 +4,7 @@
 #include "BlueprintEditorPrivatePCH.h"
 #include "Editor/PropertyEditor/Public/IDetailsView.h"
 #include "Editor/UnrealEd/Public/Kismet2/BlueprintEditorUtils.h"
+#include "Editor/UnrealEd/Public/Kismet2/ComponentEditorUtils.h"
 #include "BlueprintUtilities.h"
 #include "ComponentAssetBroker.h"
 
@@ -747,6 +748,43 @@ FSCSEditorTreeNodePtrType FSCSEditorTreeNode::FindChild(const UActorComponent* I
 	return Result;
 }
 
+FSCSEditorTreeNodePtrType FSCSEditorTreeNode::FindChild(const FName& InVariableOrInstanceName, bool bRecursiveSearch, uint32* OutDepth) const
+{
+	FSCSEditorTreeNodePtrType Result;
+
+	// Ensure that the given name is valid
+	if(InVariableOrInstanceName != NAME_None)
+	{
+		// Look for a match in our set of child nodes
+		for(int32 ChildIndex = 0; ChildIndex < Children.Num() && !Result.IsValid(); ++ChildIndex)
+		{
+			FName ItemName = Children[ChildIndex]->GetVariableName();
+			if(ItemName == NAME_None)
+			{
+				UActorComponent* ComponentTemplateOrInstance = Children[ChildIndex]->GetComponentTemplate();
+				check(ComponentTemplateOrInstance != nullptr);
+				ItemName = ComponentTemplateOrInstance->GetFName();
+			}
+
+			if(InVariableOrInstanceName == ItemName)
+			{
+				Result = Children[ChildIndex];
+			}
+			else if(bRecursiveSearch)
+			{
+				Result = Children[ChildIndex]->FindChild(InVariableOrInstanceName, true, OutDepth);
+			}
+		}
+	}
+
+	if(OutDepth && Result.IsValid())
+	{
+		*OutDepth += 1;
+	}
+
+	return Result;
+}
+
 void FSCSEditorTreeNode::RemoveChild(FSCSEditorTreeNodePtrType InChildNodePtr)
 {
 	// Remove the given node as a child and reset its parent link
@@ -792,7 +830,23 @@ void FSCSEditorTreeNode::OnCompleteRename(const FText& InNewName)
 		TransactionContext = new FScopedTransaction(LOCTEXT("RenameComponentVariable", "Rename Component Variable"));
 	}
 
-	FBlueprintEditorUtils::RenameComponentMemberVariable(GetBlueprint(), GetSCSNode(), FName( *InNewName.ToString() ));
+	if(bIsInstanced)
+	{
+		UActorComponent* ComponentInstance = GetComponentTemplate();
+		check(ComponentInstance != nullptr);
+
+		ERenameFlags RenameFlags = REN_DontCreateRedirectors;
+		if(!TransactionContext)
+		{
+			RenameFlags |= REN_NonTransactional;
+		}
+
+		ComponentInstance->Rename(*InNewName.ToString(), nullptr, RenameFlags);
+	}
+	else
+	{
+		FBlueprintEditorUtils::RenameComponentMemberVariable(GetBlueprint(), GetSCSNode(), FName( *InNewName.ToString() ));
+	}
 
 	if(TransactionContext)
 	{
@@ -2090,14 +2144,22 @@ UBlueprint* SSCS_RowWidget::GetBlueprint() const
 
 bool SSCS_RowWidget::OnNameTextVerifyChanged(const FText& InNewText, FText& OutErrorMessage)
 {
-	USCS_Node* SCS_Node = NodePtr->GetSCSNode();
-	if(SCS_Node != NULL && !InNewText.IsEmpty() && !SCS_Node->IsValidVariableNameString(InNewText.ToString()))
+	if(!InNewText.IsEmpty() && !FComponentEditorUtils::IsValidVariableNameString(NodePtr->GetComponentTemplate(), InNewText.ToString()))
 	{
 		OutErrorMessage = LOCTEXT("RenameFailed_NotValid", "This name is reserved for engine use.");
 		return false;
 	}
 
-	TSharedPtr<INameValidatorInterface> NameValidator = MakeShareable(new FKismetNameValidator(GetBlueprint(), NodePtr->GetVariableName()));
+	UBlueprint* Blueprint = GetBlueprint();
+	TSharedPtr<INameValidatorInterface> NameValidator;
+	if(Blueprint != nullptr)
+	{
+		NameValidator = MakeShareable(new FKismetNameValidator(GetBlueprint(), NodePtr->GetVariableName()));
+	}
+	else
+	{
+		NameValidator = MakeShareable(new FStringSetNameValidator(NodePtr->GetComponentTemplate()->GetName()));
+	}
 
 	EValidatorResult ValidatorResult = NameValidator->IsValid(InNewText.ToString());
 	if(ValidatorResult == EValidatorResult::AlreadyInUse)
@@ -2124,6 +2186,14 @@ bool SSCS_RowWidget::OnNameTextVerifyChanged(const FText& InNewText, FText& OutE
 void SSCS_RowWidget::OnNameTextCommit(const FText& InNewName, ETextCommit::Type InTextCommit)
 {
 	NodePtr->OnCompleteRename(InNewName);
+
+	// No need to call UpdateTree() in SCS editor mode; it will already be called by MBASM internally
+	check(SCSEditor.IsValid());
+	TSharedPtr<SSCSEditor> PinnedEditor = SCSEditor.Pin();
+	if(PinnedEditor.IsValid() && PinnedEditor->EditorMode.Get() == SSCSEditor::EEditorMode::ActorInstance)
+	{
+		PinnedEditor->UpdateTree();
+	}
 }
 
 //////////////////////////////////////////////////////////////////////////
@@ -2287,11 +2357,21 @@ FReply SSCSEditor::OnKeyDown( const FGeometry& MyGeometry, const FKeyEvent& InKe
 
 TSharedRef<ITableRow> SSCSEditor::MakeTableRowWidget( FSCSEditorTreeNodePtrType InNodePtr, const TSharedRef<STableViewBase>& OwnerTable )
 {
-	// Native components do not have a SCSNode, so ignore them
-	if(InNodePtr->GetSCSNode() && DeferredRenameRequest.Get() == InNodePtr->GetSCSNode())
+	if(DeferredRenameRequest != NAME_None)
 	{
-		SCSTreeWidget->SetSelection(InNodePtr);
-		OnRenameComponent(false);
+		FName ItemName = InNodePtr->GetVariableName();
+		if(ItemName == NAME_None)
+		{
+			UActorComponent* ComponentTemplateOrInstance = InNodePtr->GetComponentTemplate();
+			check(ComponentTemplateOrInstance != nullptr);
+			ItemName = ComponentTemplateOrInstance->GetFName();
+		}
+
+		if(DeferredRenameRequest == ItemName)
+		{
+			SCSTreeWidget->SetSelection(InNodePtr);
+			OnRenameComponent(false);
+		}
 	}
 
 	// Setup a meta tag for this node
@@ -2933,9 +3013,9 @@ void SSCSEditor::UpdateTree(bool bRegenerateTreeNodes)
 		}
 
 		// If we have a pending deferred rename request, redirect it to the new tree node
-		if(DeferredRenameRequest.IsValid())
+		if(DeferredRenameRequest != NAME_None)
 		{
-			FSCSEditorTreeNodePtrType NodeToRenamePtr = FindTreeNode(DeferredRenameRequest.Get());
+			FSCSEditorTreeNodePtrType NodeToRenamePtr = FindTreeNode(DeferredRenameRequest);
 			if(NodeToRenamePtr.IsValid())
 			{
 				SCSTreeWidget->RequestScrollIntoView(NodeToRenamePtr);
@@ -3442,7 +3522,21 @@ void SSCSEditor::OnDeleteNodes()
 	}
 	else    // SSCSEditor::EEditorMode::ActorInstance
 	{
-		// @TODO - Actor instance mode
+		AActor* ActorInstance = ActorContext.Get();
+		check(ActorInstance != nullptr);
+
+		ActorInstance->Modify();
+
+		TArray<FSCSEditorTreeNodePtrType> SelectedNodes = SCSTreeWidget->GetSelectedItems();
+		for (int32 i = 0; i < SelectedNodes.Num(); ++i)
+		{
+			auto Node = SelectedNodes[i];
+
+			RemoveComponentNode(Node);
+		}
+
+		// Rebuild the tree view to reflect the new component hierarchy
+		UpdateTree();
 	}
 
 	// Do this AFTER marking the Blueprint as modified
@@ -3483,7 +3577,14 @@ void SSCSEditor::RemoveComponentNode(FSCSEditorTreeNodePtrType InNodePtr)
 	}
 	else    // SSCSEditor::EEditorMode::ActorInstance
 	{
-		// @TODO - Actor instance mode
+		AActor* ActorInstance = ActorContext.Get();
+		check(ActorInstance != nullptr);
+
+		UActorComponent* ComponentInstance = InNodePtr->GetComponentTemplate();
+		check(ComponentInstance != nullptr);
+
+		// Destroy the component instance
+		ComponentInstance->DestroyComponent();
 	}
 }
 
@@ -3751,13 +3852,67 @@ FSCSEditorTreeNodePtrType SSCSEditor::FindTreeNode(const UActorComponent* InComp
 	return NodePtr;
 }
 
+FSCSEditorTreeNodePtrType SSCSEditor::FindTreeNode(const FName& InVariableOrInstanceName, FSCSEditorTreeNodePtrType InStartNodePtr) const
+{
+	FSCSEditorTreeNodePtrType NodePtr;
+	if(InVariableOrInstanceName != NAME_None)
+	{
+		// Start at the scene root node if none was given
+		if(!InStartNodePtr.IsValid())
+		{
+			InStartNodePtr = SceneRootNodePtr;
+		}
+
+		if(InStartNodePtr.IsValid())
+		{
+			FName ItemName = InStartNodePtr->GetVariableName();
+			if(ItemName == NAME_None)
+			{
+				UActorComponent* ComponentTemplateOrInstance = InStartNodePtr->GetComponentTemplate();
+				check(ComponentTemplateOrInstance != nullptr);
+				ItemName = ComponentTemplateOrInstance->GetFName();
+			}
+
+			// Check to see if the given name matches the item name
+			if(InVariableOrInstanceName == ItemName)
+			{
+				NodePtr = InStartNodePtr;
+			}
+			else
+			{
+				// Recursively search for the node in our child set
+				NodePtr = InStartNodePtr->FindChild(InVariableOrInstanceName);
+				if(!NodePtr.IsValid())
+				{
+					for(int32 i = 0; i < InStartNodePtr->GetChildren().Num() && !NodePtr.IsValid(); ++i)
+					{
+						NodePtr = FindTreeNode(InVariableOrInstanceName, InStartNodePtr->GetChildren()[i]);
+					}
+				}
+			}
+		}
+	}
+
+	return NodePtr;
+}
+
 void SSCSEditor::OnItemScrolledIntoView( FSCSEditorTreeNodePtrType InItem, const TSharedPtr<ITableRow>& InWidget)
 {
-	USCS_Node* SCS_Node = DeferredRenameRequest.Get();
-	if( SCS_Node && SCS_Node == InItem->GetSCSNode() )
+	if(DeferredRenameRequest != NAME_None)
 	{
-		DeferredRenameRequest.Reset();
-		InItem->OnRequestRename(bIsDeferredRenameRequestTransactional);
+		FName ItemName = InItem->GetVariableName();
+		if(ItemName == NAME_None)
+		{
+			UActorComponent* ComponentTemplateOrInstance = InItem->GetComponentTemplate();
+			check(ComponentTemplateOrInstance != nullptr);
+			ItemName = ComponentTemplateOrInstance->GetFName();
+		}
+
+		if(DeferredRenameRequest == ItemName)
+		{
+			DeferredRenameRequest = NAME_None;
+			InItem->OnRequestRename(bIsDeferredRenameRequestTransactional);
+		}
 	}
 }
 
@@ -3769,7 +3924,14 @@ void SSCSEditor::OnRenameComponent(bool bTransactional)
 	check(SelectedItems.Num() == 1);
 
 	SCSTreeWidget->RequestScrollIntoView(SelectedItems[0]);
-	DeferredRenameRequest = SelectedItems[0]->GetSCSNode();
+	DeferredRenameRequest = SelectedItems[0]->GetVariableName();
+	if(DeferredRenameRequest == NAME_None)
+	{
+		UActorComponent* ComponentTemplateOrInstance = SelectedItems[0]->GetComponentTemplate();
+		check(ComponentTemplateOrInstance != nullptr);
+		DeferredRenameRequest = ComponentTemplateOrInstance->GetFName();
+	}
+
 	bIsDeferredRenameRequestTransactional = bTransactional;
 }
 
diff --git a/Engine/Source/Editor/Kismet/Public/SSCSEditor.h b/Engine/Source/Editor/Kismet/Public/SSCSEditor.h
--- a/Engine/Source/Editor/Kismet/Public/SSCSEditor.h
+++ b/Engine/Source/Editor/Kismet/Public/SSCSEditor.h
@@ -113,7 +113,7 @@ public:
 	/** 
 	 * @return Whether or not this object represents a node that can be renamed from the SCS tree.
 	 */
-	bool CanRename() const { return !IsNative() && !IsInherited() && !IsDefaultSceneRoot(); }
+	bool CanRename() const { return IsUserInstanced() || (!IsNative() && !IsInherited() && !IsDefaultSceneRoot()); }
 	/** 
 	 * @return Whether or not this object represents a node that can be reparented to other nodes based on its context.
 	 */
@@ -175,6 +175,16 @@ public:
 	 */
 	FSCSEditorTreeNodePtrType FindChild(const UActorComponent* InComponentTemplate, bool bRecursiveSearch = false, uint32* OutDepth = NULL) const;
 
+	/** 
+	 * Attempts to find a reference to the child node that matches the given component variable or instance name.
+	 *
+	 * @param InVariableOrInstanceName The component variable or instance name to match.
+	 * @param bRecursiveSearch Whether or not to recursively search child nodes (default == false).
+	 * @param OutDepth If non-NULL, the depth of the child node will be returned in this parameter on success (default == NULL).
+	 * @return The child node with a component variable or instance name that matches the given name, or an invalid node reference if no match was found.
+	 */
+	FSCSEditorTreeNodePtrType FindChild(const FName& InVariableOrInstanceName, bool bRecursiveSearch = false, uint32* OutDepth = NULL) const;
+
 	/** 
 	 * Removes the given node from the list of child nodes.
 	 *
@@ -626,6 +636,9 @@ protected:
 	/** Helper method to recursively find a tree node for the given scene component starting at the given tree node */
 	FSCSEditorTreeNodePtrType FindTreeNode(const UActorComponent* InComponent, FSCSEditorTreeNodePtrType InStartNodePtr = FSCSEditorTreeNodePtrType()) const;
 
+	/** Helper method to recursively find a tree node for the given variable or instance name starting at the given tree node */
+	FSCSEditorTreeNodePtrType FindTreeNode(const FName& InVariableOrInstanceName, FSCSEditorTreeNodePtrType InStartNodePtr = FSCSEditorTreeNodePtrType()) const;
+
 	/** Callback when a component item is scrolled into view */
 	void OnItemScrolledIntoView( FSCSEditorTreeNodePtrType InItem, const TSharedPtr<ITableRow>& InWidget);
 
@@ -644,8 +657,8 @@ public:
 	/** Command list for handling actions in the SSCSEditor */
 	TSharedPtr< FUICommandList > CommandList;
 
-	/** Weak reference to an SCS node that has been requested to be renamed */
-	TWeakObjectPtr<USCS_Node> DeferredRenameRequest;
+	/** Name of a node that has been requested to be renamed */
+	FName DeferredRenameRequest;
 
 	/** Whether or not the deferred rename request was flagged as transactional */
 	bool bIsDeferredRenameRequestTransactional;
diff --git a/Engine/Source/Editor/UnrealEd/Private/Kismet2/ComponentEditorUtils.cpp b/Engine/Source/Editor/UnrealEd/Private/Kismet2/ComponentEditorUtils.cpp
--- a/Engine/Source/Editor/UnrealEd/Private/Kismet2/ComponentEditorUtils.cpp
+++ b/Engine/Source/Editor/UnrealEd/Private/Kismet2/ComponentEditorUtils.cpp
@@ -90,6 +90,23 @@ void FComponentEditorUtils::GetArchetypeInstances( UObject* Object, TArray<UObje
 	}
 }
 
+bool FComponentEditorUtils::IsValidVariableNameString(const UActorComponent* InComponent, const FString& InString)
+{
+	// First test to make sure the string is not empty and does not equate to the DefaultSceneRoot node name
+	bool bIsValid = !InString.IsEmpty() && !InString.Equals(USimpleConstructionScript::DefaultSceneRootVariableName.ToString());
+	if(bIsValid && InComponent != NULL)
+	{
+		// Next test to make sure the string doesn't conflict with the format that MakeUniqueObjectName() generates
+		FString MakeUniqueObjectNamePrefix = FString::Printf(TEXT("%s_"), *InComponent->GetClass()->GetName());
+		if(InString.StartsWith(MakeUniqueObjectNamePrefix))
+		{
+			bIsValid = !InString.Replace(*MakeUniqueObjectNamePrefix, TEXT("")).IsNumeric();
+		}
+	}
+
+	return bIsValid;
+}
+
 void FComponentEditorUtils::PropagateTransformPropertyChange(
 	class USceneComponent* InSceneComponentTemplate,
 	const FTransformData& OldDefaultTransform,
diff --git a/Engine/Source/Editor/UnrealEd/Public/Kismet2/ComponentEditorUtils.h b/Engine/Source/Editor/UnrealEd/Public/Kismet2/ComponentEditorUtils.h
--- a/Engine/Source/Editor/UnrealEd/Public/Kismet2/ComponentEditorUtils.h
+++ b/Engine/Source/Editor/UnrealEd/Public/Kismet2/ComponentEditorUtils.h
@@ -13,6 +13,9 @@ public:
 
 	static void GetArchetypeInstances(class UObject* Object, TArray<class UObject*>& ArchetypeInstances);
 
+	/** Test whether or not the given string is a valid variable name string for the given component instance */
+	static bool IsValidVariableNameString(const UActorComponent* InComponent, const FString& InString);
+
 	struct FTransformData
 	{
 		FTransformData(const class USceneComponent& Component)
diff --git a/Engine/Source/Editor/UnrealEd/Public/Kismet2/Kismet2NameValidators.h b/Engine/Source/Editor/UnrealEd/Public/Kismet2/Kismet2NameValidators.h
--- a/Engine/Source/Editor/UnrealEd/Public/Kismet2/Kismet2NameValidators.h
+++ b/Engine/Source/Editor/UnrealEd/Public/Kismet2/Kismet2NameValidators.h
@@ -98,7 +98,6 @@ public:
 	virtual EValidatorResult IsValid(const FName& Name, bool bOriginal) override;
 	// End FNameValidatorInterface
 
-protected:
 	// This class is a base class for anything that just needs to validate a string is unique
 	FStringSetNameValidator(const FString& InExistingName)
 		: ExistingName(InExistingName)
diff --git a/Engine/Source/Runtime/Engine/Classes/Engine/SCS_Node.h b/Engine/Source/Runtime/Engine/Classes/Engine/SCS_Node.h
--- a/Engine/Source/Runtime/Engine/Classes/Engine/SCS_Node.h
+++ b/Engine/Source/Runtime/Engine/Classes/Engine/SCS_Node.h
@@ -152,9 +152,6 @@ class USCS_Node : public UObject
 
 	/** Finds and returns the parent component template through the given Blueprint */
 	ENGINE_API USceneComponent* GetParentComponentTemplate(UBlueprint* InBlueprint) const;
-
-	/** Test whether or not the given string is a valid variable name string */
-	ENGINE_API bool IsValidVariableNameString(const FString& InString);
 #endif
 		
 private:
diff --git a/Engine/Source/Runtime/Engine/Private/SCS_Node.cpp b/Engine/Source/Runtime/Engine/Private/SCS_Node.cpp
--- a/Engine/Source/Runtime/Engine/Private/SCS_Node.cpp
+++ b/Engine/Source/Runtime/Engine/Private/SCS_Node.cpp
@@ -382,23 +382,6 @@ USceneComponent* USCS_Node::GetParentComponentTemplate(UBlueprint* InBlueprint)
 	return ParentComponentTemplate;
 }
 
-bool USCS_Node::IsValidVariableNameString(const FString& InString)
-{
-	// First test to make sure the string is not empty and does not equate to the DefaultSceneRoot node name
-	bool bIsValid = !InString.IsEmpty() && !InString.Equals(USimpleConstructionScript::DefaultSceneRootVariableName.ToString());
-	if(bIsValid && ComponentTemplate != NULL)
-	{
-		// Next test to make sure the string doesn't conflict with the format that MakeUniqueObjectName() generates
-		FString MakeUniqueObjectNamePrefix = FString::Printf(TEXT("%s_"), *ComponentTemplate->GetClass()->GetName());
-		if(InString.StartsWith(MakeUniqueObjectNamePrefix))
-		{
-			bIsValid = !InString.Replace(*MakeUniqueObjectNamePrefix, TEXT("")).IsNumeric();
-		}
-	}
-
-	return bIsValid;
-}
-
 void USCS_Node::GenerateListOfExistingNames( TArray<FName>& CurrentNames ) const
 {
 	const USimpleConstructionScript* SCS = GetSCS();
diff --git a/Engine/Source/Runtime/Engine/Private/SimpleConstructionScript.cpp b/Engine/Source/Runtime/Engine/Private/SimpleConstructionScript.cpp
--- a/Engine/Source/Runtime/Engine/Private/SimpleConstructionScript.cpp
+++ b/Engine/Source/Runtime/Engine/Private/SimpleConstructionScript.cpp
@@ -6,6 +6,7 @@
 #if WITH_EDITOR
 #include "Kismet2/CompilerResultsLog.h"
 #include "Kismet2/BlueprintEditorUtils.h"
+#include "Kismet2/ComponentEditorUtils.h"
 #include "Kismet2/Kismet2NameValidators.h"
 #endif
 #include "Engine/SimpleConstructionScript.h"
@@ -778,7 +779,7 @@ void USimpleConstructionScript::ValidateNodeVariableNames(FCompilerResultsLog& M
 			// Replace missing or invalid component variable names
 			if( Node->VariableName == NAME_None
 				|| Node->bVariableNameAutoGenerated_DEPRECATED
-				|| !Node->IsValidVariableNameString(Node->VariableName.ToString()) )
+				|| !FComponentEditorUtils::IsValidVariableNameString(Node->ComponentTemplate, Node->VariableName.ToString()) )
 			{
 				FName OldName = Node->VariableName;
 
