This change adds the ability in the blueprint details panel to set blueprint variable values using the owning config file.
#UE-6467 - Blueprints need to be able to make properties config (.ini)
#Branch UE4
#Proj Kismet, KismetCompiler

#info please note that currently for a blueprint called MyTestBlueprint with a variable MyVariable the config entry would be for example [/Game/Blueprints/MyTestBlueprint.MyTestBlueprint_C] MyVariable=99.0

diff --git a/Engine/Source/Editor/Kismet/Private/BlueprintDetailsCustomization.cpp b/Engine/Source/Editor/Kismet/Private/BlueprintDetailsCustomization.cpp
--- a/Engine/Source/Editor/Kismet/Private/BlueprintDetailsCustomization.cpp
+++ b/Engine/Source/Editor/Kismet/Private/BlueprintDetailsCustomization.cpp
@@ -223,6 +223,24 @@ void FBlueprintVarActionDetails::CustomizeDetails( IDetailLayoutBuilder& DetailL
 		.ToolTip(ExposeToMatineeTooltip)
 	];
 
+	TSharedPtr<SToolTip> ExposeToConfigTooltip = IDocumentation::Get()->CreateToolTip(LOCTEXT("VariableExposeToConfig_Tooltip", "Allow this variable be set from the config?"), NULL, DocLink, TEXT("ExposeToConfig"));
+
+	Category.AddCustomRow( LOCTEXT("VariableExposeToConfig", "Config Variable") )
+	.NameContent()
+	[
+		SNew(STextBlock)
+		.ToolTip( ExposeToConfigTooltip )
+		.Text( LOCTEXT("ExposeToConfigLabel", "Config Variable") )
+		.Font( DetailFontInfo )
+	]
+	.ValueContent()
+	[
+		SNew(SCheckBox)
+		.ToolTip( ExposeToConfigTooltip )
+		.IsChecked( this, &FBlueprintVarActionDetails::OnGetConfigVariableCheckboxState )
+		.OnCheckStateChanged( this, &FBlueprintVarActionDetails::OnSetConfigVariableState )
+	];
+
 	PopulateCategories(MyBlueprint.Pin().Get(), CategorySource);
 	TSharedPtr<SComboButton> NewComboButton;
 	TSharedPtr<SListView<TSharedPtr<FString>>> NewListView;
@@ -1407,6 +1425,48 @@ EVisibility FBlueprintVarActionDetails::ExposeToMatineeVisibility() const
 	return EVisibility::Collapsed;
 }
 
+ECheckBoxState FBlueprintVarActionDetails::OnGetConfigVariableCheckboxState() const
+{
+	UBlueprint* Blueprint = GetBlueprintObj();
+	const FName VarName = GetVariableName();
+	ECheckBoxState CheckboxValue = ECheckBoxState::Unchecked;
+
+	if( Blueprint && VarName != NAME_None )
+	{
+		const int32 VarIndex = FBlueprintEditorUtils::FindNewVariableIndex( Blueprint, VarName );
+
+		if( VarIndex != INDEX_NONE && Blueprint->NewVariables[ VarIndex ].PropertyFlags & CPF_Config )
+		{
+			CheckboxValue = ECheckBoxState::Checked;
+		}
+	}
+	return CheckboxValue;
+}
+
+void FBlueprintVarActionDetails::OnSetConfigVariableState( ECheckBoxState InNewState )
+{
+	UBlueprint* Blueprint = GetBlueprintObj();
+	const FName VarName = GetVariableName();
+
+	if( Blueprint && VarName != NAME_None )
+	{
+		const int32 VarIndex = FBlueprintEditorUtils::FindNewVariableIndex( Blueprint, VarName );
+
+		if( VarIndex != INDEX_NONE )
+		{
+			if( InNewState == ECheckBoxState::Checked )
+			{
+				Blueprint->NewVariables[ VarIndex ].PropertyFlags |= CPF_Config;
+			}
+			else
+			{
+				Blueprint->NewVariables[ VarIndex ].PropertyFlags &= ~CPF_Config;
+			}
+			FBlueprintEditorUtils::MarkBlueprintAsStructurallyModified( Blueprint );
+		}
+	}
+}
+
 FText FBlueprintVarActionDetails::OnGetMetaKeyValue(FName Key) const
 {
 	FName VarName = GetVariableName();
diff --git a/Engine/Source/Editor/Kismet/Private/BlueprintDetailsCustomization.h b/Engine/Source/Editor/Kismet/Private/BlueprintDetailsCustomization.h
--- a/Engine/Source/Editor/Kismet/Private/BlueprintDetailsCustomization.h
+++ b/Engine/Source/Editor/Kismet/Private/BlueprintDetailsCustomization.h
@@ -120,6 +120,9 @@ private:
 	void OnExposedToMatineeChanged(ECheckBoxState InNewState);
 	EVisibility ExposeToMatineeVisibility() const;
 
+	ECheckBoxState OnGetConfigVariableCheckboxState() const;
+	void OnSetConfigVariableState(ECheckBoxState InNewState);
+
 	FText OnGetMetaKeyValue(FName Key) const;
 	void OnMetaKeyValueChanged(const FText& NewMinValue, ETextCommit::Type CommitInfo, FName Key);
 	EVisibility RangeVisibility() const;
diff --git a/Engine/Source/Editor/KismetCompiler/Private/KismetCompiler.cpp b/Engine/Source/Editor/KismetCompiler/Private/KismetCompiler.cpp
--- a/Engine/Source/Editor/KismetCompiler/Private/KismetCompiler.cpp
+++ b/Engine/Source/Editor/KismetCompiler/Private/KismetCompiler.cpp
@@ -1757,6 +1757,11 @@ void FKismetCompilerContext::FinishCompilingClass(UClass* Class)
 				Property->RepNotifyFunc = NAME_None;
 			}
 		}
+		if( Property->HasAnyPropertyFlags( CPF_Config ))
+		{
+			// If we have properties that are set from the config, then the class needs to also have CLASS_Config flags
+			Class->ClassFlags |= CLASS_Config;
+		}
 	}
 
 	// Set class metadata as needed
