Fix for Config option being shown for components in blueprints details panel.
#UE-7908 - Config checkbox appears in blueprint details panel for components incorrectly.
#Branch UE4
#Proj Kismet

diff --git a/Engine/Source/Editor/Kismet/Private/BlueprintDetailsCustomization.cpp b/Engine/Source/Editor/Kismet/Private/BlueprintDetailsCustomization.cpp
--- a/Engine/Source/Editor/Kismet/Private/BlueprintDetailsCustomization.cpp
+++ b/Engine/Source/Editor/Kismet/Private/BlueprintDetailsCustomization.cpp
@@ -226,6 +226,7 @@ void FBlueprintVarActionDetails::CustomizeDetails( IDetailLayoutBuilder& DetailL
 	TSharedPtr<SToolTip> ExposeToConfigTooltip = IDocumentation::Get()->CreateToolTip(LOCTEXT("VariableExposeToConfig_Tooltip", "Allow this variable be set from the config?"), NULL, DocLink, TEXT("ExposeToConfig"));
 
 	Category.AddCustomRow( LOCTEXT("VariableExposeToConfig", "Config Variable") )
+	.Visibility(TAttribute<EVisibility>(this, &FBlueprintVarActionDetails::ExposeConfigVisibility))
 	.NameContent()
 	[
 		SNew(STextBlock)
@@ -1467,6 +1468,19 @@ void FBlueprintVarActionDetails::OnSetConfigVariableState( ECheckBoxState InNewS
 	}
 }
 
+EVisibility FBlueprintVarActionDetails::ExposeConfigVisibility() const
+{
+	UProperty* Property = SelectionAsProperty();
+	if (Property)
+	{
+		if (IsABlueprintVariable(Property) && !IsAComponentVariable(Property))
+		{
+			return EVisibility::Visible;
+		}
+	}
+	return EVisibility::Collapsed;
+}
+
 FText FBlueprintVarActionDetails::OnGetMetaKeyValue(FName Key) const
 {
 	FName VarName = GetVariableName();
diff --git a/Engine/Source/Editor/Kismet/Private/BlueprintDetailsCustomization.h b/Engine/Source/Editor/Kismet/Private/BlueprintDetailsCustomization.h
--- a/Engine/Source/Editor/Kismet/Private/BlueprintDetailsCustomization.h
+++ b/Engine/Source/Editor/Kismet/Private/BlueprintDetailsCustomization.h
@@ -122,6 +122,7 @@ private:
 
 	ECheckBoxState OnGetConfigVariableCheckboxState() const;
 	void OnSetConfigVariableState(ECheckBoxState InNewState);
+	EVisibility ExposeConfigVisibility() const;
 
 	FText OnGetMetaKeyValue(FName Key) const;
 	void OnMetaKeyValueChanged(const FText& NewMinValue, ETextCommit::Type CommitInfo, FName Key);
