UE-7903 : SCSEditor is clearing the details view, even when we're focussed on teh graph view

diff --git a/Engine/Source/Editor/Kismet/Private/SSCSEditor.cpp b/Engine/Source/Editor/Kismet/Private/SSCSEditor.cpp
--- a/Engine/Source/Editor/Kismet/Private/SSCSEditor.cpp
+++ b/Engine/Source/Editor/Kismet/Private/SSCSEditor.cpp
@@ -2869,7 +2869,10 @@ void SSCSEditor::UpdateTree(bool bRegenerateTreeNodes)
 		TArray<FSCSEditorTreeNodePtrType> SelectedTreeNodes = SCSTreeWidget->GetSelectedItems();
 
 		// Clear the current tree
-		SCSTreeWidget->ClearSelection();
+		if (SelectedTreeNodes.Num() != 0)
+		{
+			SCSTreeWidget->ClearSelection();
+		}
 		RootNodes.Empty();
 
 		// Reset the scene root node
