Blueprint Inheritable Components - fix for 3d Widget in SCS Editor

Merging 2418656

//depot/UE4-Releases/4.7/Engine/Source/Editor/Kismet/...

to //depot/UE4/Engine/Source/Editor/Kismet/...

diff --git a/Engine/Source/Editor/Kismet/Private/BlueprintEditor.cpp b/Engine/Source/Editor/Kismet/Private/BlueprintEditor.cpp
--- a/Engine/Source/Editor/Kismet/Private/BlueprintEditor.cpp
+++ b/Engine/Source/Editor/Kismet/Private/BlueprintEditor.cpp
@@ -624,22 +624,22 @@ void FBlueprintEditor::RefreshEditors(ERefreshBlueprintEditorReason::Type Reason
 }
 
 void FBlueprintEditor::ClearSelectionInAllEditors()
-{
-	TArray< TSharedPtr<SDockTab> > GraphEditorTabs;
-	DocumentManager->FindAllTabsForFactory(GraphEditorTabFactoryPtr, /*out*/ GraphEditorTabs);
-
-	for (auto GraphEditorTabIt = GraphEditorTabs.CreateIterator(); GraphEditorTabIt; ++GraphEditorTabIt)
 	{
-		TSharedRef<SGraphEditor> Editor = StaticCastSharedRef<SGraphEditor>((*GraphEditorTabIt)->GetContent());
+		TArray< TSharedPtr<SDockTab> > GraphEditorTabs;
+		DocumentManager->FindAllTabsForFactory(GraphEditorTabFactoryPtr, /*out*/ GraphEditorTabs);
 
-		Editor->ClearSelectionSet();
-	}
+		for ( auto GraphEditorTabIt = GraphEditorTabs.CreateIterator(); GraphEditorTabIt; ++GraphEditorTabIt )
+		{
+			TSharedRef<SGraphEditor> Editor = StaticCastSharedRef<SGraphEditor>(( *GraphEditorTabIt )->GetContent());
 
-	if(SCSEditor.IsValid())
-	{
-		SCSEditor->ClearSelection();
+			Editor->ClearSelectionSet();
+		}
+
+		if ( SCSEditor.IsValid() )
+		{
+			SCSEditor->ClearSelection();
+		}
 	}
-}
 
 void FBlueprintEditor::SummonSearchUI(bool bSetFindWithinBlueprint, FString NewSearchTerms, bool bSelectFirstResult)
 {
@@ -702,33 +702,24 @@ void FBlueprintEditor::OnSCSEditorTreeViewSelectionChanged(const TArray<FSCSEdit
 }
 
 void FBlueprintEditor::OnSCSEditorUpdateSelectionFromNodes(const TArray<FSCSEditorTreeNodePtrType>& SelectedNodes)
-{
-	// Convert the selection set to an array of UObject* pointers
-	FText InspectorTitle = FText::GetEmpty();
-	TArray<UObject*> InspectorObjects;
-	InspectorObjects.Empty(SelectedNodes.Num());
-	for (auto NodeIt = SelectedNodes.CreateConstIterator(); NodeIt; ++NodeIt)
 	{
-		auto NodePtr = *NodeIt;
-		if (NodePtr.IsValid())
+		// Convert the selection set to an array of UObject* pointers
+		FText InspectorTitle = FText::GetEmpty();
+		TArray<UObject*> InspectorObjects;
+		InspectorObjects.Empty(SelectedNodes.Num());
+		for (auto NodeIt = SelectedNodes.CreateConstIterator(); NodeIt; ++NodeIt)
 		{
-			UActorComponent* EditableComponent = nullptr;
-			if (NodePtr->CanEditDefaults())
-			{
-				EditableComponent = NodePtr->GetComponentTemplate();
-			}
-			else if (!NodePtr->IsNative() && NodePtr->IsInherited())
+			auto NodePtr = *NodeIt;
+			if (NodePtr.IsValid())
 			{
-				EditableComponent = NodePtr->GetOverridenComponentTemplate(GetBlueprintObj(), true);
-			}
-
-			if (EditableComponent)
-		{
-			InspectorTitle = FText::FromString(NodePtr->GetDisplayString());
-				InspectorObjects.Add(EditableComponent);
+				UActorComponent* EditableComponent = NodePtr->GetEditableComponentTemplate(GetBlueprintObj());
+				if (EditableComponent)
+				{
+					InspectorTitle = FText::FromString(NodePtr->GetDisplayString());
+					InspectorObjects.Add(EditableComponent);
+				}
 			}
 		}
-	}
 
 	// Update the details panel
 	if (Inspector.IsValid())
@@ -1101,21 +1092,21 @@ FGraphAppearanceInfo FBlueprintEditor::GetGraphAppearance(UEdGraph* InGraph) con
 	UBlueprint* Blueprint = (InGraph != nullptr) ? FBlueprintEditorUtils::FindBlueprintForGraph(InGraph) : GetBlueprintObj();
 	if (Blueprint != NULL)
 	{
-		switch (Blueprint->BlueprintType)
-		{
-		case BPTYPE_LevelScript:
-			AppearanceInfo.CornerText = LOCTEXT("AppearanceCornerText_LevelScript", "LEVEL BLUEPRINT");
-			break;
-		case BPTYPE_MacroLibrary:
-			AppearanceInfo.CornerText = LOCTEXT("AppearanceCornerText_Macro", "MACRO");
-			break;
-		case BPTYPE_Interface:
-			AppearanceInfo.CornerText = LOCTEXT("AppearanceCornerText_Interface", "INTERFACE");
-			break;
-		default:
-			AppearanceInfo.CornerText = LOCTEXT("AppearanceCornerText_Blueprint", "BLUEPRINT");
-			break;
-		}
+	switch (Blueprint->BlueprintType)
+	{
+	case BPTYPE_LevelScript:
+		AppearanceInfo.CornerText = LOCTEXT("AppearanceCornerText_LevelScript", "LEVEL BLUEPRINT");
+		break;
+	case BPTYPE_MacroLibrary:
+		AppearanceInfo.CornerText = LOCTEXT("AppearanceCornerText_Macro", "MACRO");
+		break;
+	case BPTYPE_Interface:
+		AppearanceInfo.CornerText = LOCTEXT("AppearanceCornerText_Interface", "INTERFACE");
+		break;
+	default:
+		AppearanceInfo.CornerText = LOCTEXT("AppearanceCornerText_Blueprint", "BLUEPRINT");
+		break;
+	}
 	}
 
 	UEdGraph const* EditingGraph = GetFocusedGraph();
@@ -1482,7 +1473,7 @@ void FBlueprintEditor::InitBlueprintEditor(const EToolkitMode::Type Mode, const
 	PostLayoutBlueprintEditorInitialization();
 
 	// Find and set any instances of this blueprint type if any exists and we are not already editing one
-	FBlueprintEditorUtils::FindAndSetDebuggableBlueprintInstances();	
+	FBlueprintEditorUtils::FindAndSetDebuggableBlueprintInstances();
 }
 
 void FBlueprintEditor::RegisterApplicationModes(const TArray<UBlueprint*>& InBlueprints, bool bShouldOpenInDefaultsMode, bool bNewlyCreated/* = false*/)
@@ -1511,11 +1502,11 @@ void FBlueprintEditor::RegisterApplicationModes(const TArray<UBlueprint*>& InBlu
 		}
 		else if ( SingleBP->BlueprintType == BPTYPE_FunctionLibrary )
 		{
-			AddApplicationMode(
-				FBlueprintEditorApplicationModes::StandardBlueprintEditorMode,
-				MakeShareable(new FBlueprintEditorApplicationMode(SharedThis(this), FBlueprintEditorApplicationModes::StandardBlueprintEditorMode, FBlueprintEditorApplicationModes::GetLocalizedMode)));
-			SetCurrentMode(FBlueprintEditorApplicationModes::StandardBlueprintEditorMode);
-		}
+				AddApplicationMode(
+					FBlueprintEditorApplicationModes::StandardBlueprintEditorMode,
+					MakeShareable(new FBlueprintEditorApplicationMode(SharedThis(this), FBlueprintEditorApplicationModes::StandardBlueprintEditorMode, FBlueprintEditorApplicationModes::GetLocalizedMode)));
+				SetCurrentMode(FBlueprintEditorApplicationModes::StandardBlueprintEditorMode);
+			}
 		else
 		{
 			if ( GetDefault<UEditorExperimentalSettings>()->bUnifiedBlueprintEditor )
@@ -1985,15 +1976,15 @@ void FBlueprintEditor::CreateDefaultTabContents(const TArray<UBlueprint*>& InBlu
 		InBlueprint->ParentClass->IsChildOf(AActor::StaticClass()) && 
 		InBlueprint->SimpleConstructionScript )
 	{
-		SCSEditor = SAssignNew(SCSEditor, SSCSEditor)
-			.ActorContext(this, &FBlueprintEditor::GetSCSEditorActorContext)
-			.PreviewActor(this, &FBlueprintEditor::GetPreviewActor)
-			.AllowEditing(this, &FBlueprintEditor::InEditingMode)
+	SCSEditor = SAssignNew(SCSEditor, SSCSEditor)
+		.ActorContext(this, &FBlueprintEditor::GetSCSEditorActorContext)
+		.PreviewActor(this, &FBlueprintEditor::GetPreviewActor)
+		.AllowEditing(this, &FBlueprintEditor::InEditingMode)
 			.OnTreeViewSelectionChanged(this, &FBlueprintEditor::OnSCSEditorTreeViewSelectionChanged)
 			.OnUpdateSelectionFromNodes(this, &FBlueprintEditor::OnSCSEditorUpdateSelectionFromNodes);
-		SCSViewport = SAssignNew(SCSViewport, SSCSEditorViewport)
-			.BlueprintEditor(SharedThis(this));
-	}
+	SCSViewport = SAssignNew(SCSViewport, SSCSEditorViewport)
+		.BlueprintEditor(SharedThis(this));
+}
 }
 
 void FBlueprintEditor::OnLogTokenClicked(const TSharedRef<IMessageToken>& Token)
@@ -2404,11 +2395,11 @@ void FBlueprintEditor::EditGlobalOptions_Clicked()
 		if (CurrentUISelection == FBlueprintEditor::GraphPanel)
 		{
 			ClearSelectionInAllEditors();
-		}
+}
 		if (CurrentUISelection == FBlueprintEditor::MyBlueprint)
-		{
+{
 			if (MyBlueprintWidget.IsValid())
-			{
+	{
 				MyBlueprintWidget->ClearGraphActionMenuSelection();
 			}
 		}
@@ -2688,7 +2679,7 @@ void FBlueprintEditor::OnSelectedNodesChanged(const FGraphPanelSelectionSet& New
 	{
 		// clear MyBlueprint selection
 		if (MyBlueprintWidget.IsValid())
-		{
+	{
 			MyBlueprintWidget->ClearGraphActionMenuSelection();
 		}
 	}
diff --git a/Engine/Source/Editor/Kismet/Private/SCSEditorViewportClient.cpp b/Engine/Source/Editor/Kismet/Private/SCSEditorViewportClient.cpp
--- a/Engine/Source/Editor/Kismet/Private/SCSEditorViewportClient.cpp
+++ b/Engine/Source/Editor/Kismet/Private/SCSEditorViewportClient.cpp
@@ -107,12 +107,12 @@ namespace
 
 FSCSEditorViewportClient::FSCSEditorViewportClient(TWeakPtr<FBlueprintEditor>& InBlueprintEditorPtr, FPreviewScene& InPreviewScene, const TSharedRef<SSCSEditorViewport>& InSCSEditorViewport)
 	: FEditorViewportClient(nullptr, &InPreviewScene, StaticCastSharedRef<SEditorViewport>(InSCSEditorViewport))
-	,BlueprintEditorPtr(InBlueprintEditorPtr)
-	,PreviewBlueprint(NULL)
-	,PreviewActorBounds(ForceInitToZero)
-	,bIsManipulating(false)
-	,ScopedTransaction(NULL)
-	,bIsSimulateEnabled(false)
+	, BlueprintEditorPtr(InBlueprintEditorPtr)
+	, PreviewBlueprint(NULL)
+	, PreviewActorBounds(ForceInitToZero)
+	, bIsManipulating(false)
+	, ScopedTransaction(NULL)
+	, bIsSimulateEnabled(false)
 {
 	WidgetMode = FWidget::WM_Translate;
 	WidgetCoordSystem = COORD_Local;
@@ -370,9 +370,10 @@ bool FSCSEditorViewportClient::InputWidgetDelta( FViewport* Viewport, EAxisList:
 	{
 		bHandled = true;
 		AActor* PreviewActor = GetPreviewActor();
-		if(PreviewActor)
+		auto BlueprintEditor = BlueprintEditorPtr.Pin();
+		if (PreviewActor && BlueprintEditor.IsValid())
 		{
-			TArray<FSCSEditorTreeNodePtrType> SelectedNodes = BlueprintEditorPtr.Pin()->GetSelectedSCSEditorTreeNodes();
+			TArray<FSCSEditorTreeNodePtrType> SelectedNodes = BlueprintEditor->GetSelectedSCSEditorTreeNodes();
 			if(SelectedNodes.Num() > 0)
 			{
 				FVector ModifiedScale = Scale;
@@ -387,13 +388,13 @@ bool FSCSEditorViewportClient::InputWidgetDelta( FViewport* Viewport, EAxisList:
 				{
 					FSCSEditorTreeNodePtrType SelectedNodePtr = *It;
 					// Don't allow editing of a root node, inherited SCS node or child node that also has a movable (non-root) parent node selected
-					const bool bCanEdit =  !SelectedNodePtr->IsRoot() && !SelectedNodePtr->IsInherited()
+					const bool bCanEdit = !SelectedNodePtr->IsRoot() && !SelectedNodePtr->IsInherited()
 						&& !IsMovableParentNodeSelected(SelectedNodePtr, SelectedNodes);
 
 					if(bCanEdit)
 					{
 						USceneComponent* SceneComp = Cast<USceneComponent>(SelectedNodePtr->FindComponentInstanceInActor(PreviewActor));
-						USceneComponent* SelectedTemplate = Cast<USceneComponent>(SelectedNodePtr->GetComponentTemplate());
+						USceneComponent* SelectedTemplate = Cast<USceneComponent>(SelectedNodePtr->GetEditableComponentTemplate(BlueprintEditor->GetBlueprintObj()));
 						if(SceneComp != NULL && SelectedTemplate != NULL)
 						{
 							// Cache the current default values for propagation
@@ -416,10 +417,10 @@ bool FSCSEditorViewportClient::InputWidgetDelta( FViewport* Viewport, EAxisList:
 									Rot = (ParentToWorldSpace.Inverse().GetRotation() * Rot.Quaternion() * ParentToWorldSpace.GetRotation()).Rotator();
 								}
 							}
-
+							
 							FComponentEditorUtils::FTransformData OldDefaultTransform(*SelectedTemplate);
 
-							TSharedPtr<ISCSEditorCustomization> Customization = BlueprintEditorPtr.Pin()->CustomizeSCSEditor(SceneComp);
+							TSharedPtr<ISCSEditorCustomization> Customization = BlueprintEditor->CustomizeSCSEditor(SceneComp);
 							if(Customization.IsValid() && Customization->HandleViewportDrag(SceneComp, SelectedTemplate, Drag, Rot, ModifiedScale, GetWidgetLocation()))
 							{
 								UpdatedComponents.Add(SceneComp);
@@ -566,7 +567,7 @@ FWidget::EWidgetMode FSCSEditorViewportClient::GetWidgetMode() const
 		const TSharedPtr<FBlueprintEditor> BluePrintEditor = BlueprintEditorPtr.Pin();
 		if ( BluePrintEditor.IsValid() )
 		{
-			TArray<FSCSEditorTreeNodePtrType> SelectedNodes = BlueprintEditorPtr.Pin()->GetSelectedSCSEditorTreeNodes();
+			TArray<FSCSEditorTreeNodePtrType> SelectedNodes = BluePrintEditor->GetSelectedSCSEditorTreeNodes();
 			const TArray<FSCSEditorTreeNodePtrType>& RootNodes = BluePrintEditor->GetSCSEditor()->RootNodes;
 
 			// if the selected nodes array is empty, or only contains entries from the
@@ -574,7 +575,9 @@ FWidget::EWidgetMode FSCSEditorViewportClient::GetWidgetMode() const
 			for ( int32 CurrentNodeIndex=0; CurrentNodeIndex < SelectedNodes.Num(); CurrentNodeIndex++ )
 			{
 				FSCSEditorTreeNodePtrType CurrentNodePtr = SelectedNodes[CurrentNodeIndex];
-				if ( CurrentNodePtr.IsValid() && !RootNodes.Contains( CurrentNodePtr ) && !CurrentNodePtr->IsRoot() && CurrentNodePtr->CanEditDefaults() && CurrentNodePtr->FindComponentInstanceInActor(PreviewActor))
+				if (CurrentNodePtr.IsValid() && !RootNodes.Contains(CurrentNodePtr) && !CurrentNodePtr->IsRoot() && 
+					CurrentNodePtr->GetEditableComponentTemplate(BluePrintEditor->GetBlueprintObj()) &&
+					CurrentNodePtr->FindComponentInstanceInActor(PreviewActor))
 				{
 					// a non-NULL, non-root item is selected, draw the widget
 					ReturnWidgetMode = WidgetMode;
@@ -675,7 +678,7 @@ void FSCSEditorViewportClient::InvalidatePreview(bool bResetCamera)
 	{
 		return;
 	}
-
+	
 	UBlueprint* Blueprint = BlueprintEditorPtr.Pin()->GetBlueprintObj();
 	check(Blueprint);
 
@@ -837,10 +840,10 @@ void FSCSEditorViewportClient::BeginTransaction(const FText& Description)
 			FBlueprintEditorUtils::MarkBlueprintAsModified(PreviewBlueprint);
 		}
 
-		AActor* PreviewActor = GetPreviewActor();
-		TArray<FSCSEditorTreeNodePtrType> SelectedNodes = BlueprintEditorPtr.Pin()->GetSelectedSCSEditorTreeNodes();
-		if(SelectedNodes.Num() > 0)
+		auto BlueprintEditor = BlueprintEditorPtr.Pin();
+		if (BlueprintEditor.IsValid())
 		{
+			TArray<FSCSEditorTreeNodePtrType> SelectedNodes = BlueprintEditor->GetSelectedSCSEditorTreeNodes();
 			for(auto SelectedSCSNodeIter(SelectedNodes.CreateIterator()); SelectedSCSNodeIter; ++SelectedSCSNodeIter)
 			{
 				FSCSEditorTreeNodePtrType Node = *SelectedSCSNodeIter;
@@ -852,13 +855,14 @@ void FSCSEditorViewportClient::BeginTransaction(const FText& Description)
 					}
 
 					// Modify both the component template and the instance in the preview actor (provided there is one)
-					UActorComponent* ComponentTemplate = Node->GetComponentTemplate();
+					UActorComponent* ComponentTemplate = Node->GetEditableComponentTemplate(BlueprintEditor->GetBlueprintObj());
 					if (ComponentTemplate != nullptr)
 					{
 						ComponentTemplate->SetFlags(RF_Transactional);
 						ComponentTemplate->Modify();
 					}
 
+					AActor* PreviewActor = GetPreviewActor();
 					if (PreviewActor)
 					{
 						UActorComponent* ComponentPreviewInstance = Node->FindComponentInstanceInActor(PreviewActor);
diff --git a/Engine/Source/Editor/Kismet/Private/SSCSEditor.cpp b/Engine/Source/Editor/Kismet/Private/SSCSEditor.cpp
--- a/Engine/Source/Editor/Kismet/Private/SSCSEditor.cpp
+++ b/Engine/Source/Editor/Kismet/Private/SSCSEditor.cpp
@@ -275,6 +275,22 @@ UActorComponent* FSCSEditorTreeNode::GetComponentTemplate() const
 	return ComponentTemplatePtr.Get();
 }
 
+UActorComponent* FSCSEditorTreeNode::GetEditableComponentTemplate(UBlueprint* ActualEditedBlueprint)
+{
+	if (CanEditDefaults())
+	{
+		return GetComponentTemplate();
+	}
+	
+	auto SCSNode = GetSCSNode();
+	if (!IsNative() && IsInherited() && SCSNode && ActualEditedBlueprint)
+	{
+		return GetOverridenComponentTemplate(ActualEditedBlueprint, true);
+	}
+
+	return nullptr;
+}
+
 FName FSCSEditorTreeNode::GetVariableName() const
 {
 	FName VariableName = NAME_None;
diff --git a/Engine/Source/Editor/Kismet/Public/SSCSEditor.h b/Engine/Source/Editor/Kismet/Public/SSCSEditor.h
--- a/Engine/Source/Editor/Kismet/Public/SSCSEditor.h
+++ b/Engine/Source/Editor/Kismet/Public/SSCSEditor.h
@@ -55,6 +55,12 @@ public:
 	 * @return The component template or instance represented by this object.
 	 */
 	UActorComponent* GetComponentTemplate() const;
+	/**
+	 * @param ActualEditedBlueprint currently edited blueprint
+	 *
+	 * @return The component template that can be editable for actual class.
+	 */
+	UActorComponent* GetEditableComponentTemplate(UBlueprint* ActualEditedBlueprint);
 	/**
 	 * Finds the component instance represented by this node contained within a given Actor instance.
 	 *
