[INTEGRATE] Change 2411020 by [REDACTED_PERSON][REDACTED_HANDLE] on 2015/01/19 12:06:22

	If the root component is instanced and has already been added recursively by one of its children, don't try to add it to the tree again.

diff --git a/Engine/Source/Editor/Kismet/Private/SSCSEditor.cpp b/Engine/Source/Editor/Kismet/Private/SSCSEditor.cpp
--- a/Engine/Source/Editor/Kismet/Private/SSCSEditor.cpp
+++ b/Engine/Source/Editor/Kismet/Private/SSCSEditor.cpp
@@ -3783,7 +3783,7 @@ FSCSEditorTreeNodePtrType SSCSEditor::AddTreeNode(USceneComponent* InSceneCompon
 			// Expand the scene root node by default
 			SCSTreeWidget->SetItemExpansion(SceneRootNodePtr, true);
 		}
-		else
+		else if (SceneRootNodePtr->GetComponentTemplate() != InSceneComponent)
 		{
 			NewNodePtr = SceneRootNodePtr->AddChild(InSceneComponent);
 		}
