[INTEGRATE] Change 2411381 by [REDACTED_PERSON][REDACTED_HANDLE] on 2015/01/19 16:39:31

	Editor - Adding the unified blueprint editor view.  Combines all the modes into a single blueprint editor.  Still requires being enabled through the experimental option in the editor.

diff --git a/Engine/Source/Editor/BlueprintGraph/Private/BlueprintEditorSettings.cpp b/Engine/Source/Editor/BlueprintGraph/Private/BlueprintEditorSettings.cpp
--- a/Engine/Source/Editor/BlueprintGraph/Private/BlueprintEditorSettings.cpp
+++ b/Engine/Source/Editor/BlueprintGraph/Private/BlueprintEditorSettings.cpp
@@ -18,6 +18,8 @@ UBlueprintEditorSettings::UBlueprintEditorSettings(const FObjectInitializer& Obj
 	, bFavorPureCastNodes(false)
 	, bAutoCastObjectConnections(false)
 	, bUseLegacyMenuingSystem(false)
+	, bShowInheritedVariables(false)
+	, bShowEmptySections(true)
 	// Compiler Settings
 	, SaveOnCompile(SoC_Never)
 	, bJumpToNodeErrors(false)
diff --git a/Engine/Source/Editor/BlueprintGraph/Public/BlueprintEditorSettings.h b/Engine/Source/Editor/BlueprintGraph/Public/BlueprintEditorSettings.h
--- a/Engine/Source/Editor/BlueprintGraph/Public/BlueprintEditorSettings.h
+++ b/Engine/Source/Editor/BlueprintGraph/Public/BlueprintEditorSettings.h
@@ -59,6 +59,14 @@ public:
 	UPROPERTY(EditAnywhere, AdvancedDisplay, config, Category=Workflow)
 	bool bUseLegacyMenuingSystem;
 
+	/** If set we'll show the inherited variables in the My Blueprint view. */
+	UPROPERTY(config)
+	bool bShowInheritedVariables;
+
+	/** If set we'll show empty sections in the My Blueprint view. */
+	UPROPERTY(config)
+	bool bShowEmptySections;
+
 // Compiler Settings
 public:
 	/** Determines when to save Blueprints post-compile */
diff --git a/Engine/Source/Editor/EditorStyle/Private/SlateEditorStyle.cpp b/Engine/Source/Editor/EditorStyle/Private/SlateEditorStyle.cpp
--- a/Engine/Source/Editor/EditorStyle/Private/SlateEditorStyle.cpp
+++ b/Engine/Source/Editor/EditorStyle/Private/SlateEditorStyle.cpp
@@ -318,6 +318,14 @@ void FSlateEditorStyle::FStyle::SetupGeneralStyles()
 			.SetPressed(BOX_BRUSH( "Common/RoundedSelection_16x", 4.0f/16.0f, SelectionColor_Pressed ))
 		);
 
+		//FSlateColorBrush(FLinearColor::White)
+
+		Set("RoundButton", FButtonStyle(Button)
+			.SetNormal(BOX_BRUSH("Common/RoundedSelection_16x", 4.0f / 16.0f, FLinearColor(1, 1, 1, 0.1f)))
+			.SetHovered(BOX_BRUSH("Common/RoundedSelection_16x", 4.0f / 16.0f, SelectionColor))
+			.SetPressed(BOX_BRUSH("Common/RoundedSelection_16x", 4.0f / 16.0f, SelectionColor_Pressed))
+			);
+
 		/* Create a checkbox style for "ToggleButton" ... */
 		const FCheckBoxStyle ToggleButtonStyle = FCheckBoxStyle()
 			.SetCheckBoxType(ESlateCheckBoxType::ToggleButton)
@@ -718,11 +726,68 @@ void FSlateEditorStyle::FStyle::SetupGeneralStyles()
 			.SetTopShadowBrush( IMAGE_BRUSH( "Common/ScrollBoxShadowTop", FVector2D(64,8) ) )
 			.SetBottomShadowBrush( IMAGE_BRUSH( "Common/ScrollBoxShadowBottom", FVector2D(64,8) ) )
 			);
-	}
+	}//
+
+
+	NormalTableRowStyle = FTableRowStyle()
+		.SetEvenRowBackgroundBrush(FSlateNoResource())
+		.SetEvenRowBackgroundHoveredBrush(IMAGE_BRUSH("Common/Selection", Icon8x8, FLinearColor(1.0f, 1.0f, 1.0f, 0.1f)))
+		.SetOddRowBackgroundBrush(FSlateNoResource())
+		.SetOddRowBackgroundHoveredBrush(IMAGE_BRUSH("Common/Selection", Icon8x8, FLinearColor(1.0f, 1.0f, 1.0f, 0.1f)))
+		.SetSelectorFocusedBrush(BORDER_BRUSH("Common/Selector", FMargin(4.f / 16.f), SelectorColor))
+		.SetActiveBrush(IMAGE_BRUSH("Common/Selection", Icon8x8, SelectionColor))
+		.SetActiveHoveredBrush(IMAGE_BRUSH("Common/Selection", Icon8x8, SelectionColor))
+		.SetInactiveBrush(IMAGE_BRUSH("Common/Selection", Icon8x8, SelectionColor_Inactive))
+		.SetInactiveHoveredBrush(IMAGE_BRUSH("Common/Selection", Icon8x8, SelectionColor_Inactive))
+		.SetTextColor(DefaultForeground)
+		.SetSelectedTextColor(InvertedForeground);
 
 	// Lists, Trees
 	{
 		Set( "TableView.Row", FTableRowStyle( NormalTableRowStyle) );
+		Set( "TableView.DarkRow",FTableRowStyle( NormalTableRowStyle)
+			.SetEvenRowBackgroundBrush(IMAGE_BRUSH("PropertyView/DetailCategoryMiddle", FVector2D(16, 16)))
+			.SetEvenRowBackgroundHoveredBrush(IMAGE_BRUSH("PropertyView/DetailCategoryMiddle_Hovered", FVector2D(16, 16)))
+			.SetOddRowBackgroundBrush(IMAGE_BRUSH("PropertyView/DetailCategoryMiddle", FVector2D(16, 16)))
+			.SetOddRowBackgroundHoveredBrush(IMAGE_BRUSH("PropertyView/DetailCategoryMiddle_Hovered", FVector2D(16, 16)))
+			.SetSelectorFocusedBrush(BORDER_BRUSH("Common/Selector", FMargin(4.f / 16.f), SelectorColor))
+			.SetActiveBrush(IMAGE_BRUSH("Common/Selection", Icon8x8, SelectionColor))
+			.SetActiveHoveredBrush(IMAGE_BRUSH("Common/Selection", Icon8x8, SelectionColor))
+			.SetInactiveBrush(IMAGE_BRUSH("Common/Selection", Icon8x8, SelectionColor_Inactive))
+			.SetInactiveHoveredBrush(IMAGE_BRUSH("Common/Selection", Icon8x8, SelectionColor_Inactive))
+		);
+
+		//Set("DetailsView.TreeView.TableRow", FTableRowStyle()
+		//	.SetEvenRowBackgroundBrush(FSlateNoResource())
+		//	.SetEvenRowBackgroundHoveredBrush(FSlateNoResource())
+		//	.SetOddRowBackgroundBrush(FSlateNoResource())
+		//	.SetOddRowBackgroundHoveredBrush(FSlateNoResource())
+		//	.SetSelectorFocusedBrush(FSlateNoResource())
+		//	.SetActiveBrush(FSlateNoResource())
+		//	.SetActiveHoveredBrush(FSlateNoResource())
+		//	.SetInactiveBrush(FSlateNoResource())
+		//	.SetInactiveHoveredBrush(FSlateNoResource())
+		//	.SetTextColor(DefaultForeground)
+		//	.SetSelectedTextColor(InvertedForeground)
+		//	);
+
+		//Set("DetailsView.CategoryMiddle", new IMAGE_BRUSH("PropertyView/DetailCategoryMiddle", FVector2D(16, 16)));
+		//Set("DetailsView.CategoryMiddle_Hovered", new IMAGE_BRUSH("PropertyView/DetailCategoryMiddle_Hovered", FVector2D(16, 16)));
+		//Set("DetailsView.CategoryMiddle_Highlighted", new BOX_BRUSH("Common/TextBox_Special_Active", FMargin(8.0f / 32.0f)));
+
+		//if ( IsHighlighted() )
+		//{
+		//	return FEditorStyle::GetBrush("DetailsView.CategoryMiddle_Highlighted");
+		//}
+		//else if ( IsHovered() )
+		//{
+		//	return FEditorStyle::GetBrush("DetailsView.CategoryMiddle_Hovered");
+		//}
+		//else
+		//{
+		//	return FEditorStyle::GetBrush("DetailsView.CategoryMiddle");
+		//}
+
 		Set( "TreeArrow_Collapsed", new IMAGE_BRUSH( "Common/TreeArrow_Collapsed", Icon10x10, DefaultForeground ) );
 		Set( "TreeArrow_Collapsed_Hovered", new IMAGE_BRUSH( "Common/TreeArrow_Collapsed_Hovered", Icon10x10, DefaultForeground ) );
 		Set( "TreeArrow_Expanded", new IMAGE_BRUSH( "Common/TreeArrow_Expanded", Icon10x10, DefaultForeground ) );
@@ -4515,6 +4580,9 @@ void FSlateEditorStyle::FStyle::SetupPersonaStyle()
 		Set( "FullBlueprintEditor.EditGlobalOptions", new IMAGE_BRUSH( "Icons/icon_Blueprint_Options_40px", Icon40x40));
 		Set( "FullBlueprintEditor.EditGlobalOptions.Small", new IMAGE_BRUSH( "Icons/icon_Blueprint_Options_40px", Icon20x20));
 
+		Set("FullBlueprintEditor.EditClassDefaults", new IMAGE_BRUSH("Icons/icon_BlueprintEditor_Defaults_40x", Icon40x40));
+		Set("FullBlueprintEditor.EditClassDefaults.Small", new IMAGE_BRUSH("Icons/icon_BlueprintEditor_Defaults_40x", Icon20x20));
+
 		Set( "BlueprintEditor.Details.DeleteButton", new IMAGE_BRUSH( "/Icons/GenericDelete_Black", Icon16x16 ) );
 
 		Set( "BlueprintEditor.Details.ArgUpButton", new IMAGE_BRUSH( "/Icons/icon_FunctionArgUp", Icon16x16 ) );
diff --git a/Engine/Source/Editor/GraphEditor/Private/GraphActionNode.cpp b/Engine/Source/Editor/GraphEditor/Private/GraphActionNode.cpp
--- a/Engine/Source/Editor/GraphEditor/Private/GraphActionNode.cpp
+++ b/Engine/Source/Editor/GraphEditor/Private/GraphActionNode.cpp
@@ -190,10 +190,26 @@ TSharedPtr<FGraphActionNode> FGraphActionNode::AddChild(FGraphActionListBuilderB
 
 	TSharedPtr<FGraphActionNode> ActionNode = FGraphActionNode::NewActionNode(ActionSet.Actions);
 	AddChildRecursively(CategoryStack, ActionNode);
-
+	
 	return ActionNode;
 }
 
+//------------------------------------------------------------------------------
+TSharedPtr<FGraphActionNode> FGraphActionNode::AddSection(int32 Grouping, int32 InSectionID)
+{
+	if ( !ChildSections.Contains(InSectionID) )
+	{
+		ChildSections.Add(InSectionID);
+
+		TSharedPtr<FGraphActionNode> Section = NewSectionHeadingNode(SharedThis(this), Grouping, InSectionID);
+		InsertChild(Section);
+
+		return Section;
+	}
+
+	return nullptr;
+}
+
 //------------------------------------------------------------------------------
 void FGraphActionNode::SortChildren(bool bAlphabetically/* = true*/, bool bRecursive/* = true*/)
 {
@@ -435,28 +451,45 @@ TSharedPtr<FGraphActionNode> FGraphActionNode::NewGroupDividerNode(TWeakPtr<FGra
 //------------------------------------------------------------------------------
 void FGraphActionNode::AddChildRecursively(TArray<FString>& CategoryStack, TSharedPtr<FGraphActionNode> NodeToAdd)
 {
-	if (CategoryStack.Num() > 0)
+	TSharedPtr<FGraphActionNode> FoundSectionNode;
+	for ( TSharedPtr<FGraphActionNode> const& ChildNode : Children )
 	{
-		FString CategorySection = CategoryStack[0];
-		CategoryStack.RemoveAt(0, 1);
+		if ( NodeToAdd->SectionID == ChildNode->SectionID && ChildNode->IsSectionHeadingNode() )
+		{
+			FoundSectionNode = ChildNode;
+			break;
+		}
+	}
 
-		// make sure we don't already have a child that this can nest under
-		TSharedPtr<FGraphActionNode> ExistingNode = FindMatchingParent(CategorySection, NodeToAdd);
-		if (ExistingNode.IsValid())
+	if ( FoundSectionNode.IsValid() )
+	{
+		FoundSectionNode->AddChildRecursively(CategoryStack, NodeToAdd);
+	}
+	else
+	{
+		if ( CategoryStack.Num() > 0 )
 		{
-			ExistingNode->AddChildRecursively(CategoryStack, NodeToAdd);
+			FString CategorySection = CategoryStack[0];
+			CategoryStack.RemoveAt(0, 1);
+
+			// make sure we don't already have a child that this can nest under
+			TSharedPtr<FGraphActionNode> ExistingNode = FindMatchingParent(CategorySection, NodeToAdd);
+			if ( ExistingNode.IsValid() )
+			{
+				ExistingNode->AddChildRecursively(CategoryStack, NodeToAdd);
+			}
+			else
+			{
+				TSharedPtr<FGraphActionNode> CategoryNode = NewCategoryNode(CategorySection, NodeToAdd->Grouping, NodeToAdd->SectionID);
+				InsertChild(CategoryNode);
+				CategoryNode->AddChildRecursively(CategoryStack, NodeToAdd);
+			}
 		}
 		else
 		{
-			TSharedPtr<FGraphActionNode> CategoryNode = NewCategoryNode(CategorySection, NodeToAdd->Grouping, NodeToAdd->SectionID);
-			InsertChild(CategoryNode);
-			CategoryNode->AddChildRecursively(CategoryStack, NodeToAdd);
+			InsertChild(NodeToAdd);
 		}
 	}
-	else
-	{
-		InsertChild(NodeToAdd);
-	}
 }
 
 //------------------------------------------------------------------------------
@@ -503,7 +536,7 @@ TSharedPtr<FGraphActionNode> FGraphActionNode::FindMatchingParent(FString const&
 void FGraphActionNode::InsertChild(TSharedPtr<FGraphActionNode> NodeToAdd)
 {
 	ensure(!NodeToAdd->IsRootNode());
-	ensure(!IsSeparator());
+	//ensure(!IsSeparator());
 
 	NodeToAdd->ParentNode = this->AsShared();
 
@@ -519,7 +552,14 @@ void FGraphActionNode::InsertChild(TSharedPtr<FGraphActionNode> NodeToAdd)
 		if (bAddSectionHeading)
 		{
 			ChildSections.Add(NodeToAdd->SectionID); // to avoid recursion, add before we insert
-			InsertChild(NewSectionHeadingNode(NodeToAdd->ParentNode, NodeToAdd->Grouping, NodeToAdd->SectionID));
+			//InsertChild(NewSectionHeadingNode(NodeToAdd->ParentNode, NodeToAdd->Grouping, NodeToAdd->SectionID));
+
+			TSharedPtr<FGraphActionNode> NewSection = NewSectionHeadingNode(NodeToAdd->ParentNode, NodeToAdd->Grouping, NodeToAdd->SectionID);
+			InsertChild(NewSection);
+
+			NodeToAdd->InsertOrder = NewSection->Children.Num();
+			NewSection->Children.Add(NodeToAdd);
+			return;
 		}
 	}
 	// we don't use group-dividers inside of sections (we use groups to more to
diff --git a/Engine/Source/Editor/GraphEditor/Private/GraphActionNode.h b/Engine/Source/Editor/GraphEditor/Private/GraphActionNode.h
--- a/Engine/Source/Editor/GraphEditor/Private/GraphActionNode.h
+++ b/Engine/Source/Editor/GraphEditor/Private/GraphActionNode.h
@@ -42,6 +42,8 @@ public:
 	 */
 	TSharedPtr<FGraphActionNode> AddChild(FGraphActionListBuilderBase::ActionGroup const& ActionSet);
 
+	TSharedPtr<FGraphActionNode> AddSection(int32 Grouping, int32 InSectionID);
+
 	/**
 	 * Sorts all child nodes by section, group, and type (additionally, can
 	 * sort alphabetically if wanted).
diff --git a/Engine/Source/Editor/GraphEditor/Private/SGraphActionMenu.cpp b/Engine/Source/Editor/GraphEditor/Private/SGraphActionMenu.cpp
--- a/Engine/Source/Editor/GraphEditor/Private/SGraphActionMenu.cpp
+++ b/Engine/Source/Editor/GraphEditor/Private/SGraphActionMenu.cpp
@@ -12,6 +12,76 @@
 
 #define LOCTEXT_NAMESPACE "GraphActionMenu"
 
+//////////////////////////////////////////////////////////////////////////
+
+template<typename ItemType>
+class SCategoryHeaderTableRow : public STableRow < ItemType >
+{
+public:
+	SLATE_BEGIN_ARGS(SCategoryHeaderTableRow)
+	{}
+		SLATE_DEFAULT_SLOT(typename SCategoryHeaderTableRow::FArguments, Content)
+	SLATE_END_ARGS()
+
+	void Construct(const FArguments& InArgs, const TSharedRef<STableViewBase>& InOwnerTableView)
+	{
+		ChildSlot
+		.Padding(0.0f, 2.0f, 0.0f, 0.0f)
+		[
+			SAssignNew(ContentBorder, SBorder)
+			.BorderImage(this, &SCategoryHeaderTableRow::GetBackgroundImage)
+			.Padding(FMargin(0.0f, 3.0f))
+			.BorderBackgroundColor(FLinearColor(.6, .6, .6, 1.0f))
+			[
+				SNew(SHorizontalBox)
+
+				+ SHorizontalBox::Slot()
+				.VAlign(VAlign_Center)
+				.Padding(2.0f, 2.0f, 2.0f, 2.0f)
+				.AutoWidth()
+				[
+					SNew(SExpanderArrow, SharedThis(this))
+				]
+
+				+ SHorizontalBox::Slot()
+				.VAlign(VAlign_Center)
+				[
+					InArgs._Content.Widget
+				]
+			]
+		];
+
+		STableRow < ItemType >::ConstructInternal(
+			STableRow::FArguments()
+			.Style(FEditorStyle::Get(), "DetailsView.TreeView.TableRow")
+			.ShowSelection(false),
+			InOwnerTableView
+			);
+	}
+
+	const FSlateBrush* GetBackgroundImage() const
+	{
+		if ( IsHovered() )
+		{
+			return IsItemExpanded() ? FEditorStyle::GetBrush("DetailsView.CategoryTop_Hovered") : FEditorStyle::GetBrush("DetailsView.CollapsedCategory_Hovered");
+		}
+		else
+		{
+			return IsItemExpanded() ? FEditorStyle::GetBrush("DetailsView.CategoryTop") : FEditorStyle::GetBrush("DetailsView.CollapsedCategory");
+		}
+	}
+
+	virtual void SetContent(TSharedRef< SWidget > InContent) override
+	{
+		ContentBorder->SetContent(InContent);
+	}
+
+private:
+	TSharedPtr<SBorder> ContentBorder;
+};
+
+//////////////////////////////////////////////////////////////////////////
+
 namespace GraphActionMenuHelpers
 {
 	bool ActionMatchesName(const FEdGraphSchemaAction* InGraphAction, const FName& ItemName)
@@ -38,6 +108,8 @@ namespace GraphActionMenuHelpers
 	}
 }
 
+//////////////////////////////////////////////////////////////////////////
+
 void SDefaultGraphActionWidget::Construct(const FArguments& InArgs, const FCreateWidgetForActionData* InCreateData)
 {
 	ActionPtr = InCreateData->Action;
@@ -166,6 +238,7 @@ void SGraphActionMenu::Construct( const FArguments& InArgs, bool bIsReadOnly/* =
 {
 	this->SelectedSuggestion = INDEX_NONE;
 	this->bIgnoreUIUpdate = false;
+	this->bUseSectionStyling = InArgs._UseSectionStyling;
 
 	this->bAutoExpandActionMenu = InArgs._AutoExpandActionMenu;
 	this->bShowFilterTextBox = InArgs._ShowFilterTextBox;
@@ -177,9 +250,12 @@ void SGraphActionMenu::Construct( const FArguments& InArgs, bool bIsReadOnly/* =
 	this->OnCreateWidgetForAction = InArgs._OnCreateWidgetForAction;
 	this->OnCreateCustomRowExpander = InArgs._OnCreateCustomRowExpander;
 	this->OnCollectAllActions = InArgs._OnCollectAllActions;
+	this->OnCollectStaticSections = InArgs._OnCollectStaticSections;
 	this->OnCategoryTextCommitted = InArgs._OnCategoryTextCommitted;
 	this->OnCanRenameSelectedAction = InArgs._OnCanRenameSelectedAction;
 	this->OnGetSectionTitle = InArgs._OnGetSectionTitle;
+	this->OnGetSectionToolTip = InArgs._OnGetSectionToolTip;
+	this->OnGetSectionWidget = InArgs._OnGetSectionWidget;
 	this->FilteredRootAction = FGraphActionNode::NewRootNode();
 	this->OnActionMatchesName = InArgs._OnActionMatchesName;
 	
@@ -241,7 +317,7 @@ void SGraphActionMenu::RefreshAllActions(bool bPreserveExpansion, bool bHandleOn
 {
 	// Save Selection (of only the first selected thing)
 	TArray< TSharedPtr<FGraphActionNode> > SelectedNodes = TreeView->GetSelectedItems();
-	TSharedPtr<FGraphActionNode> SelectedAction = SelectedNodes.Num() > 0 ? SelectedNodes[0] : NULL;
+	TSharedPtr<FGraphActionNode> SelectedAction = SelectedNodes.Num() > 0 ? SelectedNodes[0] : nullptr;
 
 	AllActions.Empty();
 	OnCollectAllActions.ExecuteIfBound(AllActions);
@@ -265,6 +341,26 @@ void SGraphActionMenu::RefreshAllActions(bool bPreserveExpansion, bool bHandleOn
 	}
 }
 
+void SGraphActionMenu::GetSectionExpansion(TMap<int32, bool>& SectionExpansion) const
+{
+
+}
+
+void SGraphActionMenu::SetSectionExpansion(const TMap<int32, bool>& InSectionExpansion)
+{
+	for ( auto& PossibleSection : FilteredRootAction->Children )
+	{
+		if ( PossibleSection->IsSectionHeadingNode() )
+		{
+			const bool* IsExpanded = InSectionExpansion.Find(PossibleSection->SectionID);
+			if ( IsExpanded != nullptr )
+			{
+				TreeView->SetItemExpansion(PossibleSection, *IsExpanded);
+			}
+		}
+	}
+}
+
 TSharedRef<SEditableTextBox> SGraphActionMenu::GetFilterTextBox()
 {
 	return FilterTextBox.ToSharedRef();
@@ -512,6 +608,18 @@ void SGraphActionMenu::GenerateFilteredItems(bool bPreserveExpansion)
 	// Clear the filtered root action
 	FilteredRootAction->ClearChildren();
 
+	// Collect the list of always visible sections if any, and force the creation of those sections.
+	if ( OnCollectStaticSections.IsBound() )
+	{
+		TArray<int32> StaticSectionIDs;
+		OnCollectStaticSections.Execute(StaticSectionIDs);
+
+		for ( int32 i = 0; i < StaticSectionIDs.Num(); i++ )
+		{
+			FilteredRootAction->AddSection(0, StaticSectionIDs[i]);
+		}
+	}
+	
 	// Trim and sanitized the filter text (so that it more likely matches the action descriptions)
 	FString TrimmedFilterString = FText::TrimPrecedingAndTrailing(GetFilterText()).ToString();
 
@@ -806,18 +914,53 @@ void SGraphActionMenu::OnItemScrolledIntoView( TSharedPtr<FGraphActionNode> InAc
 
 TSharedRef<ITableRow> SGraphActionMenu::MakeWidget( TSharedPtr<FGraphActionNode> InItem, const TSharedRef<STableViewBase>& OwnerTable, bool bIsReadOnly )
 {
+	FText SectionTitle;
+	TSharedPtr<IToolTip> SectionToolTip;
+
+	if ( InItem->IsSectionHeadingNode() )
+	{
+		if ( OnGetSectionTitle.IsBound() )
+		{
+			SectionTitle = OnGetSectionTitle.Execute(InItem->SectionID);
+		}
+
+		if ( OnGetSectionToolTip.IsBound() )
+		{
+			SectionToolTip = OnGetSectionToolTip.Execute(InItem->SectionID);
+		}
+	}
+
 	// In the case of FGraphActionNodes that have multiple actions, all of the actions will
 	// have the same text as they will have been created at the same point - only the actual
 	// action itself will differ, which is why parts of this function only refer to InItem->Actions[0]
 	// rather than iterating over the array
 
 	// Create the widget but do not add any content, the widget is needed to pass the IsSelectedExclusively function down to the potential SInlineEditableTextBlock widget
-	TSharedPtr< STableRow< TSharedPtr<FGraphActionNode> > > TableRow = SNew(STableRow< TSharedPtr<FGraphActionNode> >, OwnerTable)
-		.OnDragDetected( this, &SGraphActionMenu::OnItemDragDetected )
-		.ShowSelection( !InItem->IsSeparator() );
+	TSharedPtr< STableRow< TSharedPtr<FGraphActionNode> > > TableRow;
+	
+	if ( InItem->IsSectionHeadingNode() )
+	{
+		TableRow = SNew(SCategoryHeaderTableRow< TSharedPtr<FGraphActionNode> >, OwnerTable)
+			.ToolTip(SectionToolTip);
+	}
+	else
+	{
+		const FTableRowStyle* Style = bUseSectionStyling ? &FEditorStyle::Get().GetWidgetStyle<FTableRowStyle>("TableView.DarkRow") : &FCoreStyle::Get().GetWidgetStyle<FTableRowStyle>("TableView.Row");
 
+		TableRow = SNew(STableRow< TSharedPtr<FGraphActionNode> >, OwnerTable)
+			.Style(Style)
+			.OnDragDetected(this, &SGraphActionMenu::OnItemDragDetected)
+			.ShowSelection(!InItem->IsSeparator());
+	}
+
+	TSharedPtr<SHorizontalBox> RowContainer;
+	TableRow->SetContent
+	( 
+		SAssignNew(RowContainer, SHorizontalBox)
+	);
 
 	TSharedPtr<SWidget> RowContent;
+	FMargin RowPadding = FMargin(0, 2);
 
 	if( InItem->IsActionNode() )
 	{
@@ -871,46 +1014,23 @@ TSharedRef<ITableRow> SGraphActionMenu::MakeWidget( TSharedPtr<FGraphActionNode>
 	}
 	else if( InItem->IsSeparator() )
 	{
+		RowPadding = FMargin(0);
+
 		FText SectionTitle;
 		if( OnGetSectionTitle.IsBound() )
 		{
 			SectionTitle = OnGetSectionTitle.Execute(InItem->SectionID);
 		}
 
-		if( SectionTitle.IsEmpty() == true )
-		{
-			RowContent = SNew( SVerticalBox )
-				.Visibility(EVisibility::HitTestInvisible)
-				+SVerticalBox::Slot()
-				.AutoHeight()
-				// Add some empty space before the line, and a tiny bit after it
-				.Padding( 0.0f, 5.f, 0.0f, 5.f )
-				[
-					SNew( SBorder )
-
-					// We'll use the border's padding to actually create the horizontal line
-					.Padding(FEditorStyle::GetMargin(TEXT("Menu.Separator.Padding")))
-
-					// Separator graphic
-					.BorderImage( FEditorStyle::GetBrush( TEXT( "Menu.Separator" ) ) )
-				];
-		}
-		else
+		if( SectionTitle.IsEmpty() )
 		{
 			RowContent = SNew( SVerticalBox )
 			.Visibility(EVisibility::HitTestInvisible)
-			+SVerticalBox::Slot()
+
+			+ SVerticalBox::Slot()
 			.AutoHeight()
-			.Padding( 0.0f, 2.f, 0.0f, 0.f )
-			[
-				SNew(STextBlock)
-				.Text( SectionTitle )
-				.TextStyle( FEditorStyle::Get(), TEXT("Menu.Heading") )
-			]
-			+SVerticalBox::Slot()
-			.AutoHeight()			
 			// Add some empty space before the line, and a tiny bit after it
-			.Padding( 0.0f, 2.f, 0.0f, 5.f )
+			.Padding( 0.0f, 5.f, 0.0f, 5.f )
 			[
 				SNew( SBorder )
 
@@ -919,15 +1039,30 @@ TSharedRef<ITableRow> SGraphActionMenu::MakeWidget( TSharedPtr<FGraphActionNode>
 
 				// Separator graphic
 				.BorderImage( FEditorStyle::GetBrush( TEXT( "Menu.Separator" ) ) )
-			];	
+			];
 		}
-	}
+		else
+		{
+			RowContent = SNew(SHorizontalBox)
 
-	TSharedPtr<SHorizontalBox> RowContainer;
-	TableRow->SetContent
-	( 
-		SAssignNew(RowContainer, SHorizontalBox)
-	);
+			+ SHorizontalBox::Slot()
+			.VAlign(VAlign_Center)
+			[
+				SNew(STextBlock)
+				.Text(SectionTitle)
+				.Font(FEditorStyle::GetFontStyle("DetailsView.CategoryFontStyle"))
+				.ShadowOffset(FVector2D(1.0f, 1.0f))
+			]
+
+			+ SHorizontalBox::Slot()
+			.VAlign(VAlign_Center)
+			.HAlign(HAlign_Right)
+			.Padding(FMargin(0,0,2,0))
+			[
+				OnGetSectionWidget.IsBound() ? OnGetSectionWidget.Execute(TableRow.ToSharedRef(), InItem->SectionID) : SNullWidget::NullWidget
+			];
+		}
+	}
 
 	TSharedPtr<SExpanderArrow> ExpanderWidget;
 	if (OnCreateCustomRowExpander.IsBound())
@@ -946,7 +1081,9 @@ TSharedRef<ITableRow> SGraphActionMenu::MakeWidget( TSharedPtr<FGraphActionNode>
 	}
 	else 
 	{
-		ExpanderWidget = SNew(SExpanderArrow, TableRow);
+		ExpanderWidget =
+			SNew(SExpanderArrow, TableRow)
+			.BaseIndentLevel(1);
 	}
 
 	RowContainer->AddSlot()
@@ -958,7 +1095,8 @@ TSharedRef<ITableRow> SGraphActionMenu::MakeWidget( TSharedPtr<FGraphActionNode>
 	];
 
 	RowContainer->AddSlot()
-		.FillWidth(1.0)
+	.FillWidth(1.0)
+	.Padding(RowPadding)
 	[
 		RowContent.ToSharedRef()
 	];
diff --git a/Engine/Source/Editor/GraphEditor/Private/SGraphPanel.cpp b/Engine/Source/Editor/GraphEditor/Private/SGraphPanel.cpp
--- a/Engine/Source/Editor/GraphEditor/Private/SGraphPanel.cpp
+++ b/Engine/Source/Editor/GraphEditor/Private/SGraphPanel.cpp
@@ -906,12 +906,7 @@ TSharedPtr<SWidget> SGraphPanel::SummonContextMenu(const FVector2D& WhereToSummo
 
 		FActionMenuContent FocusedContent = OnGetContextMenuFor.Execute(SpawnInfo);
 
-		TSharedRef<SWidget> MenuContent =
-			SNew( SBorder )
-			.BorderImage( FEditorStyle::GetBrush("Menu.Background") )
-			[
-				FocusedContent.Content
-			];
+		TSharedRef<SWidget> MenuContent = FocusedContent.Content;
 		
 		FSlateApplication::Get().PushMenu(
 			AsShared(),
diff --git a/Engine/Source/Editor/GraphEditor/Public/SGraphActionMenu.h b/Engine/Source/Editor/GraphEditor/Public/SGraphActionMenu.h
--- a/Engine/Source/Editor/GraphEditor/Public/SGraphActionMenu.h
+++ b/Engine/Source/Editor/GraphEditor/Public/SGraphActionMenu.h
@@ -81,23 +81,30 @@ public:
 	DECLARE_DELEGATE_RetVal_TwoParams( FReply, FOnCategoryDragged, const FString&, const FPointerEvent& );
 	/** Delegate executed when the list of all actions needs to be refreshed */
 	DECLARE_DELEGATE_OneParam( FOnCollectAllActions, FGraphActionListBuilderBase& );
+	/** Delegate executed when the list of all actions needs to be refreshed, should return any sections that should always be visible, even if they don't have children. */
+	DECLARE_DELEGATE_OneParam( FOnCollectStaticSections, TArray<int32>& )
 	/** Delegate executed when a category is being renamed so any post-rename actions can be handled */
 	DECLARE_DELEGATE_ThreeParams( FOnCategoryTextCommitted, const FText&, ETextCommit::Type, TWeakPtr< FGraphActionNode >);
 	/** Delegate executed to check if the selected action is valid for renaming */
 	DECLARE_DELEGATE_RetVal_OneParam( bool, FCanRenameSelectedAction, TWeakPtr< FGraphActionNode > );
 	/** Delegate to get the name of a section if the widget is a section separator. */
 	DECLARE_DELEGATE_RetVal_OneParam( FText, FGetSectionTitle, int32 );
+	/** Delegate to get the tooltip of a section if the widget is a section separator. */
+	DECLARE_DELEGATE_RetVal_OneParam( TSharedPtr<IToolTip>, FGetSectionToolTip, int32 );
+	/** Delegate to get the widget that appears on the section bar in the section separator. */
+	DECLARE_DELEGATE_RetVal_TwoParams( TSharedRef<SWidget>, FGetSectionWidget, TSharedRef<SWidget>, int32 );
 	/** Delegate to get the filter text */
 	DECLARE_DELEGATE_RetVal( FText, FGetFilterText);
 	/** Delegate to check if an action matches a specified name (used for renaming items etc.) */
 	DECLARE_DELEGATE_RetVal_TwoParams( bool, FOnActionMatchesName, FEdGraphSchemaAction*, const FName& );
 
-	SLATE_BEGIN_ARGS( SGraphActionMenu ) 
-		{
-			_AutoExpandActionMenu = false;
-			_ShowFilterTextBox = true;
-			_AlphaSortItems = true;
-		}
+	SLATE_BEGIN_ARGS(SGraphActionMenu)
+		: _UseSectionStyling(false)
+		, _AutoExpandActionMenu(false)
+		, _ShowFilterTextBox(true)
+		, _AlphaSortItems(true)
+		{ }
+
 		SLATE_EVENT( FOnActionSelected, OnActionSelected )
 		SLATE_EVENT( FOnActionDoubleClicked, OnActionDoubleClicked )
 		SLATE_EVENT( FOnActionDragged, OnActionDragged )
@@ -106,14 +113,18 @@ public:
 		SLATE_EVENT( FOnCreateWidgetForAction, OnCreateWidgetForAction )
 		SLATE_EVENT( FOnCreateCustomRowExpander, OnCreateCustomRowExpander )
 		SLATE_EVENT( FOnCollectAllActions, OnCollectAllActions )
+		SLATE_EVENT( FOnCollectStaticSections, OnCollectStaticSections )
 		SLATE_EVENT( FOnCategoryTextCommitted, OnCategoryTextCommitted )
 		SLATE_EVENT( FCanRenameSelectedAction, OnCanRenameSelectedAction )
-		SLATE_EVENT( FGetSectionTitle, OnGetSectionTitle )		
+		SLATE_EVENT( FGetSectionTitle, OnGetSectionTitle )
+		SLATE_EVENT( FGetSectionToolTip, OnGetSectionToolTip )
+		SLATE_EVENT( FGetSectionWidget, OnGetSectionWidget )
 		SLATE_EVENT( FGetFilterText, OnGetFilterText )
 		SLATE_EVENT( FOnActionMatchesName, OnActionMatchesName )
 		SLATE_ARGUMENT( bool, AutoExpandActionMenu )
 		SLATE_ARGUMENT( bool, AlphaSortItems )
 		SLATE_ARGUMENT( bool, ShowFilterTextBox )
+		SLATE_ARGUMENT( bool, UseSectionStyling )
 	SLATE_END_ARGS()
 
 	void Construct( const FArguments& InArgs, bool bIsReadOnly = true );
@@ -128,6 +139,12 @@ public:
 	 * @param bHandleOnSelectionEvent		TRUE if the item should be selected and any actions that occur with selection will be handled. FALSE and only selection will occur */
 	void RefreshAllActions(bool bPreserveExpansion, bool bHandleOnSelectionEvent = true);
 
+	/** Returns a map of all top level sections and their current expansion state. */
+	void GetSectionExpansion(TMap<int32, bool>& SectionExpansion) const;
+
+	/** Sets the sections to be expanded of all top level sections. */
+	void SetSectionExpansion(const TMap<int32, bool>& SectionExpansion);
+
 protected:
 	/** Tree view for showing actions */
 	TSharedPtr< STreeView< TSharedPtr<FGraphActionNode> > > TreeView;
@@ -153,6 +170,8 @@ protected:
 	bool bShowFilterTextBox;
 	/** Don't sort items alphabetically */
 	bool bAlphaSortItems;
+	/** Should the rows and sections be styled like the details panel? */
+	bool bUseSectionStyling;
 
 	/** Delegate to call when action is selected */
 	FOnActionSelected OnActionSelected;
@@ -168,12 +187,18 @@ protected:
 	FOnCreateCustomRowExpander OnCreateCustomRowExpander;
 	/** Delegate to call to collect all actions */
 	FOnCollectAllActions OnCollectAllActions;
+	/** Delegate to call to collect all always visible sections */
+	FOnCollectStaticSections OnCollectStaticSections;
 	/** Delegate to call to handle any post-category rename events */
 	FOnCategoryTextCommitted OnCategoryTextCommitted;
 	/** Delegate to call to check if a selected action is valid for renaming */
 	FCanRenameSelectedAction OnCanRenameSelectedAction;
-	/** Delegate to get the name of a section seperator. */
+	/** Delegate to get the name of a section separator. */
 	FGetSectionTitle OnGetSectionTitle;
+	/** Delegate to get the tooltip of a section separator. */
+	FGetSectionToolTip OnGetSectionToolTip;
+	/** Delegate to get the widgets of a section separator. */
+	FGetSectionWidget OnGetSectionWidget;
 	/** Delegate to get the filter text if supplied from an external source */
 	FGetFilterText OnGetFilterText;
 	/** Delegate to check if an action matches a specified name (used for renaming items etc.) */
diff --git a/Engine/Source/Editor/Kismet/Private/BlueprintEditor.cpp b/Engine/Source/Editor/Kismet/Private/BlueprintEditor.cpp
--- a/Engine/Source/Editor/Kismet/Private/BlueprintEditor.cpp
+++ b/Engine/Source/Editor/Kismet/Private/BlueprintEditor.cpp
@@ -79,6 +79,8 @@
 #include "Editor/UnrealEd/Public/SourceCodeNavigation.h"
 #include "Developer/HotReload/Public/IHotReload.h"
 
+#include "AudioDevice.h"
+
 // Blueprint merging
 #include "Merge.h"
 #include "SHyperlink.h"
@@ -111,6 +113,16 @@ TSharedRef<SWidget> FSelectionDetailsSummoner::CreateTabBody(const FWorkflowTabS
 	return BlueprintEditorPtr->GetInspector();
 }
 
+TSharedRef<SDockTab> FSelectionDetailsSummoner::SpawnTab(const FWorkflowTabSpawnInfo& Info) const
+{
+	TSharedRef<SDockTab> Tab = FWorkflowTabFactory::SpawnTab(Info);
+
+	TSharedPtr<FBlueprintEditor> BlueprintEditorPtr = StaticCastSharedPtr<FBlueprintEditor>(HostingApp.Pin());
+	BlueprintEditorPtr->GetInspector()->SetOwnerTab(Tab);
+
+	return Tab;
+}
+
 /////////////////////////////////////////////////////
 // SChildGraphPicker
 
@@ -1839,6 +1851,9 @@ void FBlueprintEditor::SetupViewForBlueprintEditingMode()
 
 FBlueprintEditor::~FBlueprintEditor()
 {
+	// Clean up the preview
+	DestroyPreview();
+
 	// NOTE: Any tabs that we still have hanging out when destroyed will be cleaned up by FBaseToolkit's destructor
 	UEditorEngine* Editor = (UEditorEngine*)GEngine;
 	if (Editor)
@@ -1895,17 +1910,24 @@ FBlueprintEditor::~FBlueprintEditor()
 
 void FBlueprintEditor::FocusInspectorOnGraphSelection(const FGraphPanelSelectionSet& NewSelection, bool bForceRefresh)
 {
-	if (NewSelection.Array().Num())
+	if ( GetDefault<UEditorExperimentalSettings>()->bUnifiedBlueprintEditor )
 	{
-		SKismetInspector::FShowDetailsOptions ShowDetailsOptions;
-		ShowDetailsOptions.bForceRefresh = bForceRefresh;
-
-		Inspector->ShowDetailsForObjects(NewSelection.Array(), ShowDetailsOptions);
+		//TODO UNIFIEDBP Do something smart here, like keep selection if it comes from the my blueprint view, or the components view, but not if it's from a graph view.
 	}
 	else
 	{
-		// Clear the Inspector if nothing is selected, the MyBlueprints window will set the selection if it means to, but we need to make sure invalid nodes aren't still appearing
-		Inspector->ShowDetailsForSingleObject(nullptr);
+		if ( NewSelection.Array().Num() )
+		{
+			SKismetInspector::FShowDetailsOptions ShowDetailsOptions;
+			ShowDetailsOptions.bForceRefresh = bForceRefresh;
+
+			Inspector->ShowDetailsForObjects(NewSelection.Array(), ShowDetailsOptions);
+		}
+		else
+		{
+			// Clear the Inspector if nothing is selected, the MyBlueprints window will set the selection if it means to, but we need to make sure invalid nodes aren't still appearing
+			Inspector->ShowDetailsForSingleObject(nullptr);
+		}
 	}
 }
 
@@ -2070,6 +2092,11 @@ void FBlueprintEditor::CreateDefaultCommands()
 		FExecuteAction::CreateSP(this, &FBlueprintEditor::EditGlobalOptions_Clicked),
 		FCanExecuteAction());
 
+	ToolkitCommands->MapAction(
+		FFullBlueprintEditorCommands::Get().EditClassDefaults,
+		FExecuteAction::CreateSP(this, &FBlueprintEditor::EditClassDefaults_Clicked),
+		FCanExecuteAction());
+
 	// Edit menu actions
 	ToolkitCommands->MapAction( FBlueprintEditorCommands::Get().FindInBlueprint,
 		FExecuteAction::CreateSP(this, &FBlueprintEditor::FindInBlueprint_Clicked),
@@ -2390,23 +2417,83 @@ bool FBlueprintEditor::ReparentBlueprint_IsVisible() const
 void FBlueprintEditor::EditGlobalOptions_Clicked()
 {
 	UBlueprint* Blueprint = GetBlueprintObj();
-	if(Blueprint != NULL)
+	if ( Blueprint != nullptr )
 	{
-		if (CurrentUISelection == FBlueprintEditor::GraphPanel)
+		if ( CurrentUISelection == FBlueprintEditor::GraphPanel )
 		{
 			ClearSelectionInAllEditors();
+		}
+		if ( CurrentUISelection == FBlueprintEditor::MyBlueprint )
+		{
+			if ( MyBlueprintWidget.IsValid() )
+			{
+				MyBlueprintWidget->ClearGraphActionMenuSelection();
+			}
+		}
+		CurrentUISelection = FBlueprintEditor::ClassSettings;
+
+		// Show details for the Blueprint instance we're editing
+		Inspector->ShowDetailsForSingleObject(Blueprint);
+
+		TSharedPtr<SDockTab> OwnerTab = Inspector->GetOwnerTab();
+		if ( OwnerTab.IsValid() )
+		{
+			OwnerTab->ActivateInParent(ETabActivationCause::SetDirectly);
+			OwnerTab->FlashTab();
+		}
+	}
 }
-		if (CurrentUISelection == FBlueprintEditor::MyBlueprint)
+
+void FBlueprintEditor::EditClassDefaults_Clicked()
 {
-			if (MyBlueprintWidget.IsValid())
+	UBlueprint* Blueprint = GetBlueprintObj();
+	if ( Blueprint != nullptr )
 	{
+		if ( CurrentUISelection == FBlueprintEditor::GraphPanel )
+		{
+			ClearSelectionInAllEditors();
+		}
+		else if ( CurrentUISelection == FBlueprintEditor::MyBlueprint )
+		{
+			if ( MyBlueprintWidget.IsValid() )
+			{
 				MyBlueprintWidget->ClearGraphActionMenuSelection();
 			}
 		}
-		CurrentUISelection = FBlueprintEditor::BlueprintProps;
 
-		// Show details for the Blueprint instance we're editing
-		Inspector->ShowDetailsForSingleObject(Blueprint);
+		CurrentUISelection = FBlueprintEditor::ClassDefaults;
+
+		const bool bForceRefresh = true;
+
+		if ( IsEditingSingleBlueprint() )
+		{
+			if ( GetBlueprintObj()->GeneratedClass != NULL )
+			{
+				Inspector->ShowDetailsForSingleObject(GetBlueprintObj()->GeneratedClass->GetDefaultObject(), SKismetInspector::FShowDetailsOptions(DefaultEditString(), bForceRefresh));
+				TSharedPtr<SDockTab> OwnerTab = Inspector->GetOwnerTab();
+				if ( OwnerTab.IsValid() )
+				{
+					OwnerTab->ActivateInParent(ETabActivationCause::SetDirectly);
+					OwnerTab->FlashTab();
+				}
+			}
+		}
+		else if ( GetEditingObjects().Num() > 0 )
+		{
+			TArray<UObject*> DefaultObjects;
+			for ( int32 i = 0; i < GetEditingObjects().Num(); ++i )
+			{
+				auto Blueprint = (UBlueprint*)( GetEditingObjects()[i] );
+				if ( Blueprint->GeneratedClass )
+				{
+					DefaultObjects.Add(Blueprint->GeneratedClass->GetDefaultObject());
+				}
+			}
+			if ( DefaultObjects.Num() )
+			{
+				DefaultEditor->ShowDetailsForObjects(DefaultObjects);
+			}
+		}
 	}
 }
 
@@ -6014,6 +6101,12 @@ void FBlueprintEditor::RequestSaveEditedObjectState()
 
 void FBlueprintEditor::Tick(float DeltaTime)
 {
+	// Create or update the Blueprint actor instance in the preview scene
+	if ( GetPreviewActor() == nullptr )
+	{
+		UpdatePreviewActor(GetBlueprintObj(), true);
+	}
+
 	if (bRequestedSavingOpenDocumentState)
 	{
 		bRequestedSavingOpenDocumentState = false;
@@ -6925,9 +7018,140 @@ void FBlueprintEditor::OnFindInstancesCustomEvent()
 
 AActor* FBlueprintEditor::GetPreviewActor() const
 {
-	return SCSViewport.IsValid() ? SCSViewport->GetPreviewActor() : nullptr;
+	//return SCSViewport.IsValid() ? SCSViewport->GetPreviewActor() : nullptr;
+
+	UBlueprint* PreviewBlueprint = GetBlueprintObj();
+
+	// Note: The weak ptr can become stale if the actor is reinstanced due to a Blueprint change, etc. In that case we look to see if we can find the new instance in the preview world and then update the weak ptr.
+	if ( PreviewActorPtr.IsStale(true) && PreviewBlueprint )
+	{
+		UWorld* PreviewWorld = PreviewScene.GetWorld();
+		for ( TActorIterator<AActor> It(PreviewWorld); It; ++It )
+		{
+			AActor* Actor = *It;
+			if ( !Actor->IsPendingKillPending()
+				&& Actor->GetClass()->ClassGeneratedBy == PreviewBlueprint )
+			{
+				PreviewActorPtr = Actor;
+				break;
+			}
+		}
+	}
+
+	return PreviewActorPtr.Get();
 }
 
+void FBlueprintEditor::UpdatePreviewActor(UBlueprint* InBlueprint, bool bInForceFullUpdate/* = false*/)
+{
+	AActor* PreviewActor = GetPreviewActor();
+
+	// Signal that we're going to be constructing editor components
+	if ( InBlueprint != NULL && InBlueprint->SimpleConstructionScript != NULL )
+	{
+		InBlueprint->SimpleConstructionScript->BeginEditorComponentConstruction();
+	}
+
+	UBlueprint* PreviewBlueprint = GetBlueprintObj();
+
+	// If the Blueprint is changing
+	if ( InBlueprint != PreviewBlueprint || bInForceFullUpdate )
+	{
+		// Destroy the previous actor instance
+		DestroyPreview();
+
+		// Save the Blueprint we're creating a preview for
+		PreviewBlueprint = InBlueprint;
+
+		// Spawn a new preview actor based on the Blueprint's generated class if it's Actor-based
+		if ( PreviewBlueprint && PreviewBlueprint->GeneratedClass && PreviewBlueprint->GeneratedClass->IsChildOf(AActor::StaticClass()) )
+		{
+			FVector SpawnLocation = FVector::ZeroVector;
+			FRotator SpawnRotation = FRotator::ZeroRotator;
+
+			// Spawn an Actor based on the Blueprint's generated class
+			FActorSpawnParameters SpawnInfo;
+			SpawnInfo.bNoCollisionFail = true;
+			SpawnInfo.bNoFail = true;
+			SpawnInfo.ObjectFlags = RF_Transient;
+
+			// Temporarily remove the deprecated flag so we can respawn the Blueprint in the viewport
+			bool bIsClassDeprecated = PreviewBlueprint->GeneratedClass->HasAnyClassFlags(CLASS_Deprecated);
+			PreviewBlueprint->GeneratedClass->ClassFlags &= ~CLASS_Deprecated;
+
+			PreviewActorPtr = PreviewActor = PreviewScene.GetWorld()->SpawnActor(PreviewBlueprint->GeneratedClass, &SpawnLocation, &SpawnRotation, SpawnInfo);
+
+			// Reassign the deprecated flag if it was previously assigned
+			if ( bIsClassDeprecated )
+			{
+				PreviewBlueprint->GeneratedClass->ClassFlags |= CLASS_Deprecated;
+			}
+
+			check(PreviewActor);
+
+			// Ensure that the actor is visible
+			if ( PreviewActor->bHidden )
+			{
+				PreviewActor->bHidden = false;
+				PreviewActor->MarkComponentsRenderStateDirty();
+				PreviewScene.GetWorld()->SendAllEndOfFrameUpdates();
+			}
+
+			// Prevent any audio from playing as a result of spawning
+			if ( GEngine->AudioDevice )
+			{
+				GEngine->AudioDevice->Flush(PreviewScene.GetWorld());
+			}
+
+			// Set the reference to the preview actor for component editing purposes
+			if ( PreviewBlueprint->SimpleConstructionScript != nullptr )
+			{
+				PreviewBlueprint->SimpleConstructionScript->SetComponentEditorActorInstance(PreviewActor);
+			}
+
+			// Run the construction scripts again, otherwise the actor will appear as though it's had a script pass first, rather than the default properties as shown in the details panel
+			PreviewActor->RerunConstructionScripts();
+		}
+	}
+	else if ( PreviewActor )
+	{
+		PreviewActor->RerunConstructionScripts();
+	}
+
+	// Signal that we're done constructing editor components
+	if ( InBlueprint != nullptr && InBlueprint->SimpleConstructionScript != nullptr )
+	{
+		InBlueprint->SimpleConstructionScript->EndEditorComponentConstruction();
+	}
+}
+
+void FBlueprintEditor::DestroyPreview()
+{
+	AActor* PreviewActor = GetPreviewActor();
+	if ( PreviewActor != nullptr )
+	{
+		check(PreviewScene.GetWorld());
+		PreviewScene.GetWorld()->EditorDestroyActor(PreviewActor, false);
+	}
+
+	UBlueprint* PreviewBlueprint = GetBlueprintObj();
+
+	if ( PreviewBlueprint != nullptr )
+	{
+		if ( PreviewBlueprint->SimpleConstructionScript != nullptr
+			&& PreviewActor == PreviewBlueprint->SimpleConstructionScript->GetComponentEditorActorInstance() )
+		{
+			// Ensure that all editable component references are cleared
+			PreviewBlueprint->SimpleConstructionScript->ClearEditorComponentReferences();
+
+			// Clear the reference to the preview actor instance
+			PreviewBlueprint->SimpleConstructionScript->SetComponentEditorActorInstance(nullptr);
+		}
+
+		PreviewBlueprint = nullptr;
+	}
+}
+
+
 FReply FBlueprintEditor::OnSpawnGraphNodeByShortcut(FInputGesture InGesture, const FVector2D& InPosition, UEdGraph* InGraph)
 {
 	UEdGraph* Graph = InGraph;
diff --git a/Engine/Source/Editor/Kismet/Private/BlueprintEditorModes.cpp b/Engine/Source/Editor/Kismet/Private/BlueprintEditorModes.cpp
--- a/Engine/Source/Editor/Kismet/Private/BlueprintEditorModes.cpp
+++ b/Engine/Source/Editor/Kismet/Private/BlueprintEditorModes.cpp
@@ -43,7 +43,7 @@ const FName FBlueprintEditorApplicationModes::BlueprintMacroMode( TEXT("MacroNam
 
 TSharedPtr<FTabManager::FLayout> GetDefaltEditorLayout(TSharedPtr<class FBlueprintEditor> InBlueprintEditor)
 {
-	return FTabManager::NewLayout( "Standalone_BlueprintEditor_Layout_v5" )
+	return FTabManager::NewLayout( "Standalone_BlueprintEditor_Layout_v6" )
 	->AddArea
 	(
 		FTabManager::NewPrimaryArea() ->SetOrientation(Orient_Vertical)
diff --git a/Engine/Source/Editor/Kismet/Private/SBlueprintEditorSelectedDebugObjectWidget.cpp b/Engine/Source/Editor/Kismet/Private/SBlueprintEditorSelectedDebugObjectWidget.cpp
--- a/Engine/Source/Editor/Kismet/Private/SBlueprintEditorSelectedDebugObjectWidget.cpp
+++ b/Engine/Source/Editor/Kismet/Private/SBlueprintEditorSelectedDebugObjectWidget.cpp
@@ -238,12 +238,7 @@ void SBlueprintEditorSelectedDebugObjectWidget::GenerateDebugWorldNames(bool bRe
 	DebugWorlds.Add(nullptr);
 	DebugWorldNames.Add(MakeShareable(new FString(GetDebugAllWorldsString())));
 
-	UWorld* PreviewWorld = nullptr;
-	TSharedPtr<SSCSEditorViewport> PreviewViewportPtr = BlueprintEditor.Pin()->GetSCSViewport();
-	if (PreviewViewportPtr.IsValid())
-	{
-		PreviewWorld = PreviewViewportPtr->GetPreviewScene().GetWorld();
-	}
+	UWorld* PreviewWorld = BlueprintEditor.Pin()->GetPreviewScene()->GetWorld();
 
 	for (TObjectIterator<UWorld> It; It; ++It)
 	{
@@ -374,12 +369,7 @@ void SBlueprintEditorSelectedDebugObjectWidget::GenerateDebugObjectNames(bool bR
 		}
 	}
 
-	UWorld* PreviewWorld = nullptr;
-	TSharedPtr<SSCSEditorViewport> PreviewViewportPtr = BlueprintEditor.Pin()->GetSCSViewport();
-	if (PreviewViewportPtr.IsValid())
-	{
-		PreviewWorld = PreviewViewportPtr->GetPreviewScene().GetWorld();
-	}
+	UWorld* PreviewWorld = BlueprintEditor.Pin()->GetPreviewScene()->GetWorld();
 
 	for (TObjectIterator<UObject> It; It; ++It)
 	{
diff --git a/Engine/Source/Editor/Kismet/Private/SBlueprintEditorToolbar.cpp b/Engine/Source/Editor/Kismet/Private/SBlueprintEditorToolbar.cpp
--- a/Engine/Source/Editor/Kismet/Private/SBlueprintEditorToolbar.cpp
+++ b/Engine/Source/Editor/Kismet/Private/SBlueprintEditorToolbar.cpp
@@ -343,8 +343,9 @@ void FFullBlueprintEditorCommands::RegisterCommands()
 	UI_COMMAND(SwitchToScriptingMode, "Graph", "Switches to Graph Editing Mode", EUserInterfaceActionType::ToggleButton, FInputGesture());
 	UI_COMMAND(SwitchToBlueprintDefaultsMode, "Defaults", "Switches to Blueprint Defaults Mode", EUserInterfaceActionType::ToggleButton, FInputGesture());
 	UI_COMMAND(SwitchToComponentsMode, "Components", "Switches to Components Mode", EUserInterfaceActionType::ToggleButton, FInputGesture());
-	// Renaming this to Options because Blueprint props was requested to be changed and Blueprint properties was too long and looks odd.
-	UI_COMMAND(EditGlobalOptions, "Options", "Edit Blueprint Options", EUserInterfaceActionType::Button, FInputGesture());
+
+	UI_COMMAND(EditGlobalOptions, "Class Settings", "Edit Class Settings (Previously known as Blueprint Props)", EUserInterfaceActionType::Button, FInputGesture());
+	UI_COMMAND(EditClassDefaults, "Class Defaults", "Edit the initial values of your class.", EUserInterfaceActionType::Button, FInputGesture());
 
 	UI_COMMAND(JumpToErrorNode, "Jump to Error Node", "When enabled, then the Blueprint will snap focus to nodes producing an error during compilation", EUserInterfaceActionType::ToggleButton, FInputGesture());
 }
@@ -506,6 +507,11 @@ void FBlueprintEditorToolbar::FillBlueprintGlobalOptionsToolbar(FToolBarBuilder&
 	if(BlueprintObj != NULL)
 	{
 		ToolbarBuilder.AddToolBarButton(Commands.EditGlobalOptions);
+
+		if ( GetDefault<UEditorExperimentalSettings>()->bUnifiedBlueprintEditor )
+		{
+			ToolbarBuilder.AddToolBarButton(Commands.EditClassDefaults);
+		}
 	}
 	
 	ToolbarBuilder.EndSection();
diff --git a/Engine/Source/Editor/Kismet/Private/SCSEditorViewportClient.cpp b/Engine/Source/Editor/Kismet/Private/SCSEditorViewportClient.cpp
--- a/Engine/Source/Editor/Kismet/Private/SCSEditorViewportClient.cpp
+++ b/Engine/Source/Editor/Kismet/Private/SCSEditorViewportClient.cpp
@@ -105,8 +105,8 @@ namespace
 /////////////////////////////////////////////////////////////////////////
 // FSCSEditorViewportClient
 
-FSCSEditorViewportClient::FSCSEditorViewportClient(TWeakPtr<FBlueprintEditor>& InBlueprintEditorPtr, FPreviewScene& InPreviewScene, const TSharedRef<SSCSEditorViewport>& InSCSEditorViewport)
-	: FEditorViewportClient(nullptr, &InPreviewScene, StaticCastSharedRef<SEditorViewport>(InSCSEditorViewport))
+FSCSEditorViewportClient::FSCSEditorViewportClient(TWeakPtr<FBlueprintEditor>& InBlueprintEditorPtr, FPreviewScene* InPreviewScene, const TSharedRef<SSCSEditorViewport>& InSCSEditorViewport)
+	: FEditorViewportClient(nullptr, InPreviewScene, StaticCastSharedRef<SEditorViewport>(InSCSEditorViewport))
 	, BlueprintEditorPtr(InBlueprintEditorPtr)
 	, PreviewBlueprint(NULL)
 	, PreviewActorBounds(ForceInitToZero)
@@ -156,9 +156,6 @@ FSCSEditorViewportClient::~FSCSEditorViewportClient()
 {
 	// Ensure that an in-progress transaction is ended
 	EndTransaction();
-
-	// Clean up the preview
-	DestroyPreview();
 }
 
 void FSCSEditorViewportClient::Tick(float DeltaSeconds)
@@ -168,7 +165,7 @@ void FSCSEditorViewportClient::Tick(float DeltaSeconds)
 	// Register the selection override delegate for the preview actor's components
 	TSharedRef<SSCSEditor> SCSEditor = BlueprintEditorPtr.Pin()->GetSCSEditor();
 	AActor* PreviewActor = GetPreviewActor();
-	if (PreviewActor != NULL)
+	if (PreviewActor != nullptr)
 	{
 		TInlineComponentArray<UPrimitiveComponent*> PrimitiveComponents;
 		PreviewActor->GetComponents(PrimitiveComponents);
@@ -187,6 +184,14 @@ void FSCSEditorViewportClient::Tick(float DeltaSeconds)
 		InvalidatePreview(false);
 	}
 
+	if ( PreviewActor != LastPreviewActor.Get() || PreviewActor == nullptr )
+	{
+		PreviewActor = PreviewActor;
+
+		Invalidate();
+		RefreshPreviewBounds();
+	}
+
 	// Tick the preview scene world.
 	if (!GIntraFrameDebuggingGameThread)
 	{
@@ -679,12 +684,6 @@ void FSCSEditorViewportClient::InvalidatePreview(bool bResetCamera)
 		return;
 	}
 	
-	UBlueprint* Blueprint = BlueprintEditorPtr.Pin()->GetBlueprintObj();
-	check(Blueprint);
-
-	// Create or update the Blueprint actor instance in the preview scene
-	UpdatePreviewActorForBlueprint(Blueprint, !IsPreviewSceneValid());
-
 	if( bResetCamera )
 	{
 		ResetCamera();
@@ -735,9 +734,9 @@ void FSCSEditorViewportClient::ToggleRealtimePreview()
 	Invalidate();
 }
 
-bool FSCSEditorViewportClient::IsPreviewSceneValid() const
+AActor* FSCSEditorViewportClient::GetPreviewActor() const
 {
-	return GetPreviewActor() != NULL;
+	return BlueprintEditorPtr.Pin()->GetPreviewActor();
 }
 
 void FSCSEditorViewportClient::FocusViewportToSelection()
@@ -770,7 +769,7 @@ bool FSCSEditorViewportClient::GetIsSimulateEnabled()
 void FSCSEditorViewportClient::ToggleIsSimulateEnabled() 
 {
 	// Must destroy existing actors before we toggle the world state
-	DestroyPreview();
+	BlueprintEditorPtr.Pin()->DestroyPreview();
 
 	bIsSimulateEnabled = !bIsSimulateEnabled;
 	PreviewScene->GetWorld()->bBegunPlay = bIsSimulateEnabled;
@@ -781,7 +780,8 @@ void FSCSEditorViewportClient::ToggleIsSimulateEnabled()
 	TSharedRef<SWidget> Inspector = BlueprintEditorPtr.Pin()->GetInspector();
 
 	// When simulate is enabled, we don't want to allow the user to modify the components
-	UpdatePreviewActorForBlueprint(PreviewBlueprint, true);
+	BlueprintEditorPtr.Pin()->UpdatePreviewActor(PreviewBlueprint, true);
+
 	SCSEditor->SetEnabled(!bIsSimulateEnabled);
 	Inspector->SetEnabled(!bIsSimulateEnabled);
 
@@ -891,111 +891,6 @@ void FSCSEditorViewportClient::EndTransaction()
 	}
 }
 
-AActor* FSCSEditorViewportClient::GetPreviewActor() const
-{
-	// Note: The weak ptr can become stale if the actor is reinstanced due to a Blueprint change, etc. In that case we look to see if we can find the new instance in the preview world and then update the weak ptr.
-	if(PreviewActorPtr.IsStale(true) && PreviewBlueprint)
-	{
-		UWorld* PreviewWorld = PreviewScene->GetWorld();
-		for(TActorIterator<AActor> It(PreviewWorld); It; ++It)
-		{
-			AActor* Actor = *It;
-			if(!Actor->IsPendingKillPending()
-				&& Actor->GetClass()->ClassGeneratedBy == PreviewBlueprint)
-			{
-				PreviewActorPtr = Actor;
-				break;
-			}
-		}
-	}
-
-	return PreviewActorPtr.Get();
-}
-
-void FSCSEditorViewportClient::UpdatePreviewActorForBlueprint(UBlueprint* InBlueprint, bool bInForceFullUpdate/* = false*/)
-{
-	AActor* PreviewActor = GetPreviewActor();
-
-	// Signal that we're going to be constructing editor components
-	if(InBlueprint != NULL && InBlueprint->SimpleConstructionScript != NULL)
-	{
-		InBlueprint->SimpleConstructionScript->BeginEditorComponentConstruction();
-	}
-
-	// If the Blueprint is changing
-	if(InBlueprint != PreviewBlueprint || bInForceFullUpdate)
-	{
-		// Destroy the previous actor instance
-		DestroyPreview();
-
-		// Save the Blueprint we're creating a preview for
-		PreviewBlueprint = InBlueprint;
-
-		// Spawn a new preview actor based on the Blueprint's generated class if it's Actor-based
-		if(PreviewBlueprint && PreviewBlueprint->GeneratedClass && PreviewBlueprint->GeneratedClass->IsChildOf(AActor::StaticClass()))
-		{
-			FVector SpawnLocation = FVector::ZeroVector;
-			FRotator SpawnRotation = FRotator::ZeroRotator;
-
-			// Spawn an Actor based on the Blueprint's generated class
-			FActorSpawnParameters SpawnInfo;
-			SpawnInfo.bNoCollisionFail = true;
-			SpawnInfo.bNoFail = true;
-			SpawnInfo.ObjectFlags = RF_Transient;
-
-			// Temporarily remove the deprecated flag so we can respawn the Blueprint in the viewport
-			bool bIsClassDeprecated = PreviewBlueprint->GeneratedClass->HasAnyClassFlags(CLASS_Deprecated);
-			PreviewBlueprint->GeneratedClass->ClassFlags &= ~CLASS_Deprecated;
-
-			PreviewActorPtr = PreviewActor = PreviewScene->GetWorld()->SpawnActor( PreviewBlueprint->GeneratedClass, &SpawnLocation, &SpawnRotation, SpawnInfo );
-
-			// Reassign the deprecated flag if it was previously assigned
-			if(bIsClassDeprecated)
-			{
-				PreviewBlueprint->GeneratedClass->ClassFlags |= CLASS_Deprecated;
-			}
-
-			check(PreviewActor);
-
-			// Ensure that the actor is visible
-			if(PreviewActor->bHidden)
-			{
-				PreviewActor->bHidden = false;
-				PreviewActor->MarkComponentsRenderStateDirty();
-				GetWorld()->SendAllEndOfFrameUpdates();
-			}
-
-			// Prevent any audio from playing as a result of spawning
-			if(GEngine->AudioDevice)
-			{
-				GEngine->AudioDevice->Flush(GetWorld());
-			}
-
-			// Set the reference to the preview actor for component editing purposes
-			if(PreviewBlueprint->SimpleConstructionScript != NULL)
-			{
-				PreviewBlueprint->SimpleConstructionScript->SetComponentEditorActorInstance(PreviewActor);
-			}
-
-			// Run the construction scripts again, otherwise the actor will appear as though it's had a script pass first, rather than the default properties as shown in the details panel
-			PreviewActor->RerunConstructionScripts();
-		}
-	}
-	else if(PreviewActor)
-	{
-		PreviewActor->RerunConstructionScripts();
-	}
-
-	// Signal that we're done constructing editor components
-	if(InBlueprint != NULL && InBlueprint->SimpleConstructionScript != NULL)
-	{
-		InBlueprint->SimpleConstructionScript->EndEditorComponentConstruction();
-	}
-
-	Invalidate();
-	RefreshPreviewBounds();
-}
-
 void FSCSEditorViewportClient::RefreshPreviewBounds()
 {
 	AActor* PreviewActor = GetPreviewActor();
@@ -1018,29 +913,3 @@ void FSCSEditorViewportClient::RefreshPreviewBounds()
 		}
 	}
 }
-
-void FSCSEditorViewportClient::DestroyPreview()
-{
-	AActor* PreviewActor = GetPreviewActor();
-	if(PreviewActor != NULL)
-	{
-		check(PreviewScene);
-		check(PreviewScene->GetWorld());
-		PreviewScene->GetWorld()->EditorDestroyActor(PreviewActor, false);
-	}
-
-	if(PreviewBlueprint != NULL)
-	{
-		if(PreviewBlueprint->SimpleConstructionScript != NULL
-			&& PreviewActor == PreviewBlueprint->SimpleConstructionScript->GetComponentEditorActorInstance())
-		{
-			// Ensure that all editable component references are cleared
-			PreviewBlueprint->SimpleConstructionScript->ClearEditorComponentReferences();
-
-			// Clear the reference to the preview actor instance
-			PreviewBlueprint->SimpleConstructionScript->SetComponentEditorActorInstance(NULL);
-		}
-
-		PreviewBlueprint = NULL;
-	}
-}
diff --git a/Engine/Source/Editor/Kismet/Private/SCSEditorViewportClient.h b/Engine/Source/Editor/Kismet/Private/SCSEditorViewportClient.h
--- a/Engine/Source/Editor/Kismet/Private/SCSEditorViewportClient.h
+++ b/Engine/Source/Editor/Kismet/Private/SCSEditorViewportClient.h
@@ -16,7 +16,7 @@ public:
 	 * @param InBlueprintEditorPtr A weak reference to the Blueprint Editor context.
 	 * @param InPreviewScene The preview scene to use.
 	 */
-	FSCSEditorViewportClient(TWeakPtr<class FBlueprintEditor>& InBlueprintEditorPtr, FPreviewScene& InPreviewScene, const TSharedRef<SSCSEditorViewport>& InSCSEditorViewport);
+	FSCSEditorViewportClient(TWeakPtr<class FBlueprintEditor>& InBlueprintEditorPtr, FPreviewScene* InPreviewScene, const TSharedRef<SSCSEditorViewport>& InSCSEditorViewport);
 
 	/**
 	 * Destructor.
@@ -49,12 +49,12 @@ public:
 	void InvalidatePreview(bool bResetCamera = true);
 
 	/**
-     * Resets the camera position
+	 * Resets the camera position
 	 */
 	void ResetCamera();
 
 	/**
-     * Determines whether or not realtime preview is enabled.
+	 * Determines whether or not realtime preview is enabled.
 	 * 
 	 * @return true if realtime preview is enabled, false otherwise.
 	 */
@@ -64,17 +64,10 @@ public:
 	}
 
 	/**
-     * Toggles realtime preview on/off.
+	 * Toggles realtime preview on/off.
 	 */
 	void ToggleRealtimePreview();
 
-	/**
-     * Determines whether or not the preview scene is valid.
-	 * 
-	 * @return true if the preview scene is valid, false otherwise.
-	 */
-	bool IsPreviewSceneValid() const;
-
 	/**
 	 * Focuses the viewport on the selected components
 	 */
@@ -126,19 +119,6 @@ protected:
 	 */
 	void EndTransaction();
 
-	/**
-	 * Creates/updates the preview actor for the given blueprint.
-	 *
-	 * @param InBlueprint			The Blueprint to create or update the preview for.
-	 * @param bInForceFullUpdate	Force a full update to respawn actors.
-	 */
-	void UpdatePreviewActorForBlueprint(UBlueprint* InBlueprint, bool bInForceFullUpdate = false);
-
-	/**
-	 * Destroy the Blueprint preview.
-	 */
-	void DestroyPreview();
-
 	/**
 	 * Updates preview bounds and floor positioning
 	 */
@@ -154,8 +134,8 @@ private:
 	/** The Blueprint associated with the current preview */
 	UBlueprint* PreviewBlueprint;
 
-	/** The preview actor representing the current preview */
-	mutable TWeakObjectPtr<AActor> PreviewActorPtr;
+	/** The preview actor from last tick so we can tell if it changed. */
+	TWeakObjectPtr<AActor> LastPreviewActor;
 
 	/** The full bounds of the preview scene (encompasses all visible components) */
 	FBoxSphereBounds PreviewActorBounds;
diff --git a/Engine/Source/Editor/Kismet/Private/SKismetInspector.cpp b/Engine/Source/Editor/Kismet/Private/SKismetInspector.cpp
--- a/Engine/Source/Editor/Kismet/Private/SKismetInspector.cpp
+++ b/Engine/Source/Editor/Kismet/Private/SKismetInspector.cpp
@@ -104,6 +104,16 @@ TSharedRef<SWidget> SKismetInspector::MakeContextualEditingWidget(struct FKismet
 	return ContextualEditingWidget;
 }
 
+void SKismetInspector::SetOwnerTab(TSharedRef<SDockTab> Tab)
+{
+	OwnerTab = Tab;
+}
+
+TSharedPtr<SDockTab> SKismetInspector::GetOwnerTab() const
+{
+	return OwnerTab.Pin();
+}
+
 FText SKismetInspector::GetContextualEditingWidgetTitle() const
 {
 	FText Title = PropertyViewTitle;
diff --git a/Engine/Source/Editor/Kismet/Private/SMyBlueprint.cpp b/Engine/Source/Editor/Kismet/Private/SMyBlueprint.cpp
--- a/Engine/Source/Editor/Kismet/Private/SMyBlueprint.cpp
+++ b/Engine/Source/Editor/Kismet/Private/SMyBlueprint.cpp
@@ -190,10 +190,17 @@ void SMyBlueprint::Construct(const FArguments& InArgs, TWeakPtr<FBlueprintEditor
 		TSharedPtr<FExtender> Extender = MakeShareable(new FExtender);
 		Toolbar->AddNewToolbar(Extender);
 
-		FToolBarBuilder ToolbarBuilder(ToolKitCommandList, FMultiBoxCustomization::None, Extender);
-		ToolbarBuilder.BeginSection("MyBlueprint");
-		ToolbarBuilder.EndSection();
-		ToolbarBuilderWidget = ToolbarBuilder.MakeWidget();
+		if ( GetDefault<UEditorExperimentalSettings>()->bUnifiedBlueprintEditor )
+		{
+			ToolbarBuilderWidget = SNullWidget::NullWidget;
+		}
+		else
+		{
+			FToolBarBuilder ToolbarBuilder(ToolKitCommandList, FMultiBoxCustomization::None, Extender);
+			ToolbarBuilder.BeginSection("MyBlueprint");
+			ToolbarBuilder.EndSection();
+			ToolbarBuilderWidget = ToolbarBuilder.MakeWidget();
+		}
 
 		ToolKitCommandList->MapAction(FGenericCommands::Get().Rename,
 			FExecuteAction::CreateSP(this, &SMyBlueprint::OnRequestRenameOnActionNode),
@@ -207,14 +214,79 @@ void SMyBlueprint::Construct(const FArguments& InArgs, TWeakPtr<FBlueprintEditor
 		ToolbarBuilderWidget = SNew(SBox);
 	}
 
+	TSharedPtr<SWidget> AddNewMenu = SNullWidget::NullWidget;
+	if ( GetDefault<UEditorExperimentalSettings>()->bUnifiedBlueprintEditor )
+	{
+		AddNewMenu = SNew(SComboButton)
+			.ComboButtonStyle(FEditorStyle::Get(), "ToolbarComboButton")
+			.ForegroundColor(FLinearColor::White)
+			.ToolTipText(LOCTEXT("AddNewToolTip", "Add a new Variable, Graph, Function, Macro, or Event Dispatcher."))
+			.OnGetMenuContent(this, &SMyBlueprint::CreateAddNewMenuWidget)
+			.HasDownArrow(true)
+			.ContentPadding(FMargin(1, 0, 2, 0))
+			.AddMetaData<FTagMetaData>(FTagMetaData(TEXT("MyBlueprintAddNewCombo")))
+			.ButtonContent()
+			[
+				SNew(SHorizontalBox)
+
+				+ SHorizontalBox::Slot()
+				.AutoWidth()
+				.Padding(FMargin(0, 1))
+				[
+					SNew(SImage)
+					.Image(FEditorStyle::GetBrush("Plus"))
+				]
+
+				+ SHorizontalBox::Slot()
+				.VAlign(VAlign_Center)
+				.AutoWidth()
+				.Padding(FMargin(2, 0, 2, 0))
+				[
+					SNew(STextBlock)
+					.Text(LOCTEXT("AddNew", "Add New"))
+				]
+			];
+	}
+
+	FMenuBuilder ViewOptions(true, nullptr);
+
+	ViewOptions.AddMenuEntry(
+		LOCTEXT("ShowInheritedVariables", "Show Inherited Variables"),
+		LOCTEXT("ShowInheritedVariablesTooltip", "Should inherited variables from parent classes and blueprints be shown in the tree?"),
+		FSlateIcon(),
+		FUIAction( 
+			FExecuteAction::CreateSP( this, &SMyBlueprint::OnToggleShowInheritedVariables ),
+			FCanExecuteAction(),
+			FIsActionChecked::CreateSP( this, &SMyBlueprint::IsShowingInheritedVariables )
+		),
+		NAME_None,
+		EUserInterfaceActionType::ToggleButton,
+		TEXT("MyBlueprint_ShowInheritedVariables")
+	);
+
+	ViewOptions.AddMenuEntry(
+		LOCTEXT("ShowEmptySections", "Show Empty Sections"),
+		LOCTEXT("ShowEmptySectionsTooltip", "Should we show empty sections? eg. Graphs, Functions...etc."),
+		FSlateIcon(),
+		FUIAction( 
+			FExecuteAction::CreateSP( this, &SMyBlueprint::OnToggleShowEmptySections ),
+			FCanExecuteAction(),
+			FIsActionChecked::CreateSP(this, &SMyBlueprint::IsShowingEmptySections)
+		),
+		NAME_None,
+		EUserInterfaceActionType::ToggleButton,
+		TEXT("MyBlueprint_ShowEmptySections")
+	);
+
 	SAssignNew(FilterBox, SSearchBox)
 		.OnTextChanged( this, &SMyBlueprint::OnFilterTextChanged );
 
 	// create the main action list piece of this widget
-	TSharedRef<SWidget> MyBlueprintActionMenu = SAssignNew(GraphActionMenu, SGraphActionMenu, false)
+	SAssignNew(GraphActionMenu, SGraphActionMenu, false)
 		.OnGetFilterText(this, &SMyBlueprint::GetFilterText)
 		.OnCreateWidgetForAction(this, &SMyBlueprint::OnCreateWidgetForAction)
 		.OnCollectAllActions(this, &SMyBlueprint::CollectAllActions)
+		.OnCollectStaticSections(this, &SMyBlueprint::CollectStaticSections)
 		.OnActionDragged(this, &SMyBlueprint::OnActionDragged)
 		.OnCategoryDragged(this, &SMyBlueprint::OnCategoryDragged)
 		.OnActionSelected(this, &SMyBlueprint::OnGlobalActionSelected)
@@ -222,61 +294,87 @@ void SMyBlueprint::Construct(const FArguments& InArgs, TWeakPtr<FBlueprintEditor
 		.OnContextMenuOpening(this, &SMyBlueprint::OnContextMenuOpening)
 		.OnCategoryTextCommitted(this, &SMyBlueprint::OnCategoryNameCommitted)
 		.OnCanRenameSelectedAction(this, &SMyBlueprint::CanRequestRenameOnActionNode)
-		.OnGetSectionTitle(this,&SMyBlueprint::OnGetSectionTitle)
-		.AlphaSortItems(false);
+		.OnGetSectionTitle(this, &SMyBlueprint::OnGetSectionTitle)
+		.OnGetSectionWidget(this, &SMyBlueprint::OnGetSectionWidget)
+		.AlphaSortItems(false)
+		.UseSectionStyling(true);
+
 
 	// now piece together all the content for this widget
 	ChildSlot
 	[
-		SNew(SBorder)
-		.Padding(4.0f)
-		.BorderImage(FEditorStyle::GetBrush("ToolPanel.GroupBorder"))
-		.AddMetaData<FTagMetaData>(FTagMetaData(TEXT("MyBlueprintPanel")))
-		[
-			SNew(SVerticalBox)
+		SNew(SVerticalBox)
 
-			+ SVerticalBox::Slot()
-			.AutoHeight()
-			[
-				ToolbarBuilderWidget.ToSharedRef()
-			]
-				
-			+ SVerticalBox::Slot()
-			.AutoHeight()
-			[
-				FilterBox.ToSharedRef()
-			]
-				
-			+ SVerticalBox::Slot()
-			.FillHeight(1.0f)
+		+ SVerticalBox::Slot()
+		.AutoHeight()
+		[
+			SNew(SBorder)
+			.Padding(4.0f)
+			.BorderImage(FEditorStyle::GetBrush("ToolPanel.GroupBorder"))
+			.AddMetaData<FTagMetaData>(FTagMetaData(TEXT("MyBlueprintPanel")))
 			[
-				SAssignNew(ActionMenuContainer, SSplitter)
-				.Orientation(Orient_Vertical)
+				SNew(SVerticalBox)
 
-				+ SSplitter::Slot()
+				+ SVerticalBox::Slot()
+				.AutoHeight()
 				[
-					MyBlueprintActionMenu
+					ToolbarBuilderWidget.ToSharedRef()
 				]
-			]
 
-			+ SVerticalBox::Slot()
-			.AutoHeight()
-			.Padding(FMargin(3.0f, 2.0f, 0.0f, 0.0f))
-			[
-				SNew(SCheckBox)
-				.IsChecked(this, &SMyBlueprint::OnUserVarsCheckState)
-				.OnCheckStateChanged(this, &SMyBlueprint::OnUserVarsCheckStateChanged)
+				+ SVerticalBox::Slot()
+				.AutoHeight()
 				[
-					SNew(STextBlock)
-					.Text(LOCTEXT("ShowInheritedVariables", "Show inherited variables"))
-					.ToolTip(IDocumentation::Get()->CreateToolTip(
-						LOCTEXT("ShowInheritedVariablesTooltip", "Should inherited variables from parent classes and blueprints be shown in the tree?"),
-						NULL,
-						TEXT("Shared/Editors/BlueprintEditor"),
-						TEXT("MyBlueprint_ShowInheritedVariables")))
+					SNew(SHorizontalBox)
+
+					+ SHorizontalBox::Slot()
+					.AutoWidth()
+					.Padding(0, 0, 2, 0)
+					[
+						AddNewMenu.ToSharedRef()
+					]
+
+					+ SHorizontalBox::Slot()
+					.FillWidth(1.0f)
+					.VAlign(VAlign_Center)
+					[
+						FilterBox.ToSharedRef()
+					]
+
+					+ SHorizontalBox::Slot()
+					.AutoWidth()
+					.Padding(2, 0, 0, 0)
+					[
+						SNew(SComboButton)
+						.ComboButtonStyle(FEditorStyle::Get(), "ToolbarComboButton")
+						.ForegroundColor(FSlateColor::UseForeground())
+						.HasDownArrow(true)
+						.ContentPadding(FMargin(1, 0))
+						.AddMetaData<FTagMetaData>(FTagMetaData(TEXT("ViewOptions")))
+						.MenuContent()
+						[
+							ViewOptions.MakeWidget()
+						]
+						.ButtonContent()
+						[
+							SNew(SImage)
+							.Image(FEditorStyle::GetBrush("GenericViewButton"))
+						]
+					]
 				]
 			]
 		]
+		
+		+ SVerticalBox::Slot()
+		.FillHeight(1.0f)
+		[
+			SAssignNew(ActionMenuContainer, SSplitter)
+			.Orientation(Orient_Vertical)
+
+			+ SSplitter::Slot()
+			[
+				GraphActionMenu.ToSharedRef()
+			]
+		]
 	];
 
 	if (GetLocalActionsListVisibility() == EVisibility::Visible)
@@ -294,6 +392,16 @@ void SMyBlueprint::Construct(const FArguments& InArgs, TWeakPtr<FBlueprintEditor
 	{
 		Refresh();
 	}
+
+	TMap<int32, bool> ExpandedSections;
+	ExpandedSections.Add(NodeSectionID::VARIABLE, true);
+	ExpandedSections.Add(NodeSectionID::FUNCTION, true);
+	ExpandedSections.Add(NodeSectionID::MACRO, true);
+	ExpandedSections.Add(NodeSectionID::DELEGATE, true);
+	ExpandedSections.Add(NodeSectionID::GRAPH, true);
+	ExpandedSections.Add(NodeSectionID::LOCAL_VARIABLE, true);
+
+	GraphActionMenu->SetSectionExpansion(ExpandedSections);
 }
 
 void SMyBlueprint::OnCategoryNameCommitted(const FText& InNewText, ETextCommit::Type InTextCommit, TWeakPtr< FGraphActionNode > InAction )
@@ -423,6 +531,109 @@ FText SMyBlueprint::OnGetSectionTitle( int32 InSectionID )
 	return SeperatorTitle;
 }
 
+TSharedRef<SWidget> SMyBlueprint::OnGetSectionWidget(TSharedRef<SWidget> RowWidget, int32 InSectionID)
+{
+	TWeakPtr<SWidget> WeakRowWidget = RowWidget;
+
+	FText AddNewText;
+
+	switch ( InSectionID )
+	{
+	case NodeSectionID::VARIABLE:
+		AddNewText = LOCTEXT("AddNewVariable", "Variable");
+		break;
+	case NodeSectionID::FUNCTION:
+		AddNewText = LOCTEXT("AddNewFunction", "Function");
+		break;
+	case NodeSectionID::MACRO:
+		AddNewText = LOCTEXT("AddNewMacro", "Macro");
+		break;
+	case NodeSectionID::DELEGATE:
+		AddNewText = LOCTEXT("AddNewDelegate", "Event Dispatcher");
+		break;
+	case NodeSectionID::GRAPH:
+		AddNewText = LOCTEXT("AddNewGraph", "New Graph");
+		break;
+	case NodeSectionID::LOCAL_VARIABLE:
+		AddNewText = LOCTEXT("AddNewLocalVariable", "Local Variable");
+		break;
+	default:
+		return SNullWidget::NullWidget;
+	}
+
+	return SNew(SButton)
+		.ButtonStyle(FEditorStyle::Get(), "RoundButton")
+		.ForegroundColor(FEditorStyle::GetSlateColor("DefaultForeground"))
+		.OnClicked(this, &SMyBlueprint::OnAddButtonClickedOnSection, InSectionID)
+		.HAlign(HAlign_Center)
+		.VAlign(VAlign_Center)
+		.ContentPadding(FMargin(2, 0))
+		//.IsEnabled(this, &SMyBlueprint::OnAddButtonClickedOnSection)
+		[
+			SNew(SHorizontalBox)
+
+			+ SHorizontalBox::Slot()
+			.AutoWidth()
+			.Padding(FMargin(0, 1))
+			[
+				SNew(SImage)
+				.Image(FEditorStyle::GetBrush("Plus"))
+			]
+
+			+ SHorizontalBox::Slot()
+			.VAlign(VAlign_Center)
+			.AutoWidth()
+			.Padding(FMargin(2,0,0,0))
+			[
+				SNew(STextBlock)
+				.Font(IDetailLayoutBuilder::GetDetailFontBold())
+				.Text(AddNewText)
+				.Visibility(this, &SMyBlueprint::OnGetSectionTextVisibility, WeakRowWidget)
+				.ShadowOffset(FVector2D(1,1))
+			]
+		];
+}
+
+FReply SMyBlueprint::OnAddButtonClickedOnSection(int32 InSectionID)
+{
+	switch ( InSectionID )
+	{
+	case NodeSectionID::VARIABLE:
+		BlueprintEditorPtr.Pin()->GetToolkitCommands()->ExecuteAction(FBlueprintEditorCommands::Get().AddNewVariable.ToSharedRef());
+		break;
+	case NodeSectionID::FUNCTION:
+		BlueprintEditorPtr.Pin()->GetToolkitCommands()->ExecuteAction(FBlueprintEditorCommands::Get().AddNewFunction.ToSharedRef());
+		break;
+	case NodeSectionID::MACRO:
+		BlueprintEditorPtr.Pin()->GetToolkitCommands()->ExecuteAction(FBlueprintEditorCommands::Get().AddNewMacroDeclaration.ToSharedRef());
+		break;
+	case NodeSectionID::DELEGATE:
+		BlueprintEditorPtr.Pin()->GetToolkitCommands()->ExecuteAction(FBlueprintEditorCommands::Get().AddNewDelegate.ToSharedRef());
+		break;
+	case NodeSectionID::GRAPH:
+		BlueprintEditorPtr.Pin()->GetToolkitCommands()->ExecuteAction(FBlueprintEditorCommands::Get().AddNewEventGraph.ToSharedRef());
+		break;
+	case NodeSectionID::LOCAL_VARIABLE:
+		//BlueprintEditorPtr.Pin()->GetToolkitCommands()->ExecuteAction(FBlueprintEditorCommands::Get().AddNewLocalVariable.ToSharedRef());
+		OnAddNewLocalVariable();
+		break;
+	}
+
+	return FReply::Handled();
+}
+
+EVisibility SMyBlueprint::OnGetSectionTextVisibility(TWeakPtr<SWidget> RowWidget) const
+{
+	if ( RowWidget.Pin()->IsHovered() )
+	{
+		return EVisibility::SelfHitTestInvisible;
+	}
+	else
+	{
+		return EVisibility::Collapsed;
+	}
+}
+
 bool SMyBlueprint::CanRequestRenameOnActionNode(TWeakPtr<FGraphActionNode> InSelectedNode) const
 {
 	bool bIsReadOnly = true;
@@ -491,23 +702,27 @@ void SMyBlueprint::Refresh()
 	GraphActionMenu->RefreshAllActions(true);
 	
 	bool bLocalActionsAreVisible = (GetLocalActionsListVisibility() == EVisibility::Visible);
-	if (bLocalActionsAreVisible)
+
+	if ( !GetDefault<UEditorExperimentalSettings>()->bUnifiedBlueprintEditor )
 	{
-		if (!LocalGraphActionMenu.IsValid())
+		if ( bLocalActionsAreVisible )
 		{
-			ActionMenuContainer->AddSlot()
-				.Value(0.33f)
-			[
-				ConstructLocalActionPanel()
-			];
+			if ( !LocalGraphActionMenu.IsValid() )
+			{
+				ActionMenuContainer->AddSlot()
+					.Value(0.33f)
+					[
+						ConstructLocalActionPanel()
+					];
+			}
+			check(LocalGraphActionMenu.IsValid());
+			LocalGraphActionMenu->RefreshAllActions(true);
+		}
+		else if ( LocalGraphActionMenu.IsValid() && ensure(ActionMenuContainer->GetChildren()->Num() > 1) )
+		{
+			ActionMenuContainer->RemoveAt(1);
+			LocalGraphActionMenu = nullptr;
 		}
-		check(LocalGraphActionMenu.IsValid());
-		LocalGraphActionMenu->RefreshAllActions(true);
-	}
-	else if (LocalGraphActionMenu.IsValid() && ensure(ActionMenuContainer->GetChildren()->Num() > 1))
-	{
-		ActionMenuContainer->RemoveAt(1);
-		LocalGraphActionMenu = NULL;
 	}
 }
 
@@ -671,6 +886,10 @@ void SMyBlueprint::GetLocalVariables(FGraphActionListBuilderBase& OutAllActions)
 				{
 					TSharedPtr<FEdGraphSchemaAction_K2LocalVar> NewVarAction = MakeShareable(new FEdGraphSchemaAction_K2LocalVar(Category, FText::FromName(Variable.VarName), TEXT(""), 0));
 					NewVarAction->SetVariableInfo(Variable.VarName, Func);
+					if ( GetDefault<UEditorExperimentalSettings>()->bUnifiedBlueprintEditor )
+					{
+						NewVarAction->SectionID = NodeSectionID::LOCAL_VARIABLE;
+					}
 					OutAllActions.AddAction(NewVarAction);
 				}
 			}
@@ -700,7 +919,7 @@ void SMyBlueprint::CollectAllActions(FGraphActionListBuilderBase& OutAllActions)
 	check(Blueprint);
 
 	EFieldIteratorFlags::SuperClassFlags FieldIteratorSuperFlag = EFieldIteratorFlags::IncludeSuper;
-	if(bShowUserVarsOnly)
+	if ( ShowUserVarsOnly() )
 	{
 		FieldIteratorSuperFlag = EFieldIteratorFlags::ExcludeSuper;
 	}
@@ -725,11 +944,11 @@ void SMyBlueprint::CollectAllActions(FGraphActionListBuilderBase& OutAllActions)
 				const int32 SortIndex = GetSortIndex( RootCategory ) + Action->SectionID;
 
 				static const FName DefaultCategory( TEXT( "Default" ));
-				if( RootCategory != DefaultCategory )
-				{
-					// Strip sections out of non default categories to enable them to be grouped together
-					Action->SectionID = NodeSectionID::USER_SORTED;
-				}
+				//if( RootCategory != DefaultCategory )
+				//{
+				//	// Strip sections out of non default categories to enable them to be grouped together
+				//	Action->SectionID = NodeSectionID::USER_SORTED;
+				//}
 				SortedActions.Add( SortIndex, Action );
 			}
 
@@ -826,10 +1045,17 @@ void SMyBlueprint::CollectAllActions(FGraphActionListBuilderBase& OutAllActions)
 
 		if (bShouldShowAsVar)
 		{
+			const bool bComponentProperty = Obj && Obj->PropertyClass ? Obj->PropertyClass->IsChildOf<UActorComponent>() : false;
+
+			// By default components go into the variable section under the component category unless a custom category is specified.
+			if ( bComponentProperty && CategoryName == NAME_None )
+			{
+				PropertyCategory = LOCTEXT("Components", "Components").ToString();
+			}
+
 			TSharedPtr<FEdGraphSchemaAction_K2Var> NewVarAction = MakeShareable(new FEdGraphSchemaAction_K2Var(PropertyCategory, PropertyDesc, PropertyTooltip.ToString(), 0));
 			NewVarAction->SetVariableInfo(PropertyName, Blueprint->SkeletonGeneratedClass);
-			const bool bComponentProperty = Obj && Obj->PropertyClass ? Obj->PropertyClass->IsChildOf<UActorComponent>() : false;
-			NewVarAction->SectionID = bComponentProperty ? NodeSectionID::COMPONENT : NodeSectionID::VARIABLE;
+			NewVarAction->SectionID = NodeSectionID::VARIABLE;
 			SortList.AddAction( UserCategoryName, NewVarAction );
 		}
 		else if (bShouldShowAsDelegate)
@@ -1043,19 +1269,60 @@ void SMyBlueprint::CollectAllActions(FGraphActionListBuilderBase& OutAllActions)
 		GetChildGraphs(Graph, OutAllActions);
 		GetChildEvents(Graph, NewFuncAction->SectionID, OutAllActions);
 	}
+
+	if ( GetLocalActionsListVisibility().IsVisible() )
+	{
+		GetLocalVariables(OutAllActions);
+	}
+}
+
+void SMyBlueprint::CollectStaticSections(TArray<int32>& StaticSectionIDs)
+{
+	if ( IsShowingEmptySections() )
+	{
+		StaticSectionIDs.Add(NodeSectionID::GRAPH);
+		StaticSectionIDs.Add(NodeSectionID::MACRO);
+		StaticSectionIDs.Add(NodeSectionID::FUNCTION);
+		StaticSectionIDs.Add(NodeSectionID::VARIABLE);
+		StaticSectionIDs.Add(NodeSectionID::DELEGATE);
+	}
+
+	if ( GetLocalActionsListVisibility().IsVisible() )
+	{
+		StaticSectionIDs.Add(NodeSectionID::LOCAL_VARIABLE);
+	}
 }
 
-ECheckBoxState SMyBlueprint::OnUserVarsCheckState() const
+bool SMyBlueprint::IsShowingInheritedVariables() const
 {
-	return !bShowUserVarsOnly ? ECheckBoxState::Checked : ECheckBoxState::Unchecked;
+	return GetMutableDefault<UBlueprintEditorSettings>()->bShowInheritedVariables;
 }
 
-void SMyBlueprint::OnUserVarsCheckStateChanged(ECheckBoxState InNewState)
+void SMyBlueprint::OnToggleShowInheritedVariables()
 {
-	bShowUserVarsOnly = (InNewState != ECheckBoxState::Checked);
+	UBlueprintEditorSettings* Settings = GetMutableDefault<UBlueprintEditorSettings>();
+	Settings->bShowInheritedVariables = !Settings->bShowInheritedVariables;
+	Settings->PostEditChange();
+	Settings->SaveConfig();
+
 	Refresh();
 }
 
+void SMyBlueprint::OnToggleShowEmptySections()
+{
+	UBlueprintEditorSettings* Settings = GetMutableDefault<UBlueprintEditorSettings>();
+	Settings->bShowEmptySections = !Settings->bShowEmptySections;
+	Settings->PostEditChange();
+	Settings->SaveConfig();
+
+	Refresh();
+}
+
+bool SMyBlueprint::IsShowingEmptySections() const
+{
+	return GetMutableDefault<UBlueprintEditorSettings>()->bShowEmptySections;
+}
+
 FReply SMyBlueprint::OnActionDragged( const TArray< TSharedPtr<FEdGraphSchemaAction> >& InActions, const FPointerEvent& MouseEvent )
 {
 	TSharedPtr<FEdGraphSchemaAction> InAction( InActions.Num() > 0 ? InActions[0] : NULL );
@@ -1383,10 +1650,18 @@ FEdGraphSchemaAction_K2Var* SMyBlueprint::SelectionAsVar() const
 
 FEdGraphSchemaAction_K2LocalVar* SMyBlueprint::SelectionAsLocalVar() const
 {
-	if(LocalGraphActionMenu.IsValid())
+	if ( GetDefault<UEditorExperimentalSettings>()->bUnifiedBlueprintEditor )
 	{
-		return SelectionAsType<FEdGraphSchemaAction_K2LocalVar>( LocalGraphActionMenu );
+		return SelectionAsType<FEdGraphSchemaAction_K2LocalVar>(GraphActionMenu);
 	}
+	else
+	{
+		if ( LocalGraphActionMenu.IsValid() )
+		{
+			return SelectionAsType<FEdGraphSchemaAction_K2LocalVar>(LocalGraphActionMenu);
+		}
+	}
+
 	return NULL;
 }
 
@@ -1524,21 +1799,36 @@ TSharedPtr<SWidget> SMyBlueprint::OnContextMenuOpening()
 	}
 	else
 	{
-		MenuBuilder.BeginSection("AddNewItem", LOCTEXT("AddOperations", "Add New"));
-		{
-			MenuBuilder.AddMenuEntry(FBlueprintEditorCommands::Get().AddNewVariable);
-			MenuBuilder.AddMenuEntry(FBlueprintEditorCommands::Get().AddNewLocalVariable);
-			MenuBuilder.AddMenuEntry(FBlueprintEditorCommands::Get().AddNewFunction);
-			MenuBuilder.AddMenuEntry(FBlueprintEditorCommands::Get().AddNewMacroDeclaration);
-			MenuBuilder.AddMenuEntry(FBlueprintEditorCommands::Get().AddNewEventGraph);
-			MenuBuilder.AddMenuEntry(FBlueprintEditorCommands::Get().AddNewDelegate);
-		}
-		MenuBuilder.EndSection();
+		BuildAddNewMenu(MenuBuilder);
 	}
 
 	return MenuBuilder.MakeWidget();
 }
 
+TSharedRef<SWidget> SMyBlueprint::CreateAddNewMenuWidget()
+{
+	const bool bShouldCloseWindowAfterMenuSelection = true;
+	FMenuBuilder MenuBuilder(bShouldCloseWindowAfterMenuSelection, BlueprintEditorPtr.Pin()->GetToolkitCommands());
+
+	BuildAddNewMenu(MenuBuilder);
+
+	return MenuBuilder.MakeWidget();
+}
+
+void SMyBlueprint::BuildAddNewMenu(FMenuBuilder& MenuBuilder)
+{
+	MenuBuilder.BeginSection("AddNewItem", LOCTEXT("AddOperations", "Add New"));
+	{
+		MenuBuilder.AddMenuEntry(FBlueprintEditorCommands::Get().AddNewVariable);
+		MenuBuilder.AddMenuEntry(FBlueprintEditorCommands::Get().AddNewLocalVariable);
+		MenuBuilder.AddMenuEntry(FBlueprintEditorCommands::Get().AddNewFunction);
+		MenuBuilder.AddMenuEntry(FBlueprintEditorCommands::Get().AddNewMacroDeclaration);
+		MenuBuilder.AddMenuEntry(FBlueprintEditorCommands::Get().AddNewEventGraph);
+		MenuBuilder.AddMenuEntry(FBlueprintEditorCommands::Get().AddNewDelegate);
+	}
+	MenuBuilder.EndSection();
+}
+
 bool SMyBlueprint::CanOpenGraph() const
 {
 	const FEdGraphSchemaAction_K2Graph* GraphAction = SelectionAsGraph();
diff --git a/Engine/Source/Editor/Kismet/Private/SMyBlueprint.h b/Engine/Source/Editor/Kismet/Private/SMyBlueprint.h
--- a/Engine/Source/Editor/Kismet/Private/SMyBlueprint.h
+++ b/Engine/Source/Editor/Kismet/Private/SMyBlueprint.h
@@ -98,7 +98,7 @@ public:
 	UBlueprint* GetBlueprintObj() const {return Blueprint;}
 
 	/** Gets whether we are showing user variables only or not */
-	bool ShowUserVarsOnly() const {return bShowUserVarsOnly;}
+	bool ShowUserVarsOnly() const { return !IsShowingInheritedVariables(); }
 
 	/** Gets our parent blueprint editor */
 	TWeakPtr<FBlueprintEditor> GetBlueprintEditor() {return BlueprintEditorPtr;}
@@ -135,6 +135,7 @@ private:
 
 	/** Callback used to populate all actions list in SGraphActionMenu */
 	void CollectAllActions(FGraphActionListBuilderBase& OutAllActions);
+	void CollectStaticSections(TArray<int32>& StaticSectionIDs);
 	void GetChildGraphs(UEdGraph* EdGraph, FGraphActionListBuilderBase& OutAllActions, FString ParentCategory = FString());
 	void GetChildEvents(UEdGraph const* EdGraph, int32 const SectionId, FGraphActionListBuilderBase& OutAllActions) const;
 	void GetLocalVariables(FGraphActionListBuilderBase& OutAllActions) const;
@@ -152,13 +153,23 @@ private:
 	void OnActionDoubleClicked(const TArray< TSharedPtr<FEdGraphSchemaAction> >& InActions);
 	TSharedPtr<SWidget> OnContextMenuOpening();
 
+	TSharedRef<SWidget> CreateAddNewMenuWidget();
+	void BuildAddNewMenu(FMenuBuilder& MenuBuilder);
+
 	void OnCategoryNameCommitted(const FText& InNewText, ETextCommit::Type InTextCommit, TWeakPtr< struct FGraphActionNode > InAction );
 	bool CanRequestRenameOnActionNode(TWeakPtr<struct FGraphActionNode> InSelectedNode) const;
 	FText OnGetSectionTitle( int32 InSectionID );
+	TSharedRef<SWidget> OnGetSectionWidget( TSharedRef<SWidget> RowWidget, int32 InSectionID );
+	EVisibility OnGetSectionTextVisibility(TWeakPtr<SWidget> RowWidget) const;
+	FReply OnAddButtonClickedOnSection(int32 InSectionID);
 
 	/** Support functions for checkbox to manage displaying user variables only */
-	ECheckBoxState OnUserVarsCheckState() const;
-	void OnUserVarsCheckStateChanged(ECheckBoxState InNewState);
+	bool IsShowingInheritedVariables() const;
+	void OnToggleShowInheritedVariables();
+
+	/** Support functions for view options for Show Empty Sections */
+	void OnToggleShowEmptySections();
+	bool IsShowingEmptySections() const;
 	
 	/** Helper function to open the selected graph */
 	void OpenGraph(FDocumentTracker::EOpenDocumentCause InCause);
@@ -216,9 +227,6 @@ private:
 	/** Graph Action Menu for displaying all local variables */
 	TSharedPtr<class SGraphActionMenu> LocalGraphActionMenu;
 
-	/** Whether or not we're only showing user-created variables */
-	bool bShowUserVarsOnly;
-
 	/** The last pin type used (including the function editor last pin type) */
 	FEdGraphPinType LastPinType;
 	FEdGraphPinType LastFunctionPinType;
diff --git a/Engine/Source/Editor/Kismet/Private/SSCSEditorViewport.cpp b/Engine/Source/Editor/Kismet/Private/SSCSEditorViewport.cpp
--- a/Engine/Source/Editor/Kismet/Private/SSCSEditorViewport.cpp
+++ b/Engine/Source/Editor/Kismet/Private/SSCSEditorViewport.cpp
@@ -223,6 +223,8 @@ EVisibility SSCSEditorViewport::GetWidgetVisibility() const
 
 TSharedRef<FEditorViewportClient> SSCSEditorViewport::MakeEditorViewportClient()
 {
+	FPreviewScene* PreviewScene = BlueprintEditorPtr.Pin()->GetPreviewScene();
+
 	// Construct a new viewport client instance.
 	ViewportClient = MakeShareable(new FSCSEditorViewportClient(BlueprintEditorPtr, PreviewScene, SharedThis(this)));
 	ViewportClient->SetRealtime(true);
@@ -354,8 +356,3 @@ EActiveTimerReturnType SSCSEditorViewport::DeferredUpdatePreview(double InCurren
 	bIsActiveTimerRegistered = false;
 	return EActiveTimerReturnType::Stop;
 }
-
-AActor* SSCSEditorViewport::GetPreviewActor() const
-{
-	return ViewportClient->GetPreviewActor();
-}
diff --git a/Engine/Source/Editor/Kismet/Private/SSCSEditorViewport.h b/Engine/Source/Editor/Kismet/Private/SSCSEditorViewport.h
--- a/Engine/Source/Editor/Kismet/Private/SSCSEditorViewport.h
+++ b/Engine/Source/Editor/Kismet/Private/SSCSEditorViewport.h
@@ -2,7 +2,6 @@
 
 #pragma once
 
-#include "PreviewScene.h"
 #include "SEditorViewport.h"
 
 /**
@@ -62,19 +61,6 @@ public:
 	 */
 	bool GetIsSimulateEnabled();
 
-	/** 
-	 * Provides access to the preview scene.
-	 */
-	const FPreviewScene& GetPreviewScene() const
-	{
-		return PreviewScene;
-	}
-
-	/**
-	 * Gets the current preview actor instance.
-	 */
-	AActor* GetPreviewActor() const;
-
 protected:
 	/**
 	 * Determines if the viewport widget is visible.
@@ -100,9 +86,6 @@ private:
 	/** Pointer back to editor tool (owner) */
 	TWeakPtr<class FBlueprintEditor> BlueprintEditorPtr;
 
-	/** Blueprint preview scene */
-	FPreviewScene PreviewScene;
-
 	/** Viewport client */
 	TSharedPtr<class FSCSEditorViewportClient> ViewportClient;
 
diff --git a/Engine/Source/Editor/Kismet/Public/BlueprintEditor.h b/Engine/Source/Editor/Kismet/Public/BlueprintEditor.h
--- a/Engine/Source/Editor/Kismet/Public/BlueprintEditor.h
+++ b/Engine/Source/Editor/Kismet/Public/BlueprintEditor.h
@@ -10,6 +10,7 @@
 #include "GraphEditor.h"
 #include "EditorUndoClient.h"
 #include "Engine/UserDefinedEnum.h"
+#include "PreviewScene.h"
 #include "Developer/Merge/Public/Merge.h" // for FOnMergeResolved
 
 class USCS_Node;
@@ -50,6 +51,7 @@ public:
 	FSelectionDetailsSummoner(TSharedPtr<class FAssetEditorToolkit> InHostingApp);
 
 	virtual TSharedRef<SWidget> CreateTabBody(const FWorkflowTabSpawnInfo& Info) const override;
+	virtual TSharedRef<SDockTab> SpawnTab(const FWorkflowTabSpawnInfo& Info) const override;
 
 	virtual FText GetTabToolTipText(const FWorkflowTabSpawnInfo& Info) const override
 	{
@@ -198,8 +200,33 @@ public:
 	TSharedRef<class SSCSEditor> GetSCSEditor() const {return SCSEditor.ToSharedRef();}
 	TSharedPtr<class SSCSEditorViewport> GetSCSViewport() const {return SCSViewport;}
 	TSharedPtr<class SMyBlueprint> GetMyBlueprintWidget() const {return MyBlueprintWidget;}
+
+	/**
+	 * Provides access to the preview actor.
+	 */
 	AActor* GetPreviewActor() const;
 
+	/**
+	 * Provides access to the preview scene.
+	 */
+	FPreviewScene* GetPreviewScene()
+	{
+		return &PreviewScene;
+	}
+
+	/**
+	* Creates/updates the preview actor for the given blueprint.
+	*
+	* @param InBlueprint			The Blueprint to create or update the preview for.
+	* @param bInForceFullUpdate	Force a full update to respawn actors.
+	*/
+	void UpdatePreviewActor(UBlueprint* InBlueprint, bool bInForceFullUpdate = false);
+
+	/**
+	* Destroy the Blueprint preview.
+	*/
+	void DestroyPreview();
+
 	TSharedPtr<class FBlueprintEditorToolbar> GetToolbarBuilder() {return Toolbar;}
 
 	/** @return the documentation location for this editor */
@@ -424,7 +451,8 @@ public:
 		NoSelection,
 		MyBlueprint,
 		GraphPanel,
-		BlueprintProps,
+		ClassSettings,
+		ClassDefaults,
 	};
 
 	/** Gets the UI selection state of this editor */
@@ -539,9 +567,12 @@ protected:
 	/** Called to start a quick find (focus the search box in the explorer tab) */
 	void FindInBlueprint_Clicked();
 
-	/** Edit Blueprint global options */
+	/** Edit the class settings aka Blueprint global options */
 	void EditGlobalOptions_Clicked();
 
+	/** Edit the class defaults */
+	void EditClassDefaults_Clicked();
+
 	/** Called to undo the last action */
 	void UndoGraphAction();
 
@@ -997,6 +1028,12 @@ protected:
 	/** Whether we are already in the process of closing this editor */
 	bool bEditorMarkedAsClosed;
 
+	/** Blueprint preview scene */
+	FPreviewScene PreviewScene;
+
+	/** The preview actor representing the current preview */
+	mutable TWeakObjectPtr<AActor> PreviewActorPtr;
+
 public:
 	// Tries to open the specified graph and bring it's document to the front (note: this can return NULL)
 	TSharedPtr<SGraphEditor> OpenGraphAndBringToFront(UEdGraph* Graph);
diff --git a/Engine/Source/Editor/Kismet/Public/SBlueprintEditorToolbar.h b/Engine/Source/Editor/Kismet/Public/SBlueprintEditorToolbar.h
--- a/Engine/Source/Editor/Kismet/Public/SBlueprintEditorToolbar.h
+++ b/Engine/Source/Editor/Kismet/Public/SBlueprintEditorToolbar.h
@@ -56,6 +56,7 @@ public:
 	
 	/** Edit Blueprint global options */
 	TSharedPtr<FUICommandInfo> EditGlobalOptions;
+	TSharedPtr<FUICommandInfo> EditClassDefaults;
 
 	/** Initialize commands */
 	virtual void RegisterCommands() override;
diff --git a/Engine/Source/Editor/Kismet/Public/SKismetInspector.h b/Engine/Source/Editor/Kismet/Public/SKismetInspector.h
--- a/Engine/Source/Editor/Kismet/Public/SKismetInspector.h
+++ b/Engine/Source/Editor/Kismet/Public/SKismetInspector.h
@@ -63,6 +63,9 @@ public:
 
 	TSharedPtr<class IDetailsView> GetPropertyView() const { return PropertyView; }
 
+	void SetOwnerTab(TSharedRef<SDockTab> Tab);
+	TSharedPtr<SDockTab> GetOwnerTab() const;
+
 protected:
 	/** Update the inspector window to show information on the supplied objects */
 	void UpdateFromObjects(const TArray<UObject*>& PropertyObjects, struct FKismetSelectionInfo& SelectionInfo, const FShowDetailsOptions& Options);
@@ -73,6 +76,9 @@ protected:
 	/** Pointer back to the kismet 2 tool that owns us */
 	TWeakPtr<FBlueprintEditor> Kismet2Ptr;
 
+	/** The tab that owns this details view. */
+	TWeakPtr<SDockTab> OwnerTab;
+
 	/** String used as the title above the property window */
 	FText PropertyViewTitle;
 
diff --git a/Engine/Source/Runtime/Slate/Private/Widgets/Layout/SBorder.cpp b/Engine/Source/Runtime/Slate/Private/Widgets/Layout/SBorder.cpp
--- a/Engine/Source/Runtime/Slate/Private/Widgets/Layout/SBorder.cpp
+++ b/Engine/Source/Runtime/Slate/Private/Widgets/Layout/SBorder.cpp
@@ -45,7 +45,7 @@ void SBorder::Construct( const SBorder::FArguments& InArgs )
  *
  * @param	InContent	The widget to use as content for the border
  */
-void SBorder::SetContent( const TSharedRef< SWidget >& InContent )
+void SBorder::SetContent( TSharedRef< SWidget > InContent )
 {
 	ChildSlot
 	[
diff --git a/Engine/Source/Runtime/Slate/Public/Widgets/Docking/SDockTab.h b/Engine/Source/Runtime/Slate/Public/Widgets/Docking/SDockTab.h
--- a/Engine/Source/Runtime/Slate/Public/Widgets/Docking/SDockTab.h
+++ b/Engine/Source/Runtime/Slate/Public/Widgets/Docking/SDockTab.h
@@ -96,7 +96,8 @@ public:
 
 	virtual FReply OnTouchEnded( const FGeometry& MyGeometry, const FPointerEvent& InTouchEvent ) override;
 
-	void SetContent( TSharedRef<SWidget> InContent );
+	virtual void SetContent(TSharedRef<SWidget> InContent) override;
+
 	void SetLeftContent( TSharedRef<SWidget> InContent );
 	void SetRightContent( TSharedRef<SWidget> InContent );
 
diff --git a/Engine/Source/Runtime/Slate/Public/Widgets/Layout/SBorder.h b/Engine/Source/Runtime/Slate/Public/Widgets/Layout/SBorder.h
--- a/Engine/Source/Runtime/Slate/Public/Widgets/Layout/SBorder.h
+++ b/Engine/Source/Runtime/Slate/Public/Widgets/Layout/SBorder.h
@@ -75,7 +75,7 @@ public:
 	 *
 	 * @param	InContent	The widget to use as content for the border
 	 */
-	void SetContent( const TSharedRef< SWidget >& InContent );
+	virtual void SetContent( TSharedRef< SWidget > InContent );
 
 	/**
 	 * Gets the content for this border
