[INTEGRATE] Change 2411673 by [REDACTED_PERSON][REDACTED_HANDLE] on 2015/01/19 19:20:57

	SCS Editor now shows a root actor that you can select and see properties.
	Draging and dropping components to the graph from SCS editor creates Getters.
	Add components drop down has more common entries and better naming for some.

diff --git a/Engine/Source/Editor/Documentation/Private/Documentation.cpp b/Engine/Source/Editor/Documentation/Private/Documentation.cpp
--- a/Engine/Source/Editor/Documentation/Private/Documentation.cpp
+++ b/Engine/Source/Editor/Documentation/Private/Documentation.cpp
@@ -194,7 +194,7 @@ bool FDocumentation::PageExists(const FString& Link, const FCultureRef& Culture)
 	return FPaths::FileExists(SourcePath);
 }
 
-TSharedRef< class SToolTip > FDocumentation::CreateToolTip( const TAttribute<FText>& Text, const TSharedPtr<SWidget>& OverrideContent, const FString& Link, const FString& ExcerptName ) const
+TSharedRef< class SToolTip > FDocumentation::CreateToolTip(const TAttribute<FText>& Text, const TSharedPtr<SWidget>& OverrideContent, const FString& Link, const FString& ExcerptName) const
 {
 	TSharedPtr< SDocumentationToolTip > DocToolTip;
 
@@ -227,4 +227,30 @@ TSharedRef< class SToolTip > FDocumentation::CreateToolTip( const TAttribute<FTe
 		];
 }
 
+TSharedRef< class SToolTip > FDocumentation::CreateToolTip(const TAttribute<FText>& Text, const TSharedRef<SWidget>& OverrideContent, const TSharedPtr<SVerticalBox>& DocVerticalBox, const FString& Link, const FString& ExcerptName) const
+{
+	TSharedRef<SDocumentationToolTip> DocToolTip =
+		SNew(SDocumentationToolTip)
+		.Text(Text)
+		.DocumentationLink(Link)
+		.ExcerptName(ExcerptName)
+		.AddDocumentation(false)
+		.DocumentationMargin(7)
+		[
+			OverrideContent
+		];
+
+	if (DocVerticalBox.IsValid())
+	{
+		DocToolTip->AddDocumentation(DocVerticalBox);
+	}
+
+	return SNew(SToolTip)
+		.TextMargin(1)
+		.IsInteractive(DocToolTip, &SDocumentationToolTip::IsInteractive)
+		[
+			DocToolTip
+		];
+}
+
 #undef LOCTEXT_NAMESPACE
diff --git a/Engine/Source/Editor/Documentation/Private/Documentation.h b/Engine/Source/Editor/Documentation/Private/Documentation.h
--- a/Engine/Source/Editor/Documentation/Private/Documentation.h
+++ b/Engine/Source/Editor/Documentation/Private/Documentation.h
@@ -32,6 +32,8 @@ public:
 	virtual bool PageExists(const FString& Link, const FCultureRef& Culture) const override;
 
 	virtual TSharedRef< class SToolTip > CreateToolTip( const TAttribute<FText>& Text, const TSharedPtr<SWidget>& OverrideContent, const FString& Link, const FString& ExcerptName ) const override;
+	
+	virtual TSharedRef< class SToolTip > CreateToolTip(const TAttribute<FText>& Text, const TSharedRef<SWidget>& OverrideContent, const TSharedPtr<SVerticalBox>& DocVerticalBox, const FString& Link, const FString& ExcerptName) const override;
 
 private:
 
diff --git a/Engine/Source/Editor/Documentation/Private/SDocumentationToolTip.cpp b/Engine/Source/Editor/Documentation/Private/SDocumentationToolTip.cpp
--- a/Engine/Source/Editor/Documentation/Private/SDocumentationToolTip.cpp
+++ b/Engine/Source/Editor/Documentation/Private/SDocumentationToolTip.cpp
@@ -18,6 +18,8 @@ void SDocumentationToolTip::Construct( const FArguments& InArgs )
 	ColorAndOpacity = InArgs._ColorAndOpacity;
 	DocumentationLink = InArgs._DocumentationLink;
 	IsDisplayingDocumentationLink = false;
+	bAddDocumentation = InArgs._AddDocumentation;
+	DocumentationMargin = InArgs._DocumentationMargin;
 
 	if ( !DocumentationLink.IsEmpty() )
 	{
@@ -39,9 +41,7 @@ void SDocumentationToolTip::Construct( const FArguments& InArgs )
 
 	ChildSlot
 	[
-		SAssignNew( WidgetContent, SBorder )
-		.BorderImage( FEditorStyle::GetBrush( "NoBorder" ) )
-		.Padding( 0 )
+		SAssignNew(WidgetContent, SBox)
 		[
 			SimpleTipContent.ToSharedRef()
 		]
@@ -75,15 +75,19 @@ void SDocumentationToolTip::ConstructSimpleTipContent()
 			+SVerticalBox::Slot()
 			.FillHeight( 1.0f )
 			[
-				SNew( STextBlock )
-				.Text( TextContent )
-				.TextStyle( &StyleInfo )
-				.ColorAndOpacity( ColorAndOpacity )
-				.WrapTextAt_Static( &SToolTip::GetToolTipWrapWidth )
+				OverrideContent.ToSharedRef()
 			]
 		];
 	}
 
+	if (bAddDocumentation)
+	{
+		AddDocumentation(VerticalBox);
+	}
+}
+
+void SDocumentationToolTip::AddDocumentation(TSharedPtr< SVerticalBox > VerticalBox)
+{
 	if ( !DocumentationLink.IsEmpty() )
 	{
 		IsDisplayingDocumentationLink = GEditor->EditorUserSettings->bDisplayDocumentationLink;
@@ -259,12 +263,16 @@ void SDocumentationToolTip::ConstructFullTipContent()
 		{
 			TSharedPtr< SVerticalBox > Box;
 			FullTipContent = 
-				SAssignNew( Box, SVerticalBox )
-				+SVerticalBox::Slot()
-				.HAlign( HAlign_Center )
-				.AutoHeight()
+				SNew(SBox)
+				.Padding(DocumentationMargin)
 				[
-					Excerpts[ ExcerptIndex ].Content.ToSharedRef()
+					SAssignNew(Box, SVerticalBox)
+					+ SVerticalBox::Slot()
+					.HAlign(HAlign_Center)
+					.AutoHeight()
+					[
+						Excerpts[ExcerptIndex].Content.ToSharedRef()
+					]
 				];
 
 			FString* FullDocumentationLink = Excerpts[ ExcerptIndex ].Variables.Find( TEXT("ToolTipFullLink") );
diff --git a/Engine/Source/Editor/Documentation/Public/IDocumentation.h b/Engine/Source/Editor/Documentation/Public/IDocumentation.h
--- a/Engine/Source/Editor/Documentation/Public/IDocumentation.h
+++ b/Engine/Source/Editor/Documentation/Public/IDocumentation.h
@@ -191,4 +191,6 @@ public:
 	virtual bool PageExists(const FString& Link, const FCultureRef& Culture) const = 0;
 
 	virtual TSharedRef< class SToolTip > CreateToolTip( const TAttribute<FText>& Text, const TSharedPtr<SWidget>& OverrideContent, const FString& Link, const FString& ExcerptName ) const = 0;
+
+	virtual TSharedRef< class SToolTip > CreateToolTip(const TAttribute<FText>& Text, const TSharedRef<SWidget>& OverrideContent, const TSharedPtr<SVerticalBox>& DocVerticalBox, const FString& Link, const FString& ExcerptName) const = 0;
 };
diff --git a/Engine/Source/Editor/Documentation/Public/SDocumentationToolTip.h b/Engine/Source/Editor/Documentation/Public/SDocumentationToolTip.h
--- a/Engine/Source/Editor/Documentation/Public/SDocumentationToolTip.h
+++ b/Engine/Source/Editor/Documentation/Public/SDocumentationToolTip.h
@@ -2,7 +2,7 @@
 #pragma once
 #include "IDocumentationPage.h"
 
-class SDocumentationToolTip : public SCompoundWidget
+class DOCUMENTATION_API SDocumentationToolTip : public SCompoundWidget
 {
 
 public:
@@ -11,6 +11,8 @@ public:
 		: _Text()
 		, _Style( TEXT("Documentation.SDocumentationTooltip") )
 		, _ColorAndOpacity( FLinearColor( 1.0f, 1.0f, 1.0f, 1.0f ) )
+		, _AddDocumentation( true )
+		, _DocumentationMargin(0)
 		, _Content()
 		{}
 
@@ -23,6 +25,11 @@ public:
 		/** Font color and opacity */
 		SLATE_ATTRIBUTE( FSlateColor, ColorAndOpacity )
 
+		/**  */
+		SLATE_ARGUMENT( bool, AddDocumentation )
+
+		SLATE_ARGUMENT( FMargin, DocumentationMargin )
+
 		/**  */
 		SLATE_ARGUMENT( FString, DocumentationLink )
 
@@ -50,6 +57,14 @@ public:
 		return TextContent.Get();
 	}
 
+	/**
+	 * Adds slots to the provided Vertical Box containing the documentation information.
+	 * If you specify not to add it (AddDocumentation = false) you may call this externally to do custom tooltip layout
+	 *
+	 * @param	VerticalBox	The vertical box to add it to
+	 */
+	void AddDocumentation(TSharedPtr< SVerticalBox > VerticalBox);
+
 private:
 
 	void ConstructSimpleTipContent();
@@ -73,7 +88,7 @@ private:
 	FString ExcerptName;
 
 	/** Content widget */
-	TSharedPtr< SBorder > WidgetContent;
+	TSharedPtr< SBox > WidgetContent;
 
 	TSharedPtr< SWidget > SimpleTipContent;
 	bool IsDisplayingDocumentationLink;
@@ -82,4 +97,7 @@ private:
 
 	TSharedPtr< IDocumentationPage > DocumentationPage;
 	bool IsShowingFullTip;
+
+	bool bAddDocumentation;
+	FMargin DocumentationMargin;
 };
diff --git a/Engine/Source/Editor/EditorStyle/Private/SlateEditorStyle.cpp b/Engine/Source/Editor/EditorStyle/Private/SlateEditorStyle.cpp
--- a/Engine/Source/Editor/EditorStyle/Private/SlateEditorStyle.cpp
+++ b/Engine/Source/Editor/EditorStyle/Private/SlateEditorStyle.cpp
@@ -862,8 +862,6 @@ void FSlateEditorStyle::FStyle::SetupGeneralStyles()
 			);
 	}
 
-	Set( "ToolTip.Background", new BOX_BRUSH( "Old/ToolTip_Background", FMargin(8.0f/64.0f) ) );
-
 	// Throbber
 	{
 		Set( "Throbber.Chunk", new IMAGE_BRUSH( "Old/Throbber/Throbber_Piece", FVector2D(16,16) ) );
@@ -3897,6 +3895,7 @@ void FSlateEditorStyle::FStyle::SetupGraphEditorStyles()
 
 		Set("SCSEditor.ToggleComponentEditing" , new IMAGE_BRUSH( "Icons/icon_translate_40x", Icon40x40) );
 		Set("SCSEditor.ToggleComponentEditing.Small" , new IMAGE_BRUSH( "Icons/icon_translate_40x", Icon20x20) );
+		Set("SCSEditor.TileViewTooltip.NonContentBorder", new BOX_BRUSH("/Docking/TabContentArea", FMargin(4 / 16.0f)));
 
 		// Timeline Editor
 		{
diff --git a/Engine/Source/Editor/FoliageEdit/Private/SFoliageEditMeshDisplayItem.cpp b/Engine/Source/Editor/FoliageEdit/Private/SFoliageEditMeshDisplayItem.cpp
--- a/Engine/Source/Editor/FoliageEdit/Private/SFoliageEditMeshDisplayItem.cpp
+++ b/Engine/Source/Editor/FoliageEdit/Private/SFoliageEditMeshDisplayItem.cpp
@@ -1481,7 +1481,7 @@ void SFoliageEditMeshDisplayItem::Construct(const FArguments& InArgs)
 				.BorderBackgroundColor(TAttribute<FSlateColor>(this, &SFoliageEditMeshDisplayItem::GetBorderColor))
 				[
 					SNew(SBorder)
-					.BorderImage(FEditorStyle::GetBrush("ToolTip.Background"))
+					.BorderImage(FCoreStyle::Get().GetBrush("ToolTip.Background"))
 					.Padding(0.0f)
 					[
 						SNew(SCheckBox)
diff --git a/Engine/Source/Editor/Kismet/Private/BlueprintEditor.cpp b/Engine/Source/Editor/Kismet/Private/BlueprintEditor.cpp
--- a/Engine/Source/Editor/Kismet/Private/BlueprintEditor.cpp
+++ b/Engine/Source/Editor/Kismet/Private/BlueprintEditor.cpp
@@ -689,7 +689,27 @@ AActor* FBlueprintEditor::GetSCSEditorActorContext() const
 	return nullptr;
 }
 
-void FBlueprintEditor::OnSCSEditorTreeViewSelectionChanged(const TArray<FSCSEditorTreeNodePtrType>& SelectedNodes)
+void FBlueprintEditor::OnRootSelected(AActor * DefaultActor)
+{
+	if (Inspector.IsValid())
+	{	
+		// Clear the my blueprints selection
+		MyBlueprintWidget->ClearGraphActionMenuSelection();
+
+		// Build an object list
+		TArray<UObject*> InspectorObjects;
+		InspectorObjects.Add(DefaultActor);
+		
+		// Update the details panel
+		FString Title;
+		DefaultActor->GetName(Title);
+		SKismetInspector::FShowDetailsOptions Options(FText::FromString(Title), true);
+		Options.bShowComponents = false;
+		Inspector->ShowDetailsForObjects(InspectorObjects, Options);
+	}
+}
+
+void FBlueprintEditor::OnSelectionUpdated(const TArray<FSCSEditorTreeNodePtrType>& SelectedNodes)
 {
 	if (SCSViewport.IsValid())
 	{
@@ -711,10 +731,16 @@ void FBlueprintEditor::OnSCSEditorTreeViewSelectionChanged(const TArray<FSCSEdit
 			PrimitiveComponents[Idx]->PushSelectionToProxy();
 		}
 	}
-}
 
-void FBlueprintEditor::OnSCSEditorUpdateSelectionFromNodes(const TArray<FSCSEditorTreeNodePtrType>& SelectedNodes)
+	if (Inspector.IsValid())
 	{
+		// Clear the my blueprints selection
+		bool bSingleLayoutBPEditor = GetDefault<UEditorExperimentalSettings>()->bUnifiedBlueprintEditor;
+		if (bSingleLayoutBPEditor)
+		{
+			MyBlueprintWidget->ClearGraphActionMenuSelection();
+		}
+
 		// Convert the selection set to an array of UObject* pointers
 		FText InspectorTitle = FText::GetEmpty();
 		TArray<UObject*> InspectorObjects;
@@ -724,7 +750,16 @@ void FBlueprintEditor::OnSCSEditorUpdateSelectionFromNodes(const TArray<FSCSEdit
 			auto NodePtr = *NodeIt;
 			if (NodePtr.IsValid())
 			{
-				UActorComponent* EditableComponent = NodePtr->GetEditableComponentTemplate(GetBlueprintObj());
+				UActorComponent* EditableComponent = nullptr;
+				if (NodePtr->CanEditDefaults())
+				{
+					EditableComponent = NodePtr->GetComponentTemplate();
+				}
+				else if (!NodePtr->IsNative() && NodePtr->IsInherited())
+				{
+					EditableComponent = NodePtr->GetOverridenComponentTemplate(GetBlueprintObj(), true);
+				}
+
 				if (EditableComponent)
 				{
 					InspectorTitle = FText::FromString(NodePtr->GetDisplayString());
@@ -733,9 +768,7 @@ void FBlueprintEditor::OnSCSEditorUpdateSelectionFromNodes(const TArray<FSCSEdit
 			}
 		}
 
-	// Update the details panel
-	if (Inspector.IsValid())
-	{
+		// Update the details panel
 		SKismetInspector::FShowDetailsOptions Options(InspectorTitle, true);
 		Inspector->ShowDetailsForObjects(InspectorObjects, Options);
 	}
@@ -1207,8 +1240,8 @@ void FBlueprintEditor::EnsureBlueprintIsUpToDate(UBlueprint* BlueprintObj)
 				.ActorContext(this, &FBlueprintEditor::GetSCSEditorActorContext)
 				.PreviewActor(this, &FBlueprintEditor::GetPreviewActor)
 				.AllowEditing(this, &FBlueprintEditor::InEditingMode)
-				.OnTreeViewSelectionChanged(this, &FBlueprintEditor::OnSCSEditorTreeViewSelectionChanged)
-				.OnUpdateSelectionFromNodes(this, &FBlueprintEditor::OnSCSEditorUpdateSelectionFromNodes);
+				.OnRootSelected(this, &FBlueprintEditor::OnRootSelected)
+				.OnSelectionUpdated(this, &FBlueprintEditor::OnSelectionUpdated);
 			SCSViewport = SAssignNew(SCSViewport, SSCSEditorViewport)
 				.BlueprintEditor(SharedThis(this));
 		}
@@ -1926,7 +1959,7 @@ void FBlueprintEditor::FocusInspectorOnGraphSelection(const FGraphPanelSelection
 		else
 		{
 			// Clear the Inspector if nothing is selected, the MyBlueprints window will set the selection if it means to, but we need to make sure invalid nodes aren't still appearing
-			Inspector->ShowDetailsForSingleObject(nullptr);
+			//Inspector->ShowDetailsForSingleObject(nullptr);
 		}
 	}
 }
@@ -1998,15 +2031,15 @@ void FBlueprintEditor::CreateDefaultTabContents(const TArray<UBlueprint*>& InBlu
 		InBlueprint->ParentClass->IsChildOf(AActor::StaticClass()) && 
 		InBlueprint->SimpleConstructionScript )
 	{
-	SCSEditor = SAssignNew(SCSEditor, SSCSEditor)
-		.ActorContext(this, &FBlueprintEditor::GetSCSEditorActorContext)
-		.PreviewActor(this, &FBlueprintEditor::GetPreviewActor)
-		.AllowEditing(this, &FBlueprintEditor::InEditingMode)
-			.OnTreeViewSelectionChanged(this, &FBlueprintEditor::OnSCSEditorTreeViewSelectionChanged)
-			.OnUpdateSelectionFromNodes(this, &FBlueprintEditor::OnSCSEditorUpdateSelectionFromNodes);
-	SCSViewport = SAssignNew(SCSViewport, SSCSEditorViewport)
-		.BlueprintEditor(SharedThis(this));
-}
+		SCSEditor = SAssignNew(SCSEditor, SSCSEditor)
+			.ActorContext(this, &FBlueprintEditor::GetSCSEditorActorContext)
+			.PreviewActor(this, &FBlueprintEditor::GetPreviewActor)
+			.AllowEditing(this, &FBlueprintEditor::InEditingMode)
+			.OnRootSelected(this, &FBlueprintEditor::OnRootSelected)
+			.OnSelectionUpdated(this, &FBlueprintEditor::OnSelectionUpdated);
+		SCSViewport = SAssignNew(SCSViewport, SSCSEditorViewport)
+			.BlueprintEditor(SharedThis(this));
+	}
 }
 
 void FBlueprintEditor::OnLogTokenClicked(const TSharedRef<IMessageToken>& Token)
diff --git a/Engine/Source/Editor/Kismet/Private/SCSDiff.cpp b/Engine/Source/Editor/Kismet/Private/SCSDiff.cpp
--- a/Engine/Source/Editor/Kismet/Private/SCSDiff.cpp
+++ b/Engine/Source/Editor/Kismet/Private/SCSDiff.cpp
@@ -30,7 +30,7 @@ FSCSDiff::FSCSDiff(const UBlueprint* InBlueprint)
 				.ActorContext(InBlueprint->GeneratedClass->GetDefaultObject<AActor>())
 				.AllowEditing(false)
 				.HideComponentClassCombo(true)
-				.OnUpdateSelectionFromNodes(SSCSEditor::FOnUpdateSelectionFromNodes::CreateRaw(this, &FSCSDiff::OnSCSEditorUpdateSelectionFromNodes))
+				.OnSelectionUpdated(SSCSEditor::FOnSelectionUpdated::CreateRaw(this, &FSCSDiff::OnSCSEditorUpdateSelectionFromNodes))
 				.OnHighlightPropertyInDetailsView(SSCSEditor::FOnHighlightPropertyInDetailsView::CreateRaw(this, &FSCSDiff::OnSCSEditorHighlightPropertyInDetailsView))
 		]
 		+ SSplitter::Slot()
diff --git a/Engine/Source/Editor/Kismet/Private/SKismetInspector.cpp b/Engine/Source/Editor/Kismet/Private/SKismetInspector.cpp
--- a/Engine/Source/Editor/Kismet/Private/SKismetInspector.cpp
+++ b/Engine/Source/Editor/Kismet/Private/SKismetInspector.cpp
@@ -505,6 +505,7 @@ void SKismetInspector::UpdateFromObjects(const TArray<UObject*>& PropertyObjects
 	}
 
 	PropertyViewTitle = Options.ForcedTitle;
+	bShowComponents = Options.bShowComponents;
 
 	// Update our context-sensitive editing widget
 	ContextualEditingBorderWidget->SetContent( MakeContextualEditingWidget(SelectionInfo, Options) );
@@ -551,6 +552,13 @@ bool SKismetInspector::IsPropertyVisible( const FPropertyAndParent& PropertyAndP
 		return false;
 	}
 
+	bool bIsComponent = (ObjectProperty != nullptr && ObjectProperty->PropertyClass->IsChildOf(UActorComponent::StaticClass()));
+	if (!bShowComponents && bIsComponent)
+	{
+		// Don't show sub components properties, thats what selecting components in the component tree is for.
+		return false;
+	}
+
 	// Filter down to selected properties only if set.
 	// If the current property is selected then it is visible or if its parent is selected and the current property did not fail any of the above tests it should be visible.
 	if( SelectedObjectProperties.Find( &Property ) || ( PropertyAndParent.ParentProperty && SelectedObjectProperties.Find( PropertyAndParent.ParentProperty ) ) )
diff --git a/Engine/Source/Editor/Kismet/Private/SSCSEditor.cpp b/Engine/Source/Editor/Kismet/Private/SSCSEditor.cpp
--- a/Engine/Source/Editor/Kismet/Private/SSCSEditor.cpp
+++ b/Engine/Source/Editor/Kismet/Private/SSCSEditor.cpp
@@ -125,12 +125,12 @@ FReply SSCSEditorDragDropTree::OnDrop( const FGeometry& MyGeometry, const FDragD
 //////////////////////////////////////////////////////////////////////////
 //
 
-class FSCSRowDragDropOp : public FDragDropOperation
+class FSCSRowDragDropOp : public FKismetVariableDragDropAction
 {
 public:
-	DRAG_DROP_OPERATOR_TYPE(FSCSRowDragDropOp, FDragDropOperation)
+	DRAG_DROP_OPERATOR_TYPE(FSCSRowDragDropOp, FKismetVariableDragDropAction)
 
-	/** Available drop actions */
+		/** Available drop actions */
 	enum EDropActionType
 	{
 		DropAction_None,
@@ -148,52 +148,60 @@ public:
 
 	/** The type of drop action that's pending while dragging */
 	EDropActionType PendingDropAction;
-	
+
 	/** The widget decorator to use */
-	virtual TSharedPtr<SWidget> GetDefaultDecorator() const override
-	{
-		return SNew(SBorder)
-			.BorderImage(FEditorStyle::GetBrush("Graph.ConnectorFeedback.Border"))
-			[				
-				SNew(SHorizontalBox)
-				+SHorizontalBox::Slot()
-				.AutoWidth()
-				.Padding(0,0,3,0)
-				[
-					SNew(SImage)
-					.Image(this, &FSCSRowDragDropOp::GetIcon)
-				]
-				+SHorizontalBox::Slot()
-				.AutoWidth()
-				.VAlign(VAlign_Center)
-				[
-					SNew(STextBlock) 
-					.Text(this, &FSCSRowDragDropOp::GetHoverText)
-				]
-			];
-	}
+	virtual TSharedPtr<SWidget> GetDefaultDecorator() const override;
+	FText GetHoverText() const;
+	const FSlateBrush* GetIcon() const;
+	static TSharedRef<FSCSRowDragDropOp> New(FName InVariableName, UStruct* InVariableSource, FNodeCreationAnalytic AnalyticCallback);
+};
 
-	FText GetHoverText() const
-	{
-		return !CurrentHoverText.IsEmpty()
-			? CurrentHoverText
-			: LOCTEXT("DropActionToolTip_InvalidDropTarget", "Cannot drop here.");
-	}
+TSharedPtr<SWidget> FSCSRowDragDropOp::GetDefaultDecorator() const
+{
+	return SNew(SBorder)
+		.BorderImage(FEditorStyle::GetBrush("Graph.ConnectorFeedback.Border"))
+		[
+			SNew(SHorizontalBox)
+			+ SHorizontalBox::Slot()
+			.AutoWidth()
+			.Padding(0, 0, 3, 0)
+			[
+				SNew(SImage)
+				.Image(this, &FSCSRowDragDropOp::GetIcon)
+			]
+			+ SHorizontalBox::Slot()
+			.AutoWidth()
+			.VAlign(VAlign_Center)
+			[
+				SNew(STextBlock)
+				.Text(this, &FSCSRowDragDropOp::GetHoverText)
+			]
+		];
+}
 
-	const FSlateBrush* GetIcon() const
-	{
-		return PendingDropAction != DropAction_None
-			? FEditorStyle::GetBrush(TEXT("Graph.ConnectorFeedback.OK"))
-			: FEditorStyle::GetBrush(TEXT("Graph.ConnectorFeedback.Error"));
-	}
+FText FSCSRowDragDropOp::GetHoverText() const
+{
+	return !CurrentHoverText.IsEmpty()
+		? CurrentHoverText
+		: LOCTEXT("DropActionToolTip_InvalidDropTarget", "Cannot drop here.");
+}
 
-	static TSharedRef<FSCSRowDragDropOp> New()
-	{
-		TSharedPtr<FSCSRowDragDropOp> Operation = MakeShareable(new FSCSRowDragDropOp);
-		Operation->Construct();
-		return Operation.ToSharedRef();
-	}
-};
+const FSlateBrush* FSCSRowDragDropOp::GetIcon() const
+{
+	return PendingDropAction != DropAction_None
+		? FEditorStyle::GetBrush(TEXT("Graph.ConnectorFeedback.OK"))
+		: FEditorStyle::GetBrush(TEXT("Graph.ConnectorFeedback.Error"));
+}
+
+TSharedRef<FSCSRowDragDropOp> FSCSRowDragDropOp::New(FName InVariableName, UStruct* InVariableSource, FNodeCreationAnalytic AnalyticCallback)
+{
+	TSharedPtr<FSCSRowDragDropOp> Operation = MakeShareable(new FSCSRowDragDropOp);
+	Operation->VariableName = InVariableName;
+	Operation->VariableSource = InVariableSource;
+	Operation->AnalyticCallback = AnalyticCallback;
+	Operation->Construct();
+	return Operation.ToSharedRef();
+}
 
 //////////////////////////////////////////////////////////////////////////
 // FSCSEditorTreeNode
@@ -1014,7 +1022,7 @@ SSCS_RowWidget::~SSCS_RowWidget()
 	}
 }
 
-
+BEGIN_SLATE_FUNCTION_BUILD_OPTIMIZATION
 TSharedRef<SWidget> SSCS_RowWidget::GenerateWidgetForColumn( const FName& ColumnName )
 {
 	if(ColumnName == SCS_ColumnName_ComponentClass)
@@ -1035,13 +1043,11 @@ TSharedRef<SWidget> SSCS_RowWidget::GenerateWidgetForColumn( const FName& Column
 				.IsReadOnly( !NodePtr->CanRename() || (SCSEditor.IsValid() && !SCSEditor.Pin()->IsEditingAllowed()) );
 
 		NodePtr->SetRenameRequestedDelegate(FSCSEditorTreeNode::FOnRenameRequested::CreateSP(InlineWidget.Get(), &SInlineEditableTextBlock::EnterEditingMode));
+		
+		TSharedRef<SToolTip> Tooltip = CreateToolTipWidget();
 
 		return	SNew(SHorizontalBox)
-				.ToolTip(IDocumentation::Get()->CreateToolTip(
-					TAttribute<FText>(this, &SSCS_RowWidget::GetTooltipText),
-					NULL,
-					GetDocumentationLink(),
-					GetDocumentationExcerptName()))
+				.ToolTip(Tooltip)
 				+SHorizontalBox::Slot()
 					.AutoWidth()
 					.VAlign(VAlign_Center)
@@ -1082,7 +1088,7 @@ TSharedRef<SWidget> SSCS_RowWidget::GenerateWidgetForColumn( const FName& Column
 				SNew(STextBlock)
 				.Visibility(this, &SSCS_RowWidget::GetAssetVisibility)
 				.Text(this, &SSCS_RowWidget::GetAssetName)
-				.ToolTipText(this, &SSCS_RowWidget::GetAssetNameTooltip)
+				.ToolTipText(this, &SSCS_RowWidget::GetAssetPath)
 			];
 	}
 	else if (ColumnName == SCS_ColumnName_Mobility)
@@ -1111,6 +1117,125 @@ TSharedRef<SWidget> SSCS_RowWidget::GenerateWidgetForColumn( const FName& Column
 	}
 }
 
+void AddToToolTipInfoBox(const TSharedRef<SVerticalBox>& InfoBox, const FText& Key, TSharedRef<SWidget> ValueIcon, const TAttribute<FText>& Value, bool bImportant)
+{
+	FWidgetStyle ImportantStyle;
+	ImportantStyle.SetForegroundColor(FLinearColor(1, 0.5, 0, 1));
+
+	InfoBox->AddSlot()
+		.AutoHeight()
+		.Padding(0, 1)
+		[
+			SNew(SHorizontalBox)
+
+			+ SHorizontalBox::Slot()
+			.AutoWidth()
+			.Padding(0, 0, 4, 0)
+			[
+				SNew(STextBlock).Text(FText::Format(LOCTEXT("AssetViewTooltipFormat", "{0}:"), Key))
+				.ColorAndOpacity(bImportant ? ImportantStyle.GetSubduedForegroundColor() : FSlateColor::UseSubduedForeground())
+			]
+
+			+ SHorizontalBox::Slot()
+			.AutoWidth()
+			[
+				ValueIcon
+			]
+
+			+ SHorizontalBox::Slot()
+			.AutoWidth()
+			[
+				SNew(STextBlock).Text(Value)
+				.ColorAndOpacity(bImportant ? ImportantStyle.GetForegroundColor() : FSlateColor::UseForeground())
+			]
+		];
+}
+
+TSharedRef<SToolTip> SSCS_RowWidget::CreateToolTipWidget() const
+{
+	bool bSingleLayoutBPEditor = GetDefault<UEditorExperimentalSettings>()->bUnifiedBlueprintEditor;
+	if (!bSingleLayoutBPEditor)
+	{
+		return IDocumentation::Get()->CreateToolTip(
+		TAttribute<FText>(this, &SSCS_RowWidget::GetTooltipText),
+		NULL,
+		GetDocumentationLink(),
+		GetDocumentationExcerptName());
+	}
+
+	// Create a box to hold every line of info in the body of the tooltip
+	TSharedRef<SVerticalBox> InfoBox = SNew(SVerticalBox);
+
+	// Add asset if applicable to this node
+	if (GetAssetVisibility() == EVisibility::Visible)
+	{
+		AddToToolTipInfoBox(InfoBox, LOCTEXT("TooltipAsset", "Asset"), SNullWidget::NullWidget, TAttribute<FText>(this, &SSCS_RowWidget::GetAssetName), false);
+	}
+
+	TWeakObjectPtr<USCS_Node> SCSNode(NodePtr->GetSCSNode());
+
+	// Add mobility
+	TSharedRef<SImage> MobilityIcon = SNew(SImage).Image(this, &SSCS_RowWidget::GetMobilityIconImage, SCSNode);
+	AddToToolTipInfoBox(InfoBox, LOCTEXT("TooltipMobility", "Mobility"), MobilityIcon, TAttribute<FText>::Create(TAttribute<FText>::FGetter::CreateSP(this, &SSCS_RowWidget::GetMobilityToolTipText, SCSNode)), false);
+	
+	TSharedRef<SBorder> TooltipContent = SNew(SBorder)
+		.Padding(4)
+		.BorderImage(FEditorStyle::GetBrush("SCSEditor.TileViewTooltip.NonContentBorder"))
+		[
+			SNew(SVerticalBox)
+			
+			+SVerticalBox::Slot()
+			.AutoHeight()
+			.Padding(0, 0, 0, 4)
+			[
+				SNew(SBorder)
+				.Padding(0)
+				.BorderImage(FEditorStyle::GetBrush("ToolPanel.GroupBorder"))
+				[
+					SNew(SVerticalBox)
+
+					+ SVerticalBox::Slot()
+					.AutoHeight()
+					[
+						SNew(SHorizontalBox)
+
+						+ SHorizontalBox::Slot()
+						.AutoWidth()
+						.VAlign(VAlign_Center)
+						.Padding(4)
+						[
+							SNew(STextBlock)
+							.Text(this, &SSCS_RowWidget::GetTooltipText)
+							.Font(FEditorStyle::GetFontStyle("ContentBrowser.TileViewTooltip.NameFont"))
+						]
+						/*
+						+ SHorizontalBox::Slot()
+						.AutoWidth()
+						.VAlign(VAlign_Center)
+						[
+							SNew(STextBlock).Text(this, SSCS_RowWidget::GetTypeAsString())
+							.HighlightText(HighlightText)
+						]*/
+					]
+				]
+			]
+
+			+ SVerticalBox::Slot()
+			.AutoHeight()
+			[
+				SNew(SBorder)
+				.Padding(4)
+				.BorderImage(FEditorStyle::GetBrush("ToolPanel.GroupBorder"))
+				[
+					InfoBox
+				]
+			]
+		];
+
+	return IDocumentation::Get()->CreateToolTip(TAttribute<FText>(this, &SSCS_RowWidget::GetTooltipText), TooltipContent, InfoBox, GetDocumentationLink(), GetDocumentationExcerptName());
+}
+END_SLATE_FUNCTION_BUILD_OPTIMIZATION
+
 FSlateBrush const* SSCS_RowWidget::GetMobilityIconImage(TWeakObjectPtr<USCS_Node> SCSNode) const
 {
 	if (!SCSNode.IsValid())
@@ -1183,7 +1308,7 @@ FText SSCS_RowWidget::GetAssetName() const
 	return AssetName;
 }
 
-FText SSCS_RowWidget::GetAssetNameTooltip() const
+FText SSCS_RowWidget::GetAssetPath() const
 {
 	FText AssetName = LOCTEXT("None", "None");
 	if(NodePtr.IsValid() && NodePtr->GetComponentTemplate())
@@ -1242,11 +1367,17 @@ FSlateColor SSCS_RowWidget::GetColorTint() const
 		}
 	}
 
-	return FSlateColor::UseForeground();
+	return FLinearColor(1.0f, 1.0f, 1.0f);
 }
 
 EVisibility SSCS_RowWidget::GetRootLabelVisibility() const
 {
+	bool bSingleLayoutBPEditor = GetDefault<UEditorExperimentalSettings>()->bUnifiedBlueprintEditor;
+	if (bSingleLayoutBPEditor)
+	{
+		return EVisibility::Collapsed;
+	}
+
 	if(NodePtr.IsValid() && SCSEditor.IsValid() && NodePtr == SCSEditor.Pin()->SceneRootNodePtr)
 	{
 		return EVisibility::Visible;
@@ -1318,15 +1449,20 @@ FReply SSCS_RowWidget::OnDragDetected( const FGeometry& MyGeometry, const FPoint
 			&& SCSEditorPtr.IsValid()
 			&& SCSEditorPtr->IsEditingAllowed() ) //can only drag when editing
 	{
-		TSharedRef<FSCSRowDragDropOp> Operation = FSCSRowDragDropOp::New();
-		Operation->CurrentHoverText = FText::GetEmpty();
-		Operation->PendingDropAction = FSCSRowDragDropOp::DropAction_None;
-
 		TArray<TSharedPtr<FSCSEditorTreeNode>> SelectedNodePtrs = SCSEditorPtr->GetSelectedNodes();
 		if(SelectedNodePtrs.Num() == 0)
 		{
 			SelectedNodePtrs.Add(NodePtr);
 		}
+
+		TSharedPtr<FSCSEditorTreeNode> FirstNode = SelectedNodePtrs[0];
+		
+		TSharedRef<FSCSRowDragDropOp> Operation = FSCSRowDragDropOp::New(FirstNode->GetVariableName(), FirstNode->GetBlueprint()->SkeletonGeneratedClass, FNodeCreationAnalytic());
+		//Operation->SetAltDrag(MouseEvent.IsAltDown());
+		//Operation->SetCtrlDrag(MouseEvent.IsLeftControlDown() || MouseEvent.IsRightControlDown());
+		Operation->SetCtrlDrag(true); // Always put a getter
+		Operation->CurrentHoverText = FText::GetEmpty();
+		Operation->PendingDropAction = FSCSRowDragDropOp::DropAction_None;
 		
 		for(const auto& SelectedNodePtr : SelectedNodePtrs)
 		{
@@ -1362,11 +1498,9 @@ void SSCS_RowWidget::OnDragEnter( const FGeometry& MyGeometry, const FDragDropEv
 		return;
 	}
 
-	if ( Operation->IsOfType<FSCSRowDragDropOp>() )
+	TSharedPtr<FSCSRowDragDropOp> DragRowOp = DragDropEvent.GetOperationAs<FSCSRowDragDropOp>();
+	if (DragRowOp.IsValid())
 	{
-		TSharedPtr<FSCSRowDragDropOp> DragRowOp = StaticCastSharedPtr<FSCSRowDragDropOp>( Operation );
-		check(DragRowOp.IsValid());
-
 		// If the hover text is already set, skip everything below, because it means we already know we can't drag-and-drop one or more of the selected nodes.
 		if(!DragRowOp->CurrentHoverText.IsEmpty())
 		{
@@ -2208,15 +2342,15 @@ void SSCS_RowWidget::OnNameTextCommit(const FText& InNewName, ETextCommit::Type
 //////////////////////////////////////////////////////////////////////////
 // SSCSEditor
 
-
+BEGIN_SLATE_FUNCTION_BUILD_OPTIMIZATION
 void SSCSEditor::Construct( const FArguments& InArgs )
 {
 	EditorMode = InArgs._EditorMode;
 	ActorContext = InArgs._ActorContext;
 	AllowEditing = InArgs._AllowEditing;
 	PreviewActor = InArgs._PreviewActor;
-	OnTreeViewSelectionChanged = InArgs._OnTreeViewSelectionChanged;
-	OnUpdateSelectionFromNodes = InArgs._OnUpdateSelectionFromNodes;
+	OnRootSelected = InArgs._OnRootSelected;
+	OnSelectionUpdated = InArgs._OnSelectionUpdated;
 	OnHighlightPropertyInDetailsView = InArgs._OnHighlightPropertyInDetailsView;
 
 	CommandList = MakeShareable( new FUICommandList );
@@ -2248,21 +2382,37 @@ void SSCSEditor::Construct( const FArguments& InArgs )
 		);
 
 	FSlateBrush const* MobilityHeaderBrush = FEditorStyle::GetBrush(TEXT("ClassIcon.ComponentMobilityHeaderIcon"));
+	
+	TSharedPtr<SHeaderRow> HeaderRow;
+	
+	bool bSingleLayoutBPEditor = GetDefault<UEditorExperimentalSettings>()->bUnifiedBlueprintEditor;
+	if (bSingleLayoutBPEditor)
+	{
+		HeaderRow = SNew(SHeaderRow)
 
-	SAssignNew(SCSTreeWidget, SSCSTreeType)
-		.ToolTipText(LOCTEXT("DropAssetToAddComponent", "Drop asset here to add component."))
-		.SCSEditor(this)
-		.TreeItemsSource(&RootNodes)
-		.SelectionMode(ESelectionMode::Multi)
-		.OnGenerateRow(this, &SSCSEditor::MakeTableRowWidget)
-		.OnGetChildren(this, &SSCSEditor::OnGetChildrenForTree)
-		.OnSelectionChanged(this, &SSCSEditor::OnTreeSelectionChanged)
-		.OnContextMenuOpening(this, &SSCSEditor::CreateContextMenu)
-		.OnItemScrolledIntoView(this, &SSCSEditor::OnItemScrolledIntoView)
-		.ItemHeight(24)
-		.HeaderRow
-		(
-		SNew(SHeaderRow)
+		+ SHeaderRow::Column(SCS_ColumnName_Mobility)
+		.DefaultLabel(LOCTEXT("MobilityColumnLabel", "Mobility"))
+		.FixedWidth(16.0f) // mobility icons are 16px (16 slate-units = 16px, when application scale == 1)
+		.HeaderContent()
+		[
+			SNew(SHorizontalBox)
+			.ToolTip(SNew(SToolTip).Text(LOCTEXT("MobilityColumnTooltip", "Mobility")))
+			+ SHorizontalBox::Slot()
+			.FillWidth(1.0f)
+			.VAlign(VAlign_Center)
+			.HAlign(HAlign_Center)
+			[
+				SNew(SImage).Image(MobilityHeaderBrush)
+			]
+		]
+
+		+ SHeaderRow::Column(SCS_ColumnName_ComponentClass)
+		.DefaultLabel(LOCTEXT("Class", "Class"))
+		.FillWidth(4);
+	}
+	else
+	{
+		HeaderRow = SNew(SHeaderRow)
 
 		+ SHeaderRow::Column(SCS_ColumnName_Mobility)
 		.DefaultLabel(LOCTEXT("MobilityColumnLabel", "Mobility"))
@@ -2280,17 +2430,110 @@ void SSCSEditor::Construct( const FArguments& InArgs )
 			]
 		]
 
-	+ SHeaderRow::Column(SCS_ColumnName_ComponentClass)
+		+ SHeaderRow::Column(SCS_ColumnName_ComponentClass)
 		.DefaultLabel(LOCTEXT("Class", "Class"))
 		.FillWidth(4)
 
 		+ SHeaderRow::Column(SCS_ColumnName_Asset)
-		.DefaultLabel(LOCTEXT("Asset", "Asset"))
-		.FillWidth(3)
+		.DefaultLabel(LOCTEXT("Asset", "Asset"))\
+		.FillWidth(3);
+	}
+
+	SAssignNew(SCSTreeWidget, SSCSTreeType)
+		.ToolTipText(LOCTEXT("DropAssetToAddComponent", "Drop asset here to add component."))
+		.SCSEditor(this)
+		.TreeItemsSource(&RootNodes)
+		.SelectionMode(ESelectionMode::Multi)
+		.OnGenerateRow(this, &SSCSEditor::MakeTableRowWidget)
+		.OnGetChildren(this, &SSCSEditor::OnGetChildrenForTree)
+		.OnSelectionChanged(this, &SSCSEditor::OnTreeSelectionChanged)
+		.OnContextMenuOpening(this, &SSCSEditor::CreateContextMenu)
+		.OnItemScrolledIntoView(this, &SSCSEditor::OnItemScrolledIntoView)
+		.ItemHeight(24)
+		.HeaderRow
+		(
+			HeaderRow
 		);
 
+	if (bSingleLayoutBPEditor)
+	{
+		SCSTreeWidget->GetHeaderRow()->SetVisibility(EVisibility::Collapsed);
+	}
+
+	TSharedRef<SToolTip> Tooltip = CreateToolTipWidget();
+
 	TSharedPtr<SWidget> Contents;
-	if( InArgs._HideComponentClassCombo.Get() )
+
+	if (bSingleLayoutBPEditor)
+	{
+		Contents = SNew(SVerticalBox)
+
+		// Root Actor
+		+ SVerticalBox::Slot()
+		.Padding(0.f)
+		.AutoHeight()
+		[
+			SNew(SVerticalBox)
+
+			+ SVerticalBox::Slot()
+			.AutoHeight()
+			.VAlign(VAlign_Center)
+			.HAlign(HAlign_Left)
+			.Padding(2.0f, 2.0f)
+			[
+				SNew(SComponentClassCombo)
+				.OnComponentClassSelected(this, &SSCSEditor::PerformComboAddClass)
+				.ToolTipText(LOCTEXT("AddComponent_Tooltip", "Add a component."))
+			]
+
+			+ SVerticalBox::Slot()
+			.VAlign(VAlign_Fill)
+			.HAlign(HAlign_Fill)
+			.Padding(0.0f, 0.0f)
+			[
+				SNew(SCheckBox)
+				.Style(FCoreStyle::Get(), "ToggleButtonRowStyle")
+				.IsFocusable(true)
+				.OnCheckStateChanged(this, &SSCSEditor::OnActorSelected)
+				.IsChecked(this, &SSCSEditor::OnIsActorSelected)
+				.ToolTip(Tooltip)
+				.Padding(FMargin(2.0f))
+				[
+					SNew(SHorizontalBox)
+
+					+ SHorizontalBox::Slot()
+					.AutoWidth()
+					.VAlign(VAlign_Center)
+					.Padding(FMargin(0.f, 0.f, 6.f, 0.f))
+					[
+						SNew(SBox)
+						[
+							SNew(SImage)
+							.Image(this, &SSCSEditor::GetActorIcon)
+						]
+					]
+
+					+ SHorizontalBox::Slot()
+					.HAlign(HAlign_Left)
+					.VAlign(VAlign_Center)
+					.Padding(0.0f, 0.0f)
+					[
+						SNew(STextBlock)
+						.Text(this, &SSCSEditor::GetActorDisplayText)
+						.ColorAndOpacity(FSlateColor::UseForeground())
+					]
+				]
+			]
+		]
+
+		// Tree
+		+ SVerticalBox::Slot()
+		.Padding(0.f, 0.f, 0.f, 2.f)
+		[
+			SCSTreeWidget.ToSharedRef()
+		];
+	}
+	else if( InArgs._HideComponentClassCombo.Get() )
 	{
 		Contents = SCSTreeWidget;
 	}
@@ -2338,7 +2581,93 @@ void SSCSEditor::Construct( const FArguments& InArgs )
 	{
 		SCSTreeWidget->SetItemExpansion(SceneRootNodePtr, true);
 	}
+
+	if (bSingleLayoutBPEditor)
+	{
+		// Select the root actor
+		OnActorSelected(ECheckBoxState::Checked);
+	}
 }
+TSharedRef<SToolTip> SSCSEditor::CreateToolTipWidget() const
+{
+	// Create a box to hold every line of info in the body of the tooltip
+	TSharedRef<SVerticalBox> InfoBox = SNew(SVerticalBox);
+
+	// Add class
+	AddToToolTipInfoBox(InfoBox, LOCTEXT("TooltipClass", "Class"), SNullWidget::NullWidget, TAttribute<FText>::Create(TAttribute<FText>::FGetter::CreateSP(this, &SSCSEditor::GetActorClassNameText)), false);
+
+	// Add super class
+	AddToToolTipInfoBox(InfoBox, LOCTEXT("TooltipSuperClass", "Parent Class"), SNullWidget::NullWidget, TAttribute<FText>::Create(TAttribute<FText>::FGetter::CreateSP(this, &SSCSEditor::GetActorSuperClassNameText)), false);
+
+	// Add mobility
+	AddToToolTipInfoBox(InfoBox, LOCTEXT("TooltipMobility", "Mobility"), SNullWidget::NullWidget, TAttribute<FText>::Create(TAttribute<FText>::FGetter::CreateSP(this, &SSCSEditor::GetActorMobilityText)), false);
+
+	TSharedRef<SBorder> TooltipContent = SNew(SBorder)
+		.Padding(4)
+		.BorderImage(FEditorStyle::GetBrush("SCSEditor.TileViewTooltip.NonContentBorder"))
+		[
+			SNew(SVerticalBox)
+
+			+ SVerticalBox::Slot()
+			.AutoHeight()
+			.Padding(0, 0, 0, 4)
+			[
+				SNew(SBorder)
+				.Padding(0)
+				.BorderImage(FEditorStyle::GetBrush("ToolPanel.GroupBorder"))
+				[
+					SNew(SVerticalBox)
+
+					+ SVerticalBox::Slot()
+					.AutoHeight()
+					[
+						SNew(SHorizontalBox)
+
+						+ SHorizontalBox::Slot()
+						.AutoWidth()
+						.VAlign(VAlign_Center)
+						.Padding(4)
+						[
+							SNew(STextBlock)
+							.Text(this, &SSCSEditor::GetActorDisplayText)
+							.Font(FEditorStyle::GetFontStyle("ContentBrowser.TileViewTooltip.NameFont"))
+						]
+					]
+				]
+			]
+
+			+ SVerticalBox::Slot()
+			.AutoHeight()
+			[
+				SNew(SBorder)
+				.Padding(4)
+				.BorderImage(FEditorStyle::GetBrush("ToolPanel.GroupBorder"))
+				[
+					InfoBox
+				]
+			]
+		];
+
+	return IDocumentation::Get()->CreateToolTip(TAttribute<FText>(this, &SSCSEditor::GetActorDisplayText), TooltipContent, InfoBox, TEXT(""), TEXT(""));
+}
+END_SLATE_FUNCTION_BUILD_OPTIMIZATION
+/*
+FString SSCSEditor::GetDocumentationLink() const
+{
+	if (NodePtr == SCSEditor.Pin()->SceneRootNodePtr
+		|| NodePtr->IsNative() || NodePtr->IsInherited())
+	{
+		return TEXT("Shared/Editors/BlueprintEditor/ComponentsMode");
+		FText ToolTipText;
+		if (ActorContext.IsSet())
+		{
+			AActor* DefaultActor = ActorContext.Get();
+			ToolTipText = FText::FromString(DefaultActor->GetClass()->GetSuperClass()->GetName());
+		}
+	}
+
+	return TEXT("");
+}*/
 
 UBlueprint* SSCSEditor::GetBlueprint() const
 {
@@ -2649,7 +2978,7 @@ void SSCSEditor::OnGetChildrenForTree( FSCSEditorTreeNodePtrType InNodePtr, TArr
 
 
 void SSCSEditor::PerformComboAddClass(TSubclassOf<UActorComponent> ComponentClass)
-	{
+{
 	UClass* NewClass = ComponentClass;
 
 	FEditorDelegates::LoadSelectedAssetsIfNeeded.Broadcast();
@@ -2689,6 +3018,91 @@ void SSCSEditor::PerformComboAddClass(TSubclassOf<UActorComponent> ComponentClas
 	}
 }
 
+const FSlateBrush* SSCSEditor::GetActorIcon() const
+{
+	if (ActorContext.IsSet())
+	{
+		return FClassIconFinder::FindIconForActor(ActorContext.Get());
+	}
+	return nullptr;
+}
+
+FText SSCSEditor::GetActorDisplayText() const
+{
+	if (ActorContext.IsSet())
+	{
+		FString Name;
+		AActor* DefaultActor = ActorContext.Get();
+		UBlueprint* Blueprint = UBlueprint::GetBlueprintFromClass(DefaultActor->GetClass());
+		if (Blueprint != nullptr)
+		{
+			Blueprint->GetName(Name);
+			return FText::Format(LOCTEXT("DefaultActor_Name", "{0} (self)"), FText::FromString(Name));
+		}
+		else
+		{
+			DefaultActor->GetName(Name);
+			return FText::Format(LOCTEXT("DefaultActor_Name", "{0}"), FText::FromString(Name));
+		}
+	}
+	return FText::GetEmpty();	
+}
+
+FText SSCSEditor::GetActorClassNameText() const
+{
+	FText Text;
+	if (ActorContext.IsSet())
+	{
+		AActor* DefaultActor = ActorContext.Get();
+		Text = FText::FromString(DefaultActor->GetClass()->GetName());
+	}
+
+	return Text;
+}
+
+FText SSCSEditor::GetActorSuperClassNameText() const
+{
+	FText Text;
+	if (ActorContext.IsSet())
+	{
+		AActor* DefaultActor = ActorContext.Get();
+		Text = FText::FromString(DefaultActor->GetClass()->GetSuperClass()->GetName());
+	}
+
+	return Text;
+}
+
+FText SSCSEditor::GetActorMobilityText() const
+{
+	FText Text;
+	if (ActorContext.IsSet())
+	{
+		AActor* DefaultActor = ActorContext.Get();
+		USceneComponent* RootComponent = DefaultActor->GetRootComponent();
+		if (RootComponent)
+		{
+			if (RootComponent->Mobility == EComponentMobility::Static)
+			{
+				Text = LOCTEXT("ComponentMobility_Static", "Static");
+			}
+			else if (RootComponent->Mobility == EComponentMobility::Stationary)
+			{
+				Text = LOCTEXT("ComponentMobility_Stationary", "Stationary");
+			}
+			else if (RootComponent->Mobility == EComponentMobility::Movable)
+			{
+				Text = LOCTEXT("ComponentMobility_Movable", "Movable");
+			}
+		}
+		else
+		{
+			Text = LOCTEXT("ComponentMobility_NoRoot", "No root component, unknown mobility");
+		}
+	}
+
+	return Text;
+}
+
 TArray<FSCSEditorTreeNodePtrType>  SSCSEditor::GetSelectedNodes() const
 {
 	TArray<FSCSEditorTreeNodePtrType> SelectedTreeNodes = SCSTreeWidget->GetSelectedItems();
@@ -3043,13 +3457,8 @@ void SSCSEditor::ClearSelection()
 {
 	check(SCSTreeWidget.IsValid());
 	SCSTreeWidget->ClearSelection();
-
-	// Invoke the selection changed delegate with an empty selection set
-	OnTreeViewSelectionChanged.ExecuteIfBound(TArray<FSCSEditorTreeNodePtrType>());
 }
 
-
-
 void SSCSEditor::SaveSCSCurrentState( USimpleConstructionScript* SCSObj )
 {
 	if( SCSObj )
@@ -3094,7 +3503,8 @@ UActorComponent* SSCSEditor::AddNewComponent( UClass* NewComponentClass, UObject
 		Blueprint->Modify();
 		SaveSCSCurrentState(Blueprint->SimpleConstructionScript);
 
-		return AddNewNode(Blueprint->SimpleConstructionScript->CreateNode(NewComponentClass), Asset, true);
+		FName NewVariableName = Asset != nullptr ? Asset->GetFName() : NAME_None;
+		return AddNewNode(Blueprint->SimpleConstructionScript->CreateNode(NewComponentClass, NewVariableName), Asset, true);
 	}
 	else    // EEditorMode::ActorInstance
 	{
@@ -3600,10 +4010,10 @@ void SSCSEditor::RemoveComponentNode(FSCSEditorTreeNodePtrType InNodePtr)
 	}
 }
 
-void SSCSEditor::UpdateSelectionFromNodes(const TArray<FSCSEditorTreeNodePtrType> &SelectedNodes )
+void SSCSEditor::UpdateSelectionFromNodes(const TArray<FSCSEditorTreeNodePtrType> &SelectedNodes)
 {
-	// Invoke the delegate to update any selection details
-	OnUpdateSelectionFromNodes.ExecuteIfBound(SelectedNodes);
+	// Notify that the selection has updated
+	OnSelectionUpdated.ExecuteIfBound(SelectedNodes);
 }
 
 void SSCSEditor::RefreshSelectionDetails()
@@ -3613,16 +4023,29 @@ void SSCSEditor::RefreshSelectionDetails()
 
 void SSCSEditor::OnTreeSelectionChanged(FSCSEditorTreeNodePtrType, ESelectInfo::Type /*SelectInfo*/)
 {
-	// Get the selected tree nodes
-	const TArray<FSCSEditorTreeNodePtrType>& SelectedNodes = SCSTreeWidget->GetSelectedItems();
+	bIsActorSelected = false;
 
-	// Update the details view selection
-	UpdateSelectionFromNodes(SelectedNodes);
+	UpdateSelectionFromNodes(SCSTreeWidget->GetSelectedItems());
+}
 
-	// Pass the tree selection change event along to the delegate
-	OnTreeViewSelectionChanged.ExecuteIfBound(SelectedNodes);
+void SSCSEditor::OnActorSelected(const ECheckBoxState NewCheckedState)
+{
+	// Clear tree selection
+	SCSTreeWidget->ClearSelection();
+		
+	// Set the Actor selected, done after the Tree clear as it will deselect the actor.
+	bIsActorSelected = true;
+
+	// Notify that the root has been selected
+	OnRootSelected.ExecuteIfBound(ActorContext.Get(nullptr));
+}
+
+ECheckBoxState SSCSEditor::OnIsActorSelected() const
+{
+	return (bIsActorSelected) ? ECheckBoxState::Checked : ECheckBoxState::Unchecked;
 }
 
+
 bool SSCSEditor::IsNodeInSimpleConstructionScript( USCS_Node* Node ) const
 {
 	check(Node);
@@ -3701,7 +4124,7 @@ FSCSEditorTreeNodePtrType SSCSEditor::AddTreeNode(USCS_Node* InSCSNode, FSCSEdit
 
 			// Add it to the root set
 			NodeSCS->AddNode(InSCSNode);
-			RootNodes.Add(NewNodePtr);
+			RootNodes.Insert(NewNodePtr, 0);
 
 			// Make it the scene root node
 			SceneRootNodePtr = NewNodePtr;
@@ -3715,16 +4138,7 @@ FSCSEditorTreeNodePtrType SSCSEditor::AddTreeNode(USCS_Node* InSCSNode, FSCSEdit
 		// If the given SCS node does not contain a scene component template, we create a new root node
 		NewNodePtr = MakeShareable(new FSCSEditorTreeNode(InSCSNode, bIsInherited));
 
-		// Add it to the root set prior to the scene root node that represents the start of the scene hierarchy
-		int32 SceneRootNodeIndex = RootNodes.Find(SceneRootNodePtr);
-		if(SceneRootNodeIndex != INDEX_NONE)
-		{
-			RootNodes.Insert(NewNodePtr, SceneRootNodeIndex);
-		}
-		else
-		{
-			RootNodes.Add(NewNodePtr);
-		}
+		RootNodes.Add(NewNodePtr);
 
 		// If the SCS root node array does not already contain the given node, this will add it (this should only occur after node creation)
 		if(NodeSCS != NULL)
@@ -3775,7 +4189,7 @@ FSCSEditorTreeNodePtrType SSCSEditor::AddTreeNode(USceneComponent* InSceneCompon
 			NewNodePtr = MakeShareable(new FSCSEditorTreeNode(InSceneComponent));
 
 			// Add it to the root set
-			RootNodes.Add(NewNodePtr);
+			RootNodes.Insert(NewNodePtr, 0);
 
 			// Make it the scene root node
 			SceneRootNodePtr = NewNodePtr;
diff --git a/Engine/Source/Editor/Kismet/Public/BlueprintEditor.h b/Engine/Source/Editor/Kismet/Public/BlueprintEditor.h
--- a/Engine/Source/Editor/Kismet/Public/BlueprintEditor.h
+++ b/Engine/Source/Editor/Kismet/Public/BlueprintEditor.h
@@ -384,15 +384,17 @@ public:
 
 	/** Refresh the preview viewport to reflect changes in the SCS */
 	void UpdateSCSPreview(bool bUpdateNow = false);
-	
+
 	/** Delegate invoked when the SCS editor needs to obtain the Actor context for editing */
 	AActor* GetSCSEditorActorContext() const;
 
-	/** Delegate invoked when the tree view selection is changed in the SCS editor widget */
-	void OnSCSEditorTreeViewSelectionChanged(const TArray< TSharedPtr<class FSCSEditorTreeNode> >& SelectedNodes);
+	/** Delegate invoked when the root is selected in the SCS editor widget */
+	void OnRootSelected(class AActor * DefaultActor);
+
+	/** Delegate invoked when the selection is changed in the SCS editor widget */
+	void OnSelectionUpdated(const TArray<TSharedPtr<class FSCSEditorTreeNode>>& SelectedNodes);
+
 
-	/** Delegate to update the Inspector (details) panel from the given set of selected nodes */
-	void OnSCSEditorUpdateSelectionFromNodes(const TArray< TSharedPtr<class FSCSEditorTreeNode> >& SelectedNodes);
 
 	/** Pin visibility accessors */
 	void SetPinVisibility(SGraphEditor::EPinVisibility Visibility);
diff --git a/Engine/Source/Editor/Kismet/Public/SKismetInspector.h b/Engine/Source/Editor/Kismet/Public/SKismetInspector.h
--- a/Engine/Source/Editor/Kismet/Public/SKismetInspector.h
+++ b/Engine/Source/Editor/Kismet/Public/SKismetInspector.h
@@ -40,11 +40,13 @@ public:
 	{
 		FText ForcedTitle;
 		bool bForceRefresh;
+		bool bShowComponents;
 		bool bHideFilterArea;
 
 		FShowDetailsOptions(const FText& InForcedTitle = FText::GetEmpty(), bool bInForceRefresh = false)
 			:ForcedTitle(InForcedTitle)
 			,bForceRefresh(bInForceRefresh)
+			,bShowComponents(true)
 			,bHideFilterArea(false)
 		{}
 	};
@@ -85,6 +87,9 @@ protected:
 	/** Should we currently show the property view */
 	bool bShowInspectorPropertyView;
 
+	/** Should we currently show components */
+	bool bShowComponents;
+
 	/** State of CheckBox representing whether to show only the public variables*/
 	ECheckBoxState	PublicViewState;
 
diff --git a/Engine/Source/Editor/Kismet/Public/SSCSEditor.h b/Engine/Source/Editor/Kismet/Public/SSCSEditor.h
--- a/Engine/Source/Editor/Kismet/Public/SSCSEditor.h
+++ b/Engine/Source/Editor/Kismet/Public/SSCSEditor.h
@@ -264,7 +264,7 @@ public:
 	FString GetDocumentationExcerptName() const;
 	
 	FText GetAssetName() const;
-	FText GetAssetNameTooltip() const;
+	FText GetAssetPath() const;
 	EVisibility GetAssetVisibility() const;
 
 	EVisibility GetRootLabelVisibility() const;
@@ -285,6 +285,9 @@ private:
 	/** Builds a context menu popup for dropping a child node onto the scene root node */
 	TSharedPtr<SWidget> BuildSceneRootDropActionMenu(FSCSEditorTreeNodePtrType DroppedNodePtr);
 
+	/** Creates a tooltip for this row */
+	TSharedRef<SToolTip> CreateToolTipWidget() const;
+
 	/** Handler for attaching a single node to this node */
 	void OnAttachToDropAction(FSCSEditorTreeNodePtrType DroppedNodePtr)
 	{
@@ -415,8 +418,8 @@ public:
 
 	DECLARE_DELEGATE_RetVal_OneParam(class USCS_Node*, FOnAddNewComponent, class UClass*);
 	DECLARE_DELEGATE_RetVal_OneParam(class USCS_Node*, FOnAddExistingComponent, class UActorComponent*);
-	DECLARE_DELEGATE_OneParam(FOnTreeViewSelectionChanged, const TArray<FSCSEditorTreeNodePtrType>&);
-	DECLARE_DELEGATE_OneParam(FOnUpdateSelectionFromNodes, const TArray<FSCSEditorTreeNodePtrType>&);
+	DECLARE_DELEGATE_OneParam(FOnRootSelected, class AActor*);
+	DECLARE_DELEGATE_OneParam(FOnSelectionUpdated, const TArray<FSCSEditorTreeNodePtrType>&);
 	DECLARE_DELEGATE_OneParam(FOnHighlightPropertyInDetailsView, const class FPropertyPath&);
 
 	SLATE_BEGIN_ARGS( SSCSEditor )
@@ -425,8 +428,8 @@ public:
 		,_PreviewActor(nullptr)
 		,_AllowEditing(true)
 		,_HideComponentClassCombo(false)
-		,_OnTreeViewSelectionChanged()
-		,_OnUpdateSelectionFromNodes()
+		,_OnRootSelected()
+		,_OnSelectionUpdated()
 		,_OnHighlightPropertyInDetailsView()
 		{}
 
@@ -435,8 +438,8 @@ public:
 		SLATE_ATTRIBUTE(class AActor*, PreviewActor)
 		SLATE_ATTRIBUTE(bool, AllowEditing)
 		SLATE_ATTRIBUTE(bool, HideComponentClassCombo)
-		SLATE_EVENT(FOnTreeViewSelectionChanged, OnTreeViewSelectionChanged)
-		SLATE_EVENT(FOnUpdateSelectionFromNodes, OnUpdateSelectionFromNodes)
+		SLATE_EVENT(FOnRootSelected, OnRootSelected)
+		SLATE_EVENT(FOnSelectionUpdated, OnSelectionUpdated)
 		SLATE_EVENT(FOnHighlightPropertyInDetailsView, OnHighlightPropertyInDetailsView)
 
 	SLATE_END_ARGS()
@@ -504,6 +507,12 @@ public:
 	/** Called when selection in the tree changes */
 	void OnTreeSelectionChanged(FSCSEditorTreeNodePtrType InSelectedNodePtr, ESelectInfo::Type SelectInfo);
 
+	/** Called when the Actor is selected. */
+	void OnActorSelected(const ECheckBoxState NewCheckedState);
+
+	/** Called to determine if actor is selected. */
+	ECheckBoxState OnIsActorSelected() const;
+
 	/** Update any associated selection (e.g. details view) from the passed in nodes */
 	void UpdateSelectionFromNodes(const TArray<FSCSEditorTreeNodePtrType> &SelectedNodes );
 
@@ -589,6 +598,24 @@ protected:
 	/** Add a component from the selection in the combo box */
 	void PerformComboAddClass(TSubclassOf<UActorComponent> ComponentClass);
 
+	/** Called to get an the actors icon */
+	const FSlateBrush* GetActorIcon() const;
+
+	/** Called to get an the actors display text */
+	FText GetActorDisplayText() const;
+
+	/** Called to get an the actors class name text */
+	FText GetActorClassNameText() const;
+
+	/** Called to get an the actors super class name text */
+	FText GetActorSuperClassNameText() const;
+
+	/** Called to get an the actors mobility text */
+	FText GetActorMobilityText() const;
+
+	/** Creates a tooltip for the root component*/
+	TSharedRef<SToolTip> CreateToolTipWidget() const;
+
 	/** Called to display context menu when right clicking on the widget */
 	TSharedPtr< SWidget > CreateContextMenu();
 
@@ -681,11 +708,11 @@ public:
 	/** Attribute to indicate whether or not editing is allowed. */
 	TAttribute<bool> AllowEditing;
 
-	/** Delegate to invoke on tree view selection change. */
-	FOnTreeViewSelectionChanged OnTreeViewSelectionChanged;
+	/** Delegate to invoke on selection of the root. */
+	FOnRootSelected OnRootSelected;
 
-	/** Delegate to invoke when details should be updated. */
-	FOnUpdateSelectionFromNodes OnUpdateSelectionFromNodes;
+	/** Delegate to invoke on selection update. */
+	FOnSelectionUpdated OnSelectionUpdated;
 
 	/** Delegate to invoke when the given property should be highlighted in the details view (e.g. diff). */
 	FOnHighlightPropertyInDetailsView OnHighlightPropertyInDetailsView;
@@ -693,4 +720,7 @@ private:
 
 	/** Flag to enable/disable component editing */
 	bool	bEnableComponentEditing;
+
+	/** Flag to track if we have the Actor selected curently */
+	bool bIsActorSelected;
 };
diff --git a/Engine/Source/Editor/LevelEditor/Private/SLevelEditor.cpp b/Engine/Source/Editor/LevelEditor/Private/SLevelEditor.cpp
--- a/Engine/Source/Editor/LevelEditor/Private/SLevelEditor.cpp
+++ b/Engine/Source/Editor/LevelEditor/Private/SLevelEditor.cpp
@@ -555,8 +555,8 @@ public:
 			SAssignNew(SCSEditor, SSCSEditor)
 			.EditorMode(SSCSEditor::EEditorMode::ActorInstance)
 			.ActorContext(this, &SActorDetails::GetSelectedActor)												// Get the instance of the actor in the world
-			.OnTreeViewSelectionChanged(this, &SActorDetails::OnSCSEditorTreeViewSelectionChanged)				// A selection has been made in the tree view, so inform the level editor
-			//.OnUpdateSelectionFromNodes(this, &SLevelEditor::OnSCSEditorUpdateSelectionFromNodes)				// Unsure, don't think it's needed
+			.OnRootSelected(this, &SActorDetails::OnSCSEditorRootSelected)										// A selection of the root has been made
+			.OnSelectionUpdated(this, &SActorDetails::OnSCSEditorTreeViewSelectionChanged)						// A selection has been made in the tree view, so inform the level editor
 			//.OnHighlightPropertyInDetailsView(this, &SLevelEditor::OnSCSEditorHighlightPropertyInDetailsView)	// Also unsure and don't think it's needed
 		);
 
@@ -678,6 +678,35 @@ private:
 		return Cast<AActor>(*GEditor->GetSelectedActorIterator());
 	}
 
+	void OnSCSEditorRootSelected(AActor*)
+	{
+		if (!bSelectionGuard)
+		{
+			auto Actor = GetSelectedActor();
+			if (Actor)
+			{
+				// Enable the selection guard to prevent OnEditorSelectionChanged() from altering the contents of the SCSTreeWidget
+				TGuardValue<bool> SelectionGuard(bSelectionGuard, true);
+
+				// Update the editor's component selection
+				TArray<UObject*> DetailsObjects;
+				USelection* SelectedComponents = GEditor->GetSelectedComponents();
+				DetailsObjects.Add(Actor);
+
+				SelectedComponents->BeginBatchSelectOperation();
+				SelectedComponents->DeselectAll();
+
+				DetailsView->SetObjects(DetailsObjects, false);
+
+				SelectedComponents->EndBatchSelectOperation();
+
+				GUnrealEd->SetActorSelectionFlags(Actor);
+				GUnrealEd->UpdatePivotLocationForSelection(true);
+				GEditor->RedrawLevelEditingViewports();
+			}
+		}
+	}
+
 	void OnSCSEditorTreeViewSelectionChanged(const TArray<FSCSEditorTreeNodePtrType>& SelectedNodes)
 	{
 		if (!bSelectionGuard && SelectedNodes.Num() > 0 )
@@ -1666,4 +1695,4 @@ void SLevelEditor::OnActorSelectionChanged( const TArray<UObject*>& NewSelection
 			// remove stray entries here
 		}
 	}
-}
\ No newline at end of file
+}
diff --git a/Engine/Source/Editor/UnrealEd/Private/SComponentClassCombo.cpp b/Engine/Source/Editor/UnrealEd/Private/SComponentClassCombo.cpp
--- a/Engine/Source/Editor/UnrealEd/Private/SComponentClassCombo.cpp
+++ b/Engine/Source/Editor/UnrealEd/Private/SComponentClassCombo.cpp
@@ -54,6 +54,7 @@ void SComponentClassCombo::Construct(const FArguments& InArgs)
 			SNew(STextBlock)
 			.Text(LOCTEXT("AddComponentButtonLabel", "Add Component"))
 			.TextStyle(FEditorStyle::Get(), "ContentBrowser.TopBar.Font")
+			.Visibility(InArgs._IncludeText.Get() ? EVisibility::Visible : EVisibility::Collapsed)
 		]
 	]
 	.MenuContent()
@@ -123,7 +124,7 @@ void SComponentClassCombo::GenerateFilteredComponentList(const FString& InSearch
 		{
 			FComponentClassComboEntryPtr& CurrentEntry = ComponentClassList[ComponentIndex];
 
-			if ( CurrentEntry->IsClass() )
+			if (CurrentEntry->IsClass() && CurrentEntry->IsIncludedInFilter())
 			{
 				FString FriendlyComponentName = GetSanitizedComponentName( CurrentEntry->GetComponentClass() );
 
@@ -375,18 +376,18 @@ void SComponentClassCombo::UpdateComponentClassList()
 			TArray<FString> ClassGroupNames;
 			Class->GetClassGroupNames( ClassGroupNames );
 
-			FString ClassGroup;
 			if (ClassGroupNames.Contains(SortComboEntry::CommonClassGroup))
 			{
-				ClassGroup = SortComboEntry::CommonClassGroup;
+				FString ClassGroup = SortComboEntry::CommonClassGroup;
+				FComponentClassComboEntryPtr NewEntry(new FComponentClassComboEntry(ClassGroup, Class, ClassGroupNames.Num() <= 1));
+				SortedClassList.Add(NewEntry);
 			}
-			else if (ClassGroupNames.Num())
+			if (ClassGroupNames.Num() && !ClassGroupNames[0].Equals(SortComboEntry::CommonClassGroup))
 			{
-				ClassGroup = ClassGroupNames[0];
+				FString ClassGroup = ClassGroupNames[0];
+				FComponentClassComboEntryPtr NewEntry(new FComponentClassComboEntry(ClassGroup, Class, true));
+				SortedClassList.Add(NewEntry);
 			}
-
-			FComponentClassComboEntryPtr NewEntry(new FComponentClassComboEntry(ClassGroup,	Class));
-			SortedClassList.Add(NewEntry);
 		}
 	}
 
@@ -422,9 +423,18 @@ void SComponentClassCombo::UpdateComponentClassList()
 
 FString SComponentClassCombo::GetSanitizedComponentName( UClass* ComponentClass )
 {
-	FString OriginalName = ComponentClass->GetName();
-	FString ComponentNameWithoutComponent = OriginalName.Replace( TEXT("Component"), TEXT(""), ESearchCase::IgnoreCase );
-	return FName::NameToDisplayString( ComponentNameWithoutComponent, false );
+	FString DisplayName;
+	if (ComponentClass->HasMetaData(TEXT("DisplayName")))
+	{
+		DisplayName = ComponentClass->GetMetaData(TEXT("DisplayName"));
+	}
+	else
+	{
+		FString OriginalName = ComponentClass->GetName();
+		DisplayName = OriginalName.Replace( TEXT("Component"), TEXT(""), ESearchCase::IgnoreCase );
+	}
+
+	return FName::NameToDisplayString(DisplayName, false);
 }
 
 #undef LOCTEXT_NAMESPACE
diff --git a/Engine/Source/Editor/UnrealEd/Public/SComponentClassCombo.h b/Engine/Source/Editor/UnrealEd/Public/SComponentClassCombo.h
--- a/Engine/Source/Editor/UnrealEd/Public/SComponentClassCombo.h
+++ b/Engine/Source/Editor/UnrealEd/Public/SComponentClassCombo.h
@@ -12,9 +12,10 @@ DECLARE_DELEGATE_OneParam(FComponentClassSelected, TSubclassOf<UActorComponent>)
 class FComponentClassComboEntry: public TSharedFromThis<FComponentClassComboEntry>
 {
 public:
-	FComponentClassComboEntry( const FString& InHeadingText, TSubclassOf<UActorComponent> InComponentClass )
+	FComponentClassComboEntry( const FString& InHeadingText, TSubclassOf<UActorComponent> InComponentClass, bool InIncludedInFilter )
 		: ComponentClass(InComponentClass)
 		, HeadingText(InHeadingText)
+		, bIncludedInFilter(InIncludedInFilter)
 	{}
 
 	FComponentClassComboEntry( const FString& InHeadingText )
@@ -43,11 +44,15 @@ public:
 	{
 		return (ComponentClass!=NULL);
 	}
+	bool IsIncludedInFilter() const
+	{
+		return bIncludedInFilter;
+	}
 
 private:
 	TSubclassOf<UActorComponent> ComponentClass;
-
 	FString HeadingText;
+	bool bIncludedInFilter;
 };
 
 //////////////////////////////////////////////////////////////////////////
@@ -56,8 +61,10 @@ class UNREALED_API SComponentClassCombo : public SComboButton
 {
 public:
 	SLATE_BEGIN_ARGS( SComponentClassCombo )
+		: _IncludeText(true)
 	{}
 
+		SLATE_ATTRIBUTE(bool, IncludeText)
 		SLATE_EVENT( FComponentClassSelected, OnComponentClassSelected )
 
 	SLATE_END_ARGS()
diff --git a/Engine/Source/Runtime/Engine/Classes/Components/ArrowComponent.h b/Engine/Source/Runtime/Engine/Classes/Components/ArrowComponent.h
--- a/Engine/Source/Runtime/Engine/Classes/Components/ArrowComponent.h
+++ b/Engine/Source/Runtime/Engine/Classes/Components/ArrowComponent.h
@@ -7,7 +7,7 @@
 /** 
  * A simple arrow rendered using lines. Useful for indicating which way an object is facing.
  */
-UCLASS(ClassGroup=Shapes, hidecategories=(Object,LOD,Physics,Lighting,TextureStreaming,Activation,"Components|Activation",Collision), editinlinenew, meta=(BlueprintSpawnableComponent), MinimalAPI)
+UCLASS(ClassGroup=Utility, hidecategories=(Object,LOD,Physics,Lighting,TextureStreaming,Activation,"Components|Activation",Collision), editinlinenew, meta=(BlueprintSpawnableComponent), MinimalAPI)
 class UArrowComponent : public UPrimitiveComponent
 {
 	GENERATED_UCLASS_BODY()
diff --git a/Engine/Source/Runtime/Engine/Classes/Components/BoxComponent.h b/Engine/Source/Runtime/Engine/Classes/Components/BoxComponent.h
--- a/Engine/Source/Runtime/Engine/Classes/Components/BoxComponent.h
+++ b/Engine/Source/Runtime/Engine/Classes/Components/BoxComponent.h
@@ -8,7 +8,7 @@
 /** 
  * A box generally used for simple collision. Bounds are rendered as lines in the editor.
  */
-UCLASS(ClassGroup=Shapes, hidecategories=(Object,LOD,Lighting,TextureStreaming), editinlinenew, meta=(BlueprintSpawnableComponent))
+UCLASS(ClassGroup="Physics", hidecategories=(Object,LOD,Lighting,TextureStreaming), editinlinenew, meta=(DisplayName="Box Collision", BlueprintSpawnableComponent))
 class ENGINE_API UBoxComponent : public UShapeComponent
 {
 	GENERATED_UCLASS_BODY()
diff --git a/Engine/Source/Runtime/Engine/Classes/Components/CapsuleComponent.h b/Engine/Source/Runtime/Engine/Classes/Components/CapsuleComponent.h
--- a/Engine/Source/Runtime/Engine/Classes/Components/CapsuleComponent.h
+++ b/Engine/Source/Runtime/Engine/Classes/Components/CapsuleComponent.h
@@ -7,7 +7,7 @@
 /** 
  * A capsule generally used for simple collision. Bounds are rendered as lines in the editor.
  */
-UCLASS(ClassGroup=Shapes, editinlinenew, hidecategories=(Object,LOD,Lighting,TextureStreaming), meta=(BlueprintSpawnableComponent))
+UCLASS(ClassGroup="Physics", editinlinenew, hidecategories=(Object,LOD,Lighting,TextureStreaming), meta=(DisplayName="Capsule Collision", BlueprintSpawnableComponent))
 class ENGINE_API UCapsuleComponent : public UShapeComponent
 {
 	GENERATED_UCLASS_BODY()
diff --git a/Engine/Source/Runtime/Engine/Classes/Components/PointLightComponent.h b/Engine/Source/Runtime/Engine/Classes/Components/PointLightComponent.h
--- a/Engine/Source/Runtime/Engine/Classes/Components/PointLightComponent.h
+++ b/Engine/Source/Runtime/Engine/Classes/Components/PointLightComponent.h
@@ -8,7 +8,7 @@
 /**
  * A light component which emits light from a single point equally in all directions.
  */
-UCLASS(ClassGroup=Lights, hidecategories=(Object, LightShafts), editinlinenew, meta=(BlueprintSpawnableComponent))
+UCLASS(ClassGroup=(Lights,Common), hidecategories=(Object, LightShafts), editinlinenew, meta=(BlueprintSpawnableComponent))
 class ENGINE_API UPointLightComponent : public ULightComponent
 {
 	GENERATED_UCLASS_BODY()
diff --git a/Engine/Source/Runtime/Engine/Classes/Components/SphereComponent.h b/Engine/Source/Runtime/Engine/Classes/Components/SphereComponent.h
--- a/Engine/Source/Runtime/Engine/Classes/Components/SphereComponent.h
+++ b/Engine/Source/Runtime/Engine/Classes/Components/SphereComponent.h
@@ -8,7 +8,7 @@
 /** 
  * A sphere generally used for simple collision. Bounds are rendered as lines in the editor.
  */
-UCLASS(ClassGroup=Shapes, editinlinenew, hidecategories=(Object,LOD,Lighting,TextureStreaming), meta=(BlueprintSpawnableComponent))
+UCLASS(ClassGroup="Physics", editinlinenew, hidecategories=(Object,LOD,Lighting,TextureStreaming), meta=(DisplayName="Sphere Collision", BlueprintSpawnableComponent))
 class ENGINE_API USphereComponent : public UShapeComponent
 {
 	GENERATED_UCLASS_BODY()
diff --git a/Engine/Source/Runtime/Engine/Classes/Components/SplineComponent.h b/Engine/Source/Runtime/Engine/Classes/Components/SplineComponent.h
--- a/Engine/Source/Runtime/Engine/Classes/Components/SplineComponent.h
+++ b/Engine/Source/Runtime/Engine/Classes/Components/SplineComponent.h
@@ -24,7 +24,7 @@ namespace ESplinePointType
  *	A spline component is a spline shape which can be used for other purposes (e.g. animating objects). It does not contain rendering capabilities itself (outside the editor) 
  *	@see https://docs.unrealengine.com/latest/INT/Resources/ContentExamples/Blueprint_Splines
  */
-UCLASS(ClassGroup=Shapes, meta=(BlueprintSpawnableComponent))
+UCLASS(ClassGroup=Utility, meta=(BlueprintSpawnableComponent))
 class ENGINE_API USplineComponent : public UPrimitiveComponent
 {
 	GENERATED_UCLASS_BODY()
diff --git a/Engine/Source/Runtime/Engine/Classes/Particles/ParticleSystemComponent.h b/Engine/Source/Runtime/Engine/Classes/Particles/ParticleSystemComponent.h
--- a/Engine/Source/Runtime/Engine/Classes/Particles/ParticleSystemComponent.h
+++ b/Engine/Source/Runtime/Engine/Classes/Particles/ParticleSystemComponent.h
@@ -247,7 +247,7 @@ struct FParticleEventKismetData : public FParticleEventData
 /** 
  * A particle emitter.
  */
-UCLASS(ClassGroup=Rendering, hidecategories=Object, hidecategories=Physics, hidecategories=Collision, showcategories=Trigger, editinlinenew, meta=(BlueprintSpawnableComponent))
+UCLASS(ClassGroup=(Rendering, Common), hidecategories=Object, hidecategories=Physics, hidecategories=Collision, showcategories=Trigger, editinlinenew, meta=(BlueprintSpawnableComponent))
 class ENGINE_API UParticleSystemComponent : public UPrimitiveComponent
 {
 	GENERATED_UCLASS_BODY()
diff --git a/Engine/Source/Runtime/SlateCore/Private/Styling/CoreStyle.cpp b/Engine/Source/Runtime/SlateCore/Private/Styling/CoreStyle.cpp
--- a/Engine/Source/Runtime/SlateCore/Private/Styling/CoreStyle.cpp
+++ b/Engine/Source/Runtime/SlateCore/Private/Styling/CoreStyle.cpp
@@ -374,6 +374,7 @@ TSharedRef<ISlateStyle> FCoreStyle::Create( const FName& InStyleSetName )
 			.SetUndeterminedPressedImage( FSlateNoResource() );
 		Style->Set( "TransparentCheckBox", BasicTransparentCheckBoxStyle );
 
+		/* Default Style for a toggleable button */
 		const FCheckBoxStyle ToggleButtonStyle = FCheckBoxStyle()
 			.SetCheckBoxType(ESlateCheckBoxType::ToggleButton)
 			.SetUncheckedImage( FSlateNoResource() )
@@ -384,6 +385,17 @@ TSharedRef<ISlateStyle> FCoreStyle::Create( const FName& InStyleSetName )
 			.SetCheckedPressedImage( BOX_BRUSH("Common/RoundedSelection_16x",  4.0f/16.0f, SelectionColor_Pressed ) );
 		Style->Set( "ToggleButtonCheckbox", ToggleButtonStyle );
 
+		/* Style for a toggleable button that mimics the coloring and look of a Table Row */
+		const FCheckBoxStyle ToggleButtonRowStyle = FCheckBoxStyle()
+			.SetCheckBoxType(ESlateCheckBoxType::ToggleButton)
+			.SetUncheckedImage(FSlateNoResource())
+			.SetUncheckedHoveredImage(IMAGE_BRUSH("Common/Selection", Icon8x8, SelectionColor_Inactive))
+			.SetUncheckedPressedImage(IMAGE_BRUSH("Common/Selection", Icon8x8, SelectionColor_Inactive))
+			.SetCheckedImage(IMAGE_BRUSH("Common/Selection", Icon8x8, SelectionColor))
+			.SetCheckedHoveredImage(IMAGE_BRUSH("Common/Selection", Icon8x8, SelectionColor))
+			.SetCheckedPressedImage(BOX_BRUSH("Common/Selector", 4.0f / 16.0f, SelectorColor));
+		Style->Set("ToggleButtonRowStyle", ToggleButtonRowStyle);
+
 		/* A radio button is actually just a SCheckBox box with different images */
 		/* Set images for various radio button (SCheckBox) states ... */
 		const FCheckBoxStyle BasicRadioButtonStyle = FCheckBoxStyle()
@@ -494,8 +506,8 @@ TSharedRef<ISlateStyle> FCoreStyle::Create( const FName& InStyleSetName )
 
 	// SToolTip defaults...
 	{
-		Style->Set( "ToolTip.Font", TTF_FONT( "Fonts/Roboto-Regular", 8 ) );
-		Style->Set( "ToolTip.Background", new BOX_BRUSH( "Old/ToolTip_Background", FMargin( 8.0f/64.0f ) ) );
+		Style->Set("ToolTip.Font", TTF_FONT("Fonts/Roboto-Regular", 8));
+		Style->Set("ToolTip.Background", new BOX_BRUSH("Old/ToolTip_Background", FMargin(8.0f / 64.0f)));
 	}
 
 	// SBorder defaults...
@@ -646,22 +658,27 @@ TSharedRef<ISlateStyle> FCoreStyle::Create( const FName& InStyleSetName )
 
 	// TableView defaults...
 	{
-		Style->Set( "TableView.Row", FTableRowStyle()
-			.SetEvenRowBackgroundBrush( FSlateNoResource() )
-			.SetEvenRowBackgroundHoveredBrush( IMAGE_BRUSH( "Common/Selection", Icon8x8, FLinearColor(1.0f, 1.0f, 1.0f, 0.1f) ) )
-			.SetOddRowBackgroundBrush( FSlateNoResource() )
-			.SetOddRowBackgroundHoveredBrush( IMAGE_BRUSH( "Common/Selection", Icon8x8, FLinearColor(1.0f, 1.0f, 1.0f, 0.1f) ) )
-			.SetSelectorFocusedBrush( BORDER_BRUSH( "Common/Selector", FMargin(4.f/16.f), SelectorColor ) )
-			.SetActiveBrush( IMAGE_BRUSH( "Common/Selection", Icon8x8, SelectionColor ) )
-			.SetActiveHoveredBrush( IMAGE_BRUSH( "Common/Selection", Icon8x8, SelectionColor ) )
-			.SetInactiveBrush( IMAGE_BRUSH( "Common/Selection", Icon8x8, SelectionColor_Inactive ) )
-			.SetInactiveHoveredBrush( IMAGE_BRUSH( "Common/Selection", Icon8x8, SelectionColor_Inactive ) )
-			.SetTextColor( DefaultForeground )
-			.SetSelectedTextColor( InvertedForeground )
-			.SetDropIndicator_Above(BOX_BRUSH("Common/DropZoneIndicator_Above", FMargin(10.0f/16.0f, 10.0f/16.0f, 0,0), SelectionColor ))
+		const FTableRowStyle DefaultTableRowStyle = FTableRowStyle()
+			.SetEvenRowBackgroundBrush(FSlateNoResource())
+			.SetEvenRowBackgroundHoveredBrush(IMAGE_BRUSH("Common/Selection", Icon8x8, FLinearColor(1.0f, 1.0f, 1.0f, 0.1f)))
+			.SetOddRowBackgroundBrush(FSlateNoResource())
+			.SetOddRowBackgroundHoveredBrush(IMAGE_BRUSH("Common/Selection", Icon8x8, FLinearColor(1.0f, 1.0f, 1.0f, 0.1f)))
+			.SetSelectorFocusedBrush(BORDER_BRUSH("Common/Selector", FMargin(4.f / 16.f), SelectorColor))
+			.SetActiveBrush(IMAGE_BRUSH("Common/Selection", Icon8x8, SelectionColor))
+			.SetActiveHoveredBrush(IMAGE_BRUSH("Common/Selection", Icon8x8, SelectionColor))
+			.SetInactiveBrush(IMAGE_BRUSH("Common/Selection", Icon8x8, SelectionColor_Inactive))
+			.SetInactiveHoveredBrush(IMAGE_BRUSH("Common/Selection", Icon8x8, SelectionColor_Inactive))
+			.SetTextColor(DefaultForeground)
+			.SetSelectedTextColor(InvertedForeground)
+			.SetDropIndicator_Above(BOX_BRUSH("Common/DropZoneIndicator_Above", FMargin(10.0f / 16.0f, 10.0f / 16.0f, 0, 0), SelectionColor))
 			.SetDropIndicator_Onto(BOX_BRUSH("Common/DropZoneIndicator_Onto", FMargin(4.0f / 16.0f), SelectionColor))
-			.SetDropIndicator_Below(BOX_BRUSH("Common/DropZoneIndicator_Below", FMargin(10.0f / 16.0f, 0, 0, 10.0f / 16.0f), SelectionColor))
-		);
+			.SetDropIndicator_Below(BOX_BRUSH("Common/DropZoneIndicator_Below", FMargin(10.0f / 16.0f, 0, 0, 10.0f / 16.0f), SelectionColor));
+		Style->Set("TableView.Row", DefaultTableRowStyle);
+
+		const FTableRowStyle DarkTableRowStyle = FTableRowStyle(DefaultTableRowStyle)
+			.SetEvenRowBackgroundBrush(IMAGE_BRUSH("Common/Selection", Icon8x8, FLinearColor(0.0f, 0.0f, 0.0f, 0.1f)))
+			.SetOddRowBackgroundBrush(IMAGE_BRUSH("Common/Selection", Icon8x8, FLinearColor(0.0f, 0.0f, 0.0f, 0.1f)));
+		Style->Set("TableView.DarkRow", DarkTableRowStyle);
 
 		Style->Set( "TreeArrow_Collapsed", new IMAGE_BRUSH( "Common/TreeArrow_Collapsed", Icon10x10, DefaultForeground ) );
 		Style->Set( "TreeArrow_Collapsed_Hovered", new IMAGE_BRUSH( "Common/TreeArrow_Collapsed_Hovered", Icon10x10, DefaultForeground ) );
