[INTEGRATE] Change 2412130 by [REDACTED_PERSON][REDACTED_HANDLE] on 2015/01/20 09:15:30

	Fix another Mac build error (assigning variable to itself).

diff --git a/Engine/Source/Editor/Kismet/Private/SCSEditorViewportClient.cpp b/Engine/Source/Editor/Kismet/Private/SCSEditorViewportClient.cpp
--- a/Engine/Source/Editor/Kismet/Private/SCSEditorViewportClient.cpp
+++ b/Engine/Source/Editor/Kismet/Private/SCSEditorViewportClient.cpp
@@ -186,7 +186,7 @@ void FSCSEditorViewportClient::Tick(float DeltaSeconds)
 
 	if ( PreviewActor != LastPreviewActor.Get() || PreviewActor == nullptr )
 	{
-		PreviewActor = PreviewActor;
+//		PreviewActor = PreviewActor;
 
 		Invalidate();
 		RefreshPreviewBounds();
