[INTEGRATE] Change 2412134 by [REDACTED_PERSON][REDACTED_HANDLE] on 2015/01/20 09:17:39

	Probable fix for another Mac build error.

diff --git a/Engine/Source/Editor/Kismet/Private/SMyBlueprint.cpp b/Engine/Source/Editor/Kismet/Private/SMyBlueprint.cpp
--- a/Engine/Source/Editor/Kismet/Private/SMyBlueprint.cpp
+++ b/Engine/Source/Editor/Kismet/Private/SMyBlueprint.cpp
@@ -41,6 +41,8 @@
 #include "Components/TimelineComponent.h"
 #include "Engine/TimelineTemplate.h"
 
+#include "BlueprintEditorSettings.h"
+
 #define LOCTEXT_NAMESPACE "MyBlueprint"
 
 //////////////////////////////////////////////////////////////////////////
