[INTEGRATE] Change 2412172 by [REDACTED_PERSON][REDACTED_HANDLE] on 2015/01/20 09:47:02

	Fixed SCS editor showing actor FNames instead of actual actor labels

diff --git a/Engine/Source/Editor/Kismet/Private/SSCSEditor.cpp b/Engine/Source/Editor/Kismet/Private/SSCSEditor.cpp
--- a/Engine/Source/Editor/Kismet/Private/SSCSEditor.cpp
+++ b/Engine/Source/Editor/Kismet/Private/SSCSEditor.cpp
@@ -3041,7 +3041,7 @@ FText SSCSEditor::GetActorDisplayText() const
 		}
 		else
 		{
-			DefaultActor->GetName(Name);
+			Name = DefaultActor->GetActorLabel();
 			return FText::Format(LOCTEXT("DefaultActor_Name", "{0}"), FText::FromString(Name));
 		}
 	}
