[INTEGRATE] Change 2412829 by [REDACTED_PERSON][REDACTED_HANDLE] on 2015/01/20 16:32:25

	BP - Disabling the old local variable area.

diff --git a/Engine/Source/Editor/Kismet/Private/SMyBlueprint.cpp b/Engine/Source/Editor/Kismet/Private/SMyBlueprint.cpp
--- a/Engine/Source/Editor/Kismet/Private/SMyBlueprint.cpp
+++ b/Engine/Source/Editor/Kismet/Private/SMyBlueprint.cpp
@@ -705,27 +705,27 @@ void SMyBlueprint::Refresh()
 	
 	bool bLocalActionsAreVisible = (GetLocalActionsListVisibility() == EVisibility::Visible);
 
-	if ( !GetDefault<UEditorExperimentalSettings>()->bUnifiedBlueprintEditor )
-	{
-		if ( bLocalActionsAreVisible )
-		{
-			if ( !LocalGraphActionMenu.IsValid() )
-			{
-				ActionMenuContainer->AddSlot()
-					.Value(0.33f)
-					[
-						ConstructLocalActionPanel()
-					];
-			}
-			check(LocalGraphActionMenu.IsValid());
-			LocalGraphActionMenu->RefreshAllActions(true);
-		}
-		else if ( LocalGraphActionMenu.IsValid() && ensure(ActionMenuContainer->GetChildren()->Num() > 1) )
-		{
-			ActionMenuContainer->RemoveAt(1);
-			LocalGraphActionMenu = nullptr;
-		}
-	}
+	//if ( !GetDefault<UEditorExperimentalSettings>()->bUnifiedBlueprintEditor )
+	//{
+	//	if ( bLocalActionsAreVisible )
+	//	{
+	//		if ( !LocalGraphActionMenu.IsValid() )
+	//		{
+	//			ActionMenuContainer->AddSlot()
+	//				.Value(0.33f)
+	//				[
+	//					ConstructLocalActionPanel()
+	//				];
+	//		}
+	//		check(LocalGraphActionMenu.IsValid());
+	//		LocalGraphActionMenu->RefreshAllActions(true);
+	//	}
+	//	else if ( LocalGraphActionMenu.IsValid() && ensure(ActionMenuContainer->GetChildren()->Num() > 1) )
+	//	{
+	//		ActionMenuContainer->RemoveAt(1);
+	//		LocalGraphActionMenu = nullptr;
+	//	}
+	//}
 }
 
 TSharedRef<SWidget> SMyBlueprint::OnCreateWidgetForAction(FCreateWidgetForActionData* const InCreateData)
