[INTEGRATE] Change 2412911 by [REDACTED_PERSON][REDACTED_HANDLE] on 2015/01/20 17:09:54

	Renaming components in the details view will no longer deselect the component when regenerating the tree.

diff --git a/Engine/Source/Editor/Kismet/Private/SSCSEditor.cpp b/Engine/Source/Editor/Kismet/Private/SSCSEditor.cpp
--- a/Engine/Source/Editor/Kismet/Private/SSCSEditor.cpp
+++ b/Engine/Source/Editor/Kismet/Private/SSCSEditor.cpp
@@ -876,6 +876,7 @@ void FSCSEditorTreeNode::OnCompleteRename(const FText& InNewName)
 		}
 
 		ComponentInstance->Rename(*InNewName.ToString(), nullptr, RenameFlags);
+		InstancedComponentName = *InNewName.ToString();
 	}
 	else
 	{
