[INTEGRATE] Change 2412999 by [REDACTED_PERSON][REDACTED_HANDLE] on 2015/01/20 17:54:09

	Exclude "Add Events" option from the SCS editor context menu when in "instanced" mode (per discussion).

diff --git a/Engine/Source/Editor/Kismet/Private/SSCSEditor.cpp b/Engine/Source/Editor/Kismet/Private/SSCSEditor.cpp
--- a/Engine/Source/Editor/Kismet/Private/SSCSEditor.cpp
+++ b/Engine/Source/Editor/Kismet/Private/SSCSEditor.cpp
@@ -2771,29 +2771,32 @@ TSharedPtr< SWidget > SSCSEditor::CreateContextMenu()
 				MenuBuilder.AddMenuEntry( FGenericCommands::Get().Delete );
 				MenuBuilder.AddMenuEntry( FGenericCommands::Get().Rename );
 
-				// Collect the classes of all selected objects
-				TArray<UClass*> SelectionClasses;
-				for( auto NodeIter = SelectedNodes.CreateConstIterator(); NodeIter; ++NodeIter )
+				if(EditorMode.Get() == EEditorMode::BlueprintSCS)
 				{
-					auto TreeNode = *NodeIter;
-					if( TreeNode->GetComponentTemplate() )
+					// Collect the classes of all selected objects
+					TArray<UClass*> SelectionClasses;
+					for( auto NodeIter = SelectedNodes.CreateConstIterator(); NodeIter; ++NodeIter )
 					{
-						SelectionClasses.Add( TreeNode->GetComponentTemplate()->GetClass() );
+						auto TreeNode = *NodeIter;
+						if( TreeNode->GetComponentTemplate() )
+						{
+							SelectionClasses.Add( TreeNode->GetComponentTemplate()->GetClass() );
+						}
 					}
-				}
 
-				if ( SelectionClasses.Num() )
-				{
-					// Find the common base class of all selected classes
-					UClass* SelectedClass = UClass::FindCommonBase( SelectionClasses );
-					// Build an event submenu if we can generate events
-					if( FBlueprintEditorUtils::CanClassGenerateEvents( SelectedClass ))
+					if ( SelectionClasses.Num() )
 					{
-						MenuBuilder.AddSubMenu(	LOCTEXT("AddEventSubMenu", "Add Event"), 
-							LOCTEXT("ActtionsSubMenu_ToolTip", "Add Event"), 
-							FNewMenuDelegate::CreateStatic( &SSCSEditor::BuildMenuEventsSection,
+						// Find the common base class of all selected classes
+						UClass* SelectedClass = UClass::FindCommonBase( SelectionClasses );
+						// Build an event submenu if we can generate events
+						if( FBlueprintEditorUtils::CanClassGenerateEvents( SelectedClass ))
+						{
+							MenuBuilder.AddSubMenu(	LOCTEXT("AddEventSubMenu", "Add Event"), 
+								LOCTEXT("ActtionsSubMenu_ToolTip", "Add Event"), 
+								FNewMenuDelegate::CreateStatic( &SSCSEditor::BuildMenuEventsSection,
 								GetBlueprint(), SelectedClass, FCanExecuteAction::CreateSP(this, &SSCSEditor::IsEditingAllowed),
 								FGetSelectedObjectsDelegate::CreateSP(this, &SSCSEditor::GetSelectedItemsForContextMenu)));
+						}
 					}
 				}
 			}
