[INTEGRATE] Change 2413308 by [REDACTED_PERSON][REDACTED_HANDLE] on 2015/01/20 21:24:44

	Fix for layout of the SCS Editor, now two areas much more consistant (more changes to the tree to come)
	Fixed component combo box not updating properly with asset selection
	Fixed component combo box not clearing the search on close. UE-7839
	Fixed some inconsistent padding. UE-7863
	Fixed crashes when drag and dropping (visual issues with the drag and drop decorator still) UE-7895

diff --git a/Engine/Source/Editor/Kismet/Private/BPVariableDragDropAction.cpp b/Engine/Source/Editor/Kismet/Private/BPVariableDragDropAction.cpp
--- a/Engine/Source/Editor/Kismet/Private/BPVariableDragDropAction.cpp
+++ b/Engine/Source/Editor/Kismet/Private/BPVariableDragDropAction.cpp
@@ -45,6 +45,11 @@ void FKismetVariableDragDropAction::GetLinksThatWillBreak(	UEdGraphNode* Node, U
 void FKismetVariableDragDropAction::HoverTargetChanged()
 {
 	UProperty* VariableProperty = GetVariableProperty();
+	if (VariableProperty == nullptr)
+	{
+		return;
+	}
+
 	FString VariableString = VariableName.ToString();
 
 	// Icon/text to draw on tooltip
diff --git a/Engine/Source/Editor/Kismet/Private/BPVariableDragDropAction.h b/Engine/Source/Editor/Kismet/Private/BPVariableDragDropAction.h
--- a/Engine/Source/Editor/Kismet/Private/BPVariableDragDropAction.h
+++ b/Engine/Source/Editor/Kismet/Private/BPVariableDragDropAction.h
@@ -51,11 +51,13 @@ public:
 
 	UProperty* GetVariableProperty()
 	{
-		check(VariableSource.IsValid());
-		check(VariableName != NAME_None);
-		UProperty* VariableProperty = FindField<UProperty>(VariableSource.Get(), VariableName);
-		check(VariableProperty != NULL);
-		return VariableProperty;
+		if (VariableSource.IsValid() && VariableName != NAME_None)
+		{
+			UProperty* VariableProperty = FindField<UProperty>(VariableSource.Get(), VariableName);
+			check(VariableProperty != NULL);
+			return VariableProperty;
+		}
+		return nullptr;
 	}
 
 	/** Set if operation is modified by alt */
diff --git a/Engine/Source/Editor/Kismet/Private/SCSDiff.cpp b/Engine/Source/Editor/Kismet/Private/SCSDiff.cpp
--- a/Engine/Source/Editor/Kismet/Private/SCSDiff.cpp
+++ b/Engine/Source/Editor/Kismet/Private/SCSDiff.cpp
@@ -73,11 +73,12 @@ TArray< FSCSResolvedIdentifier > FSCSDiff::GetDisplayedHierarchy() const
 
 	if( SCSEditor.IsValid() )
 	{
-		for (int32 Iter = 0; Iter != SCSEditor->RootNodes.Num(); ++Iter)
+		const TArray<FSCSEditorTreeNodePtrType>& RootNodes = SCSEditor->GetRootComponentNodes();
+		for (int32 Iter = 0; Iter != RootNodes.Num(); ++Iter)
 		{
 			TArray< int32 > TreeAddress;
 			TreeAddress.Push(Iter);
-			GetDisplayedHierarchyRecursive(TreeAddress, *SCSEditor->RootNodes[Iter], Ret);
+			GetDisplayedHierarchyRecursive(TreeAddress, *RootNodes[Iter], Ret);
 		}
 	}
 
diff --git a/Engine/Source/Editor/Kismet/Private/SCSEditorViewportClient.cpp b/Engine/Source/Editor/Kismet/Private/SCSEditorViewportClient.cpp
--- a/Engine/Source/Editor/Kismet/Private/SCSEditorViewportClient.cpp
+++ b/Engine/Source/Editor/Kismet/Private/SCSEditorViewportClient.cpp
@@ -573,7 +573,7 @@ FWidget::EWidgetMode FSCSEditorViewportClient::GetWidgetMode() const
 		if ( BluePrintEditor.IsValid() )
 		{
 			TArray<FSCSEditorTreeNodePtrType> SelectedNodes = BluePrintEditor->GetSelectedSCSEditorTreeNodes();
-			const TArray<FSCSEditorTreeNodePtrType>& RootNodes = BluePrintEditor->GetSCSEditor()->RootNodes;
+			const TArray<FSCSEditorTreeNodePtrType>& RootNodes = BluePrintEditor->GetSCSEditor()->GetRootComponentNodes();
 
 			// if the selected nodes array is empty, or only contains entries from the
 			// root nodes array, or isn't visible in the preview actor, then don't display a transform widget
diff --git a/Engine/Source/Editor/Kismet/Private/SSCSEditor.cpp b/Engine/Source/Editor/Kismet/Private/SSCSEditor.cpp
--- a/Engine/Source/Editor/Kismet/Private/SSCSEditor.cpp
+++ b/Engine/Source/Editor/Kismet/Private/SSCSEditor.cpp
@@ -1089,6 +1089,7 @@ TSharedRef<SWidget> SSCS_RowWidget::GenerateWidgetForColumn( const FName& Column
 					]
 				+SHorizontalBox::Slot()
 					.VAlign(VAlign_Center)
+					.Padding(2, 0, 0, 0)
 					[
 						InlineWidget.ToSharedRef()
 					];
@@ -1472,8 +1473,9 @@ FReply SSCS_RowWidget::OnDragDetected( const FGeometry& MyGeometry, const FPoint
 		}
 
 		TSharedPtr<FSCSEditorTreeNode> FirstNode = SelectedNodePtrs[0];
+		UBlueprint* Blueprint = FirstNode->GetBlueprint();
 		
-		TSharedRef<FSCSRowDragDropOp> Operation = FSCSRowDragDropOp::New(FirstNode->GetVariableName(), FirstNode->GetBlueprint()->SkeletonGeneratedClass, FNodeCreationAnalytic());
+		TSharedRef<FSCSRowDragDropOp> Operation = FSCSRowDragDropOp::New(FirstNode->GetVariableName(), Blueprint != nullptr ? Blueprint->SkeletonGeneratedClass : nullptr, FNodeCreationAnalytic());
 		//Operation->SetAltDrag(MouseEvent.IsAltDown());
 		//Operation->SetCtrlDrag(MouseEvent.IsLeftControlDown() || MouseEvent.IsRightControlDown());
 		Operation->SetCtrlDrag(true); // Always put a getter
@@ -2406,22 +2408,6 @@ void SSCSEditor::Construct( const FArguments& InArgs )
 	{
 		HeaderRow = SNew(SHeaderRow)
 
-		+ SHeaderRow::Column(SCS_ColumnName_Mobility)
-		.DefaultLabel(LOCTEXT("MobilityColumnLabel", "Mobility"))
-		.FixedWidth(16.0f) // mobility icons are 16px (16 slate-units = 16px, when application scale == 1)
-		.HeaderContent()
-		[
-			SNew(SHorizontalBox)
-			.ToolTip(SNew(SToolTip).Text(LOCTEXT("MobilityColumnTooltip", "Mobility")))
-			+ SHorizontalBox::Slot()
-			.FillWidth(1.0f)
-			.VAlign(VAlign_Center)
-			.HAlign(HAlign_Center)
-			[
-				SNew(SImage).Image(MobilityHeaderBrush)
-			]
-		]
-
 		+ SHeaderRow::Column(SCS_ColumnName_ComponentClass)
 		.DefaultLabel(LOCTEXT("Class", "Class"))
 		.FillWidth(4);
@@ -2483,110 +2469,136 @@ void SSCSEditor::Construct( const FArguments& InArgs )
 	if (bSingleLayoutBPEditor)
 	{
 		Contents = SNew(SVerticalBox)
-
-		// Root Actor
+		
 		+ SVerticalBox::Slot()
-		.Padding(0.f)
-		.AutoHeight()
+		.Padding(0.0f)
 		[
 			SNew(SVerticalBox)
 
 			+ SVerticalBox::Slot()
 			.AutoHeight()
-			.VAlign(VAlign_Center)
-			.HAlign(HAlign_Left)
-			.Padding(2.0f, 2.0f)
+			.VAlign(VAlign_Top)
+			.Padding(0.0f, 2.0f)
 			[
-				SNew(SComponentClassCombo)
-				.OnComponentClassSelected(this, &SSCSEditor::PerformComboAddClass)
-				.ToolTipText(LOCTEXT("AddComponent_Tooltip", "Add a component."))
+				SNew(SBorder)
+				.Padding(2.0f)
+				.BorderImage(FEditorStyle::GetBrush("ToolPanel.GroupBorder"))
+				.AddMetaData<FTagMetaData>(FTagMetaData(TEXT("ComponentsPanel")))
+				[
+					SNew(SBox)
+					.HAlign(HAlign_Left)
+					[
+						SNew(SComponentClassCombo)
+						.OnComponentClassSelected(this, &SSCSEditor::PerformComboAddClass)
+						.ToolTipText(LOCTEXT("AddComponent_Tooltip", "Add a component."))
+					]
+				]
 			]
 
 			+ SVerticalBox::Slot()
-			.VAlign(VAlign_Fill)
-			.HAlign(HAlign_Fill)
 			.Padding(0.0f, 0.0f)
 			[
-				SNew(SCheckBox)
-				.Style(FCoreStyle::Get(), "ToggleButtonRowStyle")
-				.IsFocusable(true)
-				.OnCheckStateChanged(this, &SSCSEditor::OnActorSelected)
-				.IsChecked(this, &SSCSEditor::OnIsActorSelected)
-				.ToolTip(Tooltip)
-				.Padding(FMargin(2.0f))
+				SNew(SBorder)
+				.Padding(2.0f)
+				.BorderImage(FEditorStyle::GetBrush("ToolPanel.GroupBorder"))
+				.AddMetaData<FTagMetaData>(FTagMetaData(TEXT("ComponentsPanel")))
+				
 				[
-					SNew(SHorizontalBox)
+					SNew(SVerticalBox)
 
-					+ SHorizontalBox::Slot()
-					.AutoWidth()
-					.VAlign(VAlign_Center)
-					.Padding(FMargin(0.f, 0.f, 6.f, 0.f))
+					// Root Actor
+					+SVerticalBox::Slot()
+					.AutoHeight()
+					.Padding(0.0f, 0.0f)
 					[
-						SNew(SBox)
+						SNew(SCheckBox)
+						.Style(FCoreStyle::Get(), "ToggleButtonRowStyle")
+						.IsFocusable(true)
+						.OnCheckStateChanged(this, &SSCSEditor::OnActorSelected)
+						.IsChecked(this, &SSCSEditor::OnIsActorSelected)
+						.ToolTip(Tooltip)
+						.Padding(FMargin(2.0f))
 						[
-							SNew(SImage)
-							.Image(this, &SSCSEditor::GetActorIcon)
+							SNew(SHorizontalBox)
+
+							+ SHorizontalBox::Slot()
+							.AutoWidth()
+							.VAlign(VAlign_Center)
+							.Padding(FMargin(0.f, 0.f, 6.f, 0.f))
+							[
+								SNew(SImage)
+								.Image(this, &SSCSEditor::GetActorIcon)
+							]
+
+							+ SHorizontalBox::Slot()
+							.HAlign(HAlign_Left)
+							.VAlign(VAlign_Center)
+							.Padding(0.0f, 0.0f)
+							[
+								SNew(STextBlock)
+								.Text(this, &SSCSEditor::GetActorDisplayText)
+								.ColorAndOpacity(FSlateColor::UseForeground())
+							]
 						]
 					]
 
-					+ SHorizontalBox::Slot()
-					.HAlign(HAlign_Left)
-					.VAlign(VAlign_Center)
-					.Padding(0.0f, 0.0f)
+					// Tree
+					+ SVerticalBox::Slot()
+					.FillHeight(1.0f)
+					.Padding(0.f, 0.f, 0.f, 2.f)
 					[
-						SNew(STextBlock)
-						.Text(this, &SSCSEditor::GetActorDisplayText)
-						.ColorAndOpacity(FSlateColor::UseForeground())
+						SCSTreeWidget.ToSharedRef()
 					]
 				]
 			]
-		]
-
-		// Tree
-		+ SVerticalBox::Slot()
-		.Padding(0.f, 0.f, 0.f, 2.f)
-		[
-			SCSTreeWidget.ToSharedRef()
 		];
 	}
 	else if( InArgs._HideComponentClassCombo.Get() )
 	{
-		Contents = SCSTreeWidget;
+		Contents = SNew(SBorder)
+
+		.Padding(2.0f)
+		.BorderImage(FEditorStyle::GetBrush("ToolPanel.GroupBorder"))
+		.AddMetaData<FTagMetaData>(FTagMetaData(TEXT("ComponentsPanel")))
+		[
+			SCSTreeWidget.ToSharedRef()
+		];
 	}
 	else
 	{
-		Contents = SNew(SVerticalBox)
+		Contents = SNew(SBorder)
 
-		// Component picker
-		+ SVerticalBox::Slot()
-		.Padding(1.f)
-		.AutoHeight()
+		.Padding(2.0f)
+		.BorderImage(FEditorStyle::GetBrush("ToolPanel.GroupBorder"))
+		.AddMetaData<FTagMetaData>(FTagMetaData(TEXT("ComponentsPanel")))
 		[
-			SNew(SHorizontalBox)
-			+ SHorizontalBox::Slot()
-			.AutoWidth()
+			SNew(SVerticalBox)
+
+			// Component picker
+			+ SVerticalBox::Slot()
+			.Padding(1.f)
+			.AutoHeight()
 			[
-				SNew(SComponentClassCombo)
-				.OnComponentClassSelected(this, &SSCSEditor::PerformComboAddClass)
+				SNew(SHorizontalBox)
+				+ SHorizontalBox::Slot()
+				.AutoWidth()
+				[
+					SNew(SComponentClassCombo)
+					.OnComponentClassSelected(this, &SSCSEditor::PerformComboAddClass)
+				]
+			]
+			// Tree
+			+ SVerticalBox::Slot()
+			.Padding(0.f, 0.f, 0.f, 2.f)
+			[
+				SCSTreeWidget.ToSharedRef()
 			]
-		]
-		// Tree
-		+ SVerticalBox::Slot()
-		.Padding(0.f, 0.f, 0.f, 2.f)
-		[
-			SCSTreeWidget.ToSharedRef()
 		];
 	}
 
 	this->ChildSlot
 	[
-		SNew(SBorder)
-		.Padding(2.0f)
-		.BorderImage(FEditorStyle::GetBrush("ToolPanel.GroupBorder"))
-		.AddMetaData<FTagMetaData>(FTagMetaData(TEXT("ComponentsPanel")))
-		[
-			Contents.ToSharedRef()
-		]
+		Contents.ToSharedRef()
 	];
 
 	// Refresh the tree widget
@@ -4408,5 +4420,10 @@ void SSCSEditor::GetCollapsedNodes(const FSCSEditorTreeNodePtrType& InNodePtr, T
 	}
 }
 
+const TArray<FSCSEditorTreeNodePtrType>& SSCSEditor::GetRootComponentNodes()
+{
+	return RootNodes;
+}
+
 #undef LOCTEXT_NAMESPACE
 
diff --git a/Engine/Source/Editor/Kismet/Public/SSCSEditor.h b/Engine/Source/Editor/Kismet/Public/SSCSEditor.h
--- a/Engine/Source/Editor/Kismet/Public/SSCSEditor.h
+++ b/Engine/Source/Editor/Kismet/Public/SSCSEditor.h
@@ -80,7 +80,7 @@ public:
 	 * @return The Blueprint to which this node belongs.
 	 */
 	UBlueprint* GetBlueprint() const;
-	/** 
+	/**
 	 * @return Whether or not this object represents a root component.
 	 */
 	bool IsRoot() const;
@@ -332,10 +332,6 @@ protected:
 	FSCSEditorTreeNodePtrType NodePtr;
 };
 
-
-
-
-
 //////////////////////////////////////////////////////////////////////////
 // SSCSEditorDragDropTree - implements STreeView for our specific node type and adds drag/drop functionality
 class SSCSEditorDragDropTree : public STreeView<FSCSEditorTreeNodePtrType>
@@ -592,6 +588,9 @@ public:
 	/** Provides access to the Blueprint context that's being edited */
 	class UBlueprint* GetBlueprint() const;
 
+	/** Returns the set of root nodes */
+	const TArray<FSCSEditorTreeNodePtrType>& GetRootComponentNodes();
+
 protected:
 	FString GetSelectedClassText() const;
 
@@ -677,13 +676,11 @@ protected:
 
 	/** Returns the set of expandable nodes that are currently collapsed in the UI */
 	void GetCollapsedNodes(const FSCSEditorTreeNodePtrType& InNodePtr, TSet<FSCSEditorTreeNodePtrType>& OutCollapsedNodes) const;
+
 public:
 	/** Tree widget */
 	TSharedPtr<SSCSTreeType> SCSTreeWidget;
 
-	/** Root set of tree */
-	TArray<FSCSEditorTreeNodePtrType> RootNodes;
-
 	/** The node that represents the root component in the scene hierarchy */
 	FSCSEditorTreeNodePtrType SceneRootNodePtr;
 
@@ -718,6 +715,9 @@ public:
 	FOnHighlightPropertyInDetailsView OnHighlightPropertyInDetailsView;
 private:
 
+	/** Root set of tree */
+	TArray<FSCSEditorTreeNodePtrType> RootNodes;
+
 	/** Flag to enable/disable component editing */
 	bool	bEnableComponentEditing;
 
diff --git a/Engine/Source/Editor/UnrealEd/Private/SComponentClassCombo.cpp b/Engine/Source/Editor/UnrealEd/Private/SComponentClassCombo.cpp
--- a/Engine/Source/Editor/UnrealEd/Private/SComponentClassCombo.cpp
+++ b/Engine/Source/Editor/UnrealEd/Private/SComponentClassCombo.cpp
@@ -86,7 +86,7 @@ void SComponentClassCombo::Construct(const FArguments& InArgs)
 		]
 	]
 	.IsFocusable(true)
-	.ContentPadding(FMargin(2))
+	.ContentPadding(FMargin(0))
 	.ComboButtonStyle(FEditorStyle::Get(), "ContentBrowser.NewAsset.Style")
 	.ForegroundColor(FLinearColor::White)
 	.OnComboBoxOpened(this, &SComponentClassCombo::ClearSelection);
@@ -244,50 +244,6 @@ TSharedRef<ITableRow> SComponentClassCombo::GenerateAddComponentRow( FComponentC
 	}
 	else
 	{
-		// Get a user friendly string from the component name
-		FString FriendlyComponentName = GetSanitizedComponentName( Entry->GetComponentClass() );
-
-		// Search the selected assets and look for any that can be used as a source asset for this type of component
-		// If there is one we append the asset name to the component name, if there are many we append "Multiple Assets"
-		FString AssetName;
-		UObject* PreviousMatchingAsset = NULL;
-
-		FEditorDelegates::LoadSelectedAssetsIfNeeded.Broadcast();
-		USelection* Selection = GEditor->GetSelectedObjects();
-		for ( FSelectionIterator ObjectIter(*Selection); ObjectIter; ++ObjectIter )
-		{
-			UObject* Object = *ObjectIter;
-			UClass* Class = Object->GetClass();
-
-			TArray<TSubclassOf<UActorComponent> > ComponentClasses = FComponentAssetBrokerage::GetComponentsForAsset(Object);
-			for ( int32 ComponentIndex = 0; ComponentIndex < ComponentClasses.Num(); ComponentIndex++ )
-			{
-				if ( ComponentClasses[ComponentIndex]->IsChildOf( Entry->GetComponentClass() ) )
-				{
-					if ( AssetName.IsEmpty() )
-					{
-						// If there is no previous asset then we just accept the name
-						AssetName = Object->GetName();
-						PreviousMatchingAsset = Object;
-					}
-					else
-					{
-						// if there is a previous asset then check that we didn't just find multiple appropriate components
-						// in a single asset - if the asset differs then we don't display the name, just "Multiple Assets"
-						if ( PreviousMatchingAsset != Object )
-						{
-							AssetName = LOCTEXT("MultipleAssetsForComponentAnnotation", "Multiple Assets").ToString();
-							PreviousMatchingAsset = Object;
-						}
-					}
-				}
-			}
-		}
-
-		if ( !AssetName.IsEmpty() )
-		{
-			FriendlyComponentName += FString(" (") + AssetName + FString(")");
-		}
 
 		return
 			SNew( SComboRow< TSharedPtr<FString> >, OwnerTable )
@@ -321,7 +277,7 @@ TSharedRef<ITableRow> SComponentClassCombo::GenerateAddComponentRow( FComponentC
 				[
 					SNew(STextBlock)
 					.HighlightText(this, &SComponentClassCombo::GetCurrentSearchString)
-					.Text( FText::FromString(FriendlyComponentName) )
+					.Text(this, &SComponentClassCombo::GetFriendlyComponentName, Entry)
 				]
 			];
 	}
@@ -421,6 +377,61 @@ void SComponentClassCombo::UpdateComponentClassList()
 	}
 }
 
+FReply SComponentClassCombo::OnButtonClicked()
+{
+	SearchBox->SetText(FText::GetEmpty());
+	return SComboButton::OnButtonClicked();
+}
+
+FText SComponentClassCombo::GetFriendlyComponentName(FComponentClassComboEntryPtr Entry) const
+{
+	// Get a user friendly string from the component name
+	FString FriendlyComponentName = GetSanitizedComponentName(Entry->GetComponentClass());
+
+	// Search the selected assets and look for any that can be used as a source asset for this type of component
+	// If there is one we append the asset name to the component name, if there are many we append "Multiple Assets"
+	FString AssetName;
+	UObject* PreviousMatchingAsset = NULL;
+
+	FEditorDelegates::LoadSelectedAssetsIfNeeded.Broadcast();
+	USelection* Selection = GEditor->GetSelectedObjects();
+	for (FSelectionIterator ObjectIter(*Selection); ObjectIter; ++ObjectIter)
+	{
+		UObject* Object = *ObjectIter;
+		UClass* Class = Object->GetClass();
+
+		TArray<TSubclassOf<UActorComponent> > ComponentClasses = FComponentAssetBrokerage::GetComponentsForAsset(Object);
+		for (int32 ComponentIndex = 0; ComponentIndex < ComponentClasses.Num(); ComponentIndex++)
+		{
+			if (ComponentClasses[ComponentIndex]->IsChildOf(Entry->GetComponentClass()))
+			{
+				if (AssetName.IsEmpty())
+				{
+					// If there is no previous asset then we just accept the name
+					AssetName = Object->GetName();
+					PreviousMatchingAsset = Object;
+				}
+				else
+				{
+					// if there is a previous asset then check that we didn't just find multiple appropriate components
+					// in a single asset - if the asset differs then we don't display the name, just "Multiple Assets"
+					if (PreviousMatchingAsset != Object)
+					{
+						AssetName = LOCTEXT("MultipleAssetsForComponentAnnotation", "Multiple Assets").ToString();
+						PreviousMatchingAsset = Object;
+					}
+				}
+			}
+		}
+	}
+
+	if (!AssetName.IsEmpty())
+	{
+		FriendlyComponentName += FString(" (") + AssetName + FString(")");
+	}
+	return FText::FromString(FriendlyComponentName);
+}
+
 FString SComponentClassCombo::GetSanitizedComponentName( UClass* ComponentClass )
 {
 	FString DisplayName;
diff --git a/Engine/Source/Editor/UnrealEd/Public/SComponentClassCombo.h b/Engine/Source/Editor/UnrealEd/Public/SComponentClassCombo.h
--- a/Engine/Source/Editor/UnrealEd/Public/SComponentClassCombo.h
+++ b/Engine/Source/Editor/UnrealEd/Public/SComponentClassCombo.h
@@ -101,8 +101,13 @@ public:
 	/** Returns a component name without the substring "Component" and sanitized for display */
 	static FString GetSanitizedComponentName( UClass* ComponentClass );
 
+protected:
+	virtual FReply OnButtonClicked();
+
 private:
 
+	FText GetFriendlyComponentName(FComponentClassComboEntryPtr Entry) const;
+
 	FComponentClassSelected OnComponentClassSelected;
 
 	/** List of component class names used by combo box */
