[INTEGRATE] Change 2413664 by [REDACTED_PERSON][REDACTED_HANDLE] on 2015/01/21 08:33:11

	BP - In the unified BP editor we now check if we're selecting components to know if we should enable the the customization.

diff --git a/Engine/Source/Editor/Kismet/Private/SKismetInspector.cpp b/Engine/Source/Editor/Kismet/Private/SKismetInspector.cpp
--- a/Engine/Source/Editor/Kismet/Private/SKismetInspector.cpp
+++ b/Engine/Source/Editor/Kismet/Private/SKismetInspector.cpp
@@ -371,6 +371,24 @@ void SKismetInspector::AddPropertiesRecursive(UProperty* Property)
 
 void SKismetInspector::UpdateFromObjects(const TArray<UObject*>& PropertyObjects, struct FKismetSelectionInfo& SelectionInfo, const FShowDetailsOptions& Options)
 {
+	// If we're using the unified blueprint editor, there's not an explicit point where
+	// we ender a kind of component editing mode, so instead, just look at what we're selecting.
+	// If we select a component, then enable the customization.
+	if ( GetDefault<UEditorExperimentalSettings>()->bUnifiedBlueprintEditor )
+	{
+		bool bEnableComponentCustomization = false;
+		for ( UObject* PropertyObject : PropertyObjects )
+		{
+			if ( PropertyObject->IsA<UActorComponent>() )
+			{
+				bEnableComponentCustomization = true;
+				break;
+			}
+		}
+
+		EnableComponentDetailsCustomization(bEnableComponentCustomization);
+	}
+
 	PropertyView->OnFinishedChangingProperties().Clear();
 	PropertyView->OnFinishedChangingProperties().Add( UserOnFinishedChangingProperties );
 
