[INTEGRATE] Change 2413953 by [REDACTED_PERSON][REDACTED_HANDLE] on 2015/01/21 12:18:28

	Update selection after deleting nodes in the SCS editor "Actor Instance" mode.

diff --git a/Engine/Source/Editor/Kismet/Private/SSCSEditor.cpp b/Engine/Source/Editor/Kismet/Private/SSCSEditor.cpp
--- a/Engine/Source/Editor/Kismet/Private/SSCSEditor.cpp
+++ b/Engine/Source/Editor/Kismet/Private/SSCSEditor.cpp
@@ -3985,14 +3985,38 @@ void SSCSEditor::OnDeleteNodes()
 
 		ActorInstance->Modify();
 
+		FSCSEditorTreeNodePtrType NewSelection;
 		TArray<FSCSEditorTreeNodePtrType> SelectedNodes = SCSTreeWidget->GetSelectedItems();
 		for (int32 i = 0; i < SelectedNodes.Num(); ++i)
 		{
 			auto Node = SelectedNodes[i];
 
+			// Find an appropriate node to select after removal
+			if(!NewSelection.IsValid() || NewSelection == Node)
+			{
+				// Default to the parent node
+				NewSelection = Node->GetParent();
+				if(NewSelection.IsValid())
+				{
+					// If we have sibling nodes, find the one that immediately precedes the one being removed
+					const TArray<FSCSEditorTreeNodePtrType>& ChildNodes = NewSelection->GetChildren();
+					for (int32 ChildIndex = 0; ChildIndex < ChildNodes.Num() && Node != ChildNodes[ChildIndex]; ++ChildIndex)
+					{
+						NewSelection = ChildNodes[ChildIndex];
+					}
+				}
+			}
+
+			// This will clear the current selection
 			RemoveComponentNode(Node);
 		}
 
+		// Reset the selection
+		if(NewSelection.IsValid())
+		{
+			SCSTreeWidget->SetItemSelection(NewSelection, true);
+		}
+
 		// Rebuild the tree view to reflect the new component hierarchy
 		UpdateTree();
 	}
