[INTEGRATE] Change 2414301 by [REDACTED_PERSON][REDACTED_HANDLE] on 2015/01/21 15:45:08

	Added a Tick() override to the SCS editor for Actor Instance mode to keep component tree nodes in sync w/ the selected Actor instance.

diff --git a/Engine/Source/Editor/Kismet/Private/SSCSEditor.cpp b/Engine/Source/Editor/Kismet/Private/SSCSEditor.cpp
--- a/Engine/Source/Editor/Kismet/Private/SSCSEditor.cpp
+++ b/Engine/Source/Editor/Kismet/Private/SSCSEditor.cpp
@@ -40,6 +40,8 @@
 
 #define LOCTEXT_NAMESPACE "SSCSEditor"
 
+DEFINE_LOG_CATEGORY_STATIC(LogSCSEditor, Log, All);
+
 static const FName SCS_ColumnName_ComponentClass( "ComponentClass" );
 static const FName SCS_ColumnName_Asset( "Asset" );
 static const FName SCS_ColumnName_Mobility( "Mobility" );
@@ -2713,6 +2715,34 @@ UBlueprint* SSCSEditor::GetBlueprint() const
 	return nullptr;
 }
 
+void SSCSEditor::Tick( const FGeometry& AllottedGeometry, const double InCurrentTime, const float InDeltaTime )
+{
+	SCompoundWidget::Tick(AllottedGeometry, InCurrentTime, InDeltaTime);
+
+	if(EditorMode.Get() == EEditorMode::ActorInstance)
+	{
+		TFunction<bool (const TArray<FSCSEditorTreeNodePtrType>&, int32&)> AreAnyNodesInvalidLambda = [&](const TArray<FSCSEditorTreeNodePtrType>& InNodes, int32& OutNumValidNodes) -> bool
+		{
+			bool bFoundInvalidNode = false;
+			for(auto NodeIt = InNodes.CreateConstIterator(); NodeIt && !bFoundInvalidNode; ++NodeIt)
+			{
+				const UActorComponent* InstancedComponent = (*NodeIt)->GetComponentTemplate();
+				bFoundInvalidNode = !InstancedComponent || InstancedComponent->IsPendingKill() || AreAnyNodesInvalidLambda((*NodeIt)->GetChildren(), ++OutNumValidNodes);
+			}
+
+			return bFoundInvalidNode;
+		};
+
+		int32 NumComponentNodes = 0;
+		if(AreAnyNodesInvalidLambda(RootNodes, NumComponentNodes) || NumComponentNodes != ActorContext.Get()->GetComponents().Num())
+		{
+			UE_LOG(LogSCSEditor, Log, TEXT("Calling UpdateTree() from Tick()."));
+
+			UpdateTree();
+		}
+	}
+}
+
 FReply SSCSEditor::OnKeyDown( const FGeometry& MyGeometry, const FKeyEvent& InKeyEvent )
 {
 	if ( CommandList->ProcessCommandBindings( InKeyEvent ) )
diff --git a/Engine/Source/Editor/Kismet/Public/SSCSEditor.h b/Engine/Source/Editor/Kismet/Public/SSCSEditor.h
--- a/Engine/Source/Editor/Kismet/Public/SSCSEditor.h
+++ b/Engine/Source/Editor/Kismet/Public/SSCSEditor.h
@@ -442,7 +442,8 @@ public:
 
 	void Construct(const FArguments& InArgs);
 
-	/** Override OnKeyDown */
+	/** SWidget interface */
+	virtual void Tick( const FGeometry& AllottedGeometry, const double InCurrentTime, const float InDeltaTime );
 	virtual FReply OnKeyDown( const FGeometry& MyGeometry, const FKeyEvent& InKeyEvent );
 
 	/** Used by tree control - make a widget for a table row from a node */
