[INTEGRATE] Change 2414562 by [REDACTED_PERSON][REDACTED_HANDLE] on 2015/01/21 18:00:41

	Fixed - SCS editor Tick() was calling UpdateTree() every frame when the selected Actor instance had editor-only components.

diff --git a/Engine/Source/Editor/Kismet/Private/SSCSEditor.cpp b/Engine/Source/Editor/Kismet/Private/SSCSEditor.cpp
--- a/Engine/Source/Editor/Kismet/Private/SSCSEditor.cpp
+++ b/Engine/Source/Editor/Kismet/Private/SSCSEditor.cpp
@@ -2733,8 +2733,21 @@ void SSCSEditor::Tick( const FGeometry& AllottedGeometry, const double InCurrent
 			return bFoundInvalidNode;
 		};
 
+		const AActor* ActorInstance = ActorContext.Get();
+		check(ActorInstance);
+
+		int32 NumComponentInstances = 0;
+		for(auto CompIt = ActorInstance->GetComponents().CreateConstIterator(); CompIt; ++CompIt)
+		{
+			// Don't count editor-only components, because we don't show them
+			if(!(*CompIt)->IsEditorOnly())
+			{
+				++NumComponentInstances;
+			}
+		}
+
 		int32 NumComponentNodes = 0;
-		if(AreAnyNodesInvalidLambda(RootNodes, NumComponentNodes) || NumComponentNodes != ActorContext.Get()->GetComponents().Num())
+		if(AreAnyNodesInvalidLambda(RootNodes, NumComponentNodes) || NumComponentNodes != NumComponentInstances)
 		{
 			UE_LOG(LogSCSEditor, Log, TEXT("Calling UpdateTree() from Tick()."));
 
