[INTEGRATE] Change 2414788 by [REDACTED_PERSON][REDACTED_HANDLE] on 2015/01/21 21:11:27

	[UE-8000] [UE-7812] [UE-7798]

	Fixed transformations of components for in-world editing
	Fixed selection resetting when clicking in the tree but not on an item

diff --git a/Engine/Source/Editor/Kismet/Private/SCSEditorViewportClient.cpp b/Engine/Source/Editor/Kismet/Private/SCSEditorViewportClient.cpp
--- a/Engine/Source/Editor/Kismet/Private/SCSEditorViewportClient.cpp
+++ b/Engine/Source/Editor/Kismet/Private/SCSEditorViewportClient.cpp
@@ -393,7 +393,7 @@ bool FSCSEditorViewportClient::InputWidgetDelta( FViewport* Viewport, EAxisList:
 				{
 					FSCSEditorTreeNodePtrType SelectedNodePtr = *It;
 					// Don't allow editing of a root node, inherited SCS node or child node that also has a movable (non-root) parent node selected
-					const bool bCanEdit = !SelectedNodePtr->IsRoot() && !SelectedNodePtr->IsInherited()
+					const bool bCanEdit =  !SelectedNodePtr->IsRoot() && !SelectedNodePtr->IsInherited()
 						&& !IsMovableParentNodeSelected(SelectedNodePtr, SelectedNodes);
 
 					if(bCanEdit)
@@ -407,22 +407,9 @@ bool FSCSEditorViewportClient::InputWidgetDelta( FViewport* Viewport, EAxisList:
 							FRotator OldRelativeRotation = SelectedTemplate->RelativeRotation;
 							FVector OldRelativeScale3D = SelectedTemplate->RelativeScale3D;
 
-							USceneComponent* ParentSceneComp = SceneComp->GetAttachParent();
-							if( ParentSceneComp )
-							{
-								const FTransform ParentToWorldSpace = ParentSceneComp->GetSocketTransform(SceneComp->AttachSocketName);
+							// Adjust the deltas as necessary
+							FComponentEditorUtils::AdjustComponentDelta(SceneComp, Drag, Rot);
 
-								if(!SceneComp->bAbsoluteLocation)
-								{
-									Drag = ParentToWorldSpace.Inverse().TransformVector(Drag);
-								}
-								
-								if(!SceneComp->bAbsoluteRotation)
-								{
-									Rot = (ParentToWorldSpace.Inverse().GetRotation() * Rot.Quaternion() * ParentToWorldSpace.GetRotation()).Rotator();
-								}
-							}
-							
 							FComponentEditorUtils::FTransformData OldDefaultTransform(*SelectedTemplate);
 
 							TSharedPtr<ISCSEditorCustomization> Customization = BlueprintEditor->CustomizeSCSEditor(SceneComp);
diff --git a/Engine/Source/Editor/Kismet/Private/SSCSEditor.cpp b/Engine/Source/Editor/Kismet/Private/SSCSEditor.cpp
--- a/Engine/Source/Editor/Kismet/Private/SSCSEditor.cpp
+++ b/Engine/Source/Editor/Kismet/Private/SSCSEditor.cpp
@@ -692,7 +692,7 @@ void FSCSEditorTreeNode::AddChild(FSCSEditorTreeNodePtrType InChildNodePtr)
 		check(ParentInstance != nullptr);
 
 		// Handle attachment at the instance level
-		ChildInstance->AttachTo(ParentInstance);
+		ChildInstance->AttachTo(ParentInstance, NAME_None, EAttachLocation::KeepWorldPosition);
 	}
 }
 
@@ -845,7 +845,7 @@ void FSCSEditorTreeNode::RemoveChild(FSCSEditorTreeNodePtrType InChildNodePtr)
 		check(ChildInstance != nullptr);
 
 		// Handle detachment at the instance level
-		ChildInstance->DetachFromParent();
+		ChildInstance->DetachFromParent(true);
 	}
 }
 
@@ -2455,6 +2455,7 @@ void SSCSEditor::Construct( const FArguments& InArgs )
 		.OnSelectionChanged(this, &SSCSEditor::OnTreeSelectionChanged)
 		.OnContextMenuOpening(this, &SSCSEditor::CreateContextMenu)
 		.OnItemScrolledIntoView(this, &SSCSEditor::OnItemScrolledIntoView)
+		.ClearSelectionOnClick(InArgs._EditorMode == EEditorMode::BlueprintSCS ? true : false)
 		.ItemHeight(24)
 		.HeaderRow
 		(
diff --git a/Engine/Source/Editor/UnrealEd/Private/Kismet2/ComponentEditorUtils.cpp b/Engine/Source/Editor/UnrealEd/Private/Kismet2/ComponentEditorUtils.cpp
--- a/Engine/Source/Editor/UnrealEd/Private/Kismet2/ComponentEditorUtils.cpp
+++ b/Engine/Source/Editor/UnrealEd/Private/Kismet2/ComponentEditorUtils.cpp
@@ -141,6 +141,25 @@ FString FComponentEditorUtils::GenerateValidVariableName(TSubclassOf<UActorCompo
 	return ComponentInstanceName;
 }
 
+void FComponentEditorUtils::AdjustComponentDelta(USceneComponent* Component, FVector& Drag, FRotator& Rotation)
+{
+	USceneComponent* ParentSceneComp = Component->GetAttachParent();
+	if (ParentSceneComp)
+	{
+		const FTransform ParentToWorldSpace = ParentSceneComp->GetSocketTransform(Component->AttachSocketName);
+
+		if (!Component->bAbsoluteLocation)
+		{
+			Drag = ParentToWorldSpace.Inverse().TransformVector(Drag);
+		}
+
+		if (!Component->bAbsoluteRotation)
+		{
+			Rotation = ( ParentToWorldSpace.Inverse().GetRotation() * Rotation.Quaternion() * ParentToWorldSpace.GetRotation() ).Rotator();
+		}
+	}
+}
+
 void FComponentEditorUtils::PropagateTransformPropertyChange(
 	class USceneComponent* InSceneComponentTemplate,
 	const FTransformData& OldDefaultTransform,
diff --git a/Engine/Source/Editor/UnrealEd/Private/LevelEditorViewport.cpp b/Engine/Source/Editor/UnrealEd/Private/LevelEditorViewport.cpp
--- a/Engine/Source/Editor/UnrealEd/Private/LevelEditorViewport.cpp
+++ b/Engine/Source/Editor/UnrealEd/Private/LevelEditorViewport.cpp
@@ -48,6 +48,7 @@
 #include "Editor/ActorPositioning.h"
 #include "NotificationManager.h"
 #include "SNotificationList.h"
+#include "ComponentEditorUtils.h"
 
 DEFINE_LOG_CATEGORY(LogEditorViewport);
 
@@ -3045,11 +3046,15 @@ void FLevelEditorViewportClient::ApplyDeltaToComponent(USceneComponent* InCompon
 		ModifiedDeltaScale = FVector::ZeroVector;
 	}
 
+	FVector AdjustedDrag = InDeltaDrag;
+	FRotator AdjustedRot = InDeltaRot;
+	FComponentEditorUtils::AdjustComponentDelta(InComponent, AdjustedDrag, AdjustedRot);
+
 	GEditor->ApplyDeltaToComponent(
 		InComponent,
 		true,
-		&InDeltaDrag,
-		&InDeltaRot,
+		&AdjustedDrag,
+		&AdjustedRot,
 		&ModifiedDeltaScale,
 		InComponent->RelativeLocation);
 }
diff --git a/Engine/Source/Editor/UnrealEd/Public/Kismet2/ComponentEditorUtils.h b/Engine/Source/Editor/UnrealEd/Public/Kismet2/ComponentEditorUtils.h
--- a/Engine/Source/Editor/UnrealEd/Public/Kismet2/ComponentEditorUtils.h
+++ b/Engine/Source/Editor/UnrealEd/Public/Kismet2/ComponentEditorUtils.h
@@ -32,6 +32,9 @@ public:
 		FVector Scale;
 	};
 
+	/** Potentially transforms the delta to be applied to a component into the appropriate space */
+	static void AdjustComponentDelta(USceneComponent* Component, FVector& Drag, FRotator& Rotation);
+
 	// Given a template, propagates a default transform change to all instances of the template
 	static void PropagateTransformPropertyChange(class USceneComponent* InSceneComponentTemplate, const FTransformData& OldDefaultTransform, const FTransformData& NewDefaultTransform, TSet<class USceneComponent*>& UpdatedComponents);
 
