[INTEGRATE] Change 2415251 by [REDACTED_PERSON][REDACTED_HANDLE] on 2015/01/22 08:59:29

	Fixed a few styling regressions with Blueprint tool-tips

diff --git a/Engine/Source/Editor/GraphEditor/Private/KismetNodes/SGraphNodeK2Base.cpp b/Engine/Source/Editor/GraphEditor/Private/KismetNodes/SGraphNodeK2Base.cpp
--- a/Engine/Source/Editor/GraphEditor/Private/KismetNodes/SGraphNodeK2Base.cpp
+++ b/Engine/Source/Editor/GraphEditor/Private/KismetNodes/SGraphNodeK2Base.cpp
@@ -171,6 +171,11 @@ TSharedPtr<SToolTip> SGraphNodeK2Base::GetComplexTooltip()
 	SAssignNew(NodeToolTip, SToolTip)
 		.Visibility_Static(&LocalUtils::IsToolTipVisible, ThisRef)
 		.IsInteractive_Static(&LocalUtils::IsInteractive)
+
+		// Emulate text-only tool-tip styling that SToolTip uses when no custom content is supplied.  We want node tool-tips to 
+		// be styled just like text-only tool-tips
+		.BorderImage( FCoreStyle::Get().GetBrush("ToolTip.BrightBackground") )
+		.TextMargin(FMargin(11.0f))
 	[
 		SAssignNew(VerticalBoxWidget, SVerticalBox)
 		// heading container
@@ -182,8 +187,7 @@ TSharedPtr<SToolTip> SGraphNodeK2Base::GetComplexTooltip()
 				.AutoHeight()
 			[
 				SNew(STextBlock)
-					.Font(FEditorStyle::GetFontStyle(TEXT("Kismet.Tooltip.SubtextFont")))
-					.ColorAndOpacity(FSlateColor::UseSubduedForeground())
+					.TextStyle( FEditorStyle::Get(), "Documentation.SDocumentationTooltipSubdued")
 					.Text(this, &SGraphNodeK2Base::GetToolTipHeading)
 			]
 			+SVerticalBox::Slot()
diff --git a/Engine/Source/Editor/Kismet/Private/SBlueprintPalette.cpp b/Engine/Source/Editor/Kismet/Private/SBlueprintPalette.cpp
--- a/Engine/Source/Editor/Kismet/Private/SBlueprintPalette.cpp
+++ b/Engine/Source/Editor/Kismet/Private/SBlueprintPalette.cpp
@@ -400,6 +400,11 @@ static TSharedRef<IToolTip> ConstructToolTipWithActionPath(TSharedPtr<FEdGraphSc
 			.SetColorAndOpacity(FLinearColor(0.4f, 0.4f, 0.4f));
 
 		NewToolTip = SNew(SToolTip)
+
+		// Emulate text-only tool-tip styling that SToolTip uses when no custom content is supplied.  We want node tool-tips to 
+		// be styled just like text-only tool-tips
+		.BorderImage( FCoreStyle::Get().GetBrush("ToolTip.BrightBackground") )
+		.TextMargin(FMargin(11.0f))
 		[
 			SNew(SVerticalBox)
 			+ SVerticalBox::Slot()
@@ -411,7 +416,7 @@ static TSharedRef<IToolTip> ConstructToolTipWithActionPath(TSharedPtr<FEdGraphSc
 			.HAlign(EHorizontalAlignment::HAlign_Right)
 			[
 				SNew(STextBlock)
-				.ColorAndOpacity(FSlateColor::UseSubduedForeground())
+				.TextStyle( FEditorStyle::Get(), "Documentation.SDocumentationTooltip")
 				.Text(FText::FromString(ActionItem.ToString()))
 				//.TextStyle(&PathStyle)
 			]
@@ -1505,6 +1510,11 @@ TSharedPtr<SToolTip> SBlueprintPaletteItem::ConstructToolTipWidget() const
 
 			return 
 				SNew(SToolTip)
+
+				// Emulate text-only tool-tip styling that SToolTip uses when no custom content is supplied.  We want node tool-tips to 
+				// be styled just like text-only tool-tips
+				.BorderImage( FCoreStyle::Get().GetBrush("ToolTip.BrightBackground") )
+				.TextMargin(FMargin(11.0f))
 				[
 					SNew(SVerticalBox)
 					+SVerticalBox::Slot()
@@ -1517,9 +1527,8 @@ TSharedPtr<SToolTip> SBlueprintPaletteItem::ConstructToolTipWidget() const
 					[
 
 						SNew( STextBlock )
-						.ColorAndOpacity( FSlateColor::UseSubduedForeground() )
 						.Text( LOCTEXT( "NativeNodeName", "hold (Alt) for native node name" ) )
-						.TextStyle( &FEditorStyle::GetWidgetStyle<FTextBlockStyle>(TEXT("Documentation.SDocumentationTooltip")) )
+						.TextStyle( FEditorStyle::Get(), "Documentation.SDocumentationTooltipSubdued")
 						.Visibility_Static(&Local::GetNativeNodeNameVisibility)
 					]
 				];
