[INTEGRATE] Change 2415569 by [REDACTED_PERSON][REDACTED_HANDLE] on 2015/01/22 14:01:48

	Attach a newly created Instance SceneComponent to the root

diff --git a/Engine/Source/Editor/Kismet/Private/SSCSEditor.cpp b/Engine/Source/Editor/Kismet/Private/SSCSEditor.cpp
--- a/Engine/Source/Editor/Kismet/Private/SSCSEditor.cpp
+++ b/Engine/Source/Editor/Kismet/Private/SSCSEditor.cpp
@@ -3596,6 +3596,11 @@ UActorComponent* SSCSEditor::AddNewComponent( UClass* NewComponentClass, UObject
 		UActorComponent* NewComponentInstance = ConstructObject<UActorComponent>(NewComponentClass, ActorInstance, NewComponentName, RF_Transactional);
 		check(NewComponentInstance != nullptr);
 
+		if (USceneComponent* NewSceneComponent = Cast<USceneComponent>(NewComponentInstance))
+		{
+			NewSceneComponent->AttachTo(ActorInstance->GetRootComponent());
+		}
+
 		// Add to SerializedComponents array so it gets saved
 		ActorInstance->InstanceComponents.Add(NewComponentInstance);
 
