[INTEGRATE] Change 2415745 by [REDACTED_PERSON][REDACTED_HANDLE] on 2015/01/22 15:26:02

	UE-7951 - BP - Fixing variables becoming deselected after compile.
	UE-7966 - BP - Fixing variables becoming deselected after changing a property.
	UE-8059 - BP - Level Blueprint Default layout now matches the new unified layout.
	UE-8056 - BP - Close Other Tabs, now effects the Viewport Tab.

diff --git a/Engine/Source/Editor/BlueprintGraph/Private/BlueprintEditorSettings.cpp b/Engine/Source/Editor/BlueprintGraph/Private/BlueprintEditorSettings.cpp
--- a/Engine/Source/Editor/BlueprintGraph/Private/BlueprintEditorSettings.cpp
+++ b/Engine/Source/Editor/BlueprintGraph/Private/BlueprintEditorSettings.cpp
@@ -18,8 +18,9 @@ UBlueprintEditorSettings::UBlueprintEditorSettings(const FObjectInitializer& Obj
 	, bFavorPureCastNodes(false)
 	, bAutoCastObjectConnections(false)
 	, bUseLegacyMenuingSystem(false)
+	, bShowViewportOnSimulate(false)
 	, bShowInheritedVariables(false)
-	, bShowEmptySections(false)
+	, bShowEmptySections(true)
 	// Compiler Settings
 	, SaveOnCompile(SoC_Never)
 	, bJumpToNodeErrors(false)
diff --git a/Engine/Source/Editor/BlueprintGraph/Public/BlueprintEditorSettings.h b/Engine/Source/Editor/BlueprintGraph/Public/BlueprintEditorSettings.h
--- a/Engine/Source/Editor/BlueprintGraph/Public/BlueprintEditorSettings.h
+++ b/Engine/Source/Editor/BlueprintGraph/Public/BlueprintEditorSettings.h
@@ -59,6 +59,10 @@ public:
 	UPROPERTY(EditAnywhere, AdvancedDisplay, config, Category=Workflow)
 	bool bUseLegacyMenuingSystem;
 
+	/** If true will show the viewport tab when simulate is clicked. */
+	UPROPERTY(EditAnywhere, config, Category=Workflow)
+	bool bShowViewportOnSimulate;
+
 	/** If set we'll show the inherited variables in the My Blueprint view. */
 	UPROPERTY(config)
 	bool bShowInheritedVariables;
diff --git a/Engine/Source/Editor/Kismet/Private/BlueprintEditor.cpp b/Engine/Source/Editor/Kismet/Private/BlueprintEditor.cpp
--- a/Engine/Source/Editor/Kismet/Private/BlueprintEditor.cpp
+++ b/Engine/Source/Editor/Kismet/Private/BlueprintEditor.cpp
@@ -387,6 +387,12 @@ static UEdGraphNode* BlueprintEditorImpl::FindNodeWithError(FCompilerResultsLog
 }
 
 
+FName FBlueprintEditor::SelectionState_MyBlueprint(TEXT("MyBlueprint"));
+FName FBlueprintEditor::SelectionState_Components(TEXT("Components"));
+FName FBlueprintEditor::SelectionState_Graph(TEXT("Graph"));
+FName FBlueprintEditor::SelectionState_ClassSettings(TEXT("ClassSettings"));
+FName FBlueprintEditor::SelectionState_ClassDefaults(TEXT("ClassDefaults"));
+
 
 bool FBlueprintEditor::IsASubGraph( const UEdGraph* GraphPtr )
 {
@@ -602,13 +608,13 @@ void FBlueprintEditor::AnalyticsTrackCompileEvent( UBlueprint* Blueprint, int32
 void FBlueprintEditor::RefreshEditors(ERefreshBlueprintEditorReason::Type Reason)
 {
 	if( Reason != ERefreshBlueprintEditorReason::BlueprintCompiled )
-	{
-		DocumentManager->CleanInvalidTabs();
+{
+	DocumentManager->CleanInvalidTabs();
 
-		DocumentManager->RefreshAllTabs();
+	DocumentManager->RefreshAllTabs();
 
-		// The workflow manager only tracks document tabs.
-		FocusInspectorOnGraphSelection(GetSelectedNodes(), true);
+	// The workflow manager only tracks document tabs.
+	FocusInspectorOnGraphSelection(GetSelectedNodes(), true);
 	}
 
 	if (MyBlueprintWidget.IsValid())
@@ -635,23 +641,45 @@ void FBlueprintEditor::RefreshEditors(ERefreshBlueprintEditorReason::Type Reason
 	BroadcastRefresh();
 }
 
-void FBlueprintEditor::ClearSelectionInAllEditors()
+void FBlueprintEditor::SetUISelectionState(FName SelectionOwner)
+{
+	if ( SelectionOwner != CurrentUISelection )
 	{
-		TArray< TSharedPtr<SDockTab> > GraphEditorTabs;
-		DocumentManager->FindAllTabsForFactory(GraphEditorTabFactoryPtr, /*out*/ GraphEditorTabs);
+		ClearSelectionStateFor(CurrentUISelection);
 
-		for ( auto GraphEditorTabIt = GraphEditorTabs.CreateIterator(); GraphEditorTabIt; ++GraphEditorTabIt )
-		{
-			TSharedRef<SGraphEditor> Editor = StaticCastSharedRef<SGraphEditor>(( *GraphEditorTabIt )->GetContent());
+		CurrentUISelection = SelectionOwner;
+	}
+}
 
-			Editor->ClearSelectionSet();
-		}
+void FBlueprintEditor::ClearSelectionStateFor(FName SelectionOwner)
+{
+	if ( SelectionOwner == SelectionState_Graph )
+{
+	TArray< TSharedPtr<SDockTab> > GraphEditorTabs;
+	DocumentManager->FindAllTabsForFactory(GraphEditorTabFactoryPtr, /*out*/ GraphEditorTabs);
 
-		if ( SCSEditor.IsValid() )
+	for (auto GraphEditorTabIt = GraphEditorTabs.CreateIterator(); GraphEditorTabIt; ++GraphEditorTabIt)
+	{
+		TSharedRef<SGraphEditor> Editor = StaticCastSharedRef<SGraphEditor>((*GraphEditorTabIt)->GetContent());
+
+		Editor->ClearSelectionSet();
+	}
+	}
+	else if ( SelectionOwner == SelectionState_Components )
+	{
+	if(SCSEditor.IsValid())
+	{
+		SCSEditor->ClearSelection();
+	}
+}
+	else if ( SelectionOwner == SelectionState_MyBlueprint )
+	{
+		if ( MyBlueprintWidget.IsValid() )
 		{
-			SCSEditor->ClearSelection();
+			MyBlueprintWidget->ClearGraphActionMenuSelection();
 		}
 	}
+}
 
 void FBlueprintEditor::SummonSearchUI(bool bSetFindWithinBlueprint, FString NewSearchTerms, bool bSelectFirstResult)
 {
@@ -694,7 +722,7 @@ void FBlueprintEditor::OnRootSelected(AActor * DefaultActor)
 	if (Inspector.IsValid())
 	{	
 		// Clear the my blueprints selection
-		MyBlueprintWidget->ClearGraphActionMenuSelection();
+		SetUISelectionState(FBlueprintEditor::SelectionState_Components);
 
 		// Build an object list
 		TArray<UObject*> InspectorObjects;
@@ -736,9 +764,9 @@ void FBlueprintEditor::OnSelectionUpdated(const TArray<FSCSEditorTreeNodePtrType
 	{
 		// Clear the my blueprints selection
 		bool bSingleLayoutBPEditor = GetDefault<UEditorExperimentalSettings>()->bUnifiedBlueprintEditor;
-		if (bSingleLayoutBPEditor)
+		if ( bSingleLayoutBPEditor && SelectedNodes.Num() > 0 )
 		{
-			MyBlueprintWidget->ClearGraphActionMenuSelection();
+			SetUISelectionState(FBlueprintEditor::SelectionState_Components);
 		}
 
 		// Convert the selection set to an array of UObject* pointers
@@ -774,6 +802,15 @@ void FBlueprintEditor::OnSelectionUpdated(const TArray<FSCSEditorTreeNodePtrType
 	}
 }
 
+void FBlueprintEditor::OnComponentDoubleClicked(TSharedPtr<class FSCSEditorTreeNode> Node)
+{
+	TSharedPtr<SDockTab> OwnerTab = Inspector->GetOwnerTab();
+	if ( OwnerTab.IsValid() )
+	{
+		GetTabManager()->InvokeTab(FBlueprintEditorTabs::SCSViewportID);
+	}
+}
+
 /** Create new tab for the supplied graph - don't call this directly.*/
 TSharedRef<SGraphEditor> FBlueprintEditor::CreateGraphEditorWidget(TSharedRef<FTabInfo> InTabInfo, UEdGraph* InGraph)
 {
@@ -1193,7 +1230,7 @@ FBlueprintEditor::FBlueprintEditor()
 	, bBlueprintModifiedOnOpen (false)
 	, PinVisibility(SGraphEditor::Pin_Show)
 	, bIsActionMenuContextSensitive(true)
-	, CurrentUISelection(FBlueprintEditor::NoSelection)
+	, CurrentUISelection(NAME_None)
 	, bEditorMarkedAsClosed(false)
 	, bCodeBasedProject(false)
 	, HasOpenActionMenu(nullptr)
@@ -1236,14 +1273,7 @@ void FBlueprintEditor::EnsureBlueprintIsUpToDate(UBlueprint* BlueprintObj)
 			BlueprintObj->SimpleConstructionScript->SetFlags(RF_Transactional);
 
 			// Recreate (or create) any widgets that depend on the SCS
-			SCSEditor = SAssignNew(SCSEditor, SSCSEditor)
-				.ActorContext(this, &FBlueprintEditor::GetSCSEditorActorContext)
-				.PreviewActor(this, &FBlueprintEditor::GetPreviewActor)
-				.AllowEditing(this, &FBlueprintEditor::InEditingMode)
-				.OnRootSelected(this, &FBlueprintEditor::OnRootSelected)
-				.OnSelectionUpdated(this, &FBlueprintEditor::OnSelectionUpdated);
-			SCSViewport = SAssignNew(SCSViewport, SSCSEditorViewport)
-				.BlueprintEditor(SharedThis(this));
+			CreateSCSEditors();
 		}
 
 		// If we should have a UCS but don't yet, make it
@@ -1518,7 +1548,7 @@ void FBlueprintEditor::InitBlueprintEditor(const EToolkitMode::Type Mode, const
 	PostLayoutBlueprintEditorInitialization();
 
 	// Find and set any instances of this blueprint type if any exists and we are not already editing one
-	FBlueprintEditorUtils::FindAndSetDebuggableBlueprintInstances();
+	FBlueprintEditorUtils::FindAndSetDebuggableBlueprintInstances();	
 }
 
 void FBlueprintEditor::RegisterApplicationModes(const TArray<UBlueprint*>& InBlueprints, bool bShouldOpenInDefaultsMode, bool bNewlyCreated/* = false*/)
@@ -1547,11 +1577,11 @@ void FBlueprintEditor::RegisterApplicationModes(const TArray<UBlueprint*>& InBlu
 		}
 		else if ( SingleBP->BlueprintType == BPTYPE_FunctionLibrary )
 		{
-				AddApplicationMode(
-					FBlueprintEditorApplicationModes::StandardBlueprintEditorMode,
-					MakeShareable(new FBlueprintEditorApplicationMode(SharedThis(this), FBlueprintEditorApplicationModes::StandardBlueprintEditorMode, FBlueprintEditorApplicationModes::GetLocalizedMode)));
-				SetCurrentMode(FBlueprintEditorApplicationModes::StandardBlueprintEditorMode);
-			}
+			AddApplicationMode(
+				FBlueprintEditorApplicationModes::StandardBlueprintEditorMode,
+				MakeShareable(new FBlueprintEditorApplicationMode(SharedThis(this), FBlueprintEditorApplicationModes::StandardBlueprintEditorMode, FBlueprintEditorApplicationModes::GetLocalizedMode)));
+			SetCurrentMode(FBlueprintEditorApplicationModes::StandardBlueprintEditorMode);
+		}
 		else
 		{
 			if ( GetDefault<UEditorExperimentalSettings>()->bUnifiedBlueprintEditor )
@@ -1943,25 +1973,15 @@ FBlueprintEditor::~FBlueprintEditor()
 
 void FBlueprintEditor::FocusInspectorOnGraphSelection(const FGraphPanelSelectionSet& NewSelection, bool bForceRefresh)
 {
-	if ( GetDefault<UEditorExperimentalSettings>()->bUnifiedBlueprintEditor )
-	{
-		//TODO UNIFIEDBP Do something smart here, like keep selection if it comes from the my blueprint view, or the components view, but not if it's from a graph view.
-	}
-	else
-	{
-		if ( NewSelection.Array().Num() )
+	if ( NewSelection.Array().Num() > 0 )
 		{
+		SetUISelectionState(FBlueprintEditor::SelectionState_Graph);
+
 			SKismetInspector::FShowDetailsOptions ShowDetailsOptions;
 			ShowDetailsOptions.bForceRefresh = bForceRefresh;
 
 			Inspector->ShowDetailsForObjects(NewSelection.Array(), ShowDetailsOptions);
 		}
-		else
-		{
-			// Clear the Inspector if nothing is selected, the MyBlueprints window will set the selection if it means to, but we need to make sure invalid nodes aren't still appearing
-			//Inspector->ShowDetailsForSingleObject(nullptr);
-		}
-	}
 }
 
 void FBlueprintEditor::CreateDefaultTabContents(const TArray<UBlueprint*>& InBlueprints)
@@ -2031,16 +2051,23 @@ void FBlueprintEditor::CreateDefaultTabContents(const TArray<UBlueprint*>& InBlu
 		InBlueprint->ParentClass->IsChildOf(AActor::StaticClass()) && 
 		InBlueprint->SimpleConstructionScript )
 	{
+		CreateSCSEditors();
+	}
+}
+
+void FBlueprintEditor::CreateSCSEditors()
+{
 		SCSEditor = SAssignNew(SCSEditor, SSCSEditor)
 			.ActorContext(this, &FBlueprintEditor::GetSCSEditorActorContext)
 			.PreviewActor(this, &FBlueprintEditor::GetPreviewActor)
 			.AllowEditing(this, &FBlueprintEditor::InEditingMode)
 			.OnRootSelected(this, &FBlueprintEditor::OnRootSelected)
-			.OnSelectionUpdated(this, &FBlueprintEditor::OnSelectionUpdated);
+		.OnSelectionUpdated(this, &FBlueprintEditor::OnSelectionUpdated)
+		.OnItemDoubleClicked(this, &FBlueprintEditor::OnComponentDoubleClicked);
+
 		SCSViewport = SAssignNew(SCSViewport, SSCSEditorViewport)
 			.BlueprintEditor(SharedThis(this));
 	}
-}
 
 void FBlueprintEditor::OnLogTokenClicked(const TSharedRef<IMessageToken>& Token)
 {
@@ -2452,18 +2479,7 @@ void FBlueprintEditor::EditGlobalOptions_Clicked()
 	UBlueprint* Blueprint = GetBlueprintObj();
 	if ( Blueprint != nullptr )
 	{
-		if ( CurrentUISelection == FBlueprintEditor::GraphPanel )
-		{
-			ClearSelectionInAllEditors();
-		}
-		if ( CurrentUISelection == FBlueprintEditor::MyBlueprint )
-		{
-			if ( MyBlueprintWidget.IsValid() )
-			{
-				MyBlueprintWidget->ClearGraphActionMenuSelection();
-			}
-		}
-		CurrentUISelection = FBlueprintEditor::ClassSettings;
+		SetUISelectionState(FBlueprintEditor::SelectionState_ClassSettings);
 
 		// Show details for the Blueprint instance we're editing
 		Inspector->ShowDetailsForSingleObject(Blueprint);
@@ -2482,19 +2498,7 @@ void FBlueprintEditor::EditClassDefaults_Clicked()
 	UBlueprint* Blueprint = GetBlueprintObj();
 	if ( Blueprint != nullptr )
 	{
-		if ( CurrentUISelection == FBlueprintEditor::GraphPanel )
-		{
-			ClearSelectionInAllEditors();
-		}
-		else if ( CurrentUISelection == FBlueprintEditor::MyBlueprint )
-		{
-			if ( MyBlueprintWidget.IsValid() )
-			{
-				MyBlueprintWidget->ClearGraphActionMenuSelection();
-			}
-		}
-
-		CurrentUISelection = FBlueprintEditor::ClassDefaults;
+		SetUISelectionState(FBlueprintEditor::SelectionState_ClassDefaults);
 
 		const bool bForceRefresh = true;
 
@@ -2502,7 +2506,16 @@ void FBlueprintEditor::EditClassDefaults_Clicked()
 		{
 			if ( GetBlueprintObj()->GeneratedClass != NULL )
 			{
-				Inspector->ShowDetailsForSingleObject(GetBlueprintObj()->GeneratedClass->GetDefaultObject(), SKismetInspector::FShowDetailsOptions(DefaultEditString(), bForceRefresh));
+				UObject* DefaultObject = GetBlueprintObj()->GeneratedClass->GetDefaultObject();
+
+				// Update the details panel
+				FString Title;
+				DefaultObject->GetName(Title);
+				SKismetInspector::FShowDetailsOptions Options(FText::FromString(Title), true);
+				Options.bShowComponents = false;
+
+				Inspector->ShowDetailsForSingleObject(DefaultObject, Options);
+
 				TSharedPtr<SDockTab> OwnerTab = Inspector->GetOwnerTab();
 				if ( OwnerTab.IsValid() )
 				{
@@ -2610,7 +2623,7 @@ void FBlueprintEditor::PostUndo(bool bSuccess)
 	// Clear selection, to avoid holding refs to nodes that go away
 	if (bSuccess && GetBlueprintObj())
 	{
-		ClearSelectionInAllEditors();
+		SetUISelectionState(NAME_None);
 
 		// Will cause a call to RefreshEditors()
 		FBlueprintEditorUtils::MarkBlueprintAsStructurallyModified(GetBlueprintObj());
@@ -2652,7 +2665,7 @@ void FBlueprintEditor::RedoGraphAction()
 	// Clear selection, to avoid holding refs to nodes that go away
 	if (GetBlueprintObj())
 	{
-		ClearSelectionInAllEditors();
+		SetUISelectionState(NAME_None);
 
 		GEditor->RedoTransaction();
 
@@ -2669,10 +2682,9 @@ bool FBlueprintEditor::CanRedoGraphAction() const
 
 void FBlueprintEditor::OnActiveTabChanged( TSharedPtr<SDockTab> PreviouslyActive, TSharedPtr<SDockTab> NewlyActivated )
 {
-	if (!NewlyActivated.IsValid())
+	if (NewlyActivated.IsValid() == false)
 	{
-		TArray<UObject*> ObjArray;
-		Inspector->ShowDetailsForObjects(ObjArray);
+		SetUISelectionState(NAME_None);
 		//UE_LOG(LogBlueprint, Log, TEXT("OnActiveTabChanged: NONE"));
 	}
 	else
@@ -2795,15 +2807,11 @@ void FBlueprintEditor::OnGraphActionMenuClosed(bool bActionExecuted, bool bConte
 
 void FBlueprintEditor::OnSelectedNodesChanged(const FGraphPanelSelectionSet& NewSelection)
 {
-	if (CurrentUISelection == FBlueprintEditor::MyBlueprint)
+	if ( NewSelection.Num() > 0 )
 	{
-		// clear MyBlueprint selection
-		if (MyBlueprintWidget.IsValid())
-	{
-			MyBlueprintWidget->ClearGraphActionMenuSelection();
+		SetUISelectionState(FBlueprintEditor::SelectionState_Graph);
 		}
-	}
-	CurrentUISelection = NewSelection.Num() > 0 ? FBlueprintEditor::GraphPanel : FBlueprintEditor::NoSelection;
+
 	Inspector->ShowDetailsForObjects(NewSelection.Array());
 }
 
@@ -4645,7 +4653,7 @@ void FBlueprintEditor::DeleteSelectedNodes()
 
 	const FGraphPanelSelectionSet SelectedNodes = GetSelectedNodes();
 
-	ClearSelectionInAllEditors();
+	SetUISelectionState(NAME_None);
 
 	for (FGraphPanelSelectionSet::TConstIterator NodeIt( SelectedNodes ); NodeIt; ++NodeIt)
 	{
@@ -5018,7 +5026,7 @@ void FBlueprintEditor::PasteNodesHere(class UEdGraph* DestinationGraph, const FV
 	DestinationGraph->Modify();
 
 	// Clear the selection set (newly pasted stuff will be selected)
-	ClearSelectionInAllEditors();
+	SetUISelectionState(NAME_None);
 
 	// Grab the text to paste from the clipboard.
 	FString TextToImport;
diff --git a/Engine/Source/Editor/Kismet/Private/BlueprintEditorModes.cpp b/Engine/Source/Editor/Kismet/Private/BlueprintEditorModes.cpp
--- a/Engine/Source/Editor/Kismet/Private/BlueprintEditorModes.cpp
+++ b/Engine/Source/Editor/Kismet/Private/BlueprintEditorModes.cpp
@@ -549,29 +549,75 @@ FBlueprintEditorUnifiedMode::FBlueprintEditorUnifiedMode(TSharedPtr<class FBluep
 
 	if ( bRegisterViewport )
 	{
-		TabLayout = FTabManager::NewLayout( "Standalone_BlueprintUnifiedEditor_Layout_v1_1" )
+		TabLayout = FTabManager::NewLayout( "Blueprints_Unified_Components_v1" )
 		->AddArea
 		(
-			FTabManager::NewPrimaryArea() ->SetOrientation(Orient_Vertical)
+			FTabManager::NewPrimaryArea()->SetOrientation(Orient_Vertical)
 			->Split
 			(
-				FTabManager::NewSplitter() ->SetOrientation(Orient_Horizontal)
+				FTabManager::NewSplitter()->SetOrientation(Orient_Horizontal)
+				->Split
+				(
+					FTabManager::NewSplitter()->SetOrientation( Orient_Vertical )
+					->SetSizeCoefficient(0.70f)
+					->Split
+					(
+						FTabManager::NewStack()
+						->SetSizeCoefficient(0.186721f)
+						->SetHideTabWell(true)
+						->AddTab(InBlueprintEditor->GetToolbarTabId(), ETabState::OpenedTab)
+					)
+					->Split
+					(
+						FTabManager::NewStack()
+						->SetSizeCoefficient( 0.80f )
+						->AddTab(FBlueprintEditorTabs::SCSViewportID, ETabState::OpenedTab)
+						->AddTab( "Document", ETabState::ClosedTab )
+					)
+					->Split
+					(
+						FTabManager::NewStack()
+						->SetSizeCoefficient( 0.20f )
+						->AddTab( FBlueprintEditorTabs::CompilerResultsID, ETabState::ClosedTab )
+						->AddTab( FBlueprintEditorTabs::FindResultsID, ETabState::ClosedTab )
+					)
+				)
 				->Split
 				(
 					FTabManager::NewSplitter() ->SetOrientation( Orient_Vertical )
-					->SetSizeCoefficient(0.15f)
+					->SetSizeCoefficient(0.20f)
 					->Split
 					(
 						FTabManager::NewStack()
-						->SetSizeCoefficient(0.25f)
+						->SetSizeCoefficient(0.20f)
 						->AddTab( FBlueprintEditorTabs::ConstructionScriptEditorID, ETabState::OpenedTab )
 					)
 					->Split
 					(
-						FTabManager::NewStack()->SetSizeCoefficient(0.75f)
+						FTabManager::NewStack()
+						->SetSizeCoefficient(0.30f)
 						->AddTab( FBlueprintEditorTabs::MyBlueprintID, ETabState::OpenedTab )
 					)
+					->Split
+					(
+						FTabManager::NewStack()
+						->SetSizeCoefficient(0.50f)
+						->AddTab( FBlueprintEditorTabs::DetailsID, ETabState::OpenedTab )
+						//->AddTab( FBlueprintEditorTabs::DefaultEditorID, ETabState::ClosedTab )
+					)
 				)
+			)
+		);
+	}
+	else
+	{
+		TabLayout = FTabManager::NewLayout( "Blueprints_Unified_v1" )
+		->AddArea
+		(
+			FTabManager::NewPrimaryArea() ->SetOrientation(Orient_Vertical)
+			->Split
+			(
+				FTabManager::NewSplitter() ->SetOrientation(Orient_Horizontal)
 				->Split
 				(
 					FTabManager::NewSplitter()->SetOrientation( Orient_Vertical )
@@ -587,7 +633,6 @@ FBlueprintEditorUnifiedMode::FBlueprintEditorUnifiedMode(TSharedPtr<class FBluep
 					(
 						FTabManager::NewStack()
 						->SetSizeCoefficient( 0.80f )
-						->AddTab(FBlueprintEditorTabs::SCSViewportID, ETabState::OpenedTab)
 						->AddTab( "Document", ETabState::ClosedTab )
 					)
 					->Split
@@ -601,20 +646,24 @@ FBlueprintEditorUnifiedMode::FBlueprintEditorUnifiedMode(TSharedPtr<class FBluep
 				->Split
 				(
 					FTabManager::NewSplitter() ->SetOrientation( Orient_Vertical )
-					->SetSizeCoefficient(0.15f)
+					->SetSizeCoefficient(0.20f)
+					->Split
+					(
+						FTabManager::NewStack()
+						->SetSizeCoefficient(0.40f)
+						->AddTab( FBlueprintEditorTabs::MyBlueprintID, ETabState::OpenedTab )
+					)
 					->Split
 					(
 						FTabManager::NewStack()
+						->SetSizeCoefficient(0.60f)
 						->AddTab( FBlueprintEditorTabs::DetailsID, ETabState::OpenedTab )
+						//->AddTab( FBlueprintEditorTabs::DefaultEditorID, ETabState::ClosedTab )
 					)
 				)
 			)
 		);
 	}
-	else
-	{
-		TabLayout = GetDefaltEditorLayout(InBlueprintEditor);
-	}
 	
 	// setup toolbar
 	//@TODO: Keep this in sync with AnimBlueprintMode.cpp
diff --git a/Engine/Source/Editor/Kismet/Private/BlueprintEditorTabFactories.cpp b/Engine/Source/Editor/Kismet/Private/BlueprintEditorTabFactories.cpp
--- a/Engine/Source/Editor/Kismet/Private/BlueprintEditorTabFactories.cpp
+++ b/Engine/Source/Editor/Kismet/Private/BlueprintEditorTabFactories.cpp
@@ -167,6 +167,7 @@ FSCSViewportSummoner::FSCSViewportSummoner(TSharedPtr<class FAssetEditorToolkit>
 	TabIcon = FSlateIcon(FEditorStyle::GetStyleSetName(), "LevelEditor.Tabs.Viewports");
 
 	bIsSingleton = true;
+	TabRole = ETabRole::DocumentTab;
 
 	ViewMenuDescription = LOCTEXT("SCSViewportView", "Viewport");
 	ViewMenuTooltip = LOCTEXT("SCSViewportView_ToolTip", "Show the viewport view");
diff --git a/Engine/Source/Editor/Kismet/Private/SKismetInspector.cpp b/Engine/Source/Editor/Kismet/Private/SKismetInspector.cpp
--- a/Engine/Source/Editor/Kismet/Private/SKismetInspector.cpp
+++ b/Engine/Source/Editor/Kismet/Private/SKismetInspector.cpp
@@ -198,6 +198,7 @@ void SKismetInspector::Construct(const FArguments& InArgs)
 {
 	bShowInspectorPropertyView = true;
 	PublicViewState = ECheckBoxState::Unchecked;
+	bComponenetDetailsCustomizationEnabled = false;
 
 	Kismet2Ptr = InArgs._Kismet2;
 	bShowPublicView = InArgs._ShowPublicViewControl;
@@ -296,6 +297,8 @@ void SKismetInspector::EnableComponentDetailsCustomization(bool bEnable)
 		}
 	};
 
+	bComponenetDetailsCustomizationEnabled = bEnable;
+
 	if (bEnable)
 	{
 		FOnGetDetailCustomizationInstance ActorOverrideDetails = FOnGetDetailCustomizationInstance::CreateStatic(&FActorDetailsOverrideCustomization::MakeInstance);
diff --git a/Engine/Source/Editor/Kismet/Private/SMyBlueprint.cpp b/Engine/Source/Editor/Kismet/Private/SMyBlueprint.cpp
--- a/Engine/Source/Editor/Kismet/Private/SMyBlueprint.cpp
+++ b/Engine/Source/Editor/Kismet/Private/SMyBlueprint.cpp
@@ -1460,16 +1460,17 @@ void SMyBlueprint::OnActionSelected( const TArray< TSharedPtr<FEdGraphSchemaActi
 	UBlueprint* CurrentBlueprint = Blueprint;
 	TSharedPtr<SKismetInspector> CurrentInspector = Inspector.Pin();
 
-	if( BlueprintEditorPtr.IsValid() )
+	TSharedPtr<FBlueprintEditor> BlueprintEditor = BlueprintEditorPtr.Pin();
+
+	if ( BlueprintEditor.IsValid() )
 	{
-		if (BlueprintEditorPtr.Pin()->GetUISelectionState() == FBlueprintEditor::GraphPanel)
+		if ( InAction.IsValid() )
 		{
-			// clear graph panel selection
-			BlueprintEditorPtr.Pin()->ClearSelectionInAllEditors();
+			BlueprintEditor->SetUISelectionState(FBlueprintEditor::SelectionState_MyBlueprint);
 		}
-		BlueprintEditorPtr.Pin()->GetUISelectionState() = InAction.IsValid() ? FBlueprintEditor::MyBlueprint : FBlueprintEditor::NoSelection;
-		CurrentBlueprint = BlueprintEditorPtr.Pin()->GetBlueprintObj();
-		CurrentInspector = BlueprintEditorPtr.Pin()->GetInspector();
+
+		CurrentBlueprint = BlueprintEditor->GetBlueprintObj();
+		CurrentInspector = BlueprintEditor->GetInspector();
 	}
 
 	OnActionSelectedHelper( InAction, Blueprint, CurrentInspector.ToSharedRef() );
diff --git a/Engine/Source/Editor/Kismet/Private/SSCSEditor.cpp b/Engine/Source/Editor/Kismet/Private/SSCSEditor.cpp
--- a/Engine/Source/Editor/Kismet/Private/SSCSEditor.cpp
+++ b/Engine/Source/Editor/Kismet/Private/SSCSEditor.cpp
@@ -2196,11 +2196,11 @@ FText SSCS_RowWidget::GetTooltipText() const
 	{
 		if(NodePtr->IsInherited())
 		{
-			return LOCTEXT("InheritedDefaultSceneRootToolTip", "This is the default scene root component. It cannot be copied, renamed or deleted. It has been inherited from the parent class, so its properties cannot be edited here. New scene components will automatically be attached to it.");
+			return LOCTEXT("InheritedDefaultSceneRootToolTip", "This is the default scene root component. It cannot be copied, renamed or deleted.\nIt has been inherited from the parent class, so its properties cannot be edited here.\nNew scene components will automatically be attached to it.");
 		}
 		else
 		{
-			return LOCTEXT("DefaultSceneRootToolTip", "This is the default scene root component. It cannot be copied, renamed or deleted. Adding a new scene component will automatically replace it as the new root.");
+			return LOCTEXT("DefaultSceneRootToolTip", "This is the default scene root component. It cannot be copied, renamed or deleted.\nAdding a new scene component will automatically replace it as the new root.");
 		}
 	}
 	else
@@ -2373,7 +2373,9 @@ void SSCSEditor::Construct( const FArguments& InArgs )
 	PreviewActor = InArgs._PreviewActor;
 	OnRootSelected = InArgs._OnRootSelected;
 	OnSelectionUpdated = InArgs._OnSelectionUpdated;
+	OnItemDoubleClicked = InArgs._OnItemDoubleClicked;
 	OnHighlightPropertyInDetailsView = InArgs._OnHighlightPropertyInDetailsView;
+	bUpdatingSelection = false;
 
 	CommandList = MakeShareable( new FUICommandList );
 	CommandList->MapAction( FGenericCommands::Get().Cut,
@@ -2455,6 +2457,7 @@ void SSCSEditor::Construct( const FArguments& InArgs )
 		.OnSelectionChanged(this, &SSCSEditor::OnTreeSelectionChanged)
 		.OnContextMenuOpening(this, &SSCSEditor::CreateContextMenu)
 		.OnItemScrolledIntoView(this, &SSCSEditor::OnItemScrolledIntoView)
+		.OnMouseButtonDoubleClick(this, &SSCSEditor::HandleItemDoubleClicked)
 		.ClearSelectionOnClick(InArgs._EditorMode == EEditorMode::BlueprintSCS ? true : false)
 		.ItemHeight(24)
 		.HeaderRow
@@ -3532,8 +3535,13 @@ void SSCSEditor::UpdateTree(bool bRegenerateTreeNodes)
 
 void SSCSEditor::ClearSelection()
 {
-	check(SCSTreeWidget.IsValid());
-	SCSTreeWidget->ClearSelection();
+	if ( bUpdatingSelection == false )
+	{
+		bIsActorSelected = false;
+
+		check(SCSTreeWidget.IsValid());
+		SCSTreeWidget->ClearSelection();
+	}
 }
 
 void SSCSEditor::SaveSCSCurrentState( USimpleConstructionScript* SCSObj )
@@ -4121,8 +4129,12 @@ void SSCSEditor::RemoveComponentNode(FSCSEditorTreeNodePtrType InNodePtr)
 
 void SSCSEditor::UpdateSelectionFromNodes(const TArray<FSCSEditorTreeNodePtrType> &SelectedNodes)
 {
+	bUpdatingSelection = true;
+
 	// Notify that the selection has updated
 	OnSelectionUpdated.ExecuteIfBound(SelectedNodes);
+
+	bUpdatingSelection = false;
 }
 
 void SSCSEditor::RefreshSelectionDetails()
@@ -4451,6 +4463,12 @@ void SSCSEditor::OnItemScrolledIntoView( FSCSEditorTreeNodePtrType InItem, const
 	}
 }
 
+void SSCSEditor::HandleItemDoubleClicked(FSCSEditorTreeNodePtrType InItem)
+{
+	// Notify that the selection has updated
+	OnItemDoubleClicked.ExecuteIfBound(InItem);
+}
+
 void SSCSEditor::OnRenameComponent(bool bTransactional)
 {
 	TArray< FSCSEditorTreeNodePtrType > SelectedItems = SCSTreeWidget->GetSelectedItems();
diff --git a/Engine/Source/Editor/Kismet/Private/SSCSEditorViewport.cpp b/Engine/Source/Editor/Kismet/Private/SSCSEditorViewport.cpp
--- a/Engine/Source/Editor/Kismet/Private/SSCSEditorViewport.cpp
+++ b/Engine/Source/Editor/Kismet/Private/SSCSEditorViewport.cpp
@@ -12,6 +12,7 @@
 #include "EditorViewportCommands.h"
 #include "SEditorViewportToolBarMenu.h"
 #include "BlueprintEditorTabs.h"
+#include "BlueprintEditorSettings.h"
 
 /*-----------------------------------------------------------------------------
    SSCSEditorViewportToolBar
@@ -289,7 +290,10 @@ void SSCSEditorViewport::ToggleIsSimulateEnabled()
 		// Only trigger the switch if the simulation is starting.
 		if ( !ViewportClient->GetIsSimulateEnabled() )
 		{
-			BlueprintEditorPtr.Pin()->GetTabManager()->InvokeTab(FBlueprintEditorTabs::SCSViewportID);
+			if ( GetDefault<UBlueprintEditorSettings>()->bShowViewportOnSimulate )
+			{
+				BlueprintEditorPtr.Pin()->GetTabManager()->InvokeTab(FBlueprintEditorTabs::SCSViewportID);
+			}
 		}
 	}
 
diff --git a/Engine/Source/Editor/Kismet/Public/BlueprintEditor.h b/Engine/Source/Editor/Kismet/Public/BlueprintEditor.h
--- a/Engine/Source/Editor/Kismet/Public/BlueprintEditor.h
+++ b/Engine/Source/Editor/Kismet/Public/BlueprintEditor.h
@@ -394,6 +394,9 @@ public:
 	/** Delegate invoked when the selection is changed in the SCS editor widget */
 	void OnSelectionUpdated(const TArray<TSharedPtr<class FSCSEditorTreeNode>>& SelectedNodes);
 
+	/** Delegate invoked when an item is double clicked in the SCS editor widget */
+	void OnComponentDoubleClicked(TSharedPtr<class FSCSEditorTreeNode> Node);
+
 
 
 	/** Pin visibility accessors */
@@ -440,25 +443,21 @@ public:
 	{
 		return NewDocument_IsVisibleForType(GraphType) ? EVisibility::Visible : EVisibility::Collapsed;
 	}
-	
-	/** Clear selection across all editors */
-	void ClearSelectionInAllEditors();
+
+	static FName SelectionState_MyBlueprint;
+	static FName SelectionState_Components;
+	static FName SelectionState_Graph;
+	static FName SelectionState_ClassSettings;
+	static FName SelectionState_ClassDefaults;
 	
 	/** Gets or sets the flag for context sensitivity in the graph action menu */
-	bool& GetIsContextSensitive() {return bIsActionMenuContextSensitive;}
-
-	/** Selection state, because all selection in this editor is mutually exclusive */
-	enum ESelectionState
-	{
-		NoSelection,
-		MyBlueprint,
-		GraphPanel,
-		ClassSettings,
-		ClassDefaults,
-	};
+	bool& GetIsContextSensitive() { return bIsActionMenuContextSensitive; }
 
 	/** Gets the UI selection state of this editor */
-	ESelectionState& GetUISelectionState() {return CurrentUISelection;}
+	FName GetUISelectionState() const { return CurrentUISelection; }
+	void SetUISelectionState(FName SelectionOwner);
+
+	virtual void ClearSelectionStateFor(FName SelectionOwner);
 
 	/** Find all instances of the selected custom event. */
 	void OnFindInstancesCustomEvent();
@@ -905,6 +904,9 @@ private:
 	/** Function to check whether the give graph is a subgraph */
 	static bool IsASubGraph( const class UEdGraph* GraphPtr );
 
+	/** Creates the SCSEditor tree component view and the SCS Viewport. */
+	void CreateSCSEditors();
+
 	/** Callback when a token is clicked on in the compiler results log */
 	void OnLogTokenClicked(const TSharedRef<class IMessageToken>& Token);
 
@@ -1025,7 +1027,7 @@ protected:
 	bool bIsActionMenuContextSensitive;
 	
 	/** The current UI selection state of this editor */
-	ESelectionState CurrentUISelection;
+	FName CurrentUISelection;
 
 	/** Whether we are already in the process of closing this editor */
 	bool bEditorMarkedAsClosed;
diff --git a/Engine/Source/Editor/Kismet/Public/SKismetInspector.h b/Engine/Source/Editor/Kismet/Public/SKismetInspector.h
--- a/Engine/Source/Editor/Kismet/Public/SKismetInspector.h
+++ b/Engine/Source/Editor/Kismet/Public/SKismetInspector.h
@@ -111,6 +111,9 @@ protected:
 	/** If true show the kismet inspector title widget */
 	bool bShowTitleArea;
 
+	/** Component details customization enabled. */
+	bool bComponenetDetailsCustomizationEnabled;
+
 	/** Set of object properties that should be visible */
 	TSet<TWeakObjectPtr<UProperty> > SelectedObjectProperties;
 	
diff --git a/Engine/Source/Editor/Kismet/Public/SSCSEditor.h b/Engine/Source/Editor/Kismet/Public/SSCSEditor.h
--- a/Engine/Source/Editor/Kismet/Public/SSCSEditor.h
+++ b/Engine/Source/Editor/Kismet/Public/SSCSEditor.h
@@ -416,6 +416,7 @@ public:
 	DECLARE_DELEGATE_RetVal_OneParam(class USCS_Node*, FOnAddExistingComponent, class UActorComponent*);
 	DECLARE_DELEGATE_OneParam(FOnRootSelected, class AActor*);
 	DECLARE_DELEGATE_OneParam(FOnSelectionUpdated, const TArray<FSCSEditorTreeNodePtrType>&);
+	DECLARE_DELEGATE_OneParam(FOnItemDoubleClicked, const FSCSEditorTreeNodePtrType);
 	DECLARE_DELEGATE_OneParam(FOnHighlightPropertyInDetailsView, const class FPropertyPath&);
 
 	SLATE_BEGIN_ARGS( SSCSEditor )
@@ -436,6 +437,7 @@ public:
 		SLATE_ATTRIBUTE(bool, HideComponentClassCombo)
 		SLATE_EVENT(FOnRootSelected, OnRootSelected)
 		SLATE_EVENT(FOnSelectionUpdated, OnSelectionUpdated)
+		SLATE_EVENT(FOnItemDoubleClicked, OnItemDoubleClicked)
 		SLATE_EVENT(FOnHighlightPropertyInDetailsView, OnHighlightPropertyInDetailsView)
 
 	SLATE_END_ARGS()
@@ -675,6 +677,9 @@ protected:
 	/** Callback when a component item is scrolled into view */
 	void OnItemScrolledIntoView( FSCSEditorTreeNodePtrType InItem, const TSharedPtr<ITableRow>& InWidget);
 
+	/** Callback when a component item is double clicked. */
+	void HandleItemDoubleClicked(FSCSEditorTreeNodePtrType InItem);
+
 	/** Returns the set of expandable nodes that are currently collapsed in the UI */
 	void GetCollapsedNodes(const FSCSEditorTreeNodePtrType& InNodePtr, TSet<FSCSEditorTreeNodePtrType>& OutCollapsedNodes) const;
 
@@ -712,6 +717,9 @@ public:
 	/** Delegate to invoke on selection update. */
 	FOnSelectionUpdated OnSelectionUpdated;
 
+	/** Delegate to invoke when an item in the tree is double clicked. */
+	FOnItemDoubleClicked OnItemDoubleClicked;
+
 	/** Delegate to invoke when the given property should be highlighted in the details view (e.g. diff). */
 	FOnHighlightPropertyInDetailsView OnHighlightPropertyInDetailsView;
 private:
@@ -722,6 +730,9 @@ private:
 	/** Flag to enable/disable component editing */
 	bool	bEnableComponentEditing;
 
+	/** Gate to prevent changing the selection while selection change is being broadcast. */
+	bool bUpdatingSelection;
+
 	/** Flag to track if we have the Actor selected curently */
 	bool bIsActorSelected;
 };
