[INTEGRATE] Change 2415789 by [REDACTED_PERSON][REDACTED_HANDLE] on 2015/01/22 15:57:24

	UE-8097 - BP - Now checking if components mode is needed before creating or deleting the preview in the blueprint editor.

diff --git a/Engine/Source/Editor/Kismet/Private/BlueprintEditor.cpp b/Engine/Source/Editor/Kismet/Private/BlueprintEditor.cpp
--- a/Engine/Source/Editor/Kismet/Private/BlueprintEditor.cpp
+++ b/Engine/Source/Editor/Kismet/Private/BlueprintEditor.cpp
@@ -7084,6 +7084,12 @@ AActor* FBlueprintEditor::GetPreviewActor() const
 
 void FBlueprintEditor::UpdatePreviewActor(UBlueprint* InBlueprint, bool bInForceFullUpdate/* = false*/)
 {
+	// If the components mode isn't available there's no reason to update the preview actor.
+	if ( !CanAccessComponentsMode() )
+	{
+		return;
+	}
+
 	AActor* PreviewActor = GetPreviewActor();
 
 	// Signal that we're going to be constructing editor components
@@ -7167,6 +7173,12 @@ void FBlueprintEditor::UpdatePreviewActor(UBlueprint* InBlueprint, bool bInForce
 
 void FBlueprintEditor::DestroyPreview()
 {
+	// If the components mode isn't available there's no reason to delete the preview actor.
+	if ( !CanAccessComponentsMode() )
+	{
+		return;
+	}
+
 	AActor* PreviewActor = GetPreviewActor();
 	if ( PreviewActor != nullptr )
 	{
