[INTEGRATE] Change 2415917 by [REDACTED_PERSON][REDACTED_HANDLE] on 2015/01/22 17:20:04

	BP - Even with show empty sections enabled, we now limit their visibility to the ones that actually apply to this kind of blueprint.

diff --git a/Engine/Source/Editor/Kismet/Private/SMyBlueprint.cpp b/Engine/Source/Editor/Kismet/Private/SMyBlueprint.cpp
--- a/Engine/Source/Editor/Kismet/Private/SMyBlueprint.cpp
+++ b/Engine/Source/Editor/Kismet/Private/SMyBlueprint.cpp
@@ -1282,11 +1282,28 @@ void SMyBlueprint::CollectStaticSections(TArray<int32>& StaticSectionIDs)
 {
 	if ( IsShowingEmptySections() )
 	{
-		StaticSectionIDs.Add(NodeSectionID::GRAPH);
-		StaticSectionIDs.Add(NodeSectionID::MACRO);
-		StaticSectionIDs.Add(NodeSectionID::FUNCTION);
-		StaticSectionIDs.Add(NodeSectionID::VARIABLE);
-		StaticSectionIDs.Add(NodeSectionID::DELEGATE);
+		TSharedPtr<FBlueprintEditor> BlueprintEditor = BlueprintEditorPtr.Pin();
+		
+		if ( BlueprintEditor->NewDocument_IsVisibleForType(FBlueprintEditor::CGT_NewEventGraph) )
+		{
+			StaticSectionIDs.Add(NodeSectionID::GRAPH);
+		}
+		if ( BlueprintEditor->NewDocument_IsVisibleForType(FBlueprintEditor::CGT_NewMacroGraph) )
+		{
+			StaticSectionIDs.Add(NodeSectionID::MACRO);
+		}
+		if ( BlueprintEditor->NewDocument_IsVisibleForType(FBlueprintEditor::CGT_NewFunctionGraph) )
+		{
+			StaticSectionIDs.Add(NodeSectionID::FUNCTION);
+		}
+		if ( BlueprintEditor->NewDocument_IsVisibleForType(FBlueprintEditor::CGT_NewVariable) )
+		{
+			StaticSectionIDs.Add(NodeSectionID::VARIABLE);
+		}
+		if ( BlueprintEditor->FBlueprintEditor::AddNewDelegateIsVisible() )
+		{
+			StaticSectionIDs.Add(NodeSectionID::DELEGATE);
+		}
 	}
 
 	if ( GetLocalActionsListVisibility().IsVisible() )
