[INTEGRATE] Change 2417088 by [REDACTED_PERSON][REDACTED_HANDLE] on 2015/01/23 14:12:07

	BP - Improvements to the starting conditions of several blueprint types, macro, function libraries and interfaces no longer begin with an empty graph screne, you now get a default new element.  Additionally we remember state of those asset types now, so they recall what you were editing last.

diff --git a/Engine/Source/Editor/Kismet/Private/BlueprintEditor.cpp b/Engine/Source/Editor/Kismet/Private/BlueprintEditor.cpp
--- a/Engine/Source/Editor/Kismet/Private/BlueprintEditor.cpp
+++ b/Engine/Source/Editor/Kismet/Private/BlueprintEditor.cpp
@@ -1548,7 +1548,26 @@ void FBlueprintEditor::InitBlueprintEditor(const EToolkitMode::Type Mode, const
 	PostLayoutBlueprintEditorInitialization();
 
 	// Find and set any instances of this blueprint type if any exists and we are not already editing one
-	FBlueprintEditorUtils::FindAndSetDebuggableBlueprintInstances();	
+	FBlueprintEditorUtils::FindAndSetDebuggableBlueprintInstances();
+
+	if ( bNewlyCreated )
+	{
+		if ( UBlueprint* Blueprint = GetBlueprintObj() )
+		{
+			if ( Blueprint->BlueprintType == BPTYPE_MacroLibrary )
+			{
+				NewDocument_OnClick(CGT_NewMacroGraph);
+			}
+			else if ( Blueprint->BlueprintType == BPTYPE_Interface)
+			{
+				NewDocument_OnClick(CGT_NewFunctionGraph);
+			}
+			else if ( Blueprint->BlueprintType == BPTYPE_FunctionLibrary )
+			{
+				NewDocument_OnClick(CGT_NewFunctionGraph);
+			}
+		}
+	}
 }
 
 void FBlueprintEditor::RegisterApplicationModes(const TArray<UBlueprint*>& InBlueprints, bool bShouldOpenInDefaultsMode, bool bNewlyCreated/* = false*/)
@@ -1577,10 +1596,20 @@ void FBlueprintEditor::RegisterApplicationModes(const TArray<UBlueprint*>& InBlu
 		}
 		else if ( SingleBP->BlueprintType == BPTYPE_FunctionLibrary )
 		{
-			AddApplicationMode(
-				FBlueprintEditorApplicationModes::StandardBlueprintEditorMode,
-				MakeShareable(new FBlueprintEditorApplicationMode(SharedThis(this), FBlueprintEditorApplicationModes::StandardBlueprintEditorMode, FBlueprintEditorApplicationModes::GetLocalizedMode)));
-			SetCurrentMode(FBlueprintEditorApplicationModes::StandardBlueprintEditorMode);
+			if ( GetDefault<UEditorExperimentalSettings>()->bUnifiedBlueprintEditor )
+			{
+				AddApplicationMode(
+					FBlueprintEditorApplicationModes::StandardBlueprintEditorMode,
+					MakeShareable(new FBlueprintEditorUnifiedMode(SharedThis(this), FBlueprintEditorApplicationModes::StandardBlueprintEditorMode, FBlueprintEditorApplicationModes::GetLocalizedMode, CanAccessComponentsMode())));
+				SetCurrentMode(FBlueprintEditorApplicationModes::StandardBlueprintEditorMode);
+			}
+			else
+			{
+				AddApplicationMode(
+					FBlueprintEditorApplicationModes::StandardBlueprintEditorMode,
+					MakeShareable(new FBlueprintEditorApplicationMode(SharedThis(this), FBlueprintEditorApplicationModes::StandardBlueprintEditorMode, FBlueprintEditorApplicationModes::GetLocalizedMode)));
+				SetCurrentMode(FBlueprintEditorApplicationModes::StandardBlueprintEditorMode);
+			}
 		}
 		else
 		{
diff --git a/Engine/Source/Editor/Kismet/Private/BlueprintEditorModes.cpp b/Engine/Source/Editor/Kismet/Private/BlueprintEditorModes.cpp
--- a/Engine/Source/Editor/Kismet/Private/BlueprintEditorModes.cpp
+++ b/Engine/Source/Editor/Kismet/Private/BlueprintEditorModes.cpp
@@ -60,19 +60,4 @@
 			->Split
 			(
 				FTabManager::NewSplitter() ->SetOrientation( Orient_Vertical )
-				->SetSizeCoefficient(0.15f)
-				->Split
-				(
-					FTabManager::NewStack() ->SetSizeCoefficient(0.5f)
-					->AddTab( FBlueprintEditorTabs::MyBlueprintID, ETabState::OpenedTab )
-				)
-				->Split
-				(
-					FTabManager::NewStack() ->SetSizeCoefficient(0.5f)
-					->AddTab( FBlueprintEditorTabs::DetailsID, ETabState::OpenedTab )
-				)
-			)
-			->Split
-			(
-				FTabManager::NewSplitter() ->SetOrientation( Orient_Vertical )
 				->SetSizeCoefficient(0.70f)
@@ -96,8 +81,15 @@
 				->SetSizeCoefficient(0.15f)
 				->Split
 				(
-					FTabManager::NewStack()
+					FTabManager::NewStack() ->SetSizeCoefficient(0.5f)
+					->AddTab( FBlueprintEditorTabs::MyBlueprintID, ETabState::OpenedTab )
+				)
+				->Split
+				(
+					FTabManager::NewStack() ->SetSizeCoefficient(0.5f)
+					->AddTab( FBlueprintEditorTabs::DetailsID, ETabState::OpenedTab )
 					->AddTab( FBlueprintEditorTabs::PaletteID, ETabState::ClosedTab )
+					->AddTab( FBlueprintEditorTabs::DefaultEditorID, ETabState::ClosedTab )
 				)
 			)
 		)
@@ -365,39 +357,24 @@
 	BlueprintInterfaceTabFactories.RegisterFactory(MakeShareable(new FFindResultsSummoner(InBlueprintEditor)));
 	BlueprintInterfaceTabFactories.RegisterFactory(MakeShareable(new FSelectionDetailsSummoner(InBlueprintEditor)));
 
-	TabLayout = FTabManager::NewLayout( "Standalone_BlueprintInterface_Layout_v1" )
+	TabLayout = FTabManager::NewLayout( "Standalone_BlueprintInterface_Layout_v3" )
 		->AddArea
 		(
 			FTabManager::NewPrimaryArea() ->SetOrientation(Orient_Vertical)
 			->Split
 			(
-				FTabManager::NewStack()
-				->SetSizeCoefficient( 0.186721f )
-				->SetHideTabWell(true)
-				->AddTab( InBlueprintEditor->GetToolbarTabId(), ETabState::OpenedTab )
-			)
-			->Split
-			(
-				FTabManager::NewSplitter() ->SetOrientation(Orient_Horizontal)
+				FTabManager::NewSplitter()->SetOrientation(Orient_Horizontal)
 				->Split
 				(
-					FTabManager::NewSplitter() ->SetOrientation( Orient_Vertical )
-					->SetSizeCoefficient(0.15f)
-					->Split
-					(
-						FTabManager::NewStack() ->SetSizeCoefficient(0.5f)
-						->AddTab( FBlueprintEditorTabs::MyBlueprintID, ETabState::OpenedTab )
-					)
-					->Split
-					(
-						FTabManager::NewStack() ->SetSizeCoefficient(0.5f)
-						->AddTab( FBlueprintEditorTabs::DetailsID, ETabState::OpenedTab )
-					)
-				)
-				->Split
-				(
-					FTabManager::NewSplitter() ->SetOrientation( Orient_Vertical )
+					FTabManager::NewSplitter()->SetOrientation( Orient_Vertical )
 					->SetSizeCoefficient(0.70f)
 					->Split
 					(
 						FTabManager::NewStack()
+						->SetSizeCoefficient(0.186721f)
+						->SetHideTabWell(true)
+						->AddTab(InBlueprintEditor->GetToolbarTabId(), ETabState::OpenedTab)
+					)
+					->Split
+					(
+						FTabManager::NewStack()
@@ -412,6 +389,21 @@
 						->AddTab( FBlueprintEditorTabs::FindResultsID, ETabState::ClosedTab )
 					)
 				)
+				->Split
+				(
+					FTabManager::NewSplitter() ->SetOrientation( Orient_Vertical )
+					->SetSizeCoefficient(0.15f)
+					->Split
+					(
+						FTabManager::NewStack() ->SetSizeCoefficient(0.5f)
+						->AddTab( FBlueprintEditorTabs::MyBlueprintID, ETabState::OpenedTab )
+					)
+					->Split
+					(
+						FTabManager::NewStack() ->SetSizeCoefficient(0.5f)
+						->AddTab( FBlueprintEditorTabs::DetailsID, ETabState::OpenedTab )
+					)
+				)
 			)
 		);
 
@@ -432,6 +424,24 @@
 	FApplicationMode::RegisterTabFactories(InTabManager);
 }
 
+void FBlueprintInterfaceApplicationMode::PreDeactivateMode()
+{
+	FApplicationMode::PreDeactivateMode();
+
+	TSharedPtr<FBlueprintEditor> BP = MyBlueprintEditor.Pin();
+
+	BP->SaveEditedObjectState();
+}
+
+void FBlueprintInterfaceApplicationMode::PostActivateMode()
+{
+	// Reopen any documents that were open when the blueprint was last saved
+	TSharedPtr<FBlueprintEditor> BP = MyBlueprintEditor.Pin();
+	BP->RestoreEditedObjectState();
+
+	FApplicationMode::PostActivateMode();
+}
+
 ////////////////////////////////////////
 //
 FBlueprintMacroApplicationMode::FBlueprintMacroApplicationMode(TSharedPtr<class FBlueprintEditor> InBlueprintEditor)
@@ -446,39 +456,24 @@
 	BlueprintMacroTabFactories.RegisterFactory(MakeShareable(new FFindResultsSummoner(InBlueprintEditor)));
 	BlueprintMacroTabFactories.RegisterFactory(MakeShareable(new FSelectionDetailsSummoner(InBlueprintEditor)));
 
-	TabLayout = FTabManager::NewLayout( "Standalone_BlueprintMacro_Layout_v1" )
+	TabLayout = FTabManager::NewLayout( "Standalone_BlueprintMacro_Layout_v3" )
 		->AddArea
 		(
 			FTabManager::NewPrimaryArea() ->SetOrientation(Orient_Vertical)
 			->Split
 			(
-				FTabManager::NewStack()
-				->SetSizeCoefficient( 0.2f )
-				->SetHideTabWell(true)
-				->AddTab( InBlueprintEditor->GetToolbarTabId(), ETabState::OpenedTab )
-			)
-			->Split
-			(
-				FTabManager::NewSplitter() ->SetOrientation(Orient_Horizontal)
+				FTabManager::NewSplitter()->SetOrientation(Orient_Horizontal)
 				->Split
 				(
-					FTabManager::NewSplitter() ->SetOrientation( Orient_Vertical )
-					->SetSizeCoefficient(0.15f)
-					->Split
-					(
-						FTabManager::NewStack() ->SetSizeCoefficient(0.5f)
-						->AddTab( FBlueprintEditorTabs::MyBlueprintID, ETabState::OpenedTab )
-					)
-					->Split
-					(
-						FTabManager::NewStack() ->SetSizeCoefficient(0.5f)
-						->AddTab( FBlueprintEditorTabs::DetailsID, ETabState::OpenedTab )
-					)
-				)
-				->Split
-				(
-					FTabManager::NewSplitter() ->SetOrientation( Orient_Vertical )
+					FTabManager::NewSplitter()->SetOrientation( Orient_Vertical )
 					->SetSizeCoefficient(0.70f)
 					->Split
 					(
 						FTabManager::NewStack()
+						->SetSizeCoefficient(0.186721f)
+						->SetHideTabWell(true)
+						->AddTab(InBlueprintEditor->GetToolbarTabId(), ETabState::OpenedTab)
+					)
+					->Split
+					(
+						FTabManager::NewStack()
@@ -498,7 +493,13 @@
 					->SetSizeCoefficient(0.15f)
 					->Split
 					(
-						FTabManager::NewStack()
+						FTabManager::NewStack() ->SetSizeCoefficient(0.5f)
+						->AddTab( FBlueprintEditorTabs::MyBlueprintID, ETabState::OpenedTab )
+					)
+					->Split
+					(
+						FTabManager::NewStack() ->SetSizeCoefficient(0.5f)
+						->AddTab( FBlueprintEditorTabs::DetailsID, ETabState::OpenedTab )
 						->AddTab( FBlueprintEditorTabs::PaletteID, ETabState::ClosedTab )
 					)
 				)
@@ -524,6 +525,26 @@
 	FApplicationMode::RegisterTabFactories(InTabManager);
 }
 
+void FBlueprintMacroApplicationMode::PreDeactivateMode()
+{
+	FApplicationMode::PreDeactivateMode();
+
+	TSharedPtr<FBlueprintEditor> BP = MyBlueprintEditor.Pin();
+
+	BP->SaveEditedObjectState();
+}
+
+void FBlueprintMacroApplicationMode::PostActivateMode()
+{
+	// Reopen any documents that were open when the blueprint was last saved
+	TSharedPtr<FBlueprintEditor> BP = MyBlueprintEditor.Pin();
+	BP->RestoreEditedObjectState();
+
+	FApplicationMode::PostActivateMode();
+}
+
+////////////////////////////////////////
+//
 FBlueprintEditorUnifiedMode::FBlueprintEditorUnifiedMode(TSharedPtr<class FBlueprintEditor> InBlueprintEditor, FName InModeName, FText(*GetLocalizedMode)( const FName ), const bool bRegisterViewport, const bool bRegisterDefaultsTab)
 	: FApplicationMode(InModeName, GetLocalizedMode)
 {
@@ -603,7 +624,8 @@
 						FTabManager::NewStack()
 						->SetSizeCoefficient(0.50f)
 						->AddTab( FBlueprintEditorTabs::DetailsID, ETabState::OpenedTab )
-						//->AddTab( FBlueprintEditorTabs::DefaultEditorID, ETabState::ClosedTab )
+						->AddTab( FBlueprintEditorTabs::PaletteID, ETabState::ClosedTab )
+						->AddTab( FBlueprintEditorTabs::DefaultEditorID, ETabState::ClosedTab )
 					)
 				)
 			)
@@ -658,7 +680,7 @@
 						FTabManager::NewStack()
 						->SetSizeCoefficient(0.60f)
 						->AddTab( FBlueprintEditorTabs::DetailsID, ETabState::OpenedTab )
-						//->AddTab( FBlueprintEditorTabs::DefaultEditorID, ETabState::ClosedTab )
+						->AddTab( FBlueprintEditorTabs::DefaultEditorID, ETabState::ClosedTab )
 					)
 				)
 			)
diff --git a/Engine/Source/Editor/Kismet/Private/SMyBlueprint.cpp b/Engine/Source/Editor/Kismet/Private/SMyBlueprint.cpp
--- a/Engine/Source/Editor/Kismet/Private/SMyBlueprint.cpp
+++ b/Engine/Source/Editor/Kismet/Private/SMyBlueprint.cpp
@@ -616,7 +616,6 @@ FReply SMyBlueprint::OnAddButtonClickedOnSection(int32 InSectionID)
 		BlueprintEditorPtr.Pin()->GetToolkitCommands()->ExecuteAction(FBlueprintEditorCommands::Get().AddNewEventGraph.ToSharedRef());
 		break;
 	case NodeSectionID::LOCAL_VARIABLE:
-		//BlueprintEditorPtr.Pin()->GetToolkitCommands()->ExecuteAction(FBlueprintEditorCommands::Get().AddNewLocalVariable.ToSharedRef());
 		OnAddNewLocalVariable();
 		break;
 	}
diff --git a/Engine/Source/Editor/Kismet/Public/BlueprintEditorModes.h b/Engine/Source/Editor/Kismet/Public/BlueprintEditorModes.h
--- a/Engine/Source/Editor/Kismet/Public/BlueprintEditorModes.h
+++ b/Engine/Source/Editor/Kismet/Public/BlueprintEditorModes.h
@@ -105,6 +105,8 @@ public:
 	FBlueprintInterfaceApplicationMode(TSharedPtr<class FBlueprintEditor> InBlueprintEditor);
 
 	virtual void RegisterTabFactories(TSharedPtr<FTabManager> InTabManager) override;
+	virtual void PreDeactivateMode() override;
+	virtual void PostActivateMode() override;
 public:
 
 protected:
@@ -120,6 +122,8 @@ public:
 	FBlueprintMacroApplicationMode(TSharedPtr<class FBlueprintEditor> InBlueprintEditor);
 
 	virtual void RegisterTabFactories(TSharedPtr<FTabManager> InTabManager) override;
+	virtual void PreDeactivateMode() override;
+	virtual void PostActivateMode() override;
 public:
 
 protected:
