[INTEGRATE] Change 2417177 by [REDACTED_PERSON][REDACTED_HANDLE] on 2015/01/23 15:28:42

	IWCE - Added support for root node deletion/promotion and swapping with other component instances.

	change summary:
	- moved DefaultSceneRootVariableName to FComponentEditorUtils
	- updated references to DefaultSceneRootVariableName to access from FComponentEditorUtils
	- modified UActorFactoryEmptyActor::SpawnActor() to use the default scene root component name for the root component name
	- moved some instancing code out of SSCSEditor::AddNewComponent() into AddNewNode() instead (so it runs for other ops, e.g. Paste)
	- removed name validation code on new tree node creation as (a) it's already handled elsewhere now and (b) i need to preserve the default root name
	- filled out instanced mode side of SSCS_RowWidget::OnMakeNewRoot()
	- updated SSCSEditor::AddNewNode() for instanced mode to remove the "default" scene root node on first user component add
	- updated SSCSEditor::RemoveComponentNode() to handle both removal of an instanced root node as well as promotion of non editor-only child nodes to root

diff --git a/Engine/Source/Editor/Kismet/Private/SSCSEditor.cpp b/Engine/Source/Editor/Kismet/Private/SSCSEditor.cpp
--- a/Engine/Source/Editor/Kismet/Private/SSCSEditor.cpp
+++ b/Engine/Source/Editor/Kismet/Private/SSCSEditor.cpp
@@ -244,15 +244,6 @@ FSCSEditorTreeNode::FSCSEditorTreeNode(UActorComponent* InComponentTemplate)
 	{
 		bIsInstanced = true;
 		bWasInstancedFromNativeClass = false;
-
-		// Make sure the component has a valid name
-		if (!FComponentEditorUtils::IsValidVariableNameString(InComponentTemplate, InComponentTemplate->GetName()))
-		{
-			ERenameFlags RenameFlags = REN_DontCreateRedirectors | REN_NonTransactional;
-			FString NewComponentName = FComponentEditorUtils::GenerateValidVariableName(InComponentTemplate->GetClass(), InComponentTemplate->GetOwner());
-
-			InComponentTemplate->Rename(*NewComponentName, nullptr, RenameFlags);
-		}
 		
 		InstancedComponentName = InComponentTemplate->GetFName();
 
@@ -522,14 +513,22 @@ bool FSCSEditorTreeNode::IsAttachedTo(FSCSEditorTreeNodePtrType InNodePtr) const
 bool FSCSEditorTreeNode::IsDefaultSceneRoot() const
 {
 	USCS_Node* SCS_Node = GetSCSNode();
-	if(SCS_Node != NULL)
+	if(SCS_Node != nullptr)
 	{
 		USimpleConstructionScript* SCS = SCS_Node->GetSCS();
-		if(SCS != NULL)
+		if(SCS != nullptr)
 		{
 			return SCS_Node == SCS->GetDefaultSceneRootNode();
 		}
 	}
+	else if(IsUserInstanced())
+	{
+		USceneComponent* SceneComponent = Cast<USceneComponent>(GetComponentTemplate());
+		if(SceneComponent != nullptr)
+		{
+			return SceneComponent->GetFName() == FComponentEditorUtils::GetDefaultSceneRootVariableName();
+		}
+	}
 
 	return false;
 }
@@ -2160,7 +2159,23 @@ void SSCS_RowWidget::OnMakeNewRootDropAction(FSCSEditorTreeNodePtrType DroppedNo
 	}
 	else    // SSCSEditor::EEditorMode::ActorInstance
 	{
+		if(DroppedNodePtr->GetParent().IsValid())
+		{
+			// Remove the dropped node from its existing parent
+			DroppedNodePtr->GetParent()->RemoveChild(DroppedNodePtr);
+		}
+
+		check(SceneRootNodePtr->CanReparent());
+
+		// Save old root node
+		FSCSEditorTreeNodePtrType OldSceneRootNodePtr = SceneRootNodePtr;
+
+		// Set node we are dropping as new root
+		SceneRootNodePtr = DroppedNodePtr;
 
+		// Set old root as child of new root
+		check(OldSceneRootNodePtr.IsValid());
+		SceneRootNodePtr->AddChild(OldSceneRootNodePtr);
 	}
 
 	PostDragDropAction(true);
@@ -3721,7 +3736,6 @@ UActorComponent* SSCSEditor::AddNewComponent( UClass* NewComponentClass, UObject
 		
 		ActorInstance->Modify();
 
-		// Create new component
 		FName NewComponentName = NAME_None;
 		if (Asset != nullptr)
 		{
@@ -3731,19 +3745,7 @@ UActorComponent* SSCSEditor::AddNewComponent( UClass* NewComponentClass, UObject
 		{
 			NewComponentName = *FComponentEditorUtils::GenerateValidVariableName(NewComponentClass, ActorInstance);
 		}
-
-		UActorComponent* NewComponentInstance = ConstructObject<UActorComponent>(NewComponentClass, ActorInstance, NewComponentName, RF_Transactional);
-		check(NewComponentInstance != nullptr);
-
-		if (USceneComponent* NewSceneComponent = Cast<USceneComponent>(NewComponentInstance))
-		{
-			NewSceneComponent->AttachTo(ActorInstance->GetRootComponent());
-		}
-
-		// Add to SerializedComponents array so it gets saved
-		ActorInstance->AddInstanceComponent(NewComponentInstance);
-
-		return AddNewNode(NewComponentInstance, Asset, true);
+		return AddNewNode(ConstructObject<UActorComponent>(NewComponentClass, ActorInstance, NewComponentName, RF_Transactional), Asset, true);
 	}
 }
 
@@ -3800,6 +3802,17 @@ UActorComponent* SSCSEditor::AddNewNode(UActorComponent* NewInstanceComponent,
 {
 	check(NewInstanceComponent != nullptr);
 
+	AActor* ActorInstance = ActorContext.Get();
+	check(ActorInstance != nullptr);
+
+	if (USceneComponent* NewSceneComponent = Cast<USceneComponent>(NewInstanceComponent))
+	{
+		NewSceneComponent->AttachTo(ActorInstance->GetRootComponent());
+	}
+
+	// Add to SerializedComponents array so it gets saved
+	ActorInstance->AddInstanceComponent(NewInstanceComponent);
+
 	if(Asset)
 	{
 		FComponentAssetBrokerage::AssignAssetToComponent(NewInstanceComponent, Asset);
@@ -3814,6 +3827,15 @@ UActorComponent* SSCSEditor::AddNewNode(UActorComponent* NewInstanceComponent,
 	if(NewSceneComponent != nullptr)
 	{
 		NewNodePtr = AddTreeNode(NewSceneComponent);
+
+		// Remove the old scene root node if it's set to the default one
+		if(SceneRootNodePtr.IsValid() && SceneRootNodePtr->IsDefaultSceneRoot())
+		{
+			RemoveComponentNode(SceneRootNodePtr);
+			SceneRootNodePtr.Reset();
+
+			UpdateTree();
+		}
 	}
 	else
 	{
@@ -4048,9 +4070,6 @@ void SSCSEditor::PasteNodes()
 			FString NewComponentName = FComponentEditorUtils::GenerateValidVariableName(NewActorComponent->GetClass(), ActorInstance);
 			NewActorComponent->Rename(*NewComponentName, ActorInstance, REN_DontCreateRedirectors | REN_DoNotDirty);
 
-			// Add to SerializedComponents array so it gets saved
-			ActorInstance->AddInstanceComponent(NewActorComponent);
-
 			// Create a new node to contain the new component instance and add it to the tree
 			NewActorComponent = AddNewNode(NewActorComponent, NULL, false);
 		}
@@ -4253,6 +4272,91 @@ void SSCSEditor::RemoveComponentNode(FSCSEditorTreeNodePtrType InNodePtr)
 		UActorComponent* ComponentInstance = InNodePtr->GetComponentTemplate();
 		check(ComponentInstance != nullptr);
 
+		ComponentInstance->Modify();
+
+		// Handle removal of a node within the scene component hierarchy
+		USceneComponent* SceneComponent = Cast<USceneComponent>(ComponentInstance);
+		if(SceneComponent != nullptr)
+		{
+			// Find an appropriate child node to promote to this node's position in the hierarchy
+			USceneComponent* ChildToPromote = nullptr;
+			if(SceneComponent->AttachChildren.Num() > 0)
+			{
+				// Start with the first child node
+				ChildToPromote = SceneComponent->AttachChildren[0];
+				check(ChildToPromote != nullptr);
+
+				// Always choose non editor-only child nodes over editor-only child nodes (since we don't want editor-only nodes to end up with non editor-only child nodes)
+				if(ChildToPromote->IsEditorOnly())
+				{
+					for(int32 ChildIndex = 1; ChildIndex < SceneComponent->AttachChildren.Num(); ++ChildIndex)
+					{
+						USceneComponent* Child = SceneComponent->AttachChildren[ChildIndex];
+						if (Child != nullptr && !Child->IsEditorOnly())
+						{
+							ChildToPromote = Child;
+							break;
+						}
+					}
+				}
+			}
+
+			// Handle removal of the root node
+			if(SceneComponent == ActorInstance->GetRootComponent())
+			{
+				// We only promote non editor-only components to root in instanced mode
+				if(ChildToPromote == nullptr || ChildToPromote->IsEditorOnly())
+				{
+					// Construct a new default root component
+					USceneComponent* NewRootComponent = ConstructObject<USceneComponent>(USceneComponent::StaticClass(), ActorInstance, FComponentEditorUtils::GetDefaultSceneRootVariableName(), RF_Transactional);
+					NewRootComponent->Mobility = SceneComponent->Mobility;
+					NewRootComponent->SetWorldLocationAndRotation(SceneComponent->GetComponentLocation(), SceneComponent->GetComponentRotation());
+
+					// Designate the new default root as the child we're promoting
+					ChildToPromote = NewRootComponent;
+				}
+
+				ActorInstance->Modify();
+
+				// Set the selected child node as the new root
+				check(ChildToPromote != nullptr);
+				ActorInstance->SetRootComponent(ChildToPromote);
+			}
+			else    // ...not the root node, so we'll promote the selected child node to this position in its AttachParent's child array.
+			{
+				// Get the node's AttachParent
+				USceneComponent* AttachParent = SceneComponent->AttachParent;
+				check(AttachParent != nullptr);
+
+				// Find the node's position in its AttachParent's child array
+				int32 Index = AttachParent->AttachChildren.Find(SceneComponent);
+				check(Index != INDEX_NONE);
+
+				// Detach the node from its parent
+				SceneComponent->DetachFromParent();
+
+				// Attach the child node that we're promoting to the parent and move it to the same position as the old node was in the array
+				ChildToPromote->AttachTo(AttachParent, NAME_None, EAttachLocation::KeepWorldPosition);
+				AttachParent->AttachChildren.Remove(ChildToPromote);
+				AttachParent->AttachChildren.Insert(ChildToPromote, Index);
+			}
+
+			// Detach child nodes from the node that's being removed and re-attach them to the child that's being promoted
+			TArray<USceneComponent*> AttachChildrenLocalCopy(SceneComponent->AttachChildren);
+			for(auto ChildCompIt = AttachChildrenLocalCopy.CreateIterator(); ChildCompIt; ++ChildCompIt)
+			{
+				USceneComponent* Child = *ChildCompIt;
+				check(Child != nullptr);
+
+				// Note: This will internally call Modify(), so we don't need to call it here
+				Child->DetachFromParent(true);
+				if(Child != ChildToPromote)
+				{
+					Child->AttachTo(ChildToPromote, NAME_None, EAttachLocation::KeepWorldPosition);
+				}
+			}
+		}
+
 		// Destroy the component instance
 		ComponentInstance->DestroyComponent();
 	}
@@ -4414,8 +4518,9 @@ FSCSEditorTreeNodePtrType SSCSEditor::AddTreeNode(USceneComponent* InSceneCompon
 
 	check(InSceneComponent != NULL);
 
-	// If the given component has a parent
-	if(InSceneComponent->AttachParent != NULL)
+	// If the given component has a parent, and if we're not in "instance" mode OR the owner of the parent matches the Actor instance we're editing
+	if(InSceneComponent->AttachParent != NULL
+		&& (EditorMode != EEditorMode::ActorInstance || InSceneComponent->AttachParent->GetOwner() == ActorContext.Get()))
 	{
 		// Attempt to find the parent node in the current tree
 		FSCSEditorTreeNodePtrType ParentNodePtr = FindTreeNode(InSceneComponent->AttachParent);
diff --git a/Engine/Source/Editor/Kismet/Public/SSCSEditor.h b/Engine/Source/Editor/Kismet/Public/SSCSEditor.h
--- a/Engine/Source/Editor/Kismet/Public/SSCSEditor.h
+++ b/Engine/Source/Editor/Kismet/Public/SSCSEditor.h
@@ -115,15 +115,15 @@ public:
 	/** 
 	 * @return Whether or not this object represents a node that can be deleted from the SCS tree.
 	 */
-	bool CanDelete() const { return IsUserInstanced() || (!IsNative() && !IsInherited() && !IsDefaultSceneRoot()); }
+	bool CanDelete() const { return (IsUserInstanced() || (!IsNative() && !IsInherited())) && !IsDefaultSceneRoot(); }
 	/** 
 	 * @return Whether or not this object represents a node that can be renamed from the SCS tree.
 	 */
-	bool CanRename() const { return IsUserInstanced() || (!IsNative() && !IsInherited() && !IsDefaultSceneRoot()); }
+	bool CanRename() const { return (IsUserInstanced() || (!IsNative() && !IsInherited())) && !IsDefaultSceneRoot(); }
 	/** 
 	 * @return Whether or not this object represents a node that can be reparented to other nodes based on its context.
 	 */
-	bool CanReparent() const { return (IsUserInstanced() || (!IsNative() && !IsInherited() && !IsDefaultSceneRoot())) && Cast<USceneComponent>(GetComponentTemplate()) != NULL; }
+	bool CanReparent() const { return (IsUserInstanced() || (!IsNative() && !IsInherited())) && !IsDefaultSceneRoot() && Cast<USceneComponent>(GetComponentTemplate()) != NULL; }
 	/** 
 	 * @return Whether or not we can edit default properties for the component template represented by this object.
 	 */
diff --git a/Engine/Source/Editor/UnrealEd/Private/Factories/ActorFactory.cpp b/Engine/Source/Editor/UnrealEd/Private/Factories/ActorFactory.cpp
--- a/Engine/Source/Editor/UnrealEd/Private/Factories/ActorFactory.cpp
+++ b/Engine/Source/Editor/UnrealEd/Private/Factories/ActorFactory.cpp
@@ -64,6 +64,7 @@ ActorFactory.cpp:
 #include "Engine/Polys.h"
 #include "Components/VectorFieldComponent.h"
 #include "Animation/AnimBlueprintGeneratedClass.h"
+#include "Kismet2/ComponentEditorUtils.h"
 
 DEFINE_LOG_CATEGORY_STATIC(LogActorFactory, Log, All);
 
@@ -1151,7 +1152,7 @@ AActor* UActorFactoryEmptyActor::SpawnActor( UObject* Asset, ULevel* InLevel, co
 		// Spawn a temporary actor for dragging around
 		NewActor = Super::SpawnActor(Asset, InLevel, Location, Rotation, ObjectFlags, Name);
 
-		USceneComponent* RootComponent = ConstructObject<USceneComponent>(USceneComponent::StaticClass(), NewActor, FName("Root"), RF_Transactional);
+		USceneComponent* RootComponent = ConstructObject<USceneComponent>(USceneComponent::StaticClass(), NewActor, FComponentEditorUtils::GetDefaultSceneRootVariableName(), RF_Transactional);
 		RootComponent->Mobility = EComponentMobility::Movable;
 		RootComponent->SetWorldLocationAndRotation(Location, Rotation);
 		NewActor->SetRootComponent(RootComponent);
diff --git a/Engine/Source/Editor/UnrealEd/Private/Kismet2/ComponentEditorUtils.cpp b/Engine/Source/Editor/UnrealEd/Private/Kismet2/ComponentEditorUtils.cpp
--- a/Engine/Source/Editor/UnrealEd/Private/Kismet2/ComponentEditorUtils.cpp
+++ b/Engine/Source/Editor/UnrealEd/Private/Kismet2/ComponentEditorUtils.cpp
@@ -7,6 +7,13 @@
 #include "Engine/SCS_Node.h"
 #include "Engine/SimpleConstructionScript.h"
 
+const FName& FComponentEditorUtils::GetDefaultSceneRootVariableName()
+{
+	static FName DefaultSceneRootVariableName = FName(TEXT("DefaultSceneRoot"));
+
+	return DefaultSceneRootVariableName;
+}
+
 USceneComponent* FComponentEditorUtils::GetSceneComponent( UObject* Object, UObject* SubObject /*= NULL*/ )
 {
 	if( Object )
@@ -93,7 +100,7 @@ void FComponentEditorUtils::GetArchetypeInstances( UObject* Object, TArray<UObje
 bool FComponentEditorUtils::IsValidVariableNameString(const UActorComponent* InComponent, const FString& InString)
 {
 	// First test to make sure the string is not empty and does not equate to the DefaultSceneRoot node name
-	bool bIsValid = !InString.IsEmpty() && !InString.Equals(USimpleConstructionScript::DefaultSceneRootVariableName.ToString());
+	bool bIsValid = !InString.IsEmpty() && !InString.Equals(GetDefaultSceneRootVariableName().ToString());
 	if(bIsValid && InComponent != NULL)
 	{
 		// Next test to make sure the string doesn't conflict with the format that MakeUniqueObjectName() generates
diff --git a/Engine/Source/Editor/UnrealEd/Public/Kismet2/ComponentEditorUtils.h b/Engine/Source/Editor/UnrealEd/Public/Kismet2/ComponentEditorUtils.h
--- a/Engine/Source/Editor/UnrealEd/Public/Kismet2/ComponentEditorUtils.h
+++ b/Engine/Source/Editor/UnrealEd/Public/Kismet2/ComponentEditorUtils.h
@@ -8,6 +8,8 @@
 class UNREALED_API FComponentEditorUtils
 {
 public:
+	/** The name to use for the default scene root variable */
+	static const FName& GetDefaultSceneRootVariableName();
 
 	static class USceneComponent* GetSceneComponent(class UObject* Object, class UObject* SubObject = NULL);
 
diff --git a/Engine/Source/Runtime/Engine/Classes/Engine/SimpleConstructionScript.h b/Engine/Source/Runtime/Engine/Classes/Engine/SimpleConstructionScript.h
--- a/Engine/Source/Runtime/Engine/Classes/Engine/SimpleConstructionScript.h
+++ b/Engine/Source/Runtime/Engine/Classes/Engine/SimpleConstructionScript.h
@@ -123,9 +123,6 @@ public:
 		return EditorActorInstancePtr.Get();
 	}
 
-	/** The name to use for the default scene root variable */
-	ENGINE_API static const FName DefaultSceneRootVariableName;
-
 private:
 	/** Actor instance used to host components in the SCS editor */
 	TWeakObjectPtr<class AActor> EditorActorInstancePtr;
diff --git a/Engine/Source/Runtime/Engine/Private/SimpleConstructionScript.cpp b/Engine/Source/Runtime/Engine/Private/SimpleConstructionScript.cpp
--- a/Engine/Source/Runtime/Engine/Private/SimpleConstructionScript.cpp
+++ b/Engine/Source/Runtime/Engine/Private/SimpleConstructionScript.cpp
@@ -14,10 +14,6 @@
 //////////////////////////////////////////////////////////////////////////
 // USimpleConstructionScript
 
-#if WITH_EDITOR
-const FName USimpleConstructionScript::DefaultSceneRootVariableName = FName(TEXT("DefaultSceneRoot"));
-#endif
-
 namespace
 {
 	// Helper method to register instanced components post-construction
@@ -651,7 +647,7 @@ void USimpleConstructionScript::ValidateSceneRootNodes()
 			&& FBlueprintEditorUtils::IsActorBased(Blueprint)
 			&& Blueprint->BlueprintType != BPTYPE_MacroLibrary)
 		{
-			DefaultSceneRootNode = CreateNode(USceneComponent::StaticClass(), DefaultSceneRootVariableName);
+			DefaultSceneRootNode = CreateNode(USceneComponent::StaticClass(), FComponentEditorUtils::GetDefaultSceneRootVariableName());
 		}
 	}
 
