[INTEGRATE] Change 2417546 by [REDACTED_PERSON][REDACTED_HANDLE] on 2015/01/23 19:26:20

	Move the SCS Root Actor into the treeview.

diff --git a/Engine/Source/Editor/Kismet/Private/BPVariableDragDropAction.cpp b/Engine/Source/Editor/Kismet/Private/BPVariableDragDropAction.cpp
--- a/Engine/Source/Editor/Kismet/Private/BPVariableDragDropAction.cpp
+++ b/Engine/Source/Editor/Kismet/Private/BPVariableDragDropAction.cpp
@@ -432,10 +432,10 @@ FReply FKismetVariableDragDropAction::DroppedOnPanel( const TSharedRef< SWidget
 	if (Cast<const UEdGraphSchema_K2>(Graph.GetSchema()) != NULL)
 	{
 		UProperty* VariableProperty = GetVariableProperty();
-		UStruct* Outer = CastChecked<UStruct>(VariableProperty->GetOuter());
-
-		if(CanVariableBeDropped(VariableProperty, Graph))
+		if (VariableProperty != nullptr && CanVariableBeDropped(VariableProperty, Graph))
 		{
+			UStruct* Outer = CastChecked<UStruct>(VariableProperty->GetOuter());
+			
 			FNodeConstructionParams NewNodeParams;
 			NewNodeParams.VariableName = VariableName;
 			NewNodeParams.Graph = &Graph;
diff --git a/Engine/Source/Editor/Kismet/Private/BlueprintDetailsCustomization.cpp b/Engine/Source/Editor/Kismet/Private/BlueprintDetailsCustomization.cpp
--- a/Engine/Source/Editor/Kismet/Private/BlueprintDetailsCustomization.cpp
+++ b/Engine/Source/Editor/Kismet/Private/BlueprintDetailsCustomization.cpp
@@ -4758,7 +4758,7 @@ bool FBlueprintComponentDetails::IsNodeAttachable() const
 
 	// See if node is an attachable type
 	bool bCanAttach = false;	
-	if (!CachedNodePtr->IsRoot())
+	if (!CachedNodePtr->IsRootComponent())
 	{
 		// check component is the correct type
 		if (USceneComponent* SceneComponent = Cast<USceneComponent>(CachedNodePtr->GetComponentTemplate()))
diff --git a/Engine/Source/Editor/Kismet/Private/BlueprintEditor.cpp b/Engine/Source/Editor/Kismet/Private/BlueprintEditor.cpp
--- a/Engine/Source/Editor/Kismet/Private/BlueprintEditor.cpp
+++ b/Engine/Source/Editor/Kismet/Private/BlueprintEditor.cpp
@@ -608,13 +608,13 @@ void FBlueprintEditor::AnalyticsTrackCompileEvent( UBlueprint* Blueprint, int32
 void FBlueprintEditor::RefreshEditors(ERefreshBlueprintEditorReason::Type Reason)
 {
 	if( Reason != ERefreshBlueprintEditorReason::BlueprintCompiled )
-{
-	DocumentManager->CleanInvalidTabs();
+	{
+		DocumentManager->CleanInvalidTabs();
 
-	DocumentManager->RefreshAllTabs();
+		DocumentManager->RefreshAllTabs();
 
-	// The workflow manager only tracks document tabs.
-	FocusInspectorOnGraphSelection(GetSelectedNodes(), true);
+		// The workflow manager only tracks document tabs.
+		FocusInspectorOnGraphSelection(GetSelectedNodes(), true);
 	}
 
 	if (MyBlueprintWidget.IsValid())
@@ -654,24 +654,24 @@ void FBlueprintEditor::SetUISelectionState(FName SelectionOwner)
 void FBlueprintEditor::ClearSelectionStateFor(FName SelectionOwner)
 {
 	if ( SelectionOwner == SelectionState_Graph )
-{
-	TArray< TSharedPtr<SDockTab> > GraphEditorTabs;
-	DocumentManager->FindAllTabsForFactory(GraphEditorTabFactoryPtr, /*out*/ GraphEditorTabs);
-
-	for (auto GraphEditorTabIt = GraphEditorTabs.CreateIterator(); GraphEditorTabIt; ++GraphEditorTabIt)
 	{
-		TSharedRef<SGraphEditor> Editor = StaticCastSharedRef<SGraphEditor>((*GraphEditorTabIt)->GetContent());
+		TArray< TSharedPtr<SDockTab> > GraphEditorTabs;
+		DocumentManager->FindAllTabsForFactory(GraphEditorTabFactoryPtr, /*out*/ GraphEditorTabs);
 
-		Editor->ClearSelectionSet();
-	}
+		for ( auto GraphEditorTabIt = GraphEditorTabs.CreateIterator(); GraphEditorTabIt; ++GraphEditorTabIt )
+		{
+			TSharedRef<SGraphEditor> Editor = StaticCastSharedRef<SGraphEditor>(( *GraphEditorTabIt )->GetContent());
+
+			Editor->ClearSelectionSet();
+		}
 	}
 	else if ( SelectionOwner == SelectionState_Components )
 	{
-	if(SCSEditor.IsValid())
-	{
-		SCSEditor->ClearSelection();
+		if ( SCSEditor.IsValid() )
+		{
+			SCSEditor->ClearSelection();
+		}
 	}
-}
 	else if ( SelectionOwner == SelectionState_MyBlueprint )
 	{
 		if ( MyBlueprintWidget.IsValid() )
@@ -717,26 +717,6 @@ AActor* FBlueprintEditor::GetSCSEditorActorContext() const
 	return nullptr;
 }
 
-void FBlueprintEditor::OnRootSelected(AActor * DefaultActor)
-{
-	if (Inspector.IsValid())
-	{	
-		// Clear the my blueprints selection
-		SetUISelectionState(FBlueprintEditor::SelectionState_Components);
-
-		// Build an object list
-		TArray<UObject*> InspectorObjects;
-		InspectorObjects.Add(DefaultActor);
-		
-		// Update the details panel
-		FString Title;
-		DefaultActor->GetName(Title);
-		SKismetInspector::FShowDetailsOptions Options(FText::FromString(Title), true);
-		Options.bShowComponents = false;
-		Inspector->ShowDetailsForObjects(InspectorObjects, Options);
-	}
-}
-
 void FBlueprintEditor::OnSelectionUpdated(const TArray<FSCSEditorTreeNodePtrType>& SelectedNodes)
 {
 	if (SCSViewport.IsValid())
@@ -772,12 +752,25 @@ void FBlueprintEditor::OnSelectionUpdated(const TArray<FSCSEditorTreeNodePtrType
 		// Convert the selection set to an array of UObject* pointers
 		FText InspectorTitle = FText::GetEmpty();
 		TArray<UObject*> InspectorObjects;
+		bool bShowComponents = true;
 		InspectorObjects.Empty(SelectedNodes.Num());
 		for (auto NodeIt = SelectedNodes.CreateConstIterator(); NodeIt; ++NodeIt)
 		{
 			auto NodePtr = *NodeIt;
 			if (NodePtr.IsValid())
 			{
+				if (NodePtr->IsRootActor())
+				{
+					AActor* DefaultActor = GetSCSEditorActorContext();
+					InspectorObjects.Add(DefaultActor);
+					
+					FString Title; 
+					DefaultActor->GetName(Title);
+					InspectorTitle = FText::FromString(Title);
+					bShowComponents = false;
+				}
+				else
+				{
 				UActorComponent* EditableComponent = nullptr;
 				if (NodePtr->CanEditDefaults())
 				{
@@ -795,9 +788,11 @@ void FBlueprintEditor::OnSelectionUpdated(const TArray<FSCSEditorTreeNodePtrType
 				}
 			}
 		}
+		}
 
 		// Update the details panel
 		SKismetInspector::FShowDetailsOptions Options(InspectorTitle, true);
+		Options.bShowComponents = bShowComponents;
 		Inspector->ShowDetailsForObjects(InspectorObjects, Options);
 	}
 }
@@ -2003,14 +1998,14 @@ FBlueprintEditor::~FBlueprintEditor()
 void FBlueprintEditor::FocusInspectorOnGraphSelection(const FGraphPanelSelectionSet& NewSelection, bool bForceRefresh)
 {
 	if ( NewSelection.Array().Num() > 0 )
-		{
+	{
 		SetUISelectionState(FBlueprintEditor::SelectionState_Graph);
 
-			SKismetInspector::FShowDetailsOptions ShowDetailsOptions;
-			ShowDetailsOptions.bForceRefresh = bForceRefresh;
+		SKismetInspector::FShowDetailsOptions ShowDetailsOptions;
+		ShowDetailsOptions.bForceRefresh = bForceRefresh;
 
-			Inspector->ShowDetailsForObjects(NewSelection.Array(), ShowDetailsOptions);
-		}
+		Inspector->ShowDetailsForObjects(NewSelection.Array(), ShowDetailsOptions);
+	}
 }
 
 void FBlueprintEditor::CreateDefaultTabContents(const TArray<UBlueprint*>& InBlueprints)
@@ -2086,17 +2081,16 @@ void FBlueprintEditor::CreateDefaultTabContents(const TArray<UBlueprint*>& InBlu
 
 void FBlueprintEditor::CreateSCSEditors()
 {
-		SCSEditor = SAssignNew(SCSEditor, SSCSEditor)
-			.ActorContext(this, &FBlueprintEditor::GetSCSEditorActorContext)
-			.PreviewActor(this, &FBlueprintEditor::GetPreviewActor)
-			.AllowEditing(this, &FBlueprintEditor::InEditingMode)
-			.OnRootSelected(this, &FBlueprintEditor::OnRootSelected)
+	SCSEditor = SAssignNew(SCSEditor, SSCSEditor)
+		.ActorContext(this, &FBlueprintEditor::GetSCSEditorActorContext)
+		.PreviewActor(this, &FBlueprintEditor::GetPreviewActor)
+		.AllowEditing(this, &FBlueprintEditor::InEditingMode)
 		.OnSelectionUpdated(this, &FBlueprintEditor::OnSelectionUpdated)
 		.OnItemDoubleClicked(this, &FBlueprintEditor::OnComponentDoubleClicked);
 
-		SCSViewport = SAssignNew(SCSViewport, SSCSEditorViewport)
-			.BlueprintEditor(SharedThis(this));
-	}
+	SCSViewport = SAssignNew(SCSViewport, SSCSEditorViewport)
+		.BlueprintEditor(SharedThis(this));
+}
 
 void FBlueprintEditor::OnLogTokenClicked(const TSharedRef<IMessageToken>& Token)
 {
@@ -2839,7 +2833,7 @@ void FBlueprintEditor::OnSelectedNodesChanged(const FGraphPanelSelectionSet& New
 	if ( NewSelection.Num() > 0 )
 	{
 		SetUISelectionState(FBlueprintEditor::SelectionState_Graph);
-		}
+	}
 
 	Inspector->ShowDetailsForObjects(NewSelection.Array());
 }
diff --git a/Engine/Source/Editor/Kismet/Private/SCSEditorViewportClient.cpp b/Engine/Source/Editor/Kismet/Private/SCSEditorViewportClient.cpp
--- a/Engine/Source/Editor/Kismet/Private/SCSEditorViewportClient.cpp
+++ b/Engine/Source/Editor/Kismet/Private/SCSEditorViewportClient.cpp
@@ -83,7 +83,7 @@ namespace
 		{
 			// Check for a valid parent node
 			FSCSEditorTreeNodePtrType ParentNodePtr = NodePtr->GetParent();
-			if(ParentNodePtr.IsValid() && !ParentNodePtr->IsRoot())
+			if(ParentNodePtr.IsValid() && !ParentNodePtr->IsRootComponent())
 			{
 				if(SelectedNodes.Contains(ParentNodePtr))
 				{
@@ -393,7 +393,7 @@ bool FSCSEditorViewportClient::InputWidgetDelta( FViewport* Viewport, EAxisList:
 				{
 					FSCSEditorTreeNodePtrType SelectedNodePtr = *It;
 					// Don't allow editing of a root node, inherited SCS node or child node that also has a movable (non-root) parent node selected
-					const bool bCanEdit =  !SelectedNodePtr->IsRoot() && !SelectedNodePtr->IsInherited()
+					const bool bCanEdit = !SelectedNodePtr->IsRootComponent() && !SelectedNodePtr->IsInherited()
 						&& !IsMovableParentNodeSelected(SelectedNodePtr, SelectedNodes);
 
 					if(bCanEdit)
@@ -409,7 +409,7 @@ bool FSCSEditorViewportClient::InputWidgetDelta( FViewport* Viewport, EAxisList:
 
 							// Adjust the deltas as necessary
 							FComponentEditorUtils::AdjustComponentDelta(SceneComp, Drag, Rot);
-
+							
 							FComponentEditorUtils::FTransformData OldDefaultTransform(*SelectedTemplate);
 
 							TSharedPtr<ISCSEditorCustomization> Customization = BlueprintEditor->CustomizeSCSEditor(SceneComp);
@@ -567,9 +567,7 @@ FWidget::EWidgetMode FSCSEditorViewportClient::GetWidgetMode() const
 			for ( int32 CurrentNodeIndex=0; CurrentNodeIndex < SelectedNodes.Num(); CurrentNodeIndex++ )
 			{
 				FSCSEditorTreeNodePtrType CurrentNodePtr = SelectedNodes[CurrentNodeIndex];
-				if (CurrentNodePtr.IsValid() && !RootNodes.Contains(CurrentNodePtr) && !CurrentNodePtr->IsRoot() && 
-					CurrentNodePtr->GetEditableComponentTemplate(BluePrintEditor->GetBlueprintObj()) &&
-					CurrentNodePtr->FindComponentInstanceInActor(PreviewActor))
+				if (CurrentNodePtr.IsValid() && !RootNodes.Contains(CurrentNodePtr) && !CurrentNodePtr->IsRootComponent() && CurrentNodePtr->CanEditDefaults() && CurrentNodePtr->FindComponentInstanceInActor(PreviewActor))
 				{
 					// a non-NULL, non-root item is selected, draw the widget
 					ReturnWidgetMode = WidgetMode;
diff --git a/Engine/Source/Editor/Kismet/Private/SSCSEditor.cpp b/Engine/Source/Editor/Kismet/Private/SSCSEditor.cpp
--- a/Engine/Source/Editor/Kismet/Private/SSCSEditor.cpp
+++ b/Engine/Source/Editor/Kismet/Private/SSCSEditor.cpp
@@ -466,7 +466,12 @@ UBlueprint* FSCSEditorTreeNode::GetBlueprint() const
 	return NULL;
 }
 
-bool FSCSEditorTreeNode::IsRoot() const
+bool FSCSEditorTreeNode::IsRootActor() const
+{
+	return (GetSCSNode() == nullptr && GetComponentTemplate() == nullptr);
+}
+
+bool FSCSEditorTreeNode::IsRootComponent() const
 {
 	bool bIsRoot = true;
 	USCS_Node* SCS_Node = GetSCSNode();
@@ -1021,7 +1026,7 @@ void SSCS_RowWidget::Construct( const FArguments& InArgs, TSharedPtr<SSCSEditor>
 	check(InNodePtr.IsValid());
 
 	SCSEditor = InSCSEditor;
-	NodePtr = InNodePtr;
+	TreeNodePtr = InNodePtr;
 
 	
 	auto Args = FSuperRowType::FArguments()
@@ -1038,7 +1043,7 @@ SSCS_RowWidget::~SSCS_RowWidget()
 	TSharedPtr<SSCSEditor> Editor = SCSEditor.Pin();
 	if(Editor.IsValid())
 	{
-		USCS_Node* SCS_Node = NodePtr->GetSCSNode();
+		USCS_Node* SCS_Node = GetNode()->GetSCSNode();
 		if(SCS_Node != NULL && Editor->IsNodeInSimpleConstructionScript(SCS_Node))
 		{
 			SCS_Node->SetOnNameChanged(FSCSNodeNameChanged());
@@ -1049,6 +1054,8 @@ SSCS_RowWidget::~SSCS_RowWidget()
 BEGIN_SLATE_FUNCTION_BUILD_OPTIMIZATION
 TSharedRef<SWidget> SSCS_RowWidget::GenerateWidgetForColumn( const FName& ColumnName )
 {
+	FSCSEditorTreeNodePtrType NodePtr = GetNode();
+	
 	if(ColumnName == SCS_ColumnName_ComponentClass)
 	{
 		// Setup a default icon brush.
@@ -1197,7 +1204,7 @@ TSharedRef<SToolTip> SSCS_RowWidget::CreateToolTipWidget() const
 		AddToToolTipInfoBox(InfoBox, LOCTEXT("TooltipAsset", "Asset"), SNullWidget::NullWidget, TAttribute<FText>(this, &SSCS_RowWidget::GetAssetName), false);
 	}
 
-	TWeakObjectPtr<USCS_Node> SCSNode(NodePtr->GetSCSNode());
+	TWeakObjectPtr<USCS_Node> SCSNode(GetNode()->GetSCSNode());
 
 	// Add mobility
 	TSharedRef<SImage> MobilityIcon = SNew(SImage).Image(this, &SSCS_RowWidget::GetMobilityIconImage, SCSNode);
@@ -1320,6 +1327,8 @@ FText SSCS_RowWidget::GetMobilityToolTipText(TWeakObjectPtr<USCS_Node> SCSNode)
 
 FText SSCS_RowWidget::GetAssetName() const
 {
+	FSCSEditorTreeNodePtrType NodePtr = GetNode();
+
 	FText AssetName = LOCTEXT("None", "None");
 	if(NodePtr.IsValid() && NodePtr->GetComponentTemplate())
 	{
@@ -1335,6 +1344,8 @@ FText SSCS_RowWidget::GetAssetName() const
 
 FText SSCS_RowWidget::GetAssetPath() const
 {
+	FSCSEditorTreeNodePtrType NodePtr = GetNode();
+
 	FText AssetName = LOCTEXT("None", "None");
 	if(NodePtr.IsValid() && NodePtr->GetComponentTemplate())
 	{
@@ -1351,6 +1362,8 @@ FText SSCS_RowWidget::GetAssetPath() const
 
 EVisibility SSCS_RowWidget::GetAssetVisibility() const
 {
+	FSCSEditorTreeNodePtrType NodePtr = GetNode();
+
 	if(NodePtr.IsValid() && NodePtr->GetComponentTemplate() && FComponentAssetBrokerage::SupportsAssets(NodePtr->GetComponentTemplate()))
 	{
 		return EVisibility::Visible;
@@ -1363,6 +1376,7 @@ EVisibility SSCS_RowWidget::GetAssetVisibility() const
 
 FSlateColor SSCS_RowWidget::GetColorTint() const
 {
+	FSCSEditorTreeNodePtrType NodePtr = GetNode();
 	if(SCSEditor.Pin()->GetEditorMode() == SSCSEditor::EEditorMode::BlueprintSCS)
 	{
 		if(NodePtr->IsNative())
@@ -1403,6 +1417,7 @@ EVisibility SSCS_RowWidget::GetRootLabelVisibility() const
 		return EVisibility::Collapsed;
 	}
 
+	FSCSEditorTreeNodePtrType NodePtr = GetNode();
 	if(NodePtr.IsValid() && SCSEditor.IsValid() && NodePtr == SCSEditor.Pin()->SceneRootNodePtr)
 	{
 		return EVisibility::Visible;
@@ -1421,9 +1436,8 @@ TSharedPtr<SWidget> SSCS_RowWidget::BuildSceneRootDropActionMenu(FSCSEditorTreeN
 	MenuBuilder.BeginSection("SceneRootNodeDropActions", LOCTEXT("SceneRootNodeDropActionContextMenu", "Drop Actions"));
 	{
 		const FText DroppedVariableNameText = FText::FromName( DroppedNodePtr->GetVariableName() );
-		const FText NodeVariableNameText = FText::FromName( NodePtr->GetVariableName() );
+		const FText NodeVariableNameText = FText::FromName( GetNode()->GetVariableName() );
 
-		check(NodePtr.IsValid());
 		bool bDroppedInSameBlueprint = true;
 		if(SCSEditor.Pin()->GetEditorMode() == SSCSEditor::EEditorMode::BlueprintSCS)
 		{
@@ -1477,13 +1491,13 @@ FReply SSCS_RowWidget::OnDragDetected( const FGeometry& MyGeometry, const FPoint
 		TArray<TSharedPtr<FSCSEditorTreeNode>> SelectedNodePtrs = SCSEditorPtr->GetSelectedNodes();
 		if(SelectedNodePtrs.Num() == 0)
 		{
-			SelectedNodePtrs.Add(NodePtr);
+			SelectedNodePtrs.Add(GetNode());
 		}
 
 		TSharedPtr<FSCSEditorTreeNode> FirstNode = SelectedNodePtrs[0];
 		UBlueprint* Blueprint = FirstNode->GetBlueprint();
 		
-		TSharedRef<FSCSRowDragDropOp> Operation = FSCSRowDragDropOp::New(FirstNode->GetVariableName(), Blueprint != nullptr ? Blueprint->SkeletonGeneratedClass : nullptr, FNodeCreationAnalytic());
+		TSharedRef<FSCSRowDragDropOp> Operation = FSCSRowDragDropOp::New(FirstNode->IsRootActor() ? FName("Self") :FirstNode->GetVariableName(), Blueprint != nullptr ? Blueprint->SkeletonGeneratedClass : nullptr, FNodeCreationAnalytic());
 		//Operation->SetAltDrag(MouseEvent.IsAltDown());
 		//Operation->SetCtrlDrag(MouseEvent.IsLeftControlDown() || MouseEvent.IsRightControlDown());
 		Operation->SetCtrlDrag(true); // Always put a getter
@@ -1537,6 +1551,12 @@ void SSCS_RowWidget::OnDragEnter( const FGeometry& MyGeometry, const FDragDropEv
 		FSCSEditorTreeNodePtrType SceneRootNodePtr = SCSEditor.Pin()->SceneRootNodePtr;
 		check(SceneRootNodePtr.IsValid());
 
+		FSlateColor IconColor = FLinearColor::White;
+		const FSlateBrush* StatusSymbol = FEditorStyle::GetBrush(TEXT("NoBrush"));
+		FText Message;
+
+		FSCSEditorTreeNodePtrType NodePtr = GetNode();
+
 		// Validate each selected node being dragged against the node that belongs to this row. Exit the loop if we have a valid tooltip OR a valid pending drop action once all nodes in the selection have been validated.
 		for(auto SourceNodeIter = DragRowOp->SourceNodes.CreateConstIterator(); SourceNodeIter && (DragRowOp->CurrentHoverText.IsEmpty() || DragRowOp->PendingDropAction != FSCSRowDragDropOp::DropAction_None); ++SourceNodeIter)
 		{
@@ -1555,11 +1575,11 @@ void SSCS_RowWidget::OnDragEnter( const FGeometry& MyGeometry, const FDragDropEv
 				// Attempted to drag and drop onto self
 				if(DragRowOp->SourceNodes.Num() > 1)
 				{
-					DragRowOp->CurrentHoverText = FText::Format(LOCTEXT("DropActionToolTip_Error_CannotAttachToSelfWithMultipleSelection", "Cannot attach the selected components here because it would result in {0} being attached to itself. Remove it from the selection and try again."), FText::FromName(DraggedNodePtr->GetVariableName()));
+					Message = FText::Format(LOCTEXT("DropActionToolTip_Error_CannotAttachToSelfWithMultipleSelection", "Cannot attach the selected components here because it would result in {0} being attached to itself. Remove it from the selection and try again."), FText::FromName(DraggedNodePtr->GetVariableName()));
 				}
 				else
 				{
-					DragRowOp->CurrentHoverText = FText::Format(LOCTEXT("DropActionToolTip_Error_CannotAttachToSelf", "Cannot attach {0} to itself."), FText::FromName(DraggedNodePtr->GetVariableName()));
+					Message = FText::Format(LOCTEXT("DropActionToolTip_Error_CannotAttachToSelf", "Cannot attach {0} to itself."), FText::FromName(DraggedNodePtr->GetVariableName()));
 				}
 			}
 			else if(NodePtr->IsAttachedTo(DraggedNodePtr))
@@ -1567,17 +1587,17 @@ void SSCS_RowWidget::OnDragEnter( const FGeometry& MyGeometry, const FDragDropEv
 				// Attempted to drop a parent onto a child
 				if(DragRowOp->SourceNodes.Num() > 1)
 				{
-					DragRowOp->CurrentHoverText = FText::Format(LOCTEXT("DropActionToolTip_Error_CannotAttachToChildWithMultipleSelection", "Cannot attach the selected components here because it would result in {0} being attached to one of its children. Remove it from the selection and try again."), FText::FromName(DraggedNodePtr->GetVariableName()));
+					Message = FText::Format(LOCTEXT("DropActionToolTip_Error_CannotAttachToChildWithMultipleSelection", "Cannot attach the selected components here because it would result in {0} being attached to one of its children. Remove it from the selection and try again."), FText::FromName(DraggedNodePtr->GetVariableName()));
 				}
 				else
 				{
-					DragRowOp->CurrentHoverText = FText::Format(LOCTEXT("DropActionToolTip_Error_CannotAttachToChild", "Cannot attach {0} to one of its children."), FText::FromName(DraggedNodePtr->GetVariableName()));
+					Message = FText::Format(LOCTEXT("DropActionToolTip_Error_CannotAttachToChild", "Cannot attach {0} to one of its children."), FText::FromName(DraggedNodePtr->GetVariableName()));
 				}
 			}
 			else if(HoveredTemplate == NULL || DraggedTemplate == NULL)
 			{
 				// Can't attach non-USceneComponent types
-				DragRowOp->CurrentHoverText = LOCTEXT("DropActionToolTip_Error_NotAttachable", "Cannot attach to this component.");
+				Message = LOCTEXT("DropActionToolTip_Error_NotAttachable", "Cannot attach to this component.");
 			}
 			else if(NodePtr == SceneRootNodePtr)
 			{
@@ -1591,21 +1611,21 @@ void SSCS_RowWidget::OnDragEnter( const FGeometry& MyGeometry, const FDragDropEv
 				if(!NodePtr->CanReparent() && (!NodePtr->IsDefaultSceneRoot() || NodePtr->IsInherited()))
 				{
 					// Cannot make the dropped node the new root if we cannot reparent the current root
-					DragRowOp->CurrentHoverText = LOCTEXT("DropActionToolTip_Error_CannotReparentRootNode", "The root component in this Blueprint cannot be replaced.");						
+					Message = LOCTEXT("DropActionToolTip_Error_CannotReparentRootNode", "The root component in this Blueprint cannot be replaced.");
 				}
 				else if(DraggedTemplate->IsEditorOnly() && !HoveredTemplate->IsEditorOnly()) 
 				{
 					// can't have a new root that's editor-only (when children would be around in-game)
-					DragRowOp->CurrentHoverText = LOCTEXT("DropActionToolTip_Error_CannotReparentEditorOnly", "Cannot re-parent game components under editor-only ones.");
+					Message = LOCTEXT("DropActionToolTip_Error_CannotReparentEditorOnly", "Cannot re-parent game components under editor-only ones.");
 				}
 				else if(DraggedTemplate->Mobility > HoveredTemplate->Mobility)
 				{
 					// can't have a new root that's movable if the existing root is static or stationary
-					DragRowOp->CurrentHoverText = LOCTEXT("DropActionToolTip_Error_CannotReparentNonMovable", "Cannot replace a non-movable scene root with a movable component.");
+					Message = LOCTEXT("DropActionToolTip_Error_CannotReparentNonMovable", "Cannot replace a non-movable scene root with a movable component.");
 				}
 				else if(DragRowOp->SourceNodes.Num() > 1)
 				{
-					DragRowOp->CurrentHoverText = LOCTEXT("DropActionToolTip_Error_CannotAssignMultipleRootNodes", "Cannot replace the scene root with multiple components. Please select only a single component and try again.");
+					Message = LOCTEXT("DropActionToolTip_Error_CannotAssignMultipleRootNodes", "Cannot replace the scene root with multiple components. Please select only a single component and try again.");
 				}
 				else
 				{
@@ -1615,7 +1635,7 @@ void SSCS_RowWidget::OnDragEnter( const FGeometry& MyGeometry, const FDragDropEv
 				if(bCanMakeNewRoot && bCanAttachToRoot)
 				{
 					// User can choose to either attach to the current root or make the dropped node the new root
-					DragRowOp->CurrentHoverText = LOCTEXT("DropActionToolTip_AttachToOrMakeNewRoot", "Drop here to see available actions.");
+					Message = LOCTEXT("DropActionToolTip_AttachToOrMakeNewRoot", "Drop here to see available actions.");
 					DragRowOp->PendingDropAction = FSCSRowDragDropOp::DropAction_AttachToOrMakeNewRoot;
 				}
 				else if(SCSEditor.Pin()->GetEditorMode() == SSCSEditor::EEditorMode::BlueprintSCS && DraggedNodePtr->GetBlueprint() != GetBlueprint())
@@ -1623,7 +1643,7 @@ void SSCS_RowWidget::OnDragEnter( const FGeometry& MyGeometry, const FDragDropEv
 					if(bCanMakeNewRoot)
 					{
 						// Only available action is to copy the dragged node to the other Blueprint and make it the new root
-						DragRowOp->CurrentHoverText = FText::Format(LOCTEXT("DropActionToolTip_DropMakeNewRootNodeFromCopy", "Drop here to copy {0} to a new variable and make it the new root."), FText::FromName(DraggedNodePtr->GetVariableName()));
+						Message = FText::Format(LOCTEXT("DropActionToolTip_DropMakeNewRootNodeFromCopy", "Drop here to copy {0} to a new variable and make it the new root."), FText::FromName(DraggedNodePtr->GetVariableName()));
 						DragRowOp->PendingDropAction = FSCSRowDragDropOp::DropAction_MakeNewRoot;
 					}
 					else if(bCanAttachToRoot)
@@ -1631,11 +1651,11 @@ void SSCS_RowWidget::OnDragEnter( const FGeometry& MyGeometry, const FDragDropEv
 						// Only available action is to copy the dragged node(s) to the other Blueprint and attach it to the root
 						if(DragRowOp->SourceNodes.Num() > 1)
 						{
-							DragRowOp->CurrentHoverText = FText::Format(LOCTEXT("DropActionToolTip_AttachToThisNodeFromCopyWithMultipleSelection", "Drop here to copy the selected components to new variables and attach them to {0}."), FText::FromName(NodePtr->GetVariableName()));
+							Message = FText::Format(LOCTEXT("DropActionToolTip_AttachToThisNodeFromCopyWithMultipleSelection", "Drop here to copy the selected components to new variables and attach them to {0}."), FText::FromName(NodePtr->GetVariableName()));
 						}
 						else
 						{
-							DragRowOp->CurrentHoverText = FText::Format(LOCTEXT("DropActionToolTip_AttachToThisNodeFromCopy", "Drop here to copy {0} to a new variable and attach it to {1}."), FText::FromName(DraggedNodePtr->GetVariableName()), FText::FromName(NodePtr->GetVariableName()));
+							Message = FText::Format(LOCTEXT("DropActionToolTip_AttachToThisNodeFromCopy", "Drop here to copy {0} to a new variable and attach it to {1}."), FText::FromName(DraggedNodePtr->GetVariableName()), FText::FromName(NodePtr->GetVariableName()));
 						}
 						
 						DragRowOp->PendingDropAction = FSCSRowDragDropOp::DropAction_AttachTo;
@@ -1644,7 +1664,7 @@ void SSCS_RowWidget::OnDragEnter( const FGeometry& MyGeometry, const FDragDropEv
 				else if(bCanMakeNewRoot)
 				{
 					// Only available action is to make the dragged node the new root
-					DragRowOp->CurrentHoverText = FText::Format(LOCTEXT("DropActionToolTip_DropMakeNewRootNode", "Drop here to make {0} the new root."), FText::FromName(DraggedNodePtr->GetVariableName()));
+					Message = FText::Format(LOCTEXT("DropActionToolTip_DropMakeNewRootNode", "Drop here to make {0} the new root."), FText::FromName(DraggedNodePtr->GetVariableName()));
 					DragRowOp->PendingDropAction = FSCSRowDragDropOp::DropAction_MakeNewRoot;
 				}
 				else if(bCanAttachToRoot)
@@ -1652,11 +1672,11 @@ void SSCS_RowWidget::OnDragEnter( const FGeometry& MyGeometry, const FDragDropEv
 					// Only available action is to attach the dragged node(s) to the root
 					if(DragRowOp->SourceNodes.Num() > 1)
 					{
-						DragRowOp->CurrentHoverText = FText::Format(LOCTEXT("DropActionToolTip_AttachToThisNodeWithMultipleSelection", "Drop here to attach the selected components to {0}."), FText::FromName(NodePtr->GetVariableName()));
+						Message = FText::Format(LOCTEXT("DropActionToolTip_AttachToThisNodeWithMultipleSelection", "Drop here to attach the selected components to {0}."), FText::FromName(NodePtr->GetVariableName()));
 					}
 					else
 					{
-						DragRowOp->CurrentHoverText = FText::Format(LOCTEXT("DropActionToolTip_AttachToThisNode", "Drop here to attach {0} to {1}."), FText::FromName(DraggedNodePtr->GetVariableName()), FText::FromName(NodePtr->GetVariableName()));
+						Message = FText::Format(LOCTEXT("DropActionToolTip_AttachToThisNode", "Drop here to attach {0} to {1}."), FText::FromName(DraggedNodePtr->GetVariableName()), FText::FromName(NodePtr->GetVariableName()));
 					}
 
 					DragRowOp->PendingDropAction = FSCSRowDragDropOp::DropAction_AttachTo;
@@ -1667,11 +1687,11 @@ void SSCS_RowWidget::OnDragEnter( const FGeometry& MyGeometry, const FDragDropEv
 				// Detach the dropped node(s) from the current node and reattach to the root node
 				if(DragRowOp->SourceNodes.Num() > 1)
 				{
-					DragRowOp->CurrentHoverText = FText::Format(LOCTEXT("DropActionToolTip_DetachFromThisNodeWithMultipleSelection", "Drop here to detach the selected components from {0}."), FText::FromName(NodePtr->GetVariableName()));
+					Message = FText::Format(LOCTEXT("DropActionToolTip_DetachFromThisNodeWithMultipleSelection", "Drop here to detach the selected components from {0}."), FText::FromName(NodePtr->GetVariableName()));
 				}
 				else
 				{
-					DragRowOp->CurrentHoverText = FText::Format(LOCTEXT("DropActionToolTip_DetachFromThisNode", "Drop here to detach {0} from {1}."), FText::FromName(DraggedNodePtr->GetVariableName()), FText::FromName(NodePtr->GetVariableName()));
+					Message = FText::Format(LOCTEXT("DropActionToolTip_DetachFromThisNode", "Drop here to detach {0} from {1}."), FText::FromName(DraggedNodePtr->GetVariableName()), FText::FromName(NodePtr->GetVariableName()));
 				}
 				
 				DragRowOp->PendingDropAction = FSCSRowDragDropOp::DropAction_DetachFrom;
@@ -1679,17 +1699,17 @@ void SSCS_RowWidget::OnDragEnter( const FGeometry& MyGeometry, const FDragDropEv
 			else if (!DraggedTemplate->IsEditorOnly() && HoveredTemplate->IsEditorOnly()) 
 			{
 				// can't have a game component child nested under an editor-only one
-				DragRowOp->CurrentHoverText = LOCTEXT("DropActionToolTip_Error_CannotAttachToEditorOnly", "Cannot attach game components to editor-only ones.");
+				Message = LOCTEXT("DropActionToolTip_Error_CannotAttachToEditorOnly", "Cannot attach game components to editor-only ones.");
 			}
 			else if ((DraggedTemplate->Mobility == EComponentMobility::Static) && ((HoveredTemplate->Mobility == EComponentMobility::Movable) || (HoveredTemplate->Mobility == EComponentMobility::Stationary)))
 			{
 				// Can't attach Static components to mobile ones
-				DragRowOp->CurrentHoverText = LOCTEXT("DropActionToolTip_Error_CannotAttachStatic", "Cannot attach Static components to movable ones.");
+				Message = LOCTEXT("DropActionToolTip_Error_CannotAttachStatic", "Cannot attach Static components to movable ones.");
 			}
 			else if ((DraggedTemplate->Mobility == EComponentMobility::Stationary) && (HoveredTemplate->Mobility == EComponentMobility::Movable))
 			{
 				// Can't attach Static components to mobile ones
-				DragRowOp->CurrentHoverText = LOCTEXT("DropActionToolTip_Error_CannotAttachStationary", "Cannot attach Stationary components to movable ones.");
+				Message = LOCTEXT("DropActionToolTip_Error_CannotAttachStationary", "Cannot attach Stationary components to movable ones.");
 			}
 			else if(HoveredTemplate->CanAttachAsChild(DraggedTemplate, NAME_None))
 			{
@@ -1698,20 +1718,20 @@ void SSCS_RowWidget::OnDragEnter( const FGeometry& MyGeometry, const FDragDropEv
 				{
 					if(DragRowOp->SourceNodes.Num() > 1)
 					{
-						DragRowOp->CurrentHoverText = FText::Format(LOCTEXT("DropActionToolTip_AttachToThisNodeFromCopyWithMultipleSelection", "Drop here to copy the selected nodes to new variables and attach to {0}."), FText::FromName(NodePtr->GetVariableName()));
+						Message = FText::Format(LOCTEXT("DropActionToolTip_AttachToThisNodeFromCopyWithMultipleSelection", "Drop here to copy the selected nodes to new variables and attach to {0}."), FText::FromName(NodePtr->GetVariableName()));
 					}
 					else
 					{
-						DragRowOp->CurrentHoverText = FText::Format(LOCTEXT("DropActionToolTip_AttachToThisNodeFromCopy", "Drop here to copy {0} to a new variable and attach it to {1}."), FText::FromName(DraggedNodePtr->GetVariableName()), FText::FromName(NodePtr->GetVariableName()));
+						Message = FText::Format(LOCTEXT("DropActionToolTip_AttachToThisNodeFromCopy", "Drop here to copy {0} to a new variable and attach it to {1}."), FText::FromName(DraggedNodePtr->GetVariableName()), FText::FromName(NodePtr->GetVariableName()));
 					}
 				}
 				else if(DragRowOp->SourceNodes.Num() > 1)
 				{
-					DragRowOp->CurrentHoverText = FText::Format(LOCTEXT("DropActionToolTip_AttachToThisNodeWithMultipleSelection", "Drop here to attach the selected nodes to {0}."), FText::FromName(NodePtr->GetVariableName()));
+					Message = FText::Format(LOCTEXT("DropActionToolTip_AttachToThisNodeWithMultipleSelection", "Drop here to attach the selected nodes to {0}."), FText::FromName(NodePtr->GetVariableName()));
 				}
 				else
 				{
-					DragRowOp->CurrentHoverText = FText::Format(LOCTEXT("DropActionToolTip_AttachToThisNode", "Drop here to attach {0} to {1}."), FText::FromName(DraggedNodePtr->GetVariableName()), FText::FromName(NodePtr->GetVariableName()));
+					Message = FText::Format(LOCTEXT("DropActionToolTip_AttachToThisNode", "Drop here to attach {0} to {1}."), FText::FromName(DraggedNodePtr->GetVariableName()), FText::FromName(NodePtr->GetVariableName()));
 				}
 
 				DragRowOp->PendingDropAction = FSCSRowDragDropOp::DropAction_AttachTo;
@@ -1719,9 +1739,11 @@ void SSCS_RowWidget::OnDragEnter( const FGeometry& MyGeometry, const FDragDropEv
 			else
 			{
 				// The dropped node cannot be attached to the current node
-				DragRowOp->CurrentHoverText = FText::Format( LOCTEXT("DropActionToolTip_Error_TooManyAttachments", "Unable to attach {0} to {1}."), FText::FromName( DraggedNodePtr->GetVariableName() ), FText::FromName( NodePtr->GetVariableName() ) );
+				Message = FText::Format(LOCTEXT("DropActionToolTip_Error_TooManyAttachments", "Unable to attach {0} to {1}."), FText::FromName(DraggedNodePtr->GetVariableName()), FText::FromName(NodePtr->GetVariableName()));
 			}
 		}
+
+		DragRowOp->SetSimpleFeedbackMessage(StatusSymbol, FLinearColor::White, Message);
 	}
 	else if ( Operation->IsOfType<FExternalDragOperation>() || Operation->IsOfType<FAssetDragDropOp>() )
 	{
@@ -1764,8 +1786,8 @@ FReply SSCS_RowWidget::OnDrop( const FGeometry& MyGeometry, const FDragDropEvent
 	{
 		return FReply::Handled();
 	}
-
-	if ( Operation->IsOfType<FSCSRowDragDropOp>() && NodePtr->GetComponentTemplate()->IsA<USceneComponent>())	
+	
+	if (Operation->IsOfType<FSCSRowDragDropOp>() && GetNode()->GetComponentTemplate()->IsA<USceneComponent>())
 	{
 		TSharedPtr<FSCSRowDragDropOp> DragRowOp = StaticCastSharedPtr<FSCSRowDragDropOp>( Operation );	
 		check(DragRowOp.IsValid());
@@ -1815,6 +1837,8 @@ FReply SSCS_RowWidget::OnDrop( const FGeometry& MyGeometry, const FDragDropEvent
 
 void SSCS_RowWidget::OnAttachToDropAction(const TArray<FSCSEditorTreeNodePtrType>& DroppedNodePtrs)
 {
+	FSCSEditorTreeNodePtrType NodePtr = GetNode();
+
 	check(NodePtr.IsValid());
 	check(DroppedNodePtrs.Num() > 0);
 
@@ -1959,6 +1983,8 @@ void SSCS_RowWidget::OnAttachToDropAction(const TArray<FSCSEditorTreeNodePtrType
 
 void SSCS_RowWidget::OnDetachFromDropAction(const TArray<FSCSEditorTreeNodePtrType>& DroppedNodePtrs)
 {
+	FSCSEditorTreeNodePtrType NodePtr = GetNode();
+
 	check(NodePtr.IsValid());
 	check(DroppedNodePtrs.Num() > 0);
 
@@ -2067,6 +2093,8 @@ void SSCS_RowWidget::OnMakeNewRootDropAction(FSCSEditorTreeNodePtrType DroppedNo
 	// Get the current scene root node
 	FSCSEditorTreeNodePtrType& SceneRootNodePtr = SCSEditorPtr->SceneRootNodePtr;
 
+	FSCSEditorTreeNodePtrType NodePtr = GetNode();
+
 	check(NodePtr.IsValid() && NodePtr == SceneRootNodePtr);
 	check(DroppedNodePtr.IsValid());
 
@@ -2183,6 +2211,8 @@ void SSCS_RowWidget::OnMakeNewRootDropAction(FSCSEditorTreeNodePtrType DroppedNo
 
 void SSCS_RowWidget::PostDragDropAction(bool bRegenerateTreeNodes)
 {
+	FSCSEditorTreeNodePtrType NodePtr = GetNode();
+
 	TSharedPtr<SSCSEditor> PinnedEditor = SCSEditor.Pin();
 	if(PinnedEditor.IsValid())
 	{
@@ -2207,11 +2237,13 @@ void SSCS_RowWidget::PostDragDropAction(bool bRegenerateTreeNodes)
 FText SSCS_RowWidget::GetNameLabel() const
 {
 	// NOTE: Whatever this returns also becomes the variable name
-	return FText::FromString(NodePtr->GetDisplayString());
+	return FText::FromString(GetNode()->GetDisplayString());
 }
 
 FText SSCS_RowWidget::GetTooltipText() const
 {
+	FSCSEditorTreeNodePtrType NodePtr = GetNode();
+
 	if(NodePtr->IsDefaultSceneRoot())
 	{
 		if(NodePtr->IsInherited())
@@ -2280,6 +2312,7 @@ FString SSCS_RowWidget::GetDocumentationLink() const
 
 	if(SCSEditor.Pin()->GetEditorMode() == SSCSEditor::EEditorMode::BlueprintSCS)
 	{
+		FSCSEditorTreeNodePtrType NodePtr = GetNode();
 		if( NodePtr == SCSEditor.Pin()->SceneRootNodePtr
 			|| NodePtr->IsNative() || NodePtr->IsInherited())
 		{
@@ -2300,6 +2333,7 @@ FString SSCS_RowWidget::GetDocumentationExcerptName() const
 
 	if(SCSEditor.Pin()->GetEditorMode() == SSCSEditor::EEditorMode::BlueprintSCS)
 	{
+		FSCSEditorTreeNodePtrType NodePtr = GetNode();
 		if( NodePtr == SCSEditor.Pin()->SceneRootNodePtr)
 		{
 			return TEXT("RootComponent");
@@ -2329,6 +2363,8 @@ UBlueprint* SSCS_RowWidget::GetBlueprint() const
 
 bool SSCS_RowWidget::OnNameTextVerifyChanged(const FText& InNewText, FText& OutErrorMessage)
 {
+	FSCSEditorTreeNodePtrType NodePtr = GetNode();
+
 	if(!InNewText.IsEmpty())
 	{
 		if (!FComponentEditorUtils::IsValidVariableNameString(NodePtr->GetComponentTemplate(), InNewText.ToString()))
@@ -2378,7 +2414,7 @@ bool SSCS_RowWidget::OnNameTextVerifyChanged(const FText& InNewText, FText& OutE
 
 void SSCS_RowWidget::OnNameTextCommit(const FText& InNewName, ETextCommit::Type InTextCommit)
 {
-	NodePtr->OnCompleteRename(InNewName);
+	GetNode()->OnCompleteRename(InNewName);
 
 	// No need to call UpdateTree() in SCS editor mode; it will already be called by MBASM internally
 	check(SCSEditor.IsValid());
@@ -2389,6 +2425,223 @@ void SSCS_RowWidget::OnNameTextCommit(const FText& InNewName, ETextCommit::Type
 	}
 }
 
+//////////////////////////////////////////////////////////////////////////
+// SSCS_RowWidget_ActorRoot
+
+FSCSEditorTreeNodePtrType SSCS_RowWidget_ActorRoot::GetNode() const
+{
+	TSharedPtr<SSCSEditor> SCSEditorPtr = SCSEditor.Pin();
+	if (SCSEditorPtr.IsValid())
+	{
+		return SCSEditorPtr->SceneRootNodePtr;
+	}
+	return nullptr;
+}
+
+TSharedRef<SWidget> SSCS_RowWidget_ActorRoot::GenerateWidgetForColumn(const FName& ColumnName)
+{
+	// We've removed the other columns for now,  implement them for the root actor is necessary
+	ensure(ColumnName == SCS_ColumnName_ComponentClass);
+
+	return SNew(SHorizontalBox)
+		.ToolTip(CreateToolTipWidget())
+
+		+ SHorizontalBox::Slot()
+		.AutoWidth()
+		.VAlign(VAlign_Center)
+		[
+			SNew(SExpanderArrow, SharedThis(this))
+		]
+
+	+ SHorizontalBox::Slot()
+		.AutoWidth()
+		.VAlign(VAlign_Center)
+		.Padding(FMargin(0.f, 0.f, 6.f, 0.f))
+		[
+			SNew(SImage)
+			.Image(this, &SSCS_RowWidget_ActorRoot::GetActorIcon)
+		]
+
+	+ SHorizontalBox::Slot()
+		.HAlign(HAlign_Left)
+		.VAlign(VAlign_Center)
+		.Padding(0.0f, 0.0f)
+		[
+			SNew(STextBlock)
+			.Text(this, &SSCS_RowWidget_ActorRoot::GetActorDisplayText)
+			.ColorAndOpacity(FSlateColor::UseForeground())
+		];
+}
+
+TSharedRef<SToolTip> SSCS_RowWidget_ActorRoot::CreateToolTipWidget() const
+{
+	// Create a box to hold every line of info in the body of the tooltip
+	TSharedRef<SVerticalBox> InfoBox = SNew(SVerticalBox);
+
+	// Add class
+	AddToToolTipInfoBox(InfoBox, LOCTEXT("TooltipClass", "Class"), SNullWidget::NullWidget, TAttribute<FText>::Create(TAttribute<FText>::FGetter::CreateSP(this, &SSCS_RowWidget_ActorRoot::GetActorClassNameText)), false);
+
+	// Add super class
+	AddToToolTipInfoBox(InfoBox, LOCTEXT("TooltipSuperClass", "Parent Class"), SNullWidget::NullWidget, TAttribute<FText>::Create(TAttribute<FText>::FGetter::CreateSP(this, &SSCS_RowWidget_ActorRoot::GetActorSuperClassNameText)), false);
+
+	// Add mobility
+	AddToToolTipInfoBox(InfoBox, LOCTEXT("TooltipMobility", "Mobility"), SNullWidget::NullWidget, TAttribute<FText>::Create(TAttribute<FText>::FGetter::CreateSP(this, &SSCS_RowWidget_ActorRoot::GetActorMobilityText)), false);
+
+	TSharedRef<SBorder> TooltipContent = SNew(SBorder)
+		.Padding(4)
+		.BorderImage(FEditorStyle::GetBrush("SCSEditor.TileViewTooltip.NonContentBorder"))
+		[
+			SNew(SVerticalBox)
+
+			+ SVerticalBox::Slot()
+			.AutoHeight()
+			.Padding(0, 0, 0, 4)
+			[
+				SNew(SBorder)
+				.Padding(0)
+				.BorderImage(FEditorStyle::GetBrush("ToolPanel.GroupBorder"))
+				[
+					SNew(SVerticalBox)
+
+					+ SVerticalBox::Slot()
+					.AutoHeight()
+					[
+						SNew(SHorizontalBox)
+
+						+ SHorizontalBox::Slot()
+						.AutoWidth()
+						.VAlign(VAlign_Center)
+						.Padding(4)
+						[
+							SNew(STextBlock)
+							.Text(this, &SSCS_RowWidget_ActorRoot::GetActorDisplayText)
+							.Font(FEditorStyle::GetFontStyle("ContentBrowser.TileViewTooltip.NameFont"))
+						]
+					]
+				]
+			]
+
+			+ SVerticalBox::Slot()
+			.AutoHeight()
+			[
+				SNew(SBorder)
+				.Padding(4)
+				.BorderImage(FEditorStyle::GetBrush("ToolPanel.GroupBorder"))
+				[
+					InfoBox
+				]
+			]
+		];
+
+	return IDocumentation::Get()->CreateToolTip(TAttribute<FText>(this, &SSCS_RowWidget_ActorRoot::GetActorDisplayText), TooltipContent, InfoBox, TEXT(""), TEXT(""));
+}
+
+const FSlateBrush* SSCS_RowWidget_ActorRoot::GetActorIcon() const
+{
+	if (SCSEditor.IsValid())
+	{
+		TSharedPtr<SSCSEditor> SCSEditorPtr = SCSEditor.Pin();
+		if (SCSEditorPtr->ActorContext.IsSet())
+		{
+			return FClassIconFinder::FindIconForActor(SCSEditorPtr->ActorContext.Get());
+		}
+	}
+	return nullptr;
+}
+
+FText SSCS_RowWidget_ActorRoot::GetActorDisplayText() const
+{
+	if (SCSEditor.IsValid())
+	{
+		TSharedPtr<SSCSEditor> SCSEditorPtr = SCSEditor.Pin();
+		if (SCSEditorPtr->ActorContext.IsSet())
+		{
+			FString Name;
+			AActor* DefaultActor = SCSEditorPtr->ActorContext.Get();
+			UBlueprint* Blueprint = UBlueprint::GetBlueprintFromClass(DefaultActor->GetClass());
+			if (Blueprint != nullptr)
+			{
+				Blueprint->GetName(Name);
+				return FText::Format(LOCTEXT("DefaultActor_Name", "{0} (self)"), FText::FromString(Name));
+			}
+			else
+			{
+				DefaultActor->GetName(Name);
+				return FText::Format(LOCTEXT("DefaultActor_Name", "{0}"), FText::FromString(Name));
+			}
+		}
+	}
+	return FText::GetEmpty();
+}
+
+FText SSCS_RowWidget_ActorRoot::GetActorClassNameText() const
+{
+	FText Text;
+	if (SCSEditor.IsValid())
+	{
+		TSharedPtr<SSCSEditor> SCSEditorPtr = SCSEditor.Pin();
+		if (SCSEditorPtr->ActorContext.IsSet())
+		{
+			AActor* DefaultActor = SCSEditorPtr->ActorContext.Get();
+			Text = FText::FromString(DefaultActor->GetClass()->GetName());
+		}
+	}
+
+	return Text;
+}
+
+FText SSCS_RowWidget_ActorRoot::GetActorSuperClassNameText() const
+{
+	FText Text;
+	if (SCSEditor.IsValid())
+	{
+		TSharedPtr<SSCSEditor> SCSEditorPtr = SCSEditor.Pin();
+		if (SCSEditorPtr->ActorContext.IsSet())
+		{
+			AActor* DefaultActor = SCSEditorPtr->ActorContext.Get();
+			Text = FText::FromString(DefaultActor->GetClass()->GetSuperClass()->GetName());
+		}
+	}
+
+	return Text;
+}
+
+FText SSCS_RowWidget_ActorRoot::GetActorMobilityText() const
+{
+	FText Text;
+	if (SCSEditor.IsValid())
+	{
+		TSharedPtr<SSCSEditor> SCSEditorPtr = SCSEditor.Pin();
+		if (SCSEditorPtr->ActorContext.IsSet())
+		{
+			AActor* DefaultActor = SCSEditorPtr->ActorContext.Get();
+			USceneComponent* RootComponent = DefaultActor->GetRootComponent();
+			if (RootComponent)
+			{
+				if (RootComponent->Mobility == EComponentMobility::Static)
+				{
+					Text = LOCTEXT("ComponentMobility_Static", "Static");
+				}
+				else if (RootComponent->Mobility == EComponentMobility::Stationary)
+				{
+					Text = LOCTEXT("ComponentMobility_Stationary", "Stationary");
+				}
+				else if (RootComponent->Mobility == EComponentMobility::Movable)
+				{
+					Text = LOCTEXT("ComponentMobility_Movable", "Movable");
+				}
+			}
+			else
+			{
+				Text = LOCTEXT("ComponentMobility_NoRoot", "No root component, unknown mobility");
+			}
+		}
+	}
+
+	return Text;
+}
+
+
 //////////////////////////////////////////////////////////////////////////
 // SSCSEditor
 
@@ -2398,7 +2651,6 @@ void SSCSEditor::Construct( const FArguments& InArgs )
 	ActorContext = InArgs._ActorContext;
 	AllowEditing = InArgs._AllowEditing;
 	PreviewActor = InArgs._PreviewActor;
-	OnRootSelected = InArgs._OnRootSelected;
 	OnSelectionUpdated = InArgs._OnSelectionUpdated;
 	OnItemDoubleClicked = InArgs._OnItemDoubleClicked;
 	OnHighlightPropertyInDetailsView = InArgs._OnHighlightPropertyInDetailsView;
@@ -2474,30 +2726,48 @@ void SSCSEditor::Construct( const FArguments& InArgs )
 		.FillWidth(3);
 	}
 
-	SAssignNew(SCSTreeWidget, SSCSTreeType)
-		.ToolTipText(LOCTEXT("DropAssetToAddComponent", "Drop asset here to add component."))
-		.SCSEditor(this)
-		.TreeItemsSource(&RootNodes)
-		.SelectionMode(ESelectionMode::Multi)
-		.OnGenerateRow(this, &SSCSEditor::MakeTableRowWidget)
-		.OnGetChildren(this, &SSCSEditor::OnGetChildrenForTree)
-		.OnSelectionChanged(this, &SSCSEditor::OnTreeSelectionChanged)
-		.OnContextMenuOpening(this, &SSCSEditor::CreateContextMenu)
-		.OnItemScrolledIntoView(this, &SSCSEditor::OnItemScrolledIntoView)
-		.OnMouseButtonDoubleClick(this, &SSCSEditor::HandleItemDoubleClicked)
-		.ClearSelectionOnClick(InArgs._EditorMode == EEditorMode::BlueprintSCS ? true : false)
-		.ItemHeight(24)
-		.HeaderRow
-		(
-			HeaderRow
-		);
-
 	if (bSingleLayoutBPEditor)
 	{
+		SAssignNew(SCSTreeWidget, SSCSTreeType)
+			.ToolTipText(LOCTEXT("DropAssetToAddComponent", "Drop asset here to add component."))
+			.SCSEditor(this)
+			.TreeItemsSource(&RootNodes)
+			.SelectionMode(ESelectionMode::Multi)
+			.OnGenerateRow(this, &SSCSEditor::MakeTableRowWidget)
+			.OnGetChildren(this, &SSCSEditor::OnGetChildrenForTree)
+			.OnSelectionChanged(this, &SSCSEditor::OnTreeSelectionChanged)
+			.OnContextMenuOpening(this, &SSCSEditor::CreateContextMenu)
+			.OnItemScrolledIntoView(this, &SSCSEditor::OnItemScrolledIntoView)
+			.OnMouseButtonDoubleClick(this, &SSCSEditor::HandleItemDoubleClicked)
+			.ClearSelectionOnClick(InArgs._EditorMode == EEditorMode::BlueprintSCS ? true : false)
+			.ItemHeight(24)
+			.HeaderRow
+			(
+				HeaderRow
+			);
+
 		SCSTreeWidget->GetHeaderRow()->SetVisibility(EVisibility::Collapsed);
 	}
-
-	TSharedRef<SToolTip> Tooltip = CreateToolTipWidget();
+	else
+	{
+		SAssignNew(SCSTreeWidget, SSCSTreeType)
+			.ToolTipText(LOCTEXT("DropAssetToAddComponent", "Drop asset here to add component."))
+			.SCSEditor(this)
+			.TreeItemsSource(&RootComponentNodes)
+			.SelectionMode(ESelectionMode::Multi)
+			.OnGenerateRow(this, &SSCSEditor::MakeTableRowWidget)
+			.OnGetChildren(this, &SSCSEditor::OnGetChildrenForTree)
+			.OnSelectionChanged(this, &SSCSEditor::OnTreeSelectionChanged)
+			.OnContextMenuOpening(this, &SSCSEditor::CreateContextMenu)
+			.OnItemScrolledIntoView(this, &SSCSEditor::OnItemScrolledIntoView)
+			.OnMouseButtonDoubleClick(this, &SSCSEditor::HandleItemDoubleClicked)
+			.ClearSelectionOnClick(InArgs._EditorMode == EEditorMode::BlueprintSCS ? true : false)
+			.ItemHeight(24)
+			.HeaderRow
+			(
+				HeaderRow
+			);
+	}
 
 	TSharedPtr<SWidget> Contents;
 
@@ -2647,53 +2917,8 @@ void SSCSEditor::Construct( const FArguments& InArgs )
 				.Padding(2.0f)
 				.BorderImage(FEditorStyle::GetBrush("SCSEditor.TreePanel"))
 				.AddMetaData<FTagMetaData>(FTagMetaData(TEXT("ComponentsPanel")))
-				
 				[
-					SNew(SVerticalBox)
-
-					// Root Actor
-					+SVerticalBox::Slot()
-					.AutoHeight()
-					.Padding(0.0f, 0.0f)
-					[
-						SNew(SCheckBox)
-						.Style(FCoreStyle::Get(), "ToggleButtonRowStyle")
-						.IsFocusable(true)
-						.OnCheckStateChanged(this, &SSCSEditor::OnActorSelected)
-						.IsChecked(this, &SSCSEditor::OnIsActorSelected)
-						.ToolTip(Tooltip)
-						.Padding(FMargin(2.0f))
-						[
-							SNew(SHorizontalBox)
-
-							+ SHorizontalBox::Slot()
-							.AutoWidth()
-							.VAlign(VAlign_Center)
-							.Padding(FMargin(0.f, 0.f, 6.f, 0.f))
-							[
-								SNew(SImage)
-								.Image(this, &SSCSEditor::GetActorIcon)
-							]
-
-							+ SHorizontalBox::Slot()
-							.HAlign(HAlign_Left)
-							.VAlign(VAlign_Center)
-							.Padding(0.0f, 0.0f)
-							[
-								SNew(STextBlock)
-								.Text(this, &SSCSEditor::GetActorDisplayText)
-								.ColorAndOpacity(FSlateColor::UseForeground())
-							]
-						]
-					]
-
-					// Tree
-					+ SVerticalBox::Slot()
-					.FillHeight(1.0f)
-					.Padding(0.f, 0.f, 0.f, 2.f)
-					[
-						SCSTreeWidget.ToSharedRef()
-					]
+					SCSTreeWidget.ToSharedRef()
 				]
 			]
 		];
@@ -2748,100 +2973,8 @@ void SSCSEditor::Construct( const FArguments& InArgs )
 
 	// Refresh the tree widget
 	UpdateTree();
-
-	// Expand the scene root node so we show all children by default
-	if(SceneRootNodePtr.IsValid())
-	{
-		SCSTreeWidget->SetItemExpansion(SceneRootNodePtr, true);
-	}
-
-	if (bSingleLayoutBPEditor)
-	{
-		// Select the root actor
-		OnActorSelected(ECheckBoxState::Checked);
-	}
-}
-
-TSharedRef<SToolTip> SSCSEditor::CreateToolTipWidget() const
-{
-	// Create a box to hold every line of info in the body of the tooltip
-	TSharedRef<SVerticalBox> InfoBox = SNew(SVerticalBox);
-
-	// Add class
-	AddToToolTipInfoBox(InfoBox, LOCTEXT("TooltipClass", "Class"), SNullWidget::NullWidget, TAttribute<FText>::Create(TAttribute<FText>::FGetter::CreateSP(this, &SSCSEditor::GetActorClassNameText)), false);
-
-	// Add super class
-	AddToToolTipInfoBox(InfoBox, LOCTEXT("TooltipSuperClass", "Parent Class"), SNullWidget::NullWidget, TAttribute<FText>::Create(TAttribute<FText>::FGetter::CreateSP(this, &SSCSEditor::GetActorSuperClassNameText)), false);
-
-	// Add mobility
-	AddToToolTipInfoBox(InfoBox, LOCTEXT("TooltipMobility", "Mobility"), SNullWidget::NullWidget, TAttribute<FText>::Create(TAttribute<FText>::FGetter::CreateSP(this, &SSCSEditor::GetActorMobilityText)), false);
-
-	TSharedRef<SBorder> TooltipContent = SNew(SBorder)
-		.Padding(4)
-		.BorderImage(FEditorStyle::GetBrush("SCSEditor.TileViewTooltip.NonContentBorder"))
-		[
-			SNew(SVerticalBox)
-
-			+ SVerticalBox::Slot()
-			.AutoHeight()
-			.Padding(0, 0, 0, 4)
-			[
-				SNew(SBorder)
-				.Padding(0)
-				.BorderImage(FEditorStyle::GetBrush("ToolPanel.GroupBorder"))
-				[
-					SNew(SVerticalBox)
-
-					+ SVerticalBox::Slot()
-					.AutoHeight()
-					[
-						SNew(SHorizontalBox)
-
-						+ SHorizontalBox::Slot()
-						.AutoWidth()
-						.VAlign(VAlign_Center)
-						.Padding(4)
-						[
-							SNew(STextBlock)
-							.Text(this, &SSCSEditor::GetActorDisplayText)
-							.Font(FEditorStyle::GetFontStyle("ContentBrowser.TileViewTooltip.NameFont"))
-						]
-					]
-				]
-			]
-
-			+ SVerticalBox::Slot()
-			.AutoHeight()
-			[
-				SNew(SBorder)
-				.Padding(4)
-				.BorderImage(FEditorStyle::GetBrush("ToolPanel.GroupBorder"))
-				[
-					InfoBox
-				]
-			]
-		];
-
-	return IDocumentation::Get()->CreateToolTip(TAttribute<FText>(this, &SSCSEditor::GetActorDisplayText), TooltipContent, InfoBox, TEXT(""), TEXT(""));
 }
 END_SLATE_FUNCTION_BUILD_OPTIMIZATION
-/*
-FString SSCSEditor::GetDocumentationLink() const
-{
-	if (NodePtr == SCSEditor.Pin()->SceneRootNodePtr
-		|| NodePtr->IsNative() || NodePtr->IsInherited())
-	{
-		return TEXT("Shared/Editors/BlueprintEditor/ComponentsMode");
-		FText ToolTipText;
-		if (ActorContext.IsSet())
-		{
-			AActor* DefaultActor = ActorContext.Get();
-			ToolTipText = FText::FromString(DefaultActor->GetClass()->GetSuperClass()->GetName());
-		}
-	}
-
-	return TEXT("");
-}*/
 
 UBlueprint* SSCSEditor::GetBlueprint() const
 {
@@ -2869,7 +3002,11 @@ void SSCSEditor::Tick( const FGeometry& AllottedGeometry, const double InCurrent
 			for(auto NodeIt = InNodes.CreateConstIterator(); NodeIt && !bFoundInvalidNode; ++NodeIt)
 			{
 				const UActorComponent* InstancedComponent = (*NodeIt)->GetComponentTemplate();
-				bFoundInvalidNode = !InstancedComponent || InstancedComponent->IsPendingKill() || AreAnyNodesInvalidLambda((*NodeIt)->GetChildren(), ++OutNumValidNodes);
+				bFoundInvalidNode = (!(*NodeIt)->IsRootActor() && (!InstancedComponent || InstancedComponent->IsPendingKill())) || AreAnyNodesInvalidLambda((*NodeIt)->GetChildren(), OutNumValidNodes);
+				if (!(*NodeIt)->IsRootActor())
+				{
+					++OutNumValidNodes;
+				}
 			}
 
 			return bFoundInvalidNode;
@@ -2889,7 +3026,7 @@ void SSCSEditor::Tick( const FGeometry& AllottedGeometry, const double InCurrent
 		}
 
 		int32 NumComponentNodes = 0;
-		if(AreAnyNodesInvalidLambda(RootNodes, NumComponentNodes) || NumComponentNodes != NumComponentInstances)
+		if (AreAnyNodesInvalidLambda(GetRootNodes(), NumComponentNodes) || NumComponentNodes != NumComponentInstances)
 		{
 			UE_LOG(LogSCSEditor, Log, TEXT("Calling UpdateTree() from Tick()."));
 
@@ -2916,8 +3053,10 @@ TSharedRef<ITableRow> SSCSEditor::MakeTableRowWidget( FSCSEditorTreeNodePtrType
 		if(ItemName == NAME_None)
 		{
 			UActorComponent* ComponentTemplateOrInstance = InNodePtr->GetComponentTemplate();
-			check(ComponentTemplateOrInstance != nullptr);
-			ItemName = ComponentTemplateOrInstance->GetFName();
+			if (ComponentTemplateOrInstance != nullptr)
+			{
+				ItemName = ComponentTemplateOrInstance->GetFName();
+			}
 		}
 
 		if(DeferredRenameRequest == ItemName)
@@ -2933,6 +3072,13 @@ TSharedRef<ITableRow> SSCSEditor::MakeTableRowWidget( FSCSEditorTreeNodePtrType
 	{
 		TagMeta.FriendlyName = FString::Printf(TEXT("TableRow,%s,0"), *InNodePtr->GetComponentTemplate()->GetReadableName());
 	}
+
+	if (InNodePtr->IsRootActor())
+	{
+		return SNew(SSCS_RowWidget_ActorRoot, SharedThis(this), InNodePtr, OwnerTable)
+			.AddMetaData<FTutorialMetaData>(TagMeta);
+	}
+
 	return SNew(SSCS_RowWidget, SharedThis(this), InNodePtr, OwnerTable)
 		.AddMetaData<FTutorialMetaData>(TagMeta);
 }
@@ -3236,91 +3382,6 @@ void SSCSEditor::PerformComboAddClass(TSubclassOf<UActorComponent> ComponentClas
 	}
 }
 
-const FSlateBrush* SSCSEditor::GetActorIcon() const
-{
-	if (ActorContext.IsSet())
-	{
-		return FClassIconFinder::FindIconForActor(ActorContext.Get());
-	}
-	return nullptr;
-}
-
-FText SSCSEditor::GetActorDisplayText() const
-{
-	if (ActorContext.IsSet())
-	{
-		FString Name;
-		AActor* DefaultActor = ActorContext.Get();
-		UBlueprint* Blueprint = UBlueprint::GetBlueprintFromClass(DefaultActor->GetClass());
-		if (Blueprint != nullptr)
-		{
-			Blueprint->GetName(Name);
-			return FText::Format(LOCTEXT("DefaultActor_Name", "{0} (Self)"), FText::FromString(Name));
-		}
-		else
-		{
-			Name = DefaultActor->GetActorLabel();
-			return FText::Format(LOCTEXT("DefaultActor_Name", "{0} (Instance)"), FText::FromString(Name));
-		}
-	}
-	return FText::GetEmpty();	
-}
-
-FText SSCSEditor::GetActorClassNameText() const
-{
-	FText Text;
-	if (ActorContext.IsSet())
-	{
-		AActor* DefaultActor = ActorContext.Get();
-		Text = FText::FromString(DefaultActor->GetClass()->GetName());
-	}
-
-	return Text;
-}
-
-FText SSCSEditor::GetActorSuperClassNameText() const
-{
-	FText Text;
-	if (ActorContext.IsSet())
-	{
-		AActor* DefaultActor = ActorContext.Get();
-		Text = FText::FromString(DefaultActor->GetClass()->GetSuperClass()->GetName());
-	}
-
-	return Text;
-}
-
-FText SSCSEditor::GetActorMobilityText() const
-{
-	FText Text;
-	if (ActorContext.IsSet())
-	{
-		AActor* DefaultActor = ActorContext.Get();
-		USceneComponent* RootComponent = DefaultActor->GetRootComponent();
-		if (RootComponent)
-		{
-			if (RootComponent->Mobility == EComponentMobility::Static)
-			{
-				Text = LOCTEXT("ComponentMobility_Static", "Static");
-			}
-			else if (RootComponent->Mobility == EComponentMobility::Stationary)
-			{
-				Text = LOCTEXT("ComponentMobility_Stationary", "Stationary");
-			}
-			else if (RootComponent->Mobility == EComponentMobility::Movable)
-			{
-				Text = LOCTEXT("ComponentMobility_Movable", "Movable");
-			}
-		}
-		else
-		{
-			Text = LOCTEXT("ComponentMobility_NoRoot", "No root component, unknown mobility");
-		}
-	}
-
-	return Text;
-}
-
 TArray<FSCSEditorTreeNodePtrType>  SSCSEditor::GetSelectedNodes() const
 {
 	TArray<FSCSEditorTreeNodePtrType> SelectedTreeNodes = SCSTreeWidget->GetSelectedItems();
@@ -3405,9 +3466,10 @@ FSCSEditorTreeNodePtrType SSCSEditor::GetNodeFromActorComponent(const UActorComp
 		// If we have a valid component archetype instance, attempt to find a tree node that corresponds to it
 		if((EditorMode == EEditorMode::BlueprintSCS && ActorComponent->IsTemplate()) || EditorMode == EEditorMode::ActorInstance)
 		{
-			for(int32 i = 0; i < RootNodes.Num() && !NodePtr.IsValid(); i++)
+			const TArray<FSCSEditorTreeNodePtrType>& Nodes = GetRootNodes();
+			for (int32 i = 0; i < Nodes.Num() && !NodePtr.IsValid(); i++)
 			{
-				NodePtr = FindTreeNode(ActorComponent, RootNodes[i]);
+				NodePtr = FindTreeNode(ActorComponent, Nodes[i]);
 			}
 		}
 
@@ -3425,6 +3487,15 @@ FSCSEditorTreeNodePtrType SSCSEditor::GetNodeFromActorComponent(const UActorComp
 	return NodePtr;
 }
 
+void SSCSEditor::SelectRoot()
+{
+	const TArray<FSCSEditorTreeNodePtrType>& Nodes = GetRootNodes();
+	if (Nodes.Num() > 0)
+	{
+		SCSTreeWidget->SetSelection(Nodes[0]);
+	}
+}
+
 void SSCSEditor::SelectNode(FSCSEditorTreeNodePtrType InNodeToSelect, bool IsCntrlDown) 
 {
 	if(SCSTreeWidget.IsValid() && InNodeToSelect.IsValid())
@@ -3462,7 +3533,7 @@ static FSCSEditorTreeNode* FindRecursive( FSCSEditorTreeNode* Node, FName Name )
 
 void SSCSEditor::HighlightTreeNode(FName TreeNodeName, const class FPropertyPath& Property)
 {
-	for( const auto& Node : RootNodes )
+	for( const auto& Node : GetRootNodes() )
 	{
 		if( auto FoundNode = FindRecursive( Node.Get(), TreeNodeName ) )
 		{
@@ -3522,10 +3593,19 @@ void SSCSEditor::UpdateTree(bool bRegenerateTreeNodes)
 			SCSTreeWidget->ClearSelection();
 		}
 		RootNodes.Empty();
+		RootComponentNodes.Empty();
 
 		// Reset the scene root node
 		SceneRootNodePtr.Reset();
 
+		bool bSingleLayoutBPEditor = GetDefault<UEditorExperimentalSettings>()->bUnifiedBlueprintEditor;
+		if (bSingleLayoutBPEditor)
+		{
+			RootTreeNode = MakeShareable(new FSCSEditorTreeNode());
+			RootNodes.Add(RootTreeNode);
+			SCSTreeWidget->SetItemExpansion(RootTreeNode, true);
+		}
+
 		// Build the tree data source according to what mode we're in
 		if(EditorMode == EEditorMode::BlueprintSCS)
 		{
@@ -3547,26 +3627,27 @@ void SSCSEditor::UpdateTree(bool bRegenerateTreeNodes)
 
 			if(CDO != NULL)
 			{
-				// Add native ActorComponent nodes to the root set first
+				
 				TInlineComponentArray<UActorComponent*> Components;
 				CDO->GetComponents(Components);
-
+				
+				// Add the native base class SceneComponent hierarchy
 				for(auto CompIter = Components.CreateIterator(); CompIter; ++CompIter)
 				{
-					UActorComponent* ActorComp = *CompIter;
-					if (!ActorComp->IsA<USceneComponent>())
+					USceneComponent* SceneComp = Cast<USceneComponent>(*CompIter);
+					if(SceneComp != NULL)
 					{
-						RootNodes.Add(MakeShareable(new FSCSEditorTreeNode(ActorComp)));
+						AddTreeNode(SceneComp);
 					}
 				}
 
-				// Add the native base class SceneComponent hierarchy
+				// Add native ActorComponent nodes that aren't SceneComponents
 				for(auto CompIter = Components.CreateIterator(); CompIter; ++CompIter)
 				{
-					USceneComponent* SceneComp = Cast<USceneComponent>(*CompIter);
-					if(SceneComp != NULL)
+					UActorComponent* ActorComp = *CompIter;
+					if (!ActorComp->IsA<USceneComponent>())
 					{
-						AddTreeNode(SceneComp);
+						AddRootComponentTreeNode(ActorComp);
 					}
 				}
 			}
@@ -3611,24 +3692,24 @@ void SSCSEditor::UpdateTree(bool bRegenerateTreeNodes)
 				// Get the full set of instanced components
 				TInlineComponentArray<UActorComponent*> Components;
 				ActorInstance->GetComponents(Components);
-
-				// Add all non-scene component instances to the root set first
+				
+				// Now add the instanced scene component hierarchy
 				for(auto CompIter = Components.CreateIterator(); CompIter; ++CompIter)
 				{
-					UActorComponent* ActorComp = *CompIter;
-					if (!ActorComp->IsA<USceneComponent>() && !ActorComp->IsEditorOnly())
+					USceneComponent* SceneComp = Cast<USceneComponent>(*CompIter);
+					if(SceneComp != nullptr && !SceneComp->IsEditorOnly())
 					{
-						RootNodes.Add(MakeShareable(new FSCSEditorTreeNode(ActorComp)));
+						AddTreeNode(SceneComp);
 					}
 				}
 
-				// Now add the instanced scene component hierarchy
+				// Add all non-scene component instances to the root set first
 				for(auto CompIter = Components.CreateIterator(); CompIter; ++CompIter)
 				{
-					USceneComponent* SceneComp = Cast<USceneComponent>(*CompIter);
-					if(SceneComp != nullptr && !SceneComp->IsEditorOnly())
+					UActorComponent* ActorComp = *CompIter;
+					if (!ActorComp->IsA<USceneComponent>() && !ActorComp->IsEditorOnly())
 					{
-						AddTreeNode(SceneComp);
+						AddRootComponentTreeNode(ActorComp);
 					}
 				}
 			}
@@ -3649,10 +3730,17 @@ void SSCSEditor::UpdateTree(bool bRegenerateTreeNodes)
 		// Restore the previous selection state on the new tree nodes
 		for(int i = 0; i < SelectedTreeNodes.Num(); ++i)
 		{
-			FSCSEditorTreeNodePtrType NodeToSelectPtr = FindTreeNode(SelectedTreeNodes[i]->GetComponentTemplate());
-			if(NodeToSelectPtr.IsValid())
+			if (SelectedTreeNodes[i]->IsRootActor())
 			{
-				SCSTreeWidget->SetItemSelection(NodeToSelectPtr, true);
+				SCSTreeWidget->SetItemSelection(RootTreeNode, true);
+			}
+			else
+			{
+				FSCSEditorTreeNodePtrType NodeToSelectPtr = FindTreeNode(SelectedTreeNodes[i]->GetComponentTemplate());
+				if (NodeToSelectPtr.IsValid())
+				{
+					SCSTreeWidget->SetItemSelection(NodeToSelectPtr, true);
+				}
 			}
 		}
 
@@ -3671,6 +3759,34 @@ void SSCSEditor::UpdateTree(bool bRegenerateTreeNodes)
 	SCSTreeWidget->RequestTreeRefresh();
 }
 
+const TArray<FSCSEditorTreeNodePtrType>& SSCSEditor::GetRootNodes() const
+{
+	bool bSingleLayoutBPEditor = GetDefault<UEditorExperimentalSettings>()->bUnifiedBlueprintEditor;
+	if (bSingleLayoutBPEditor)
+	{
+		return RootNodes;
+	}
+	
+	return RootComponentNodes;
+}
+
+TSharedPtr<FSCSEditorTreeNode> SSCSEditor::AddRootComponentTreeNode(UActorComponent* ActorComp)
+{
+	TSharedPtr<FSCSEditorTreeNode> NewTreeNode;
+	if (RootTreeNode.IsValid())
+	{
+		NewTreeNode = RootTreeNode->AddChild(ActorComp);
+	}
+	else
+	{
+		NewTreeNode = MakeShareable(new FSCSEditorTreeNode(ActorComp));
+	}
+
+	RootComponentNodes.Add(NewTreeNode);
+
+	return NewTreeNode;
+}
+
 void SSCSEditor::ClearSelection()
 {
 	if ( bUpdatingSelection == false )
@@ -3839,20 +3955,7 @@ UActorComponent* SSCSEditor::AddNewNode(UActorComponent* NewInstanceComponent,
 	}
 	else
 	{
-		NewNodePtr = MakeShareable(new FSCSEditorTreeNode(NewInstanceComponent));
-
-		// Ensure that the root node ordering is what we assume it to be
-		check(!SceneRootNodePtr.IsValid() || (RootNodes.Num() > 0 && RootNodes[RootNodes.Num() - 1] == SceneRootNodePtr));
-
-		// Add the node to the end of the list of non-scene component nodes (just before the scene component hierarchy)
-		if(RootNodes.Num() > 0)
-		{
-			RootNodes.Insert(NewNodePtr, RootNodes.Num() - 1);
-		}
-		else
-		{
-			RootNodes.Add(NewNodePtr);
-		}
+		NewNodePtr = AddRootComponentTreeNode(NewInstanceComponent);
 	}
 
 	if(bSetFocusToNewItem)
@@ -4270,7 +4373,8 @@ void SSCSEditor::RemoveComponentNode(FSCSEditorTreeNodePtrType InNodePtr)
 		check(ActorInstance != nullptr);
 
 		UActorComponent* ComponentInstance = InNodePtr->GetComponentTemplate();
-		check(ComponentInstance != nullptr);
+		if (ComponentInstance != nullptr)
+		{
 
 		ComponentInstance->Modify();
 
@@ -4356,9 +4460,9 @@ void SSCSEditor::RemoveComponentNode(FSCSEditorTreeNodePtrType InNodePtr)
 				}
 			}
 		}
-
 		// Destroy the component instance
 		ComponentInstance->DestroyComponent();
+		}
 	}
 }
 
@@ -4384,24 +4488,6 @@ void SSCSEditor::OnTreeSelectionChanged(FSCSEditorTreeNodePtrType, ESelectInfo::
 	UpdateSelectionFromNodes(SCSTreeWidget->GetSelectedItems());
 }
 
-void SSCSEditor::OnActorSelected(const ECheckBoxState NewCheckedState)
-{
-	// Clear tree selection
-	SCSTreeWidget->ClearSelection();
-		
-	// Set the Actor selected, done after the Tree clear as it will deselect the actor.
-	bIsActorSelected = true;
-
-	// Notify that the root has been selected
-	OnRootSelected.ExecuteIfBound(ActorContext.Get(nullptr));
-}
-
-ECheckBoxState SSCSEditor::OnIsActorSelected() const
-{
-	return (bIsActorSelected) ? ECheckBoxState::Checked : ECheckBoxState::Unchecked;
-}
-
-
 bool SSCSEditor::IsNodeInSimpleConstructionScript( USCS_Node* Node ) const
 {
 	check(Node);
@@ -4476,11 +4562,20 @@ FSCSEditorTreeNodePtrType SSCSEditor::AddTreeNode(USCS_Node* InSCSNode, FSCSEdit
 		else 
 		{
 			// Create a new root node
-			NewNodePtr = MakeShareable(new FSCSEditorTreeNode(InSCSNode, bIsInherited));
-
-			// Add it to the root set
+			TSharedPtr<FSCSEditorTreeNode> NewNodePtr;
+			if (RootTreeNode.IsValid())
+			{
+				NewNodePtr = RootTreeNode->AddChild(InSCSNode, bIsInherited);
+			}
+			else
+			{
+				NewNodePtr = MakeShareable(new FSCSEditorTreeNode(InSCSNode, bIsInherited));
+			}
+			
 			NodeSCS->AddNode(InSCSNode);
-			RootNodes.Insert(NewNodePtr, 0);
+			
+			// Add it to the root set
+			RootComponentNodes.Insert(NewNodePtr, 0);
 
 			// Make it the scene root node
 			SceneRootNodePtr = NewNodePtr;
@@ -4492,9 +4587,17 @@ FSCSEditorTreeNodePtrType SSCSEditor::AddTreeNode(USCS_Node* InSCSNode, FSCSEdit
 	else
 	{
 		// If the given SCS node does not contain a scene component template, we create a new root node
-		NewNodePtr = MakeShareable(new FSCSEditorTreeNode(InSCSNode, bIsInherited));
+		TSharedPtr<FSCSEditorTreeNode> NewNodePtr;
+		if (RootTreeNode.IsValid())
+		{
+			NewNodePtr = RootTreeNode->AddChild(InSCSNode, bIsInherited);
+		}
+		else
+		{
+			NewNodePtr = MakeShareable(new FSCSEditorTreeNode(InSCSNode, bIsInherited));
+		}
 
-		RootNodes.Add(NewNodePtr);
+		RootComponentNodes.Add(NewNodePtr);
 
 		// If the SCS root node array does not already contain the given node, this will add it (this should only occur after node creation)
 		if(NodeSCS != NULL)
@@ -4543,10 +4646,7 @@ FSCSEditorTreeNodePtrType SSCSEditor::AddTreeNode(USceneComponent* InSceneCompon
 		if(!SceneRootNodePtr.IsValid())
 		{
 			// Create a new root node
-			NewNodePtr = MakeShareable(new FSCSEditorTreeNode(InSceneComponent));
-
-			// Add it to the root set
-			RootNodes.Insert(NewNodePtr, 0);
+			NewNodePtr = AddRootComponentTreeNode(InSceneComponent);
 
 			// Make it the scene root node
 			SceneRootNodePtr = NewNodePtr;
@@ -5003,7 +5103,7 @@ FReply SSCSEditor::OnPromoteToBlueprintClicked()
 
 const TArray<FSCSEditorTreeNodePtrType>& SSCSEditor::GetRootComponentNodes()
 {
-	return RootNodes;
+	return RootComponentNodes;
 }
 
 #undef LOCTEXT_NAMESPACE
diff --git a/Engine/Source/Editor/Kismet/Public/BlueprintEditor.h b/Engine/Source/Editor/Kismet/Public/BlueprintEditor.h
--- a/Engine/Source/Editor/Kismet/Public/BlueprintEditor.h
+++ b/Engine/Source/Editor/Kismet/Public/BlueprintEditor.h
@@ -388,17 +388,12 @@ public:
 	/** Delegate invoked when the SCS editor needs to obtain the Actor context for editing */
 	AActor* GetSCSEditorActorContext() const;
 
-	/** Delegate invoked when the root is selected in the SCS editor widget */
-	void OnRootSelected(class AActor * DefaultActor);
-
 	/** Delegate invoked when the selection is changed in the SCS editor widget */
 	void OnSelectionUpdated(const TArray<TSharedPtr<class FSCSEditorTreeNode>>& SelectedNodes);
 
 	/** Delegate invoked when an item is double clicked in the SCS editor widget */
 	void OnComponentDoubleClicked(TSharedPtr<class FSCSEditorTreeNode> Node);
 
-
-
 	/** Pin visibility accessors */
 	void SetPinVisibility(SGraphEditor::EPinVisibility Visibility);
 	bool GetPinVisibility(SGraphEditor::EPinVisibility Visibility) const { return PinVisibility == Visibility; }
diff --git a/Engine/Source/Editor/Kismet/Public/SSCSEditor.h b/Engine/Source/Editor/Kismet/Public/SSCSEditor.h
--- a/Engine/Source/Editor/Kismet/Public/SSCSEditor.h
+++ b/Engine/Source/Editor/Kismet/Public/SSCSEditor.h
@@ -80,10 +80,14 @@ public:
 	 * @return The Blueprint to which this node belongs.
 	 */
 	UBlueprint* GetBlueprint() const;
+	/**
+	 * @return Whether or not this object represents a root actor.
+	 */
+	bool IsRootActor() const;
 	/**
 	 * @return Whether or not this object represents a root component.
 	 */
-	bool IsRoot() const;
+	bool IsRootComponent() const;
 	/**
 	 * @return Whether or not this node is a direct child of the given node.
 	 */
@@ -270,10 +274,7 @@ public:
 	EVisibility GetRootLabelVisibility() const;
 
 	/* Get the node used by the row Widget */
-	FSCSEditorTreeNodePtrType GetNode() const
-	{
-		return NodePtr;
-	}
+	virtual FSCSEditorTreeNodePtrType GetNode() const { return TreeNodePtr; };
 
 private:
 	/** Verifies the name of the component when changing it */
@@ -327,11 +328,35 @@ public:
 	/** Pointer back to owning SCSEditor 2 tool */
 	TWeakPtr<SSCSEditor> SCSEditor;
 
-protected:
+private:
 	/** Pointer to node we represent */
-	FSCSEditorTreeNodePtrType NodePtr;
+	FSCSEditorTreeNodePtrType TreeNodePtr;
 };
 
+class SSCS_RowWidget_ActorRoot : public SSCS_RowWidget
+{
+public:
+
+	// SMultiColumnTableRow<T> interface
+	virtual TSharedRef<SWidget> GenerateWidgetForColumn(const FName& ColumnName) override;
+	// End of SMultiColumnTableRow<T>
+
+	/* Get the node used by the row Widget */
+	virtual FSCSEditorTreeNodePtrType GetNode() const override;
+
+private:
+	/** Creates a tooltip for this row */
+	TSharedRef<SToolTip> CreateToolTipWidget() const;
+
+	/** Data accessors */
+	const FSlateBrush* GetActorIcon() const;
+	FText GetActorDisplayText() const;
+	FText GetActorClassNameText() const;
+	FText GetActorSuperClassNameText() const;
+	FText GetActorMobilityText() const;
+};
+
+
 //////////////////////////////////////////////////////////////////////////
 // SSCSEditorDragDropTree - implements STreeView for our specific node type and adds drag/drop functionality
 class SSCSEditorDragDropTree : public STreeView<FSCSEditorTreeNodePtrType>
@@ -414,7 +439,6 @@ public:
 
 	DECLARE_DELEGATE_RetVal_OneParam(class USCS_Node*, FOnAddNewComponent, class UClass*);
 	DECLARE_DELEGATE_RetVal_OneParam(class USCS_Node*, FOnAddExistingComponent, class UActorComponent*);
-	DECLARE_DELEGATE_OneParam(FOnRootSelected, class AActor*);
 	DECLARE_DELEGATE_OneParam(FOnSelectionUpdated, const TArray<FSCSEditorTreeNodePtrType>&);
 	DECLARE_DELEGATE_OneParam(FOnItemDoubleClicked, const FSCSEditorTreeNodePtrType);
 	DECLARE_DELEGATE_OneParam(FOnHighlightPropertyInDetailsView, const class FPropertyPath&);
@@ -425,7 +449,6 @@ public:
 		,_PreviewActor(nullptr)
 		,_AllowEditing(true)
 		,_HideComponentClassCombo(false)
-		,_OnRootSelected()
 		,_OnSelectionUpdated()
 		,_OnHighlightPropertyInDetailsView()
 		{}
@@ -435,7 +458,6 @@ public:
 		SLATE_ATTRIBUTE(class AActor*, PreviewActor)
 		SLATE_ATTRIBUTE(bool, AllowEditing)
 		SLATE_ATTRIBUTE(bool, HideComponentClassCombo)
-		SLATE_EVENT(FOnRootSelected, OnRootSelected)
 		SLATE_EVENT(FOnSelectionUpdated, OnSelectionUpdated)
 		SLATE_EVENT(FOnItemDoubleClicked, OnItemDoubleClicked)
 		SLATE_EVENT(FOnHighlightPropertyInDetailsView, OnHighlightPropertyInDetailsView)
@@ -546,6 +568,9 @@ public:
 	 */
 	FSCSEditorTreeNodePtrType GetNodeFromActorComponent(const UActorComponent* ActorComponent, bool bIncludeAttachedComponents = true) const;
 
+	/** Select the root of the tree */
+	void SelectRoot();
+
 	/** Select the given tree node */
 	void SelectNode(FSCSEditorTreeNodePtrType InNodeToSelect, bool IsCntrlDown);
 
@@ -602,24 +627,6 @@ protected:
 	/** Add a component from the selection in the combo box */
 	void PerformComboAddClass(TSubclassOf<UActorComponent> ComponentClass);
 
-	/** Called to get an the actors icon */
-	const FSlateBrush* GetActorIcon() const;
-
-	/** Called to get an the actors display text */
-	FText GetActorDisplayText() const;
-
-	/** Called to get an the actors class name text */
-	FText GetActorClassNameText() const;
-
-	/** Called to get an the actors super class name text */
-	FText GetActorSuperClassNameText() const;
-
-	/** Called to get an the actors mobility text */
-	FText GetActorMobilityText() const;
-
-	/** Creates a tooltip for the root component*/
-	TSharedRef<SToolTip> CreateToolTipWidget() const;
-
 	/** Called to display context menu when right clicking on the widget */
 	TSharedPtr< SWidget > CreateContextMenu();
 
@@ -711,6 +718,13 @@ protected:
 
 	/** Called when the promote to blueprint button is clicked */
 	FReply OnPromoteToBlueprintClicked();
+
+	/** gets a root nodes of the tree */
+	const TArray<FSCSEditorTreeNodePtrType>& GetRootNodes() const;
+
+	/** Adds a root component tree node */
+	TSharedPtr<FSCSEditorTreeNode> AddRootComponentTreeNode(UActorComponent* ActorComp);
+
 public:
 	/** Tree widget */
 	TSharedPtr<SSCSTreeType> SCSTreeWidget;
@@ -736,9 +750,6 @@ public:
 	/** Attribute to indicate whether or not editing is allowed. */
 	TAttribute<bool> AllowEditing;
 
-	/** Delegate to invoke on selection of the root. */
-	FOnRootSelected OnRootSelected;
-
 	/** Delegate to invoke on selection update. */
 	FOnSelectionUpdated OnSelectionUpdated;
 
@@ -754,8 +765,14 @@ private:
 	/** Root set of tree */
 	TArray<FSCSEditorTreeNodePtrType> RootNodes;
 
+	/** Root set of components*/
+	TArray<FSCSEditorTreeNodePtrType> RootComponentNodes;
+
+	/* Root Tree Node*/
+	TSharedPtr<FSCSEditorTreeNode> RootTreeNode;
+
 	/** Flag to enable/disable component editing */
-	bool	bEnableComponentEditing;
+	bool bEnableComponentEditing;
 
 	/** Gate to prevent changing the selection while selection change is being broadcast. */
 	bool bUpdatingSelection;
diff --git a/Engine/Source/Editor/LevelEditor/Private/SLevelEditor.cpp b/Engine/Source/Editor/LevelEditor/Private/SLevelEditor.cpp
--- a/Engine/Source/Editor/LevelEditor/Private/SLevelEditor.cpp
+++ b/Engine/Source/Editor/LevelEditor/Private/SLevelEditor.cpp
@@ -560,7 +560,6 @@ public:
 			SAssignNew(SCSEditor, SSCSEditor)
 			.EditorMode(SSCSEditor::EEditorMode::ActorInstance)
 			.ActorContext(this, &SActorDetails::GetSelectedActor)												// Get the instance of the actor in the world
-			.OnRootSelected(this, &SActorDetails::OnSCSEditorRootSelected)										// A selection of the root has been made
 			.OnSelectionUpdated(this, &SActorDetails::OnSCSEditorTreeViewSelectionChanged)						// A selection has been made in the tree view, so inform the level editor
 			//.OnHighlightPropertyInDetailsView(this, &SLevelEditor::OnSCSEditorHighlightPropertyInDetailsView)	// Also unsure and don't think it's needed
 		);
@@ -692,10 +691,27 @@ private:
 			auto Actor = GetSelectedActor();
 			if (Actor)
 			{
+				TArray<UObject*> DetailsObjects;
+
+				bool bActorSelected = false;
+				for (auto& SelectedNode : SelectedNodes)
+				{
+					if (SelectedNode.IsValid() && SelectedNode->IsRootActor())
+					{
+						bActorSelected = true;
+						break;
+					}
+				}
+
+				if (bActorSelected)
+				{
+					DetailsObjects.Add(Actor);
+				}
+
 				USelection* SelectedComponents = GEditor->GetSelectedComponents();
 
 				// Don't bother doing anything if the node selection already matches the current world selection
-				bool bSelectionChanged = GEditor->GetSelectedComponentCount() != SelectedNodes.Num();
+				bool bSelectionChanged = GEditor->GetSelectedComponentCount() != SelectedNodes.Num() - (bActorSelected ? 1 : 1);
 				if (!bSelectionChanged)
 				{
 					// Check to see if any of the selected nodes aren't already selected in the world
@@ -711,13 +727,12 @@ private:
 					}
 				}
 				
-				if (bSelectionChanged)
+				if (bActorSelected || bSelectionChanged)
 				{
 					// Enable the selection guard to prevent OnEditorSelectionChanged() from altering the contents of the SCSTreeWidget
 					TGuardValue<bool> SelectionGuard(bSelectionGuard, true);
 
 					// Update the editor's component selection to match the node selection
-					TArray<UObject*> DetailsObjects;
 					const FScopedTransaction Transaction(NSLOCTEXT("UnrealEd", "ClickingOnComponentInTree", "Clicking on Component (tree view)"));
 
 					SelectedComponents->Modify();
@@ -804,7 +819,7 @@ private:
 		}
 		else
 		{
-			//@todo Make sure the actor "root" is selected
+			SCSEditor->SelectRoot();
 		}
 	}
 
