[INTEGRATE] Change 2417801 by [REDACTED_PERSON][REDACTED_HANDLE] on 2015/01/24 11:05:31

	Renamed blueprint creation UI from "Convert to Class Blueprint" to "Convert to Blueprint Class"

diff --git a/Engine/Source/Editor/Kismet/Private/SSCSEditor.cpp b/Engine/Source/Editor/Kismet/Private/SSCSEditor.cpp
--- a/Engine/Source/Editor/Kismet/Private/SSCSEditor.cpp
+++ b/Engine/Source/Editor/Kismet/Private/SSCSEditor.cpp
@@ -2807,8 +2807,8 @@ void SSCSEditor::Construct( const FArguments& InArgs )
 
 		EditBlueprintMenuBuilder.AddMenuEntry
 		(
-			LOCTEXT("PromoteToBlueprint", "Convert to Class Blueprint"),
-			LOCTEXT("PromoteToBluerprintTooltip","Converts the existing Blueprint into a new SubClass Blueprint" ),
+			LOCTEXT("PromoteToBlueprint", "Convert to Blueprint Class"),
+			LOCTEXT("PromoteToBluerprintTooltip","Converts the existing Blueprint into a new Blueprint class that is a subclass of the current Blueprint" ),
 			FSlateIcon(),
 			FUIAction(FExecuteAction::CreateSP(this, &SSCSEditor::PromoteToBlueprint))
 		);
@@ -2848,7 +2848,7 @@ void SSCSEditor::Construct( const FArguments& InArgs )
 						.ButtonStyle(FEditorStyle::Get(), "ToggleButton")
 						.ContentPadding(0)
 						.ToolTip(IDocumentation::Get()->CreateToolTip(
-							LOCTEXT("PromoteToBluerprintTooltip","Converts this actor into a reusable Class Blueprint that can have script behavior" ),
+							LOCTEXT("PromoteToBluerprintTooltip","Converts this actor into a reusable Blueprint class that can have script behavior" ),
 							NULL,
 							TEXT("Shared/LevelEditor"),
 							TEXT("ConvertToBlueprint")))
@@ -2868,7 +2868,7 @@ void SSCSEditor::Construct( const FArguments& InArgs )
 							[
 								SNew(STextBlock)
 								.TextStyle(FEditorStyle::Get(), "ContentBrowser.TopBar.Font")
-								.Text( LOCTEXT("PromoteToBlueprint", "Convert to Class Blueprint") )
+								.Text( LOCTEXT("PromoteToBlueprint", "Convert to Blueprint Class") )
 							]
 						]
 					]
