[INTEGRATE] Change 2417828 by [REDACTED_PERSON][REDACTED_HANDLE] on 2015/01/24 12:21:09

	Added added "(instance)" text in the actor name in the level editor case.
	This code was lost when i moved some functions.

diff --git a/Engine/Source/Editor/Kismet/Private/SSCSEditor.cpp b/Engine/Source/Editor/Kismet/Private/SSCSEditor.cpp
--- a/Engine/Source/Editor/Kismet/Private/SSCSEditor.cpp
+++ b/Engine/Source/Editor/Kismet/Private/SSCSEditor.cpp
@@ -2567,7 +2567,7 @@ FText SSCS_RowWidget_ActorRoot::GetActorDisplayText() const
 			else
 			{
 				Name = DefaultActor->GetActorLabel();
-				return FText::Format(LOCTEXT("DefaultActor_Name", "{0}"), FText::FromString(Name));
+				return FText::Format(LOCTEXT("DefaultActor_Name", "{0} (Instance)"), FText::FromString(Name));
 			}
 		}
 	}
