[INTEGRATE] Change 2417868 by [REDACTED_PERSON][REDACTED_HANDLE] on 2015/01/24 14:12:42

	No need to spend time reattaching, only attach if not already attached to the node in question.

diff --git a/Engine/Source/Editor/Kismet/Private/SSCSEditor.cpp b/Engine/Source/Editor/Kismet/Private/SSCSEditor.cpp
--- a/Engine/Source/Editor/Kismet/Private/SSCSEditor.cpp
+++ b/Engine/Source/Editor/Kismet/Private/SSCSEditor.cpp
@@ -701,7 +701,10 @@ void FSCSEditorTreeNode::AddChild(FSCSEditorTreeNodePtrType InChildNodePtr)
 		check(ParentInstance != nullptr);
 
 		// Handle attachment at the instance level
-		ChildInstance->AttachTo(ParentInstance, NAME_None, EAttachLocation::KeepWorldPosition);
+		if (ChildInstance->AttachParent != ParentInstance)
+		{
+			ChildInstance->AttachTo(ParentInstance, NAME_None, EAttachLocation::KeepWorldPosition);
+		}
 	}
 }
 
