[INTEGRATE] Change 2417898 by [REDACTED_PERSON][REDACTED_HANDLE] on 2015/01/24 14:54:28

	Make ClearInstanceComponents take a boolean indicating if the instances should be destroyed

diff --git a/Engine/Source/Editor/Kismet/Private/SSCSEditor.cpp b/Engine/Source/Editor/Kismet/Private/SSCSEditor.cpp
--- a/Engine/Source/Editor/Kismet/Private/SSCSEditor.cpp
+++ b/Engine/Source/Editor/Kismet/Private/SSCSEditor.cpp
@@ -4976,7 +4976,7 @@ void SSCSEditor::OnApplyChangesToBlueprint() const
 				{
 					FKismetEditorUtilities::AddComponentsToBlueprint(Blueprint, Actor->GetInstanceComponents());
 					NumChangedProperties += Actor->GetInstanceComponents().Num();
-					Actor->ClearInstanceComponents();
+					Actor->ClearInstanceComponents(true);
 				}
 				if (NumChangedProperties > 0)
 				{
@@ -5041,18 +5041,8 @@ void SSCSEditor::OnResetToBlueprintDefaults() const
 				const auto CopyOptions = (EditorUtilities::ECopyOptions::Type)(EditorUtilities::ECopyOptions::OnlyCopyEditOrInterpProperties | EditorUtilities::ECopyOptions::CallPostEditChangeProperty);
 				NumChangedProperties = EditorUtilities::CopyActorProperties(BlueprintCDO, Actor, CopyOptions);
 			}
-			if (Actor->GetInstanceComponents().Num() > 0)
-			{
-				TArray<UActorComponent*> InstanceComponents = Actor->GetInstanceComponents();
-				NumChangedProperties += InstanceComponents.Num();
-				for (UActorComponent* ActorComponent : InstanceComponents)
-				{
-					if (ActorComponent)
-					{
-						ActorComponent->DestroyComponent();
-					}
-				}
-			}
+			NumChangedProperties += Actor->GetInstanceComponents().Num();
+			Actor->ClearInstanceComponents(true);
 		}
 
 		// Set up a notification record to indicate success/failure
diff --git a/Engine/Source/Editor/UnrealEd/Private/Kismet2/Kismet2.cpp b/Engine/Source/Editor/UnrealEd/Private/Kismet2/Kismet2.cpp
--- a/Engine/Source/Editor/UnrealEd/Private/Kismet2/Kismet2.cpp
+++ b/Engine/Source/Editor/UnrealEd/Private/Kismet2/Kismet2.cpp
@@ -1037,7 +1037,7 @@ UBlueprint* FKismetEditorUtilities::CreateBlueprintFromActor(const FName Bluepri
 				// Since we already created SCS Nodes for the instance components, temporarily cache and clear the
 				// array to avoid creating duplicates in the new CDO
 				const TArray<UActorComponent*> TempInstanceComponents(Actor->GetInstanceComponents());
-				Actor->ClearInstanceComponents();
+				Actor->ClearInstanceComponents(false);
 
 				UObject* CDO = NewBlueprint->GeneratedClass->GetDefaultObject();
 				UEditorEngine::CopyPropertiesForUnrelatedObjects(Actor, CDO);
diff --git a/Engine/Source/Runtime/Engine/Classes/GameFramework/Actor.h b/Engine/Source/Runtime/Engine/Classes/GameFramework/Actor.h
--- a/Engine/Source/Runtime/Engine/Classes/GameFramework/Actor.h
+++ b/Engine/Source/Runtime/Engine/Classes/GameFramework/Actor.h
@@ -2241,7 +2241,7 @@ public:
 	void RemoveInstanceComponent(UActorComponent* Component);
 
 	/** Clears the instance components array */
-	void ClearInstanceComponents();
+	void ClearInstanceComponents(bool bDestroyComponents);
 
 	/** Returns the instance components array */
 	const TArray<UActorComponent*>& GetInstanceComponents() const;
diff --git a/Engine/Source/Runtime/Engine/Private/Actor.cpp b/Engine/Source/Runtime/Engine/Private/Actor.cpp
--- a/Engine/Source/Runtime/Engine/Private/Actor.cpp
+++ b/Engine/Source/Runtime/Engine/Private/Actor.cpp
@@ -2060,9 +2060,23 @@ void AActor::RemoveInstanceComponent(UActorComponent* Component)
 	InstanceComponents.Remove(Component);
 }
 
-void AActor::ClearInstanceComponents()
+void AActor::ClearInstanceComponents(const bool bDestroyComponents)
 {
-	InstanceComponents.Empty();
+	if (bDestroyComponents)
+	{
+		// Need to cache because calling destroy will remove them from InstanceComponents
+		TArray<UActorComponent*> CachedComponents(InstanceComponents);
+
+		// Run in reverse to reduce memory churn when the components are removed from InstanceComponents
+		for (int32 Index=CachedComponents.Num()-1; Index >= 0; --Index)
+		{
+			CachedComponents[Index]->DestroyComponent();
+		}
+	}
+	else
+	{
+		InstanceComponents.Empty();
+	}
 }
 
 UActorComponent* AActor::FindComponentByClass(const TSubclassOf<UActorComponent> ComponentClass) const
