[INTEGRATE] Change 2417942 by [REDACTED_PERSON][REDACTED_HANDLE] on 2015/01/24 16:47:33

	Fixed tooltip text not showing correctly when dragging and dropping instanced components.

diff --git a/Engine/Source/Editor/Kismet/Private/SSCSEditor.cpp b/Engine/Source/Editor/Kismet/Private/SSCSEditor.cpp
--- a/Engine/Source/Editor/Kismet/Private/SSCSEditor.cpp
+++ b/Engine/Source/Editor/Kismet/Private/SSCSEditor.cpp
@@ -1553,8 +1553,9 @@ void SSCS_RowWidget::OnDragEnter( const FGeometry& MyGeometry, const FDragDropEv
 			return;
 		}
 
-		check(SCSEditor.IsValid());
-		FSCSEditorTreeNodePtrType SceneRootNodePtr = SCSEditor.Pin()->SceneRootNodePtr;
+		TSharedPtr<SSCSEditor> SCSEditorPin = SCSEditor.Pin();
+		check(SCSEditorPin.IsValid());
+		FSCSEditorTreeNodePtrType SceneRootNodePtr = SCSEditorPin->SceneRootNodePtr;
 		check(SceneRootNodePtr.IsValid());
 
 		FSlateColor IconColor = FLinearColor::White;
@@ -1569,6 +1570,10 @@ void SSCS_RowWidget::OnDragEnter( const FGeometry& MyGeometry, const FDragDropEv
 			FSCSEditorTreeNodePtrType DraggedNodePtr = *SourceNodeIter;
 			check(DraggedNodePtr.IsValid());
 
+
+			FText DraggedNodePtrName = SCSEditorPin->GetEditorMode() == SSCSEditor::ActorInstance ? FText::FromName( DraggedNodePtr->GetComponentTemplate()->GetFName() ) : FText::FromName(DraggedNodePtr->GetVariableName());
+			FText CurrentNodePtrName = SCSEditorPin->GetEditorMode() == SSCSEditor::ActorInstance ? FText::FromName( NodePtr->GetComponentTemplate()->GetFName() ) : FText::FromName(NodePtr->GetVariableName());
+
 			// Reset the pending drop action each time through the loop
 			DragRowOp->PendingDropAction = FSCSRowDragDropOp::DropAction_None;
 
@@ -1581,11 +1586,11 @@ void SSCS_RowWidget::OnDragEnter( const FGeometry& MyGeometry, const FDragDropEv
 				// Attempted to drag and drop onto self
 				if(DragRowOp->SourceNodes.Num() > 1)
 				{
-					Message = FText::Format(LOCTEXT("DropActionToolTip_Error_CannotAttachToSelfWithMultipleSelection", "Cannot attach the selected components here because it would result in {0} being attached to itself. Remove it from the selection and try again."), FText::FromName(DraggedNodePtr->GetVariableName()));
+					Message = FText::Format(LOCTEXT("DropActionToolTip_Error_CannotAttachToSelfWithMultipleSelection", "Cannot attach the selected components here because it would result in {0} being attached to itself. Remove it from the selection and try again."), DraggedNodePtrName);
 				}
 				else
 				{
-					Message = FText::Format(LOCTEXT("DropActionToolTip_Error_CannotAttachToSelf", "Cannot attach {0} to itself."), FText::FromName(DraggedNodePtr->GetVariableName()));
+					Message = FText::Format(LOCTEXT("DropActionToolTip_Error_CannotAttachToSelf", "Cannot attach {0} to itself."), DraggedNodePtrName);
 				}
 			}
 			else if(NodePtr->IsAttachedTo(DraggedNodePtr))
@@ -1593,11 +1598,11 @@ void SSCS_RowWidget::OnDragEnter( const FGeometry& MyGeometry, const FDragDropEv
 				// Attempted to drop a parent onto a child
 				if(DragRowOp->SourceNodes.Num() > 1)
 				{
-					Message = FText::Format(LOCTEXT("DropActionToolTip_Error_CannotAttachToChildWithMultipleSelection", "Cannot attach the selected components here because it would result in {0} being attached to one of its children. Remove it from the selection and try again."), FText::FromName(DraggedNodePtr->GetVariableName()));
+					Message = FText::Format(LOCTEXT("DropActionToolTip_Error_CannotAttachToChildWithMultipleSelection", "Cannot attach the selected components here because it would result in {0} being attached to one of its children. Remove it from the selection and try again."), DraggedNodePtrName);
 				}
 				else
 				{
-					Message = FText::Format(LOCTEXT("DropActionToolTip_Error_CannotAttachToChild", "Cannot attach {0} to one of its children."), FText::FromName(DraggedNodePtr->GetVariableName()));
+					Message = FText::Format(LOCTEXT("DropActionToolTip_Error_CannotAttachToChild", "Cannot attach {0} to one of its children."), DraggedNodePtrName);
 				}
 			}
 			else if(HoveredTemplate == NULL || DraggedTemplate == NULL)
@@ -1649,7 +1654,7 @@ void SSCS_RowWidget::OnDragEnter( const FGeometry& MyGeometry, const FDragDropEv
 					if(bCanMakeNewRoot)
 					{
 						// Only available action is to copy the dragged node to the other Blueprint and make it the new root
-						Message = FText::Format(LOCTEXT("DropActionToolTip_DropMakeNewRootNodeFromCopy", "Drop here to copy {0} to a new variable and make it the new root."), FText::FromName(DraggedNodePtr->GetVariableName()));
+						Message = FText::Format(LOCTEXT("DropActionToolTip_DropMakeNewRootNodeFromCopy", "Drop here to copy {0} to a new variable and make it the new root."), DraggedNodePtrName);
 						DragRowOp->PendingDropAction = FSCSRowDragDropOp::DropAction_MakeNewRoot;
 					}
 					else if(bCanAttachToRoot)
@@ -1657,11 +1662,11 @@ void SSCS_RowWidget::OnDragEnter( const FGeometry& MyGeometry, const FDragDropEv
 						// Only available action is to copy the dragged node(s) to the other Blueprint and attach it to the root
 						if(DragRowOp->SourceNodes.Num() > 1)
 						{
-							Message = FText::Format(LOCTEXT("DropActionToolTip_AttachToThisNodeFromCopyWithMultipleSelection", "Drop here to copy the selected components to new variables and attach them to {0}."), FText::FromName(NodePtr->GetVariableName()));
+							Message = FText::Format(LOCTEXT("DropActionToolTip_AttachToThisNodeFromCopyWithMultipleSelection", "Drop here to copy the selected components to new variables and attach them to {0}."), CurrentNodePtrName);
 						}
 						else
 						{
-							Message = FText::Format(LOCTEXT("DropActionToolTip_AttachToThisNodeFromCopy", "Drop here to copy {0} to a new variable and attach it to {1}."), FText::FromName(DraggedNodePtr->GetVariableName()), FText::FromName(NodePtr->GetVariableName()));
+							Message = FText::Format(LOCTEXT("DropActionToolTip_AttachToThisNodeFromCopy", "Drop here to copy {0} to a new variable and attach it to {1}."), DraggedNodePtrName, CurrentNodePtrName);
 						}
 						
 						DragRowOp->PendingDropAction = FSCSRowDragDropOp::DropAction_AttachTo;
@@ -1670,7 +1675,7 @@ void SSCS_RowWidget::OnDragEnter( const FGeometry& MyGeometry, const FDragDropEv
 				else if(bCanMakeNewRoot)
 				{
 					// Only available action is to make the dragged node the new root
-					Message = FText::Format(LOCTEXT("DropActionToolTip_DropMakeNewRootNode", "Drop here to make {0} the new root."), FText::FromName(DraggedNodePtr->GetVariableName()));
+					Message = FText::Format(LOCTEXT("DropActionToolTip_DropMakeNewRootNode", "Drop here to make {0} the new root."),DraggedNodePtrName);
 					DragRowOp->PendingDropAction = FSCSRowDragDropOp::DropAction_MakeNewRoot;
 				}
 				else if(bCanAttachToRoot)
@@ -1678,11 +1683,11 @@ void SSCS_RowWidget::OnDragEnter( const FGeometry& MyGeometry, const FDragDropEv
 					// Only available action is to attach the dragged node(s) to the root
 					if(DragRowOp->SourceNodes.Num() > 1)
 					{
-						Message = FText::Format(LOCTEXT("DropActionToolTip_AttachToThisNodeWithMultipleSelection", "Drop here to attach the selected components to {0}."), FText::FromName(NodePtr->GetVariableName()));
+						Message = FText::Format(LOCTEXT("DropActionToolTip_AttachToThisNodeWithMultipleSelection", "Drop here to attach the selected components to {0}."), CurrentNodePtrName);
 					}
 					else
 					{
-						Message = FText::Format(LOCTEXT("DropActionToolTip_AttachToThisNode", "Drop here to attach {0} to {1}."), FText::FromName(DraggedNodePtr->GetVariableName()), FText::FromName(NodePtr->GetVariableName()));
+						Message = FText::Format(LOCTEXT("DropActionToolTip_AttachToThisNode", "Drop here to attach {0} to {1}."), DraggedNodePtrName, CurrentNodePtrName);
 					}
 
 					DragRowOp->PendingDropAction = FSCSRowDragDropOp::DropAction_AttachTo;
@@ -1693,11 +1698,11 @@ void SSCS_RowWidget::OnDragEnter( const FGeometry& MyGeometry, const FDragDropEv
 				// Detach the dropped node(s) from the current node and reattach to the root node
 				if(DragRowOp->SourceNodes.Num() > 1)
 				{
-					Message = FText::Format(LOCTEXT("DropActionToolTip_DetachFromThisNodeWithMultipleSelection", "Drop here to detach the selected components from {0}."), FText::FromName(NodePtr->GetVariableName()));
+					Message = FText::Format(LOCTEXT("DropActionToolTip_DetachFromThisNodeWithMultipleSelection", "Drop here to detach the selected components from {0}."), CurrentNodePtrName);
 				}
 				else
 				{
-					Message = FText::Format(LOCTEXT("DropActionToolTip_DetachFromThisNode", "Drop here to detach {0} from {1}."), FText::FromName(DraggedNodePtr->GetVariableName()), FText::FromName(NodePtr->GetVariableName()));
+					Message = FText::Format(LOCTEXT("DropActionToolTip_DetachFromThisNode", "Drop here to detach {0} from {1}."), DraggedNodePtrName, CurrentNodePtrName);
 				}
 				
 				DragRowOp->PendingDropAction = FSCSRowDragDropOp::DropAction_DetachFrom;
@@ -1724,20 +1729,20 @@ void SSCS_RowWidget::OnDragEnter( const FGeometry& MyGeometry, const FDragDropEv
 				{
 					if(DragRowOp->SourceNodes.Num() > 1)
 					{
-						Message = FText::Format(LOCTEXT("DropActionToolTip_AttachToThisNodeFromCopyWithMultipleSelection", "Drop here to copy the selected nodes to new variables and attach to {0}."), FText::FromName(NodePtr->GetVariableName()));
+						Message = FText::Format(LOCTEXT("DropActionToolTip_AttachToThisNodeFromCopyWithMultipleSelection", "Drop here to copy the selected nodes to new variables and attach to {0}."), CurrentNodePtrName);
 					}
 					else
 					{
-						Message = FText::Format(LOCTEXT("DropActionToolTip_AttachToThisNodeFromCopy", "Drop here to copy {0} to a new variable and attach it to {1}."), FText::FromName(DraggedNodePtr->GetVariableName()), FText::FromName(NodePtr->GetVariableName()));
+						Message = FText::Format(LOCTEXT("DropActionToolTip_AttachToThisNodeFromCopy", "Drop here to copy {0} to a new variable and attach it to {1}."), DraggedNodePtrName, CurrentNodePtrName);
 					}
 				}
 				else if(DragRowOp->SourceNodes.Num() > 1)
 				{
-					Message = FText::Format(LOCTEXT("DropActionToolTip_AttachToThisNodeWithMultipleSelection", "Drop here to attach the selected nodes to {0}."), FText::FromName(NodePtr->GetVariableName()));
+					Message = FText::Format(LOCTEXT("DropActionToolTip_AttachToThisNodeWithMultipleSelection", "Drop here to attach the selected nodes to {0}."), CurrentNodePtrName);
 				}
 				else
 				{
-					Message = FText::Format(LOCTEXT("DropActionToolTip_AttachToThisNode", "Drop here to attach {0} to {1}."), FText::FromName(DraggedNodePtr->GetVariableName()), FText::FromName(NodePtr->GetVariableName()));
+					Message = FText::Format(LOCTEXT("DropActionToolTip_AttachToThisNode", "Drop here to attach {0} to {1}."), DraggedNodePtrName, CurrentNodePtrName);
 				}
 
 				DragRowOp->PendingDropAction = FSCSRowDragDropOp::DropAction_AttachTo;
@@ -1745,7 +1750,7 @@ void SSCS_RowWidget::OnDragEnter( const FGeometry& MyGeometry, const FDragDropEv
 			else
 			{
 				// The dropped node cannot be attached to the current node
-				Message = FText::Format(LOCTEXT("DropActionToolTip_Error_TooManyAttachments", "Unable to attach {0} to {1}."), FText::FromName(DraggedNodePtr->GetVariableName()), FText::FromName(NodePtr->GetVariableName()));
+				Message = FText::Format(LOCTEXT("DropActionToolTip_Error_TooManyAttachments", "Unable to attach {0} to {1}."), DraggedNodePtrName, CurrentNodePtrName);
 			}
 		}
 
