[INTEGRATE] Change 2418107 by [REDACTED_PERSON][REDACTED_HANDLE] on 2015/01/25 08:41:32

	misc. component editor fixes for IWCE

	change summary:
	- fixed SSCSEditor::AddTreeNode() for SCS nodes
	- fixed FSCSEditorTreeNode::IsActorRoot() to work w/ instanced nodes
	- fixed crash in child promotion when deleting a root node w/ only 1 child
	- improved post-deletion selection update for instanced component editing

diff --git a/Engine/Source/Editor/Kismet/Private/SSCSEditor.cpp b/Engine/Source/Editor/Kismet/Private/SSCSEditor.cpp
--- a/Engine/Source/Editor/Kismet/Private/SSCSEditor.cpp
+++ b/Engine/Source/Editor/Kismet/Private/SSCSEditor.cpp
@@ -249,6 +249,7 @@ FSCSEditorTreeNode::FSCSEditorTreeNode(UActorComponent* InComponentTemplate)
 		bWasInstancedFromNativeClass = false;
 		
 		InstancedComponentName = InComponentTemplate->GetFName();
+		check(InstancedComponentName != NAME_None);	// ...otherwise IsRootActor() can return a false positive.
 
 		ComponentTemplatePtr.Reset();
 		InstancedComponentOwnerPtr = Owner;
@@ -471,7 +472,7 @@ UBlueprint* FSCSEditorTreeNode::GetBlueprint() const
 
 bool FSCSEditorTreeNode::IsRootActor() const
 {
-	return (GetSCSNode() == nullptr && GetComponentTemplate() == nullptr);
+	return (GetSCSNode() == nullptr && InstancedComponentName == NAME_None && GetComponentTemplate() == nullptr);
 }
 
 bool FSCSEditorTreeNode::IsRootComponent() const
@@ -3760,22 +3761,29 @@ void SSCSEditor::UpdateTree(bool bRegenerateTreeNodes)
 			}
 		}
 
-		// Restore the previous selection state on the new tree nodes
-		for(int i = 0; i < SelectedTreeNodes.Num(); ++i)
+		if(SelectedTreeNodes.Num() > 0)
 		{
-			if (SelectedTreeNodes[i]->IsRootActor())
+			// Restore the previous selection state on the new tree nodes
+			for(int i = 0; i < SelectedTreeNodes.Num(); ++i)
 			{
-				SCSTreeWidget->SetItemSelection(RootTreeNode, true);
-			}
-			else
-			{
-				FSCSEditorTreeNodePtrType NodeToSelectPtr = FindTreeNode(SelectedTreeNodes[i]->GetComponentTemplate());
-				if (NodeToSelectPtr.IsValid())
+				if (SelectedTreeNodes[i]->IsRootActor())
 				{
-					SCSTreeWidget->SetItemSelection(NodeToSelectPtr, true);
+					SCSTreeWidget->SetItemSelection(RootTreeNode, true);
+				}
+				else
+				{
+					FSCSEditorTreeNodePtrType NodeToSelectPtr = FindTreeNode(SelectedTreeNodes[i]->GetComponentTemplate());
+					if (NodeToSelectPtr.IsValid())
+					{
+						SCSTreeWidget->SetItemSelection(NodeToSelectPtr, true);
+					}
 				}
 			}
 		}
+		else if(SceneRootNodePtr.IsValid())
+		{
+			SCSTreeWidget->SetItemSelection(SceneRootNodePtr, true);
+		}
 
 		// If we have a pending deferred rename request, redirect it to the new tree node
 		if(DeferredRenameRequest != NAME_None)
@@ -4339,7 +4347,7 @@ void SSCSEditor::OnDeleteNodes()
 			{
 				// Default to the parent node
 				NewSelection = Node->GetParent();
-				if(NewSelection.IsValid())
+				if(NewSelection.IsValid() && !NewSelection->IsRootActor())
 				{
 					// If we have sibling nodes, find the one that immediately precedes the one being removed
 					const TArray<FSCSEditorTreeNodePtrType>& ChildNodes = NewSelection->GetChildren();
@@ -4348,6 +4356,24 @@ void SSCSEditor::OnDeleteNodes()
 						NewSelection = ChildNodes[ChildIndex];
 					}
 				}
+				else if(!Node->GetComponentTemplate()->IsA<USceneComponent>())
+				{
+					int32 NodeIndex = RootComponentNodes.Find(Node);
+					if(NodeIndex != INDEX_NONE)
+					{
+						// We should always have a valid scene root component, so there should be at least 2 root nodes in this case
+						check(RootComponentNodes.Num() > 1);
+						NewSelection = RootComponentNodes[NodeIndex > 0 ? NodeIndex - 1 : NodeIndex + 1];
+					}
+					else
+					{
+						NewSelection.Reset();
+					}
+				}
+				else
+				{
+					NewSelection.Reset();
+				}
 			}
 
 			// This will clear the current selection
@@ -4408,93 +4434,96 @@ void SSCSEditor::RemoveComponentNode(FSCSEditorTreeNodePtrType InNodePtr)
 		UActorComponent* ComponentInstance = InNodePtr->GetComponentTemplate();
 		if (ComponentInstance != nullptr)
 		{
+			ComponentInstance->Modify();
 
-		ComponentInstance->Modify();
-
-		// Handle removal of a node within the scene component hierarchy
-		USceneComponent* SceneComponent = Cast<USceneComponent>(ComponentInstance);
-		if(SceneComponent != nullptr)
-		{
-			// Find an appropriate child node to promote to this node's position in the hierarchy
-			USceneComponent* ChildToPromote = nullptr;
-			if(SceneComponent->AttachChildren.Num() > 0)
+			// Handle removal of a node within the scene component hierarchy
+			USceneComponent* SceneComponent = Cast<USceneComponent>(ComponentInstance);
+			if(SceneComponent != nullptr)
 			{
-				// Start with the first child node
-				ChildToPromote = SceneComponent->AttachChildren[0];
-				check(ChildToPromote != nullptr);
-
-				// Always choose non editor-only child nodes over editor-only child nodes (since we don't want editor-only nodes to end up with non editor-only child nodes)
-				if(ChildToPromote->IsEditorOnly())
+				// Find an appropriate child node to promote to this node's position in the hierarchy
+				USceneComponent* ChildToPromote = nullptr;
+				if(SceneComponent->AttachChildren.Num() > 0)
 				{
-					for(int32 ChildIndex = 1; ChildIndex < SceneComponent->AttachChildren.Num(); ++ChildIndex)
+					// Start with the first child node
+					ChildToPromote = SceneComponent->AttachChildren[0];
+					check(ChildToPromote != nullptr);
+
+					// Always choose non editor-only child nodes over editor-only child nodes (since we don't want editor-only nodes to end up with non editor-only child nodes)
+					if(ChildToPromote->IsEditorOnly())
 					{
-						USceneComponent* Child = SceneComponent->AttachChildren[ChildIndex];
-						if (Child != nullptr && !Child->IsEditorOnly())
+						for(int32 ChildIndex = 1; ChildIndex < SceneComponent->AttachChildren.Num(); ++ChildIndex)
 						{
-							ChildToPromote = Child;
-							break;
+							USceneComponent* Child = SceneComponent->AttachChildren[ChildIndex];
+							if (Child != nullptr && !Child->IsEditorOnly())
+							{
+								ChildToPromote = Child;
+								break;
+							}
 						}
 					}
 				}
-			}
 
-			// Handle removal of the root node
-			if(SceneComponent == ActorInstance->GetRootComponent())
-			{
-				// We only promote non editor-only components to root in instanced mode
-				if(ChildToPromote == nullptr || ChildToPromote->IsEditorOnly())
+				// Handle removal of the root node
+				if(SceneComponent == ActorInstance->GetRootComponent())
 				{
-					// Construct a new default root component
-					USceneComponent* NewRootComponent = ConstructObject<USceneComponent>(USceneComponent::StaticClass(), ActorInstance, FComponentEditorUtils::GetDefaultSceneRootVariableName(), RF_Transactional);
-					NewRootComponent->Mobility = SceneComponent->Mobility;
-					NewRootComponent->SetWorldLocationAndRotation(SceneComponent->GetComponentLocation(), SceneComponent->GetComponentRotation());
-
-					// Designate the new default root as the child we're promoting
-					ChildToPromote = NewRootComponent;
-				}
+					// We only promote non editor-only components to root in instanced mode
+					if(ChildToPromote == nullptr || ChildToPromote->IsEditorOnly())
+					{
+						// Construct a new default root component
+						USceneComponent* NewRootComponent = ConstructObject<USceneComponent>(USceneComponent::StaticClass(), ActorInstance, FComponentEditorUtils::GetDefaultSceneRootVariableName(), RF_Transactional);
+						NewRootComponent->Mobility = SceneComponent->Mobility;
+						NewRootComponent->SetWorldLocationAndRotation(SceneComponent->GetComponentLocation(), SceneComponent->GetComponentRotation());
 
-				ActorInstance->Modify();
+						// Designate the new default root as the child we're promoting
+						ChildToPromote = NewRootComponent;
+					}
 
-				// Set the selected child node as the new root
-				check(ChildToPromote != nullptr);
-				ActorInstance->SetRootComponent(ChildToPromote);
-			}
-			else    // ...not the root node, so we'll promote the selected child node to this position in its AttachParent's child array.
-			{
-				// Get the node's AttachParent
-				USceneComponent* AttachParent = SceneComponent->AttachParent;
-				check(AttachParent != nullptr);
+					ActorInstance->Modify();
 
-				// Find the node's position in its AttachParent's child array
-				int32 Index = AttachParent->AttachChildren.Find(SceneComponent);
-				check(Index != INDEX_NONE);
+					// Set the selected child node as the new root
+					check(ChildToPromote != nullptr);
+					ActorInstance->SetRootComponent(ChildToPromote);
+				}
+				else    // ...not the root node, so we'll promote the selected child node to this position in its AttachParent's child array.
+				{
+					// Get the node's AttachParent
+					USceneComponent* AttachParent = SceneComponent->AttachParent;
+					check(AttachParent != nullptr);
 
-				// Detach the node from its parent
-				SceneComponent->DetachFromParent();
+					// Find the node's position in its AttachParent's child array
+					int32 Index = AttachParent->AttachChildren.Find(SceneComponent);
+					check(Index != INDEX_NONE);
 
-				// Attach the child node that we're promoting to the parent and move it to the same position as the old node was in the array
-				ChildToPromote->AttachTo(AttachParent, NAME_None, EAttachLocation::KeepWorldPosition);
-				AttachParent->AttachChildren.Remove(ChildToPromote);
-				AttachParent->AttachChildren.Insert(ChildToPromote, Index);
-			}
+					// Detach the node from its parent
+					SceneComponent->DetachFromParent();
 
-			// Detach child nodes from the node that's being removed and re-attach them to the child that's being promoted
-			TArray<USceneComponent*> AttachChildrenLocalCopy(SceneComponent->AttachChildren);
-			for(auto ChildCompIt = AttachChildrenLocalCopy.CreateIterator(); ChildCompIt; ++ChildCompIt)
-			{
-				USceneComponent* Child = *ChildCompIt;
-				check(Child != nullptr);
+					if(ChildToPromote != nullptr)
+					{
+						// Attach the child node that we're promoting to the parent and move it to the same position as the old node was in the array
+						ChildToPromote->AttachTo(AttachParent, NAME_None, EAttachLocation::KeepWorldPosition);
+						AttachParent->AttachChildren.Remove(ChildToPromote);
+						AttachParent->AttachChildren.Insert(ChildToPromote, Index);
+					}
+				}
 
-				// Note: This will internally call Modify(), so we don't need to call it here
-				Child->DetachFromParent(true);
-				if(Child != ChildToPromote)
+				// Detach child nodes from the node that's being removed and re-attach them to the child that's being promoted
+				TArray<USceneComponent*> AttachChildrenLocalCopy(SceneComponent->AttachChildren);
+				for(auto ChildCompIt = AttachChildrenLocalCopy.CreateIterator(); ChildCompIt; ++ChildCompIt)
 				{
-					Child->AttachTo(ChildToPromote, NAME_None, EAttachLocation::KeepWorldPosition);
+					USceneComponent* Child = *ChildCompIt;
+					check(Child != nullptr);
+
+					// Note: This will internally call Modify(), so we don't need to call it here
+					Child->DetachFromParent(true);
+					if(Child != ChildToPromote)
+					{
+						Child->AttachTo(ChildToPromote, NAME_None, EAttachLocation::KeepWorldPosition);
+					}
 				}
 			}
-		}
-		// Destroy the component instance
-		ComponentInstance->DestroyComponent();
+
+			// Destroy the component instance
+			ComponentInstance->DestroyComponent();
 		}
 	}
 }
@@ -4595,7 +4624,6 @@ FSCSEditorTreeNodePtrType SSCSEditor::AddTreeNode(USCS_Node* InSCSNode, FSCSEdit
 		else 
 		{
 			// Create a new root node
-			TSharedPtr<FSCSEditorTreeNode> NewNodePtr;
 			if (RootTreeNode.IsValid())
 			{
 				NewNodePtr = RootTreeNode->AddChild(InSCSNode, bIsInherited);
@@ -4620,7 +4648,6 @@ FSCSEditorTreeNodePtrType SSCSEditor::AddTreeNode(USCS_Node* InSCSNode, FSCSEdit
 	else
 	{
 		// If the given SCS node does not contain a scene component template, we create a new root node
-		TSharedPtr<FSCSEditorTreeNode> NewNodePtr;
 		if (RootTreeNode.IsValid())
 		{
 			NewNodePtr = RootTreeNode->AddChild(InSCSNode, bIsInherited);
