[INTEGRATE] Change 2418955 by [REDACTED_PERSON][REDACTED_HANDLE] on 2015/01/26 14:14:33

	revert some code that steals selection away on UpdateTree() calls (will find a better way)

diff --git a/Engine/Source/Editor/Kismet/Private/SSCSEditor.cpp b/Engine/Source/Editor/Kismet/Private/SSCSEditor.cpp
--- a/Engine/Source/Editor/Kismet/Private/SSCSEditor.cpp
+++ b/Engine/Source/Editor/Kismet/Private/SSCSEditor.cpp
@@ -3772,10 +3772,6 @@ void SSCSEditor::UpdateTree(bool bRegenerateTreeNodes)
 				}
 			}
 		}
-		else if(SceneRootNodePtr.IsValid())
-		{
-			SCSTreeWidget->SetItemSelection(SceneRootNodePtr, true);
-		}
 
 		// If we have a pending deferred rename request, redirect it to the new tree node
 		if(DeferredRenameRequest != NAME_None)
