[INTEGRATE] Change 2419070 by [REDACTED_PERSON][REDACTED_HANDLE] on 2015/01/26 15:19:32

	Fix bad includes

diff --git a/Engine/Source/Editor/Kismet/Private/SSCSEditor.cpp b/Engine/Source/Editor/Kismet/Private/SSCSEditor.cpp
--- a/Engine/Source/Editor/Kismet/Private/SSCSEditor.cpp
+++ b/Engine/Source/Editor/Kismet/Private/SSCSEditor.cpp
@@ -46,7 +46,7 @@
 #include "GameProjectGenerationModule.h"
 #include "HotReloadInterface.h"
 #include "AssetRegistryModule.h"
-#include "SCreateAssetFromActor.h"
+#include "SCreateAssetFromObject.h"
 
 #define LOCTEXT_NAMESPACE "SSCSEditor"
 
@@ -3936,7 +3936,7 @@ bool SSCSEditor::IsEditingAllowed() const
 UActorComponent* SSCSEditor::AddNewComponent( UClass* NewComponentClass, UObject* Asset  )
 {
 	const FScopedTransaction Transaction( LOCTEXT("AddComponent", "Add Component") );
-
+	
 	if(EditorMode == EEditorMode::BlueprintSCS)
 	{
 		UBlueprint* Blueprint = GetBlueprint();
diff --git a/Engine/Source/Editor/KismetWidgets/Private/CreateBlueprintFromActorDialog.cpp b/Engine/Source/Editor/KismetWidgets/Private/CreateBlueprintFromActorDialog.cpp
--- a/Engine/Source/Editor/KismetWidgets/Private/CreateBlueprintFromActorDialog.cpp
+++ b/Engine/Source/Editor/KismetWidgets/Private/CreateBlueprintFromActorDialog.cpp
@@ -2,7 +2,7 @@
 
 #include "KismetWidgetsPrivatePCH.h"
 #include "CreateBlueprintFromActorDialog.h"
-#include "SCreateAssetFromActor.h"
+#include "SCreateAssetFromObject.h"
 #include "UnrealEd.h"
 #include "Editor/UnrealEd/Public/Kismet2/KismetEditorUtilities.h"
 #include "SNotificationList.h"
diff --git a/Engine/Source/Editor/UnrealEd/Private/Editor.cpp b/Engine/Source/Editor/UnrealEd/Private/Editor.cpp
--- a/Engine/Source/Editor/UnrealEd/Private/Editor.cpp
+++ b/Engine/Source/Editor/UnrealEd/Private/Editor.cpp
@@ -43,7 +43,7 @@
 #include "AnimationUtils.h"
 #include "AudioDecompress.h"
 #include "LevelEditor.h"
-#include "SCreateAssetFromActor.h"
+#include "SCreateAssetFromObject.h"
 
 #include "Editor/ActorPositioning.h"
 
