[INTEGRATE] Change 2419087 by [REDACTED_PERSON][REDACTED_HANDLE] on 2015/01/26 15:40:20

	Fix inifinite loop when duplicating a root primitive component

diff --git a/Engine/Source/Editor/Kismet/Private/SSCSEditor.cpp b/Engine/Source/Editor/Kismet/Private/SSCSEditor.cpp
--- a/Engine/Source/Editor/Kismet/Private/SSCSEditor.cpp
+++ b/Engine/Source/Editor/Kismet/Private/SSCSEditor.cpp
@@ -3305,14 +3305,9 @@ void SSCSEditor::OnDuplicateComponent()
 			UActorComponent* ComponentTemplate = SelectedNodes[i]->GetComponentTemplate();
 			if(ComponentTemplate != NULL)
 			{
-				UActorComponent* CloneComponent = AddNewComponent(ComponentTemplate->GetClass(), NULL);
+				UActorComponent* CloneComponent = AddNewComponent(ComponentTemplate->GetClass(), ComponentTemplate);
 				UActorComponent* OriginalComponent = ComponentTemplate;
 
-				//Serialize object properties using write/read operations.
-				TArray<uint8> SavedProperties;
-				FObjectWriter Writer(OriginalComponent, SavedProperties);
-				FObjectReader(CloneComponent, SavedProperties);
-
 				// If we've duplicated a scene component, attempt to reposition the duplicate in the hierarchy if the original
 				// was attached to another scene component as a child. By default, the duplicate is attached to the scene root node.
 				USceneComponent* NewSceneComponent = Cast<USceneComponent>(CloneComponent);
@@ -3936,7 +3931,15 @@ bool SSCSEditor::IsEditingAllowed() const
 UActorComponent* SSCSEditor::AddNewComponent( UClass* NewComponentClass, UObject* Asset  )
 {
 	const FScopedTransaction Transaction( LOCTEXT("AddComponent", "Add Component") );
-	
+
+	UActorComponent* NewComponent = nullptr;
+	UActorComponent* ComponentTemplate = Cast<UActorComponent>(Asset);
+
+	if (ComponentTemplate)
+	{
+		Asset = nullptr;
+	}
+
 	if(EditorMode == EEditorMode::BlueprintSCS)
 	{
 		UBlueprint* Blueprint = GetBlueprint();
@@ -3946,7 +3949,15 @@ UActorComponent* SSCSEditor::AddNewComponent( UClass* NewComponentClass, UObject
 		SaveSCSCurrentState(Blueprint->SimpleConstructionScript);
 
 		FName NewVariableName = Asset != nullptr ? Asset->GetFName() : NAME_None;
-		return AddNewNode(Blueprint->SimpleConstructionScript->CreateNode(NewComponentClass, NewVariableName), Asset, true);
+		NewComponent = AddNewNode(Blueprint->SimpleConstructionScript->CreateNode(NewComponentClass, NewVariableName), Asset, true);
+
+		if (ComponentTemplate)
+		{
+			//Serialize object properties using write/read operations.
+			TArray<uint8> SavedProperties;
+			FObjectWriter Writer(ComponentTemplate, SavedProperties);
+			FObjectReader(NewComponent, SavedProperties);
+		}
 	}
 	else    // EEditorMode::ActorInstance
 	{
@@ -3964,8 +3975,10 @@ UActorComponent* SSCSEditor::AddNewComponent( UClass* NewComponentClass, UObject
 		{
 			NewComponentName = *FComponentEditorUtils::GenerateValidVariableName(NewComponentClass, ActorInstance);
 		}
-		return AddNewNode(ConstructObject<UActorComponent>(NewComponentClass, ActorInstance, NewComponentName, RF_Transactional), Asset, true);
+		NewComponent = AddNewNode(ConstructObject<UActorComponent>(NewComponentClass, ActorInstance, NewComponentName, RF_Transactional, ComponentTemplate), Asset, true);
 	}
+
+	return NewComponent;
 }
 
 UActorComponent* SSCSEditor::AddNewNode(USCS_Node* NewNode,  UObject* Asset, bool bMarkBlueprintModified, bool bSetFocusToNewItem)
