[INTEGRATE] Change 2419179 by [REDACTED_PERSON][REDACTED_HANDLE] on 2015/01/26 16:37:54

	Fixed Compile errors from rename.
	Updated tooltips for SCSEditor.

diff --git a/Engine/Source/Editor/EditorStyle/Private/SlateEditorStyle.cpp b/Engine/Source/Editor/EditorStyle/Private/SlateEditorStyle.cpp
--- a/Engine/Source/Editor/EditorStyle/Private/SlateEditorStyle.cpp
+++ b/Engine/Source/Editor/EditorStyle/Private/SlateEditorStyle.cpp
@@ -756,6 +756,12 @@ void FSlateEditorStyle::FStyle::SetupGeneralStyles()
 			.SetInactiveBrush(IMAGE_BRUSH("Common/Selection", Icon8x8, SelectionColor_Inactive))
 			.SetInactiveHoveredBrush(IMAGE_BRUSH("Common/Selection", Icon8x8, SelectionColor_Inactive))
 		);
+		Set("TableView.NoHoverTableRow", FTableRowStyle(NormalTableRowStyle)
+			.SetEvenRowBackgroundHoveredBrush(FSlateNoResource())
+			.SetOddRowBackgroundHoveredBrush(FSlateNoResource())
+			.SetActiveHoveredBrush(FSlateNoResource())
+			.SetInactiveHoveredBrush(FSlateNoResource())
+			);
 
 		//Set("DetailsView.TreeView.TableRow", FTableRowStyle()
 		//	.SetEvenRowBackgroundBrush(FSlateNoResource())
diff --git a/Engine/Source/Editor/GraphEditor/Private/GraphEditorDragDropAction.cpp b/Engine/Source/Editor/GraphEditor/Private/GraphEditorDragDropAction.cpp
--- a/Engine/Source/Editor/GraphEditor/Private/GraphEditorDragDropAction.cpp
+++ b/Engine/Source/Editor/GraphEditor/Private/GraphEditorDragDropAction.cpp
@@ -89,6 +89,11 @@ void FGraphEditorDragDropAction::Construct()
 	HoverTargetChanged();
 }
 
+bool FGraphEditorDragDropAction::HasFeedbackMessage()
+{
+	return CursorDecoratorWindow->GetContent() != SNullWidget::NullWidget;
+}
+
 void FGraphEditorDragDropAction::SetFeedbackMessage(const TSharedPtr<SWidget>& Message)
 {
 	if (Message.IsValid())
diff --git a/Engine/Source/Editor/GraphEditor/Public/GraphEditorDragDropAction.h b/Engine/Source/Editor/GraphEditor/Public/GraphEditorDragDropAction.h
--- a/Engine/Source/Editor/GraphEditor/Public/GraphEditorDragDropAction.h
+++ b/Engine/Source/Editor/GraphEditor/Public/GraphEditorDragDropAction.h
@@ -29,6 +29,7 @@ public:
 	
 	virtual bool IsSupportedBySchema(const class UEdGraphSchema* Schema) const { return true; }
 
+	bool HasFeedbackMessage();
 	void SetFeedbackMessage(const TSharedPtr<SWidget>& Message);
 	void SetSimpleFeedbackMessage(const FSlateBrush* Icon, const FSlateColor& IconColor, const FText& Message);
 
diff --git a/Engine/Source/Editor/Kismet/Private/BlueprintEditor.cpp b/Engine/Source/Editor/Kismet/Private/BlueprintEditor.cpp
--- a/Engine/Source/Editor/Kismet/Private/BlueprintEditor.cpp
+++ b/Engine/Source/Editor/Kismet/Private/BlueprintEditor.cpp
@@ -759,7 +759,7 @@ void FBlueprintEditor::OnSelectionUpdated(const TArray<FSCSEditorTreeNodePtrType
 			auto NodePtr = *NodeIt;
 			if (NodePtr.IsValid())
 			{
-				if (NodePtr->IsRootActor())
+				if (NodePtr->GetNodeType() == FSCSEditorTreeNode::RootActorNode)
 				{
 					AActor* DefaultActor = GetSCSEditorActorContext();
 					InspectorObjects.Add(DefaultActor);
diff --git a/Engine/Source/Editor/Kismet/Private/SSCSEditor.cpp b/Engine/Source/Editor/Kismet/Private/SSCSEditor.cpp
--- a/Engine/Source/Editor/Kismet/Private/SSCSEditor.cpp
+++ b/Engine/Source/Editor/Kismet/Private/SSCSEditor.cpp
@@ -157,56 +157,12 @@ public:
 	/** Node(s) that we started the drag from */
 	TArray<FSCSEditorTreeNodePtrType> SourceNodes;
 
-	/** String to show as hover text */
-	FText CurrentHoverText;
-
 	/** The type of drop action that's pending while dragging */
 	EDropActionType PendingDropAction;
 
-	/** The widget decorator to use */
-	virtual TSharedPtr<SWidget> GetDefaultDecorator() const override;
-	FText GetHoverText() const;
-	const FSlateBrush* GetIcon() const;
 	static TSharedRef<FSCSRowDragDropOp> New(FName InVariableName, UStruct* InVariableSource, FNodeCreationAnalytic AnalyticCallback);
 };
 
-TSharedPtr<SWidget> FSCSRowDragDropOp::GetDefaultDecorator() const
-{
-	return SNew(SBorder)
-		.BorderImage(FEditorStyle::GetBrush("Graph.ConnectorFeedback.Border"))
-		[
-			SNew(SHorizontalBox)
-			+ SHorizontalBox::Slot()
-			.AutoWidth()
-			.Padding(0, 0, 3, 0)
-			[
-				SNew(SImage)
-				.Image(this, &FSCSRowDragDropOp::GetIcon)
-			]
-			+ SHorizontalBox::Slot()
-			.AutoWidth()
-			.VAlign(VAlign_Center)
-			[
-				SNew(STextBlock)
-				.Text(this, &FSCSRowDragDropOp::GetHoverText)
-			]
-		];
-}
-
-FText FSCSRowDragDropOp::GetHoverText() const
-{
-	return !CurrentHoverText.IsEmpty()
-		? CurrentHoverText
-		: LOCTEXT("DropActionToolTip_InvalidDropTarget", "Cannot drop here.");
-}
-
-const FSlateBrush* FSCSRowDragDropOp::GetIcon() const
-{
-	return PendingDropAction != DropAction_None
-		? FEditorStyle::GetBrush(TEXT("Graph.ConnectorFeedback.OK"))
-		: FEditorStyle::GetBrush(TEXT("Graph.ConnectorFeedback.Error"));
-}
-
 TSharedRef<FSCSRowDragDropOp> FSCSRowDragDropOp::New(FName InVariableName, UStruct* InVariableSource, FNodeCreationAnalytic AnalyticCallback)
 {
 	TSharedPtr<FSCSRowDragDropOp> Operation = MakeShareable(new FSCSRowDragDropOp);
@@ -220,15 +176,17 @@ TSharedRef<FSCSRowDragDropOp> FSCSRowDragDropOp::New(FName InVariableName, UStru
 //////////////////////////////////////////////////////////////////////////
 // FSCSEditorTreeNode
 
-FSCSEditorTreeNode::FSCSEditorTreeNode()
-	:bIsInherited(false)
+FSCSEditorTreeNode::FSCSEditorTreeNode(FSCSEditorTreeNode::ENodeType InNodeType)
+	:NodeType(InNodeType)
+	,bIsInherited(false)
 	,bIsInstanced(false)
 	,bNonTransactionalRename(false)
 {
 }
 
 FSCSEditorTreeNode::FSCSEditorTreeNode(USCS_Node* InSCSNode, bool bInIsInherited)
-	:bIsInherited(bInIsInherited)
+	:NodeType(FSCSEditorTreeNode::ComponentNode)
+	,bIsInherited(bInIsInherited)
 	,bIsInstanced(false)
 	,SCSNodePtr(InSCSNode)
 	,ComponentTemplatePtr(InSCSNode != NULL ? InSCSNode->ComponentTemplate : NULL)
@@ -237,7 +195,8 @@ FSCSEditorTreeNode::FSCSEditorTreeNode(USCS_Node* InSCSNode, bool bInIsInherited
 }
 
 FSCSEditorTreeNode::FSCSEditorTreeNode(UActorComponent* InComponentTemplate)
-	:bIsInherited(false)
+	:NodeType(FSCSEditorTreeNode::ComponentNode)
+	,bIsInherited(false)
 	,bIsInstanced(false)
 	,SCSNodePtr(NULL)
 	,ComponentTemplatePtr(InComponentTemplate)
@@ -397,6 +356,20 @@ FString FSCSEditorTreeNode::GetDisplayString() const
 	}
 }
 
+FText FSCSEditorTreeNode::GetDisplayName() const
+{
+	FName VariableName = GetVariableName();
+	if (VariableName != NAME_None)
+	{
+		return FText::FromName(VariableName);
+	}
+	else if (InstancedComponentName != NAME_None)
+	{
+		return FText::FromName(InstancedComponentName);
+	}
+	return FText::GetEmpty();
+}
+
 UActorComponent* FSCSEditorTreeNode::FindComponentInstanceInActor(const AActor* InActor) const
 {
 	USCS_Node* SCS_Node = GetSCSNode();
@@ -472,9 +445,9 @@ UBlueprint* FSCSEditorTreeNode::GetBlueprint() const
 	return NULL;
 }
 
-bool FSCSEditorTreeNode::IsRootActor() const
+FSCSEditorTreeNode::ENodeType FSCSEditorTreeNode::GetNodeType() const
 {
-	return (GetSCSNode() == nullptr && InstancedComponentName == NAME_None && GetComponentTemplate() == nullptr);
+	return NodeType;
 }
 
 bool FSCSEditorTreeNode::IsRootComponent() const
@@ -1037,10 +1010,14 @@ void SSCS_RowWidget::Construct( const FArguments& InArgs, TSharedPtr<SSCSEditor>
 	SCSEditor = InSCSEditor;
 	TreeNodePtr = InNodePtr;
 
+	bool bIsSeparator = InNodePtr->GetNodeType() == FSCSEditorTreeNode::SeparatorNode;
 	
 	auto Args = FSuperRowType::FArguments()
-		.Style(&FEditorStyle::Get().GetWidgetStyle<FTableRowStyle>("SceneOutliner.TableViewRow")) //@todo create editor style for the SCS tree
-		.Padding(FMargin(0.f, 0.f, 0.f, 4.f));
+		.Style(bIsSeparator ?
+				&FEditorStyle::Get().GetWidgetStyle<FTableRowStyle>("TableView.NoHoverTableRow") :
+				&FEditorStyle::Get().GetWidgetStyle<FTableRowStyle>("SceneOutliner.TableViewRow")) //@todo create editor style for the SCS tree
+		.Padding(FMargin(0.f, 0.f, 0.f, 4.f))
+		.ShowSelection(!bIsSeparator);
 
 	SMultiColumnTableRow<FSCSEditorTreeNodePtrType>::Construct( Args, InOwnerTableView.ToSharedRef() );
 }
@@ -1479,7 +1456,7 @@ TSharedPtr<SWidget> SSCS_RowWidget::BuildSceneRootDropActionMenu(FSCSEditorTreeN
 
 FReply SSCS_RowWidget::OnMouseButtonDown(const FGeometry& MyGeometry, const FPointerEvent& MouseEvent)
 {
-	if (MouseEvent.GetEffectingButton() == EKeys::LeftMouseButton)
+	if (MouseEvent.GetEffectingButton() == EKeys::LeftMouseButton && GetNode()->GetNodeType() != FSCSEditorTreeNode::SeparatorNode)
 	{
 		FReply Reply = SMultiColumnTableRow<FSCSEditorTreeNodePtrType>::OnMouseButtonDown( MyGeometry, MouseEvent );
 		return Reply.DetectDrag( SharedThis(this) , EKeys::LeftMouseButton );
@@ -1493,50 +1470,33 @@ FReply SSCS_RowWidget::OnMouseButtonDown(const FGeometry& MyGeometry, const FPoi
 FReply SSCS_RowWidget::OnDragDetected( const FGeometry& MyGeometry, const FPointerEvent& MouseEvent )
 {
 	auto SCSEditorPtr = SCSEditor.Pin();
-	if ( MouseEvent.IsMouseButtonDown( EKeys::LeftMouseButton )
-			&& SCSEditorPtr.IsValid()
-			&& SCSEditorPtr->IsEditingAllowed() ) //can only drag when editing
+	if (MouseEvent.IsMouseButtonDown(EKeys::LeftMouseButton)
+		&& SCSEditorPtr.IsValid()
+		&& SCSEditorPtr->IsEditingAllowed()) //can only drag when editing
 	{
 		TArray<TSharedPtr<FSCSEditorTreeNode>> SelectedNodePtrs = SCSEditorPtr->GetSelectedNodes();
-		if(SelectedNodePtrs.Num() == 0)
+		if (SelectedNodePtrs.Num() == 0)
 		{
 			SelectedNodePtrs.Add(GetNode());
 		}
 
 		TSharedPtr<FSCSEditorTreeNode> FirstNode = SelectedNodePtrs[0];
-		UBlueprint* Blueprint = FirstNode->GetBlueprint();
-		
-		TSharedRef<FSCSRowDragDropOp> Operation = FSCSRowDragDropOp::New(FirstNode->IsRootActor() ? FName("Self") :FirstNode->GetVariableName(), Blueprint != nullptr ? Blueprint->SkeletonGeneratedClass : nullptr, FNodeCreationAnalytic());
-		//Operation->SetAltDrag(MouseEvent.IsAltDown());
-		//Operation->SetCtrlDrag(MouseEvent.IsLeftControlDown() || MouseEvent.IsRightControlDown());
-		Operation->SetCtrlDrag(true); // Always put a getter
-		Operation->CurrentHoverText = FText::GetEmpty();
-		Operation->PendingDropAction = FSCSRowDragDropOp::DropAction_None;
-		
-		for(const auto& SelectedNodePtr : SelectedNodePtrs)
+		if (FirstNode->GetNodeType() != FSCSEditorTreeNode::SeparatorNode)
 		{
-			Operation->SourceNodes.Add(SelectedNodePtr);
-			if(!SelectedNodePtr->CanReparent()
-				&& Operation->CurrentHoverText.IsEmpty())
-			{
-				// We set the tooltip text here because it won't change across entry/leave events
-				if(SelectedNodePtrs.Num() == 1)
-				{
-					Operation->CurrentHoverText = LOCTEXT("DropActionToolTip_Error_CannotReparent", "The selected component cannot be moved.");
-				}
-				else
-				{
-					Operation->CurrentHoverText = LOCTEXT("DropActionToolTip_Error_CannotReparentMultiple", "One or more of the selected components cannot be moved.");
-				}
-			}
-		}
+			UBlueprint* Blueprint = FirstNode->GetBlueprint();
 
-		return FReply::Handled().BeginDragDrop(Operation);
-	}
-	else
-	{
-		return FReply::Unhandled();
+			bool bIsRootActor = FirstNode->GetNodeType() == FSCSEditorTreeNode::RootActorNode;
+
+			TSharedRef<FSCSRowDragDropOp> Operation = FSCSRowDragDropOp::New(bIsRootActor ? FName("Self") : FirstNode->GetVariableName(), Blueprint != nullptr ? Blueprint->SkeletonGeneratedClass : nullptr, FNodeCreationAnalytic());
+			Operation->SetCtrlDrag(true); // Always put a getter
+			Operation->PendingDropAction = FSCSRowDragDropOp::DropAction_None;
+			Operation->SourceNodes = SelectedNodePtrs;
+
+			return FReply::Handled().BeginDragDrop(Operation);
+		}
 	}
+	
+	return FReply::Unhandled();
 }
 
 void SSCS_RowWidget::OnDragEnter( const FGeometry& MyGeometry, const FDragDropEvent& DragDropEvent )
@@ -1550,213 +1510,229 @@ void SSCS_RowWidget::OnDragEnter( const FGeometry& MyGeometry, const FDragDropEv
 	TSharedPtr<FSCSRowDragDropOp> DragRowOp = DragDropEvent.GetOperationAs<FSCSRowDragDropOp>();
 	if (DragRowOp.IsValid())
 	{
-		// If the hover text is already set, skip everything below, because it means we already know we can't drag-and-drop one or more of the selected nodes.
-		if(!DragRowOp->CurrentHoverText.IsEmpty())
-		{
-			return;
-		}
-
-		TSharedPtr<SSCSEditor> SCSEditorPin = SCSEditor.Pin();
-		check(SCSEditorPin.IsValid());
-		FSCSEditorTreeNodePtrType SceneRootNodePtr = SCSEditorPin->SceneRootNodePtr;
-		check(SceneRootNodePtr.IsValid());
-
-		FSlateColor IconColor = FLinearColor::White;
-		const FSlateBrush* StatusSymbol = FEditorStyle::GetBrush(TEXT("NoBrush"));
+		check(SCSEditor.IsValid());
+		
 		FText Message;
-
-		FSCSEditorTreeNodePtrType NodePtr = GetNode();
-
-		// Validate each selected node being dragged against the node that belongs to this row. Exit the loop if we have a valid tooltip OR a valid pending drop action once all nodes in the selection have been validated.
-		for(auto SourceNodeIter = DragRowOp->SourceNodes.CreateConstIterator(); SourceNodeIter && (DragRowOp->CurrentHoverText.IsEmpty() || DragRowOp->PendingDropAction != FSCSRowDragDropOp::DropAction_None); ++SourceNodeIter)
+		FSlateColor IconColor = FLinearColor::White;
+		
+		for (const auto& SelectedNodePtr : DragRowOp->SourceNodes)
 		{
-			FSCSEditorTreeNodePtrType DraggedNodePtr = *SourceNodeIter;
-			check(DraggedNodePtr.IsValid());
-
-
-			FText DraggedNodePtrName = SCSEditorPin->GetEditorMode() == SSCSEditor::ActorInstance ? FText::FromName( DraggedNodePtr->GetComponentTemplate()->GetFName() ) : FText::FromName(DraggedNodePtr->GetVariableName());
-			FText CurrentNodePtrName = SCSEditorPin->GetEditorMode() == SSCSEditor::ActorInstance ? FText::FromName( NodePtr->GetComponentTemplate()->GetFName() ) : FText::FromName(NodePtr->GetVariableName());
-
-			// Reset the pending drop action each time through the loop
-			DragRowOp->PendingDropAction = FSCSRowDragDropOp::DropAction_None;
-
-			// Get the component template objects associated with each node
-			USceneComponent* HoveredTemplate = Cast<USceneComponent>(NodePtr->GetComponentTemplate());
-			USceneComponent* DraggedTemplate = Cast<USceneComponent>(DraggedNodePtr->GetComponentTemplate());
-
-			if(DraggedNodePtr == NodePtr)
+			if (!SelectedNodePtr->CanReparent())
 			{
-				// Attempted to drag and drop onto self
-				if(DragRowOp->SourceNodes.Num() > 1)
-				{
-					Message = FText::Format(LOCTEXT("DropActionToolTip_Error_CannotAttachToSelfWithMultipleSelection", "Cannot attach the selected components here because it would result in {0} being attached to itself. Remove it from the selection and try again."), DraggedNodePtrName);
-				}
-				else
-				{
-					Message = FText::Format(LOCTEXT("DropActionToolTip_Error_CannotAttachToSelf", "Cannot attach {0} to itself."), DraggedNodePtrName);
-				}
-			}
-			else if(NodePtr->IsAttachedTo(DraggedNodePtr))
-			{
-				// Attempted to drop a parent onto a child
-				if(DragRowOp->SourceNodes.Num() > 1)
+				// We set the tooltip text here because it won't change across entry/leave events
+				if (DragRowOp->SourceNodes.Num() == 1)
 				{
-					Message = FText::Format(LOCTEXT("DropActionToolTip_Error_CannotAttachToChildWithMultipleSelection", "Cannot attach the selected components here because it would result in {0} being attached to one of its children. Remove it from the selection and try again."), DraggedNodePtrName);
+					Message = LOCTEXT("DropActionToolTip_Error_CannotReparent", "The selected component cannot be moved.");
 				}
 				else
 				{
-					Message = FText::Format(LOCTEXT("DropActionToolTip_Error_CannotAttachToChild", "Cannot attach {0} to one of its children."), DraggedNodePtrName);
+					Message = LOCTEXT("DropActionToolTip_Error_CannotReparentMultiple", "One or more of the selected components cannot be moved.");
 				}
 			}
-			else if(HoveredTemplate == NULL || DraggedTemplate == NULL)
+		}
+
+		if (Message.IsEmpty())
+		{
+			FSCSEditorTreeNodePtrType SceneRootNodePtr = SCSEditor.Pin()->SceneRootNodePtr;
+			check(SceneRootNodePtr.IsValid());
+
+			FSCSEditorTreeNodePtrType NodePtr = GetNode();
+			if (NodePtr->GetNodeType() == FSCSEditorTreeNode::SeparatorNode)
 			{
-				// Can't attach non-USceneComponent types
-				Message = LOCTEXT("DropActionToolTip_Error_NotAttachable", "Cannot attach to this component.");
+				Message = LOCTEXT("DropActionToolTip_Error_CannotParentToSeparator", "Cannot Parent to separator.");
 			}
-			else if(NodePtr == SceneRootNodePtr)
+
+			// Validate each selected node being dragged against the node that belongs to this row. Exit the loop if we have a valid tooltip OR a valid pending drop action once all nodes in the selection have been validated.
+			for (auto SourceNodeIter = DragRowOp->SourceNodes.CreateConstIterator(); SourceNodeIter && (Message.IsEmpty() || DragRowOp->PendingDropAction != FSCSRowDragDropOp::DropAction_None); ++SourceNodeIter)
 			{
-				bool bCanMakeNewRoot = false;
-				bool bCanAttachToRoot = !NodePtr->IsDefaultSceneRoot()
-					&& !DraggedNodePtr->IsDirectlyAttachedTo(NodePtr)
-					&& HoveredTemplate->CanAttachAsChild(DraggedTemplate, NAME_None)
-					&& DraggedTemplate->Mobility >= HoveredTemplate->Mobility
-					&& (!HoveredTemplate->IsEditorOnly() || DraggedTemplate->IsEditorOnly());
-
-				if(!NodePtr->CanReparent() && (!NodePtr->IsDefaultSceneRoot() || NodePtr->IsInherited()))
-				{
-					// Cannot make the dropped node the new root if we cannot reparent the current root
-					Message = LOCTEXT("DropActionToolTip_Error_CannotReparentRootNode", "The root component in this Blueprint cannot be replaced.");
-				}
-				else if(DraggedTemplate->IsEditorOnly() && !HoveredTemplate->IsEditorOnly()) 
+				FSCSEditorTreeNodePtrType DraggedNodePtr = *SourceNodeIter;
+				check(DraggedNodePtr.IsValid());
+
+				// Reset the pending drop action each time through the loop
+				DragRowOp->PendingDropAction = FSCSRowDragDropOp::DropAction_None;
+
+				// Get the component template objects associated with each node
+				USceneComponent* HoveredTemplate = Cast<USceneComponent>(NodePtr->GetComponentTemplate());
+				USceneComponent* DraggedTemplate = Cast<USceneComponent>(DraggedNodePtr->GetComponentTemplate());
+
+				if (DraggedNodePtr == NodePtr)
 				{
-					// can't have a new root that's editor-only (when children would be around in-game)
-					Message = LOCTEXT("DropActionToolTip_Error_CannotReparentEditorOnly", "Cannot re-parent game components under editor-only ones.");
+					// Attempted to drag and drop onto self
+					if (DragRowOp->SourceNodes.Num() > 1)
+					{
+						Message = FText::Format(LOCTEXT("DropActionToolTip_Error_CannotAttachToSelfWithMultipleSelection", "Cannot attach the selected components here because it would result in {0} being attached to itself. Remove it from the selection and try again."), DraggedNodePtr->GetDisplayName());
+					}
+					else
+					{
+						Message = FText::Format(LOCTEXT("DropActionToolTip_Error_CannotAttachToSelf", "Cannot attach {0} to itself."), DraggedNodePtr->GetDisplayName());
+					}
 				}
-				else if(DraggedTemplate->Mobility > HoveredTemplate->Mobility)
+				else if (NodePtr->IsAttachedTo(DraggedNodePtr))
 				{
-					// can't have a new root that's movable if the existing root is static or stationary
-					Message = LOCTEXT("DropActionToolTip_Error_CannotReparentNonMovable", "Cannot replace a non-movable scene root with a movable component.");
+					// Attempted to drop a parent onto a child
+					if (DragRowOp->SourceNodes.Num() > 1)
+					{
+						Message = FText::Format(LOCTEXT("DropActionToolTip_Error_CannotAttachToChildWithMultipleSelection", "Cannot attach the selected components here because it would result in {0} being attached to one of its children. Remove it from the selection and try again."), DraggedNodePtr->GetDisplayName());
+					}
+					else
+					{
+						Message = FText::Format(LOCTEXT("DropActionToolTip_Error_CannotAttachToChild", "Cannot attach {0} to one of its children."), DraggedNodePtr->GetDisplayName());
+					}
 				}
-				else if(DragRowOp->SourceNodes.Num() > 1)
+				else if (HoveredTemplate == NULL || DraggedTemplate == NULL)
 				{
-					Message = LOCTEXT("DropActionToolTip_Error_CannotAssignMultipleRootNodes", "Cannot replace the scene root with multiple components. Please select only a single component and try again.");
+					// Can't attach non-USceneComponent types
+					Message = LOCTEXT("DropActionToolTip_Error_NotAttachable", "Cannot attach to this component.");
 				}
-				else
+				else if (NodePtr == SceneRootNodePtr)
 				{
-					bCanMakeNewRoot = true;
-				}
+					bool bCanMakeNewRoot = false;
+					bool bCanAttachToRoot = !NodePtr->IsDefaultSceneRoot()
+						&& !DraggedNodePtr->IsDirectlyAttachedTo(NodePtr)
+						&& HoveredTemplate->CanAttachAsChild(DraggedTemplate, NAME_None)
+						&& DraggedTemplate->Mobility >= HoveredTemplate->Mobility
+						&& (!HoveredTemplate->IsEditorOnly() || DraggedTemplate->IsEditorOnly());
+
+					if (!NodePtr->CanReparent() && (!NodePtr->IsDefaultSceneRoot() || NodePtr->IsInherited()))
+					{
+						// Cannot make the dropped node the new root if we cannot reparent the current root
+						Message = LOCTEXT("DropActionToolTip_Error_CannotReparentRootNode", "The root component in this Blueprint cannot be replaced.");
+					}
+					else if (DraggedTemplate->IsEditorOnly() && !HoveredTemplate->IsEditorOnly())
+					{
+						// can't have a new root that's editor-only (when children would be around in-game)
+						Message = LOCTEXT("DropActionToolTip_Error_CannotReparentEditorOnly", "Cannot re-parent game components under editor-only ones.");
+					}
+					else if (DraggedTemplate->Mobility > HoveredTemplate->Mobility)
+					{
+						// can't have a new root that's movable if the existing root is static or stationary
+						Message = LOCTEXT("DropActionToolTip_Error_CannotReparentNonMovable", "Cannot replace a non-movable scene root with a movable component.");
+					}
+					else if (DragRowOp->SourceNodes.Num() > 1)
+					{
+						Message = LOCTEXT("DropActionToolTip_Error_CannotAssignMultipleRootNodes", "Cannot replace the scene root with multiple components. Please select only a single component and try again.");
+					}
+					else
+					{
+						bCanMakeNewRoot = true;
+					}
 
-				if(bCanMakeNewRoot && bCanAttachToRoot)
-				{
-					// User can choose to either attach to the current root or make the dropped node the new root
-					Message = LOCTEXT("DropActionToolTip_AttachToOrMakeNewRoot", "Drop here to see available actions.");
-					DragRowOp->PendingDropAction = FSCSRowDragDropOp::DropAction_AttachToOrMakeNewRoot;
-				}
-				else if(SCSEditor.Pin()->GetEditorMode() == SSCSEditor::EEditorMode::BlueprintSCS && DraggedNodePtr->GetBlueprint() != GetBlueprint())
-				{
-					if(bCanMakeNewRoot)
+					if (bCanMakeNewRoot && bCanAttachToRoot)
 					{
-						// Only available action is to copy the dragged node to the other Blueprint and make it the new root
-						Message = FText::Format(LOCTEXT("DropActionToolTip_DropMakeNewRootNodeFromCopy", "Drop here to copy {0} to a new variable and make it the new root."), DraggedNodePtrName);
+						// User can choose to either attach to the current root or make the dropped node the new root
+						Message = LOCTEXT("DropActionToolTip_AttachToOrMakeNewRoot", "Drop here to see available actions.");
+						DragRowOp->PendingDropAction = FSCSRowDragDropOp::DropAction_AttachToOrMakeNewRoot;
+					}
+					else if (SCSEditor.Pin()->GetEditorMode() == SSCSEditor::EEditorMode::BlueprintSCS && DraggedNodePtr->GetBlueprint() != GetBlueprint())
+					{
+						if (bCanMakeNewRoot)
+						{
+							// Only available action is to copy the dragged node to the other Blueprint and make it the new root
+							Message = FText::Format(LOCTEXT("DropActionToolTip_DropMakeNewRootNodeFromCopy", "Drop here to copy {0} to a new variable and make it the new root."), DraggedNodePtr->GetDisplayName());
+							DragRowOp->PendingDropAction = FSCSRowDragDropOp::DropAction_MakeNewRoot;
+						}
+						else if (bCanAttachToRoot)
+						{
+							// Only available action is to copy the dragged node(s) to the other Blueprint and attach it to the root
+							if (DragRowOp->SourceNodes.Num() > 1)
+							{
+								Message = FText::Format(LOCTEXT("DropActionToolTip_AttachToThisNodeFromCopyWithMultipleSelection", "Drop here to copy the selected components to new variables and attach them to {0}."), NodePtr->GetDisplayName());
+							}
+							else
+							{
+								Message = FText::Format(LOCTEXT("DropActionToolTip_AttachToThisNodeFromCopy", "Drop here to copy {0} to a new variable and attach it to {1}."), DraggedNodePtr->GetDisplayName(), NodePtr->GetDisplayName());
+							}
+
+							DragRowOp->PendingDropAction = FSCSRowDragDropOp::DropAction_AttachTo;
+						}
+					}
+					else if (bCanMakeNewRoot)
+					{
+						// Only available action is to make the dragged node the new root
+						Message = FText::Format(LOCTEXT("DropActionToolTip_DropMakeNewRootNode", "Drop here to make {0} the new root."), DraggedNodePtr->GetDisplayName());
 						DragRowOp->PendingDropAction = FSCSRowDragDropOp::DropAction_MakeNewRoot;
 					}
-					else if(bCanAttachToRoot)
+					else if (bCanAttachToRoot)
 					{
-						// Only available action is to copy the dragged node(s) to the other Blueprint and attach it to the root
-						if(DragRowOp->SourceNodes.Num() > 1)
+						// Only available action is to attach the dragged node(s) to the root
+						if (DragRowOp->SourceNodes.Num() > 1)
 						{
-							Message = FText::Format(LOCTEXT("DropActionToolTip_AttachToThisNodeFromCopyWithMultipleSelection", "Drop here to copy the selected components to new variables and attach them to {0}."), CurrentNodePtrName);
+							Message = FText::Format(LOCTEXT("DropActionToolTip_AttachToThisNodeWithMultipleSelection", "Drop here to attach the selected components to {0}."), NodePtr->GetDisplayName());
 						}
 						else
 						{
-							Message = FText::Format(LOCTEXT("DropActionToolTip_AttachToThisNodeFromCopy", "Drop here to copy {0} to a new variable and attach it to {1}."), DraggedNodePtrName, CurrentNodePtrName);
+							Message = FText::Format(LOCTEXT("DropActionToolTip_AttachToThisNode", "Drop here to attach {0} to {1}."), DraggedNodePtr->GetDisplayName(), NodePtr->GetDisplayName());
 						}
-						
+
 						DragRowOp->PendingDropAction = FSCSRowDragDropOp::DropAction_AttachTo;
 					}
 				}
-				else if(bCanMakeNewRoot)
-				{
-					// Only available action is to make the dragged node the new root
-					Message = FText::Format(LOCTEXT("DropActionToolTip_DropMakeNewRootNode", "Drop here to make {0} the new root."),DraggedNodePtrName);
-					DragRowOp->PendingDropAction = FSCSRowDragDropOp::DropAction_MakeNewRoot;
-				}
-				else if(bCanAttachToRoot)
+				else if (DraggedNodePtr->IsDirectlyAttachedTo(NodePtr)) // if dropped onto parent
 				{
-					// Only available action is to attach the dragged node(s) to the root
-					if(DragRowOp->SourceNodes.Num() > 1)
+					// Detach the dropped node(s) from the current node and reattach to the root node
+					if (DragRowOp->SourceNodes.Num() > 1)
 					{
-						Message = FText::Format(LOCTEXT("DropActionToolTip_AttachToThisNodeWithMultipleSelection", "Drop here to attach the selected components to {0}."), CurrentNodePtrName);
+						Message = FText::Format(LOCTEXT("DropActionToolTip_DetachFromThisNodeWithMultipleSelection", "Drop here to detach the selected components from {0}."), NodePtr->GetDisplayName());
 					}
 					else
 					{
-						Message = FText::Format(LOCTEXT("DropActionToolTip_AttachToThisNode", "Drop here to attach {0} to {1}."), DraggedNodePtrName, CurrentNodePtrName);
+						Message = FText::Format(LOCTEXT("DropActionToolTip_DetachFromThisNode", "Drop here to detach {0} from {1}."), DraggedNodePtr->GetDisplayName(), NodePtr->GetDisplayName());
 					}
 
-					DragRowOp->PendingDropAction = FSCSRowDragDropOp::DropAction_AttachTo;
+					DragRowOp->PendingDropAction = FSCSRowDragDropOp::DropAction_DetachFrom;
 				}
-			}
-			else if(DraggedNodePtr->IsDirectlyAttachedTo(NodePtr)) // if dropped onto parent
-			{
-				// Detach the dropped node(s) from the current node and reattach to the root node
-				if(DragRowOp->SourceNodes.Num() > 1)
+				else if (!DraggedTemplate->IsEditorOnly() && HoveredTemplate->IsEditorOnly())
 				{
-					Message = FText::Format(LOCTEXT("DropActionToolTip_DetachFromThisNodeWithMultipleSelection", "Drop here to detach the selected components from {0}."), CurrentNodePtrName);
+					// can't have a game component child nested under an editor-only one
+					Message = LOCTEXT("DropActionToolTip_Error_CannotAttachToEditorOnly", "Cannot attach game components to editor-only ones.");
 				}
-				else
+				else if ((DraggedTemplate->Mobility == EComponentMobility::Static) && ((HoveredTemplate->Mobility == EComponentMobility::Movable) || (HoveredTemplate->Mobility == EComponentMobility::Stationary)))
 				{
-					Message = FText::Format(LOCTEXT("DropActionToolTip_DetachFromThisNode", "Drop here to detach {0} from {1}."), DraggedNodePtrName, CurrentNodePtrName);
+					// Can't attach Static components to mobile ones
+					Message = LOCTEXT("DropActionToolTip_Error_CannotAttachStatic", "Cannot attach Static components to movable ones.");
 				}
-				
-				DragRowOp->PendingDropAction = FSCSRowDragDropOp::DropAction_DetachFrom;
-			}
-			else if (!DraggedTemplate->IsEditorOnly() && HoveredTemplate->IsEditorOnly()) 
-			{
-				// can't have a game component child nested under an editor-only one
-				Message = LOCTEXT("DropActionToolTip_Error_CannotAttachToEditorOnly", "Cannot attach game components to editor-only ones.");
-			}
-			else if ((DraggedTemplate->Mobility == EComponentMobility::Static) && ((HoveredTemplate->Mobility == EComponentMobility::Movable) || (HoveredTemplate->Mobility == EComponentMobility::Stationary)))
-			{
-				// Can't attach Static components to mobile ones
-				Message = LOCTEXT("DropActionToolTip_Error_CannotAttachStatic", "Cannot attach Static components to movable ones.");
-			}
-			else if ((DraggedTemplate->Mobility == EComponentMobility::Stationary) && (HoveredTemplate->Mobility == EComponentMobility::Movable))
-			{
-				// Can't attach Static components to mobile ones
-				Message = LOCTEXT("DropActionToolTip_Error_CannotAttachStationary", "Cannot attach Stationary components to movable ones.");
-			}
-			else if(HoveredTemplate->CanAttachAsChild(DraggedTemplate, NAME_None))
-			{
-				// Attach the dragged node(s) to this node
-				if(DraggedNodePtr->GetBlueprint() != GetBlueprint())
+				else if ((DraggedTemplate->Mobility == EComponentMobility::Stationary) && (HoveredTemplate->Mobility == EComponentMobility::Movable))
+				{
+					// Can't attach Static components to mobile ones
+					Message = LOCTEXT("DropActionToolTip_Error_CannotAttachStationary", "Cannot attach Stationary components to movable ones.");
+				}
+				else if (HoveredTemplate->CanAttachAsChild(DraggedTemplate, NAME_None))
 				{
-					if(DragRowOp->SourceNodes.Num() > 1)
+					// Attach the dragged node(s) to this node
+					if (DraggedNodePtr->GetBlueprint() != GetBlueprint())
 					{
-						Message = FText::Format(LOCTEXT("DropActionToolTip_AttachToThisNodeFromCopyWithMultipleSelection", "Drop here to copy the selected nodes to new variables and attach to {0}."), CurrentNodePtrName);
+						if (DragRowOp->SourceNodes.Num() > 1)
+						{
+							Message = FText::Format(LOCTEXT("DropActionToolTip_AttachToThisNodeFromCopyWithMultipleSelection", "Drop here to copy the selected nodes to new variables and attach to {0}."), NodePtr->GetDisplayName());
+						}
+						else
+						{
+							Message = FText::Format(LOCTEXT("DropActionToolTip_AttachToThisNodeFromCopy", "Drop here to copy {0} to a new variable and attach it to {1}."), DraggedNodePtr->GetDisplayName(), NodePtr->GetDisplayName());
+						}
+					}
+					else if (DragRowOp->SourceNodes.Num() > 1)
+					{
+						Message = FText::Format(LOCTEXT("DropActionToolTip_AttachToThisNodeWithMultipleSelection", "Drop here to attach the selected nodes to {0}."), NodePtr->GetDisplayName());
 					}
 					else
 					{
-						Message = FText::Format(LOCTEXT("DropActionToolTip_AttachToThisNodeFromCopy", "Drop here to copy {0} to a new variable and attach it to {1}."), DraggedNodePtrName, CurrentNodePtrName);
+						Message = FText::Format(LOCTEXT("DropActionToolTip_AttachToThisNode", "Drop here to attach {0} to {1}."), DraggedNodePtr->GetDisplayName(), NodePtr->GetDisplayName());
 					}
-				}
-				else if(DragRowOp->SourceNodes.Num() > 1)
-				{
-					Message = FText::Format(LOCTEXT("DropActionToolTip_AttachToThisNodeWithMultipleSelection", "Drop here to attach the selected nodes to {0}."), CurrentNodePtrName);
+
+					DragRowOp->PendingDropAction = FSCSRowDragDropOp::DropAction_AttachTo;
 				}
 				else
 				{
-					Message = FText::Format(LOCTEXT("DropActionToolTip_AttachToThisNode", "Drop here to attach {0} to {1}."), DraggedNodePtrName, CurrentNodePtrName);
+					// The dropped node cannot be attached to the current node
+					Message = FText::Format(LOCTEXT("DropActionToolTip_Error_TooManyAttachments", "Unable to attach {0} to {1}."), DraggedNodePtr->GetDisplayName(), NodePtr->GetDisplayName());
 				}
-
-				DragRowOp->PendingDropAction = FSCSRowDragDropOp::DropAction_AttachTo;
-			}
-			else
-			{
-				// The dropped node cannot be attached to the current node
-				Message = FText::Format(LOCTEXT("DropActionToolTip_Error_TooManyAttachments", "Unable to attach {0} to {1}."), DraggedNodePtrName, CurrentNodePtrName);
 			}
 		}
 
+		const FSlateBrush* StatusSymbol = DragRowOp->PendingDropAction != FSCSRowDragDropOp::DropAction_None
+			? FEditorStyle::GetBrush(TEXT("Graph.ConnectorFeedback.OK"))
+			: FEditorStyle::GetBrush(TEXT("Graph.ConnectorFeedback.Error"));
+
 		DragRowOp->SetSimpleFeedbackMessage(StatusSymbol, FLinearColor::White, Message);
 	}
 	else if ( Operation->IsOfType<FExternalDragOperation>() || Operation->IsOfType<FAssetDragDropOp>() )
@@ -1787,7 +1763,7 @@ void SSCS_RowWidget::OnDragLeave( const FDragDropEvent& DragDropEvent )
 		// Only clear the tooltip text if all dragged nodes support it
 		if(bCanReparentAllNodes)
 		{
-			DragRowOp->CurrentHoverText = FText::GetEmpty();
+			DragRowOp->SetFeedbackMessage(SNullWidget::NullWidget);
 			DragRowOp->PendingDropAction = FSCSRowDragDropOp::DropAction_None;
 		}
 	}
@@ -2460,13 +2436,6 @@ TSharedRef<SWidget> SSCS_RowWidget_ActorRoot::GenerateWidgetForColumn(const FNam
 	return SNew(SHorizontalBox)
 		.ToolTip(CreateToolTipWidget())
 
-		+ SHorizontalBox::Slot()
-		.AutoWidth()
-		.VAlign(VAlign_Center)
-		[
-			SNew(SExpanderArrow, SharedThis(this))
-		]
-
 	+ SHorizontalBox::Slot()
 		.AutoWidth()
 		.VAlign(VAlign_Center)
@@ -2580,7 +2549,7 @@ FText SSCS_RowWidget_ActorRoot::GetActorDisplayText() const
 			}
 			else
 			{
-				Name = DefaultActor->GetActorLabel();
+				DefaultActor->GetName(Name);
 				return FText::Format(LOCTEXT("DefaultActor_Name", "{0} (Instance)"), FText::FromString(Name));
 			}
 		}
@@ -2655,6 +2624,41 @@ FText SSCS_RowWidget_ActorRoot::GetActorMobilityText() const
 	return Text;
 }
 
+//////////////////////////////////////////////////////////////////////////
+// SSCS_RowWidget_Separator
+
+
+TSharedRef<SWidget> SSCS_RowWidget_Separator::GenerateWidgetForColumn(const FName& ColumnName)
+{
+	/*return SNew(SVerticalBox)
+
+		+ SVerticalBox::Slot()
+		.Padding(0.f)
+		.VAlign(VAlign_Top)
+		[
+			SNew(SBorder)
+			.Padding(FEditorStyle::GetMargin(TEXT("Menu.Separator.Padding")) + FMargin(1.f, 0.f))
+			.BorderImage(FEditorStyle::GetBrush(TEXT("Menu.Separator")))
+		]
+	
+		+ SVerticalBox::Slot()
+		.VAlign(VAlign_Top)
+		.AutoHeight()
+		.Padding(1.f)
+		[
+			SNew(STextBlock)
+			.Text(LOCTEXT("HeaderText", "Header Text"))
+			.TextStyle(FEditorStyle::Get(), TEXT("Menu.Heading"))
+		];*/
+
+	return SNew(SBox)
+		.Padding(1.f)
+		[
+			SNew(SBorder)
+			.Padding(FEditorStyle::GetMargin(TEXT("Menu.Separator.Padding")))
+			.BorderImage(FEditorStyle::GetBrush(TEXT("Menu.Separator")))
+		];
+}
 
 //////////////////////////////////////////////////////////////////////////
 // SSCSEditor
@@ -2821,8 +2825,8 @@ void SSCSEditor::Construct( const FArguments& InArgs )
 
 		EditBlueprintMenuBuilder.AddMenuEntry
 		(
-			LOCTEXT("PromoteToBlueprint", "Convert to Blueprint Class"),
-			LOCTEXT("PromoteToBluerprintTooltip","Converts the existing Blueprint into a new Blueprint class that is a subclass of the current Blueprint" ),
+			LOCTEXT("PromoteToBlueprint", "Convert to Class Blueprint"),
+			LOCTEXT("PromoteToBluerprintTooltip","Converts the existing Blueprint into a new SubClass Blueprint" ),
 			FSlateIcon(),
 			FUIAction(FExecuteAction::CreateSP(this, &SSCSEditor::PromoteToBlueprint))
 		);
@@ -2862,7 +2866,7 @@ void SSCSEditor::Construct( const FArguments& InArgs )
 						.ButtonStyle(FEditorStyle::Get(), "ToggleButton")
 						.ContentPadding(0)
 						.ToolTip(IDocumentation::Get()->CreateToolTip(
-							LOCTEXT("PromoteToBluerprintTooltip","Converts this actor into a reusable Blueprint class that can have script behavior" ),
+							LOCTEXT("PromoteToBluerprintTooltip","Converts this actor into a reusable Class Blueprint that can have script behavior" ),
 							NULL,
 							TEXT("Shared/LevelEditor"),
 							TEXT("ConvertToBlueprint")))
@@ -2882,7 +2886,7 @@ void SSCSEditor::Construct( const FArguments& InArgs )
 							[
 								SNew(STextBlock)
 								.TextStyle(FEditorStyle::Get(), "ContentBrowser.TopBar.Font")
-								.Text( LOCTEXT("PromoteToBlueprint", "Convert to Blueprint Class") )
+								.Text( LOCTEXT("PromoteToBlueprint", "Convert to Class Blueprint") )
 							]
 						]
 					]
@@ -3015,9 +3019,11 @@ void SSCSEditor::Tick( const FGeometry& AllottedGeometry, const double InCurrent
 			bool bFoundInvalidNode = false;
 			for(auto NodeIt = InNodes.CreateConstIterator(); NodeIt && !bFoundInvalidNode; ++NodeIt)
 			{
+				bool bIsComponent = (*NodeIt)->GetNodeType() == FSCSEditorTreeNode::ComponentNode;
+
 				const UActorComponent* InstancedComponent = (*NodeIt)->GetComponentTemplate();
-				bFoundInvalidNode = (!(*NodeIt)->IsRootActor() && (!InstancedComponent || InstancedComponent->IsPendingKill())) || AreAnyNodesInvalidLambda((*NodeIt)->GetChildren(), OutNumValidNodes);
-				if (!(*NodeIt)->IsRootActor())
+				bFoundInvalidNode = (bIsComponent && (!InstancedComponent || InstancedComponent->IsPendingKill())) || AreAnyNodesInvalidLambda((*NodeIt)->GetChildren(), OutNumValidNodes);
+				if (bIsComponent)
 				{
 					++OutNumValidNodes;
 				}
@@ -3087,11 +3093,16 @@ TSharedRef<ITableRow> SSCSEditor::MakeTableRowWidget( FSCSEditorTreeNodePtrType
 		TagMeta.FriendlyName = FString::Printf(TEXT("TableRow,%s,0"), *InNodePtr->GetComponentTemplate()->GetReadableName());
 	}
 
-	if (InNodePtr->IsRootActor())
+	if (InNodePtr->GetNodeType() == FSCSEditorTreeNode::RootActorNode)
 	{
 		return SNew(SSCS_RowWidget_ActorRoot, SharedThis(this), InNodePtr, OwnerTable)
 			.AddMetaData<FTutorialMetaData>(TagMeta);
 	}
+	else if (InNodePtr->GetNodeType() == FSCSEditorTreeNode::SeparatorNode)
+	{
+		return SNew(SSCS_RowWidget_Separator, SharedThis(this), InNodePtr, OwnerTable)
+			.AddMetaData<FTutorialMetaData>(TagMeta);
+	}
 
 	return SNew(SSCS_RowWidget, SharedThis(this), InNodePtr, OwnerTable)
 		.AddMetaData<FTutorialMetaData>(TagMeta);
@@ -3110,7 +3121,7 @@ void SSCSEditor::GetSelectedItemsForContextMenu(TArray<FComponentEventConstructi
 	}
 }
 
-TSharedPtr< SWidget > SSCSEditor::CreateContextMenu()
+TSharedPtr< SWidget > SSCSEditor::CreateContextMenu() // @todo: make the context manual meaningful for the actor and avoid separators
 {
 	TArray<FSCSEditorTreeNodePtrType> SelectedNodes = SCSTreeWidget->GetSelectedItems();
 
@@ -3619,12 +3630,18 @@ void SSCSEditor::UpdateTree(bool bRegenerateTreeNodes)
 		// Reset the scene root node
 		SceneRootNodePtr.Reset();
 
+		TSharedPtr<FSCSEditorTreeNode> ActorTreeNode;
 		bool bSingleLayoutBPEditor = GetDefault<UEditorExperimentalSettings>()->bUnifiedBlueprintEditor;
 		if (bSingleLayoutBPEditor)
 		{
-			RootTreeNode = MakeShareable(new FSCSEditorTreeNode());
-			RootNodes.Add(RootTreeNode);
-			SCSTreeWidget->SetItemExpansion(RootTreeNode, true);
+			ActorTreeNode = MakeShareable(new FSCSEditorTreeNode(FSCSEditorTreeNode::RootActorNode));
+			
+			//RootTreeNode = ActorTreeNode
+			//RootNodes.Add(RootTreeNode);
+			//SCSTreeWidget->SetItemExpansion(RootTreeNode, true);
+
+			RootNodes.Add(ActorTreeNode);
+			RootNodes.Add(MakeShareable(new FSCSEditorTreeNode(FSCSEditorTreeNode::SeparatorNode)));
 		}
 
 		// Build the tree data source according to what mode we're in
@@ -3663,11 +3680,17 @@ void SSCSEditor::UpdateTree(bool bRegenerateTreeNodes)
 				}
 
 				// Add native ActorComponent nodes that aren't SceneComponents
+				bool bSeparatorAdded = false;
 				for(auto CompIter = Components.CreateIterator(); CompIter; ++CompIter)
 				{
 					UActorComponent* ActorComp = *CompIter;
 					if (!ActorComp->IsA<USceneComponent>())
 					{
+						if (bSingleLayoutBPEditor && !bSeparatorAdded)
+						{
+							bSeparatorAdded = true;
+							RootNodes.Add(MakeShareable(new FSCSEditorTreeNode(FSCSEditorTreeNode::SeparatorNode)));
+						}
 						AddRootComponentTreeNode(ActorComp);
 					}
 				}
@@ -3725,11 +3748,17 @@ void SSCSEditor::UpdateTree(bool bRegenerateTreeNodes)
 				}
 
 				// Add all non-scene component instances to the root set first
+				bool bSeparatorAdded = false;
 				for(auto CompIter = Components.CreateIterator(); CompIter; ++CompIter)
 				{
 					UActorComponent* ActorComp = *CompIter;
 					if (!ActorComp->IsA<USceneComponent>() && !ActorComp->IsEditorOnly())
 					{
+						if (bSingleLayoutBPEditor && !bSeparatorAdded)
+						{
+							bSeparatorAdded = true;
+							RootNodes.Add(MakeShareable(new FSCSEditorTreeNode(FSCSEditorTreeNode::SeparatorNode)));
+						}
 						AddRootComponentTreeNode(ActorComp);
 					}
 				}
@@ -3751,11 +3780,11 @@ void SSCSEditor::UpdateTree(bool bRegenerateTreeNodes)
 		if(SelectedTreeNodes.Num() > 0)
 		{
 			// Restore the previous selection state on the new tree nodes
-			for(int i = 0; i < SelectedTreeNodes.Num(); ++i)
+			for (int i = 0; i < SelectedTreeNodes.Num(); ++i)
 			{
-				if (SelectedTreeNodes[i]->IsRootActor())
+				if (SelectedTreeNodes[i]->GetNodeType() == FSCSEditorTreeNode::RootActorNode)
 				{
-					SCSTreeWidget->SetItemSelection(RootTreeNode, true);
+					SCSTreeWidget->SetItemSelection(ActorTreeNode, true);
 				}
 				else
 				{
@@ -3804,6 +3833,7 @@ TSharedPtr<FSCSEditorTreeNode> SSCSEditor::AddRootComponentTreeNode(UActorCompon
 	else
 	{
 		NewTreeNode = MakeShareable(new FSCSEditorTreeNode(ActorComp));
+		RootNodes.Add(NewTreeNode);
 	}
 
 	RootComponentNodes.Add(NewTreeNode);
@@ -4422,7 +4452,7 @@ void SSCSEditor::OnDeleteNodes()
 			{
 				// Default to the parent node
 				NewSelection = Node->GetParent();
-				if(NewSelection.IsValid() && !NewSelection->IsRootActor())
+				if (NewSelection.IsValid() && !NewSelection->GetNodeType() == FSCSEditorTreeNode::RootActorNode)
 				{
 					// If we have sibling nodes, find the one that immediately precedes the one being removed
 					const TArray<FSCSEditorTreeNodePtrType>& ChildNodes = NewSelection->GetChildren();
@@ -4706,6 +4736,7 @@ FSCSEditorTreeNodePtrType SSCSEditor::AddTreeNode(USCS_Node* InSCSNode, FSCSEdit
 			else
 			{
 				NewNodePtr = MakeShareable(new FSCSEditorTreeNode(InSCSNode, bIsInherited));
+				RootNodes.Add(NewNodePtr);
 			}
 			
 			NodeSCS->AddNode(InSCSNode);
@@ -4730,6 +4761,7 @@ FSCSEditorTreeNodePtrType SSCSEditor::AddTreeNode(USCS_Node* InSCSNode, FSCSEdit
 		else
 		{
 			NewNodePtr = MakeShareable(new FSCSEditorTreeNode(InSCSNode, bIsInherited));
+			RootNodes.Add(NewNodePtr);
 		}
 
 		RootComponentNodes.Add(NewNodePtr);
diff --git a/Engine/Source/Editor/Kismet/Public/SSCSEditor.h b/Engine/Source/Editor/Kismet/Public/SSCSEditor.h
--- a/Engine/Source/Editor/Kismet/Public/SSCSEditor.h
+++ b/Engine/Source/Editor/Kismet/Public/SSCSEditor.h
@@ -20,10 +20,18 @@ typedef TSharedPtr<class FSCSEditorTreeNode> FSCSEditorTreeNodePtrType;
 class KISMET_API FSCSEditorTreeNode : public TSharedFromThis<FSCSEditorTreeNode>
 {
 public:
+
+	enum ENodeType
+	{
+		ComponentNode,
+		RootActorNode,
+		SeparatorNode,
+	};
+
 	/** 
 	 * Constructs an empty tree node.
 	 */
-	FSCSEditorTreeNode();
+	FSCSEditorTreeNode(FSCSEditorTreeNode::ENodeType InNodeType);
 
 	/**
 	 * Constructs a wrapper around a component template contained within an SCS tree node.
@@ -48,6 +56,10 @@ public:
 	 * @return The string to be used in the tree display.
 	 */
 	FString GetDisplayString() const;
+	/**
+	* @return The name of this node in text.
+	*/
+	FText GetDisplayName() const;
 	/** 
      * @return The SCS node that is represented by this object, or NULL if there is no SCS node associated with the component template.
 	 */
@@ -82,9 +94,9 @@ public:
 	 */
 	UBlueprint* GetBlueprint() const;
 	/**
-	 * @return Whether or not this object represents a root actor.
+	 * @return Type of node
 	 */
-	bool IsRootActor() const;
+	ENodeType GetNodeType() const;
 	/**
 	 * @return Whether or not this object represents a root component.
 	 */
@@ -219,6 +231,7 @@ public:
 	UActorComponent* GetOverridenComponentTemplate(UBlueprint* Blueprint, bool bCreateIfNecessary) const;
 
 private:
+	ENodeType NodeType;
 	bool bIsInherited;
 	bool bIsInstanced;
 	TWeakObjectPtr<class USCS_Node> SCSNodePtr;
@@ -357,6 +370,17 @@ private:
 	FText GetActorMobilityText() const;
 };
 
+class SSCS_RowWidget_Separator : public SSCS_RowWidget
+{
+public:
+
+	// SMultiColumnTableRow<T> interface
+	virtual TSharedRef<SWidget> GenerateWidgetForColumn(const FName& ColumnName) override;
+	// End of SMultiColumnTableRow<T>
+
+private:
+
+};
 
 //////////////////////////////////////////////////////////////////////////
 // SSCSEditorDragDropTree - implements STreeView for our specific node type and adds drag/drop functionality
diff --git a/Engine/Source/Editor/LevelEditor/Private/SLevelEditor.cpp b/Engine/Source/Editor/LevelEditor/Private/SLevelEditor.cpp
--- a/Engine/Source/Editor/LevelEditor/Private/SLevelEditor.cpp
+++ b/Engine/Source/Editor/LevelEditor/Private/SLevelEditor.cpp
@@ -696,7 +696,7 @@ private:
 				bool bActorSelected = false;
 				for (auto& SelectedNode : SelectedNodes)
 				{
-					if (SelectedNode.IsValid() && SelectedNode->IsRootActor())
+					if (SelectedNode.IsValid() && SelectedNode->GetNodeType() == FSCSEditorTreeNode::RootActorNode)
 					{
 						bActorSelected = true;
 						break;
