[INTEGRATE] Change 2419193 by [REDACTED_PERSON][REDACTED_HANDLE] on 2015/01/26 16:42:13

	Math expression node's pin types can be modified.

diff --git a/Engine/Source/Editor/Kismet/Private/BlueprintDetailsCustomization.cpp b/Engine/Source/Editor/Kismet/Private/BlueprintDetailsCustomization.cpp
--- a/Engine/Source/Editor/Kismet/Private/BlueprintDetailsCustomization.cpp
+++ b/Engine/Source/Editor/Kismet/Private/BlueprintDetailsCustomization.cpp
@@ -1864,7 +1864,8 @@ void FBlueprintGraphArgumentLayout::GenerateHeaderRowContent( FDetailWidgetRow&
 				.Schema(K2Schema)
 				.bAllowExec(TargetNode->CanModifyExecutionWires())
 				.bAllowWildcard(ShouldAllowWildcard(TargetNode))
-				.IsEnabled(!ShouldPinBeReadOnly())
+				.bAllowArrays(!ShouldPinBeReadOnly())
+				.IsEnabled(!ShouldPinBeReadOnly(true))
 				.Font( IDetailLayoutBuilder::GetDetailFont() )
 		]
 		+SHorizontalBox::Slot()
@@ -2003,7 +2004,7 @@ FReply FBlueprintGraphArgumentLayout::OnArgMoveDown()
 	return FReply::Handled();
 }
 
-bool FBlueprintGraphArgumentLayout::ShouldPinBeReadOnly() const
+bool FBlueprintGraphArgumentLayout::ShouldPinBeReadOnly(bool bIsEditingPinType/* = false*/) const
 {
 	const UEdGraphSchema_K2* Schema = GetDefault<UEdGraphSchema_K2>();
 
@@ -2017,19 +2018,19 @@ bool FBlueprintGraphArgumentLayout::ShouldPinBeReadOnly() const
 		else
 		{
 			// Check if pin editing is read only
-			return IsPinEditingReadOnly();
+			return IsPinEditingReadOnly(bIsEditingPinType);
 		}
 	}
 	
 	return false;
 }
 
-bool FBlueprintGraphArgumentLayout::IsPinEditingReadOnly() const
+bool FBlueprintGraphArgumentLayout::IsPinEditingReadOnly(bool bIsEditingPinType/* = false*/) const
 {
 	if(UEdGraph* NodeGraph = TargetNode->GetGraph())
 	{
-		// Math expression should not be modified directly, do not let the user tweak the parameters
-		if ( Cast<UK2Node_MathExpression>(NodeGraph->GetOuter()) )
+		// Math expression should not be modified directly (except for the pin type), do not let the user tweak the parameters
+		if (!bIsEditingPinType && Cast<UK2Node_MathExpression>(NodeGraph->GetOuter()) )
 		{
 			return true;
 		}
@@ -2119,7 +2120,7 @@ void FBlueprintGraphArgumentLayout::PinInfoChanged(const FEdGraphPinType& PinTyp
 		if (GraphActionDetailsPtr.IsValid())
 		{
 			GraphActionDetailsPtr.Pin()->GetMyBlueprint().Pin()->GetLastFunctionPinTypeUsed() = PinType;
-			if( !ShouldPinBeReadOnly() )
+			if( !ShouldPinBeReadOnly(true) )
 			{
 				GraphActionDetailsPtr.Pin()->OnParamsChanged(TargetNode);
 			}
@@ -2129,7 +2130,7 @@ void FBlueprintGraphArgumentLayout::PinInfoChanged(const FEdGraphPinType& PinTyp
 
 void FBlueprintGraphArgumentLayout::OnPrePinInfoChange(const FEdGraphPinType& PinType)
 {
-	if( !ShouldPinBeReadOnly() && TargetNode )
+	if( !ShouldPinBeReadOnly(true) && TargetNode )
 	{
 		TargetNode->Modify();
 	}
diff --git a/Engine/Source/Editor/Kismet/Private/BlueprintDetailsCustomization.h b/Engine/Source/Editor/Kismet/Private/BlueprintDetailsCustomization.h
--- a/Engine/Source/Editor/Kismet/Private/BlueprintDetailsCustomization.h
+++ b/Engine/Source/Editor/Kismet/Private/BlueprintDetailsCustomization.h
@@ -351,10 +351,10 @@ private:
 
 private:
 	/** Determines if this pin should not be editable */
-	bool ShouldPinBeReadOnly() const;
+	bool ShouldPinBeReadOnly(bool bIsEditingPinType = false) const;
 	
 	/** Determines if editing the pins on the node should be read only */
-	bool IsPinEditingReadOnly() const;
+	bool IsPinEditingReadOnly(bool bIsEditingPinType = false) const;
 
 	/** Callbacks for all the functionality for modifying arguments */
 	void OnRemoveClicked();
