[INTEGRATE] Change 2419313 by [REDACTED_PERSON][REDACTED_HANDLE] on 2015/01/26 17:38:11

	BP - Making sure the palette appears in the correct place when opened.

diff --git a/Engine/Source/Editor/Kismet/Private/BlueprintEditorModes.cpp b/Engine/Source/Editor/Kismet/Private/BlueprintEditorModes.cpp
--- a/Engine/Source/Editor/Kismet/Private/BlueprintEditorModes.cpp
+++ b/Engine/Source/Editor/Kismet/Private/BlueprintEditorModes.cpp
@@ -402,6 +402,7 @@ FBlueprintInterfaceApplicationMode::FBlueprintInterfaceApplicationMode(TSharedPt
 					(
 						FTabManager::NewStack() ->SetSizeCoefficient(0.5f)
 						->AddTab( FBlueprintEditorTabs::DetailsID, ETabState::OpenedTab )
+						->AddTab( FBlueprintEditorTabs::PaletteID, ETabState::ClosedTab )
 					)
 				)
 			)
@@ -570,7 +571,7 @@ FBlueprintEditorUnifiedMode::FBlueprintEditorUnifiedMode(TSharedPtr<class FBluep
 
 	if ( bRegisterViewport )
 	{
-		TabLayout = FTabManager::NewLayout( "Blueprints_Unified_Components_v1" )
+		TabLayout = FTabManager::NewLayout( "Blueprints_Unified_Components_v2" )
 		->AddArea
 		(
 			FTabManager::NewPrimaryArea()->SetOrientation(Orient_Vertical)
@@ -633,7 +634,7 @@ FBlueprintEditorUnifiedMode::FBlueprintEditorUnifiedMode(TSharedPtr<class FBluep
 	}
 	else
 	{
-		TabLayout = FTabManager::NewLayout( "Blueprints_Unified_v1" )
+		TabLayout = FTabManager::NewLayout( "Blueprints_Unified_v2" )
 		->AddArea
 		(
 			FTabManager::NewPrimaryArea() ->SetOrientation(Orient_Vertical)
@@ -680,6 +681,7 @@ FBlueprintEditorUnifiedMode::FBlueprintEditorUnifiedMode(TSharedPtr<class FBluep
 						FTabManager::NewStack()
 						->SetSizeCoefficient(0.60f)
 						->AddTab( FBlueprintEditorTabs::DetailsID, ETabState::OpenedTab )
+						->AddTab( FBlueprintEditorTabs::PaletteID, ETabState::ClosedTab )
 						->AddTab( FBlueprintEditorTabs::DefaultEditorID, ETabState::ClosedTab )
 					)
 				)
