[INTEGRATE] Change 2419667 by [REDACTED_PERSON][REDACTED_HANDLE] on 2015/01/26 19:59:40

	Blueprint Editor: Change messaging to consistently say "Class Defaults" everywhere when editing the CDO
	- Updated BP Editor and Persona to never call anything "Blueprint Defaults"
	- Removed dead method GetDefaultEditorTitle
	[UE-8164]

diff --git a/Engine/Source/Developer/AssetTools/Private/AssetTypeActions/AssetTypeActions_Blueprint.cpp b/Engine/Source/Developer/AssetTools/Private/AssetTypeActions/AssetTypeActions_Blueprint.cpp
--- a/Engine/Source/Developer/AssetTools/Private/AssetTypeActions/AssetTypeActions_Blueprint.cpp
+++ b/Engine/Source/Developer/AssetTools/Private/AssetTypeActions/AssetTypeActions_Blueprint.cpp
@@ -154,7 +154,7 @@ void FAssetTypeActions_Blueprint::ExecuteEditDefaults(TArray<TWeakObjectPtr<UBlu
 	}
 
 	// Report errors
-	EditorErrors.Notify(LOCTEXT("OpenDefaults_Failed", "Opening Blueprint Defaults Failed!"));
+	EditorErrors.Notify(LOCTEXT("OpenDefaults_Failed", "Opening Class Defaults Failed!"));
 }
 
 void FAssetTypeActions_Blueprint::ExecuteNewDerivedBlueprint(TWeakObjectPtr<UBlueprint> InObject)
diff --git a/Engine/Source/Editor/DetailCustomizations/Private/SceneComponentDetails.cpp b/Engine/Source/Editor/DetailCustomizations/Private/SceneComponentDetails.cpp
--- a/Engine/Source/Editor/DetailCustomizations/Private/SceneComponentDetails.cpp
+++ b/Engine/Source/Editor/DetailCustomizations/Private/SceneComponentDetails.cpp
@@ -615,7 +615,7 @@ void FSceneComponentDetails::MakeTransformDetails( IDetailLayoutBuilder& DetailB
 	DetailBuilder.HideProperty( DetailBuilder.GetProperty( GET_MEMBER_NAME_CHECKED(USceneComponent, RelativeRotation) ) );
 	DetailBuilder.HideProperty( DetailBuilder.GetProperty( GET_MEMBER_NAME_CHECKED(USceneComponent, RelativeScale3D) ) );
 
-	// Determine whether or not we are editing Blueprint defaults through the CDO
+	// Determine whether or not we are editing Class Defaults through the CDO
 	bool bIsEditingBlueprintDefaults = false;
 	const TArray<TWeakObjectPtr<UObject> >& SelectedObjects = DetailBuilder.GetDetailsView().GetSelectedObjects();
 	for(int32 ObjectIndex = 0; ObjectIndex < SelectedObjects.Num(); ++ObjectIndex)
@@ -631,7 +631,7 @@ void FSceneComponentDetails::MakeTransformDetails( IDetailLayoutBuilder& DetailB
 		}
 	}
 
-	// If there are any actors selected and we're not editing Blueprint defaults, the transform section is shown as part of the actor customization
+	// If there are any actors selected and we're not editing Class Defaults, the transform section is shown as part of the actor customization
 	if( SelectedActors.Num() == 0 || bIsEditingBlueprintDefaults)
 	{
 		TArray< TWeakObjectPtr<UObject> > SceneComponentObjects;
diff --git a/Engine/Source/Editor/Kismet/Private/BlueprintEditor.cpp b/Engine/Source/Editor/Kismet/Private/BlueprintEditor.cpp
--- a/Engine/Source/Editor/Kismet/Private/BlueprintEditor.cpp
+++ b/Engine/Source/Editor/Kismet/Private/BlueprintEditor.cpp
@@ -6611,7 +6611,7 @@ FText FBlueprintEditor::GetToolkitName() const
 	FFormatNamedArguments Args;
 	Args.Add( TEXT("NumberOfObjects"), EditingObjects.Num() );
 	Args.Add( TEXT("ObjectName"), FText::FromString( SharedParentClass->GetName() ) );
-	return FText::Format( NSLOCTEXT("KismetEditor", "ToolkitTitle_UniqueLayerName", "{NumberOfObjects} {ClassName} - Blueprint Defaults"), Args );
+	return FText::Format( NSLOCTEXT("KismetEditor", "ToolkitTitle_UniqueLayerName", "{NumberOfObjects} {ClassName} - Class Defaults"), Args );
 }
 
 
@@ -6656,11 +6656,6 @@ FText FBlueprintEditor::DefaultEditString()
 	return LOCTEXT("BlueprintEditingDefaults", "Editing defaults");
 }
 
-FText FBlueprintEditor::GetDefaultEditorTitle()
-{
-	return LOCTEXT("BlueprintDefaultsTabTitle", "Blueprint Defaults");
-}
-
 bool FBlueprintEditor::GetBoundsForSelectedNodes(class FSlateRect& Rect, float Padding)
 {
 	TSharedPtr<SGraphEditor> FocusedGraphEd = FocusedGraphEdPtr.Pin();
diff --git a/Engine/Source/Editor/Kismet/Private/BlueprintEditorTabFactories.cpp b/Engine/Source/Editor/Kismet/Private/BlueprintEditorTabFactories.cpp
--- a/Engine/Source/Editor/Kismet/Private/BlueprintEditorTabFactories.cpp
+++ b/Engine/Source/Editor/Kismet/Private/BlueprintEditorTabFactories.cpp
@@ -79,7 +79,8 @@ void FTimelineEditorSummoner::OnTabRefreshed(TSharedPtr<SDockTab> Tab) const
 	TimelineEditor->OnTimelineChanged();
 }
 
-FTimelineEditorSummoner::FTimelineEditorSummoner(TSharedPtr<class FBlueprintEditor> InBlueprintEditorPtr) : FDocumentTabFactoryForObjects<UTimelineTemplate>(FBlueprintEditorTabs::TimelineEditorID, InBlueprintEditorPtr)
+FTimelineEditorSummoner::FTimelineEditorSummoner(TSharedPtr<class FBlueprintEditor> InBlueprintEditorPtr)
+	: FDocumentTabFactoryForObjects<UTimelineTemplate>(FBlueprintEditorTabs::TimelineEditorID, InBlueprintEditorPtr)
 , BlueprintEditorPtr(InBlueprintEditorPtr)
 {
 
@@ -125,9 +126,10 @@ TSharedRef<SWidget> FDebugInfoSummoner::CreateTabBody(const FWorkflowTabSpawnInf
 	return BlueprintEditorPtr->GetDebuggingView();
 }
 
-FDefaultsEditorSummoner::FDefaultsEditorSummoner(TSharedPtr<class FAssetEditorToolkit> InHostingApp) : FWorkflowTabFactory(FBlueprintEditorTabs::DefaultEditorID, InHostingApp)
+FDefaultsEditorSummoner::FDefaultsEditorSummoner(TSharedPtr<class FAssetEditorToolkit> InHostingApp)
+	: FWorkflowTabFactory(FBlueprintEditorTabs::DefaultEditorID, InHostingApp)
 {
-	TabLabel = LOCTEXT("BlueprintDefaultsTabTitle", "Blueprint Defaults"); //@TODO: ANIMATION: GetDefaultEditorTitle(); !!!
+	TabLabel = LOCTEXT("ClassDefaultsTabTitle", "Class Defaults");
 	TabIcon = FSlateIcon(FEditorStyle::GetStyleSetName(), "Kismet.Tabs.BlueprintDefaults");
 
 	bIsSingleton = true;
@@ -143,7 +145,8 @@ TSharedRef<SWidget> FDefaultsEditorSummoner::CreateTabBody(const FWorkflowTabSpa
 	return BlueprintEditorPtr->GetDefaultEditor();
 }
 
-FConstructionScriptEditorSummoner::FConstructionScriptEditorSummoner(TSharedPtr<class FAssetEditorToolkit> InHostingApp) : FWorkflowTabFactory(FBlueprintEditorTabs::ConstructionScriptEditorID, InHostingApp)
+FConstructionScriptEditorSummoner::FConstructionScriptEditorSummoner(TSharedPtr<class FAssetEditorToolkit> InHostingApp)
+	: FWorkflowTabFactory(FBlueprintEditorTabs::ConstructionScriptEditorID, InHostingApp)
 {
 	TabLabel = LOCTEXT("ComponentsTabLabel", "Components");
 	TabIcon = FSlateIcon(FEditorStyle::GetStyleSetName(), "Kismet.Tabs.Components");
@@ -161,7 +164,8 @@ TSharedRef<SWidget> FConstructionScriptEditorSummoner::CreateTabBody(const FWork
 	return BlueprintEditorPtr->GetSCSEditor();
 }
 
-FSCSViewportSummoner::FSCSViewportSummoner(TSharedPtr<class FAssetEditorToolkit> InHostingApp) : FWorkflowTabFactory(FBlueprintEditorTabs::SCSViewportID, InHostingApp)
+FSCSViewportSummoner::FSCSViewportSummoner(TSharedPtr<class FAssetEditorToolkit> InHostingApp)
+	: FWorkflowTabFactory(FBlueprintEditorTabs::SCSViewportID, InHostingApp)
 {
 	TabLabel = LOCTEXT("SCSViewportTabLabel", "Viewport");
 	TabIcon = FSlateIcon(FEditorStyle::GetStyleSetName(), "LevelEditor.Tabs.Viewports");
@@ -195,7 +199,8 @@ TSharedRef<SWidget> FSCSViewportSummoner::CreateTabBody(const FWorkflowTabSpawnI
 	}
 }
 
-FPaletteSummoner::FPaletteSummoner(TSharedPtr<class FAssetEditorToolkit> InHostingApp) : FWorkflowTabFactory(FBlueprintEditorTabs::PaletteID, InHostingApp)
+FPaletteSummoner::FPaletteSummoner(TSharedPtr<class FAssetEditorToolkit> InHostingApp)
+	: FWorkflowTabFactory(FBlueprintEditorTabs::PaletteID, InHostingApp)
 {
 	TabLabel = LOCTEXT("PaletteTabTitle", "Palette");
 	TabIcon = FSlateIcon(FEditorStyle::GetStyleSetName(), "Kismet.Tabs.Palette");
@@ -213,7 +218,8 @@ TSharedRef<SWidget> FPaletteSummoner::CreateTabBody(const FWorkflowTabSpawnInfo&
 	return BlueprintEditorPtr->GetPalette();
 }
 
-FMyBlueprintSummoner::FMyBlueprintSummoner(TSharedPtr<class FAssetEditorToolkit> InHostingApp) : FWorkflowTabFactory(FBlueprintEditorTabs::MyBlueprintID, InHostingApp)
+FMyBlueprintSummoner::FMyBlueprintSummoner(TSharedPtr<class FAssetEditorToolkit> InHostingApp)
+	: FWorkflowTabFactory(FBlueprintEditorTabs::MyBlueprintID, InHostingApp)
 {
 	TabLabel = LOCTEXT("MyBlueprintTabLabel", "My Blueprint");
 	TabIcon = FSlateIcon(FEditorStyle::GetStyleSetName(), "ClassIcon.BlueprintCore");
@@ -231,7 +237,8 @@ TSharedRef<SWidget> FMyBlueprintSummoner::CreateTabBody(const FWorkflowTabSpawnI
 	return BlueprintEditorPtr->GetMyBlueprintWidget().ToSharedRef();
 }
 
-FCompilerResultsSummoner::FCompilerResultsSummoner(TSharedPtr<class FAssetEditorToolkit> InHostingApp) : FWorkflowTabFactory(FBlueprintEditorTabs::CompilerResultsID, InHostingApp)
+FCompilerResultsSummoner::FCompilerResultsSummoner(TSharedPtr<class FAssetEditorToolkit> InHostingApp)
+	: FWorkflowTabFactory(FBlueprintEditorTabs::CompilerResultsID, InHostingApp)
 {
 	TabLabel = LOCTEXT("CompilerResultsTabTitle", "Compiler Results");
 	TabIcon = FSlateIcon(FEditorStyle::GetStyleSetName(), "Kismet.Tabs.CompilerResults");
@@ -249,7 +256,8 @@ TSharedRef<SWidget> FCompilerResultsSummoner::CreateTabBody(const FWorkflowTabSp
 	return BlueprintEditorPtr->GetCompilerResults();
 }
 
-FFindResultsSummoner::FFindResultsSummoner(TSharedPtr<class FAssetEditorToolkit> InHostingApp) : FWorkflowTabFactory(FBlueprintEditorTabs::FindResultsID, InHostingApp)
+FFindResultsSummoner::FFindResultsSummoner(TSharedPtr<class FAssetEditorToolkit> InHostingApp)
+	: FWorkflowTabFactory(FBlueprintEditorTabs::FindResultsID, InHostingApp)
 {
 	TabLabel = LOCTEXT("FindResultsTabTitle", "Find Results");
 	TabIcon = FSlateIcon(FEditorStyle::GetStyleSetName(), "Kismet.Tabs.FindResults");
diff --git a/Engine/Source/Editor/Kismet/Private/BlueprintEditorTabFactories.h b/Engine/Source/Editor/Kismet/Private/BlueprintEditorTabFactories.h
--- a/Engine/Source/Editor/Kismet/Private/BlueprintEditorTabFactories.h
+++ b/Engine/Source/Editor/Kismet/Private/BlueprintEditorTabFactories.h
@@ -70,7 +70,7 @@ public:
 
 	virtual FText GetTabToolTipText(const FWorkflowTabSpawnInfo& Info) const override
 	{
-		return LOCTEXT("DefaultsEditorTooltip", "The defaults editor lets you set the default value for all variables in your Blueprint.");
+		return LOCTEXT("DefaultsEditorTooltip", "The class defaults editor lets you set the default value for all variables in your Blueprint.");
 	}
 };
 
diff --git a/Engine/Source/Editor/Kismet/Private/SBlueprintEditorToolbar.cpp b/Engine/Source/Editor/Kismet/Private/SBlueprintEditorToolbar.cpp
--- a/Engine/Source/Editor/Kismet/Private/SBlueprintEditorToolbar.cpp
+++ b/Engine/Source/Editor/Kismet/Private/SBlueprintEditorToolbar.cpp
@@ -341,7 +341,7 @@ void FFullBlueprintEditorCommands::RegisterCommands()
 	UI_COMMAND(SaveOnCompile_Always, "Always", "Sets the save-on-compile option to 'Always', meaning that your Blueprints will be saved whenever they ar compiled (even if there were errors)", EUserInterfaceActionType::RadioButton, FInputGesture());
 
 	UI_COMMAND(SwitchToScriptingMode, "Graph", "Switches to Graph Editing Mode", EUserInterfaceActionType::ToggleButton, FInputGesture());
-	UI_COMMAND(SwitchToBlueprintDefaultsMode, "Defaults", "Switches to Blueprint Defaults Mode", EUserInterfaceActionType::ToggleButton, FInputGesture());
+	UI_COMMAND(SwitchToBlueprintDefaultsMode, "Defaults", "Switches to Class Defaults Mode", EUserInterfaceActionType::ToggleButton, FInputGesture());
 	UI_COMMAND(SwitchToComponentsMode, "Components", "Switches to Components Mode", EUserInterfaceActionType::ToggleButton, FInputGesture());
 
 	UI_COMMAND(EditGlobalOptions, "Class Settings", "Edit Class Settings (Previously known as Blueprint Props)", EUserInterfaceActionType::Button, FInputGesture());
@@ -686,7 +686,7 @@ TArray< TSharedPtr< SWidget> > FBlueprintEditorToolbar::GenerateToolbarWidgets(c
 			.OnSetActiveMode(ActiveModeSetter)
 			.CanBeSelected(BlueprintObj ? FBlueprintEditorUtils::DoesSupportDefaults(BlueprintObj) : false)
 			.ToolTip(IDocumentation::Get()->CreateToolTip(
-				LOCTEXT("BlueprintDefaultsModeButtonTooltip", "Switch to Blueprint Defaults Mode"),
+				LOCTEXT("BlueprintDefaultsModeButtonTooltip", "Switch to Class Defaults Mode"),
 				NULL,
 				TEXT("Shared/Editors/BlueprintEditor"),
 				TEXT("DefaultsMode")))
diff --git a/Engine/Source/Editor/Kismet/Private/SSCSEditor.cpp b/Engine/Source/Editor/Kismet/Private/SSCSEditor.cpp
--- a/Engine/Source/Editor/Kismet/Private/SSCSEditor.cpp
+++ b/Engine/Source/Editor/Kismet/Private/SSCSEditor.cpp
@@ -5203,7 +5203,7 @@ void SSCSEditor::OnResetToBlueprintDefaults() const
 
 	if (Actor != NULL && Blueprint != NULL && Actor->GetClass()->ClassGeneratedBy == Blueprint)
 	{
-		const FScopedTransaction Transaction(LOCTEXT("ResetToBlueprintDefaults_Transaction", "Reset to Blueprint Defaults"));
+		const FScopedTransaction Transaction(LOCTEXT("ResetToBlueprintDefaults_Transaction", "Reset to Class Defaults"));
 
 		{
 			AActor* BlueprintCDO = Actor->GetClass()->GetDefaultObject<AActor>();
diff --git a/Engine/Source/Editor/Kismet/Public/BlueprintEditor.h b/Engine/Source/Editor/Kismet/Public/BlueprintEditor.h
--- a/Engine/Source/Editor/Kismet/Public/BlueprintEditor.h
+++ b/Engine/Source/Editor/Kismet/Public/BlueprintEditor.h
@@ -607,9 +607,6 @@ protected:
 	/** Handles the unloading of Blueprints (by closing the editor, if it operating on the Blueprint being unloaded)*/
 	void OnBlueprintUnloaded(UBlueprint* InBlueprint);
 
-	/** Get title for Inspector 2 tab*/
-	virtual FText GetDefaultEditorTitle();
-
 	//@TODO: Should the breakpoint/watch modification operations be whole-blueprint, or current-graph?
 
 	/** Deletes all breakpoints for the blueprint being edited */
diff --git a/Engine/Source/Editor/Kismet/Public/BlueprintEditorModes.h b/Engine/Source/Editor/Kismet/Public/BlueprintEditorModes.h
--- a/Engine/Source/Editor/Kismet/Public/BlueprintEditorModes.h
+++ b/Engine/Source/Editor/Kismet/Public/BlueprintEditorModes.h
@@ -74,7 +74,7 @@ public:
 protected:
 	TWeakPtr<FBlueprintEditor> MyBlueprintEditor;
 
-	// Set of spawnable tabs in blueprint defaults mode
+	// Set of spawnable tabs in Class Defaults mode
 	FWorkflowAllowedTabSet BlueprintDefaultsTabFactories;
 };
 
diff --git a/Engine/Source/Editor/Persona/Private/Persona.cpp b/Engine/Source/Editor/Persona/Private/Persona.cpp
--- a/Engine/Source/Editor/Persona/Private/Persona.cpp
+++ b/Engine/Source/Editor/Persona/Private/Persona.cpp
@@ -1653,11 +1653,6 @@ void FPersona::Compile()
 	}
 }
 
-FText FPersona::GetDefaultEditorTitle()
-{
-	return NSLOCTEXT("Kismet", "PreviewParamatersTabTitle", "Preview Parameters");
-}
-
 FName FPersona::GetToolkitContextFName() const
 {
 	return GetToolkitFName();
diff --git a/Engine/Source/Editor/Persona/Private/Persona.h b/Engine/Source/Editor/Persona/Private/Persona.h
--- a/Engine/Source/Editor/Persona/Private/Persona.h
+++ b/Engine/Source/Editor/Persona/Private/Persona.h
@@ -219,7 +219,6 @@ protected:
 	virtual void StartEditingDefaults(bool bAutoFocus, bool bForceRefresh = false) override;
 	virtual void Compile() override;
 	virtual void OnGraphEditorFocused(const TSharedRef<class SGraphEditor>& InGraphEditor) override;
-	virtual FText GetDefaultEditorTitle() override;
 	virtual void OnConvertToSequenceEvaluator() override;
 	virtual void OnConvertToSequencePlayer() override;
 	virtual void OnConvertToBlendSpaceEvaluator() override;
diff --git a/Engine/Source/Editor/Persona/Private/Shared/PersonaMode.cpp b/Engine/Source/Editor/Persona/Private/Shared/PersonaMode.cpp
--- a/Engine/Source/Editor/Persona/Private/Shared/PersonaMode.cpp
+++ b/Engine/Source/Editor/Persona/Private/Shared/PersonaMode.cpp
@@ -21,7 +21,7 @@ const FName FPersonaTabs::SkeletonTreeViewID("SkeletonTreeView");		//@TODO: Name
 const FName FPersonaTabs::RetargetManagerID("RetargetManager");
 // Anim Blueprint params
 // Explorer
-// Blueprint Defaults
+// Class Defaults
 const FName FPersonaTabs::AnimBlueprintDefaultsEditorID("AnimBlueprintDefaultsEditor");
 
 const FName FPersonaTabs::AnimBlueprintParentPlayerEditorID("AnimBlueprintParentPlayerEditor");
@@ -262,7 +262,7 @@ FAnimBlueprintDefaultsEditorSummoner::FAnimBlueprintDefaultsEditorSummoner(TShar
 	CurrentMode = EAnimBlueprintEditorMode::PreviewMode;
 
 	ViewMenuDescription = LOCTEXT("AnimBlueprintDefaultsView", "Defaults");
-	ViewMenuTooltip = LOCTEXT("AnimBlueprintDefaultsView_ToolTip", "Shows the animation blueprint defaults/preview editor view");
+	ViewMenuTooltip = LOCTEXT("AnimBlueprintDefaultsView_ToolTip", "Shows the animation class defaults/preview editor view");
 }
 
 TSharedRef<SWidget> FAnimBlueprintDefaultsEditorSummoner::CreateTabBody(const FWorkflowTabSpawnInfo& Info) const
@@ -307,7 +307,7 @@ TSharedRef<SWidget> FAnimBlueprintDefaultsEditorSummoner::CreateTabBody(const FW
 						.Style(FEditorStyle::Get(), "RadioButton")
 						.IsChecked( this, &FAnimBlueprintDefaultsEditorSummoner::IsChecked, EAnimBlueprintEditorMode::DefaultsMode )
 						.OnCheckStateChanged( this, &FAnimBlueprintDefaultsEditorSummoner::OnCheckedChanged, EAnimBlueprintEditorMode::DefaultsMode )
-						.ToolTip(IDocumentation::Get()->CreateToolTip(	LOCTEXT("AnimBlueprintPropertyEditorDefaultMode", "Switch to editing the blueprint property defaults"),
+						.ToolTip(IDocumentation::Get()->CreateToolTip(	LOCTEXT("AnimBlueprintPropertyEditorDefaultMode", "Switch to editing the class defaults"),
 																		NULL,
 																		TEXT("Shared/Editors/Persona"),
 																		TEXT("AnimBlueprintPropertyEditorDefaultMode")))
diff --git a/Engine/Source/Editor/Persona/Private/Shared/PersonaMode.h b/Engine/Source/Editor/Persona/Private/Shared/PersonaMode.h
--- a/Engine/Source/Editor/Persona/Private/Shared/PersonaMode.h
+++ b/Engine/Source/Editor/Persona/Private/Shared/PersonaMode.h
@@ -21,7 +21,7 @@ struct FPersonaTabs
 	// Skeleton/Sockets
 	// Anim Blueprint Params
 	// Explorer
-	// Blueprint Defaults
+	// Class Defaults
 	static const FName AnimBlueprintDefaultsEditorID;
 	static const FName AnimBlueprintParentPlayerEditorID;
 	// Anim Document
diff --git a/Engine/Source/Editor/UnrealEd/Private/Kismet2/BlueprintEditorUtils.cpp b/Engine/Source/Editor/UnrealEd/Private/Kismet2/BlueprintEditorUtils.cpp
--- a/Engine/Source/Editor/UnrealEd/Private/Kismet2/BlueprintEditorUtils.cpp
+++ b/Engine/Source/Editor/UnrealEd/Private/Kismet2/BlueprintEditorUtils.cpp
@@ -1438,7 +1438,7 @@ void FBlueprintEditorUtils::PostDuplicateBlueprint(UBlueprint* Blueprint)
 	{
 		check(Blueprint->GeneratedClass != NULL);
 		{
-			// Grab the old CDO, which contains the blueprint defaults
+			// Grab the old CDO, which contains the class defaults
 			UClass* OldBPGCAsClass = Blueprint->GeneratedClass;
 			UBlueprintGeneratedClass* OldBPGC = (UBlueprintGeneratedClass*)(OldBPGCAsClass);
 			UObject* OldCDO = OldBPGC->GetDefaultObject();
