[INTEGRATE] Change 2419679 by [REDACTED_PERSON][REDACTED_HANDLE] on 2015/01/26 20:06:41

	Blueprint Editor: Make sure tab spawners pass their tooltip in when registering for the Windows view [UE-8354]

diff --git a/Engine/Source/Editor/Kismet/Private/WorkflowCentricApplication.cpp b/Engine/Source/Editor/Kismet/Private/WorkflowCentricApplication.cpp
--- a/Engine/Source/Editor/Kismet/Private/WorkflowCentricApplication.cpp
+++ b/Engine/Source/Editor/Kismet/Private/WorkflowCentricApplication.cpp
@@ -92,6 +92,7 @@ void FWorkflowCentricApplication::PushTabFactories(FWorkflowAllowedTabSet& Facto
 		TSharedPtr<FWorkflowTabFactory> SomeFactory = FactoryIt.Value();
 		FTabSpawnerEntry& SpawnerEntry = TabManager->RegisterTabSpawner(SomeFactory->GetIdentifier(), FOnSpawnTab::CreateRaw(this, &FWorkflowCentricApplication::CreatePanelTab, SomeFactory))
 			.SetDisplayName(SomeFactory->ConstructTabName(SpawnInfo).Get())
+			.SetTooltipText(SomeFactory->GetTabToolTipText(SpawnInfo))
 			.SetGroup(CurrentAppModePtr->GetWorkspaceMenuCategory());
 		
 		// Add the tab icon to the menu entry if one was provided
