[INTEGRATE] Change 2419766 by [REDACTED_PERSON][REDACTED_HANDLE] on 2015/01/26 20:29:32

	Blueprint Editor: Made "Class Defaults" and "Class Settings" toggle buttons, indicating when they are the thing currently being edited by the main details panel [UE-8353]

diff --git a/Engine/Source/Editor/Kismet/Private/BlueprintEditor.cpp b/Engine/Source/Editor/Kismet/Private/BlueprintEditor.cpp
--- a/Engine/Source/Editor/Kismet/Private/BlueprintEditor.cpp
+++ b/Engine/Source/Editor/Kismet/Private/BlueprintEditor.cpp
@@ -2173,12 +2173,14 @@ void FBlueprintEditor::CreateDefaultCommands()
 	ToolkitCommands->MapAction(
 		FFullBlueprintEditorCommands::Get().EditGlobalOptions,
 		FExecuteAction::CreateSP(this, &FBlueprintEditor::EditGlobalOptions_Clicked),
-		FCanExecuteAction());
+		FCanExecuteAction(),
+		FIsActionChecked::CreateSP(this, &FBlueprintEditor::IsDetailsPanelEditingGlobalOptions));
 
 	ToolkitCommands->MapAction(
 		FFullBlueprintEditorCommands::Get().EditClassDefaults,
 		FExecuteAction::CreateSP(this, &FBlueprintEditor::EditClassDefaults_Clicked),
-		FCanExecuteAction());
+		FCanExecuteAction(),
+		FIsActionChecked::CreateSP(this, &FBlueprintEditor::IsDetailsPanelEditingClassDefaults));
 
 	// Edit menu actions
 	ToolkitCommands->MapAction( FBlueprintEditorCommands::Get().FindInBlueprint,
@@ -2497,6 +2499,11 @@ bool FBlueprintEditor::ReparentBlueprint_IsVisible() const
 	}
 }
 
+bool FBlueprintEditor::IsDetailsPanelEditingGlobalOptions() const
+{
+	return CurrentUISelection == FBlueprintEditor::SelectionState_ClassSettings;
+}
+
 void FBlueprintEditor::EditGlobalOptions_Clicked()
 {
 	UBlueprint* Blueprint = GetBlueprintObj();
@@ -2516,6 +2523,11 @@ void FBlueprintEditor::EditGlobalOptions_Clicked()
 	}
 }
 
+bool FBlueprintEditor::IsDetailsPanelEditingClassDefaults() const
+{
+	return CurrentUISelection == FBlueprintEditor::SelectionState_ClassDefaults;
+}
+
 void FBlueprintEditor::EditClassDefaults_Clicked()
 {
 	UBlueprint* Blueprint = GetBlueprintObj();
diff --git a/Engine/Source/Editor/Kismet/Private/SBlueprintEditorToolbar.cpp b/Engine/Source/Editor/Kismet/Private/SBlueprintEditorToolbar.cpp
--- a/Engine/Source/Editor/Kismet/Private/SBlueprintEditorToolbar.cpp
+++ b/Engine/Source/Editor/Kismet/Private/SBlueprintEditorToolbar.cpp
@@ -344,8 +344,8 @@ void FFullBlueprintEditorCommands::RegisterCommands()
 	UI_COMMAND(SwitchToBlueprintDefaultsMode, "Defaults", "Switches to Class Defaults Mode", EUserInterfaceActionType::ToggleButton, FInputGesture());
 	UI_COMMAND(SwitchToComponentsMode, "Components", "Switches to Components Mode", EUserInterfaceActionType::ToggleButton, FInputGesture());
 
-	UI_COMMAND(EditGlobalOptions, "Class Settings", "Edit Class Settings (Previously known as Blueprint Props)", EUserInterfaceActionType::Button, FInputGesture());
-	UI_COMMAND(EditClassDefaults, "Class Defaults", "Edit the initial values of your class.", EUserInterfaceActionType::Button, FInputGesture());
+	UI_COMMAND(EditGlobalOptions, "Class Settings", "Edit Class Settings (Previously known as Blueprint Props)", EUserInterfaceActionType::ToggleButton, FInputGesture());
+	UI_COMMAND(EditClassDefaults, "Class Defaults", "Edit the initial values of your class.", EUserInterfaceActionType::ToggleButton, FInputGesture());
 
 	UI_COMMAND(JumpToErrorNode, "Jump to Error Node", "When enabled, then the Blueprint will snap focus to nodes producing an error during compilation", EUserInterfaceActionType::ToggleButton, FInputGesture());
 }
diff --git a/Engine/Source/Editor/Kismet/Public/BlueprintEditor.h b/Engine/Source/Editor/Kismet/Public/BlueprintEditor.h
--- a/Engine/Source/Editor/Kismet/Public/BlueprintEditor.h
+++ b/Engine/Source/Editor/Kismet/Public/BlueprintEditor.h
@@ -563,9 +563,15 @@ protected:
 	/** Called to start a quick find (focus the search box in the explorer tab) */
 	void FindInBlueprint_Clicked();
 
+	// Is the main details panel currently showing 'Global options' (e.g., class metadata)?
+	bool IsDetailsPanelEditingGlobalOptions() const;
+
 	/** Edit the class settings aka Blueprint global options */
 	void EditGlobalOptions_Clicked();
 
+	// Is the main details panel currently showing 'Class defaults' (Note: Has nothing to do with the standalone class defaults panel)?
+	bool IsDetailsPanelEditingClassDefaults() const;
+
 	/** Edit the class defaults */
 	void EditClassDefaults_Clicked();
 
