[INTEGRATE] Change 2419912 by [REDACTED_PERSON][REDACTED_HANDLE] on 2015/01/26 23:36:23

	Build fix for Clang. !X == Y does not do what was intended.

diff --git a/Engine/Source/Editor/Kismet/Private/SSCSEditor.cpp b/Engine/Source/Editor/Kismet/Private/SSCSEditor.cpp
--- a/Engine/Source/Editor/Kismet/Private/SSCSEditor.cpp
+++ b/Engine/Source/Editor/Kismet/Private/SSCSEditor.cpp
@@ -4459,7 +4459,7 @@ void SSCSEditor::OnDeleteNodes()
 			{
 				// Default to the parent node
 				NewSelection = Node->GetParent();
-				if (NewSelection.IsValid() && !NewSelection->GetNodeType() == FSCSEditorTreeNode::RootActorNode)
+				if (NewSelection.IsValid() && NewSelection->GetNodeType() != FSCSEditorTreeNode::RootActorNode)
 				{
 					// If we have sibling nodes, find the one that immediately precedes the one being removed
 					const TArray<FSCSEditorTreeNodePtrType>& ChildNodes = NewSelection->GetChildren();
