UE-8380:  Editor crashes if Timeline Play Rate node is copied/pasted in blueprints.  Added explicit check for classes being the same, in case the BP hasn't been compiled yet since the member was added

diff --git a/Engine/Source/Editor/Kismet/Private/BlueprintEditor.cpp b/Engine/Source/Editor/Kismet/Private/BlueprintEditor.cpp
--- a/Engine/Source/Editor/Kismet/Private/BlueprintEditor.cpp
+++ b/Engine/Source/Editor/Kismet/Private/BlueprintEditor.cpp
@@ -4900,7 +4900,7 @@ struct FUpdatePastedNodes
 
 					const bool bCanCurrentBlueprintReplace = TargetClass
 						&& CurrentClass->IsChildOf(TargetClass) // If current class if of the same type, it has the called member
-						&& !CallOnMember->MemberVariableToCallOn.IsSelfContext()
+						&& (!CallOnMember->MemberVariableToCallOn.IsSelfContext() && (TargetClass != CurrentClass)) // Make sure the class isn't self, using a explicit check in case the class hasn't been compiled since the member was added
 						&& bTargetIsNullOrSingleLinked;
 
 					if (bCanCurrentBlueprintReplace) 
