Fix crash in the scs editor with an invalid actor context

diff --git a/Engine/Source/Editor/Kismet/Private/SSCSEditor.cpp b/Engine/Source/Editor/Kismet/Private/SSCSEditor.cpp
--- a/Engine/Source/Editor/Kismet/Private/SSCSEditor.cpp
+++ b/Engine/Source/Editor/Kismet/Private/SSCSEditor.cpp
@@ -2541,16 +2541,19 @@ FText SSCS_RowWidget_ActorRoot::GetActorDisplayText() const
 		{
 			FString Name;
 			AActor* DefaultActor = SCSEditorPtr->ActorContext.Get();
-			UBlueprint* Blueprint = UBlueprint::GetBlueprintFromClass(DefaultActor->GetClass());
-			if (Blueprint != nullptr)
+			if( DefaultActor )
 			{
-				Blueprint->GetName(Name);
-				return FText::Format(LOCTEXT("DefaultActor_Name", "{0} (self)"), FText::FromString(Name));
-			}
-			else
-			{
-				DefaultActor->GetName(Name);
-				return FText::Format(LOCTEXT("DefaultActor_Name", "{0} (Instance)"), FText::FromString(Name));
+				UBlueprint* Blueprint = UBlueprint::GetBlueprintFromClass(DefaultActor->GetClass());
+				if(Blueprint != nullptr)
+				{
+					Blueprint->GetName(Name);
+					return FText::Format(LOCTEXT("DefaultActor_Name", "{0} (self)"), FText::FromString(Name));
+				}
+				else
+				{
+					DefaultActor->GetName(Name);
+					return FText::Format(LOCTEXT("DefaultActor_Name", "{0} (Instance)"), FText::FromString(Name));
+				}
 			}
 		}
 	}
