[INTEGRATE] Change 2419894 by [REDACTED_PERSON][REDACTED_HANDLE] on 2015/01/26 23:19:28

	Component Editor: Removed a number of invalid assumptions about ActorContext being non-null if ever set
	- Moved SCS editor mode enum to be a separate namespace to prevent mistakes (found a couple of ActorInstance ? X : Y instances that should have been (Actor != null) ? X : Y but were always using X since ActorInstance = 1)
	[UE-8368]

diff --git a/Engine/Source/Editor/Kismet/Private/SSCSEditor.cpp b/Engine/Source/Editor/Kismet/Private/SSCSEditor.cpp
--- a/Engine/Source/Editor/Kismet/Private/SSCSEditor.cpp
+++ b/Engine/Source/Editor/Kismet/Private/SSCSEditor.cpp
@@ -1363,7 +1363,7 @@ EVisibility SSCS_RowWidget::GetAssetVisibility() const
 FSlateColor SSCS_RowWidget::GetColorTint() const
 {
 	FSCSEditorTreeNodePtrType NodePtr = GetNode();
-	if(SCSEditor.Pin()->GetEditorMode() == SSCSEditor::EEditorMode::BlueprintSCS)
+	if(SCSEditor.Pin()->GetEditorMode() == EComponentEditorMode::BlueprintSCS)
 	{
 		if(NodePtr->IsNative())
 		{
@@ -1374,7 +1374,7 @@ FSlateColor SSCS_RowWidget::GetColorTint() const
 			return FLinearColor(0.08f,0.35f,0.6f);
 		}
 	}
-	else    // SSCSEditor::EEditorMode::ActorInstance
+	else    // EComponentEditorMode::ActorInstance
 	{
 		if(!NodePtr->IsUserInstanced())
 		{
@@ -1425,7 +1425,7 @@ TSharedPtr<SWidget> SSCS_RowWidget::BuildSceneRootDropActionMenu(FSCSEditorTreeN
 		const FText NodeVariableNameText = FText::FromName( GetNode()->GetVariableName() );
 
 		bool bDroppedInSameBlueprint = true;
-		if(SCSEditor.Pin()->GetEditorMode() == SSCSEditor::EEditorMode::BlueprintSCS)
+		if (SCSEditor.Pin()->GetEditorMode() == EComponentEditorMode::BlueprintSCS)
 		{
 			bDroppedInSameBlueprint = DroppedNodePtr->GetBlueprint() == GetBlueprint();
 		}
@@ -1623,7 +1623,7 @@ void SSCS_RowWidget::OnDragEnter( const FGeometry& MyGeometry, const FDragDropEv
 						Message = LOCTEXT("DropActionToolTip_AttachToOrMakeNewRoot", "Drop here to see available actions.");
 						DragRowOp->PendingDropAction = FSCSRowDragDropOp::DropAction_AttachToOrMakeNewRoot;
 					}
-					else if (SCSEditor.Pin()->GetEditorMode() == SSCSEditor::EEditorMode::BlueprintSCS && DraggedNodePtr->GetBlueprint() != GetBlueprint())
+					else if (SCSEditor.Pin()->GetEditorMode() == EComponentEditorMode::BlueprintSCS && DraggedNodePtr->GetBlueprint() != GetBlueprint())
 					{
 						if (bCanMakeNewRoot)
 						{
@@ -1838,7 +1838,7 @@ void SSCS_RowWidget::OnAttachToDropAction(const TArray<FSCSEditorTreeNodePtrType
 	bool bRegenerateTreeNodes = false;
 	const FScopedTransaction TransactionContext(DroppedNodePtrs.Num() > 1 ? LOCTEXT("AttachComponents", "Attach Components") : LOCTEXT("AttachComponent", "Attach Component"));
 
-	if(SCSEditorPtr->GetEditorMode() == SSCSEditor::EEditorMode::BlueprintSCS)
+	if (SCSEditorPtr->GetEditorMode() == EComponentEditorMode::BlueprintSCS)
 	{
 		// Get the current Blueprint context
 		UBlueprint* Blueprint = GetBlueprint();
@@ -1948,7 +1948,7 @@ void SSCS_RowWidget::OnAttachToDropAction(const TArray<FSCSEditorTreeNodePtrType
 			}
 		}
 	}
-	else    // SSCSEditor::EEditorMode::ActorInstance
+	else    // EComponentEditorMode::ActorInstance
 	{
 		for(const auto& DroppedNodePtr : DroppedNodePtrs)
 		{
@@ -1983,7 +1983,7 @@ void SSCS_RowWidget::OnDetachFromDropAction(const TArray<FSCSEditorTreeNodePtrTy
 
 	const FScopedTransaction TransactionContext(DroppedNodePtrs.Num() > 1 ? LOCTEXT("DetachComponents", "Detach Components") : LOCTEXT("DetachComponent", "Detach Component"));
 
-	if(SCSEditorPtr->GetEditorMode() == SSCSEditor::EEditorMode::BlueprintSCS)
+	if (SCSEditorPtr->GetEditorMode() == EComponentEditorMode::BlueprintSCS)
 	{
 		// Get the current "preview" Actor instance
 		AActor* PreviewActor = SCSEditorPtr->PreviewActor.Get();
@@ -2057,7 +2057,7 @@ void SSCS_RowWidget::OnDetachFromDropAction(const TArray<FSCSEditorTreeNodePtrTy
 			}
 		}
 	}
-	else    // SSCSEditor::EEditorMode::ActorInstance
+	else    // EComponentEditorMode::ActorInstance
 	{
 		for(const auto& DroppedNodePtr : DroppedNodePtrs)
 		{
@@ -2091,7 +2091,7 @@ void SSCS_RowWidget::OnMakeNewRootDropAction(FSCSEditorTreeNodePtrType DroppedNo
 	// Create a transaction record
 	const FScopedTransaction TransactionContext(LOCTEXT("MakeNewSceneRoot", "Make New Scene Root"));
 
-	if(SCSEditorPtr->GetEditorMode() == SSCSEditor::EEditorMode::BlueprintSCS)
+	if (SCSEditorPtr->GetEditorMode() == EComponentEditorMode::BlueprintSCS)
 	{
 		// Get the current Blueprint context
 		UBlueprint* Blueprint = GetBlueprint();
@@ -2175,7 +2175,7 @@ void SSCS_RowWidget::OnMakeNewRootDropAction(FSCSEditorTreeNodePtrType DroppedNo
 			SceneRootNodePtr->AddChild(OldSceneRootNodePtr);
 		}
 	}
-	else    // SSCSEditor::EEditorMode::ActorInstance
+	else    // EComponentEditorMode::ActorInstance
 	{
 		if(DroppedNodePtr->GetParent().IsValid())
 		{
@@ -2210,7 +2210,7 @@ void SSCS_RowWidget::PostDragDropAction(bool bRegenerateTreeNodes)
 
 		PinnedEditor->RefreshSelectionDetails();
 
-		if(PinnedEditor->GetEditorMode() == SSCSEditor::EEditorMode::BlueprintSCS)
+		if (PinnedEditor->GetEditorMode() == EComponentEditorMode::BlueprintSCS)
 		{
 			if(NodePtr.IsValid())
 			{
@@ -2300,7 +2300,7 @@ FString SSCS_RowWidget::GetDocumentationLink() const
 {
 	check(SCSEditor.IsValid());
 
-	if(SCSEditor.Pin()->GetEditorMode() == SSCSEditor::EEditorMode::BlueprintSCS)
+	if (SCSEditor.Pin()->GetEditorMode() == EComponentEditorMode::BlueprintSCS)
 	{
 		FSCSEditorTreeNodePtrType NodePtr = GetNode();
 		if( NodePtr == SCSEditor.Pin()->SceneRootNodePtr
@@ -2309,7 +2309,7 @@ FString SSCS_RowWidget::GetDocumentationLink() const
 			return TEXT("Shared/Editors/BlueprintEditor/ComponentsMode");
 		}
 	}
-	else    // SSCSEditor::EEditorMode::ActorInstance
+	else    // EComponentEditorMode::ActorInstance
 	{
 		// @TODO - Actor instance mode
 	}
@@ -2321,7 +2321,7 @@ FString SSCS_RowWidget::GetDocumentationExcerptName() const
 {
 	check(SCSEditor.IsValid());
 
-	if(SCSEditor.Pin()->GetEditorMode() == SSCSEditor::EEditorMode::BlueprintSCS)
+	if (SCSEditor.Pin()->GetEditorMode() == EComponentEditorMode::BlueprintSCS)
 	{
 		FSCSEditorTreeNodePtrType NodePtr = GetNode();
 		if( NodePtr == SCSEditor.Pin()->SceneRootNodePtr)
@@ -2337,7 +2337,7 @@ FString SSCS_RowWidget::GetDocumentationExcerptName() const
 			return TEXT("InheritedComponents");
 		}
 	}
-	else    // SSCSEditor::EEditorMode::ActorInstance
+	else    // EComponentEditorMode::ActorInstance
 	{
 		// @TODO - Actor instance mode
 	}
@@ -2409,7 +2409,7 @@ void SSCS_RowWidget::OnNameTextCommit(const FText& InNewName, ETextCommit::Type
 	// No need to call UpdateTree() in SCS editor mode; it will already be called by MBASM internally
 	check(SCSEditor.IsValid());
 	TSharedPtr<SSCSEditor> PinnedEditor = SCSEditor.Pin();
-	if(PinnedEditor.IsValid() && PinnedEditor->GetEditorMode() == SSCSEditor::EEditorMode::ActorInstance)
+	if (PinnedEditor.IsValid() && PinnedEditor->GetEditorMode() == EComponentEditorMode::ActorInstance)
 	{
 		PinnedEditor->UpdateTree();
 	}
@@ -2526,7 +2526,7 @@ const FSlateBrush* SSCS_RowWidget_ActorRoot::GetActorIcon() const
 		TSharedPtr<SSCSEditor> SCSEditorPtr = SCSEditor.Pin();
 		if (SCSEditorPtr->ActorContext.IsSet())
 		{
-			return FClassIconFinder::FindIconForActor(SCSEditorPtr->ActorContext.Get());
+			return FClassIconFinder::FindIconForActor(SCSEditorPtr->GetActorContext());
 		}
 	}
 	return nullptr;
@@ -2539,18 +2539,19 @@ FText SSCS_RowWidget_ActorRoot::GetActorDisplayText() const
 		TSharedPtr<SSCSEditor> SCSEditorPtr = SCSEditor.Pin();
 		if (SCSEditorPtr->ActorContext.IsSet())
 		{
-			FString Name;
 			AActor* DefaultActor = SCSEditorPtr->ActorContext.Get();
 			if( DefaultActor )
 			{
 				UBlueprint* Blueprint = UBlueprint::GetBlueprintFromClass(DefaultActor->GetClass());
 				if(Blueprint != nullptr)
 				{
+					FString Name;
 					Blueprint->GetName(Name);
 					return FText::Format(LOCTEXT("DefaultActor_Name", "{0} (self)"), FText::FromString(Name));
 				}
 				else
 				{
+					FString Name;
 					DefaultActor->GetName(Name);
 					return FText::Format(LOCTEXT("DefaultActor_Name", "{0} (Instance)"), FText::FromString(Name));
 				}
@@ -2566,9 +2567,8 @@ FText SSCS_RowWidget_ActorRoot::GetActorClassNameText() const
 	if (SCSEditor.IsValid())
 	{
 		TSharedPtr<SSCSEditor> SCSEditorPtr = SCSEditor.Pin();
-		if (SCSEditorPtr->ActorContext.IsSet())
+		if (AActor* DefaultActor = SCSEditorPtr->GetActorContext())
 		{
-			AActor* DefaultActor = SCSEditorPtr->ActorContext.Get();
 			Text = FText::FromString(DefaultActor->GetClass()->GetName());
 		}
 	}
@@ -2582,9 +2582,8 @@ FText SSCS_RowWidget_ActorRoot::GetActorSuperClassNameText() const
 	if (SCSEditor.IsValid())
 	{
 		TSharedPtr<SSCSEditor> SCSEditorPtr = SCSEditor.Pin();
-		if (SCSEditorPtr->ActorContext.IsSet())
+		if (AActor* DefaultActor = SCSEditorPtr->GetActorContext())
 		{
-			AActor* DefaultActor = SCSEditorPtr->ActorContext.Get();
 			Text = FText::FromString(DefaultActor->GetClass()->GetSuperClass()->GetName());
 		}
 	}
@@ -2598,9 +2597,8 @@ FText SSCS_RowWidget_ActorRoot::GetActorMobilityText() const
 	if (SCSEditor.IsValid())
 	{
 		TSharedPtr<SSCSEditor> SCSEditorPtr = SCSEditor.Pin();
-		if (SCSEditorPtr->ActorContext.IsSet())
+		if (AActor* DefaultActor = SCSEditorPtr->GetActorContext())
 		{
-			AActor* DefaultActor = SCSEditorPtr->ActorContext.Get();
 			USceneComponent* RootComponent = DefaultActor->GetRootComponent();
 			if (RootComponent)
 			{
@@ -2760,7 +2758,7 @@ void SSCSEditor::Construct( const FArguments& InArgs )
 			.OnContextMenuOpening(this, &SSCSEditor::CreateContextMenu)
 			.OnItemScrolledIntoView(this, &SSCSEditor::OnItemScrolledIntoView)
 			.OnMouseButtonDoubleClick(this, &SSCSEditor::HandleItemDoubleClicked)
-			.ClearSelectionOnClick(InArgs._EditorMode == EEditorMode::BlueprintSCS ? true : false)
+			.ClearSelectionOnClick(InArgs._EditorMode == EComponentEditorMode::BlueprintSCS ? true : false)
 			.ItemHeight(24)
 			.HeaderRow
 			(
@@ -2782,7 +2780,7 @@ void SSCSEditor::Construct( const FArguments& InArgs )
 			.OnContextMenuOpening(this, &SSCSEditor::CreateContextMenu)
 			.OnItemScrolledIntoView(this, &SSCSEditor::OnItemScrolledIntoView)
 			.OnMouseButtonDoubleClick(this, &SSCSEditor::HandleItemDoubleClicked)
-			.ClearSelectionOnClick(InArgs._EditorMode == EEditorMode::BlueprintSCS ? true : false)
+			.ClearSelectionOnClick(InArgs._EditorMode == EComponentEditorMode::BlueprintSCS ? true : false)
 			.ItemHeight(24)
 			.HeaderRow
 			(
@@ -2999,8 +2997,7 @@ END_SLATE_FUNCTION_BUILD_OPTIMIZATION
 
 UBlueprint* SSCSEditor::GetBlueprint() const
 {
-	AActor* Actor = ActorContext.Get();
-	if(Actor != nullptr)
+	if (AActor* Actor = GetActorContext())
 	{
 		UClass* ActorClass = Actor->GetClass();
 		check(ActorClass != nullptr);
@@ -3015,7 +3012,7 @@ void SSCSEditor::Tick( const FGeometry& AllottedGeometry, const double InCurrent
 {
 	SCompoundWidget::Tick(AllottedGeometry, InCurrentTime, InDeltaTime);
 
-	if(EditorMode == EEditorMode::ActorInstance)
+	if (EditorMode == EComponentEditorMode::ActorInstance)
 	{
 		TFunction<bool (const TArray<FSCSEditorTreeNodePtrType>&, int32&)> AreAnyNodesInvalidLambda = [&](const TArray<FSCSEditorTreeNodePtrType>& InNodes, int32& OutNumValidNodes) -> bool
 		{
@@ -3035,25 +3032,25 @@ void SSCSEditor::Tick( const FGeometry& AllottedGeometry, const double InCurrent
 			return bFoundInvalidNode;
 		};
 
-		const AActor* ActorInstance = ActorContext.Get();
-		check(ActorInstance);
-
-		int32 NumComponentInstances = 0;
-		for(auto CompIt = ActorInstance->GetComponents().CreateConstIterator(); CompIt; ++CompIt)
+		if (const AActor* ActorInstance = GetActorContext())
 		{
-			// Don't count editor-only components, because we don't show them
-			if(!(*CompIt)->IsEditorOnly())
+			int32 NumComponentInstances = 0;
+			for (auto CompIt = ActorInstance->GetComponents().CreateConstIterator(); CompIt; ++CompIt)
 			{
-				++NumComponentInstances;
+				// Don't count editor-only components, because we don't show them
+				if (!(*CompIt)->IsEditorOnly())
+				{
+					++NumComponentInstances;
+				}
 			}
-		}
 
-		int32 NumComponentNodes = 0;
-		if (AreAnyNodesInvalidLambda(GetRootNodes(), NumComponentNodes) || NumComponentNodes != NumComponentInstances)
-		{
-			UE_LOG(LogSCSEditor, Log, TEXT("Calling UpdateTree() from Tick()."));
+			int32 NumComponentNodes = 0;
+			if (AreAnyNodesInvalidLambda(GetRootNodes(), NumComponentNodes) || NumComponentNodes != NumComponentInstances)
+			{
+				UE_LOG(LogSCSEditor, Log, TEXT("Calling UpdateTree() from Tick()."));
 
-			UpdateTree();
+				UpdateTree();
+			}
 		}
 	}
 }
@@ -3144,7 +3141,7 @@ TSharedPtr< SWidget > SSCSEditor::CreateContextMenu() // @todo: make the context
 				MenuBuilder.AddMenuEntry( FGenericCommands::Get().Delete );
 				MenuBuilder.AddMenuEntry( FGenericCommands::Get().Rename );
 
-				if(EditorMode == EEditorMode::BlueprintSCS)
+				if (EditorMode == EComponentEditorMode::BlueprintSCS)
 				{
 					// Collect the classes of all selected objects
 					TArray<UClass*> SelectionClasses;
@@ -3441,7 +3438,7 @@ FSCSEditorTreeNodePtrType SSCSEditor::GetNodeFromActorComponent(const UActorComp
 
 	if(ActorComponent)
 	{
-		if (EditorMode == EEditorMode::BlueprintSCS)
+		if (EditorMode == EComponentEditorMode::BlueprintSCS)
 		{
 			// If the given component instance is not already an archetype object
 			if (!ActorComponent->IsTemplate())
@@ -3499,7 +3496,7 @@ FSCSEditorTreeNodePtrType SSCSEditor::GetNodeFromActorComponent(const UActorComp
 		}
 
 		// If we have a valid component archetype instance, attempt to find a tree node that corresponds to it
-		if((EditorMode == EEditorMode::BlueprintSCS && ActorComponent->IsTemplate()) || EditorMode == EEditorMode::ActorInstance)
+		if ((EditorMode == EComponentEditorMode::BlueprintSCS && ActorComponent->IsTemplate()) || EditorMode == EComponentEditorMode::ActorInstance)
 		{
 			const TArray<FSCSEditorTreeNodePtrType>& Nodes = GetRootNodes();
 			for (int32 i = 0; i < Nodes.Num() && !NodePtr.IsValid(); i++)
@@ -3648,13 +3645,13 @@ void SSCSEditor::UpdateTree(bool bRegenerateTreeNodes)
 		}
 
 		// Build the tree data source according to what mode we're in
-		if(EditorMode == EEditorMode::BlueprintSCS)
+		if (EditorMode == EComponentEditorMode::BlueprintSCS)
 		{
 			// Get the class default object
 			AActor* CDO = NULL;
 			TArray<UBlueprint*> ParentBPStack;
-			AActor* Actor = ActorContext.Get();
-			if(Actor != nullptr)
+
+			if(AActor* Actor = GetActorContext())
 			{
 				UClass* ActorClass = Actor->GetClass();
 				if(ActorClass != nullptr)
@@ -3730,11 +3727,10 @@ void SSCSEditor::UpdateTree(bool bRegenerateTreeNodes)
 				}
 			}
 		}
-		else    // EEditorMode::ActorInstance
+		else    // EComponentEditorMode::ActorInstance
 		{
 			// Get the actor instance that we're editing
-			AActor* ActorInstance = ActorContext.Get();
-			if(ActorInstance != nullptr)
+			if (AActor* ActorInstance = GetActorContext())
 			{
 				// Get the full set of instanced components
 				TInlineComponentArray<UActorComponent*> Components;
@@ -3973,7 +3969,7 @@ UActorComponent* SSCSEditor::AddNewComponent( UClass* NewComponentClass, UObject
 		Asset = nullptr;
 	}
 
-	if(EditorMode == EEditorMode::BlueprintSCS)
+	if (EditorMode == EComponentEditorMode::BlueprintSCS)
 	{
 		UBlueprint* Blueprint = GetBlueprint();
 		check(Blueprint != nullptr && Blueprint->SimpleConstructionScript != nullptr);
@@ -3992,30 +3988,30 @@ UActorComponent* SSCSEditor::AddNewComponent( UClass* NewComponentClass, UObject
 			FObjectReader(NewComponent, SavedProperties);
 		}
 	}
-	else    // EEditorMode::ActorInstance
+	else    // EComponentEditorMode::ActorInstance
 	{
-		AActor* ActorInstance = ActorContext.Get();
-		check(ActorInstance != nullptr);
-		
-		ActorInstance->Modify();
-
-		FName NewComponentName = NAME_None;
-		if (NewComponentClass->ClassGeneratedBy && NewComponentClass->GetName().EndsWith(TEXT("_C")))
+		if (AActor* ActorInstance = GetActorContext())
 		{
-			const FString NewClassName = NewComponentClass->GetName();
-			const int32 NewStrLen = NewClassName.Len() - 2;
-			NewComponentName = MakeUniqueObjectName(ActorInstance, NewComponentClass, FName(*NewClassName.Left(NewStrLen)));
-		}
+			ActorInstance->Modify();
 
-		if (Asset != nullptr)
-		{
-			NewComponentName = *FComponentEditorUtils::GenerateValidVariableNameFromAsset(Asset, ActorInstance);
-		}
-		else
-		{
-			NewComponentName = *FComponentEditorUtils::GenerateValidVariableName(NewComponentClass, ActorInstance);
+			FName NewComponentName = NAME_None;
+			if (NewComponentClass->ClassGeneratedBy && NewComponentClass->GetName().EndsWith(TEXT("_C")))
+			{
+				const FString NewClassName = NewComponentClass->GetName();
+				const int32 NewStrLen = NewClassName.Len() - 2;
+				NewComponentName = MakeUniqueObjectName(ActorInstance, NewComponentClass, FName(*NewClassName.Left(NewStrLen)));
+			}
+
+			if (Asset != nullptr)
+			{
+				NewComponentName = *FComponentEditorUtils::GenerateValidVariableNameFromAsset(Asset, ActorInstance);
+			}
+			else
+			{
+				NewComponentName = *FComponentEditorUtils::GenerateValidVariableName(NewComponentClass, ActorInstance);
+			}
+			NewComponent = AddNewNode(ConstructObject<UActorComponent>(NewComponentClass, ActorInstance, NewComponentName, RF_Transactional, ComponentTemplate), Asset, true);
 		}
-		NewComponent = AddNewNode(ConstructObject<UActorComponent>(NewComponentClass, ActorInstance, NewComponentName, RF_Transactional, ComponentTemplate), Asset, true);
 	}
 
 	return NewComponent;
@@ -4291,7 +4287,7 @@ void SSCSEditor::PasteNodes()
 	// Get the object that's being edited
 	UBlueprint* Blueprint = nullptr;
 	AActor* ActorInstance = nullptr;
-	if(EditorMode == EEditorMode::BlueprintSCS)
+	if (EditorMode == EComponentEditorMode::BlueprintSCS)
 	{
 		Blueprint = GetBlueprint();
 		check(Blueprint != nullptr && Blueprint->SimpleConstructionScript != nullptr);
@@ -4299,12 +4295,13 @@ void SSCSEditor::PasteNodes()
 		Blueprint->Modify();
 		SaveSCSCurrentState(Blueprint->SimpleConstructionScript);
 	}
-	else    // EEditorMode::ActorInstance
+	else    // EComponentEditorMode::ActorInstance
 	{
-		ActorInstance = ActorContext.Get();
-		check(ActorInstance != nullptr);
-
-		ActorInstance->Modify();
+		ActorInstance = GetActorContext();
+		if (ActorInstance != nullptr)
+		{
+			ActorInstance->Modify();
+		}
 	}
 
 	// Create a new tree node for each new (pasted) component
@@ -4326,6 +4323,7 @@ void SSCSEditor::PasteNodes()
 		else
 		{
 			// Relocate the instance from the transient package to the Actor and assign it a unique object name
+			check(ActorInstance != nullptr); //@TODO: This doesn't look justified, but the code below was already assuming it
 			FString NewComponentName = FComponentEditorUtils::GenerateValidVariableName(NewActorComponent->GetClass(), ActorInstance);
 			NewActorComponent->Rename(*NewComponentName, ActorInstance, REN_DontCreateRedirectors | REN_DoNotDirty);
 
@@ -4397,7 +4395,7 @@ void SSCSEditor::OnDeleteNodes()
 {
 	const FScopedTransaction Transaction( LOCTEXT("RemoveComponent", "Remove Component") );
 
-	if(EditorMode == SSCSEditor::EEditorMode::BlueprintSCS)
+	if (EditorMode == EComponentEditorMode::BlueprintSCS)
 	{
 		// Remove node from SCS
 		UBlueprint* Blueprint = GetBlueprint();
@@ -4444,12 +4442,12 @@ void SSCSEditor::OnDeleteNodes()
 			ThumbnailTools::CacheEmptyThumbnail( BPFullName, BPPackage );
 		}
 	}
-	else    // SSCSEditor::EEditorMode::ActorInstance
+	else    // EComponentEditorMode::ActorInstance
 	{
-		AActor* ActorInstance = ActorContext.Get();
-		check(ActorInstance != nullptr);
-
-		ActorInstance->Modify();
+		if (AActor* ActorInstance = GetActorContext())
+		{
+			ActorInstance->Modify();
+		}
 
 		FSCSEditorTreeNodePtrType NewSelection;
 		TArray<FSCSEditorTreeNodePtrType> SelectedNodes = SCSTreeWidget->GetSelectedItems();
@@ -4519,7 +4517,7 @@ void SSCSEditor::RemoveComponentNode(FSCSEditorTreeNodePtrType InNodePtr)
 		SCSTreeWidget->ClearSelection();
 	}
 
-	if(EditorMode == SSCSEditor::EEditorMode::BlueprintSCS)
+	if (EditorMode == EComponentEditorMode::BlueprintSCS)
 	{
 		USCS_Node* SCS_Node = InNodePtr->GetSCSNode();
 		if(SCS_Node != NULL)
@@ -4541,13 +4539,12 @@ void SSCSEditor::RemoveComponentNode(FSCSEditorTreeNodePtrType InNodePtr)
 			SCS_Node->SetOnNameChanged(FSCSNodeNameChanged());
 		}
 	}
-	else    // SSCSEditor::EEditorMode::ActorInstance
+	else    // EComponentEditorMode::ActorInstance
 	{
-		AActor* ActorInstance = ActorContext.Get();
-		check(ActorInstance != nullptr);
+		AActor* ActorInstance = GetActorContext();
 
 		UActorComponent* ComponentInstance = InNodePtr->GetComponentTemplate();
-		if (ComponentInstance != nullptr)
+		if ((ActorInstance != nullptr) && (ComponentInstance != nullptr))
 		{
 			ComponentInstance->Modify();
 
@@ -4800,7 +4797,7 @@ FSCSEditorTreeNodePtrType SSCSEditor::AddTreeNode(USceneComponent* InSceneCompon
 
 	// If the given component has a parent, and if we're not in "instance" mode OR the owner of the parent matches the Actor instance we're editing
 	if(InSceneComponent->AttachParent != NULL
-		&& (EditorMode != EEditorMode::ActorInstance || InSceneComponent->AttachParent->GetOwner() == ActorContext.Get()))
+		&& (EditorMode != EComponentEditorMode::ActorInstance || InSceneComponent->AttachParent->GetOwner() == GetActorContext()))
 	{
 		// Attempt to find the parent node in the current tree
 		FSCSEditorTreeNodePtrType ParentNodePtr = FindTreeNode(InSceneComponent->AttachParent);
@@ -5029,10 +5026,9 @@ void SSCSEditor::GetCollapsedNodes(const FSCSEditorTreeNodePtrType& InNodePtr, T
 EVisibility SSCSEditor::GetPromoteToBlueprintButtonVisibility() const
 {
 	EVisibility Visibility = EVisibility::Collapsed;
-	if( EditorMode == SSCSEditor::EEditorMode::ActorInstance )
+	if (EditorMode == EComponentEditorMode::ActorInstance)
 	{
-		AActor* Actor = ActorContext.Get();
-		if( Actor && Actor->GetClass()->ClassGeneratedBy == nullptr )
+		if (GetBlueprint() == nullptr)
 		{
 			Visibility = EVisibility::Visible;
 		}
@@ -5044,10 +5040,9 @@ EVisibility SSCSEditor::GetPromoteToBlueprintButtonVisibility() const
 EVisibility SSCSEditor::GetEditBlueprintButtonVisibility() const
 {
 	EVisibility Visibility = EVisibility::Collapsed;
-	if(EditorMode == SSCSEditor::EEditorMode::ActorInstance)
+	if (EditorMode == EComponentEditorMode::ActorInstance)
 	{
-		AActor* Actor = ActorContext.Get();
-		if(Actor && Actor->GetClass()->ClassGeneratedBy != nullptr )
+		if (GetBlueprint() != nullptr)
 		{
 			Visibility = EVisibility::Visible;
 		}
@@ -5060,8 +5055,8 @@ FText SSCSEditor::OnGetApplyChangesToBlueprintTooltip() const
 {
 	int32 NumChangedProperties = 0;
 
-	AActor* Actor = ActorContext.Get();
-	UBlueprint* Blueprint = ActorInstance ? Cast<UBlueprint>(Actor->GetClass()->ClassGeneratedBy) : nullptr;
+	AActor* Actor = GetActorContext();
+	UBlueprint* Blueprint = (Actor != nullptr) ? Cast<UBlueprint>(Actor->GetClass()->ClassGeneratedBy) : nullptr;
 
 	if(Actor != NULL && Blueprint != NULL && Actor->GetClass()->ClassGeneratedBy == Blueprint)
 	{
@@ -5093,8 +5088,8 @@ FText SSCSEditor::OnGetResetToBlueprintDefaultsTooltip() const
 {
 	int32 NumChangedProperties = 0;
 
-	AActor* Actor = ActorContext.Get();
-	UBlueprint* Blueprint = ActorInstance ? Cast<UBlueprint>(Actor->GetClass()->ClassGeneratedBy) : nullptr;
+	AActor* Actor = GetActorContext();
+	UBlueprint* Blueprint = (Actor != nullptr) ? Cast<UBlueprint>(Actor->GetClass()->ClassGeneratedBy) : nullptr;
 	if(Actor != NULL && Blueprint != NULL && Actor->GetClass()->ClassGeneratedBy == Blueprint)
 	{
 		AActor* BlueprintCDO = Actor->GetClass()->GetDefaultObject<AActor>();
@@ -5122,8 +5117,8 @@ FText SSCSEditor::OnGetResetToBlueprintDefaultsTooltip() const
 
 void SSCSEditor::OnOpenBlueprintEditor() const
 {
-	AActor* ActorInstance = ActorContext.Get();
-	UBlueprint* Blueprint = ActorInstance ? Cast<UBlueprint>( ActorInstance->GetClass()->ClassGeneratedBy ) : nullptr;
+	AActor* ActorInstance = GetActorContext();
+	UBlueprint* Blueprint = (ActorInstance != nullptr) ? Cast<UBlueprint>( ActorInstance->GetClass()->ClassGeneratedBy ) : nullptr;
 
 	FAssetEditorManager::Get().OpenEditorForAsset(Blueprint);
 }
@@ -5132,8 +5127,8 @@ void SSCSEditor::OnApplyChangesToBlueprint() const
 {
 	int32 NumChangedProperties = 0;
 
-	AActor* Actor = ActorContext.Get();
-	UBlueprint* Blueprint = Actor ? Cast<UBlueprint>(Actor->GetClass()->ClassGeneratedBy) : nullptr;
+	AActor* Actor = GetActorContext();
+	UBlueprint* Blueprint = (Actor != nullptr) ? Cast<UBlueprint>(Actor->GetClass()->ClassGeneratedBy) : nullptr;
 
 	if (Actor != NULL && Blueprint != NULL && Actor->GetClass()->ClassGeneratedBy == Blueprint)
 	{
@@ -5201,10 +5196,10 @@ void SSCSEditor::OnResetToBlueprintDefaults() const
 {
 	int32 NumChangedProperties = 0;
 
-	AActor* Actor = ActorContext.Get();
-	UBlueprint* Blueprint = Actor ? Cast<UBlueprint>(Actor->GetClass()->ClassGeneratedBy) : nullptr;
+	AActor* Actor = GetActorContext();
+	UBlueprint* Blueprint = (Actor != nullptr) ? Cast<UBlueprint>(Actor->GetClass()->ClassGeneratedBy) : nullptr;
 
-	if (Actor != NULL && Blueprint != NULL && Actor->GetClass()->ClassGeneratedBy == Blueprint)
+	if ((Actor != NULL) && (Blueprint != NULL) && (Actor->GetClass()->ClassGeneratedBy == Blueprint))
 	{
 		const FScopedTransaction Transaction(LOCTEXT("ResetToBlueprintDefaults_Transaction", "Reset to Class Defaults"));
 
@@ -5259,7 +5254,7 @@ void SSCSEditor::OnResetToBlueprintDefaults() const
 void SSCSEditor::PromoteToBlueprint() const
 {
 	bool bHarvest = false;
-	FCreateBlueprintFromActorDialog::OpenDialog(bHarvest, ActorContext.Get());
+	FCreateBlueprintFromActorDialog::OpenDialog(bHarvest, GetActorContext());
 }
 
 FReply SSCSEditor::OnPromoteToBlueprintClicked()
@@ -5273,5 +5268,11 @@ const TArray<FSCSEditorTreeNodePtrType>& SSCSEditor::GetRootComponentNodes()
 	return RootComponentNodes;
 }
 
+/** Returns the Actor context for which we are viewing/editing the SCS.  Can return null.  Should not be cached as it may change from frame to frame. */
+AActor* SSCSEditor::GetActorContext() const
+{
+	return ActorContext.Get(nullptr);
+}
+
 #undef LOCTEXT_NAMESPACE
 
diff --git a/Engine/Source/Editor/Kismet/Public/SSCSEditor.h b/Engine/Source/Editor/Kismet/Public/SSCSEditor.h
--- a/Engine/Source/Editor/Kismet/Public/SSCSEditor.h
+++ b/Engine/Source/Editor/Kismet/Public/SSCSEditor.h
@@ -450,18 +450,21 @@ private:
 
 typedef SSCSEditorDragDropTree SSCSTreeType;
 
-class KISMET_API SSCSEditor : public SCompoundWidget
+/* Component editor mode */
+namespace EComponentEditorMode
 {
-public:
-	/* Editor mode */
-	enum EEditorMode
+	enum Type
 	{
 		/* View/edit the SCS in a BGPC */
 		BlueprintSCS,
 		/* View/edit the Actor instance */
 		ActorInstance
 	};
+};
 
+class KISMET_API SSCSEditor : public SCompoundWidget
+{
+public:
 	DECLARE_DELEGATE_RetVal_OneParam(class USCS_Node*, FOnAddNewComponent, class UClass*);
 	DECLARE_DELEGATE_RetVal_OneParam(class USCS_Node*, FOnAddExistingComponent, class UActorComponent*);
 	DECLARE_DELEGATE_OneParam(FOnSelectionUpdated, const TArray<FSCSEditorTreeNodePtrType>&);
@@ -469,7 +472,7 @@ public:
 	DECLARE_DELEGATE_OneParam(FOnHighlightPropertyInDetailsView, const class FPropertyPath&);
 
 	SLATE_BEGIN_ARGS( SSCSEditor )
-		:_EditorMode(BlueprintSCS)
+		:_EditorMode(EComponentEditorMode::BlueprintSCS)
 		,_ActorContext(nullptr)
 		,_PreviewActor(nullptr)
 		,_AllowEditing(true)
@@ -478,7 +481,7 @@ public:
 		,_OnHighlightPropertyInDetailsView()
 		{}
 
-		SLATE_ARGUMENT(EEditorMode, EditorMode)
+		SLATE_ARGUMENT(EComponentEditorMode::Type, EditorMode)
 		SLATE_ATTRIBUTE(class AActor*, ActorContext)
 		SLATE_ATTRIBUTE(class AActor*, PreviewActor)
 		SLATE_ATTRIBUTE(bool, AllowEditing)
@@ -645,7 +648,7 @@ public:
 	const TArray<FSCSEditorTreeNodePtrType>& GetRootComponentNodes();
 
 	/** @return The current editor mode (editing live actors or editing blueprints) */
-	EEditorMode GetEditorMode() const { return EditorMode; }
+	EComponentEditorMode::Type GetEditorMode() const { return EditorMode; }
 protected:
 	FString GetSelectedClassText() const;
 
@@ -800,9 +803,12 @@ public:
 
 	/** Delegate to invoke when the given property should be highlighted in the details view (e.g. diff). */
 	FOnHighlightPropertyInDetailsView OnHighlightPropertyInDetailsView;
+
+	/** Returns the Actor context for which we are viewing/editing the SCS.  Can return null.  Should not be cached as it may change from frame to frame. */
+	class AActor* GetActorContext() const;
 private:
 	/** Indicates which editor mode we're in. */
-	EEditorMode EditorMode;
+	EComponentEditorMode::Type EditorMode;
 
 	/** Root set of tree */
 	TArray<FSCSEditorTreeNodePtrType> RootNodes;
@@ -819,6 +825,6 @@ private:
 	/** Gate to prevent changing the selection while selection change is being broadcast. */
 	bool bUpdatingSelection;
 
-	/** Flag to track if we have the Actor selected curently */
+	/** Flag to track if we have the Actor selected currently */
 	bool bIsActorSelected;
 };
diff --git a/Engine/Source/Editor/LevelEditor/Private/SLevelEditor.cpp b/Engine/Source/Editor/LevelEditor/Private/SLevelEditor.cpp
--- a/Engine/Source/Editor/LevelEditor/Private/SLevelEditor.cpp
+++ b/Engine/Source/Editor/LevelEditor/Private/SLevelEditor.cpp
@@ -558,7 +558,7 @@ public:
 		ComponentsBox->SetContent
 		(
 			SAssignNew(SCSEditor, SSCSEditor)
-			.EditorMode(SSCSEditor::EEditorMode::ActorInstance)
+			.EditorMode(EComponentEditorMode::ActorInstance)
 			.ActorContext(this, &SActorDetails::GetSelectedActor)												// Get the instance of the actor in the world
 			.OnSelectionUpdated(this, &SActorDetails::OnSCSEditorTreeViewSelectionChanged)						// A selection has been made in the tree view, so inform the level editor
 			//.OnHighlightPropertyInDetailsView(this, &SLevelEditor::OnSCSEditorHighlightPropertyInDetailsView)	// Also unsure and don't think it's needed
