[INTEGRATE] Change 2420338 by [REDACTED_PERSON][REDACTED_HANDLE] on 2015/01/27 10:15:10

	Update the Actor's root component when dragging instance components

diff --git a/Engine/Source/Editor/Kismet/Private/SSCSEditor.cpp b/Engine/Source/Editor/Kismet/Private/SSCSEditor.cpp
--- a/Engine/Source/Editor/Kismet/Private/SSCSEditor.cpp
+++ b/Engine/Source/Editor/Kismet/Private/SSCSEditor.cpp
@@ -682,6 +682,11 @@ void FSCSEditorTreeNode::AddChild(FSCSEditorTreeNodePtrType InChildNodePtr)
 		// Handle attachment at the instance level
 		if (ChildInstance->AttachParent != ParentInstance)
 		{
+			AActor* Owner = ParentInstance->GetOwner();
+			if (Owner->GetRootComponent() == ChildInstance)
+			{
+				Owner->SetRootComponent(ParentInstance);
+			}
 			ChildInstance->AttachTo(ParentInstance, NAME_None, EAttachLocation::KeepWorldPosition);
 		}
 	}
