[INTEGRATE] Change 2420448 by [REDACTED_PERSON][REDACTED_HANDLE] on 2015/01/27 11:11:02

	BP - Now ensuring that the component editing mode isavilable before checking if the selecting component needs the component editing customization.

diff --git a/Engine/Source/Editor/Kismet/Private/SKismetInspector.cpp b/Engine/Source/Editor/Kismet/Private/SKismetInspector.cpp
--- a/Engine/Source/Editor/Kismet/Private/SKismetInspector.cpp
+++ b/Engine/Source/Editor/Kismet/Private/SKismetInspector.cpp
@@ -200,10 +200,10 @@ void SKismetInspector::Construct(const FArguments& InArgs)
 	PublicViewState = ECheckBoxState::Unchecked;
 	bComponenetDetailsCustomizationEnabled = false;
 
-	Kismet2Ptr = InArgs._Kismet2;
+	BlueprintEditorPtr = InArgs._Kismet2;
 	bShowPublicView = InArgs._ShowPublicViewControl;
 	bShowTitleArea = InArgs._ShowTitleArea;
-	TSharedPtr<FBlueprintEditor> Kismet2 = Kismet2Ptr.Pin();
+	TSharedPtr<FBlueprintEditor> Kismet2 = BlueprintEditorPtr.Pin();
 
 	// Create a property view
 	FPropertyEditorModule& EditModule = FModuleManager::Get().GetModuleChecked<FPropertyEditorModule>("PropertyEditor");
@@ -245,20 +245,20 @@ void SKismetInspector::Construct(const FArguments& InArgs)
 
 	if (Kismet2.IsValid() && Kismet2->IsEditingSingleBlueprint())
 	{
-		FOnGetDetailCustomizationInstance LayoutOptionDetails = FOnGetDetailCustomizationInstance::CreateStatic(&FBlueprintGlobalOptionsDetails::MakeInstance, Kismet2Ptr);
+		FOnGetDetailCustomizationInstance LayoutOptionDetails = FOnGetDetailCustomizationInstance::CreateStatic(&FBlueprintGlobalOptionsDetails::MakeInstance, BlueprintEditorPtr);
 		PropertyView->RegisterInstancedCustomPropertyLayout(UBlueprint::StaticClass(), LayoutOptionDetails);
 
 		FOnGetDetailCustomizationInstance LayoutFormatTextDetails = FOnGetDetailCustomizationInstance::CreateStatic(&FFormatTextDetails::MakeInstance);
 		PropertyView->RegisterInstancedCustomPropertyLayout(UK2Node_FormatText::StaticClass(), LayoutFormatTextDetails);
 
-		FOnGetDetailCustomizationInstance LayoutDocumentationDetails = FOnGetDetailCustomizationInstance::CreateStatic(&FBlueprintDocumentationDetails::MakeInstance, Kismet2Ptr);
+		FOnGetDetailCustomizationInstance LayoutDocumentationDetails = FOnGetDetailCustomizationInstance::CreateStatic(&FBlueprintDocumentationDetails::MakeInstance, BlueprintEditorPtr);
 		PropertyView->RegisterInstancedCustomPropertyLayout(UEdGraphNode_Documentation::StaticClass(), LayoutDocumentationDetails);
 
-		FOnGetDetailCustomizationInstance GraphNodeDetails = FOnGetDetailCustomizationInstance::CreateStatic(&FBlueprintGraphNodeDetails::MakeInstance, Kismet2Ptr);
+		FOnGetDetailCustomizationInstance GraphNodeDetails = FOnGetDetailCustomizationInstance::CreateStatic(&FBlueprintGraphNodeDetails::MakeInstance, BlueprintEditorPtr);
 		PropertyView->RegisterInstancedCustomPropertyLayout(UEdGraphNode::StaticClass(), GraphNodeDetails);
 
 		PropertyView->RegisterInstancedCustomPropertyLayout(UChildActorComponent::StaticClass(),
-			FOnGetDetailCustomizationInstance::CreateStatic(&FChildActorComponentDetails::MakeInstance, Kismet2Ptr));
+			FOnGetDetailCustomizationInstance::CreateStatic(&FChildActorComponentDetails::MakeInstance, BlueprintEditorPtr));
 	}
 
 	// Create the border that all of the content will get stuffed into
@@ -304,7 +304,7 @@ void SKismetInspector::EnableComponentDetailsCustomization(bool bEnable)
 		FOnGetDetailCustomizationInstance ActorOverrideDetails = FOnGetDetailCustomizationInstance::CreateStatic(&FActorDetailsOverrideCustomization::MakeInstance);
 		PropertyView->RegisterInstancedCustomPropertyLayout(AActor::StaticClass(), ActorOverrideDetails);
 
-		FOnGetDetailCustomizationInstance LayoutComponentDetails = FOnGetDetailCustomizationInstance::CreateStatic(&FBlueprintComponentDetails::MakeInstance, Kismet2Ptr);
+		FOnGetDetailCustomizationInstance LayoutComponentDetails = FOnGetDetailCustomizationInstance::CreateStatic(&FBlueprintComponentDetails::MakeInstance, BlueprintEditorPtr);
 		PropertyView->RegisterInstancedCustomPropertyLayout(UActorComponent::StaticClass(), LayoutComponentDetails);
 	}
 	else
@@ -380,12 +380,16 @@ void SKismetInspector::UpdateFromObjects(const TArray<UObject*>& PropertyObjects
 	if ( GetDefault<UEditorExperimentalSettings>()->bUnifiedBlueprintEditor )
 	{
 		bool bEnableComponentCustomization = false;
-		for ( UObject* PropertyObject : PropertyObjects )
+
+		if ( BlueprintEditorPtr.Pin()->CanAccessComponentsMode() )
 		{
-			if ( PropertyObject->IsA<UActorComponent>() )
+			for ( UObject* PropertyObject : PropertyObjects )
 			{
-				bEnableComponentCustomization = true;
-				break;
+				if ( PropertyObject->IsA<UActorComponent>() )
+				{
+					bEnableComponentCustomization = true;
+					break;
+				}
 			}
 		}
 
@@ -547,7 +551,7 @@ bool SKismetInspector::IsPropertyVisible( const FPropertyAndParent& PropertyAndP
 
 	if(const UClass* OwningClass = Cast<UClass>(Property.GetOuter()))
 	{
-		const UBlueprint* BP = Kismet2Ptr.IsValid() ? Kismet2Ptr.Pin()->GetBlueprintObj() : NULL;
+		const UBlueprint* BP = BlueprintEditorPtr.IsValid() ? BlueprintEditorPtr.Pin()->GetBlueprintObj() : NULL;
 		const bool VariableAddedInCurentBlueprint = (OwningClass->ClassGeneratedBy == BP);
 
 		// If we did not add this var, hide it!
@@ -631,7 +635,7 @@ void SKismetInspector::SetPublicViewCheckboxState( ECheckBoxState InIsChecked )
 		ShowDetailsForSingleObject(Objs[0], FShowDetailsOptions(PropertyViewTitle));
 	}
 	
-	Kismet2Ptr.Pin()->StartEditingDefaults();
+	BlueprintEditorPtr.Pin()->StartEditingDefaults();
 }
 
 //////////////////////////////////////////////////////////////////////////
diff --git a/Engine/Source/Editor/Kismet/Public/SKismetInspector.h b/Engine/Source/Editor/Kismet/Public/SKismetInspector.h
--- a/Engine/Source/Editor/Kismet/Public/SKismetInspector.h
+++ b/Engine/Source/Editor/Kismet/Public/SKismetInspector.h
@@ -76,7 +76,7 @@ protected:
 	void AddPropertiesRecursive(UProperty* Property);
 
 	/** Pointer back to the kismet 2 tool that owns us */
-	TWeakPtr<FBlueprintEditor> Kismet2Ptr;
+	TWeakPtr<FBlueprintEditor> BlueprintEditorPtr;
 
 	/** The tab that owns this details view. */
 	TWeakPtr<SDockTab> OwnerTab;
