[INTEGRATE] Change 2420656 by [REDACTED_PERSON][REDACTED_HANDLE] on 2015/01/27 13:02:22

	Fix crash when copying from BP created component and pasting as instance component

diff --git a/Engine/Source/Editor/Kismet/Private/SSCSEditor.cpp b/Engine/Source/Editor/Kismet/Private/SSCSEditor.cpp
--- a/Engine/Source/Editor/Kismet/Private/SSCSEditor.cpp
+++ b/Engine/Source/Editor/Kismet/Private/SSCSEditor.cpp
@@ -950,7 +950,7 @@ struct FSCSEditorComponentObjectTextFactory : public FCustomizableTextObjectFact
 	}
 
 	// Constructs a new object factory from the given text buffer
-	static TSharedRef<FSCSEditorComponentObjectTextFactory> Get(FString InTextBuffer)
+	static TSharedRef<FSCSEditorComponentObjectTextFactory> Get(const EComponentEditorMode::Type EditorMode, const FString& InTextBuffer)
 	{
 		// Construct a new instance
 		TSharedPtr<FSCSEditorComponentObjectTextFactory> FactoryPtr = MakeShareable(new FSCSEditorComponentObjectTextFactory());
@@ -959,8 +959,14 @@ struct FSCSEditorComponentObjectTextFactory : public FCustomizableTextObjectFact
 		// Create new objects if we're allowed to
 		if(FactoryPtr->CanCreateObjectsFromText(InTextBuffer))
 		{
+			EObjectFlags ObjectFlags = RF_Transactional;
+			if (EditorMode == EComponentEditorMode::BlueprintSCS)
+			{
+				ObjectFlags |= RF_ArchetypeObject;
+			}
+
 			// Use the transient package initially for creating the objects, since the variable name is used when copying
-			FactoryPtr->ProcessBuffer(GetTransientPackage(), RF_ArchetypeObject|RF_Transactional, InTextBuffer);
+			FactoryPtr->ProcessBuffer(GetTransientPackage(), ObjectFlags, InTextBuffer);
 		}
 
 		return FactoryPtr.ToSharedRef();
@@ -4270,7 +4276,7 @@ bool SSCSEditor::CanPasteNodes() const
 	FPlatformMisc::ClipboardPaste(ClipboardContent);
 
 	// Obtain the component object text factory for the clipboard content and return whether or not we can use it
-	TSharedRef<FSCSEditorComponentObjectTextFactory> Factory = FSCSEditorComponentObjectTextFactory::Get(ClipboardContent);
+	TSharedRef<FSCSEditorComponentObjectTextFactory> Factory = FSCSEditorComponentObjectTextFactory::Get(EditorMode, ClipboardContent);
 	return Factory->NewObjectMap.Num() > 0
 		&& (SceneRootNodePtr->IsDefaultSceneRoot() || Factory->CanAttachComponentsTo(Cast<USceneComponent>(SceneRootNodePtr->GetComponentTemplate())));
 }
@@ -4284,7 +4290,7 @@ void SSCSEditor::PasteNodes()
 	FPlatformMisc::ClipboardPaste(TextToImport);
 
 	// Get a new component object factory for the clipboard content
-	TSharedRef<FSCSEditorComponentObjectTextFactory> Factory = FSCSEditorComponentObjectTextFactory::Get(TextToImport);
+	TSharedRef<FSCSEditorComponentObjectTextFactory> Factory = FSCSEditorComponentObjectTextFactory::Get(EditorMode, TextToImport);
 
 	// Clear the current selection
 	SCSTreeWidget->ClearSelection();
