[INTEGRATE] Change 2420890 by [REDACTED_PERSON][REDACTED_HANDLE] on 2015/01/27 15:22:44

	BP - Adding a fix to prevent a crash when the BP editor is not found when selection changes.

diff --git a/Engine/Source/Editor/Kismet/Private/SKismetInspector.cpp b/Engine/Source/Editor/Kismet/Private/SKismetInspector.cpp
--- a/Engine/Source/Editor/Kismet/Private/SKismetInspector.cpp
+++ b/Engine/Source/Editor/Kismet/Private/SKismetInspector.cpp
@@ -381,14 +381,18 @@ void SKismetInspector::UpdateFromObjects(const TArray<UObject*>& PropertyObjects
 	{
 		bool bEnableComponentCustomization = false;
 
-		if ( BlueprintEditorPtr.Pin()->CanAccessComponentsMode() )
+		TSharedPtr<FBlueprintEditor> BlueprintEditor = BlueprintEditorPtr.Pin();
+		if ( BlueprintEditor.IsValid() )
 		{
-			for ( UObject* PropertyObject : PropertyObjects )
+			if ( BlueprintEditor->CanAccessComponentsMode() )
 			{
-				if ( PropertyObject->IsA<UActorComponent>() )
+				for ( UObject* PropertyObject : PropertyObjects )
 				{
-					bEnableComponentCustomization = true;
-					break;
+					if ( PropertyObject->IsA<UActorComponent>() )
+					{
+						bEnableComponentCustomization = true;
+						break;
+					}
 				}
 			}
 		}
