[INTEGRATE] Change 2420897 by [REDACTED_PERSON][REDACTED_HANDLE] on 2015/01/27 15:25:52

	Fixed SCS actor label (third time is a charm, right?)

diff --git a/Engine/Source/Editor/Kismet/Private/SSCSEditor.cpp b/Engine/Source/Editor/Kismet/Private/SSCSEditor.cpp
--- a/Engine/Source/Editor/Kismet/Private/SSCSEditor.cpp
+++ b/Engine/Source/Editor/Kismet/Private/SSCSEditor.cpp
@@ -2562,8 +2562,7 @@ FText SSCS_RowWidget_ActorRoot::GetActorDisplayText() const
 				}
 				else
 				{
-					FString Name;
-					DefaultActor->GetName(Name);
+					FString Name = DefaultActor->GetActorLabel();
 					return FText::Format(LOCTEXT("DefaultActor_Name", "{0} (Instance)"), FText::FromString(Name));
 				}
 			}
