[INTEGRATE] Change 2420958 by [REDACTED_PERSON][REDACTED_HANDLE] on 2015/01/27 15:58:52

	Open editor and sync content browser to new blueprint when creating component blueprint class

diff --git a/Engine/Source/Editor/Kismet/Private/SSCSEditor.cpp b/Engine/Source/Editor/Kismet/Private/SSCSEditor.cpp
--- a/Engine/Source/Editor/Kismet/Private/SSCSEditor.cpp
+++ b/Engine/Source/Editor/Kismet/Private/SSCSEditor.cpp
@@ -3917,6 +3917,13 @@ UClass* SSCSEditor::CreateNewBPComponent(TSubclassOf<UActorComponent> ComponentC
 				Package->MarkPackageDirty();
 
 				NewClass = NewBP->GeneratedClass;
+
+				TArray<UObject*> Objects;
+				Objects.Add(NewBP);
+				GEditor->SyncBrowserToObjects( Objects );
+
+				// Open the editor for the new blueprint
+				FAssetEditorManager::Get().OpenEditorForAsset(NewBP);
 			}
 		}
 	}
