[INTEGRATE] Change 2421194 by [REDACTED_PERSON][REDACTED_HANDLE] on 2015/01/27 17:44:17

	Improvements to naming of newly-created C++ and BP components through SCS tree
	- Strip off "_C" suffix on newly created Blueprint classes
	- Don't strip "Component" from Blueprint class names unless it is the suffix
	- Changed prefix for newly-created C++ components from "My" to "New"
	- New Blueprint components default to the name of the class unless they are Blueprints of UActorComponent, in which case the name is just "NewComponent"
	- New C++ components default to the name of the class unless they inherit from UActorComponent, in which case the name is just "NewComponent"

diff --git a/Engine/Source/Editor/GameProjectGeneration/Private/GameProjectGenerationModule.cpp b/Engine/Source/Editor/GameProjectGeneration/Private/GameProjectGenerationModule.cpp
--- a/Engine/Source/Editor/GameProjectGeneration/Private/GameProjectGenerationModule.cpp
+++ b/Engine/Source/Editor/GameProjectGeneration/Private/GameProjectGenerationModule.cpp
@@ -61,9 +61,9 @@ void FGameProjectGenerationModule::OpenAddCodeToProjectDialog(const TSharedPtr<S
 	AddCodeToProjectDialogOpenedEvent.Broadcast();
 }
 
-void FGameProjectGenerationModule::OpenAddCodeToProjectDialog(const UClass* InClass, const FString& InInitialPath, const TSharedPtr<SWindow>& InParentWindow, bool bModal, FOnCodeAddedToProject OnCodeAddedToProject)
+void FGameProjectGenerationModule::OpenAddCodeToProjectDialog(const UClass* InClass, const FString& InInitialPath, const TSharedPtr<SWindow>& InParentWindow, bool bModal, FOnCodeAddedToProject OnCodeAddedToProject, const FString InDefaultClassPrefix, const FString InDefaultClassName)
 {
-	GameProjectUtils::OpenAddCodeToProjectDialog(InClass, InInitialPath, InParentWindow, bModal, OnCodeAddedToProject);
+	GameProjectUtils::OpenAddCodeToProjectDialog(InClass, InInitialPath, InParentWindow, bModal, OnCodeAddedToProject, InDefaultClassPrefix, InDefaultClassName);
 	AddCodeToProjectDialogOpenedEvent.Broadcast();
 }
 
diff --git a/Engine/Source/Editor/GameProjectGeneration/Private/GameProjectUtils.cpp b/Engine/Source/Editor/GameProjectGeneration/Private/GameProjectUtils.cpp
--- a/Engine/Source/Editor/GameProjectGeneration/Private/GameProjectUtils.cpp
+++ b/Engine/Source/Editor/GameProjectGeneration/Private/GameProjectUtils.cpp
@@ -721,7 +721,7 @@ bool GameProjectUtils::UpdateGameProject(const FString& ProjectFile, const FStri
 	return UpdateGameProjectFile(ProjectFile, EngineIdentifier, NULL, OutFailReason);
 }
 
-void GameProjectUtils::OpenAddCodeToProjectDialog(const UClass* InClass, const FString& InInitialPath, const TSharedPtr<SWindow>& InParentWindow, bool bModal, FOnCodeAddedToProject OnCodeAddedToProject )
+void GameProjectUtils::OpenAddCodeToProjectDialog(const UClass* InClass, const FString& InInitialPath, const TSharedPtr<SWindow>& InParentWindow, bool bModal, FOnCodeAddedToProject OnCodeAddedToProject, const FString InDefaultClassPrefix, const FString InDefaultClassName )
 {
 	// If we've been given a class then we only show the second page of the dialog, so we can make the window smaller as that page doesn't have as much content
 	const FVector2D WindowSize = (InClass) ? FVector2D(940, 380) : FVector2D(940, 540);
@@ -737,7 +737,9 @@ void GameProjectUtils::OpenAddCodeToProjectDialog(const UClass* InClass, const F
 		SNew(SNewClassDialog)
 		.Class(InClass)
 		.InitialPath(InInitialPath)
-		.OnCodeAddedToProject( OnCodeAddedToProject );
+		.OnCodeAddedToProject( OnCodeAddedToProject )
+		.DefaultClassPrefix( InDefaultClassPrefix )
+		.DefaultClassName( InDefaultClassName );
 
 	AddCodeWindow->SetContent( NewClassDialog );
 
diff --git a/Engine/Source/Editor/GameProjectGeneration/Private/GameProjectUtils.h b/Engine/Source/Editor/GameProjectGeneration/Private/GameProjectUtils.h
--- a/Engine/Source/Editor/GameProjectGeneration/Private/GameProjectUtils.h
+++ b/Engine/Source/Editor/GameProjectGeneration/Private/GameProjectUtils.h
@@ -180,8 +180,10 @@ public:
 	 * @param	InParentWindow		The parent window the dialog should use, or null to choose a suitable default parent window (the main frame, if available)
 	 * @param	bModal			True if the window should be modal and force the user to make a decision before continuing or false to let the user proceed with other tasks while the window is open
 	 * @param	OnCodeAddedToProject	Callback for when code is successfully added to the project
+	 * @param	InDefaultClassPrefix	Optional argument that specifies the prefix for the new class name.  The user will be able to type their own name if they don't like this name.  Defaults to "My" if not specified or empty
+	 * @param	InDefaultClassName	Optional argument that specifies the default name for the new class being added.  The user will be able to type their own name if they don't like this name.  If empty, defaults to the name of the inherited class.
 	 */
-	static void OpenAddCodeToProjectDialog(const UClass* InClass, const FString& InInitialPath, const TSharedPtr<SWindow>& InParentWindow, bool bModal = false, FOnCodeAddedToProject OnCodeAddedToProject = FOnCodeAddedToProject() );
+	static void OpenAddCodeToProjectDialog(const UClass* InClass, const FString& InInitialPath, const TSharedPtr<SWindow>& InParentWindow, bool bModal = false, FOnCodeAddedToProject OnCodeAddedToProject = FOnCodeAddedToProject(), const FString InDefaultClassPrefix = FString(), const FString InDefaultClassName = FString() );
 
 	/** Returns true if the specified class name is properly formed and does not conflict with another class */
 	static bool IsValidClassNameForCreation(const FString& NewClassName, const FModuleContextInfo& ModuleInfo, const TSet<FString>& DisallowedHeaderNames, FText& OutFailReason);
diff --git a/Engine/Source/Editor/GameProjectGeneration/Private/SNewClassDialog.cpp b/Engine/Source/Editor/GameProjectGeneration/Private/SNewClassDialog.cpp
--- a/Engine/Source/Editor/GameProjectGeneration/Private/SNewClassDialog.cpp
+++ b/Engine/Source/Editor/GameProjectGeneration/Private/SNewClassDialog.cpp
@@ -103,6 +103,9 @@ void SNewClassDialog::Construct( const FArguments& InArgs )
 		}
 	}
 
+	DefaultClassPrefix = InArgs._DefaultClassPrefix;
+	DefaultClassName = InArgs._DefaultClassName;
+
 	// If we didn't get given a valid path override (see above), try and automatically work out the best default module
 	// If we have a runtime module with the same name as our project, then use that
 	// Otherwise, set out default target module as the first runtime module in the list
@@ -904,7 +907,9 @@ void SNewClassDialog::OnNamePageEntered()
 {
 	// Set the default class name based on the selected parent class, eg MyActor
 	const FString ParentClassName = ParentClassInfo.GetClassNameCPP();
-	const FString PotentialNewClassName = FString::Printf(TEXT("My%s"), ParentClassName.IsEmpty() ? TEXT("Class") : *ParentClassName);
+	const FString PotentialNewClassName = FString::Printf(TEXT("%s%s"), 
+		DefaultClassPrefix.IsEmpty() ? TEXT("My") : *DefaultClassPrefix, 
+		DefaultClassName.IsEmpty() ? (ParentClassName.IsEmpty() ? TEXT("Class") : *ParentClassName) : *DefaultClassName);
 
 	// Only set the default if the user hasn't changed the class name from the previous default
 	if(LastAutoGeneratedClassName.IsEmpty() || NewClassName == LastAutoGeneratedClassName)
diff --git a/Engine/Source/Editor/GameProjectGeneration/Private/SNewClassDialog.h b/Engine/Source/Editor/GameProjectGeneration/Private/SNewClassDialog.h
--- a/Engine/Source/Editor/GameProjectGeneration/Private/SNewClassDialog.h
+++ b/Engine/Source/Editor/GameProjectGeneration/Private/SNewClassDialog.h
@@ -20,6 +20,13 @@ public:
 	/** The initial path to use as the destination for the new class. If this is not specified, we will work out a suitable default from the available project modules */
 	SLATE_ARGUMENT(FString, InitialPath)
 
+	/** The prefix to put on new classes by default, if the user doesn't type in a new name.  Defaults to 'My'. */
+	SLATE_ARGUMENT(FString, DefaultClassPrefix)
+
+	/** If non-empty, overrides the default name of the class, when the user doesn't type a new name.  Defaults to empty, which causes the
+	    name to be the inherited class name.  Note that DefaultClassPrefix is still prepended to this name, if non-empty. */
+	SLATE_ARGUMENT(FString, DefaultClassName)
+
 	/** Event called when code is successfully added to the project */
 	SLATE_EVENT( FOnCodeAddedToProject, OnCodeAddedToProject )
 	SLATE_END_ARGS()
@@ -171,6 +178,13 @@ private:
 	/** A pointer to a class viewer **/
 	TSharedPtr<class SClassViewer> ClassViewer;
 
+	/** The prefix to put on new classes by default, if the user doesn't type in a new name.  Defaults to 'My'. */
+	FString DefaultClassPrefix;
+
+	/** If non-empty, overrides the default name of the class, when the user doesn't type a new name.  Defaults to empty, which causes the
+	    name to be the inherited class name.  Note that DefaultClassPrefix is still prepended to this name, if non-empty. */
+	FString DefaultClassName;
+
 	/** The editable text box to enter the current name */
 	TSharedPtr<SEditableTextBox> ClassNameEditBox;
 
diff --git a/Engine/Source/Editor/GameProjectGeneration/Public/GameProjectGenerationModule.h b/Engine/Source/Editor/GameProjectGeneration/Public/GameProjectGenerationModule.h
--- a/Engine/Source/Editor/GameProjectGeneration/Public/GameProjectGenerationModule.h
+++ b/Engine/Source/Editor/GameProjectGeneration/Public/GameProjectGenerationModule.h
@@ -83,8 +83,10 @@ public:
 	 * @param	InParentWindow		The parent window the dialog should use, or null to choose a suitable default parent window (the main frame, if available)
 	 * @param	bModal			True if the window should be modal and force the user to make a decision before continuing or false to let the user proceed with other tasks while the window is open
 	 * @param	OnCodeAddedToProject	Callback for when code is successfully added to the project
+	 * @param	InDefaultClassPrefix	Optional argument that specifies the prefix for the new class name.  The user will be able to type their own name if they don't like this name.  Defaults to "My" if not specified or empty
+	 * @param	InDefaultClassName	Optional argument that specifies the default name for the new class being added.  The user will be able to type their own name if they don't like this name.  If empty, defaults to the name of the inherited class.
 	 */
-	virtual void OpenAddCodeToProjectDialog(const UClass* InClass, const FString& InInitialPath, const TSharedPtr<SWindow>& InParentWindow, bool bModal = false, FOnCodeAddedToProject OnCodeAddedToProject = FOnCodeAddedToProject() );
+	virtual void OpenAddCodeToProjectDialog(const UClass* InClass, const FString& InInitialPath, const TSharedPtr<SWindow>& InParentWindow, bool bModal = false, FOnCodeAddedToProject OnCodeAddedToProject = FOnCodeAddedToProject(), const FString InDefaultClassPrefix = FString(), const FString InDefaultClassName = FString() );
 
 	/** Delegate for when the AddCodeToProject dialog is opened */
 	DECLARE_EVENT(FGameProjectGenerationModule, FAddCodeToProjectDialogOpenedEvent);
diff --git a/Engine/Source/Editor/Kismet/Private/SSCSEditor.cpp b/Engine/Source/Editor/Kismet/Private/SSCSEditor.cpp
--- a/Engine/Source/Editor/Kismet/Private/SSCSEditor.cpp
+++ b/Engine/Source/Editor/Kismet/Private/SSCSEditor.cpp
@@ -3850,6 +3850,24 @@ TSharedPtr<FSCSEditorTreeNode> SSCSEditor::AddRootComponentTreeNode(UActorCompon
 	return NewTreeNode;
 }
 
+
+void SSCSEditor::MakeDefaultComponentPrefixAndName( TSubclassOf<UActorComponent> ComponentClass, FString& DefaultClassPrefix, FString& DefaultClassName ) const
+{
+	// If the class being created inherits directly from UActorComponent, we don't want the default new class name to
+	// be called "NewActorComponent", because this will shw up as "New Actor" in the SCS editor tree which doesn't make
+	// a whole lot of sense in the context of a list of components inside an actor.
+	DefaultClassName = ComponentClass->GetName();
+	if( ComponentClass == UActorComponent::StaticClass() )
+	{
+		// NOTE: This is intentionally not localized because it will become a C++ source file and class name
+		DefaultClassName = TEXT( "Component" );
+	}
+
+	// NOTE: This is intentionally not localized because it will become a C++ source file and class name
+	DefaultClassPrefix = TEXT( "New" );
+}
+
+
 UClass* SSCSEditor::CreateNewCPPComponent( TSubclassOf<UActorComponent> ComponentClass )
 {
 	TSharedPtr<SWindow> ParentWindow = FSlateApplication::Get().FindWidgetWindow(SharedThis(this));
@@ -3864,7 +3882,11 @@ UClass* SSCSEditor::CreateNewCPPComponent( TSubclassOf<UActorComponent> Componen
 	};
 
 	const bool bModal = true;
-	FGameProjectGenerationModule::Get().OpenAddCodeToProjectDialog(ComponentClass, FString(), ParentWindow, bModal, FOnCodeAddedToProject::CreateLambda(OnCodeAddedToProject));
+
+	FString DefaultClassPrefix, DefaultClassName;
+	MakeDefaultComponentPrefixAndName( ComponentClass, DefaultClassPrefix, DefaultClassName );
+
+	FGameProjectGenerationModule::Get().OpenAddCodeToProjectDialog(ComponentClass, FString(), ParentWindow, bModal, FOnCodeAddedToProject::CreateLambda(OnCodeAddedToProject), DefaultClassPrefix, DefaultClassName);
 
 	return LoadClass<UActorComponent>(nullptr, *AddedClassName, nullptr, LOAD_None, nullptr);
 }
@@ -3888,13 +3910,16 @@ UClass* SSCSEditor::CreateNewBPComponent(TSubclassOf<UActorComponent> ComponentC
 		AssetName = FPackageName::GetLongPackageAssetName(Path);
 	};
 
+	FString DefaultClassPrefix, DefaultClassName;
+	MakeDefaultComponentPrefixAndName( ComponentClass, DefaultClassPrefix, DefaultClassName );
+
 	TSharedPtr<SCreateAssetFromObject> CreateBlueprintFromActorDialog;
 	PickBlueprintPathWidget->SetContent
 	(
 		SAssignNew(CreateBlueprintFromActorDialog, SCreateAssetFromObject, PickBlueprintPathWidget)
 		.HeadingText(LOCTEXT("CreateBlueprintFromActor_Heading", "Blueprint Name"))
 		.CreateButtonText(LOCTEXT("CreateBlueprintFromActor_ButtonLabel", "Create Blueprint"))
-		.DefaultNameOverride(FText::FromString(TEXT("BlueprintComponent")))
+		.DefaultNameOverride(FText::FromString(DefaultClassPrefix + DefaultClassName))
 		.OnCreateAssetAction(FOnPathChosen::CreateLambda(OnPathPicked))
 	);
 
@@ -4012,13 +4037,6 @@ UActorComponent* SSCSEditor::AddNewComponent( UClass* NewComponentClass, UObject
 			ActorInstance->Modify();
 
 			FName NewComponentName = NAME_None;
-			if (NewComponentClass->ClassGeneratedBy && NewComponentClass->GetName().EndsWith(TEXT("_C")))
-			{
-				const FString NewClassName = NewComponentClass->GetName();
-				const int32 NewStrLen = NewClassName.Len() - 2;
-				NewComponentName = MakeUniqueObjectName(ActorInstance, NewComponentClass, FName(*NewClassName.Left(NewStrLen)));
-			}
-
 			if (Asset != nullptr)
 			{
 				NewComponentName = *FComponentEditorUtils::GenerateValidVariableNameFromAsset(Asset, ActorInstance);
diff --git a/Engine/Source/Editor/Kismet/Public/SSCSEditor.h b/Engine/Source/Editor/Kismet/Public/SSCSEditor.h
--- a/Engine/Source/Editor/Kismet/Public/SSCSEditor.h
+++ b/Engine/Source/Editor/Kismet/Public/SSCSEditor.h
@@ -753,6 +753,9 @@ protected:
 	/** Adds a root component tree node */
 	TSharedPtr<FSCSEditorTreeNode> AddRootComponentTreeNode(UActorComponent* ActorComp);
 
+	/** Given a class, makes up a default prefix and name for a newly created derived class */
+	void MakeDefaultComponentPrefixAndName( TSubclassOf<UActorComponent> ComponentClass, FString& DefaultClassPrefix, FString& DefaultClassName ) const;
+
 	/**
 	 * Creates a new C++ component from the specified class type
 	 * The user will be prompted to pick a new subclass name and code will be recompiled
diff --git a/Engine/Source/Editor/UnrealEd/Private/Kismet2/ComponentEditorUtils.cpp b/Engine/Source/Editor/UnrealEd/Private/Kismet2/ComponentEditorUtils.cpp
--- a/Engine/Source/Editor/UnrealEd/Private/Kismet2/ComponentEditorUtils.cpp
+++ b/Engine/Source/Editor/UnrealEd/Private/Kismet2/ComponentEditorUtils.cpp
@@ -140,10 +140,22 @@ FString FComponentEditorUtils::GenerateValidVariableName(TSubclassOf<UActorCompo
 {
 	check(ComponentOwner);
 
-	int32 Counter = 1;
-	FString ComponentTypeName = *ComponentClass->GetName().Replace(TEXT("Component"), TEXT(""));
+	// Strip off "_C" suffix if it has one
+	FString ComponentTypeName = *ComponentClass->GetName();
+	if (ComponentClass->ClassGeneratedBy && ComponentTypeName.EndsWith(TEXT("_C")))
+	{
+		ComponentTypeName = ComponentTypeName.Left( ComponentTypeName.Len() - 2 );
+	}
+
+	// Strip off 'Component' if the class ends with that.  It just looks better in the UI.
+	const FString SuffixToStrip( TEXT( "Component" ) );
+	if( ComponentTypeName.EndsWith( SuffixToStrip ) )
+	{
+		ComponentTypeName = ComponentTypeName.Left( ComponentTypeName.Len() - SuffixToStrip.Len() );
+	}
 	
 	// Try to create a name without any numerical suffix first
+	int32 Counter = 1;
 	FString ComponentInstanceName = ComponentTypeName;
 	while (!IsComponentNameAvailable(ComponentInstanceName, ComponentOwner))
 	{
