[INTEGRATE] Change 2421227 by [REDACTED_PERSON][REDACTED_HANDLE] on 2015/01/27 17:59:08

	Fixed SCS tree tool-tips showing "_C" suffix on generated classes

diff --git a/Engine/Source/Editor/Kismet/Private/SSCSEditor.cpp b/Engine/Source/Editor/Kismet/Private/SSCSEditor.cpp
--- a/Engine/Source/Editor/Kismet/Private/SSCSEditor.cpp
+++ b/Engine/Source/Editor/Kismet/Private/SSCSEditor.cpp
@@ -2258,8 +2258,15 @@ FText SSCS_RowWidget::GetTooltipText() const
 	}
 	else
 	{
+		UClass* Class = ( NodePtr->GetComponentTemplate() != nullptr ) ? NodePtr->GetComponentTemplate()->GetClass() : nullptr;
+		FString ClassName = Class != nullptr ? Class->GetName() : TEXT( "(null)" );
+		if( Class != nullptr && Class->ClassGeneratedBy != nullptr && ClassName.EndsWith( TEXT( "_C" ) ) )
+		{
+			ClassName.RemoveFromEnd( TEXT( "_C" ) );
+		}
+
 		FFormatNamedArguments Args;
-		Args.Add( TEXT("ClassName"), FText::FromString( (NodePtr->GetComponentTemplate() != NULL) ? NodePtr->GetComponentTemplate()->GetClass()->GetName() : TEXT("(null)") ) );
+		Args.Add( TEXT("ClassName"), FText::FromString( ClassName ));
 		Args.Add( TEXT("NodeName"), FText::FromString( NodePtr->GetDisplayString() ) );
 
 		if ( NodePtr->IsNative() )
