[INTEGRATE] Change 2421526 by [REDACTED_PERSON][REDACTED_HANDLE] on 2015/01/27 20:39:38

	Fix tiny tooltip when dragging off the scseditor.

diff --git a/Engine/Source/Editor/Kismet/Private/SSCSEditor.cpp b/Engine/Source/Editor/Kismet/Private/SSCSEditor.cpp
--- a/Engine/Source/Editor/Kismet/Private/SSCSEditor.cpp
+++ b/Engine/Source/Editor/Kismet/Private/SSCSEditor.cpp
@@ -1774,7 +1774,8 @@ void SSCS_RowWidget::OnDragLeave( const FDragDropEvent& DragDropEvent )
 		// Only clear the tooltip text if all dragged nodes support it
 		if(bCanReparentAllNodes)
 		{
-			DragRowOp->SetFeedbackMessage(SNullWidget::NullWidget);
+			TSharedPtr<SWidget> NoWidget;
+			DragRowOp->SetFeedbackMessage(NoWidget);
 			DragRowOp->PendingDropAction = FSCSRowDragDropOp::DropAction_None;
 		}
 	}
