[INTEGRATE] Change 2421717 by [REDACTED_PERSON][REDACTED_HANDLE] on 2015/01/27 23:39:54

	Blueprint Editor: Drag-drop improvements and crash fixes in the components tree widget
	- Prevent various crashes when working with a SCS row that has no associated node (such as separators)
	- Clarified error messages to split out different reasons why a drag isn't possible and mentioned inherited components in more places
	- Made a more generic error message indicating how to succeed (drag to a component) when over a row that isn't a component (the actor self 'node' and separator entries)
	- Removed the behavior that treated the self/actor node from acting like the root node for some purposes which produced confusing messages
	[UE-8463][UE-8117]

diff --git a/Engine/Source/Editor/Kismet/Private/SSCSEditor.cpp b/Engine/Source/Editor/Kismet/Private/SSCSEditor.cpp
--- a/Engine/Source/Editor/Kismet/Private/SSCSEditor.cpp
+++ b/Engine/Source/Editor/Kismet/Private/SSCSEditor.cpp
@@ -1533,12 +1533,24 @@ void SSCS_RowWidget::OnDragEnter( const FGeometry& MyGeometry, const FDragDropEv
 				// We set the tooltip text here because it won't change across entry/leave events
 				if (DragRowOp->SourceNodes.Num() == 1)
 				{
-					Message = LOCTEXT("DropActionToolTip_Error_CannotReparent", "The selected component cannot be moved.");
+					if (Cast<USceneComponent>(SelectedNodePtr->GetComponentTemplate()) == nullptr)
+					{
+						Message = LOCTEXT("DropActionToolTip_Error_CannotReparent_NotSceneComponent", "The selected component is not a scene component and cannot be attached to other components.");
+					}
+					else if (SelectedNodePtr->IsNative() || SelectedNodePtr->IsInherited())
+					{
+						Message = LOCTEXT("DropActionToolTip_Error_CannotReparent_Inherited", "The selected component is inherited and cannot be reordered here.");
+					}
+					else
+					{
+						Message = LOCTEXT("DropActionToolTip_Error_CannotReparent", "The selected component cannot be moved.");
+					}
 				}
 				else
 				{
-					Message = LOCTEXT("DropActionToolTip_Error_CannotReparentMultiple", "One or more of the selected components cannot be moved.");
+					Message = LOCTEXT("DropActionToolTip_Error_CannotReparentMultiple", "One or more of the selected components cannot be attached.");
 				}
+				break;
 			}
 		}
 
@@ -1548,9 +1560,10 @@ void SSCS_RowWidget::OnDragEnter( const FGeometry& MyGeometry, const FDragDropEv
 			check(SceneRootNodePtr.IsValid());
 
 			FSCSEditorTreeNodePtrType NodePtr = GetNode();
-			if (NodePtr->GetNodeType() == FSCSEditorTreeNode::SeparatorNode)
+			if ((NodePtr->GetNodeType() == FSCSEditorTreeNode::SeparatorNode) || (NodePtr->GetNodeType() == FSCSEditorTreeNode::RootActorNode))
 			{
-				Message = LOCTEXT("DropActionToolTip_Error_CannotParentToSeparator", "Cannot Parent to separator.");
+				// Don't show a feedback message if over a node that makes no sense, such as a separator or the instance node
+				Message = LOCTEXT("DropActionToolTip_FriendlyError_DragToAComponent", "Drag to another component in order to attach to that component or become the root component.");
 			}
 
 			// Validate each selected node being dragged against the node that belongs to this row. Exit the loop if we have a valid tooltip OR a valid pending drop action once all nodes in the selection have been validated.
@@ -1592,8 +1605,16 @@ void SSCS_RowWidget::OnDragEnter( const FGeometry& MyGeometry, const FDragDropEv
 				}
 				else if (HoveredTemplate == NULL || DraggedTemplate == NULL)
 				{
-					// Can't attach non-USceneComponent types
-					Message = LOCTEXT("DropActionToolTip_Error_NotAttachable", "Cannot attach to this component.");
+					if (HoveredTemplate == nullptr)
+					{
+						// Can't attach non-USceneComponent types
+						Message = LOCTEXT("DropActionToolTip_Error_NotAttachable_NotSceneComponent", "Cannot attach to this component as it is not a scene component.");
+					}
+					else
+					{
+						// Can't attach non-USceneComponent types
+						Message = LOCTEXT("DropActionToolTip_Error_NotAttachable", "Cannot attach to this component.");
+					}
 				}
 				else if (NodePtr == SceneRootNodePtr)
 				{
@@ -1607,7 +1628,7 @@ void SSCS_RowWidget::OnDragEnter( const FGeometry& MyGeometry, const FDragDropEv
 					if (!NodePtr->CanReparent() && (!NodePtr->IsDefaultSceneRoot() || NodePtr->IsInherited()))
 					{
 						// Cannot make the dropped node the new root if we cannot reparent the current root
-						Message = LOCTEXT("DropActionToolTip_Error_CannotReparentRootNode", "The root component in this Blueprint cannot be replaced.");
+						Message = LOCTEXT("DropActionToolTip_Error_CannotReparentRootNode", "The root component in this Blueprint is inherited and cannot be replaced.");
 					}
 					else if (DraggedTemplate->IsEditorOnly() && !HoveredTemplate->IsEditorOnly())
 					{
@@ -1744,7 +1765,14 @@ void SSCS_RowWidget::OnDragEnter( const FGeometry& MyGeometry, const FDragDropEv
 			? FEditorStyle::GetBrush(TEXT("Graph.ConnectorFeedback.OK"))
 			: FEditorStyle::GetBrush(TEXT("Graph.ConnectorFeedback.Error"));
 
-		DragRowOp->SetSimpleFeedbackMessage(StatusSymbol, FLinearColor::White, Message);
+		if (Message.IsEmpty())
+		{
+			DragRowOp->SetFeedbackMessage(nullptr);
+		}
+		else
+		{
+			DragRowOp->SetSimpleFeedbackMessage(StatusSymbol, FLinearColor::White, Message);
+		}
 	}
 	else if ( Operation->IsOfType<FExternalDragOperation>() || Operation->IsOfType<FAssetDragDropOp>() )
 	{
@@ -1789,7 +1817,7 @@ FReply SSCS_RowWidget::OnDrop( const FGeometry& MyGeometry, const FDragDropEvent
 		return FReply::Handled();
 	}
 	
-	if (Operation->IsOfType<FSCSRowDragDropOp>() && GetNode()->GetComponentTemplate()->IsA<USceneComponent>())
+	if (Operation->IsOfType<FSCSRowDragDropOp>() && (Cast<USceneComponent>(GetNode()->GetComponentTemplate()) != nullptr))
 	{
 		TSharedPtr<FSCSRowDragDropOp> DragRowOp = StaticCastSharedPtr<FSCSRowDragDropOp>( Operation );	
 		check(DragRowOp.IsValid());
@@ -2437,19 +2465,9 @@ void SSCS_RowWidget::OnNameTextCommit(const FText& InNewName, ETextCommit::Type
 //////////////////////////////////////////////////////////////////////////
 // SSCS_RowWidget_ActorRoot
 
-FSCSEditorTreeNodePtrType SSCS_RowWidget_ActorRoot::GetNode() const
-{
-	TSharedPtr<SSCSEditor> SCSEditorPtr = SCSEditor.Pin();
-	if (SCSEditorPtr.IsValid())
-	{
-		return SCSEditorPtr->SceneRootNodePtr;
-	}
-	return nullptr;
-}
-
 TSharedRef<SWidget> SSCS_RowWidget_ActorRoot::GenerateWidgetForColumn(const FName& ColumnName)
 {
-	// We've removed the other columns for now,  implement them for the root actor is necessary
+	// We've removed the other columns for now,  implement them for the root actor if necessary
 	ensure(ColumnName == SCS_ColumnName_ComponentClass);
 
 	return SNew(SHorizontalBox)
@@ -3166,9 +3184,9 @@ TSharedPtr< SWidget > SSCSEditor::CreateContextMenu() // @todo: make the context
 					for( auto NodeIter = SelectedNodes.CreateConstIterator(); NodeIter; ++NodeIter )
 					{
 						auto TreeNode = *NodeIter;
-						if( TreeNode->GetComponentTemplate() )
+						if( UActorComponent* ComponentTemplate = TreeNode->GetComponentTemplate() )
 						{
-							SelectionClasses.Add( TreeNode->GetComponentTemplate()->GetClass() );
+							SelectionClasses.Add(ComponentTemplate->GetClass());
 						}
 					}
 
@@ -4512,7 +4530,7 @@ void SSCSEditor::OnDeleteNodes()
 						NewSelection = ChildNodes[ChildIndex];
 					}
 				}
-				else if(!Node->GetComponentTemplate()->IsA<USceneComponent>())
+				else if(Cast<USceneComponent>(Node->GetComponentTemplate()) == nullptr)
 				{
 					int32 NodeIndex = RootComponentNodes.Find(Node);
 					if(NodeIndex != INDEX_NONE)
diff --git a/Engine/Source/Editor/Kismet/Public/SSCSEditor.h b/Engine/Source/Editor/Kismet/Public/SSCSEditor.h
--- a/Engine/Source/Editor/Kismet/Public/SSCSEditor.h
+++ b/Engine/Source/Editor/Kismet/Public/SSCSEditor.h
@@ -355,9 +355,6 @@ public:
 	virtual TSharedRef<SWidget> GenerateWidgetForColumn(const FName& ColumnName) override;
 	// End of SMultiColumnTableRow<T>
 
-	/* Get the node used by the row Widget */
-	virtual FSCSEditorTreeNodePtrType GetNode() const override;
-
 private:
 	/** Creates a tooltip for this row */
 	TSharedRef<SToolTip> CreateToolTipWidget() const;
