[INTEGRATE] Change 2422412 by [REDACTED_PERSON][REDACTED_HANDLE] on 2015/01/28 12:27:00

	Fixed hard-to-read documentation tool-tips in Blueprints (UE-8467)
	- Changed SCS tool-tips to use the new simplified dark-text-on-white-background styling
	- Made "Hold Ctlr+Alt" easier to read on a white background

diff --git a/Engine/Source/Editor/Documentation/Private/Documentation.cpp b/Engine/Source/Editor/Documentation/Private/Documentation.cpp
--- a/Engine/Source/Editor/Documentation/Private/Documentation.cpp
+++ b/Engine/Source/Editor/Documentation/Private/Documentation.cpp
@@ -251,8 +251,12 @@ TSharedRef< class SToolTip > FDocumentation::CreateToolTip(const TAttribute<FTex
 	}
 
 	return SNew(SToolTip)
-		.TextMargin(1)
 		.IsInteractive(DocToolTip, &SDocumentationToolTip::IsInteractive)
+
+		// Emulate text-only tool-tip styling that SToolTip uses when no custom content is supplied.  We want documentation tool-tips to 
+		// be styled just like text-only tool-tips
+		.BorderImage( FCoreStyle::Get().GetBrush("ToolTip.BrightBackground") )
+		.TextMargin(FMargin(11.0f))
 		[
 			DocToolTip
 		];
diff --git a/Engine/Source/Editor/EditorStyle/Private/SlateEditorStyle.cpp b/Engine/Source/Editor/EditorStyle/Private/SlateEditorStyle.cpp
--- a/Engine/Source/Editor/EditorStyle/Private/SlateEditorStyle.cpp
+++ b/Engine/Source/Editor/EditorStyle/Private/SlateEditorStyle.cpp
@@ -2550,7 +2550,7 @@ void FSlateEditorStyle::FStyle::SetupTutorialStyles()
 
 		const FTextBlockStyle DocumentationTooltipTextSubdued = FTextBlockStyle( NormalText )
 			.SetFont( TTF_CORE_FONT( "Fonts/Roboto-Regular", 8 ) )
-			.SetColorAndOpacity( FLinearColor( 0.3f, 0.3f, 0.3f ) );
+			.SetColorAndOpacity( FLinearColor( 0.2f, 0.2f, 0.2f ) );
 		Set("Documentation.SDocumentationTooltipSubdued", FTextBlockStyle(DocumentationTooltipTextSubdued));
 
 		const FTextBlockStyle DocumentationTooltipHyperlinkText = FTextBlockStyle( NormalText )
diff --git a/Engine/Source/Editor/Kismet/Private/SSCSEditor.cpp b/Engine/Source/Editor/Kismet/Private/SSCSEditor.cpp
--- a/Engine/Source/Editor/Kismet/Private/SSCSEditor.cpp
+++ b/Engine/Source/Editor/Kismet/Private/SSCSEditor.cpp
@@ -1148,9 +1148,6 @@ TSharedRef<SWidget> SSCS_RowWidget::GenerateWidgetForColumn( const FName& Column
 
 void AddToToolTipInfoBox(const TSharedRef<SVerticalBox>& InfoBox, const FText& Key, TSharedRef<SWidget> ValueIcon, const TAttribute<FText>& Value, bool bImportant)
 {
-	FWidgetStyle ImportantStyle;
-	ImportantStyle.SetForegroundColor(FLinearColor(1, 0.5, 0, 1));
-
 	InfoBox->AddSlot()
 		.AutoHeight()
 		.Padding(0, 1)
@@ -1162,7 +1159,7 @@ void AddToToolTipInfoBox(const TSharedRef<SVerticalBox>& InfoBox, const FText& K
 			.Padding(0, 0, 4, 0)
 			[
 				SNew(STextBlock).Text(FText::Format(LOCTEXT("AssetViewTooltipFormat", "{0}:"), Key))
-				.ColorAndOpacity(bImportant ? ImportantStyle.GetSubduedForegroundColor() : FSlateColor::UseSubduedForeground())
+				.ColorAndOpacity(bImportant ? FLinearColor( 0.05f, 0.05f, 0.05f ) : FLinearColor( 0.15f, 0.15f, 0.15f ) )
 			]
 
 			+ SHorizontalBox::Slot()
@@ -1175,7 +1172,7 @@ void AddToToolTipInfoBox(const TSharedRef<SVerticalBox>& InfoBox, const FText& K
 			.AutoWidth()
 			[
 				SNew(STextBlock).Text(Value)
-				.ColorAndOpacity(bImportant ? ImportantStyle.GetForegroundColor() : FSlateColor::UseForeground())
+				.ColorAndOpacity(bImportant ? FLinearColor( 0.0f, 0.0f, 0.0f ) : FLinearColor( 0.1f, 0.1f, 0.1f ) )
 			]
 		];
 }
@@ -1208,8 +1205,8 @@ TSharedRef<SToolTip> SSCS_RowWidget::CreateToolTipWidget() const
 	AddToToolTipInfoBox(InfoBox, LOCTEXT("TooltipMobility", "Mobility"), MobilityIcon, TAttribute<FText>::Create(TAttribute<FText>::FGetter::CreateSP(this, &SSCS_RowWidget::GetMobilityToolTipText, SCSNode)), false);
 	
 	TSharedRef<SBorder> TooltipContent = SNew(SBorder)
-		.Padding(4)
-		.BorderImage(FEditorStyle::GetBrush("SCSEditor.TileViewTooltip.NonContentBorder"))
+		.BorderImage(FEditorStyle::GetBrush("NoBorder"))
+		.Padding(0)
 		[
 			SNew(SVerticalBox)
 			
@@ -1217,34 +1214,22 @@ TSharedRef<SToolTip> SSCS_RowWidget::CreateToolTipWidget() const
 			.AutoHeight()
 			.Padding(0, 0, 0, 4)
 			[
-				SNew(SBorder)
-				.Padding(0)
-				.BorderImage(FEditorStyle::GetBrush("ToolPanel.GroupBorder"))
+				SNew(SVerticalBox)
+
+				+ SVerticalBox::Slot()
+				.AutoHeight()
 				[
-					SNew(SVerticalBox)
+					SNew(SHorizontalBox)
 
-					+ SVerticalBox::Slot()
-					.AutoHeight()
+					+ SHorizontalBox::Slot()
+					.AutoWidth()
+					.VAlign(VAlign_Center)
+					.Padding(4)
 					[
-						SNew(SHorizontalBox)
-
-						+ SHorizontalBox::Slot()
-						.AutoWidth()
-						.VAlign(VAlign_Center)
-						.Padding(4)
-						[
-							SNew(STextBlock)
-							.Text(this, &SSCS_RowWidget::GetTooltipText)
-							.Font(FEditorStyle::GetFontStyle("ContentBrowser.TileViewTooltip.NameFont"))
-						]
-						/*
-						+ SHorizontalBox::Slot()
-						.AutoWidth()
-						.VAlign(VAlign_Center)
-						[
-							SNew(STextBlock).Text(this, SSCS_RowWidget::GetTypeAsString())
-							.HighlightText(HighlightText)
-						]*/
+						SNew(STextBlock)
+						.Text(this, &SSCS_RowWidget::GetTooltipText)
+						.ColorAndOpacity( FLinearColor::Black )
+						.Font(FEditorStyle::GetFontStyle("ContentBrowser.TileViewTooltip.NameFont"))
 					]
 				]
 			]
@@ -1253,8 +1238,8 @@ TSharedRef<SToolTip> SSCS_RowWidget::CreateToolTipWidget() const
 			.AutoHeight()
 			[
 				SNew(SBorder)
+				.BorderImage(FEditorStyle::GetBrush("NoBorder"))
 				.Padding(4)
-				.BorderImage(FEditorStyle::GetBrush("ToolPanel.GroupBorder"))
 				[
 					InfoBox
 				]
@@ -2508,8 +2493,8 @@ TSharedRef<SToolTip> SSCS_RowWidget_ActorRoot::CreateToolTipWidget() const
 	AddToToolTipInfoBox(InfoBox, LOCTEXT("TooltipMobility", "Mobility"), SNullWidget::NullWidget, TAttribute<FText>::Create(TAttribute<FText>::FGetter::CreateSP(this, &SSCS_RowWidget_ActorRoot::GetActorMobilityText)), false);
 
 	TSharedRef<SBorder> TooltipContent = SNew(SBorder)
-		.Padding(4)
-		.BorderImage(FEditorStyle::GetBrush("SCSEditor.TileViewTooltip.NonContentBorder"))
+		.BorderImage(FEditorStyle::GetBrush("NoBorder"))
+		.Padding(0)
 		[
 			SNew(SVerticalBox)
 
@@ -2517,26 +2502,22 @@ TSharedRef<SToolTip> SSCS_RowWidget_ActorRoot::CreateToolTipWidget() const
 			.AutoHeight()
 			.Padding(0, 0, 0, 4)
 			[
-				SNew(SBorder)
-				.Padding(0)
-				.BorderImage(FEditorStyle::GetBrush("ToolPanel.GroupBorder"))
+				SNew(SVerticalBox)
+
+				+ SVerticalBox::Slot()
+				.AutoHeight()
 				[
-					SNew(SVerticalBox)
+					SNew(SHorizontalBox)
 
-					+ SVerticalBox::Slot()
-					.AutoHeight()
+					+ SHorizontalBox::Slot()
+					.AutoWidth()
+					.VAlign(VAlign_Center)
+					.Padding(4)
 					[
-						SNew(SHorizontalBox)
-
-						+ SHorizontalBox::Slot()
-						.AutoWidth()
-						.VAlign(VAlign_Center)
-						.Padding(4)
-						[
-							SNew(STextBlock)
-							.Text(this, &SSCS_RowWidget_ActorRoot::GetActorDisplayText)
-							.Font(FEditorStyle::GetFontStyle("ContentBrowser.TileViewTooltip.NameFont"))
-						]
+						SNew(STextBlock)
+						.Text(this, &SSCS_RowWidget_ActorRoot::GetActorDisplayText)
+						.ColorAndOpacity( FLinearColor::Black )
+						.Font(FEditorStyle::GetFontStyle("ContentBrowser.TileViewTooltip.NameFont"))
 					]
 				]
 			]
@@ -2545,8 +2526,8 @@ TSharedRef<SToolTip> SSCS_RowWidget_ActorRoot::CreateToolTipWidget() const
 			.AutoHeight()
 			[
 				SNew(SBorder)
+				.BorderImage(FEditorStyle::GetBrush("NoBorder"))
 				.Padding(4)
-				.BorderImage(FEditorStyle::GetBrush("ToolPanel.GroupBorder"))
 				[
 					InfoBox
 				]
