[INTEGRATE] Change 2422698 by [REDACTED_PERSON][REDACTED_HANDLE] on 2015/01/28 15:20:27

	Fix for being able to add a setter for the SCS generated component variables. UE-8135
	This changed the SCS added component variables to be CPF_BlueprintReadOnly which means that blueprints that include setters will no longer compile.

diff --git a/Engine/Source/Editor/BlueprintGraph/Private/VariableSetHandler.cpp b/Engine/Source/Editor/BlueprintGraph/Private/VariableSetHandler.cpp
--- a/Engine/Source/Editor/BlueprintGraph/Private/VariableSetHandler.cpp
+++ b/Engine/Source/Editor/BlueprintGraph/Private/VariableSetHandler.cpp
@@ -24,7 +24,7 @@ void FKCHandler_VariableSet::RegisterNets(FKismetFunctionContext& Context, UEdGr
 
 		if(FBlueprintEditorUtils::IsPropertyReadOnlyInCurrentBlueprint(Context.Blueprint, SetterNode->GetPropertyForVariable()))
 		{
-			CompilerContext.MessageLog.Warning(*LOCTEXT("BlueprintReadOnlyOrPrivate_Error", "The property is mark as BlueprintReadOnly or Private. It cannot be modifed in the blueprint. @@").ToString(), Node);
+			CompilerContext.MessageLog.Warning(*LOCTEXT("BlueprintReadOnlyOrPrivate_Error", "The property is marked as BlueprintReadOnly or Private. It cannot be modifed in the blueprint. @@").ToString(), Node);
 		}
 
 		// Report an error that the local variable could not be found
diff --git a/Engine/Source/Editor/Kismet/Private/BPVariableDragDropAction.cpp b/Engine/Source/Editor/Kismet/Private/BPVariableDragDropAction.cpp
--- a/Engine/Source/Editor/Kismet/Private/BPVariableDragDropAction.cpp
+++ b/Engine/Source/Editor/Kismet/Private/BPVariableDragDropAction.cpp
@@ -59,6 +59,7 @@ void FKismetVariableDragDropAction::HoverTargetChanged()
 
 	UEdGraphPin* PinUnderCursor = GetHoveredPin();
 
+	bool bCanMakeSetter = false;
 	bool bBadSchema = false;
 	bool bBadGraph = false;
 	UEdGraph* HoveredGraph = GetHoveredGraph();
@@ -72,6 +73,16 @@ void FKismetVariableDragDropAction::HoverTargetChanged()
 		{
 			bBadGraph = true;
 		}
+
+		UStruct* Outer = CastChecked<UStruct>(VariableProperty->GetOuter());
+
+		FNodeConstructionParams NewNodeParams;
+		NewNodeParams.VariableName = VariableName;
+		const UBlueprint* DropOnBlueprint = FBlueprintEditorUtils::FindBlueprintForGraph(HoveredGraph);
+		NewNodeParams.Graph = HoveredGraph;
+		NewNodeParams.VariableSource = Outer;
+		
+		bCanMakeSetter = CanExecuteMakeSetter(NewNodeParams, VariableProperty);
 	}
 
 	UEdGraphNode* VarNodeUnderCursor = Cast<UK2Node_Variable>(GetHoveredNode());
@@ -286,6 +297,14 @@ void FKismetVariableDragDropAction::HoverTargetChanged()
 			}
 		}
 	}
+	else if (bAltDrag && !bCanMakeSetter)
+	{
+		FFormatNamedArguments Args;
+		Args.Add(TEXT("VariableName"), FText::FromString(VariableString));
+
+		StatusSymbol = FEditorStyle::GetBrush(TEXT("Graph.ConnectorFeedback.Error"));
+		Message = FText::Format(LOCTEXT("CannotPlaceSetter", "Variable '{VariableName}' is readonly you cannot set this variable."), Args);
+	}
 	// Draw variable icon
 	else
 	{
@@ -446,7 +465,7 @@ FReply FKismetVariableDragDropAction::DroppedOnPanel( const TSharedRef< SWidget
 			AnalyticCallback.ExecuteIfBound();
 
 			// Asking for a getter
-			if (bControlDrag)
+			if (bControlDrag || !CanExecuteMakeSetter(NewNodeParams, VariableProperty))
 			{
 				MakeGetter(NewNodeParams);
 			}
diff --git a/Engine/Source/Editor/KismetCompiler/Private/KismetCompiler.cpp b/Engine/Source/Editor/KismetCompiler/Private/KismetCompiler.cpp
--- a/Engine/Source/Editor/KismetCompiler/Private/KismetCompiler.cpp
+++ b/Engine/Source/Editor/KismetCompiler/Private/KismetCompiler.cpp
@@ -587,7 +587,7 @@ void FKismetCompilerContext::CreateClassVariablesFromBlueprint()
 					const FString CategoryName = Node->CategoryName != NAME_None ? Node->CategoryName.ToString() : Blueprint->GetName();
 					
 					NewProperty->SetMetaData(TEXT("Category"), *CategoryName);
-					NewProperty->SetPropertyFlags(CPF_BlueprintVisible|CPF_NonTransactional);
+					NewProperty->SetPropertyFlags(CPF_BlueprintVisible | CPF_NonTransactional | CPF_BlueprintReadOnly);
 				}
 			}
 		}
diff --git a/Engine/Source/Editor/UnrealEd/Private/SComponentClassCombo.cpp b/Engine/Source/Editor/UnrealEd/Private/SComponentClassCombo.cpp
--- a/Engine/Source/Editor/UnrealEd/Private/SComponentClassCombo.cpp
+++ b/Engine/Source/Editor/UnrealEd/Private/SComponentClassCombo.cpp
@@ -105,6 +105,8 @@ void SComponentClassCombo::Construct(const FArguments& InArgs)
 
 void SComponentClassCombo::ClearSelection()
 {
+	SearchBox->SetText(FText::GetEmpty());
+
 	PrevSelectedIndex = INDEX_NONE;
 
 	// Clear the selection in such a way as to also clear the keyboard selector
@@ -393,11 +395,6 @@ void SComponentClassCombo::UpdateComponentClassList()
 	}
 }
 
-FReply SComponentClassCombo::OnButtonClicked()
-{
-	SearchBox->SetText(FText::GetEmpty());
-	return SComboButton::OnButtonClicked();
-}
 
 void SComponentClassCombo::OnProjectHotReloaded( bool bWasTriggeredAutomatically )
 {
diff --git a/Engine/Source/Editor/UnrealEd/Public/SComponentClassCombo.h b/Engine/Source/Editor/UnrealEd/Public/SComponentClassCombo.h
--- a/Engine/Source/Editor/UnrealEd/Public/SComponentClassCombo.h
+++ b/Engine/Source/Editor/UnrealEd/Public/SComponentClassCombo.h
@@ -124,10 +124,10 @@ public:
 	static FString GetSanitizedComponentName( UClass* ComponentClass );
 
 protected:
-	virtual FReply OnButtonClicked();
 
 	/** Called when a project is hot reloaded to refresh the components list */
 	void OnProjectHotReloaded( bool bWasTriggeredAutomatically );
+
 private:
 
 	FText GetFriendlyComponentName(FComponentClassComboEntryPtr Entry) const;
