[INTEGRATE] Change 2422806 by [REDACTED_PERSON][REDACTED_HANDLE] on 2015/01/28 16:18:53

	Component tree:
	- Disabled drag-drop from the root node for now, it isn't desired in the instance case and doesn't work correctly in the BP editor case currently [UE-8461] [UE-8469]
	- Collapsed some duplicated code
	- Improved the phrasing in a couple of comments

diff --git a/Engine/Source/Editor/Kismet/Private/SSCSEditor.cpp b/Engine/Source/Editor/Kismet/Private/SSCSEditor.cpp
--- a/Engine/Source/Editor/Kismet/Private/SSCSEditor.cpp
+++ b/Engine/Source/Editor/Kismet/Private/SSCSEditor.cpp
@@ -144,7 +144,7 @@ class FSCSRowDragDropOp : public FKismetVariableDragDropAction
 public:
 	DRAG_DROP_OPERATOR_TYPE(FSCSRowDragDropOp, FKismetVariableDragDropAction)
 
-		/** Available drop actions */
+	/** Available drop actions */
 	enum EDropActionType
 	{
 		DropAction_None,
@@ -1477,13 +1477,11 @@ FReply SSCS_RowWidget::OnDragDetected( const FGeometry& MyGeometry, const FPoint
 		}
 
 		TSharedPtr<FSCSEditorTreeNode> FirstNode = SelectedNodePtrs[0];
-		if (FirstNode->GetNodeType() != FSCSEditorTreeNode::SeparatorNode)
+		if (FirstNode->GetNodeType() == FSCSEditorTreeNode::ComponentNode)
 		{
 			UBlueprint* Blueprint = FirstNode->GetBlueprint();
-
-			bool bIsRootActor = FirstNode->GetNodeType() == FSCSEditorTreeNode::RootActorNode;
-
-			TSharedRef<FSCSRowDragDropOp> Operation = FSCSRowDragDropOp::New(bIsRootActor ? FName("Self") : FirstNode->GetVariableName(), Blueprint != nullptr ? Blueprint->SkeletonGeneratedClass : nullptr, FNodeCreationAnalytic());
+			const FName VariableName = FirstNode->GetVariableName();
+			TSharedRef<FSCSRowDragDropOp> Operation = FSCSRowDragDropOp::New(VariableName, Blueprint != nullptr ? Blueprint->SkeletonGeneratedClass : nullptr, FNodeCreationAnalytic());
 			Operation->SetCtrlDrag(true); // Always put a getter
 			Operation->PendingDropAction = FSCSRowDragDropOp::DropAction_None;
 			Operation->SourceNodes = SelectedNodePtrs;
@@ -1548,7 +1546,7 @@ void SSCS_RowWidget::OnDragEnter( const FGeometry& MyGeometry, const FDragDropEv
 			if ((NodePtr->GetNodeType() == FSCSEditorTreeNode::SeparatorNode) || (NodePtr->GetNodeType() == FSCSEditorTreeNode::RootActorNode))
 			{
 				// Don't show a feedback message if over a node that makes no sense, such as a separator or the instance node
-				Message = LOCTEXT("DropActionToolTip_FriendlyError_DragToAComponent", "Drag to another component in order to attach to that component or become the root component.");
+				Message = LOCTEXT("DropActionToolTip_FriendlyError_DragToAComponent", "Drag to another component in order to attach to that component or become the root component.\nDrag to a Blueprint graph in order to drop a reference.");
 			}
 
 			// Validate each selected node being dragged against the node that belongs to this row. Exit the loop if we have a valid tooltip OR a valid pending drop action once all nodes in the selection have been validated.
@@ -2762,48 +2760,28 @@ void SSCSEditor::Construct( const FArguments& InArgs )
 		.FillWidth(3);
 	}
 
+	SCSTreeWidget = SNew(SSCSTreeType)
+		.ToolTipText(LOCTEXT("DropAssetToAddComponent", "Drop asset here to add a component."))
+		.SCSEditor(this)
+		.TreeItemsSource(bSingleLayoutBPEditor ? &RootNodes : &RootComponentNodes)
+		.SelectionMode(ESelectionMode::Multi)
+		.OnGenerateRow(this, &SSCSEditor::MakeTableRowWidget)
+		.OnGetChildren(this, &SSCSEditor::OnGetChildrenForTree)
+		.OnSelectionChanged(this, &SSCSEditor::OnTreeSelectionChanged)
+		.OnContextMenuOpening(this, &SSCSEditor::CreateContextMenu)
+		.OnItemScrolledIntoView(this, &SSCSEditor::OnItemScrolledIntoView)
+		.OnMouseButtonDoubleClick(this, &SSCSEditor::HandleItemDoubleClicked)
+		.ClearSelectionOnClick(InArgs._EditorMode == EComponentEditorMode::BlueprintSCS ? true : false)
+		.ItemHeight(24)
+		.HeaderRow
+		(
+			HeaderRow
+		);
+
 	if (bSingleLayoutBPEditor)
 	{
-		SAssignNew(SCSTreeWidget, SSCSTreeType)
-			.ToolTipText(LOCTEXT("DropAssetToAddComponent", "Drop asset here to add component."))
-			.SCSEditor(this)
-			.TreeItemsSource(&RootNodes)
-			.SelectionMode(ESelectionMode::Multi)
-			.OnGenerateRow(this, &SSCSEditor::MakeTableRowWidget)
-			.OnGetChildren(this, &SSCSEditor::OnGetChildrenForTree)
-			.OnSelectionChanged(this, &SSCSEditor::OnTreeSelectionChanged)
-			.OnContextMenuOpening(this, &SSCSEditor::CreateContextMenu)
-			.OnItemScrolledIntoView(this, &SSCSEditor::OnItemScrolledIntoView)
-			.OnMouseButtonDoubleClick(this, &SSCSEditor::HandleItemDoubleClicked)
-			.ClearSelectionOnClick(InArgs._EditorMode == EComponentEditorMode::BlueprintSCS ? true : false)
-			.ItemHeight(24)
-			.HeaderRow
-			(
-				HeaderRow
-			);
-
 		SCSTreeWidget->GetHeaderRow()->SetVisibility(EVisibility::Collapsed);
 	}
-	else
-	{
-		SAssignNew(SCSTreeWidget, SSCSTreeType)
-			.ToolTipText(LOCTEXT("DropAssetToAddComponent", "Drop asset here to add component."))
-			.SCSEditor(this)
-			.TreeItemsSource(&RootComponentNodes)
-			.SelectionMode(ESelectionMode::Multi)
-			.OnGenerateRow(this, &SSCSEditor::MakeTableRowWidget)
-			.OnGetChildren(this, &SSCSEditor::OnGetChildrenForTree)
-			.OnSelectionChanged(this, &SSCSEditor::OnTreeSelectionChanged)
-			.OnContextMenuOpening(this, &SSCSEditor::CreateContextMenu)
-			.OnItemScrolledIntoView(this, &SSCSEditor::OnItemScrolledIntoView)
-			.OnMouseButtonDoubleClick(this, &SSCSEditor::HandleItemDoubleClicked)
-			.ClearSelectionOnClick(InArgs._EditorMode == EComponentEditorMode::BlueprintSCS ? true : false)
-			.ItemHeight(24)
-			.HeaderRow
-			(
-				HeaderRow
-			);
-	}
 
 	TSharedPtr<SWidget> Contents;
 
