[INTEGRATE] Change 2423037 by [REDACTED_PERSON][REDACTED_HANDLE] on 2015/01/28 19:01:22

	Add Goto code/blueprint in content browser and open source/blueprint on to the right click menu for components in the SCSEditor tree

diff --git a/Engine/Source/Editor/Kismet/Private/SSCSEditor.cpp b/Engine/Source/Editor/Kismet/Private/SSCSEditor.cpp
--- a/Engine/Source/Editor/Kismet/Private/SSCSEditor.cpp
+++ b/Engine/Source/Editor/Kismet/Private/SSCSEditor.cpp
@@ -48,6 +48,8 @@
 #include "AssetRegistryModule.h"
 #include "SCreateAssetFromObject.h"
 
+#include "SourceCodeNavigation.h"
+
 #define LOCTEXT_NAMESPACE "SSCSEditor"
 
 DEFINE_LOG_CATEGORY_STATIC(LogSCSEditor, Log, All);
@@ -3116,6 +3118,24 @@ void SSCSEditor::GetSelectedItemsForContextMenu(TArray<FComponentEventConstructi
 	}
 }
 
+void SSCSEditor::OnGoToAssetInBrowser(UObject* Asset) const
+{
+	TArray<UObject*> Objects;
+	Objects.Add(Asset);
+	GEditor->SyncBrowserToObjects( Objects );
+}
+
+void SSCSEditor::OnEditBlueprint(UObject* Asset) const
+{
+	FAssetEditorManager::Get().OpenEditorForAsset(Asset);
+}
+
+void SSCSEditor::OnOpenCodeFile(const FString CodeFileName) const
+{
+	const FString AbsoluteHeaderPath = IFileManager::Get().ConvertToAbsolutePathForExternalAppForRead(*CodeFileName);
+	FSourceCodeNavigation::OpenSourceFile( AbsoluteHeaderPath );
+}
+
 TSharedPtr< SWidget > SSCSEditor::CreateContextMenu() // @todo: make the context manual meaningful for the actor and avoid separators
 {
 	TArray<FSCSEditorTreeNodePtrType> SelectedNodes = SCSTreeWidget->GetSelectedItems();
@@ -3172,6 +3192,64 @@ TSharedPtr< SWidget > SSCSEditor::CreateContextMenu() // @todo: make the context
 		}
 		MenuBuilder.EndSection();
 
+		if (SelectedNodes.Num() == 1)
+		{
+			UActorComponent* ComponentTemplate = SelectedNodes[0]->GetComponentTemplate();
+			if (ComponentTemplate->GetClass()->ClassGeneratedBy)
+			{
+				MenuBuilder.BeginSection("ComponentAsset", LOCTEXT("ComponentAssetHeading", "Asset"));
+				{
+					MenuBuilder.AddMenuEntry(
+						FText::Format(LOCTEXT("GoToBlueprintForComponent", "Edit {0}"), FText::FromString(ComponentTemplate->GetClass()->ClassGeneratedBy->GetName())),
+						LOCTEXT("EditBlueprintForComponent_ToolTip", "Edits the class blueprint that defines this component."),
+						FSlateIcon(FEditorStyle::GetStyleSetName(), FClassIconFinder::FindIconNameForClass(ComponentTemplate->GetClass())),
+						FUIAction(
+							FExecuteAction::CreateSP(this, &SSCSEditor::OnEditBlueprint, ComponentTemplate->GetClass()->ClassGeneratedBy),
+							FCanExecuteAction()));
+
+					MenuBuilder.AddMenuEntry(
+						LOCTEXT("GoToAssetForComponent", "Find Class in Content Browser"),
+						LOCTEXT("GoToAssetForComponent_ToolTip", "Summons the content browser and goes to the class for this component."),
+						FSlateIcon(FEditorStyle::GetStyleSetName(), "SystemWideCommands.FindInContentBrowser"),
+						FUIAction(
+							FExecuteAction::CreateSP(this, &SSCSEditor::OnGoToAssetInBrowser, ComponentTemplate->GetClass()->ClassGeneratedBy),
+							FCanExecuteAction()));
+				}
+				MenuBuilder.EndSection();			
+			}
+			else
+			{
+				MenuBuilder.BeginSection("ComponentCode", LOCTEXT("ComponentCodeHeading", "C++"));
+				{
+					if (FSourceCodeNavigation::IsCompilerAvailable())
+					{
+						FString ClassHeaderPath;
+						if (FSourceCodeNavigation::FindClassHeaderPath(ComponentTemplate->GetClass(), ClassHeaderPath) && IFileManager::Get().FileSize(*ClassHeaderPath) != INDEX_NONE)
+						{
+							const FString CodeFileName = FPaths::GetCleanFilename(*ClassHeaderPath);
+
+							MenuBuilder.AddMenuEntry(
+								FText::Format(LOCTEXT("GoToCodeForComponent", "Open {0}"), FText::FromString(CodeFileName)),
+								FText::Format(LOCTEXT("GoToCodeForComponent_ToolTip", "Opens the header file for this component ({0}) in a code editing program"), FText::FromString(CodeFileName)),
+								FSlateIcon(),
+								FUIAction(
+									FExecuteAction::CreateSP(this, &SSCSEditor::OnOpenCodeFile, ClassHeaderPath),
+									FCanExecuteAction()));
+						}
+
+						MenuBuilder.AddMenuEntry(
+							LOCTEXT("GoToAssetForComponent", "Find Class in Content Browser"),
+							LOCTEXT("GoToAssetForComponent_ToolTip", "Summons the content browser and goes to the class for this component."),
+							FSlateIcon(FEditorStyle::GetStyleSetName(), "SystemWideCommands.FindInContentBrowser"),
+							FUIAction(
+								FExecuteAction::CreateSP(this, &SSCSEditor::OnGoToAssetInBrowser, (UObject*)ComponentTemplate->GetClass()),
+								FCanExecuteAction()));
+					}
+				}
+				MenuBuilder.EndSection();
+			}
+		}
+
 		return MenuBuilder.MakeWidget();
 	}
 	return TSharedPtr<SWidget>();
diff --git a/Engine/Source/Editor/Kismet/Public/SSCSEditor.h b/Engine/Source/Editor/Kismet/Public/SSCSEditor.h
--- a/Engine/Source/Editor/Kismet/Public/SSCSEditor.h
+++ b/Engine/Source/Editor/Kismet/Public/SSCSEditor.h
@@ -584,6 +584,12 @@ public:
 	 */
 	static void BuildMenuEventsSection( FMenuBuilder& Menu, UBlueprint* Blueprint, UClass* SelectedClass, FCanExecuteAction CanExecuteActionDelegate, FGetSelectedObjectsDelegate GetSelectedObjectsDelegate );
 
+	void OnEditBlueprint(UObject* Blueprint) const;
+
+	void OnGoToAssetInBrowser(UObject* Asset) const;
+
+	void OnOpenCodeFile(const FString CodeFileName) const;
+
 	/**
 	 * Given an actor component, attempts to find an associated tree node.
 	 *
