[INTEGRATE] Change 2423798 by [REDACTED_PERSON][REDACTED_HANDLE] on 2015/01/29 10:44:10

	Don't crash when right-clicking on the Actor root in the SCS Tree

diff --git a/Engine/Source/Editor/Kismet/Private/SSCSEditor.cpp b/Engine/Source/Editor/Kismet/Private/SSCSEditor.cpp
--- a/Engine/Source/Editor/Kismet/Private/SSCSEditor.cpp
+++ b/Engine/Source/Editor/Kismet/Private/SSCSEditor.cpp
@@ -3195,58 +3195,61 @@ TSharedPtr< SWidget > SSCSEditor::CreateContextMenu() // @todo: make the context
 		if (SelectedNodes.Num() == 1)
 		{
 			UActorComponent* ComponentTemplate = SelectedNodes[0]->GetComponentTemplate();
-			if (ComponentTemplate->GetClass()->ClassGeneratedBy)
+			if (ComponentTemplate)
 			{
-				MenuBuilder.BeginSection("ComponentAsset", LOCTEXT("ComponentAssetHeading", "Asset"));
+				if (ComponentTemplate->GetClass()->ClassGeneratedBy)
 				{
-					MenuBuilder.AddMenuEntry(
-						FText::Format(LOCTEXT("GoToBlueprintForComponent", "Edit {0}"), FText::FromString(ComponentTemplate->GetClass()->ClassGeneratedBy->GetName())),
-						LOCTEXT("EditBlueprintForComponent_ToolTip", "Edits the class blueprint that defines this component."),
-						FSlateIcon(FEditorStyle::GetStyleSetName(), FClassIconFinder::FindIconNameForClass(ComponentTemplate->GetClass())),
-						FUIAction(
-							FExecuteAction::CreateSP(this, &SSCSEditor::OnEditBlueprint, ComponentTemplate->GetClass()->ClassGeneratedBy),
-							FCanExecuteAction()));
-
-					MenuBuilder.AddMenuEntry(
-						LOCTEXT("GoToAssetForComponent", "Find Class in Content Browser"),
-						LOCTEXT("GoToAssetForComponent_ToolTip", "Summons the content browser and goes to the class for this component."),
-						FSlateIcon(FEditorStyle::GetStyleSetName(), "SystemWideCommands.FindInContentBrowser"),
-						FUIAction(
-							FExecuteAction::CreateSP(this, &SSCSEditor::OnGoToAssetInBrowser, ComponentTemplate->GetClass()->ClassGeneratedBy),
-							FCanExecuteAction()));
+					MenuBuilder.BeginSection("ComponentAsset", LOCTEXT("ComponentAssetHeading", "Asset"));
+					{
+						MenuBuilder.AddMenuEntry(
+							FText::Format(LOCTEXT("GoToBlueprintForComponent", "Edit {0}"), FText::FromString(ComponentTemplate->GetClass()->ClassGeneratedBy->GetName())),
+							LOCTEXT("EditBlueprintForComponent_ToolTip", "Edits the class blueprint that defines this component."),
+							FSlateIcon(FEditorStyle::GetStyleSetName(), FClassIconFinder::FindIconNameForClass(ComponentTemplate->GetClass())),
+							FUIAction(
+								FExecuteAction::CreateSP(this, &SSCSEditor::OnEditBlueprint, ComponentTemplate->GetClass()->ClassGeneratedBy),
+								FCanExecuteAction()));
+
+						MenuBuilder.AddMenuEntry(
+							LOCTEXT("GoToAssetForComponent", "Find Class in Content Browser"),
+							LOCTEXT("GoToAssetForComponent_ToolTip", "Summons the content browser and goes to the class for this component."),
+							FSlateIcon(FEditorStyle::GetStyleSetName(), "SystemWideCommands.FindInContentBrowser"),
+							FUIAction(
+								FExecuteAction::CreateSP(this, &SSCSEditor::OnGoToAssetInBrowser, ComponentTemplate->GetClass()->ClassGeneratedBy),
+								FCanExecuteAction()));
+					}
+					MenuBuilder.EndSection();			
 				}
-				MenuBuilder.EndSection();			
-			}
-			else
-			{
-				MenuBuilder.BeginSection("ComponentCode", LOCTEXT("ComponentCodeHeading", "C++"));
+				else
 				{
-					if (FSourceCodeNavigation::IsCompilerAvailable())
+					MenuBuilder.BeginSection("ComponentCode", LOCTEXT("ComponentCodeHeading", "C++"));
 					{
-						FString ClassHeaderPath;
-						if (FSourceCodeNavigation::FindClassHeaderPath(ComponentTemplate->GetClass(), ClassHeaderPath) && IFileManager::Get().FileSize(*ClassHeaderPath) != INDEX_NONE)
+						if (FSourceCodeNavigation::IsCompilerAvailable())
 						{
-							const FString CodeFileName = FPaths::GetCleanFilename(*ClassHeaderPath);
+							FString ClassHeaderPath;
+							if (FSourceCodeNavigation::FindClassHeaderPath(ComponentTemplate->GetClass(), ClassHeaderPath) && IFileManager::Get().FileSize(*ClassHeaderPath) != INDEX_NONE)
+							{
+								const FString CodeFileName = FPaths::GetCleanFilename(*ClassHeaderPath);
+
+								MenuBuilder.AddMenuEntry(
+									FText::Format(LOCTEXT("GoToCodeForComponent", "Open {0}"), FText::FromString(CodeFileName)),
+									FText::Format(LOCTEXT("GoToCodeForComponent_ToolTip", "Opens the header file for this component ({0}) in a code editing program"), FText::FromString(CodeFileName)),
+									FSlateIcon(),
+									FUIAction(
+										FExecuteAction::CreateSP(this, &SSCSEditor::OnOpenCodeFile, ClassHeaderPath),
+										FCanExecuteAction()));
+							}
 
 							MenuBuilder.AddMenuEntry(
-								FText::Format(LOCTEXT("GoToCodeForComponent", "Open {0}"), FText::FromString(CodeFileName)),
-								FText::Format(LOCTEXT("GoToCodeForComponent_ToolTip", "Opens the header file for this component ({0}) in a code editing program"), FText::FromString(CodeFileName)),
-								FSlateIcon(),
+								LOCTEXT("GoToAssetForComponent", "Find Class in Content Browser"),
+								LOCTEXT("GoToAssetForComponent_ToolTip", "Summons the content browser and goes to the class for this component."),
+								FSlateIcon(FEditorStyle::GetStyleSetName(), "SystemWideCommands.FindInContentBrowser"),
 								FUIAction(
-									FExecuteAction::CreateSP(this, &SSCSEditor::OnOpenCodeFile, ClassHeaderPath),
+									FExecuteAction::CreateSP(this, &SSCSEditor::OnGoToAssetInBrowser, (UObject*)ComponentTemplate->GetClass()),
 									FCanExecuteAction()));
 						}
-
-						MenuBuilder.AddMenuEntry(
-							LOCTEXT("GoToAssetForComponent", "Find Class in Content Browser"),
-							LOCTEXT("GoToAssetForComponent_ToolTip", "Summons the content browser and goes to the class for this component."),
-							FSlateIcon(FEditorStyle::GetStyleSetName(), "SystemWideCommands.FindInContentBrowser"),
-							FUIAction(
-								FExecuteAction::CreateSP(this, &SSCSEditor::OnGoToAssetInBrowser, (UObject*)ComponentTemplate->GetClass()),
-								FCanExecuteAction()));
 					}
+					MenuBuilder.EndSection();
 				}
-				MenuBuilder.EndSection();
 			}
 		}
 
