[INTEGRATE] Change 2423965 by [REDACTED_PERSON][REDACTED_HANDLE] on 2015/01/29 12:13:30

	Make right-click on the root node in SCSEditor show Actor actions in LevelEditor and just the Component paste option in BlueprintEditor

diff --git a/Engine/Source/Editor/Kismet/Private/SSCSEditor.cpp b/Engine/Source/Editor/Kismet/Private/SSCSEditor.cpp
--- a/Engine/Source/Editor/Kismet/Private/SSCSEditor.cpp
+++ b/Engine/Source/Editor/Kismet/Private/SSCSEditor.cpp
@@ -2687,6 +2687,7 @@ void SSCSEditor::Construct( const FArguments& InArgs )
 	ActorContext = InArgs._ActorContext;
 	AllowEditing = InArgs._AllowEditing;
 	PreviewActor = InArgs._PreviewActor;
+	ActorMenuExtender = InArgs._ActorMenuExtender;
 	OnSelectionUpdated = InArgs._OnSelectionUpdated;
 	OnItemDoubleClicked = InArgs._OnItemDoubleClicked;
 	OnHighlightPropertyInDetailsView = InArgs._OnHighlightPropertyInDetailsView;
@@ -3143,115 +3144,150 @@ TSharedPtr< SWidget > SSCSEditor::CreateContextMenu() // @todo: make the context
 	if (SelectedNodes.Num() > 0 || CanPasteNodes())
 	{
 		const bool CloseAfterSelection = true;
-		FMenuBuilder MenuBuilder( CloseAfterSelection, CommandList );
+		FMenuBuilder MenuBuilder( CloseAfterSelection, CommandList, ActorMenuExtender );
 
-		MenuBuilder.BeginSection("ComponentActions", LOCTEXT("ComponentContextMenu", "Component Actions") );
+		bool bShowJustPasteOption = false;
+
+		if (SelectedNodes.Num() > 0)
 		{
-			if(SelectedNodes.Num() > 0)
+			if (SelectedNodes.Num() == 1 && SelectedNodes[0]->GetNodeType() == FSCSEditorTreeNode::RootActorNode)
 			{
-				MenuBuilder.AddMenuEntry( FGenericCommands::Get().Cut) ;
-				MenuBuilder.AddMenuEntry( FGenericCommands::Get().Copy );
-				MenuBuilder.AddMenuEntry( FGenericCommands::Get().Paste );
-				MenuBuilder.AddMenuEntry( FGenericCommands::Get().Duplicate );
-				MenuBuilder.AddMenuEntry( FGenericCommands::Get().Delete );
-				MenuBuilder.AddMenuEntry( FGenericCommands::Get().Rename );
-
 				if (EditorMode == EComponentEditorMode::BlueprintSCS)
 				{
-					// Collect the classes of all selected objects
-					TArray<UClass*> SelectionClasses;
-					for( auto NodeIter = SelectedNodes.CreateConstIterator(); NodeIter; ++NodeIter )
-					{
-						auto TreeNode = *NodeIter;
-						if( UActorComponent* ComponentTemplate = TreeNode->GetComponentTemplate() )
-						{
-							SelectionClasses.Add(ComponentTemplate->GetClass());
-						}
-					}
-
-					if ( SelectionClasses.Num() )
-					{
-						// Find the common base class of all selected classes
-						UClass* SelectedClass = UClass::FindCommonBase( SelectionClasses );
-						// Build an event submenu if we can generate events
-						if( FBlueprintEditorUtils::CanClassGenerateEvents( SelectedClass ))
-						{
-							MenuBuilder.AddSubMenu(	LOCTEXT("AddEventSubMenu", "Add Event"), 
-								LOCTEXT("ActtionsSubMenu_ToolTip", "Add Event"), 
-								FNewMenuDelegate::CreateStatic( &SSCSEditor::BuildMenuEventsSection,
-								GetBlueprint(), SelectedClass, FCanExecuteAction::CreateSP(this, &SSCSEditor::IsEditingAllowed),
-								FGetSelectedObjectsDelegate::CreateSP(this, &SSCSEditor::GetSelectedItemsForContextMenu)));
-						}
-					}
+					bShowJustPasteOption = false;
+				}
+				else
+				{
+					// Display the Actor menu
+					MenuBuilder.BeginSection("MainSection");
+					MenuBuilder.EndSection();
 				}
 			}
 			else
 			{
-				MenuBuilder.AddMenuEntry( FGenericCommands::Get().Paste );
-			}
-		}
-		MenuBuilder.EndSection();
-
-		if (SelectedNodes.Num() == 1)
-		{
-			UActorComponent* ComponentTemplate = SelectedNodes[0]->GetComponentTemplate();
-			if (ComponentTemplate)
-			{
-				if (ComponentTemplate->GetClass()->ClassGeneratedBy)
+				for (auto SelectedNode : SelectedNodes)
 				{
-					MenuBuilder.BeginSection("ComponentAsset", LOCTEXT("ComponentAssetHeading", "Asset"));
+					if (SelectedNode->GetNodeType() != FSCSEditorTreeNode::ComponentNode)
 					{
-						MenuBuilder.AddMenuEntry(
-							FText::Format(LOCTEXT("GoToBlueprintForComponent", "Edit {0}"), FText::FromString(ComponentTemplate->GetClass()->ClassGeneratedBy->GetName())),
-							LOCTEXT("EditBlueprintForComponent_ToolTip", "Edits the class blueprint that defines this component."),
-							FSlateIcon(FEditorStyle::GetStyleSetName(), FClassIconFinder::FindIconNameForClass(ComponentTemplate->GetClass())),
-							FUIAction(
-								FExecuteAction::CreateSP(this, &SSCSEditor::OnEditBlueprint, ComponentTemplate->GetClass()->ClassGeneratedBy),
-								FCanExecuteAction()));
-
-						MenuBuilder.AddMenuEntry(
-							LOCTEXT("GoToAssetForComponent", "Find Class in Content Browser"),
-							LOCTEXT("GoToAssetForComponent_ToolTip", "Summons the content browser and goes to the class for this component."),
-							FSlateIcon(FEditorStyle::GetStyleSetName(), "SystemWideCommands.FindInContentBrowser"),
-							FUIAction(
-								FExecuteAction::CreateSP(this, &SSCSEditor::OnGoToAssetInBrowser, ComponentTemplate->GetClass()->ClassGeneratedBy),
-								FCanExecuteAction()));
+						bShowJustPasteOption = true;
+						break;
 					}
-					MenuBuilder.EndSection();			
 				}
-				else
+				if (!bShowJustPasteOption)
 				{
-					MenuBuilder.BeginSection("ComponentCode", LOCTEXT("ComponentCodeHeading", "C++"));
+					MenuBuilder.BeginSection("ComponentActions", LOCTEXT("ComponentContextMenu", "Component Actions") );
 					{
-						if (FSourceCodeNavigation::IsCompilerAvailable())
+						MenuBuilder.AddMenuEntry( FGenericCommands::Get().Cut) ;
+						MenuBuilder.AddMenuEntry( FGenericCommands::Get().Copy );
+						MenuBuilder.AddMenuEntry( FGenericCommands::Get().Paste );
+						MenuBuilder.AddMenuEntry( FGenericCommands::Get().Duplicate );
+						MenuBuilder.AddMenuEntry( FGenericCommands::Get().Delete );
+						MenuBuilder.AddMenuEntry( FGenericCommands::Get().Rename );
+
+						if (EditorMode == EComponentEditorMode::BlueprintSCS)
 						{
-							FString ClassHeaderPath;
-							if (FSourceCodeNavigation::FindClassHeaderPath(ComponentTemplate->GetClass(), ClassHeaderPath) && IFileManager::Get().FileSize(*ClassHeaderPath) != INDEX_NONE)
+							// Collect the classes of all selected objects
+							TArray<UClass*> SelectionClasses;
+							for( auto NodeIter = SelectedNodes.CreateConstIterator(); NodeIter; ++NodeIter )
 							{
-								const FString CodeFileName = FPaths::GetCleanFilename(*ClassHeaderPath);
+								auto TreeNode = *NodeIter;
+								if( UActorComponent* ComponentTemplate = TreeNode->GetComponentTemplate() )
+								{
+									SelectionClasses.Add(ComponentTemplate->GetClass());
+								}
+							}
 
+							if ( SelectionClasses.Num() )
+							{
+								// Find the common base class of all selected classes
+								UClass* SelectedClass = UClass::FindCommonBase( SelectionClasses );
+								// Build an event submenu if we can generate events
+								if( FBlueprintEditorUtils::CanClassGenerateEvents( SelectedClass ))
+								{
+									MenuBuilder.AddSubMenu(	LOCTEXT("AddEventSubMenu", "Add Event"), 
+										LOCTEXT("ActtionsSubMenu_ToolTip", "Add Event"), 
+										FNewMenuDelegate::CreateStatic( &SSCSEditor::BuildMenuEventsSection,
+										GetBlueprint(), SelectedClass, FCanExecuteAction::CreateSP(this, &SSCSEditor::IsEditingAllowed),
+										FGetSelectedObjectsDelegate::CreateSP(this, &SSCSEditor::GetSelectedItemsForContextMenu)));
+								}
+							}
+						}					
+					}
+					MenuBuilder.EndSection();
+
+					if (SelectedNodes.Num() == 1)
+					{
+						UActorComponent* ComponentTemplate = SelectedNodes[0]->GetComponentTemplate();
+						if (ComponentTemplate->GetClass()->ClassGeneratedBy)
+						{
+							MenuBuilder.BeginSection("ComponentAsset", LOCTEXT("ComponentAssetHeading", "Asset"));
+							{
 								MenuBuilder.AddMenuEntry(
-									FText::Format(LOCTEXT("GoToCodeForComponent", "Open {0}"), FText::FromString(CodeFileName)),
-									FText::Format(LOCTEXT("GoToCodeForComponent_ToolTip", "Opens the header file for this component ({0}) in a code editing program"), FText::FromString(CodeFileName)),
-									FSlateIcon(),
+									FText::Format(LOCTEXT("GoToBlueprintForComponent", "Edit {0}"), FText::FromString(ComponentTemplate->GetClass()->ClassGeneratedBy->GetName())),
+									LOCTEXT("EditBlueprintForComponent_ToolTip", "Edits the class blueprint that defines this component."),
+									FSlateIcon(FEditorStyle::GetStyleSetName(), FClassIconFinder::FindIconNameForClass(ComponentTemplate->GetClass())),
 									FUIAction(
-										FExecuteAction::CreateSP(this, &SSCSEditor::OnOpenCodeFile, ClassHeaderPath),
+										FExecuteAction::CreateSP(this, &SSCSEditor::OnEditBlueprint, ComponentTemplate->GetClass()->ClassGeneratedBy),
 										FCanExecuteAction()));
-							}
 
-							MenuBuilder.AddMenuEntry(
-								LOCTEXT("GoToAssetForComponent", "Find Class in Content Browser"),
-								LOCTEXT("GoToAssetForComponent_ToolTip", "Summons the content browser and goes to the class for this component."),
-								FSlateIcon(FEditorStyle::GetStyleSetName(), "SystemWideCommands.FindInContentBrowser"),
-								FUIAction(
-									FExecuteAction::CreateSP(this, &SSCSEditor::OnGoToAssetInBrowser, (UObject*)ComponentTemplate->GetClass()),
-									FCanExecuteAction()));
+								MenuBuilder.AddMenuEntry(
+									LOCTEXT("GoToAssetForComponent", "Find Class in Content Browser"),
+									LOCTEXT("GoToAssetForComponent_ToolTip", "Summons the content browser and goes to the class for this component."),
+									FSlateIcon(FEditorStyle::GetStyleSetName(), "SystemWideCommands.FindInContentBrowser"),
+									FUIAction(
+										FExecuteAction::CreateSP(this, &SSCSEditor::OnGoToAssetInBrowser, ComponentTemplate->GetClass()->ClassGeneratedBy),
+										FCanExecuteAction()));
+							}
+							MenuBuilder.EndSection();			
+						}
+						else
+						{
+							MenuBuilder.BeginSection("ComponentCode", LOCTEXT("ComponentCodeHeading", "C++"));
+							{
+								if (FSourceCodeNavigation::IsCompilerAvailable())
+								{
+									FString ClassHeaderPath;
+									if (FSourceCodeNavigation::FindClassHeaderPath(ComponentTemplate->GetClass(), ClassHeaderPath) && IFileManager::Get().FileSize(*ClassHeaderPath) != INDEX_NONE)
+									{
+										const FString CodeFileName = FPaths::GetCleanFilename(*ClassHeaderPath);
+
+										MenuBuilder.AddMenuEntry(
+											FText::Format(LOCTEXT("GoToCodeForComponent", "Open {0}"), FText::FromString(CodeFileName)),
+											FText::Format(LOCTEXT("GoToCodeForComponent_ToolTip", "Opens the header file for this component ({0}) in a code editing program"), FText::FromString(CodeFileName)),
+											FSlateIcon(),
+											FUIAction(
+												FExecuteAction::CreateSP(this, &SSCSEditor::OnOpenCodeFile, ClassHeaderPath),
+												FCanExecuteAction()));
+									}
+
+									MenuBuilder.AddMenuEntry(
+										LOCTEXT("GoToAssetForComponent", "Find Class in Content Browser"),
+										LOCTEXT("GoToAssetForComponent_ToolTip", "Summons the content browser and goes to the class for this component."),
+										FSlateIcon(FEditorStyle::GetStyleSetName(), "SystemWideCommands.FindInContentBrowser"),
+										FUIAction(
+											FExecuteAction::CreateSP(this, &SSCSEditor::OnGoToAssetInBrowser, (UObject*)ComponentTemplate->GetClass()),
+											FCanExecuteAction()));
+								}
+							}
+							MenuBuilder.EndSection();
 						}
 					}
-					MenuBuilder.EndSection();
 				}
 			}
 		}
+		else
+		{
+			bShowJustPasteOption = true;
+		}
+
+		if (bShowJustPasteOption)
+		{
+			MenuBuilder.BeginSection("ComponentActions", LOCTEXT("ComponentContextMenu", "Component Actions") );
+			{
+				MenuBuilder.AddMenuEntry( FGenericCommands::Get().Paste );
+			}
+			MenuBuilder.EndSection();
+		}
 
 		return MenuBuilder.MakeWidget();
 	}
diff --git a/Engine/Source/Editor/Kismet/Public/SSCSEditor.h b/Engine/Source/Editor/Kismet/Public/SSCSEditor.h
--- a/Engine/Source/Editor/Kismet/Public/SSCSEditor.h
+++ b/Engine/Source/Editor/Kismet/Public/SSCSEditor.h
@@ -479,6 +479,7 @@ public:
 		{}
 
 		SLATE_ARGUMENT(EComponentEditorMode::Type, EditorMode)
+		SLATE_ARGUMENT(TSharedPtr<FExtender>, ActorMenuExtender)
 		SLATE_ATTRIBUTE(class AActor*, ActorContext)
 		SLATE_ATTRIBUTE(class AActor*, PreviewActor)
 		SLATE_ATTRIBUTE(bool, AllowEditing)
@@ -825,6 +826,9 @@ private:
 	/* Root Tree Node*/
 	TSharedPtr<FSCSEditorTreeNode> RootTreeNode;
 
+	/* Root Tree Node*/
+	TSharedPtr<FExtender> ActorMenuExtender;
+
 	/** Flag to enable/disable component editing */
 	bool bEnableComponentEditing;
 
diff --git a/Engine/Source/Editor/LevelEditor/Private/SActorDetails.cpp b/Engine/Source/Editor/LevelEditor/Private/SActorDetails.cpp
--- a/Engine/Source/Editor/LevelEditor/Private/SActorDetails.cpp
+++ b/Engine/Source/Editor/LevelEditor/Private/SActorDetails.cpp
@@ -79,6 +79,7 @@ void SActorDetails::Construct(const FArguments& InArgs, const FName TabIdentifie
 		.EditorMode(EComponentEditorMode::ActorInstance)
 		.ActorContext(this, &SActorDetails::GetActorContext)
 		.OnSelectionUpdated(this, &SActorDetails::OnSCSEditorTreeViewSelectionChanged)
+		.ActorMenuExtender(InArgs._ActorMenuExtender)
 	);
 
 	ChildSlot
diff --git a/Engine/Source/Editor/LevelEditor/Private/SActorDetails.h b/Engine/Source/Editor/LevelEditor/Private/SActorDetails.h
--- a/Engine/Source/Editor/LevelEditor/Private/SActorDetails.h
+++ b/Engine/Source/Editor/LevelEditor/Private/SActorDetails.h
@@ -8,7 +8,8 @@
 class SActorDetails : public SCompoundWidget, public FEditorUndoClient
 {
 public:
-	SLATE_BEGIN_ARGS(SActorDetails){}
+	SLATE_BEGIN_ARGS(SActorDetails) {}
+		SLATE_ARGUMENT( TSharedPtr<FExtender>, ActorMenuExtender )
 	SLATE_END_ARGS()
 
 	void Construct(const FArguments& InArgs, const FName TabIdentifier);
diff --git a/Engine/Source/Editor/LevelEditor/Private/SLevelEditor.cpp b/Engine/Source/Editor/LevelEditor/Private/SLevelEditor.cpp
--- a/Engine/Source/Editor/LevelEditor/Private/SLevelEditor.cpp
+++ b/Engine/Source/Editor/LevelEditor/Private/SLevelEditor.cpp
@@ -501,7 +501,7 @@ TSharedRef<FTabManager> SLevelEditor::GetTabManager() const
 
 
 
-TSharedRef<SDockTab> SLevelEditor::SummonDetailsPanel( FName TabIdentifier )
+TSharedRef<SDockTab> SLevelEditor::SummonDetailsPanel( FName TabIdentifier, TSharedPtr<FExtender> ActorMenuExtender )
 {
 	TSharedPtr<SActorDetails> ActorDetails;
 
@@ -516,6 +516,7 @@ TSharedRef<SDockTab> SLevelEditor::SummonDetailsPanel( FName TabIdentifier )
 			.AddMetaData<FTutorialMetaData>(FTutorialMetaData(TEXT("ActorDetails"), TEXT("LevelEditorSelectionDetails")))
 			[
 				SAssignNew( ActorDetails, SActorDetails, TabIdentifier )
+					.ActorMenuExtender(ActorMenuExtender)
 			]
 		];
 
@@ -563,7 +564,20 @@ TSharedRef<SDockTab> SLevelEditor::SpawnLevelEditorTab( const FSpawnTabArgs& Arg
 	}
 	else if( TabIdentifier == TEXT("LevelEditorSelectionDetails") || TabIdentifier == TEXT("LevelEditorSelectionDetails2") || TabIdentifier == TEXT("LevelEditorSelectionDetails3") || TabIdentifier == TEXT("LevelEditorSelectionDetails4") )
 	{
-		TSharedRef<SDockTab> DetailsPanel = SummonDetailsPanel( TabIdentifier );
+		TWeakPtr<SLevelEditor> WeakLevelEditor = SharedThis(this);
+		TSharedPtr<FExtender> MenuExtender = MakeShareable(new FExtender);
+		MenuExtender->AddMenuExtension(
+			"MainSection", EExtensionHook::Before, GetLevelEditorActions(),
+			FMenuExtensionDelegate::CreateStatic([](FMenuBuilder& MenuBuilder, TWeakPtr<SLevelEditor> InWeakLevelEditor){
+				// Only extend the menu if we have actors selected
+				if (GEditor->GetSelectedActors()->Num())
+				{
+					FLevelEditorContextMenu::FillMenu(MenuBuilder, InWeakLevelEditor, LevelEditorMenuContext::NonViewport, TSharedPtr<FExtender>());
+				}
+			}, WeakLevelEditor)
+		);
+
+		TSharedRef<SDockTab> DetailsPanel = SummonDetailsPanel( TabIdentifier, MenuExtender );
 		GUnrealEd->UpdateFloatingPropertyWindows();
 		return DetailsPanel;
 	}
diff --git a/Engine/Source/Editor/LevelEditor/Private/SLevelEditor.h b/Engine/Source/Editor/LevelEditor/Private/SLevelEditor.h
--- a/Engine/Source/Editor/LevelEditor/Private/SLevelEditor.h
+++ b/Engine/Source/Editor/LevelEditor/Private/SLevelEditor.h
@@ -127,7 +127,7 @@ private:
 	
 	TSharedRef<SDockTab> SpawnLevelEditorTab(const FSpawnTabArgs& Args, FName TabIdentifier, FString InitializationPayload);
 	//TSharedRef<SDockTab> SpawnLevelEditorModeTab(const FSpawnTabArgs& Args, FEdMode* EditorMode);
-	TSharedRef<SDockTab> SummonDetailsPanel( FName Identifier );
+	TSharedRef<SDockTab> SummonDetailsPanel( FName Identifier, TSharedPtr<FExtender> ActorMenuExtender = nullptr );
 
 	/**
 	 * Binds UI commands to actions for the level editor                   
