[INTEGRATE] Change 2423982 by [REDACTED_PERSON][REDACTED_HANDLE] on 2015/01/29 12:25:35

	When adding a instance SceneComponent to an actor that has no root component, make it the root component.

diff --git a/Engine/Source/Editor/Kismet/Private/SSCSEditor.cpp b/Engine/Source/Editor/Kismet/Private/SSCSEditor.cpp
--- a/Engine/Source/Editor/Kismet/Private/SSCSEditor.cpp
+++ b/Engine/Source/Editor/Kismet/Private/SSCSEditor.cpp
@@ -4216,7 +4216,15 @@ UActorComponent* SSCSEditor::AddNewNode(UActorComponent* NewInstanceComponent,
 
 	if (USceneComponent* NewSceneComponent = Cast<USceneComponent>(NewInstanceComponent))
 	{
-		NewSceneComponent->AttachTo(ActorInstance->GetRootComponent());
+		USceneComponent* RootComponent = ActorInstance->GetRootComponent();
+		if (RootComponent)
+		{
+			NewSceneComponent->AttachTo(RootComponent);
+		}
+		else
+		{
+			ActorInstance->SetRootComponent(NewSceneComponent);
+		}
 	}
 
 	// Add to SerializedComponents array so it gets saved
