[INTEGRATE] Change 2424281 by [REDACTED_PERSON][REDACTED_HANDLE] on 2015/01/29 15:21:45

	SCSEditor will no longer select a separator if right clicked.
	Fixed some code that wasn't filtering out separators from the selected nodes.
	Don't clear the selection in the SCSEditor unless we are actually going to remove a node, better handling of cosing the next note to select in the treeview when a delete occurs.

diff --git a/Engine/Source/Editor/Kismet/Private/SSCSEditor.cpp b/Engine/Source/Editor/Kismet/Private/SSCSEditor.cpp
--- a/Engine/Source/Editor/Kismet/Private/SSCSEditor.cpp
+++ b/Engine/Source/Editor/Kismet/Private/SSCSEditor.cpp
@@ -2383,6 +2383,17 @@ UBlueprint* SSCS_RowWidget::GetBlueprint() const
 	return SCSEditor.Pin()->GetBlueprint();
 }
 
+ESelectionMode::Type SSCS_RowWidget::GetSelectionMode() const
+{
+	FSCSEditorTreeNodePtrType NodePtr = GetNode();
+	if (NodePtr->GetNodeType() == FSCSEditorTreeNode::SeparatorNode)
+	{
+		return ESelectionMode::None;
+	}
+	
+	return SMultiColumnTableRow<FSCSEditorTreeNodePtrType>::GetSelectionMode();
+}
+
 bool SSCS_RowWidget::OnNameTextVerifyChanged(const FText& InNewText, FText& OutErrorMessage)
 {
 	FSCSEditorTreeNodePtrType NodePtr = GetNode();
@@ -2762,7 +2773,7 @@ void SSCSEditor::Construct( const FArguments& InArgs )
 		.DefaultLabel(LOCTEXT("Asset", "Asset"))\
 		.FillWidth(3);
 	}
-
+	
 	SCSTreeWidget = SNew(SSCSTreeType)
 		.ToolTipText(LOCTEXT("DropAssetToAddComponent", "Drop asset here to add a component."))
 		.SCSEditor(this)
@@ -3055,11 +3066,10 @@ void SSCSEditor::Tick( const FGeometry& AllottedGeometry, const double InCurrent
 
 FReply SSCSEditor::OnKeyDown( const FGeometry& MyGeometry, const FKeyEvent& InKeyEvent )
 {
-	if ( CommandList->ProcessCommandBindings( InKeyEvent ) )
+	if (CommandList->ProcessCommandBindings(InKeyEvent))
 	{
 		return FReply::Handled();
 	}
-
 	return FReply::Unhandled();
 }
 
@@ -3137,12 +3147,12 @@ void SSCSEditor::OnOpenCodeFile(const FString CodeFileName) const
 	FSourceCodeNavigation::OpenSourceFile( AbsoluteHeaderPath );
 }
 
-TSharedPtr< SWidget > SSCSEditor::CreateContextMenu() // @todo: make the context manual meaningful for the actor and avoid separators
+TSharedPtr< SWidget > SSCSEditor::CreateContextMenu()
 {
 	TArray<FSCSEditorTreeNodePtrType> SelectedNodes = SCSTreeWidget->GetSelectedItems();
 
-	if (SelectedNodes.Num() > 0 || CanPasteNodes())
 	{
+		
 		const bool CloseAfterSelection = true;
 		FMenuBuilder MenuBuilder( CloseAfterSelection, CommandList, ActorMenuExtender );
 
@@ -4066,8 +4076,6 @@ void SSCSEditor::ClearSelection()
 {
 	if ( bUpdatingSelection == false )
 	{
-		bIsActorSelected = false;
-
 		check(SCSTreeWidget.IsValid());
 		SCSTreeWidget->ClearSelection();
 	}
@@ -4618,24 +4626,23 @@ void SSCSEditor::OnDeleteNodes()
 						NewSelection = ChildNodes[ChildIndex];
 					}
 				}
-				else if(Cast<USceneComponent>(Node->GetComponentTemplate()) == nullptr)
+				else
 				{
-					int32 NodeIndex = RootComponentNodes.Find(Node);
-					if(NodeIndex != INDEX_NONE)
+					const TArray<FSCSEditorTreeNodePtrType>& Nodes = GetRootNodes();
+					int32 NodeIndex = Nodes.Find(Node);
+					if (NodeIndex != INDEX_NONE && NodeIndex > 0)
 					{
-						// We should always have a valid scene root component, so there should be at least 2 root nodes in this case
-						check(RootComponentNodes.Num() > 1);
-						NewSelection = RootComponentNodes[NodeIndex > 0 ? NodeIndex - 1 : NodeIndex + 1];
+						while (NodeIndex > 0 && Nodes[--NodeIndex]->GetNodeType() == FSCSEditorTreeNode::SeparatorNode)
+						{
+							--NodeIndex;
+						}
+						NewSelection = Nodes[NodeIndex];
 					}
 					else
 					{
 						NewSelection.Reset();
 					}
 				}
-				else
-				{
-					NewSelection.Reset();
-				}
 			}
 
 			// This will clear the current selection
@@ -4660,17 +4667,17 @@ void SSCSEditor::RemoveComponentNode(FSCSEditorTreeNodePtrType InNodePtr)
 {
 	check(InNodePtr.IsValid());
 
-	// Clear selection if current
-	if(SCSTreeWidget->GetSelectedItems().Contains(InNodePtr))
-	{
-		SCSTreeWidget->ClearSelection();
-	}
-
 	if (EditorMode == EComponentEditorMode::BlueprintSCS)
 	{
 		USCS_Node* SCS_Node = InNodePtr->GetSCSNode();
 		if(SCS_Node != NULL)
 		{
+			// Clear selection if current
+			if (SCSTreeWidget->GetSelectedItems().Contains(InNodePtr))
+			{
+				SCSTreeWidget->ClearSelection();
+			}
+
 			USimpleConstructionScript* SCS = SCS_Node->GetSCS();
 			check(SCS != nullptr);
 
@@ -4695,6 +4702,12 @@ void SSCSEditor::RemoveComponentNode(FSCSEditorTreeNodePtrType InNodePtr)
 		UActorComponent* ComponentInstance = InNodePtr->GetComponentTemplate();
 		if ((ActorInstance != nullptr) && (ComponentInstance != nullptr))
 		{
+			// Clear selection if current
+			if (SCSTreeWidget->GetSelectedItems().Contains(InNodePtr))
+			{
+				SCSTreeWidget->ClearSelection();
+			}
+
 			ComponentInstance->Modify();
 
 			// Handle removal of a node within the scene component hierarchy
@@ -4806,8 +4819,6 @@ void SSCSEditor::RefreshSelectionDetails()
 
 void SSCSEditor::OnTreeSelectionChanged(FSCSEditorTreeNodePtrType, ESelectInfo::Type /*SelectInfo*/)
 {
-	bIsActorSelected = false;
-
 	UpdateSelectionFromNodes(SCSTreeWidget->GetSelectedItems());
 }
 
diff --git a/Engine/Source/Editor/Kismet/Public/SSCSEditor.h b/Engine/Source/Editor/Kismet/Public/SSCSEditor.h
--- a/Engine/Source/Editor/Kismet/Public/SSCSEditor.h
+++ b/Engine/Source/Editor/Kismet/Public/SSCSEditor.h
@@ -132,15 +132,15 @@ public:
 	/** 
 	 * @return Whether or not this object represents a node that can be deleted from the SCS tree.
 	 */
-	bool CanDelete() const { return (IsUserInstanced() || (!IsNative() && !IsInherited())) && !IsDefaultSceneRoot(); }
+	bool CanDelete() const { return NodeType == ComponentNode && ((IsUserInstanced() || (!IsNative() && !IsInherited())) && !IsDefaultSceneRoot()); }
 	/** 
 	 * @return Whether or not this object represents a node that can be renamed from the SCS tree.
 	 */
-	bool CanRename() const { return (IsUserInstanced() || (!IsNative() && !IsInherited())) && !IsDefaultSceneRoot(); }
+	bool CanRename() const { return NodeType == ComponentNode && ((IsUserInstanced() || (!IsNative() && !IsInherited())) && !IsDefaultSceneRoot()); }
 	/** 
 	 * @return Whether or not this object represents a node that can be reparented to other nodes based on its context.
 	 */
-	bool CanReparent() const { return (IsUserInstanced() || (!IsNative() && !IsInherited())) && !IsDefaultSceneRoot() && Cast<USceneComponent>(GetComponentTemplate()) != NULL; }
+	bool CanReparent() const { return NodeType == ComponentNode && ((IsUserInstanced() || (!IsNative() && !IsInherited())) && !IsDefaultSceneRoot() && Cast<USceneComponent>(GetComponentTemplate()) != NULL); }
 	/** 
 	 * @return Whether or not we can edit default properties for the component template represented by this object.
 	 */
@@ -290,6 +290,9 @@ public:
 	/* Get the node used by the row Widget */
 	virtual FSCSEditorTreeNodePtrType GetNode() const { return TreeNodePtr; };
 
+protected:
+	virtual ESelectionMode::Type GetSelectionMode() const override;
+
 private:
 	/** Verifies the name of the component when changing it */
 	bool OnNameTextVerifyChanged(const FText& InNewText, FText& OutErrorMessage);
@@ -835,6 +838,4 @@ private:
 	/** Gate to prevent changing the selection while selection change is being broadcast. */
 	bool bUpdatingSelection;
 
-	/** Flag to track if we have the Actor selected currently */
-	bool bIsActorSelected;
 };
diff --git a/Engine/Source/Editor/LevelEditor/Private/SActorDetails.cpp b/Engine/Source/Editor/LevelEditor/Private/SActorDetails.cpp
--- a/Engine/Source/Editor/LevelEditor/Private/SActorDetails.cpp
+++ b/Engine/Source/Editor/LevelEditor/Private/SActorDetails.cpp
@@ -221,25 +221,30 @@ void SActorDetails::OnSCSEditorTreeViewSelectionChanged(const TArray<FSCSEditorT
 {
 	if (!bSelectionGuard && SelectedNodes.Num() > 0)
 	{
-		const bool bDetailsLocked = DetailsView->IsLocked();
-		//AActor* Actor = bDetailsLocked ? LockedActorSelection.Get() : GetSelectedActorInEditor();
 		AActor* Actor = GetActorContext();
 		if (Actor)
 		{
 			TArray<UObject*> DetailsObjects;
 
-			// Determine if the root actor node is among the selected nodes
+			// Determine if the root actor node is among the selected nodes and Count number of components selected
 			bool bActorNodeSelected = false;
+			int NumSelectedComponentNodes = 0;
 			for (auto& SelectedNode : SelectedNodes)
 			{
-				if (SelectedNode.IsValid() && SelectedNode->GetNodeType() == FSCSEditorTreeNode::RootActorNode)
+				if (SelectedNode.IsValid())
 				{
-					bActorNodeSelected = true;
-					break;
+					if (SelectedNode->GetNodeType() == FSCSEditorTreeNode::RootActorNode)
+					{
+						bActorNodeSelected = true;
+					}
+					else if (SelectedNode->GetNodeType() == FSCSEditorTreeNode::ComponentNode)
+					{
+						++NumSelectedComponentNodes;
+					}
 				}
 			}
 
-			if (bDetailsLocked)
+			if (DetailsView->IsLocked())
 			{
 				// When the details panel is locked, we don't want to touch the editor's component selection
 				// We do want to force the locked panel to update to match the selected components, though, since they are part of the actor selection we're locked on
@@ -279,13 +284,13 @@ void SActorDetails::OnSCSEditorTreeViewSelectionChanged(const TArray<FSCSEditorT
 				USelection* SelectedComponents = GEditor->GetSelectedComponents();
 
 				// Determine if the selected non-root actor nodes differ from the editor component selection
-				bool bComponentSelectionChanged = GEditor->GetSelectedComponentCount() != SelectedNodes.Num() - ( bActorNodeSelected ? 1 : 0 );
+				bool bComponentSelectionChanged = GEditor->GetSelectedComponentCount() != NumSelectedComponentNodes;
 				if (!bComponentSelectionChanged)
 				{
 					// Check to see if any of the selected nodes aren't already selected in the world
 					for (auto& SelectedNode : SelectedNodes)
 					{
-						if (SelectedNode.IsValid() && SelectedNode->GetNodeType() != FSCSEditorTreeNode::RootActorNode)
+						if (SelectedNode.IsValid() && SelectedNode->GetNodeType() == FSCSEditorTreeNode::ComponentNode)
 						{
 							UActorComponent* ComponentInstance = SelectedNode->FindComponentInstanceInActor(Actor);
 							if (ComponentInstance && !SelectedComponents->IsSelected(ComponentInstance))
diff --git a/Engine/Source/Editor/UnrealEd/Private/Kismet2/ComponentEditorUtils.cpp b/Engine/Source/Editor/UnrealEd/Private/Kismet2/ComponentEditorUtils.cpp
--- a/Engine/Source/Editor/UnrealEd/Private/Kismet2/ComponentEditorUtils.cpp
+++ b/Engine/Source/Editor/UnrealEd/Private/Kismet2/ComponentEditorUtils.cpp
@@ -116,8 +116,11 @@ bool FComponentEditorUtils::IsValidVariableNameString(const UActorComponent* InC
 
 bool FComponentEditorUtils::IsComponentNameAvailable(const FString& InString, AActor* ComponentOwner, const UActorComponent* ComponentToIgnore)
 {
-	UObject* ExistingObject = FindObject<UObject>(ComponentOwner, *InString);
-	return (ExistingObject == nullptr || ExistingObject == ComponentToIgnore);
+	UObject* Object = FindObjectFast<UObject>(ComponentOwner, *InString);
+
+	bool bNameIsAvailable = Object == nullptr || Object == ComponentToIgnore;
+
+	return bNameIsAvailable;
 }
 
 FString FComponentEditorUtils::GenerateValidVariableName(TSubclassOf<UActorComponent> ComponentClass, AActor* ComponentOwner)
diff --git a/Engine/Source/Runtime/Slate/Public/Widgets/Views/STableRow.h b/Engine/Source/Runtime/Slate/Public/Widgets/Views/STableRow.h
--- a/Engine/Source/Runtime/Slate/Public/Widgets/Views/STableRow.h
+++ b/Engine/Source/Runtime/Slate/Public/Widgets/Views/STableRow.h
@@ -37,6 +37,10 @@ class SLATE_API ITableRow
 
 		/** Called when the expander arrow for this row is shift+clicked */
 		virtual void Private_OnExpanderArrowShiftClicked() = 0;
+
+	protected:
+		/** Called to query the selection mode for the row */
+		virtual ESelectionMode::Type GetSelectionMode() const = 0;
 };
 
 
@@ -272,7 +276,7 @@ public:
 
 		if ( MouseEvent.GetEffectingButton() == EKeys::LeftMouseButton )
 		{
-			switch( OwnerWidget->Private_GetSelectionMode() )
+			switch( GetSelectionMode() )
 			{
 			case ESelectionMode::Single:
 				{
@@ -377,7 +381,7 @@ public:
 			{
 				if ( bIsUnderMouse )
 				{
-					switch( OwnerWidget->Private_GetSelectionMode() )
+					switch( GetSelectionMode() )
 					{
 					case ESelectionMode::SingleToggle:
 						{
@@ -438,7 +442,7 @@ public:
 			// Handle selection of items when releasing the right mouse button, but only if the user isn't actively
 			// scrolling the view by holding down the right mouse button.
 
-			switch( OwnerWidget->Private_GetSelectionMode() )
+			switch( GetSelectionMode() )
 			{
 			case ESelectionMode::Single:
 			case ESelectionMode::SingleToggle:
@@ -486,7 +490,7 @@ public:
 			const TSharedPtr< ITypedTableView<ItemType> > OwnerWidget = OwnerTablePtr.Pin();
 			const ItemType* MyItem = OwnerWidget->Private_ItemFromWidget( this );
 
-			switch( OwnerWidget->Private_GetSelectionMode() )
+			switch( GetSelectionMode() )
 			{
 				default:
 				case ESelectionMode::None:
@@ -762,7 +766,7 @@ public:
 		else
 		{
 			// Add a slightly lighter background for even rows
-			const bool bAllowSelection = OwnerWidget->Private_GetSelectionMode() != ESelectionMode::None;
+			const bool bAllowSelection = GetSelectionMode() != ESelectionMode::None;
 			if( IndexInList % 2 == 0 )
 			{
 				return ( IsHovered() && bAllowSelection )
@@ -866,6 +870,12 @@ protected:
 			: NonSelectedForeground;
 	}
 
+	virtual ESelectionMode::Type GetSelectionMode() const override
+	{
+		const TSharedPtr< ITypedTableView<ItemType> > OwnerWidget = OwnerTablePtr.Pin();
+		return OwnerWidget->Private_GetSelectionMode();
+	}
+
 protected:
 
 	/** The list that owns this Selectable */
