[INTEGRATE] Change 2424331 by [REDACTED_PERSON][REDACTED_HANDLE] on 2015/01/29 15:49:24

	Change missed.

diff --git a/Engine/Source/Editor/Kismet/Private/SSCSEditor.cpp b/Engine/Source/Editor/Kismet/Private/SSCSEditor.cpp
--- a/Engine/Source/Editor/Kismet/Private/SSCSEditor.cpp
+++ b/Engine/Source/Editor/Kismet/Private/SSCSEditor.cpp
@@ -3151,8 +3151,8 @@ TSharedPtr< SWidget > SSCSEditor::CreateContextMenu()
 {
 	TArray<FSCSEditorTreeNodePtrType> SelectedNodes = SCSTreeWidget->GetSelectedItems();
 
+	if (SelectedNodes.Num() > 0 || CanPasteNodes())
 	{
-		
 		const bool CloseAfterSelection = true;
 		FMenuBuilder MenuBuilder( CloseAfterSelection, CommandList, ActorMenuExtender );
 
