[INTEGRATE] Change 2424333 by [REDACTED_PERSON][REDACTED_HANDLE] on 2015/01/29 15:51:06

	Allow Blueprinted Components to be dropped into the SCSEditor tree from the content browser

diff --git a/Engine/Source/Editor/Kismet/Private/SSCSEditor.cpp b/Engine/Source/Editor/Kismet/Private/SSCSEditor.cpp
--- a/Engine/Source/Editor/Kismet/Private/SSCSEditor.cpp
+++ b/Engine/Source/Editor/Kismet/Private/SSCSEditor.cpp
@@ -117,13 +117,18 @@ FReply SSCSEditorDragDropTree::OnDrop( const FGeometry& MyGeometry, const FDragD
 				for (int32 DroppedAssetIdx = 0; DroppedAssetIdx < NumAssets; ++DroppedAssetIdx)
 				{
 					const FAssetData& AssetData = DroppedAssetData[DroppedAssetIdx];
-				
-					TSubclassOf<UActorComponent>  ComponentClasses = FComponentAssetBrokerage::GetPrimaryComponentForAsset( AssetData.GetClass() );
+					UClass* AssetClass = AssetData.GetClass();
+					UBlueprint* BPClass = Cast<UBlueprint>( AssetData.GetAsset() );
+					TSubclassOf<UActorComponent>  ComponentClasses = FComponentAssetBrokerage::GetPrimaryComponentForAsset( AssetClass );
 					if ( NULL != ComponentClasses )
 					{
 						GWarn->StatusUpdate( DroppedAssetIdx, NumAssets, FText::Format( LOCTEXT("LoadingComponent", "Loading Component {0}"), FText::FromName( AssetData.AssetName ) ) );
 						SCSEditor->AddNewComponent(ComponentClasses, AssetData.GetAsset() );
 					}
+					else if( BPClass && BPClass->GeneratedClass && BPClass->GeneratedClass->IsChildOf( UActorComponent::StaticClass() ) )
+					{
+						SCSEditor->AddNewComponent(BPClass->GeneratedClass, nullptr );
+					}
 				}
 
 				GWarn->EndSlowTask();
