[INTEGRATE] Change 2424593 by [REDACTED_PERSON][REDACTED_HANDLE] on 2015/01/29 17:43:34

	Components Tree improvements:
	- Split FSCSEditorTreeNode into several subclasses to better handle instanced components, separators, and the root node
	- Improve tooltips to show the 'introduced in' class
	- Fix mobility display on instanced trees
	- Fix mobility display on actors that inherit a root component from another BP (with no ancestral native root)
	- Apply the same rich tooltips used for components in the BP editor for instanced components
	- Fix root component detection in instanced trees

diff --git a/Engine/Source/Editor/Kismet/Private/BlueprintDetailsCustomization.cpp b/Engine/Source/Editor/Kismet/Private/BlueprintDetailsCustomization.cpp
--- a/Engine/Source/Editor/Kismet/Private/BlueprintDetailsCustomization.cpp
+++ b/Engine/Source/Editor/Kismet/Private/BlueprintDetailsCustomization.cpp
@@ -4347,6 +4347,13 @@ void FBlueprintComponentDetails::CustomizeDetails(IDetailLayoutBuilder& DetailLa
 	{
 		IDetailCategoryBuilder& VariableCategory = DetailLayout.EditCategory("Variable", LOCTEXT("VariableDetailsCategory", "Variable"), ECategoryPriority::Variable);
 
+		VariableNameEditableTextBox = SNew(SEditableTextBox)
+			.Text(this, &FBlueprintComponentDetails::OnGetVariableText)
+			.OnTextChanged(this, &FBlueprintComponentDetails::OnVariableTextChanged)
+			.OnTextCommitted(this, &FBlueprintComponentDetails::OnVariableTextCommitted)
+			.IsReadOnly(!CachedNodePtr->CanRename())
+			.Font(IDetailLayoutBuilder::GetDetailFont());
+
 		VariableCategory.AddCustomRow(LOCTEXT("BlueprintComponentDetails_VariableNameLabel", "Variable Name"))
 		.NameContent()
 		[
@@ -4356,12 +4363,7 @@ void FBlueprintComponentDetails::CustomizeDetails(IDetailLayoutBuilder& DetailLa
 		]
 		.ValueContent()
 		[
-			SAssignNew(VariableNameEditableTextBox, SEditableTextBox)
-			.Text(this, &FBlueprintComponentDetails::OnGetVariableText)
-			.OnTextChanged( this, &FBlueprintComponentDetails::OnVariableTextChanged)
-			.OnTextCommitted( this, &FBlueprintComponentDetails::OnVariableTextCommitted)
-			.IsReadOnly(CachedNodePtr->IsNative() || CachedNodePtr->IsInherited() || CachedNodePtr->IsDefaultSceneRoot())
-			.Font(IDetailLayoutBuilder::GetDetailFont())
+			VariableNameEditableTextBox.ToSharedRef()
 		];
 
 		VariableCategory.AddCustomRow(LOCTEXT("BlueprintComponentDetails_VariableTooltipLabel", "Tooltip"))
@@ -4649,7 +4651,7 @@ bool FBlueprintComponentDetails::OnVariableCategoryChangeEnabled() const
 {
 	check(CachedNodePtr.IsValid());
 
-	return !CachedNodePtr->IsNative() && !CachedNodePtr->IsInherited();
+	return !CachedNodePtr->CanRename();
 }
 
 FText FBlueprintComponentDetails::OnGetVariableCategoryText() const
@@ -4757,41 +4759,6 @@ void FBlueprintComponentDetails::PopulateVariableCategories()
 	}
 }
 
-bool FBlueprintComponentDetails::IsNodeAttachable() const
-{
-	check(CachedNodePtr.IsValid());
-
-	// See if node is an attachable type
-	bool bCanAttach = false;	
-	if (!CachedNodePtr->IsRootComponent())
-	{
-		// check component is the correct type
-		if (USceneComponent* SceneComponent = Cast<USceneComponent>(CachedNodePtr->GetComponentTemplate()))
-		{
-			check(BlueprintEditorPtr.IsValid());
-			TSharedPtr<SSCSEditor> Editor = BlueprintEditorPtr.Pin()->GetSCSEditor();
-			check( Editor.IsValid() );
-
-			USCS_Node* SCS_Node = CachedNodePtr->GetSCSNode();
-			if (SCS_Node != NULL && Editor->IsNodeInSimpleConstructionScript(SCS_Node))
-			{
-				FSCSEditorTreeNodePtrType ParentFNode = CachedNodePtr->GetParent();
-				
-				// check parent is of the right type
-				if (ParentFNode.IsValid())
-				{
-					if (USceneComponent* ParentComponent = Cast<USceneComponent>(ParentFNode->GetComponentTemplate()))
-					{
-						bCanAttach = ParentComponent->HasAnySockets();
-					}
-				}
-			}
-		}
-	}
-	
-	return bCanAttach;
-}
-
 FText FBlueprintComponentDetails::GetSocketName() const
 {
 	check(CachedNodePtr.IsValid());
diff --git a/Engine/Source/Editor/Kismet/Private/BlueprintDetailsCustomization.h b/Engine/Source/Editor/Kismet/Private/BlueprintDetailsCustomization.h
--- a/Engine/Source/Editor/Kismet/Private/BlueprintDetailsCustomization.h
+++ b/Engine/Source/Editor/Kismet/Private/BlueprintDetailsCustomization.h
@@ -687,9 +687,6 @@ protected:
 
 	/** Build Event Menu for currently selected components */
 	TSharedRef<SWidget> BuildEventsMenuForComponents() const;
-	
-	/** True if the selected node can be attached to sockets */
-	bool IsNodeAttachable() const;
 
 	FText GetSocketName() const;
 	void OnBrowseSocket();
diff --git a/Engine/Source/Editor/Kismet/Private/SSCSEditor.cpp b/Engine/Source/Editor/Kismet/Private/SSCSEditor.cpp
--- a/Engine/Source/Editor/Kismet/Private/SSCSEditor.cpp
+++ b/Engine/Source/Editor/Kismet/Private/SSCSEditor.cpp
@@ -184,245 +184,39 @@ TSharedRef<FSCSRowDragDropOp> FSCSRowDragDropOp::New(FName InVariableName, UStru
 // FSCSEditorTreeNode
 
 FSCSEditorTreeNode::FSCSEditorTreeNode(FSCSEditorTreeNode::ENodeType InNodeType)
-	:NodeType(InNodeType)
-	,bIsInherited(false)
-	,bIsInstanced(false)
-	,bNonTransactionalRename(false)
+	: NodeType(InNodeType)
+	, bNonTransactionalRename(false)
 {
 }
 
-FSCSEditorTreeNode::FSCSEditorTreeNode(USCS_Node* InSCSNode, bool bInIsInherited)
-	:NodeType(FSCSEditorTreeNode::ComponentNode)
-	,bIsInherited(bInIsInherited)
-	,bIsInstanced(false)
-	,SCSNodePtr(InSCSNode)
-	,ComponentTemplatePtr(InSCSNode != NULL ? InSCSNode->ComponentTemplate : NULL)
-	,bNonTransactionalRename(false)
+FName FSCSEditorTreeNode::GetVariableName() const
 {
+	return NAME_None;
 }
 
-FSCSEditorTreeNode::FSCSEditorTreeNode(UActorComponent* InComponentTemplate)
-	:NodeType(FSCSEditorTreeNode::ComponentNode)
-	,bIsInherited(false)
-	,bIsInstanced(false)
-	,SCSNodePtr(NULL)
-	,ComponentTemplatePtr(InComponentTemplate)
-	,bNonTransactionalRename(false)
+FString FSCSEditorTreeNode::GetDisplayString() const
 {
-	check(InComponentTemplate != nullptr);
-	AActor* Owner = InComponentTemplate->GetOwner();
-	if(Owner != nullptr && !Owner->HasAllFlags(RF_ClassDefaultObject))
-	{
-		bIsInstanced = true;
-		bWasInstancedFromNativeClass = false;
-		
-		InstancedComponentName = InComponentTemplate->GetFName();
-		check(InstancedComponentName != NAME_None);	// ...otherwise IsRootActor() can return a false positive.
-
-		ComponentTemplatePtr.Reset();
-		InstancedComponentOwnerPtr = Owner;
-		
-		UClass* OwnerClass = Owner->GetActorClass();
-		if(OwnerClass != nullptr)
-		{
-			AActor* CDO = OwnerClass->GetDefaultObject<AActor>();
-			if(CDO != nullptr)
-			{
-				// Iterate over the Components array and attempt to find a component with a matching name
-				TInlineComponentArray<UActorComponent*> Components;
-				CDO->GetComponents(Components);
-
-				for(auto It = Components.CreateConstIterator(); It && !bWasInstancedFromNativeClass; ++It)
-				{
-					UActorComponent* ComponentTemplate = *It;
-					if(ComponentTemplate->GetFName() == InComponentTemplate->GetFName())
-					{
-						bWasInstancedFromNativeClass = true;
-					}
-				}
-			}
-		}
-	}
+	return TEXT("GetDisplayString not overridden");
 }
 
-UActorComponent* FSCSEditorTreeNode::GetComponentTemplate() const
+FText FSCSEditorTreeNode::GetDisplayName() const
 {
-	if(bIsInstanced && InstancedComponentOwnerPtr.IsValid())
-	{
-		TInlineComponentArray<UActorComponent*> Components;
-		InstancedComponentOwnerPtr.Get()->GetComponents(Components);
-
-		for(auto It = Components.CreateConstIterator(); It; ++It)
-		{
-			UActorComponent* ComponentInstance = *It;
-			if(ComponentInstance->GetFName() == InstancedComponentName)
-			{
-				return ComponentInstance;
-			}
-		}
-	}
-	
-	return ComponentTemplatePtr.Get();
+	return LOCTEXT("GetDisplayNameNotOverridden", "GetDisplayName not overridden");
 }
 
-UActorComponent* FSCSEditorTreeNode::GetEditableComponentTemplate(UBlueprint* ActualEditedBlueprint)
+class USCS_Node* FSCSEditorTreeNode::GetSCSNode() const
 {
-	if (CanEditDefaults())
-	{
-		return GetComponentTemplate();
-	}
-	
-	auto SCSNode = GetSCSNode();
-	if (!IsNative() && IsInherited() && SCSNode && ActualEditedBlueprint)
-	{
-		return GetOverridenComponentTemplate(ActualEditedBlueprint, true);
-	}
-
 	return nullptr;
 }
 
-FName FSCSEditorTreeNode::GetVariableName() const
-{
-	FName VariableName = NAME_None;
-
-	USCS_Node* SCS_Node = GetSCSNode();
-	UActorComponent* ComponentTemplate = GetComponentTemplate();
-
-	if(SCS_Node != NULL)
-	{
-		// Use the same variable name as is obtained by the compiler
-		VariableName = SCS_Node->GetVariableName();
-	}
-	else if(ComponentTemplate != NULL)
-	{
-		// If the owner class is a Blueprint class, see if there's a corresponding object property that contains the component template
-		check(ComponentTemplate->GetOwner());
-		UClass* OwnerClass = ComponentTemplate->GetOwner()->GetActorClass();
-		if(OwnerClass != NULL)
-		{
-			UBlueprint* Blueprint = UBlueprint::GetBlueprintFromClass(OwnerClass);
-			if(Blueprint != NULL && Blueprint->ParentClass != NULL)
-			{
-				for ( TFieldIterator<UProperty> It(Blueprint->ParentClass); It; ++It )
-				{
-					UProperty* Property  = *It;
-					if(UObjectProperty* ObjectProp = Cast<UObjectProperty>(Property))
-					{
-						UObject* CDO = Blueprint->ParentClass->GetDefaultObject();
-						UObject* Object = ObjectProp->GetObjectPropertyValue(ObjectProp->ContainerPtrToValuePtr<void>(CDO));
-
-						if(Object)
-						{
-							if(Object->GetClass() != ComponentTemplate->GetClass())
-							{
-								continue;
-							}
-
-							if(Object->GetFName() == ComponentTemplate->GetFName())
-							{
-								VariableName = ObjectProp->GetFName();
-								break;
-							}
-						}
-					}
-				}
-			}
-		}
-	}
-
-	return VariableName;
-}
-
-FString FSCSEditorTreeNode::GetDisplayString() const
-{
-	FName VariableName = GetVariableName();
-	UActorComponent* ComponentTemplate = GetComponentTemplate();
-
-	// Only display SCS node variable names in the tree if they have not been autogenerated
-	if (VariableName != NAME_None)
-	{
-		return VariableName.ToString();
-	}
-	else if(IsNative() && ComponentTemplate != NULL)
-	{
-		return ComponentTemplate->GetFName().ToString();
-	}
-	else
-	{
-		FString UnnamedString = LOCTEXT("UnnamedToolTip", "Unnamed").ToString();
-		FString NativeString = IsNative() ? LOCTEXT("NativeToolTip", "Native ").ToString() : TEXT("");
-
-		if(ComponentTemplate != NULL)
-		{
-			return FString::Printf(TEXT("[%s %s%s]"), *UnnamedString, *NativeString, *ComponentTemplate->GetClass()->GetName());
-		}
-		else
-		{
-			return FString::Printf(TEXT("[%s %s]"), *UnnamedString, *NativeString);
-		}
-	}
-}
-
-FText FSCSEditorTreeNode::GetDisplayName() const
+UActorComponent* FSCSEditorTreeNode::GetComponentTemplate() const
 {
-	FName VariableName = GetVariableName();
-	if (VariableName != NAME_None)
-	{
-		return FText::FromName(VariableName);
-	}
-	else if (InstancedComponentName != NAME_None)
-	{
-		return FText::FromName(InstancedComponentName);
-	}
-	return FText::GetEmpty();
+	return nullptr;
 }
 
-UActorComponent* FSCSEditorTreeNode::FindComponentInstanceInActor(const AActor* InActor) const
+UActorComponent* FSCSEditorTreeNode::GetEditableComponentTemplate(UBlueprint* ActualEditedBlueprint)
 {
-	USCS_Node* SCS_Node = GetSCSNode();
-	UActorComponent* ComponentTemplate = GetComponentTemplate();
-
-	UActorComponent* ComponentInstance = NULL;
-	if(InActor != NULL)
-	{
-		if(SCS_Node != NULL)
-		{
-			FName VariableName = SCS_Node->GetVariableName();
-			if(VariableName != NAME_None)
-			{
-				UWorld* World = InActor->GetWorld();
-				UObjectPropertyBase* Property = FindField<UObjectPropertyBase>(InActor->GetClass(), VariableName);
-				if(Property != NULL)
-				{
-					// Return the component instance that's stored in the property with the given variable name
-					ComponentInstance = Cast<UActorComponent>(Property->GetObjectPropertyValue_InContainer(InActor));
-				}
-				else if(World != nullptr && World->WorldType == EWorldType::Preview)
-				{
-					// If this is the preview actor, return the cached component instance that's being used for the preview actor prior to recompiling the Blueprint
-					ComponentInstance = SCS_Node->EditorComponentInstance;
-				}
-			}
-		}
-		else if(ComponentTemplate != NULL)
-		{
-			// Look for a native component instance with a name that matches the template name
-			TInlineComponentArray<UActorComponent*> Components;
-			InActor->GetComponents(Components);
-
-			for(auto It = Components.CreateConstIterator(); It; ++It)
-			{
-				UActorComponent* Component = *It;
-				if(Component->GetFName() == ComponentTemplate->GetFName())
-				{
-					ComponentInstance = Component;
-					break;
-				}
-			}
-		}
-	}
-
-	return ComponentInstance;
+	return nullptr;
 }
 
 UBlueprint* FSCSEditorTreeNode::GetBlueprint() const
@@ -457,34 +251,6 @@ FSCSEditorTreeNode::ENodeType FSCSEditorTreeNode::GetNodeType() const
 	return NodeType;
 }
 
-bool FSCSEditorTreeNode::IsRootComponent() const
-{
-	bool bIsRoot = true;
-	USCS_Node* SCS_Node = GetSCSNode();
-	UActorComponent* ComponentTemplate = GetComponentTemplate();
-
-	if(SCS_Node != NULL)
-	{
-		USimpleConstructionScript* SCS = SCS_Node->GetSCS();
-		if(SCS != NULL)
-		{
-			// Evaluate to TRUE if we have an SCS node reference, it is contained in the SCS root set and does not have an external parent
-			bIsRoot = SCS->GetRootNodes().Contains(SCS_Node) && SCS_Node->ParentComponentOrVariableName == NAME_None;
-		}
-	}
-	else if(ComponentTemplate != NULL)
-	{
-		AActor* CDO = ComponentTemplate->GetOwner();
-		if(CDO != NULL)
-		{
-			// Evaluate to TRUE if we have a valid component reference that matches the native root component
-			bIsRoot = (ComponentTemplate == CDO->GetRootComponent());
-		}
-	}
-
-	return bIsRoot;
-}
-
 bool FSCSEditorTreeNode::IsAttachedTo(FSCSEditorTreeNodePtrType InNodePtr) const
 { 
 	FSCSEditorTreeNodePtrType TestParentPtr = ParentNodePtr;
@@ -501,99 +267,6 @@ bool FSCSEditorTreeNode::IsAttachedTo(FSCSEditorTreeNodePtrType InNodePtr) const
 	return false; 
 }
 
-bool FSCSEditorTreeNode::IsDefaultSceneRoot() const
-{
-	USCS_Node* SCS_Node = GetSCSNode();
-	if(SCS_Node != nullptr)
-	{
-		USimpleConstructionScript* SCS = SCS_Node->GetSCS();
-		if(SCS != nullptr)
-		{
-			return SCS_Node == SCS->GetDefaultSceneRootNode();
-		}
-	}
-	else if(IsUserInstanced())
-	{
-		USceneComponent* SceneComponent = Cast<USceneComponent>(GetComponentTemplate());
-		if(SceneComponent != nullptr)
-		{
-			return SceneComponent->GetFName() == FComponentEditorUtils::GetDefaultSceneRootVariableName();
-		}
-	}
-
-	return false;
-}
-
-bool FSCSEditorTreeNode::IsUserInstanced() const
-{
-	if(bIsInstanced && !bWasInstancedFromNativeClass)
-	{
-		UActorComponent* ComponentInstance = GetComponentTemplate();
-		return ComponentInstance != nullptr && ComponentInstance->CreationMethod != EComponentCreationMethod::ConstructionScript;
-	}
-
-	return false;
-}
-
-bool FSCSEditorTreeNode::CanEditDefaults() const
-{
-	bool bCanEdit = false;
-	USCS_Node* SCS_Node = GetSCSNode();
-	UActorComponent* ComponentTemplate = GetComponentTemplate();
-
-	if(!IsNative())
-	{
-		// Evaluate to TRUE for non-native nodes if it represents a valid SCS node and it is not inherited from a parent Blueprint
-		bCanEdit = SCS_Node != NULL && !IsInherited();
-	}
-	else if(bIsInstanced)
-	{
-		// Evaluate to TRUE for all instanced components except for those instanced from the Blueprint-generated class (i.e. during SCS or UCS)
-		bCanEdit = (ComponentTemplate->CreationMethod != EComponentCreationMethod::ConstructionScript);
-	}
-	else if(ComponentTemplate != NULL)
-	{
-		// Evaluate to TRUE for native nodes if it is bound to a member variable and that variable has either EditDefaultsOnly or EditAnywhere flags set
-		check(ComponentTemplate->GetOwner());
-		UClass* OwnerClass = ComponentTemplate->GetOwner()->GetActorClass();
-		if(OwnerClass != NULL)
-		{
-			UBlueprint* Blueprint = UBlueprint::GetBlueprintFromClass(OwnerClass);
-			if(Blueprint != NULL && Blueprint->ParentClass != NULL)
-			{
-				for(TFieldIterator<UProperty> It(Blueprint->ParentClass); It; ++It)
-				{
-					UProperty* Property = *It;
-					if(UObjectProperty* ObjectProp = Cast<UObjectProperty>(Property))
-					{
-						//must be editable
-						if((Property->PropertyFlags & (CPF_Edit)) == 0)
-						{
-							continue;
-						}
-
-						UObject* ParentCDO = Blueprint->ParentClass->GetDefaultObject();
-
-						if(!ComponentTemplate->GetClass()->IsChildOf(ObjectProp->PropertyClass))
-						{
-							continue;
-						}
-
-						UObject* Object = ObjectProp->GetObjectPropertyValue(ObjectProp->ContainerPtrToValuePtr<void>(ParentCDO));
-						bCanEdit = Object != NULL && Object->GetFName() == ComponentTemplate->GetFName();
-
-						if (bCanEdit)
-						{
-							break;
-						}
-					}
-				}
-			}
-		}
-	}
-
-	return bCanEdit;
-}
 
 FSCSEditorTreeNodePtrType FSCSEditorTreeNode::FindClosestParent(TArray<FSCSEditorTreeNodePtrType> InNodes)
 {
@@ -678,7 +351,7 @@ void FSCSEditorTreeNode::AddChild(FSCSEditorTreeNodePtrType InChildNodePtr)
 			}
 		}
 	}
-	else if(bIsInstanced)
+	else if (IsInstanced())
 	{
 		USceneComponent* ChildInstance = Cast<USceneComponent>(InChildNodePtr->GetComponentTemplate());
 		check(ChildInstance != nullptr);
@@ -709,13 +382,14 @@ FSCSEditorTreeNodePtrType FSCSEditorTreeNode::AddChild(USCS_Node* InSCSNode, boo
 	if(!ChildNodePtr.IsValid())
 	{
 		// Add a child node to the SCS editor tree
-		AddChild(ChildNodePtr = MakeShareable(new FSCSEditorTreeNode(InSCSNode, bInIsInherited)));
+		ChildNodePtr = MakeShareable(new FSCSEditorTreeNodeComponent(InSCSNode, bInIsInherited));
+		AddChild(ChildNodePtr);
 	}
 
 	return ChildNodePtr;
 }
 
-FSCSEditorTreeNodePtrType FSCSEditorTreeNode::AddChild(UActorComponent* InComponentTemplate)
+FSCSEditorTreeNodePtrType FSCSEditorTreeNode::AddChildFromComponent(UActorComponent* InComponentTemplate)
 {
 	// Ensure that the given component template is valid
 	check(InComponentTemplate != NULL);
@@ -725,14 +399,66 @@ FSCSEditorTreeNodePtrType FSCSEditorTreeNode::AddChild(UActorComponent* InCompon
 	if(!ChildNodePtr.IsValid())
 	{
 		// Add a child node to the SCS editor tree
-		AddChild(ChildNodePtr = MakeShareable(new FSCSEditorTreeNode(InComponentTemplate)));
+		ChildNodePtr = FactoryNodeFromComponent(InComponentTemplate);
+		AddChild(ChildNodePtr);
 	}
 
 	return ChildNodePtr;
 }
 
-FSCSEditorTreeNodePtrType FSCSEditorTreeNode::FindChild(const USCS_Node* InSCSNode, bool bRecursiveSearch, uint32* OutDepth) const
-{
+// Tries to find a SCS node that was likely responsible for creating the specified instance component.  Note: This is not always possible to do!
+USCS_Node* FSCSEditorTreeNode::FindSCSNodeForInstance(UActorComponent* InstanceComponent, UClass* ClassToSearch)
+{ 
+	if ((ClassToSearch != nullptr) && (InstanceComponent->CreationMethod == EComponentCreationMethod::ConstructionScript))
+	{
+		for (UClass* TestClass = ClassToSearch; TestClass->ClassGeneratedBy != nullptr; TestClass = TestClass->GetSuperClass())
+		{
+			if (UBlueprint* TestBP = Cast<UBlueprint>(TestClass->ClassGeneratedBy))
+			{
+				if (TestBP->SimpleConstructionScript != nullptr)
+				{
+					if (USCS_Node* Result = TestBP->SimpleConstructionScript->FindSCSNode(InstanceComponent->GetFName()))
+					{
+						return Result;
+					}
+				}
+			}
+		}
+	}
+
+	return nullptr;
+}
+
+FSCSEditorTreeNodePtrType FSCSEditorTreeNode::FactoryNodeFromComponent(UActorComponent* InComponent)
+{
+	check(InComponent);
+
+	bool bComponentIsInAnInstance = false;
+
+	AActor* Owner = InComponent->GetOwner();
+	if ((Owner != nullptr) && !Owner->HasAllFlags(RF_ClassDefaultObject))
+	{
+		bComponentIsInAnInstance = true;
+	}
+
+	if (bComponentIsInAnInstance)
+	{
+		if (InComponent->CreationMethod == EComponentCreationMethod::Instance)
+		{
+			return MakeShareable(new FSCSEditorTreeNodeInstanceAddedComponent(Owner, InComponent->GetFName()));
+		}
+		else
+		{
+			return MakeShareable(new FSCSEditorTreeNodeInstancedInheritedComponent(Owner, InComponent->GetFName()));
+		}
+	}
+
+	// Not an instanced component, either an SCS node or a native component in BP edit mode
+	return MakeShareable(new FSCSEditorTreeNodeComponent(InComponent));
+}
+
+FSCSEditorTreeNodePtrType FSCSEditorTreeNode::FindChild(const USCS_Node* InSCSNode, bool bRecursiveSearch, uint32* OutDepth) const
+{
 	FSCSEditorTreeNodePtrType Result;
 
 	// Ensure that the given SCS node is valid
@@ -781,130 +507,612 @@ FSCSEditorTreeNodePtrType FSCSEditorTreeNode::FindChild(const UActorComponent* I
 		}
 	}
 
-	if(OutDepth && Result.IsValid())
+	if(OutDepth && Result.IsValid())
+	{
+		*OutDepth += 1;
+	}
+
+	return Result;
+}
+
+FSCSEditorTreeNodePtrType FSCSEditorTreeNode::FindChild(const FName& InVariableOrInstanceName, bool bRecursiveSearch, uint32* OutDepth) const
+{
+	FSCSEditorTreeNodePtrType Result;
+
+	// Ensure that the given name is valid
+	if(InVariableOrInstanceName != NAME_None)
+	{
+		// Look for a match in our set of child nodes
+		for(int32 ChildIndex = 0; ChildIndex < Children.Num() && !Result.IsValid(); ++ChildIndex)
+		{
+			FName ItemName = Children[ChildIndex]->GetVariableName();
+			if(ItemName == NAME_None)
+			{
+				UActorComponent* ComponentTemplateOrInstance = Children[ChildIndex]->GetComponentTemplate();
+				check(ComponentTemplateOrInstance != nullptr);
+				ItemName = ComponentTemplateOrInstance->GetFName();
+			}
+
+			if(InVariableOrInstanceName == ItemName)
+			{
+				Result = Children[ChildIndex];
+			}
+			else if(bRecursiveSearch)
+			{
+				Result = Children[ChildIndex]->FindChild(InVariableOrInstanceName, true, OutDepth);
+			}
+		}
+	}
+
+	if(OutDepth && Result.IsValid())
+	{
+		*OutDepth += 1;
+	}
+
+	return Result;
+}
+
+void FSCSEditorTreeNode::RemoveChild(FSCSEditorTreeNodePtrType InChildNodePtr)
+{
+	// Remove the given node as a child and reset its parent link
+	Children.Remove(InChildNodePtr);
+	InChildNodePtr->ParentNodePtr.Reset();
+	InChildNodePtr->RemoveMeAsChild();
+}
+
+void FSCSEditorTreeNode::OnRequestRename(bool bTransactional)
+{
+	bNonTransactionalRename = !bTransactional;
+	RenameRequestedDelegate.ExecuteIfBound();
+}
+
+void FSCSEditorTreeNode::OnCompleteRename(const FText& InNewName)
+{
+}
+
+//////////////////////////////////////////////////////////////////////////
+// FSCSEditorTreeNodeComponentBase
+
+FName FSCSEditorTreeNodeComponentBase::GetVariableName() const
+{
+	FName VariableName = NAME_None;
+
+	USCS_Node* SCS_Node = GetSCSNode();
+	UActorComponent* ComponentTemplate = GetComponentTemplate();
+
+	if (IsInstanced() && (SCS_Node == nullptr) && (ComponentTemplate != nullptr))
+	{
+		if (ComponentTemplate->GetOwner())
+		{
+			SCS_Node = FindSCSNodeForInstance(ComponentTemplate, ComponentTemplate->GetOwner()->GetClass());
+		}
+	}
+
+	if (SCS_Node != NULL)
+	{
+		// Use the same variable name as is obtained by the compiler
+		VariableName = SCS_Node->GetVariableName();
+	}
+	else if (ComponentTemplate != NULL)
+	{
+		// If the owner class is a Blueprint class, see if there's a corresponding object property that contains the component template
+		check(ComponentTemplate->GetOwner());
+		UClass* OwnerClass = ComponentTemplate->GetOwner()->GetActorClass();
+		if (OwnerClass != NULL)
+		{
+			UBlueprint* Blueprint = UBlueprint::GetBlueprintFromClass(OwnerClass);
+			if (Blueprint != NULL)
+			{
+				UClass* StartingClass = IsInstanced() ? Blueprint->GeneratedClass : Blueprint->ParentClass;
+				if (StartingClass != nullptr)
+				{
+					if (UObjectProperty* ObjectProp = FindField<UObjectProperty>(StartingClass, ComponentTemplate->GetFName()))
+					{
+						UObject* CDO = StartingClass->GetDefaultObject();
+						UObject* Object = ObjectProp->GetObjectPropertyValue(ObjectProp->ContainerPtrToValuePtr<void>(CDO));
+
+						if (Object)
+						{
+							if (Object->GetClass() == ComponentTemplate->GetClass())
+							{
+								VariableName = ComponentTemplate->GetFName();
+							}
+						}
+					}
+				}
+			}
+		}
+	}
+
+	return VariableName;
+}
+
+FString FSCSEditorTreeNodeComponentBase::GetDisplayString() const
+{
+	FName VariableName = GetVariableName();
+	UActorComponent* ComponentTemplate = GetComponentTemplate();
+
+	// Only display SCS node variable names in the tree if they have not been autogenerated
+	if (VariableName != NAME_None)
+	{
+		return VariableName.ToString();
+	}
+	else if (IsNative() && (ComponentTemplate != NULL))
+	{
+		return ComponentTemplate->GetFName().ToString();
+	}
+	else
+	{
+		FString UnnamedString = LOCTEXT("UnnamedToolTip", "Unnamed").ToString();
+		FString NativeString = IsNative() ? LOCTEXT("NativeToolTip", "Native ").ToString() : TEXT("");
+
+		if (ComponentTemplate != NULL)
+		{
+			return FString::Printf(TEXT("[%s %s%s]"), *UnnamedString, *NativeString, *ComponentTemplate->GetClass()->GetName());
+		}
+		else
+		{
+			return FString::Printf(TEXT("[%s %s]"), *UnnamedString, *NativeString);
+		}
+	}
+}
+
+//////////////////////////////////////////////////////////////////////////
+// FSCSEditorTreeNodeInstancedInheritedComponent
+
+FSCSEditorTreeNodeInstancedInheritedComponent::FSCSEditorTreeNodeInstancedInheritedComponent(AActor* Owner, FName InComponentName)
+{
+	InstancedComponentName = InComponentName;
+	check(InstancedComponentName != NAME_None);	// ...otherwise IsRootActor() can return a false positive.
+
+	InstancedComponentOwnerPtr = Owner;
+}
+
+bool FSCSEditorTreeNodeInstancedInheritedComponent::IsNative() const
+{
+	if (UActorComponent* Template = GetComponentTemplate())
+	{
+		return Template->CreationMethod == EComponentCreationMethod::Native;
+	}
+	else
+	{
+		return false;
+	}
+}
+
+bool FSCSEditorTreeNodeInstancedInheritedComponent::IsRootComponent() const
+{
+	UActorComponent* ComponentTemplate = GetComponentTemplate();
+
+	if (AActor* OwnerActor = InstancedComponentOwnerPtr.Get())
+	{
+		if (OwnerActor->GetRootComponent() == ComponentTemplate)
+		{
+			return true;
+		}
+	}
+
+	return false;
+}
+
+bool FSCSEditorTreeNodeInstancedInheritedComponent::IsInheritedSCS() const
+{
+	return false;
+}
+
+bool FSCSEditorTreeNodeInstancedInheritedComponent::IsDefaultSceneRoot() const
+{
+	return false;
+}
+
+FText FSCSEditorTreeNodeInstancedInheritedComponent::GetDisplayName() const
+{
+	FName VariableName = GetVariableName();
+	if (VariableName != NAME_None)
+	{
+		return FText::FromName(VariableName);
+	}
+
+	return FText::GetEmpty();
+}
+
+UActorComponent* FSCSEditorTreeNodeInstancedInheritedComponent::GetComponentTemplate() const
+{
+	if (InstancedComponentOwnerPtr.IsValid())
+	{
+		TInlineComponentArray<UActorComponent*> Components;
+		InstancedComponentOwnerPtr.Get()->GetComponents(Components);
+
+		for (auto It = Components.CreateConstIterator(); It; ++It)
+		{
+			UActorComponent* ComponentInstance = *It;
+			if (ComponentInstance->GetFName() == InstancedComponentName)
+			{
+				return ComponentInstance;
+			}
+		}
+	}
+
+	return nullptr;
+}
+
+UActorComponent* FSCSEditorTreeNodeInstancedInheritedComponent::GetEditableComponentTemplate(UBlueprint* ActualEditedBlueprint)
+{
+	if (CanEditDefaults())
+	{
+		return GetComponentTemplate();
+	}
+
+	return nullptr;
+}
+
+//////////////////////////////////////////////////////////////////////////
+// FSCSEditorTreeNodeInstanceAddedComponent
+
+FSCSEditorTreeNodeInstanceAddedComponent::FSCSEditorTreeNodeInstanceAddedComponent(AActor* Owner, FName InComponentName)
+{
+	InstancedComponentName = InComponentName;
+	check(InstancedComponentName != NAME_None);	// ...otherwise IsRootActor() can return a false positive.
+
+	InstancedComponentOwnerPtr = Owner;
+}
+
+bool FSCSEditorTreeNodeInstanceAddedComponent::IsRootComponent() const
+{
+	bool bIsRoot = true;
+	UActorComponent* ComponentTemplate = GetComponentTemplate();
+
+	if (ComponentTemplate != NULL)
+	{
+		AActor* CDO = ComponentTemplate->GetOwner();
+		if (CDO != NULL)
+		{
+			// Evaluate to TRUE if we have a valid component reference that matches the native root component
+			bIsRoot = (ComponentTemplate == CDO->GetRootComponent());
+		}
+	}
+
+	return bIsRoot;
+}
+
+bool FSCSEditorTreeNodeInstanceAddedComponent::IsDefaultSceneRoot() const
+{
+	if (USceneComponent* SceneComponent = Cast<USceneComponent>(GetComponentTemplate()))
+	{
+		return SceneComponent->GetFName() == FComponentEditorUtils::GetDefaultSceneRootVariableName();
+	}
+
+	return false;
+}
+
+FString FSCSEditorTreeNodeInstanceAddedComponent::GetDisplayString() const
+{
+	return InstancedComponentName.ToString();
+}
+
+FText FSCSEditorTreeNodeInstanceAddedComponent::GetDisplayName() const
+{
+	return FText::FromName(InstancedComponentName);
+}
+
+UActorComponent* FSCSEditorTreeNodeInstanceAddedComponent::GetComponentTemplate() const
+{
+	if (InstancedComponentOwnerPtr.IsValid())
+	{
+		TInlineComponentArray<UActorComponent*> Components;
+		InstancedComponentOwnerPtr.Get()->GetComponents(Components);
+
+		for (auto It = Components.CreateConstIterator(); It; ++It)
+		{
+			UActorComponent* ComponentInstance = *It;
+			if (ComponentInstance->GetFName() == InstancedComponentName)
+			{
+				return ComponentInstance;
+			}
+		}
+	}
+
+	return nullptr;
+}
+
+UActorComponent* FSCSEditorTreeNodeInstanceAddedComponent::GetEditableComponentTemplate(UBlueprint* ActualEditedBlueprint)
+{
+	return GetComponentTemplate();
+}
+
+void FSCSEditorTreeNodeInstanceAddedComponent::RemoveMeAsChild()
+{
+	USceneComponent* ChildInstance = Cast<USceneComponent>(GetComponentTemplate());
+	check(ChildInstance != nullptr);
+
+	// Handle detachment at the instance level
+	ChildInstance->DetachFromParent(true);
+}
+
+void FSCSEditorTreeNodeInstanceAddedComponent::OnCompleteRename(const FText& InNewName)
+{
+	FScopedTransaction* TransactionContext = NULL;
+	if (!GetAndClearNonTransactionalRenameFlag())
+	{
+		TransactionContext = new FScopedTransaction(LOCTEXT("RenameComponentVariable", "Rename Component Variable"));
+	}
+
+	UActorComponent* ComponentInstance = GetComponentTemplate();
+	check(ComponentInstance != nullptr);
+
+	ERenameFlags RenameFlags = REN_DontCreateRedirectors;
+	if (!TransactionContext)
+	{
+		RenameFlags |= REN_NonTransactional;
+	}
+
+	ComponentInstance->Rename(*InNewName.ToString(), nullptr, RenameFlags);
+	InstancedComponentName = *InNewName.ToString();
+
+	if (TransactionContext)
+	{
+		delete TransactionContext;
+	}
+}
+
+
+//////////////////////////////////////////////////////////////////////////
+// FSCSEditorTreeNodeComponent
+
+FSCSEditorTreeNodeComponent::FSCSEditorTreeNodeComponent(USCS_Node* InSCSNode, bool bInIsInheritedSCS)
+	: bIsInheritedSCS(bInIsInheritedSCS)
+	, SCSNodePtr(InSCSNode)
+	, ComponentTemplatePtr((InSCSNode != nullptr) ? InSCSNode->ComponentTemplate : nullptr)
+{
+}
+
+FSCSEditorTreeNodeComponent::FSCSEditorTreeNodeComponent(UActorComponent* InComponentTemplate)
+	: bIsInheritedSCS(false)
+	, SCSNodePtr(nullptr)
+	, ComponentTemplatePtr(InComponentTemplate)
+{
+	check(InComponentTemplate != nullptr);
+
+	AActor* Owner = InComponentTemplate->GetOwner();
+	if (Owner != nullptr)
+	{
+		ensureMsg(Owner->HasAllFlags(RF_ClassDefaultObject), TEXT("Use a different node class for instanced components"));
+	}
+}
+
+bool FSCSEditorTreeNodeComponent::IsNative() const
+{
+	return GetSCSNode() == NULL && GetComponentTemplate() != NULL;
+}
+
+bool FSCSEditorTreeNodeComponent::IsRootComponent() const
+{
+	bool bIsRoot = true;
+	USCS_Node* SCS_Node = GetSCSNode();
+	UActorComponent* ComponentTemplate = GetComponentTemplate();
+
+	if (SCS_Node != NULL)
+	{
+		USimpleConstructionScript* SCS = SCS_Node->GetSCS();
+		if (SCS != NULL)
+		{
+			// Evaluate to TRUE if we have an SCS node reference, it is contained in the SCS root set and does not have an external parent
+			bIsRoot = SCS->GetRootNodes().Contains(SCS_Node) && SCS_Node->ParentComponentOrVariableName == NAME_None;
+		}
+	}
+	else if (ComponentTemplate != NULL)
+	{
+		AActor* CDO = ComponentTemplate->GetOwner();
+		if (CDO != NULL)
+		{
+			// Evaluate to TRUE if we have a valid component reference that matches the native root component
+			bIsRoot = (ComponentTemplate == CDO->GetRootComponent());
+		}
+	}
+
+	return bIsRoot;
+}
+
+bool FSCSEditorTreeNodeComponent::IsInheritedSCS() const
+{
+	return bIsInheritedSCS;
+}
+
+bool FSCSEditorTreeNodeComponent::IsDefaultSceneRoot() const
+{
+	if (USCS_Node* SCS_Node = GetSCSNode())
+	{
+		USimpleConstructionScript* SCS = SCS_Node->GetSCS();
+		if (SCS != nullptr)
+		{
+			return SCS_Node == SCS->GetDefaultSceneRootNode();
+		}
+	}
+
+	return false;
+}
+
+bool FSCSEditorTreeNodeComponent::CanEditDefaults() const
+{
+	bool bCanEdit = false;
+
+	if (!IsNative())
+	{
+		USCS_Node* SCS_Node = GetSCSNode();
+		bCanEdit = (SCS_Node != NULL);
+	}
+	else if (UActorComponent* ComponentTemplate = GetComponentTemplate())
+	{
+		// Evaluate to TRUE for native nodes if it is bound to a member variable and that variable has either EditDefaultsOnly or EditAnywhere flags set
+		check(ComponentTemplate->GetOwner());
+		UClass* OwnerClass = ComponentTemplate->GetOwner()->GetActorClass();
+		if (OwnerClass != NULL)
+		{
+			UBlueprint* Blueprint = UBlueprint::GetBlueprintFromClass(OwnerClass);
+			if (Blueprint != NULL && Blueprint->ParentClass != NULL)
+			{
+				for (TFieldIterator<UProperty> It(Blueprint->ParentClass); It; ++It)
+				{
+					UProperty* Property = *It;
+					if (UObjectProperty* ObjectProp = Cast<UObjectProperty>(Property))
+					{
+						//must be editable
+						if ((Property->PropertyFlags & (CPF_Edit)) == 0)
+						{
+							continue;
+						}
+
+						UObject* ParentCDO = Blueprint->ParentClass->GetDefaultObject();
+
+						if (!ComponentTemplate->GetClass()->IsChildOf(ObjectProp->PropertyClass))
+						{
+							continue;
+						}
+
+						UObject* Object = ObjectProp->GetObjectPropertyValue(ObjectProp->ContainerPtrToValuePtr<void>(ParentCDO));
+						bCanEdit = Object != NULL && Object->GetFName() == ComponentTemplate->GetFName();
+
+						if (bCanEdit)
+						{
+							break;
+						}
+					}
+				}
+			}
+		}
+	}
+
+	return bCanEdit;
+}
+
+FText FSCSEditorTreeNodeComponent::GetDisplayName() const
+{
+	FName VariableName = GetVariableName();
+	if (VariableName != NAME_None)
 	{
-		*OutDepth += 1;
+		return FText::FromName(VariableName);
 	}
+	return FText::GetEmpty();
+}
 
-	return Result;
+class USCS_Node* FSCSEditorTreeNodeComponent::GetSCSNode() const
+{
+	return SCSNodePtr.Get();
 }
 
-FSCSEditorTreeNodePtrType FSCSEditorTreeNode::FindChild(const FName& InVariableOrInstanceName, bool bRecursiveSearch, uint32* OutDepth) const
+UActorComponent* FSCSEditorTreeNodeComponent::GetComponentTemplate() const
 {
-	FSCSEditorTreeNodePtrType Result;
+	return ComponentTemplatePtr.Get();
+}
 
-	// Ensure that the given name is valid
-	if(InVariableOrInstanceName != NAME_None)
+UActorComponent* FSCSEditorTreeNodeComponent::GetEditableComponentTemplate(UBlueprint* ActualEditedBlueprint)
+{
+	if (CanEditDefaults())
 	{
-		// Look for a match in our set of child nodes
-		for(int32 ChildIndex = 0; ChildIndex < Children.Num() && !Result.IsValid(); ++ChildIndex)
+		if (!IsNative() && IsInheritedSCS())
 		{
-			FName ItemName = Children[ChildIndex]->GetVariableName();
-			if(ItemName == NAME_None)
-			{
-				UActorComponent* ComponentTemplateOrInstance = Children[ChildIndex]->GetComponentTemplate();
-				check(ComponentTemplateOrInstance != nullptr);
-				ItemName = ComponentTemplateOrInstance->GetFName();
-			}
-
-			if(InVariableOrInstanceName == ItemName)
+			if (ActualEditedBlueprint != nullptr)
 			{
-				Result = Children[ChildIndex];
+				return INTERNAL_GetOverridenComponentTemplate(ActualEditedBlueprint, true);
 			}
-			else if(bRecursiveSearch)
+			else
 			{
-				Result = Children[ChildIndex]->FindChild(InVariableOrInstanceName, true, OutDepth);
+				return nullptr;
 			}
 		}
-	}
 
-	if(OutDepth && Result.IsValid())
-	{
-		*OutDepth += 1;
+		return GetComponentTemplate();
 	}
 
-	return Result;
+	return nullptr;
 }
 
-void FSCSEditorTreeNode::RemoveChild(FSCSEditorTreeNodePtrType InChildNodePtr)
+UActorComponent* FSCSEditorTreeNode::FindComponentInstanceInActor(const AActor* InActor) const
 {
-	// Remove the given node as a child and reset its parent link
-	Children.Remove(InChildNodePtr);
-	InChildNodePtr->ParentNodePtr.Reset();
+	USCS_Node* SCS_Node = GetSCSNode();
+	UActorComponent* ComponentTemplate = GetComponentTemplate();
 
-	// Remove the SCS node from the SCS tree, if present
-	USCS_Node* SCS_ChildNode = InChildNodePtr->GetSCSNode();
-	if(SCS_ChildNode != NULL)
+	UActorComponent* ComponentInstance = NULL;
+	if (InActor != NULL)
 	{
-		USimpleConstructionScript* SCS = SCS_ChildNode->GetSCS();
-		if(SCS != NULL)
+		if (SCS_Node != NULL)
 		{
-			SCS->RemoveNode(SCS_ChildNode);
+			FName VariableName = SCS_Node->GetVariableName();
+			if (VariableName != NAME_None)
+			{
+				UWorld* World = InActor->GetWorld();
+				UObjectPropertyBase* Property = FindField<UObjectPropertyBase>(InActor->GetClass(), VariableName);
+				if (Property != NULL)
+				{
+					// Return the component instance that's stored in the property with the given variable name
+					ComponentInstance = Cast<UActorComponent>(Property->GetObjectPropertyValue_InContainer(InActor));
+				}
+				else if (World != nullptr && World->WorldType == EWorldType::Preview)
+				{
+					// If this is the preview actor, return the cached component instance that's being used for the preview actor prior to recompiling the Blueprint
+					ComponentInstance = SCS_Node->EditorComponentInstance;
+				}
+			}
 		}
-	}
-	else if(bIsInstanced)
-	{
-		USceneComponent* ChildInstance = Cast<USceneComponent>(InChildNodePtr->GetComponentTemplate());
-		check(ChildInstance != nullptr);
+		else if (ComponentTemplate != NULL)
+		{
+			// Look for a native component instance with a name that matches the template name
+			TInlineComponentArray<UActorComponent*> Components;
+			InActor->GetComponents(Components);
 
-		// Handle detachment at the instance level
-		ChildInstance->DetachFromParent(true);
+			for (auto It = Components.CreateConstIterator(); It; ++It)
+			{
+				UActorComponent* Component = *It;
+				if (Component->GetFName() == ComponentTemplate->GetFName())
+				{
+					ComponentInstance = Component;
+					break;
+				}
+			}
+		}
 	}
-}
 
-void FSCSEditorTreeNode::OnRequestRename(bool bTransactional)
-{
-	bNonTransactionalRename = !bTransactional;
-	RenameRequestedDelegate.ExecuteIfBound();
+	return ComponentInstance;
 }
 
-void FSCSEditorTreeNode::OnCompleteRename(const FText& InNewName)
+void FSCSEditorTreeNodeComponent::OnCompleteRename(const FText& InNewName)
 {
 	FScopedTransaction* TransactionContext = NULL;
-	if(bNonTransactionalRename)
-	{
-		// Reset for next time through - if the next rename operation is not explicitly initiated by OnRequestRename(), then the rename must always be transactional.
-		bNonTransactionalRename = false;
-	}
-	else
+	if (!GetAndClearNonTransactionalRenameFlag())
 	{
 		TransactionContext = new FScopedTransaction(LOCTEXT("RenameComponentVariable", "Rename Component Variable"));
 	}
 
-	if(bIsInstanced)
-	{
-		UActorComponent* ComponentInstance = GetComponentTemplate();
-		check(ComponentInstance != nullptr);
-
-		ERenameFlags RenameFlags = REN_DontCreateRedirectors;
-		if(!TransactionContext)
-		{
-			RenameFlags |= REN_NonTransactional;
-		}
+	FBlueprintEditorUtils::RenameComponentMemberVariable(GetBlueprint(), GetSCSNode(), FName(*InNewName.ToString()));
 
-		ComponentInstance->Rename(*InNewName.ToString(), nullptr, RenameFlags);
-		InstancedComponentName = *InNewName.ToString();
-	}
-	else
+	if (TransactionContext)
 	{
-		FBlueprintEditorUtils::RenameComponentMemberVariable(GetBlueprint(), GetSCSNode(), FName( *InNewName.ToString() ));
+		delete TransactionContext;
 	}
+}
 
-	if(TransactionContext)
+void FSCSEditorTreeNodeComponent::RemoveMeAsChild()
+{
+	// Remove the SCS node from the SCS tree, if present
+	if (USCS_Node* SCS_ChildNode = GetSCSNode())
 	{
-		delete TransactionContext;
+		USimpleConstructionScript* SCS = SCS_ChildNode->GetSCS();
+		if (SCS != NULL)
+		{
+			SCS->RemoveNode(SCS_ChildNode);
+		}
 	}
 }
 
-UActorComponent* FSCSEditorTreeNode::GetOverridenComponentTemplate(UBlueprint* Blueprint, bool bCreateIfNecessary) const
+UActorComponent* FSCSEditorTreeNodeComponent::INTERNAL_GetOverridenComponentTemplate(UBlueprint* Blueprint, bool bCreateIfNecessary) const
 {
 	UActorComponent* OverridenComponent = NULL;
 
 	FComponentKey Key(GetSCSNode());
 
-	const bool BlueprintCanOverrideComponentFormKey = Key.IsValid() 
-		&& Blueprint 
-		&& Blueprint->ParentClass 
+	const bool BlueprintCanOverrideComponentFormKey = Key.IsValid()
+		&& Blueprint
+		&& Blueprint->ParentClass
 		&& Blueprint->ParentClass->IsChildOf(Key.OwnerClass);
 
 	if (BlueprintCanOverrideComponentFormKey)
@@ -1129,22 +1337,26 @@ TSharedRef<SWidget> SSCS_RowWidget::GenerateWidgetForColumn( const FName& Column
 	}
 	else if (ColumnName == SCS_ColumnName_Mobility)
 	{
-		TWeakObjectPtr<USCS_Node> SCSNode(NodePtr->GetSCSNode());
-
-		TSharedPtr<SToolTip> MobilityTooltip;
-		SAssignNew(MobilityTooltip, SToolTip)
-			.Text(this, &SSCS_RowWidget::GetMobilityToolTipText, SCSNode);
+		if (NodePtr->GetNodeType() == FSCSEditorTreeNode::ComponentNode)
+		{
+			TSharedPtr<SToolTip> MobilityTooltip = SNew(SToolTip)
+				.Text(this, &SSCS_RowWidget::GetMobilityToolTipText);
 
-		return	SNew(SHorizontalBox)
+			return SNew(SHorizontalBox)
+				.ToolTip(MobilityTooltip)
+				.Visibility(EVisibility::Visible) // so we still get tooltip text for an empty SHorizontalBox
+				+ SHorizontalBox::Slot()
+				.FillWidth(1.0f)
+				[
+					SNew(SImage)
+					.Image(this, &SSCS_RowWidget::GetMobilityIconImage)
 					.ToolTip(MobilityTooltip)
-					.Visibility(EVisibility::Visible) // so we still get tooltip text for an empty SHorizontalBox
-				+SHorizontalBox::Slot()
-					.FillWidth(1.0f)
-					[
-						SNew(SImage)
-							.Image(this, &SSCS_RowWidget::GetMobilityIconImage, SCSNode)
-							.ToolTip(MobilityTooltip)
-					];
+				];
+		}
+		else
+		{
+			return SNew(SSpacer);
+		}
 	}
 	else
 	{
@@ -1205,12 +1417,22 @@ TSharedRef<SToolTip> SSCS_RowWidget::CreateToolTipWidget() const
 		AddToToolTipInfoBox(InfoBox, LOCTEXT("TooltipAsset", "Asset"), SNullWidget::NullWidget, TAttribute<FText>(this, &SSCS_RowWidget::GetAssetName), false);
 	}
 
-	TWeakObjectPtr<USCS_Node> SCSNode(GetNode()->GetSCSNode());
+	if (FSCSEditorTreeNode* TreeNode = GetNode().Get())
+	{
+		if (TreeNode->GetNodeType() == FSCSEditorTreeNode::ComponentNode)
+		{
+			// Add introduction point
+			if (TreeNode->IsInherited())
+			{
+				AddToToolTipInfoBox(InfoBox, LOCTEXT("TooltipIntroducedIn", "Introduced in"), SNullWidget::NullWidget, TAttribute<FText>::Create(TAttribute<FText>::FGetter::CreateSP(this, &SSCS_RowWidget::GetIntroducedInToolTipText)), false);
+			}
+
+			// Add mobility
+			TSharedRef<SImage> MobilityIcon = SNew(SImage).Image(this, &SSCS_RowWidget::GetMobilityIconImage);
+			AddToToolTipInfoBox(InfoBox, LOCTEXT("TooltipMobility", "Mobility"), MobilityIcon, TAttribute<FText>::Create(TAttribute<FText>::FGetter::CreateSP(this, &SSCS_RowWidget::GetMobilityToolTipText)), false);
+		}
+	}
 
-	// Add mobility
-	TSharedRef<SImage> MobilityIcon = SNew(SImage).Image(this, &SSCS_RowWidget::GetMobilityIconImage, SCSNode);
-	AddToToolTipInfoBox(InfoBox, LOCTEXT("TooltipMobility", "Mobility"), MobilityIcon, TAttribute<FText>::Create(TAttribute<FText>::FGetter::CreateSP(this, &SSCS_RowWidget::GetMobilityToolTipText, SCSNode)), false);
-	
 	TSharedRef<SBorder> TooltipContent = SNew(SBorder)
 		.BorderImage(FEditorStyle::GetBrush("NoBorder"))
 		.Padding(0)
@@ -1257,61 +1479,134 @@ TSharedRef<SToolTip> SSCS_RowWidget::CreateToolTipWidget() const
 }
 END_SLATE_FUNCTION_BUILD_OPTIMIZATION
 
-FSlateBrush const* SSCS_RowWidget::GetMobilityIconImage(TWeakObjectPtr<USCS_Node> SCSNode) const
+FSlateBrush const* SSCS_RowWidget::GetMobilityIconImage() const
 {
-	if (!SCSNode.IsValid())
-	{
-		return NULL;
-	}
-
-	USceneComponent* SceneComponentTemplate = Cast<USceneComponent>(SCSNode->ComponentTemplate);
-	if (SceneComponentTemplate == NULL)
+	if (FSCSEditorTreeNode* TreeNode = GetNode().Get())
 	{
-		return NULL;
-	}
+		if (USceneComponent* SceneComponentTemplate = Cast<USceneComponent>(TreeNode->GetComponentTemplate()))
+		{
+			if (SceneComponentTemplate->Mobility == EComponentMobility::Movable)
+			{
+				return FEditorStyle::GetBrush(TEXT("ClassIcon.MovableMobilityIcon"));
+			}
+			else if (SceneComponentTemplate->Mobility == EComponentMobility::Stationary)
+			{
+				return FEditorStyle::GetBrush(TEXT("ClassIcon.StationaryMobilityIcon"));
+			}
 
-	if (SceneComponentTemplate->Mobility == EComponentMobility::Movable)
-	{
-		return FEditorStyle::GetBrush(TEXT("ClassIcon.MovableMobilityIcon"));
-	}
-	else if (SceneComponentTemplate->Mobility == EComponentMobility::Stationary)
-	{
-		return FEditorStyle::GetBrush(TEXT("ClassIcon.StationaryMobilityIcon"));
+			// static components don't get an icon (because static is the most common
+			// mobility type, and we'd like to keep the icon clutter to a minimum)
+		}
 	}
 
-	// static components don't get an icon (because static is the most common
-	// mobility type, and we'd like to keep the icon clutter to a minimum)
-	return NULL;
+	return nullptr;
 }
 
-FText SSCS_RowWidget::GetMobilityToolTipText(TWeakObjectPtr<USCS_Node> SCSNode) const
+FText SSCS_RowWidget::GetMobilityToolTipText() const
 {
-	FText MobilityToolTip = LOCTEXT("NoMobilityTooltip", "This component does not have 'Mobility' associated with it");
-	if (SCSNode.IsValid())
+	FText MobilityToolTip = LOCTEXT("ErrorNoMobilityTooltip", "Invalid component");
+
+	if (FSCSEditorTreeNode* TreeNode = TreeNodePtr.Get())
 	{
-		USceneComponent* SceneComponentTemplate = Cast<USceneComponent>(SCSNode->ComponentTemplate);
-		if (SceneComponentTemplate == NULL)
-		{
-			MobilityToolTip = LOCTEXT("NoMobilityTooltip", "This component does not have 'Mobility' associated with it");
-		}
-		else if (SceneComponentTemplate->Mobility == EComponentMobility::Movable)
+		if (USceneComponent* SceneComponentTemplate = Cast<USceneComponent>(TreeNode->GetComponentTemplate()))
 		{
-			MobilityToolTip = LOCTEXT("MovableMobilityTooltip", "Movable component");
+			if (SceneComponentTemplate->Mobility == EComponentMobility::Movable)
+			{
+				MobilityToolTip = LOCTEXT("MovableMobilityTooltip", "Movable component");
+			}
+			else if (SceneComponentTemplate->Mobility == EComponentMobility::Stationary)
+			{
+				MobilityToolTip = LOCTEXT("StationaryMobilityTooltip", "Stationary component");
+			}
+			else if (SceneComponentTemplate->Mobility == EComponentMobility::Static)
+			{
+				MobilityToolTip = LOCTEXT("StaticMobilityTooltip", "Static component");
+			}
+			else
+			{
+				// make sure we're the mobility type we're expecting (we've handled Movable & Stationary)
+				ensureMsgf(false, TEXT("Unhandled mobility type [%d], is this a new type that we don't handle here?"), SceneComponentTemplate->Mobility.GetValue());
+				MobilityToolTip = LOCTEXT("UnknownMobilityTooltip", "Component with unknown mobility");
+			}
 		}
-		else if (SceneComponentTemplate->Mobility == EComponentMobility::Stationary)
+		else
 		{
-			MobilityToolTip = LOCTEXT("StationaryMobilityTooltip", "Stationary component");
+			MobilityToolTip = LOCTEXT("NoMobilityTooltip", "Non-scene component");
 		}
-		else 
+	}
+
+	return MobilityToolTip;
+}
+
+FText SSCS_RowWidget::GetIntroducedInToolTipText() const
+{
+	FText IntroducedInTooltip = LOCTEXT("IntroducedHereTooltip", "this class");
+
+	if (FSCSEditorTreeNode* TreeNode = TreeNodePtr.Get())
+	{
+		if (TreeNode->IsInherited())
 		{
-			// make sure we're the mobility type we're expecting (we've handled Movable & Stationary)
-			ensureMsgf(SceneComponentTemplate->Mobility == EComponentMobility::Static, TEXT("Unhandled mobility type [%d], is this a new type that we don't handle here?"), SceneComponentTemplate->Mobility.GetValue());
+			if (UActorComponent* ComponentTemplate = TreeNode->GetComponentTemplate())
+			{
+				UClass* BestClass = nullptr;
+				AActor* OwningActor = ComponentTemplate->GetOwner();
+
+				if (TreeNode->IsNative() && (OwningActor != nullptr))
+				{
+					for (UClass* TestClass = OwningActor->GetClass(); TestClass != AActor::StaticClass(); TestClass = TestClass->GetSuperClass())
+					{
+						if (TreeNode->FindComponentInstanceInActor(Cast<AActor>(TestClass->GetDefaultObject())))
+						{
+							BestClass = TestClass;
+						}
+						else
+						{
+							break;
+						}
+					}
+				}
+				else if (!TreeNode->IsNative())
+				{
+					USCS_Node* SCSNode = TreeNode->GetSCSNode();
+
+					if ((SCSNode == nullptr) && (OwningActor != nullptr))
+					{
+						SCSNode = FSCSEditorTreeNode::FindSCSNodeForInstance(ComponentTemplate, OwningActor->GetClass());
+					}
+
+					if (SCSNode != nullptr)
+					{
+						if (UBlueprint* OwningBP = SCSNode->GetSCS()->GetBlueprint())
+						{
+							BestClass = OwningBP->GeneratedClass;
+						}
+					}
+				}
+
+				if (BestClass == nullptr)
+				{
+					ensure(ComponentTemplate->CreationMethod == EComponentCreationMethod::ConstructionScript);
+					IntroducedInTooltip = LOCTEXT("IntroducedInUnknownError", "Unknown Blueprint Class (via an Add Component call)");
+				}
+				else
+				{
+					FString ClassName = (BestClass != nullptr) ? BestClass->GetName() : TEXT("(unknown)");
+					if ((BestClass != nullptr) && (BestClass->ClassGeneratedBy != nullptr) && ClassName.EndsWith(TEXT("_C")))
+					{
+						ClassName.RemoveFromEnd(TEXT("_C"));
+					}
 
-			MobilityToolTip = LOCTEXT("StaticMobilityTooltip", "Static component");
+					IntroducedInTooltip = FText::FromString(ClassName);
+				}
+			}
+			else
+			{
+				IntroducedInTooltip = LOCTEXT("IntroducedInNoTemplateError", "[no component template found]");
+			}
 		}
 	}
 
-	return MobilityToolTip;
+	return IntroducedInTooltip;
 }
 
 FText SSCS_RowWidget::GetAssetName() const
@@ -1365,37 +1660,26 @@ EVisibility SSCS_RowWidget::GetAssetVisibility() const
 
 FSlateColor SSCS_RowWidget::GetColorTint() const
 {
+	const FLinearColor InheritedBlueprintComponentColor(0.08f, 0.35f, 0.6f);
+	const FLinearColor InheritedNativeComponentColor(0.08f, 0.15f, 0.6f);
+	const FLinearColor IntroducedHereColor(FLinearColor::White);
+
 	FSCSEditorTreeNodePtrType NodePtr = GetNode();
-	if(SCSEditor.Pin()->GetEditorMode() == EComponentEditorMode::BlueprintSCS)
+	if (NodePtr->IsInherited())
 	{
-		if(NodePtr->IsNative())
+		if (NodePtr->IsNative())
 		{
-			return FLinearColor(0.08f,0.15f,0.6f);
+			return InheritedNativeComponentColor;
 		}
-		else if(NodePtr->IsInherited())
+		else
 		{
-			return FLinearColor(0.08f,0.35f,0.6f);
+			return InheritedBlueprintComponentColor;
 		}
 	}
-	else    // EComponentEditorMode::ActorInstance
+	else
 	{
-		if(!NodePtr->IsUserInstanced())
-		{
-			UActorComponent* InstancedComponent = NodePtr->GetComponentTemplate();
-			check(InstancedComponent != nullptr);
-
-			if(InstancedComponent->CreationMethod == EComponentCreationMethod::ConstructionScript)
-			{
-				return FLinearColor(0.08f,0.35f,0.6f);
-			}
-			else
-			{
-				return FLinearColor(0.08f,0.15f,0.6f);
-			}
-		}
+		return IntroducedHereColor;
 	}
-
-	return FLinearColor(1.0f, 1.0f, 1.0f);
 }
 
 EVisibility SSCS_RowWidget::GetRootLabelVisibility() const
@@ -1523,11 +1807,11 @@ void SSCS_RowWidget::OnDragEnter( const FGeometry& MyGeometry, const FDragDropEv
 				// We set the tooltip text here because it won't change across entry/leave events
 				if (DragRowOp->SourceNodes.Num() == 1)
 				{
-					if (Cast<USceneComponent>(SelectedNodePtr->GetComponentTemplate()) == nullptr)
+					if (!SelectedNodePtr->IsSceneComponent())
 					{
 						Message = LOCTEXT("DropActionToolTip_Error_CannotReparent_NotSceneComponent", "The selected component is not a scene component and cannot be attached to other components.");
 					}
-					else if (SelectedNodePtr->IsNative() || SelectedNodePtr->IsInherited())
+					else if (SelectedNodePtr->IsInherited())
 					{
 						Message = LOCTEXT("DropActionToolTip_Error_CannotReparent_Inherited", "The selected component is inherited and cannot be reordered here.");
 					}
@@ -2264,9 +2548,9 @@ FText SSCS_RowWidget::GetTooltipText() const
 {
 	FSCSEditorTreeNodePtrType NodePtr = GetNode();
 
-	if(NodePtr->IsDefaultSceneRoot())
+	if (NodePtr->IsDefaultSceneRoot())
 	{
-		if(NodePtr->IsInherited())
+		if (NodePtr->IsInherited())
 		{
 			return LOCTEXT("InheritedDefaultSceneRootToolTip", "This is the default scene root component. It cannot be copied, renamed or deleted.\nIt has been inherited from the parent class, so its properties cannot be edited here.\nNew scene components will automatically be attached to it.");
 		}
@@ -2311,7 +2595,7 @@ FText SSCS_RowWidget::GetTooltipText() const
 		}
 		else
 		{
-			if ( NodePtr->IsInherited() )
+			if (NodePtr->IsInheritedSCS())
 			{
 				if ( NodePtr->GetComponentTemplate() != NULL )
 				{
@@ -2337,18 +2621,10 @@ FString SSCS_RowWidget::GetDocumentationLink() const
 {
 	check(SCSEditor.IsValid());
 
-	if (SCSEditor.Pin()->GetEditorMode() == EComponentEditorMode::BlueprintSCS)
-	{
-		FSCSEditorTreeNodePtrType NodePtr = GetNode();
-		if( NodePtr == SCSEditor.Pin()->SceneRootNodePtr
-			|| NodePtr->IsNative() || NodePtr->IsInherited())
-		{
-			return TEXT("Shared/Editors/BlueprintEditor/ComponentsMode");
-		}
-	}
-	else    // EComponentEditorMode::ActorInstance
+	FSCSEditorTreeNodePtrType NodePtr = GetNode();
+	if ((NodePtr == SCSEditor.Pin()->SceneRootNodePtr) || NodePtr->IsInherited())
 	{
-		// @TODO - Actor instance mode
+		return TEXT("Shared/Editors/BlueprintEditor/ComponentsMode");
 	}
 
 	return TEXT("");
@@ -2358,25 +2634,18 @@ FString SSCS_RowWidget::GetDocumentationExcerptName() const
 {
 	check(SCSEditor.IsValid());
 
-	if (SCSEditor.Pin()->GetEditorMode() == EComponentEditorMode::BlueprintSCS)
+	FSCSEditorTreeNodePtrType NodePtr = GetNode();
+	if (NodePtr == SCSEditor.Pin()->SceneRootNodePtr)
 	{
-		FSCSEditorTreeNodePtrType NodePtr = GetNode();
-		if( NodePtr == SCSEditor.Pin()->SceneRootNodePtr)
-		{
-			return TEXT("RootComponent");
-		}
-		else if(NodePtr->IsNative())
-		{
-			return TEXT("NativeComponents");
-		}
-		else if(NodePtr->IsInherited())
-		{
-			return TEXT("InheritedComponents");
-		}
+		return TEXT("RootComponent");
 	}
-	else    // EComponentEditorMode::ActorInstance
+	else if (NodePtr->IsNative())
 	{
-		// @TODO - Actor instance mode
+		return TEXT("NativeComponents");
+	}
+	else if (NodePtr->IsInherited())
+	{
+		return TEXT("InheritedComponents");
 	}
 
 	return TEXT("");
@@ -2633,7 +2902,13 @@ FText SSCS_RowWidget_ActorRoot::GetActorMobilityText() const
 		if (AActor* DefaultActor = SCSEditorPtr->GetActorContext())
 		{
 			USceneComponent* RootComponent = DefaultActor->GetRootComponent();
-			if (RootComponent)
+
+			if ((RootComponent == nullptr) && (SCSEditorPtr->SceneRootNodePtr.IsValid()))
+			{
+				RootComponent = Cast<USceneComponent>(SCSEditorPtr->SceneRootNodePtr->GetComponentTemplate());
+			}
+
+			if (RootComponent != nullptr)
 			{
 				if (RootComponent->Mobility == EComponentMobility::Static)
 				{
@@ -2664,27 +2939,6 @@ FText SSCS_RowWidget_ActorRoot::GetActorMobilityText() const
 
 TSharedRef<SWidget> SSCS_RowWidget_Separator::GenerateWidgetForColumn(const FName& ColumnName)
 {
-	/*return SNew(SVerticalBox)
-
-		+ SVerticalBox::Slot()
-		.Padding(0.f)
-		.VAlign(VAlign_Top)
-		[
-			SNew(SBorder)
-			.Padding(FEditorStyle::GetMargin(TEXT("Menu.Separator.Padding")) + FMargin(1.f, 0.f))
-			.BorderImage(FEditorStyle::GetBrush(TEXT("Menu.Separator")))
-		]
-	
-		+ SVerticalBox::Slot()
-		.VAlign(VAlign_Top)
-		.AutoHeight()
-		.Padding(1.f)
-		[
-			SNew(STextBlock)
-			.Text(LOCTEXT("HeaderText", "Header Text"))
-			.TextStyle(FEditorStyle::Get(), TEXT("Menu.Heading"))
-		];*/
-
 	return SNew(SBox)
 		.Padding(1.f)
 		[
@@ -3765,14 +4019,14 @@ void SSCSEditor::UpdateTree(bool bRegenerateTreeNodes)
 		bool bSingleLayoutBPEditor = GetDefault<UEditorExperimentalSettings>()->bUnifiedBlueprintEditor;
 		if (bSingleLayoutBPEditor)
 		{
-			ActorTreeNode = MakeShareable(new FSCSEditorTreeNode(FSCSEditorTreeNode::RootActorNode));
+			ActorTreeNode = MakeShareable(new FSCSEditorTreeNodeRootActor());
 			
 			//RootTreeNode = ActorTreeNode
 			//RootNodes.Add(RootTreeNode);
 			//SCSTreeWidget->SetItemExpansion(RootTreeNode, true);
 
 			RootNodes.Add(ActorTreeNode);
-			RootNodes.Add(MakeShareable(new FSCSEditorTreeNode(FSCSEditorTreeNode::SeparatorNode)));
+			RootNodes.Add(MakeShareable(new FSCSEditorTreeNodeSeparator()));
 		}
 
 		// Build the tree data source according to what mode we're in
@@ -3806,7 +4060,7 @@ void SSCSEditor::UpdateTree(bool bRegenerateTreeNodes)
 					USceneComponent* SceneComp = Cast<USceneComponent>(*CompIter);
 					if(SceneComp != NULL)
 					{
-						AddTreeNode(SceneComp);
+						AddTreeNodeFromComponent(SceneComp);
 					}
 				}
 
@@ -3820,7 +4074,7 @@ void SSCSEditor::UpdateTree(bool bRegenerateTreeNodes)
 						if (bSingleLayoutBPEditor && !bSeparatorAdded)
 						{
 							bSeparatorAdded = true;
-							RootNodes.Add(MakeShareable(new FSCSEditorTreeNode(FSCSEditorTreeNode::SeparatorNode)));
+							RootNodes.Add(MakeShareable(new FSCSEditorTreeNodeSeparator()));
 						}
 						AddRootComponentTreeNode(ActorComp);
 					}
@@ -3873,7 +4127,7 @@ void SSCSEditor::UpdateTree(bool bRegenerateTreeNodes)
 					USceneComponent* SceneComp = Cast<USceneComponent>(*CompIter);
 					if(SceneComp != nullptr && !SceneComp->IsEditorOnly())
 					{
-						AddTreeNode(SceneComp);
+						AddTreeNodeFromComponent(SceneComp);
 					}
 				}
 
@@ -3958,11 +4212,11 @@ TSharedPtr<FSCSEditorTreeNode> SSCSEditor::AddRootComponentTreeNode(UActorCompon
 	TSharedPtr<FSCSEditorTreeNode> NewTreeNode;
 	if (RootTreeNode.IsValid())
 	{
-		NewTreeNode = RootTreeNode->AddChild(ActorComp);
+		NewTreeNode = RootTreeNode->AddChildFromComponent(ActorComp);
 	}
 	else
 	{
-		NewTreeNode = MakeShareable(new FSCSEditorTreeNode(ActorComp));
+		NewTreeNode = FSCSEditorTreeNode::FactoryNodeFromComponent(ActorComp);
 		RootNodes.Add(NewTreeNode);
 	}
 
@@ -4164,14 +4418,14 @@ UActorComponent* SSCSEditor::AddNewComponent( UClass* NewComponentClass, UObject
 			{
 				NewComponentName = *FComponentEditorUtils::GenerateValidVariableName(NewComponentClass, ActorInstance);
 			}
-			NewComponent = AddNewNode(ConstructObject<UActorComponent>(NewComponentClass, ActorInstance, NewComponentName, RF_Transactional, ComponentTemplate), Asset, true);
+			NewComponent = AddNewNodeForInstancedComponent(ConstructObject<UActorComponent>(NewComponentClass, ActorInstance, NewComponentName, RF_Transactional, ComponentTemplate), Asset, true);
 		}
 	}
 
 	return NewComponent;
 }
 
-UActorComponent* SSCSEditor::AddNewNode(USCS_Node* NewNode,  UObject* Asset, bool bMarkBlueprintModified, bool bSetFocusToNewItem)
+UActorComponent* SSCSEditor::AddNewNode(USCS_Node* NewNode, UObject* Asset, bool bMarkBlueprintModified, bool bSetFocusToNewItem)
 {
 	check(NewNode != nullptr);
 
@@ -4220,7 +4474,7 @@ UActorComponent* SSCSEditor::AddNewNode(USCS_Node* NewNode,  UObject* Asset, boo
 	return NewNode->ComponentTemplate;
 }
 
-UActorComponent* SSCSEditor::AddNewNode(UActorComponent* NewInstanceComponent,  UObject* Asset, bool bSetFocusToNewItem)
+UActorComponent* SSCSEditor::AddNewNodeForInstancedComponent(UActorComponent* NewInstanceComponent, UObject* Asset, bool bSetFocusToNewItem)
 {
 	check(NewInstanceComponent != nullptr);
 
@@ -4256,7 +4510,7 @@ UActorComponent* SSCSEditor::AddNewNode(UActorComponent* NewInstanceComponent,
 	USceneComponent* NewSceneComponent = Cast<USceneComponent>(NewInstanceComponent);
 	if(NewSceneComponent != nullptr)
 	{
-		NewNodePtr = AddTreeNode(NewSceneComponent);
+		NewNodePtr = AddTreeNodeFromComponent(NewSceneComponent);
 
 		// Remove the old scene root node if it's set to the default one
 		if(SceneRootNodePtr.IsValid() && SceneRootNodePtr->IsDefaultSceneRoot())
@@ -4490,7 +4744,7 @@ void SSCSEditor::PasteNodes()
 			NewActorComponent->Rename(*NewComponentName, ActorInstance, REN_DontCreateRedirectors | REN_DoNotDirty);
 
 			// Create a new node to contain the new component instance and add it to the tree
-			NewActorComponent = AddNewNode(NewActorComponent, NULL, false);
+			NewActorComponent = AddNewNodeForInstancedComponent(NewActorComponent, NULL, false);
 		}
 
 		if(NewActorComponent)
@@ -4840,7 +5094,7 @@ bool SSCSEditor::IsNodeInSimpleConstructionScript( USCS_Node* Node ) const
 	return false;
 }
 
-FSCSEditorTreeNodePtrType SSCSEditor::AddTreeNode(USCS_Node* InSCSNode, FSCSEditorTreeNodePtrType InParentNodePtr, const bool bIsInherited)
+FSCSEditorTreeNodePtrType SSCSEditor::AddTreeNode(USCS_Node* InSCSNode, FSCSEditorTreeNodePtrType InParentNodePtr, const bool bIsInheritedSCS)
 {
 	FSCSEditorTreeNodePtrType NewNodePtr;
 
@@ -4866,7 +5120,7 @@ FSCSEditorTreeNodePtrType SSCSEditor::AddTreeNode(USCS_Node* InSCSNode, FSCSEdit
 		if(ParentPtr.IsValid())
 		{
 			// do this first, because we need a FSCSEditorTreeNodePtrType for the new node
-			NewNodePtr = ParentPtr->AddChild(InSCSNode, bIsInherited);
+			NewNodePtr = ParentPtr->AddChild(InSCSNode, bIsInheritedSCS);
 
 			bool bParentIsEditorOnly = ParentPtr->GetComponentTemplate()->IsEditorOnly();
 			// if you can't nest this new node under the proposed parent (then swap the two)
@@ -4903,11 +5157,11 @@ FSCSEditorTreeNodePtrType SSCSEditor::AddTreeNode(USCS_Node* InSCSNode, FSCSEdit
 			// Create a new root node
 			if (RootTreeNode.IsValid())
 			{
-				NewNodePtr = RootTreeNode->AddChild(InSCSNode, bIsInherited);
+				NewNodePtr = RootTreeNode->AddChild(InSCSNode, bIsInheritedSCS);
 			}
 			else
 			{
-				NewNodePtr = MakeShareable(new FSCSEditorTreeNode(InSCSNode, bIsInherited));
+				NewNodePtr = MakeShareable(new FSCSEditorTreeNodeComponent(InSCSNode, bIsInheritedSCS));
 				RootNodes.Add(NewNodePtr);
 			}
 			
@@ -4928,11 +5182,11 @@ FSCSEditorTreeNodePtrType SSCSEditor::AddTreeNode(USCS_Node* InSCSNode, FSCSEdit
 		// If the given SCS node does not contain a scene component template, we create a new root node
 		if (RootTreeNode.IsValid())
 		{
-			NewNodePtr = RootTreeNode->AddChild(InSCSNode, bIsInherited);
+			NewNodePtr = RootTreeNode->AddChild(InSCSNode, bIsInheritedSCS);
 		}
 		else
 		{
-			NewNodePtr = MakeShareable(new FSCSEditorTreeNode(InSCSNode, bIsInherited));
+			NewNodePtr = MakeShareable(new FSCSEditorTreeNodeComponent(InSCSNode, bIsInheritedSCS));
 			RootNodes.Add(NewNodePtr);
 		}
 
@@ -4948,13 +5202,13 @@ FSCSEditorTreeNodePtrType SSCSEditor::AddTreeNode(USCS_Node* InSCSNode, FSCSEdit
 	// Recursively add the given SCS node's child nodes
 	for(int32 NodeIndex = 0; NodeIndex < InSCSNode->ChildNodes.Num(); ++NodeIndex)
 	{
-		AddTreeNode(InSCSNode->ChildNodes[NodeIndex], NewNodePtr, bIsInherited);
+		AddTreeNode(InSCSNode->ChildNodes[NodeIndex], NewNodePtr, bIsInheritedSCS);
 	}
 
 	return NewNodePtr;
 }
 
-FSCSEditorTreeNodePtrType SSCSEditor::AddTreeNode(USceneComponent* InSceneComponent)
+FSCSEditorTreeNodePtrType SSCSEditor::AddTreeNodeFromComponent(USceneComponent* InSceneComponent)
 {
 	FSCSEditorTreeNodePtrType NewNodePtr;
 
@@ -4969,12 +5223,12 @@ FSCSEditorTreeNodePtrType SSCSEditor::AddTreeNode(USceneComponent* InSceneCompon
 		if(!ParentNodePtr.IsValid())
 		{
 			// Recursively add the parent node to the tree if it does not exist yet
-			ParentNodePtr = AddTreeNode(InSceneComponent->AttachParent);
+			ParentNodePtr = AddTreeNodeFromComponent(InSceneComponent->AttachParent);
 		}
 
 		// Add a new tree node for the given scene component
 		check(ParentNodePtr.IsValid());
-		NewNodePtr = ParentNodePtr->AddChild(InSceneComponent);
+		NewNodePtr = ParentNodePtr->AddChildFromComponent(InSceneComponent);
 
 		// Expand parent nodes by default
 		SCSTreeWidget->SetItemExpansion(ParentNodePtr, true);
@@ -4995,7 +5249,7 @@ FSCSEditorTreeNodePtrType SSCSEditor::AddTreeNode(USceneComponent* InSceneCompon
 		}
 		else if (SceneRootNodePtr->GetComponentTemplate() != InSceneComponent)
 		{
-			NewNodePtr = SceneRootNodePtr->AddChild(InSceneComponent);
+			NewNodePtr = SceneRootNodePtr->AddChildFromComponent(InSceneComponent);
 		}
 	}
 
diff --git a/Engine/Source/Editor/Kismet/Public/SSCSEditor.h b/Engine/Source/Editor/Kismet/Public/SSCSEditor.h
--- a/Engine/Source/Editor/Kismet/Public/SSCSEditor.h
+++ b/Engine/Source/Editor/Kismet/Public/SSCSEditor.h
@@ -13,13 +13,15 @@ class USimpleConstructionScript;
 typedef TSharedPtr<class FSCSEditorTreeNode> FSCSEditorTreeNodePtrType;
 
 /**
- * FSCSEditorTreeNodeBase
+ * FSCSEditorTreeNode
  *
  * Wrapper class for component template nodes displayed in the SCS editor tree widget.
  */
 class KISMET_API FSCSEditorTreeNode : public TSharedFromThis<FSCSEditorTreeNode>
 {
 public:
+	/** Delegate for when the context menu requests a rename */
+	DECLARE_DELEGATE(FOnRenameRequested);
 
 	enum ENodeType
 	{
@@ -28,60 +30,45 @@ public:
 		SeparatorNode,
 	};
 
-	/** 
+	/**
 	 * Constructs an empty tree node.
 	 */
 	FSCSEditorTreeNode(FSCSEditorTreeNode::ENodeType InNodeType);
 
 	/**
-	 * Constructs a wrapper around a component template contained within an SCS tree node.
-	 *
-	 * @param InSCSNode The SCS tree node represented by this object.
-	 * @param bInIsInherited Whether or not the SCS tree node is inherited from a parent Blueprint class.
-	 */
-	FSCSEditorTreeNode(class USCS_Node* InSCSNode, bool bInIsInherited = false);
-
-	/**
-	 * Constructs a wrapper around a component template not contained within an SCS tree node (e.g. "native" components).
-	 *
-	 * @param InComponentTemplate The component template represented by this object.
-	 */
-	FSCSEditorTreeNode(UActorComponent* InComponentTemplate);
-
-	/** 
 	 * @return The name of the variable represented by this node.
 	 */
-	FName GetVariableName() const;
+	virtual FName GetVariableName() const;
 	/**
 	 * @return The string to be used in the tree display.
 	 */
-	FString GetDisplayString() const;
+	virtual FString GetDisplayString() const;
 	/**
 	* @return The name of this node in text.
 	*/
-	FText GetDisplayName() const;
-	/** 
-     * @return The SCS node that is represented by this object, or NULL if there is no SCS node associated with the component template.
+	virtual FText GetDisplayName() const;
+	/**
+	 * @return The SCS node that is represented by this object, or NULL if there is no SCS node associated with the component template.
 	 */
-	class USCS_Node* GetSCSNode() const { return SCSNodePtr.Get(); }
+	virtual class USCS_Node* GetSCSNode() const;
 	/**
 	 * @return The component template or instance represented by this object.
 	 */
-	UActorComponent* GetComponentTemplate() const;
+	virtual UActorComponent* GetComponentTemplate() const;
 	/**
 	 * @param ActualEditedBlueprint currently edited blueprint
 	 *
 	 * @return The component template that can be editable for actual class.
 	 */
-	UActorComponent* GetEditableComponentTemplate(UBlueprint* ActualEditedBlueprint);
+	virtual UActorComponent* GetEditableComponentTemplate(UBlueprint* ActualEditedBlueprint);
 	/**
 	 * Finds the component instance represented by this node contained within a given Actor instance.
 	 *
 	 * @param InActor The Actor instance to use as the container object for finding the component instance.
 	 * @return The component instance represented by this node and contained within the given Actor instance, or NULL if not found.
 	 */
-	UActorComponent* FindComponentInstanceInActor(const AActor* InActor) const;
-	/** 
+	virtual UActorComponent* FindComponentInstanceInActor(const AActor* InActor) const;
+	/**
 	 * @return This object's parent node (or an invalid reference if no parent is assigned).
 	 */
 	FSCSEditorTreeNodePtrType GetParent() const { return ParentNodePtr; }
@@ -89,96 +76,52 @@ public:
 	 * @return The set of nodes which are parented to this node (read-only).
 	 */
 	const TArray<FSCSEditorTreeNodePtrType>& GetChildren() const { return Children; }
-	/**
-	 * @return The Blueprint to which this node belongs.
-	 */
-	UBlueprint* GetBlueprint() const;
 	/**
 	 * @return Type of node
 	 */
 	ENodeType GetNodeType() const;
-	/**
-	 * @return Whether or not this object represents a root component.
-	 */
-	bool IsRootComponent() const;
 	/**
 	 * @return Whether or not this node is a direct child of the given node.
 	 */
-	bool IsDirectlyAttachedTo(FSCSEditorTreeNodePtrType InNodePtr) const { return ParentNodePtr == InNodePtr; } 
+	bool IsDirectlyAttachedTo(FSCSEditorTreeNodePtrType InNodePtr) const { return ParentNodePtr == InNodePtr; }
 	/**
 	 * @return Whether or not this node is a child (direct or indirect) of the given node.
 	 */
 	bool IsAttachedTo(FSCSEditorTreeNodePtrType InNodePtr) const;
-	/**
-	 * @return Whether or not this object represents a "native" component template (i.e. one that is not found in the SCS tree).
-	 */
-	bool IsNative() const { return GetSCSNode() == NULL && GetComponentTemplate() != NULL; }
-	/** 
-	 * @return Whether or not this object represents a component template that was inherited from a parent Blueprint class SCS.
-	 */
-	bool IsInherited() const { return bIsInherited; }
-	/**
-	 * @return Whether or not this object represents a component instance rather than a template.
-	 */
-	bool IsInstanced() const { return bIsInstanced; }
-	/**
-	 * @return Whether or not this object represents a component instance that was created by the user and not by a native or Blueprint-generated class.
-	 */
-	bool IsUserInstanced() const;
-	/** 
-	 * @return Whether or not this object represents the default SCS scene root component.
-	 */
-	bool IsDefaultSceneRoot() const;
-	/** 
-	 * @return Whether or not this object represents a node that can be deleted from the SCS tree.
-	 */
-	bool CanDelete() const { return NodeType == ComponentNode && ((IsUserInstanced() || (!IsNative() && !IsInherited())) && !IsDefaultSceneRoot()); }
-	/** 
-	 * @return Whether or not this object represents a node that can be renamed from the SCS tree.
-	 */
-	bool CanRename() const { return NodeType == ComponentNode && ((IsUserInstanced() || (!IsNative() && !IsInherited())) && !IsDefaultSceneRoot()); }
-	/** 
-	 * @return Whether or not this object represents a node that can be reparented to other nodes based on its context.
-	 */
-	bool CanReparent() const { return NodeType == ComponentNode && ((IsUserInstanced() || (!IsNative() && !IsInherited())) && !IsDefaultSceneRoot() && Cast<USceneComponent>(GetComponentTemplate()) != NULL); }
-	/** 
-	 * @return Whether or not we can edit default properties for the component template represented by this object.
-	 */
-	bool CanEditDefaults() const;
 
 	/**
 	 * Finds the closest ancestor node in the given node set.
-	 * 
+	 *
 	 * @param InNodes The given node set.
 	 * @return One of the nodes from the set, or an invalid node reference if the set does not contain any ancestor nodes.
 	 */
 	FSCSEditorTreeNodePtrType FindClosestParent(TArray<FSCSEditorTreeNodePtrType> InNodes);
 
-	/** 
+	/**
 	 * Adds the given node as a child node.
 	 *
 	 * @param InChildNodePtr The node to add as a child node.
 	 */
 	void AddChild(FSCSEditorTreeNodePtrType InChildNodePtr);
 
-	/** 
+	/**
 	 * Adds a child node for the given SCS node.
 	 *
 	 * @param InSCSNode The SCS node to for which to create a child node.
-	 * @param bIsInherited Whether or not the given SCS node is inherited from a parent.
+	 * @param bIsInheritedSCS Whether or not the given SCS node is inherited from a parent.
 	 * @return A reference to the new child node or the existing child node if a match is found.
 	 */
-	FSCSEditorTreeNodePtrType AddChild(USCS_Node* InSCSNode, bool bIsInherited);
+	FSCSEditorTreeNodePtrType AddChild(USCS_Node* InSCSNode, bool bIsInheritedSCS);
 
-	/** 
+	/**
 	 * Adds a child node for the given component template.
 	 *
 	 * @param InComponentTemplate The component template for which to create a child node.
 	 * @return A reference to the new child node or the existing child node if a match is found.
 	 */
-	FSCSEditorTreeNodePtrType AddChild(UActorComponent* InComponentTemplate);
+	FSCSEditorTreeNodePtrType AddChildFromComponent(UActorComponent* InComponentTemplate);
 
-	/** 
+	/**
 	 * Attempts to find a reference to the child node that matches the given SCS node.
 	 *
 	 * @param InSCSNode The SCS node to match.
@@ -188,7 +131,7 @@ public:
 	 */
 	FSCSEditorTreeNodePtrType FindChild(const USCS_Node* InSCSNode, bool bRecursiveSearch = false, uint32* OutDepth = NULL) const;
 
-	/** 
+	/**
 	 * Attempts to find a reference to the child node that matches the given component template.
 	 *
 	 * @param InComponentTemplate The component template instance to match.
@@ -198,7 +141,7 @@ public:
 	 */
 	FSCSEditorTreeNodePtrType FindChild(const UActorComponent* InComponentTemplate, bool bRecursiveSearch = false, uint32* OutDepth = NULL) const;
 
-	/** 
+	/**
 	 * Attempts to find a reference to the child node that matches the given component variable or instance name.
 	 *
 	 * @param InVariableOrInstanceName The component variable or instance name to match.
@@ -208,46 +151,299 @@ public:
 	 */
 	FSCSEditorTreeNodePtrType FindChild(const FName& InVariableOrInstanceName, bool bRecursiveSearch = false, uint32* OutDepth = NULL) const;
 
-	/** 
+	/**
 	 * Removes the given node from the list of child nodes.
 	 *
 	 * @param InChildNodePtr The child node to remove.
 	 */
 	void RemoveChild(FSCSEditorTreeNodePtrType InChildNodePtr);
 
-	/** Delegate for when the context menu requests a rename */
-	DECLARE_DELEGATE(FOnRenameRequested);
+	bool IsSceneComponent() const
+	{
+		return Cast<USceneComponent>(GetComponentTemplate()) != nullptr;
+	}
+
+	// Tries to find a SCS node that was likely responsible for creating the specified instance component.  Note: This is not always possible to do!
+	static USCS_Node* FindSCSNodeForInstance(UActorComponent* InstanceComponent, UClass* ClassToSearch);
+
+	/**
+	 * Creates the correct type of node based on the component (instanced or not, etc...)
+	 */
+	static FSCSEditorTreeNodePtrType FactoryNodeFromComponent(UActorComponent* InComponent);
+protected:
+	// Called when this node is being removed via a RemoveChild call
+	virtual void RemoveMeAsChild() {}
+
+
+public:
+
+
+
+	/**
+	* @return The Blueprint to which this node belongs.
+	*/
+	UBlueprint* GetBlueprint() const;
+
+	/**
+	 * @return Whether or not this object represents a "native" component template (i.e. one that is not found in the SCS tree).
+	 */
+	virtual bool IsNative() const { return false; }
+
+	/**
+	 * @return Whether or not this object represents a root component.
+	 */
+	virtual bool IsRootComponent() const { return false; }
+
+	/**
+	 * @return Whether or not this object represents an inherited SCS node (one from a SCS node in a parent Blueprint).
+	 */
+	virtual bool IsInheritedSCS() const { return false; }
+
+	/**
+	 * @return Whether or not this object was declared in the current class (or instance).  Anything inherited cannot be reorganized (renamed, reparented, etc...).
+	 */
+	virtual bool IsInherited() const
+	{
+		return IsNative() || IsInheritedSCS();
+	}
+
+	/**
+	 * @return Whether or not this object represents a component instance rather than a template.
+	 */
+	virtual bool IsInstanced() const { return false; }
+
+	/**
+	 * @return Whether or not this object represents a component instance that was created by the user and not by a native or Blueprint-generated class.
+	 */
+	virtual bool IsUserInstanced() const { return false; }
+
+	/**
+	* @return Whether or not this object represents the default SCS scene root component.
+	*/
+	virtual bool IsDefaultSceneRoot() const { return false; }
+
+	/**
+	 * @return Whether or not this object represents a node that can be deleted from the SCS tree.
+	 */
+	virtual bool CanDelete() const { return false; }
+
+	/**
+	 * @return Whether or not this object represents a node that can be reparented to other nodes based on its context.
+	 */
+	virtual bool CanReparent() const { return false; }
+
+	/**
+	 * @return Whether or not we can edit default properties for the component template represented by this object.
+	 */
+	virtual bool CanEditDefaults() const { return false; }
+
+	/**
+	 * @return Whether or not this object represents a node that can be renamed from the components tree.
+	 */
+	virtual bool CanRename() const { return false; }
 
 	/** Requests a rename on the component */
 	void OnRequestRename(bool bTransactional);
 
 	/** Renames the component */
-	void OnCompleteRename(const FText& InNewName);
+	virtual void OnCompleteRename(const FText& InNewName);
 
 	/** Sets up the delegate for renaming a component */
 	void SetRenameRequestedDelegate(FOnRenameRequested InRenameRequested) { RenameRequestedDelegate = InRenameRequested; }
 
-	/** Get overridden template component, specialized in given blueprint */
-	UActorComponent* GetOverridenComponentTemplate(UBlueprint* Blueprint, bool bCreateIfNecessary) const;
+
+protected:
+	bool GetAndClearNonTransactionalRenameFlag()
+	{
+		const bool bResult = bNonTransactionalRename;
+		bNonTransactionalRename = false;
+		return bResult;
+	}
 
 private:
+	// The type of component tree node
 	ENodeType NodeType;
-	bool bIsInherited;
-	bool bIsInstanced;
-	TWeakObjectPtr<class USCS_Node> SCSNodePtr;
-	TWeakObjectPtr<UActorComponent> ComponentTemplatePtr;
+
+	// Actual tree structure
 	FSCSEditorTreeNodePtrType ParentNodePtr;
 	TArray<FSCSEditorTreeNodePtrType> Children;
 
 	/** Handles rename requests */
 	bool bNonTransactionalRename;
 	FOnRenameRequested RenameRequestedDelegate;
+};
+
+//////////////////////////////////////////////////////////////////////////
+//
+
+class KISMET_API FSCSEditorTreeNodeComponentBase : public FSCSEditorTreeNode
+{
+protected:
+	FSCSEditorTreeNodeComponentBase()
+		: FSCSEditorTreeNode(FSCSEditorTreeNode::ComponentNode)
+	{
+	}
+
+public:
+	// FSCSEditorTreeNode interface
+	virtual FName GetVariableName() const override;
+	virtual FString GetDisplayString() const override;
+	virtual bool CanRename() const { return !IsInherited() && !IsDefaultSceneRoot(); }
+	virtual bool CanDelete() const { return !IsInherited() && !IsDefaultSceneRoot(); }
+	virtual bool CanReparent() const { return !IsInherited() && !IsDefaultSceneRoot() && IsSceneComponent(); }
+	// End of FSCSEditorTreeNode interface
+};
+
+
+//////////////////////////////////////////////////////////////////////////
+// FSCSEditorTreeNodeInstancedInheritedComponent - A inherited component in the instanced case (either an inherited SCS node or an inherited native component)
+
+class KISMET_API FSCSEditorTreeNodeInstancedInheritedComponent : public FSCSEditorTreeNodeComponentBase
+{
+public:
+	/**
+	 * Constructs a wrapper around a component template not contained within an SCS tree node (e.g. "native" components).
+	 *
+	 * @param InComponentTemplate The component template represented by this object.
+	 */
+	FSCSEditorTreeNodeInstancedInheritedComponent(AActor* Owner, FName InComponentName);
+
+	// FSCSEditorTreeNode public interface
+	virtual bool IsNative() const override;
+	virtual bool IsRootComponent() const override;
+	virtual bool IsInheritedSCS() const override;
+	virtual bool IsInstanced() const override { return true; }
+	virtual bool IsInherited() const override { return true; }
+	virtual bool IsUserInstanced() const override { return false; }
+	virtual bool IsDefaultSceneRoot() const override;
+	virtual bool CanEditDefaults() const { return IsNative(); }
+	//virtual FName GetVariableName() const override;
+	//virtual FString GetDisplayString() const override;
+	virtual FText GetDisplayName() const override;
+	virtual UActorComponent* GetComponentTemplate() const override;
+	virtual UActorComponent* GetEditableComponentTemplate(UBlueprint* ActualEditedBlueprint) override;
+	// End of FSCSEditorTreeNode public interface
+
+private:
+	FName InstancedComponentName;
+	TWeakObjectPtr<AActor> InstancedComponentOwnerPtr;
+};
+
+//////////////////////////////////////////////////////////////////////////
+// FSCSEditorTreeNodeInstanceAddedComponent - A unique-to-this instance component
+
+class KISMET_API FSCSEditorTreeNodeInstanceAddedComponent : public FSCSEditorTreeNodeComponentBase
+{
+public:
+	/**
+	* Constructs a wrapper around a component template not contained within an SCS tree node (e.g. "native" components).
+	*
+	* @param InComponentTemplate The component template represented by this object.
+	*/
+	FSCSEditorTreeNodeInstanceAddedComponent(AActor* Owner, FName InComponentName);
+
+	// FSCSEditorTreeNode public interface
+	virtual bool IsNative() const override { return false; }
+	virtual bool IsRootComponent() const override;
+	virtual bool IsInheritedSCS() const override { return false; }
+	virtual bool IsInstanced() const override { return true; }
+	virtual bool IsUserInstanced() const override { return true; }
+	virtual bool IsDefaultSceneRoot() const override;
+	virtual bool CanEditDefaults() const { return true; }
+	virtual FName GetVariableName() const override { return NAME_None; }
+	virtual FString GetDisplayString() const override;
+	virtual FText GetDisplayName() const override;
+	virtual UActorComponent* GetComponentTemplate() const override;
+	virtual UActorComponent* GetEditableComponentTemplate(UBlueprint* ActualEditedBlueprint) override;
+	virtual void OnCompleteRename(const FText& InNewName) override;
+	// End of FSCSEditorTreeNode public interface
+
+protected:
+	// FSCSEditorTreeNode protected interface
+	virtual void RemoveMeAsChild() override;
+	// End of FSCSEditorTreeNode protected interface
 
-	bool bWasInstancedFromNativeClass;
+private:
 	FName InstancedComponentName;
 	TWeakObjectPtr<AActor> InstancedComponentOwnerPtr;
 };
 
+//////////////////////////////////////////////////////////////////////////
+// FSCSEditorTreeNodeComponent - A generic component in the non-instanced case (either a SCS node or an inherited native component)
+
+class KISMET_API FSCSEditorTreeNodeComponent : public FSCSEditorTreeNodeComponentBase
+{
+public:
+	/**
+	 * Constructs a wrapper around a component template contained within an SCS tree node.
+	 *
+	 * @param InSCSNode The SCS tree node represented by this object.
+	 * @param bInIsInherited Whether or not the SCS tree node is inherited from a parent Blueprint class.
+	 */
+	FSCSEditorTreeNodeComponent(class USCS_Node* InSCSNode, bool bInIsInherited = false);
+
+	/**
+	 * Constructs a wrapper around a component template not contained within an SCS tree node (e.g. "native" components).
+	 *
+	 * @param InComponentTemplate The component template represented by this object.
+	 */
+	FSCSEditorTreeNodeComponent(UActorComponent* InComponentTemplate);
+
+
+	// FSCSEditorTreeNode public interface
+	virtual bool IsNative() const override;
+	virtual bool IsRootComponent() const override;
+	virtual bool IsInheritedSCS() const override;
+	virtual bool IsInstanced() const override { return false; }
+	virtual bool IsUserInstanced() const override { return false; }
+	virtual bool IsDefaultSceneRoot() const override;
+	virtual bool CanEditDefaults() const override;
+	//virtual FName GetVariableName() const override;
+	//virtual FString GetDisplayString() const override;
+	virtual FText GetDisplayName() const override;
+	virtual class USCS_Node* GetSCSNode() const override;
+	virtual UActorComponent* GetComponentTemplate() const override;
+	virtual UActorComponent* GetEditableComponentTemplate(UBlueprint* ActualEditedBlueprint) override;
+	virtual void OnCompleteRename(const FText& InNewName) override;
+	// End of FSCSEditorTreeNode public interface
+
+protected:
+	// FSCSEditorTreeNode protected interface
+	virtual void RemoveMeAsChild() override;
+	// End of FSCSEditorTreeNode protected interface
+
+	/** Get overridden template component, specialized in given blueprint */
+	UActorComponent* INTERNAL_GetOverridenComponentTemplate(UBlueprint* Blueprint, bool bCreateIfNecessary) const;
+
+private:
+	// Was this component inherited from a parent class or introduced in this class?
+	bool bIsInheritedSCS;
+
+	// Is this the template coming from an SCS node?
+	TWeakObjectPtr<class USCS_Node> SCSNodePtr;
+
+	TWeakObjectPtr<UActorComponent> ComponentTemplatePtr;
+};
+
+
+class KISMET_API FSCSEditorTreeNodeRootActor : public FSCSEditorTreeNode
+{
+public:
+	FSCSEditorTreeNodeRootActor()
+		: FSCSEditorTreeNode(FSCSEditorTreeNode::RootActorNode)
+	{
+	}
+};
+
+class KISMET_API FSCSEditorTreeNodeSeparator : public FSCSEditorTreeNode
+{
+public:
+	FSCSEditorTreeNodeSeparator()
+		: FSCSEditorTreeNode(FSCSEditorTreeNode::SeparatorNode)
+	{
+	}
+};
+
 //////////////////////////////////////////////////////////////////////////
 // SSCS_RowWidget
 
@@ -329,18 +525,23 @@ private:
 	/**
 	 * Retrieves an image brush signifying the specified component's mobility (could sometimes be NULL).
 	 * 
-	 * @param   SCSNode		A weak-pointer to the component node that you want a mobility icon for
 	 * @returns A pointer to the FSlateBrush to use (NULL for Static and Non-SceneComponents)
 	 */
-	FSlateBrush const* GetMobilityIconImage(TWeakObjectPtr<USCS_Node> SCSNode) const;
+	FSlateBrush const* GetMobilityIconImage() const;
 
 	/**
 	 * Retrieves tooltip text describing the specified component's mobility.
 	 * 
-	 * @param   SCSNode		A weak-pointer to the component node that you want tooltip text for
 	 * @returns An FText object containing a description of the component's mobility
 	 */
-	FText GetMobilityToolTipText(TWeakObjectPtr<USCS_Node> SCSNode) const;
+	FText GetMobilityToolTipText() const;
+
+	/**
+	 * Retrieves tooltip text describing where the component was first introduced (for inherited components).
+	 * 
+	 * @returns An FText object containing a description of when the component was first introduced
+	 */
+	FText GetIntroducedInToolTipText() const;
 public:
 	/** Pointer back to owning SCSEditor 2 tool */
 	TWeakPtr<SSCSEditor> SCSEditor;
@@ -527,7 +728,7 @@ public:
 		@param Asset	(In) Optional asset to assign to the component
 		@param bSetFocusToNewItem (In) Select the new item and activate the inline rename widget (default is true)
 		@return The reference of the newly created ActorComponent */
-	UActorComponent* AddNewNode(UActorComponent* NewInstanceComponent, UObject* Asset, bool bSetFocusToNewItem = true);
+	UActorComponent* AddNewNodeForInstancedComponent(UActorComponent* NewInstanceComponent, UObject* Asset, bool bSetFocusToNewItem = true);
 	
 	/** Returns true if the specified component is currently selected */
 	bool IsComponentSelected(const UPrimitiveComponent* PrimComponent) const;
@@ -704,10 +905,10 @@ protected:
 	void OnDuplicateComponent();
 
 	/** Helper method to add a tree node for the given SCS node */
-	FSCSEditorTreeNodePtrType AddTreeNode(USCS_Node* InSCSNode, FSCSEditorTreeNodePtrType InParentNodePtr, const bool bIsInherited);
+	FSCSEditorTreeNodePtrType AddTreeNode(USCS_Node* InSCSNode, FSCSEditorTreeNodePtrType InParentNodePtr, const bool bIsInheritedSCS);
 
 	/** Helper method to add a tree node for the given scene component */
-	FSCSEditorTreeNodePtrType AddTreeNode(USceneComponent* InSceneComponent);
+	FSCSEditorTreeNodePtrType AddTreeNodeFromComponent(USceneComponent* InSceneComponent);
 
 	/** Helper method to recursively find a tree node for the given SCS node starting at the given tree node */
 	FSCSEditorTreeNodePtrType FindTreeNode(const USCS_Node* InSCSNode, FSCSEditorTreeNodePtrType InStartNodePtr = FSCSEditorTreeNodePtrType()) const;
@@ -823,10 +1024,10 @@ private:
 	/** Root set of tree */
 	TArray<FSCSEditorTreeNodePtrType> RootNodes;
 
-	/** Root set of components*/
+	/** Root set of components (contains the root scene component and any non-scene component nodes) */
 	TArray<FSCSEditorTreeNodePtrType> RootComponentNodes;
 
-	/* Root Tree Node*/
+	/* Root Tree Node (for scene components) */
 	TSharedPtr<FSCSEditorTreeNode> RootTreeNode;
 
 	/* Root Tree Node*/
