[INTEGRATE] Change 2424870 by [REDACTED_PERSON][REDACTED_HANDLE] on 2015/01/29 20:02:36

	Refactored component deletion code to allow for optional reparenting of the child nodes.
	Changed the level editor to delete the slected components instead of the actor if it has components selected.  UE-8221

	Fixes for how selection is on tree update: UE-8544 , UE8385

diff --git a/Engine/Source/Editor/Kismet/Private/SSCSEditor.cpp b/Engine/Source/Editor/Kismet/Private/SSCSEditor.cpp
--- a/Engine/Source/Editor/Kismet/Private/SSCSEditor.cpp
+++ b/Engine/Source/Editor/Kismet/Private/SSCSEditor.cpp
@@ -779,7 +779,7 @@ bool FSCSEditorTreeNodeInstanceAddedComponent::IsDefaultSceneRoot() const
 {
 	if (USceneComponent* SceneComponent = Cast<USceneComponent>(GetComponentTemplate()))
 	{
-		return SceneComponent->GetFName() == FComponentEditorUtils::GetDefaultSceneRootVariableName();
+		return SceneComponent->GetFName() == USceneComponent::GetDefaultSceneRootVariableName();
 	}
 
 	return false;
@@ -4179,6 +4179,15 @@ void SSCSEditor::UpdateTree(bool bRegenerateTreeNodes)
 					}
 				}
 			}
+
+			if (GetEditorMode() != EComponentEditorMode::BlueprintSCS)
+			{
+				TArray<FSCSEditorTreeNodePtrType> NewSelectedTreeNodes = SCSTreeWidget->GetSelectedItems();
+				if (NewSelectedTreeNodes.Num() == 0)
+				{
+					SCSTreeWidget->SetItemSelection(GetRootNodes()[0], true);
+				}
+			}
 		}
 
 		// If we have a pending deferred rename request, redirect it to the new tree node
@@ -4967,96 +4976,8 @@ void SSCSEditor::RemoveComponentNode(FSCSEditorTreeNodePtrType InNodePtr)
 				SCSTreeWidget->ClearSelection();
 			}
 
-			ComponentInstance->Modify();
-
-			// Handle removal of a node within the scene component hierarchy
-			USceneComponent* SceneComponent = Cast<USceneComponent>(ComponentInstance);
-			if(SceneComponent != nullptr)
-			{
-				// Find an appropriate child node to promote to this node's position in the hierarchy
-				USceneComponent* ChildToPromote = nullptr;
-				if(SceneComponent->AttachChildren.Num() > 0)
-				{
-					// Start with the first child node
-					ChildToPromote = SceneComponent->AttachChildren[0];
-					check(ChildToPromote != nullptr);
-
-					// Always choose non editor-only child nodes over editor-only child nodes (since we don't want editor-only nodes to end up with non editor-only child nodes)
-					if(ChildToPromote->IsEditorOnly())
-					{
-						for(int32 ChildIndex = 1; ChildIndex < SceneComponent->AttachChildren.Num(); ++ChildIndex)
-						{
-							USceneComponent* Child = SceneComponent->AttachChildren[ChildIndex];
-							if (Child != nullptr && !Child->IsEditorOnly())
-							{
-								ChildToPromote = Child;
-								break;
-							}
-						}
-					}
-				}
-
-				// Handle removal of the root node
-				if(SceneComponent == ActorInstance->GetRootComponent())
-				{
-					// We only promote non editor-only components to root in instanced mode
-					if(ChildToPromote == nullptr || ChildToPromote->IsEditorOnly())
-					{
-						// Construct a new default root component
-						USceneComponent* NewRootComponent = ConstructObject<USceneComponent>(USceneComponent::StaticClass(), ActorInstance, FComponentEditorUtils::GetDefaultSceneRootVariableName(), RF_Transactional);
-						NewRootComponent->Mobility = SceneComponent->Mobility;
-						NewRootComponent->SetWorldLocationAndRotation(SceneComponent->GetComponentLocation(), SceneComponent->GetComponentRotation());
-
-						// Designate the new default root as the child we're promoting
-						ChildToPromote = NewRootComponent;
-					}
-
-					ActorInstance->Modify();
-
-					// Set the selected child node as the new root
-					check(ChildToPromote != nullptr);
-					ActorInstance->SetRootComponent(ChildToPromote);
-				}
-				else    // ...not the root node, so we'll promote the selected child node to this position in its AttachParent's child array.
-				{
-					// Get the node's AttachParent
-					USceneComponent* AttachParent = SceneComponent->AttachParent;
-					check(AttachParent != nullptr);
-
-					// Find the node's position in its AttachParent's child array
-					int32 Index = AttachParent->AttachChildren.Find(SceneComponent);
-					check(Index != INDEX_NONE);
-
-					// Detach the node from its parent
-					SceneComponent->DetachFromParent();
-
-					if(ChildToPromote != nullptr)
-					{
-						// Attach the child node that we're promoting to the parent and move it to the same position as the old node was in the array
-						ChildToPromote->AttachTo(AttachParent, NAME_None, EAttachLocation::KeepWorldPosition);
-						AttachParent->AttachChildren.Remove(ChildToPromote);
-						AttachParent->AttachChildren.Insert(ChildToPromote, Index);
-					}
-				}
-
-				// Detach child nodes from the node that's being removed and re-attach them to the child that's being promoted
-				TArray<USceneComponent*> AttachChildrenLocalCopy(SceneComponent->AttachChildren);
-				for(auto ChildCompIt = AttachChildrenLocalCopy.CreateIterator(); ChildCompIt; ++ChildCompIt)
-				{
-					USceneComponent* Child = *ChildCompIt;
-					check(Child != nullptr);
-
-					// Note: This will internally call Modify(), so we don't need to call it here
-					Child->DetachFromParent(true);
-					if(Child != ChildToPromote)
-					{
-						Child->AttachTo(ChildToPromote, NAME_None, EAttachLocation::KeepWorldPosition);
-					}
-				}
-			}
-
 			// Destroy the component instance
-			ComponentInstance->DestroyComponent();
+			ComponentInstance->DestroyComponent(true);
 		}
 	}
 }
diff --git a/Engine/Source/Editor/UnrealEd/Private/EditorActor.cpp b/Engine/Source/Editor/UnrealEd/Private/EditorActor.cpp
--- a/Engine/Source/Editor/UnrealEd/Private/EditorActor.cpp
+++ b/Engine/Source/Editor/UnrealEd/Private/EditorActor.cpp
@@ -550,6 +550,22 @@ bool UUnrealEdEngine::edactDeleteSelected( UWorld* InWorld, bool bVerifyDeletion
 
 	GetSelectedActors()->Modify();
 
+	if (GetSelectedComponentCount() > 0)
+	{
+		bool bComponentWasDeleted = false;
+		// Delete selected components if they are an Instance component
+		for (FSelectionIterator It(GetSelectedComponentIterator()); It; ++It)
+		{
+			UActorComponent* CurActor = CastChecked<UActorComponent>(*It);
+			if (CurActor->CreationMethod == EComponentCreationMethod::Instance)
+			{
+				CurActor->DestroyComponent(true);
+				bComponentWasDeleted = true;
+			}
+		}
+		return bComponentWasDeleted;
+	}
+
 	// Fire ULevel::LevelDirtiedEvent when falling out of scope.
 	FScopedLevelDirtied			LevelDirtyCallback;
 
diff --git a/Engine/Source/Editor/UnrealEd/Private/Factories/ActorFactory.cpp b/Engine/Source/Editor/UnrealEd/Private/Factories/ActorFactory.cpp
--- a/Engine/Source/Editor/UnrealEd/Private/Factories/ActorFactory.cpp
+++ b/Engine/Source/Editor/UnrealEd/Private/Factories/ActorFactory.cpp
@@ -1184,7 +1184,7 @@ AActor* UActorFactoryEmptyActor::SpawnActor( UObject* Asset, ULevel* InLevel, co
 		// Spawn a temporary actor for dragging around
 		NewActor = Super::SpawnActor(Asset, InLevel, Location, Rotation, ObjectFlags, Name);
 
-		USceneComponent* RootComponent = ConstructObject<USceneComponent>(USceneComponent::StaticClass(), NewActor, FComponentEditorUtils::GetDefaultSceneRootVariableName(), RF_Transactional);
+		USceneComponent* RootComponent = ConstructObject<USceneComponent>(USceneComponent::StaticClass(), NewActor, USceneComponent::GetDefaultSceneRootVariableName(), RF_Transactional);
 		RootComponent->Mobility = EComponentMobility::Movable;
 		RootComponent->SetWorldLocationAndRotation(Location, Rotation);
 		NewActor->SetRootComponent(RootComponent);
diff --git a/Engine/Source/Editor/UnrealEd/Private/Kismet2/ComponentEditorUtils.cpp b/Engine/Source/Editor/UnrealEd/Private/Kismet2/ComponentEditorUtils.cpp
--- a/Engine/Source/Editor/UnrealEd/Private/Kismet2/ComponentEditorUtils.cpp
+++ b/Engine/Source/Editor/UnrealEd/Private/Kismet2/ComponentEditorUtils.cpp
@@ -7,12 +7,6 @@
 #include "Engine/SCS_Node.h"
 #include "Engine/SimpleConstructionScript.h"
 
-const FName& FComponentEditorUtils::GetDefaultSceneRootVariableName()
-{
-	static FName DefaultSceneRootVariableName = FName(TEXT("DefaultSceneRoot"));
-
-	return DefaultSceneRootVariableName;
-}
 
 USceneComponent* FComponentEditorUtils::GetSceneComponent( UObject* Object, UObject* SubObject /*= NULL*/ )
 {
@@ -100,7 +94,7 @@ void FComponentEditorUtils::GetArchetypeInstances( UObject* Object, TArray<UObje
 bool FComponentEditorUtils::IsValidVariableNameString(const UActorComponent* InComponent, const FString& InString)
 {
 	// First test to make sure the string is not empty and does not equate to the DefaultSceneRoot node name
-	bool bIsValid = !InString.IsEmpty() && !InString.Equals(GetDefaultSceneRootVariableName().ToString());
+	bool bIsValid = !InString.IsEmpty() && !InString.Equals(USceneComponent::GetDefaultSceneRootVariableName().ToString());
 	if(bIsValid && InComponent != NULL)
 	{
 		// Next test to make sure the string doesn't conflict with the format that MakeUniqueObjectName() generates
diff --git a/Engine/Source/Editor/UnrealEd/Public/Kismet2/ComponentEditorUtils.h b/Engine/Source/Editor/UnrealEd/Public/Kismet2/ComponentEditorUtils.h
--- a/Engine/Source/Editor/UnrealEd/Public/Kismet2/ComponentEditorUtils.h
+++ b/Engine/Source/Editor/UnrealEd/Public/Kismet2/ComponentEditorUtils.h
@@ -8,9 +8,7 @@
 class UNREALED_API FComponentEditorUtils
 {
 public:
-	/** The name to use for the default scene root variable */
-	static const FName& GetDefaultSceneRootVariableName();
-
+	
 	static class USceneComponent* GetSceneComponent(class UObject* Object, class UObject* SubObject = NULL);
 
 	static void GetArchetypeInstances(class UObject* Object, TArray<class UObject*>& ArchetypeInstances);
diff --git a/Engine/Source/Runtime/Engine/Classes/Components/ActorComponent.h b/Engine/Source/Runtime/Engine/Classes/Components/ActorComponent.h
--- a/Engine/Source/Runtime/Engine/Classes/Components/ActorComponent.h
+++ b/Engine/Source/Runtime/Engine/Classes/Components/ActorComponent.h
@@ -513,7 +513,7 @@ public:
 	void UnregisterComponent();
 
 	/** Unregister the component, remove it from its outer Actor's Components array and mark for pending kill. */
-	virtual void DestroyComponent();
+	virtual void DestroyComponent(bool bPromoteChildren = false);
 
 	/** Called when a component is created (not loaded) */
 	virtual void OnComponentCreated();
diff --git a/Engine/Source/Runtime/Engine/Classes/Components/DecalComponent.h b/Engine/Source/Runtime/Engine/Classes/Components/DecalComponent.h
--- a/Engine/Source/Runtime/Engine/Classes/Components/DecalComponent.h
+++ b/Engine/Source/Runtime/Engine/Classes/Components/DecalComponent.h
@@ -57,6 +57,9 @@ protected:
 	/** Handle for efficient management of DestroyDecalComponent timer */
 	FTimerHandle TimerHandle_DestroyDecalComponent;
 
+	/** Called when the life span of the decal has been exceeded */
+	void LifeSpanCallback();
+
 public:
 	
 	void SetLifeSpan(const float LifeSpan);
@@ -95,5 +98,6 @@ public:
 	// Begin USceneComponent Interface
 	virtual FBoxSphereBounds CalcBounds(const FTransform& LocalToWorld) const override;
 	// End USceneComponent Interface
+
 };
 
diff --git a/Engine/Source/Runtime/Engine/Classes/Components/SceneComponent.h b/Engine/Source/Runtime/Engine/Classes/Components/SceneComponent.h
--- a/Engine/Source/Runtime/Engine/Classes/Components/SceneComponent.h
+++ b/Engine/Source/Runtime/Engine/Classes/Components/SceneComponent.h
@@ -82,6 +82,9 @@ class ENGINE_API USceneComponent : public UActorComponent
 	GENERATED_BODY()
 public:
 
+	/** The name to use for the default scene root variable */
+	static const FName& GetDefaultSceneRootVariableName();
+
 	/**
 	 * UObject constructor that takes an ObjectInitializer
 	 */
@@ -501,6 +504,7 @@ public:
 	// Begin ActorComponent interface
 	virtual void OnRegister() override;
 	virtual void UpdateComponentToWorld(bool bSkipPhysicsMove = false) override final;
+	virtual void DestroyComponent(bool bPromoteChildren = false) override;
 	virtual void OnComponentDestroyed() override;
 	virtual void ApplyWorldOffset(const FVector& InOffset, bool bWorldShift) override;
 	virtual class FComponentInstanceDataBase* GetComponentInstanceData() const override;
diff --git a/Engine/Source/Runtime/Engine/Private/ActorComponent.cpp b/Engine/Source/Runtime/Engine/Private/ActorComponent.cpp
--- a/Engine/Source/Runtime/Engine/Private/ActorComponent.cpp
+++ b/Engine/Source/Runtime/Engine/Private/ActorComponent.cpp
@@ -711,7 +711,7 @@ void UActorComponent::UnregisterComponent()
 	World = NULL;
 }
 
-void UActorComponent::DestroyComponent()
+void UActorComponent::DestroyComponent(bool bPromoteChildren/*= false*/)
 {
 	// Ensure that we call UninitializeComponent before we destroy this component
 	if (bHasBeenInitialized)
diff --git a/Engine/Source/Runtime/Engine/Private/DecalComponent.cpp b/Engine/Source/Runtime/Engine/Private/DecalComponent.cpp
--- a/Engine/Source/Runtime/Engine/Private/DecalComponent.cpp
+++ b/Engine/Source/Runtime/Engine/Private/DecalComponent.cpp
@@ -44,7 +44,7 @@ void UDecalComponent::SetLifeSpan(const float LifeSpan)
 {
 	if (LifeSpan > 0.f)
 	{
-		GetWorld()->GetTimerManager().SetTimer(TimerHandle_DestroyDecalComponent, this, &UDecalComponent::DestroyComponent, LifeSpan, false);
+		GetWorld()->GetTimerManager().SetTimer(TimerHandle_DestroyDecalComponent, this, &UDecalComponent::LifeSpanCallback, LifeSpan, false);
 	}
 	else
 	{
@@ -52,6 +52,11 @@ void UDecalComponent::SetLifeSpan(const float LifeSpan)
 	}
 }
 
+void UDecalComponent::LifeSpanCallback()
+{
+	DestroyComponent();
+}
+
 void UDecalComponent::SetSortOrder(int32 Value)
 {
 	SortOrder = Value;
diff --git a/Engine/Source/Runtime/Engine/Private/SceneComponent.cpp b/Engine/Source/Runtime/Engine/Private/SceneComponent.cpp
--- a/Engine/Source/Runtime/Engine/Private/SceneComponent.cpp
+++ b/Engine/Source/Runtime/Engine/Private/SceneComponent.cpp
@@ -23,6 +23,13 @@ FOverlapInfo::FOverlapInfo(UPrimitiveComponent* InComponent, int32 InBodyIndex)
 	OverlapInfo.Item = InBodyIndex;
 }
 
+const FName& USceneComponent::GetDefaultSceneRootVariableName()
+{
+	static FName DefaultSceneRootVariableName = FName(TEXT("DefaultSceneRoot"));
+
+	return DefaultSceneRootVariableName;
+}
+
 USceneComponent::USceneComponent(const FObjectInitializer& ObjectInitializer /*= FObjectInitializer::Get()*/)
 	: Super(ObjectInitializer)
 {
@@ -249,6 +256,100 @@ void USceneComponent::EndScopedMovementUpdate(class FScopedMovementUpdate& Compl
 }
 
 
+void USceneComponent::DestroyComponent(bool bPromoteChildren/*= false*/)
+{
+	if (bPromoteChildren)
+	{
+		AActor* Owner = GetOwner();
+		if (Owner != NULL)
+		{
+			Owner->Modify();
+
+			// Find an appropriate child node to promote to this node's position in the hierarchy
+			USceneComponent* ChildToPromote = nullptr;
+			if (AttachChildren.Num() > 0)
+			{
+				// Start with the first child node
+				ChildToPromote = AttachChildren[0];
+				check(ChildToPromote != nullptr);
+
+				// Always choose non editor-only child nodes over editor-only child nodes (since we don't want editor-only nodes to end up with non editor-only child nodes)
+				if (ChildToPromote->IsEditorOnly())
+				{
+					for (int32 ChildIndex = 1; ChildIndex < AttachChildren.Num(); ++ChildIndex)
+					{
+						USceneComponent* Child = AttachChildren[ChildIndex];
+						if (Child != nullptr && !Child->IsEditorOnly())
+						{
+							ChildToPromote = Child;
+							break;
+						}
+					}
+				}
+			}
+
+			// Handle removal of the root node
+			if (this == Owner->GetRootComponent())
+			{
+				// We only promote non editor-only components to root in instanced mode
+				if (ChildToPromote == nullptr || ChildToPromote->IsEditorOnly())
+				{
+					// Construct a new default root component
+					USceneComponent* NewRootComponent = ConstructObject<USceneComponent>(USceneComponent::StaticClass(), Owner, USceneComponent::GetDefaultSceneRootVariableName(), RF_Transactional);
+					NewRootComponent->Mobility = Mobility;
+					NewRootComponent->SetWorldLocationAndRotation(GetComponentLocation(), GetComponentRotation());
+
+					// Designate the new default root as the child we're promoting
+					ChildToPromote = NewRootComponent;
+				}
+
+				Owner->Modify();
+
+				// Set the selected child node as the new root
+				check(ChildToPromote != nullptr);
+				Owner->SetRootComponent(ChildToPromote);
+			}
+			else    // ...not the root node, so we'll promote the selected child node to this position in its AttachParent's child array.
+			{
+				// Cache our AttachParent
+				USceneComponent* CachedAttachParent = AttachParent;
+				check(CachedAttachParent != nullptr);
+
+				// Find the our position in its AttachParent's child array
+				int32 Index = CachedAttachParent->AttachChildren.Find(this);
+				check(Index != INDEX_NONE);
+
+				// Detach from parent
+				DetachFromParent(true);
+
+				if (ChildToPromote != nullptr)
+				{
+					// Attach the child node that we're promoting to the parent and move it to the same position as the old node was in the array
+					ChildToPromote->AttachTo(CachedAttachParent, NAME_None, EAttachLocation::KeepWorldPosition);
+					CachedAttachParent->AttachChildren.Remove(ChildToPromote);
+					CachedAttachParent->AttachChildren.Insert(ChildToPromote, Index);
+				}
+			}
+
+			// Detach child nodes from the node that's being removed and re-attach them to the child that's being promoted
+			TArray<USceneComponent*> AttachChildrenLocalCopy(AttachChildren);
+			for (auto ChildCompIt = AttachChildrenLocalCopy.CreateIterator(); ChildCompIt; ++ChildCompIt)
+			{
+				USceneComponent* Child = *ChildCompIt;
+				check(Child != nullptr);
+
+				// Note: This will internally call Modify(), so we don't need to call it here
+				Child->DetachFromParent(true);
+				if (Child != ChildToPromote)
+				{
+					Child->AttachTo(ChildToPromote, NAME_None, EAttachLocation::KeepWorldPosition);
+				}
+			}
+		}
+	}
+	Super::DestroyComponent(bPromoteChildren);
+}
+
 void USceneComponent::OnComponentDestroyed()
 {
 	Super::OnComponentDestroyed();
diff --git a/Engine/Source/Runtime/Engine/Private/SimpleConstructionScript.cpp b/Engine/Source/Runtime/Engine/Private/SimpleConstructionScript.cpp
--- a/Engine/Source/Runtime/Engine/Private/SimpleConstructionScript.cpp
+++ b/Engine/Source/Runtime/Engine/Private/SimpleConstructionScript.cpp
@@ -650,7 +650,7 @@ void USimpleConstructionScript::ValidateSceneRootNodes()
 			&& FBlueprintEditorUtils::IsActorBased(Blueprint)
 			&& Blueprint->BlueprintType != BPTYPE_MacroLibrary)
 		{
-			DefaultSceneRootNode = CreateNode(USceneComponent::StaticClass(), FComponentEditorUtils::GetDefaultSceneRootVariableName());
+			DefaultSceneRootNode = CreateNode(USceneComponent::StaticClass(), USceneComponent::GetDefaultSceneRootVariableName());
 		}
 	}
 
diff --git a/Engine/Source/Runtime/Landscape/Classes/LandscapeComponent.h b/Engine/Source/Runtime/Landscape/Classes/LandscapeComponent.h
--- a/Engine/Source/Runtime/Landscape/Classes/LandscapeComponent.h
+++ b/Engine/Source/Runtime/Landscape/Classes/LandscapeComponent.h
@@ -307,7 +307,7 @@ public:
 	// End UPrimitiveComponent interface.
 
 	// Begin USceneComponent interface.
-	virtual void DestroyComponent() override;
+	virtual void DestroyComponent(bool bPromoteChildren = false) override;
 	virtual FBoxSphereBounds CalcBounds(const FTransform& LocalToWorld) const override;
 	// End USceneComponent interface.
 
diff --git a/Engine/Source/Runtime/Landscape/Classes/LandscapeHeightfieldCollisionComponent.h b/Engine/Source/Runtime/Landscape/Classes/LandscapeHeightfieldCollisionComponent.h
--- a/Engine/Source/Runtime/Landscape/Classes/LandscapeHeightfieldCollisionComponent.h
+++ b/Engine/Source/Runtime/Landscape/Classes/LandscapeHeightfieldCollisionComponent.h
@@ -134,7 +134,7 @@ class ULandscapeHeightfieldCollisionComponent : public UPrimitiveComponent
 	// End UActorComponent interface.
 
 	// Begin USceneComponent interface.
-	virtual void DestroyComponent() override;
+	virtual void DestroyComponent(bool bPromoteChildren = false) override;
 	virtual FBoxSphereBounds CalcBounds(const FTransform &BoundTransform) const override;
 
 	virtual ECollisionEnabled::Type GetCollisionEnabled() const override;
diff --git a/Engine/Source/Runtime/Landscape/Classes/LandscapeMeshCollisionComponent.h b/Engine/Source/Runtime/Landscape/Classes/LandscapeMeshCollisionComponent.h
--- a/Engine/Source/Runtime/Landscape/Classes/LandscapeMeshCollisionComponent.h
+++ b/Engine/Source/Runtime/Landscape/Classes/LandscapeMeshCollisionComponent.h
@@ -74,7 +74,7 @@ class ULandscapeMeshCollisionComponent : public ULandscapeHeightfieldCollisionCo
 	// End UActorComponent interface.
 
 	// Begin USceneComponent interface.
-	virtual void DestroyComponent() override;
+	virtual void DestroyComponent(bool bPromoteChildren = false) override;
 	// End USceneComponent interface.
 
 	// Begin UPrimitiveComponent interface
diff --git a/Engine/Source/Runtime/Landscape/Private/Landscape.cpp b/Engine/Source/Runtime/Landscape/Private/Landscape.cpp
--- a/Engine/Source/Runtime/Landscape/Private/Landscape.cpp
+++ b/Engine/Source/Runtime/Landscape/Private/Landscape.cpp
@@ -776,7 +776,7 @@ FPrimitiveSceneProxy* ULandscapeComponent::CreateSceneProxy()
 	return Proxy;
 }
 
-void ULandscapeComponent::DestroyComponent()
+void ULandscapeComponent::DestroyComponent(bool bPromoteChildren/*= false*/)
 {
 	ALandscapeProxy* Proxy = GetLandscapeProxy();
 	if (Proxy)
diff --git a/Engine/Source/Runtime/Landscape/Private/LandscapeCollision.cpp b/Engine/Source/Runtime/Landscape/Private/LandscapeCollision.cpp
--- a/Engine/Source/Runtime/Landscape/Private/LandscapeCollision.cpp
+++ b/Engine/Source/Runtime/Landscape/Private/LandscapeCollision.cpp
@@ -860,7 +860,7 @@ void ULandscapeMeshCollisionComponent::ApplyWorldOffset(const FVector& InOffset,
 	}
 }
 
-void ULandscapeMeshCollisionComponent::DestroyComponent()
+void ULandscapeMeshCollisionComponent::DestroyComponent(bool bPromoteChildren/*= false*/)
 {
 	ALandscapeProxy* Proxy = GetLandscapeProxy();
 	if (Proxy)
@@ -960,7 +960,7 @@ void ULandscapeHeightfieldCollisionComponent::UpdateHeightfieldRegion(int32 Comp
 }
 #endif// WITH_EDITOR
 
-void ULandscapeHeightfieldCollisionComponent::DestroyComponent()
+void ULandscapeHeightfieldCollisionComponent::DestroyComponent(bool bPromoteChildren/*= false*/)
 {
 	ALandscapeProxy* Proxy = GetLandscapeProxy();
 	if (Proxy)
diff --git a/Engine/Source/Runtime/UMG/Private/Components/WidgetComponent.cpp b/Engine/Source/Runtime/UMG/Private/Components/WidgetComponent.cpp
--- a/Engine/Source/Runtime/UMG/Private/Components/WidgetComponent.cpp
+++ b/Engine/Source/Runtime/UMG/Private/Components/WidgetComponent.cpp
@@ -435,7 +435,7 @@ void UWidgetComponent::OnUnregister()
 	Super::OnUnregister();
 }
 
-void UWidgetComponent::DestroyComponent()
+void UWidgetComponent::DestroyComponent(bool bPromoteChildren/*= false*/)
 {
 	Super::DestroyComponent();
 
diff --git a/Engine/Source/Runtime/UMG/Public/Components/WidgetComponent.h b/Engine/Source/Runtime/UMG/Public/Components/WidgetComponent.h
--- a/Engine/Source/Runtime/UMG/Public/Components/WidgetComponent.h
+++ b/Engine/Source/Runtime/UMG/Public/Components/WidgetComponent.h
@@ -30,7 +30,7 @@ public:
 	virtual FCollisionShape GetCollisionShape(float Inflation) const override;
 	virtual void OnRegister() override;
 	virtual void OnUnregister() override;
-	virtual void DestroyComponent() override;
+	virtual void DestroyComponent(bool bPromoteChildren = false) override;
 	virtual void TickComponent(float DeltaTime, enum ELevelTick TickType, FActorComponentTickFunction *ThisTickFunction) override;
 
 	virtual FComponentInstanceDataBase* GetComponentInstanceData() const override;
