Build fix after integration. Not previously integrated in 2418846.

diff --git a/Engine/Source/Editor/Kismet/Private/BlueprintEditor.cpp b/Engine/Source/Editor/Kismet/Private/BlueprintEditor.cpp
--- a/Engine/Source/Editor/Kismet/Private/BlueprintEditor.cpp
+++ b/Engine/Source/Editor/Kismet/Private/BlueprintEditor.cpp
@@ -771,24 +771,15 @@ void FBlueprintEditor::OnSelectionUpdated(const TArray<FSCSEditorTreeNodePtrType
 				}
 				else
 				{
-				UActorComponent* EditableComponent = nullptr;
-				if (NodePtr->CanEditDefaults())
-				{
-					EditableComponent = NodePtr->GetComponentTemplate();
-				}
-				else if (!NodePtr->IsNative() && NodePtr->IsInherited())
-				{
-					EditableComponent = NodePtr->GetOverridenComponentTemplate(GetBlueprintObj(), true);
-				}
-
-				if (EditableComponent)
-				{
-					InspectorTitle = FText::FromString(NodePtr->GetDisplayString());
-					InspectorObjects.Add(EditableComponent);
+					UActorComponent* EditableComponent = NodePtr->GetEditableComponentTemplate(GetBlueprintObj());
+					if (EditableComponent)
+					{
+						InspectorTitle = FText::FromString(NodePtr->GetDisplayString());
+						InspectorObjects.Add(EditableComponent);
+					}
 				}
 			}
 		}
-		}
 
 		// Update the details panel
 		SKismetInspector::FShowDetailsOptions Options(InspectorTitle, true);
