[INTEGRATE] Change 2424943 by [REDACTED_PERSON][REDACTED_HANDLE] on 2015/01/29 20:59:00

	Component editor: added a virtual dtor to FSCSEditorTreeNode

diff --git a/Engine/Source/Editor/Kismet/Public/SSCSEditor.h b/Engine/Source/Editor/Kismet/Public/SSCSEditor.h
--- a/Engine/Source/Editor/Kismet/Public/SSCSEditor.h
+++ b/Engine/Source/Editor/Kismet/Public/SSCSEditor.h
@@ -170,6 +170,9 @@ public:
 	 * Creates the correct type of node based on the component (instanced or not, etc...)
 	 */
 	static FSCSEditorTreeNodePtrType FactoryNodeFromComponent(UActorComponent* InComponent);
+
+	// Destructor
+	virtual ~FSCSEditorTreeNode() {}
 protected:
 	// Called when this node is being removed via a RemoveChild call
 	virtual void RemoveMeAsChild() {}
