[INTEGRATE] Change 2425146 by [REDACTED_PERSON][REDACTED_HANDLE] on 2015/01/29 23:38:19

	Blueprints: Fix up construction script when reparenting a Blueprint, in case it gains or loses a root native component [UE-8685]

diff --git a/Engine/Source/Editor/Kismet/Private/BlueprintEditor.cpp b/Engine/Source/Editor/Kismet/Private/BlueprintEditor.cpp
--- a/Engine/Source/Editor/Kismet/Private/BlueprintEditor.cpp
+++ b/Engine/Source/Editor/Kismet/Private/BlueprintEditor.cpp
@@ -1270,6 +1270,13 @@ void FBlueprintEditor::EnsureBlueprintIsUpToDate(UBlueprint* BlueprintObj)
 			FBlueprintEditorUtils::AddFunctionGraph(BlueprintObj, UCSGraph, /*bIsUserCreated=*/ false, AActor::StaticClass());
 			UCSGraph->bAllowDeletion = false;
 		}
+
+		// Check to see if we have gained a component from our parent (that would require us removing our scene root)
+		// (or lost one, which requires adding one)
+		if (BlueprintObj->SimpleConstructionScript != nullptr)
+		{
+			BlueprintObj->SimpleConstructionScript->ValidateSceneRootNodes();
+		}
 	}
 	else
 	{
@@ -2466,6 +2473,9 @@ void FBlueprintEditor::ReparentBlueprint_NewParentChosen(UClass* ChosenClass)
 
 			Compile();
 
+			// Ensure that the Blueprint is up-to-date (valid SCS etc.) after compiling (new parent class)
+			EnsureBlueprintIsUpToDate(BlueprintObj);
+
 			if (SCSEditor.IsValid())
 			{
 				SCSEditor->UpdateTree();
diff --git a/Engine/Source/Runtime/Engine/Classes/Engine/SimpleConstructionScript.h b/Engine/Source/Runtime/Engine/Classes/Engine/SimpleConstructionScript.h
--- a/Engine/Source/Runtime/Engine/Classes/Engine/SimpleConstructionScript.h
+++ b/Engine/Source/Runtime/Engine/Classes/Engine/SimpleConstructionScript.h
@@ -64,9 +64,8 @@ class USimpleConstructionScript : public UObject
 	/** Find the SCS_Node node by name and return it if found */
 	ENGINE_API USCS_Node* FindSCSNodeByGuid(FGuid Guid);
 
-protected:
 	/** Checks the root node set for scene components and ensures that it is valid (e.g. after a removal) */
-	void ValidateSceneRootNodes();
+	ENGINE_API void ValidateSceneRootNodes();
 
 private:
 	/** Root nodes of the construction script */
