[INTEGRATE] Change 2425195 by [REDACTED_PERSON][REDACTED_HANDLE] on 2015/01/30 00:34:33

	Component tree: Allow drag-dropping C++ classes from the CB to the components tree in addition to Blueprinted components (and add validation for abstract/deprecated and BlueprintSpawnableComponent) [UE-8460]

diff --git a/Engine/Source/Editor/Kismet/Private/SSCSEditor.cpp b/Engine/Source/Editor/Kismet/Private/SSCSEditor.cpp
--- a/Engine/Source/Editor/Kismet/Private/SSCSEditor.cpp
+++ b/Engine/Source/Editor/Kismet/Private/SSCSEditor.cpp
@@ -86,12 +86,32 @@ FReply SSCSEditorDragDropTree::OnDragOver( const FGeometry& MyGeometry, const FD
 {
 	FReply Handled = FReply::Unhandled();
 
-	if ( SCSEditor != NULL )
+	if (SCSEditor != nullptr)
 	{
-		const bool bIsValidDrag = DragDropEvent.GetOperationAs<FExternalDragOperation>().IsValid();
-		if (bIsValidDrag)
+		TSharedPtr<FDragDropOperation> Operation = DragDropEvent.GetOperation();
+		if (Operation.IsValid() && (Operation->IsOfType<FExternalDragOperation>() || Operation->IsOfType<FAssetDragDropOp>()))
 		{
 			Handled = AssetUtil::CanHandleAssetDrag(DragDropEvent);
+
+			if (!Handled.IsEventHandled())
+			{
+				if (Operation->IsOfType<FAssetDragDropOp>())
+				{
+					const auto& AssetDragDropOp = StaticCastSharedPtr<FAssetDragDropOp>(Operation);
+
+					for (const FAssetData& AssetData : AssetDragDropOp->AssetData)
+					{
+						if (UClass* AssetClass = AssetData.GetClass())
+						{
+							if (AssetClass->IsChildOf(UClass::StaticClass()))
+							{
+								Handled = FReply::Handled();
+								break;
+							}
+						}
+					}
+				}
+			}
 		}
 	}
 
@@ -102,7 +122,7 @@ FReply SSCSEditorDragDropTree::OnDrop( const FGeometry& MyGeometry, const FDragD
 {
 	FReply Handled = FReply::Unhandled();
 
-	if ( SCSEditor != NULL )
+	if (SCSEditor != nullptr)
 	{
 		TSharedPtr<FDragDropOperation> Operation = DragDropEvent.GetOperation();
 		if (Operation.IsValid() && (Operation->IsOfType<FExternalDragOperation>() || Operation->IsOfType<FAssetDragDropOp>()))
@@ -110,24 +130,42 @@ FReply SSCSEditorDragDropTree::OnDrop( const FGeometry& MyGeometry, const FDragD
 			TArray< FAssetData > DroppedAssetData = AssetUtil::ExtractAssetDataFromDrag( DragDropEvent );
 			const int32 NumAssets = DroppedAssetData.Num();
 
-			if ( NumAssets > 0 )
+			if (NumAssets > 0)
 			{
 				GWarn->BeginSlowTask( LOCTEXT("LoadingComponents", "Loading Component(s)"), true );
 
 				for (int32 DroppedAssetIdx = 0; DroppedAssetIdx < NumAssets; ++DroppedAssetIdx)
 				{
 					const FAssetData& AssetData = DroppedAssetData[DroppedAssetIdx];
+
+					if (!AssetData.IsAssetLoaded())
+					{
+						GWarn->StatusUpdate(DroppedAssetIdx, NumAssets, FText::Format(LOCTEXT("LoadingAsset", "Loading Asset {0}"), FText::FromName(AssetData.AssetName)));
+					}
+
 					UClass* AssetClass = AssetData.GetClass();
-					UBlueprint* BPClass = Cast<UBlueprint>( AssetData.GetAsset() );
-					TSubclassOf<UActorComponent>  ComponentClasses = FComponentAssetBrokerage::GetPrimaryComponentForAsset( AssetClass );
-					if ( NULL != ComponentClasses )
+					UObject* Asset = AssetData.GetAsset();
+
+					UBlueprint* BPClass = Cast<UBlueprint>(Asset);
+					UClass* PotentialComponentClass = nullptr;
+					
+					if ((BPClass != nullptr) && (BPClass->GeneratedClass != nullptr) && (BPClass->GeneratedClass->IsChildOf(UActorComponent::StaticClass())))
+					{
+						PotentialComponentClass = BPClass->GeneratedClass;
+					}
+					else if (AssetClass->IsChildOf(UClass::StaticClass()) && CastChecked<UClass>(Asset)->IsChildOf(UActorComponent::StaticClass()))
+					{
+						PotentialComponentClass = Cast<UClass>(Asset);
+					}
+
+					TSubclassOf<UActorComponent>  MatchingComponentClassForAsset = FComponentAssetBrokerage::GetPrimaryComponentForAsset( AssetClass );
+					if (MatchingComponentClassForAsset != nullptr)
 					{
-						GWarn->StatusUpdate( DroppedAssetIdx, NumAssets, FText::Format( LOCTEXT("LoadingComponent", "Loading Component {0}"), FText::FromName( AssetData.AssetName ) ) );
-						SCSEditor->AddNewComponent(ComponentClasses, AssetData.GetAsset() );
+						SCSEditor->AddNewComponent(MatchingComponentClassForAsset, Asset);
 					}
-					else if( BPClass && BPClass->GeneratedClass && BPClass->GeneratedClass->IsChildOf( UActorComponent::StaticClass() ) )
+					else if ((PotentialComponentClass != nullptr) && (!PotentialComponentClass->HasAnyClassFlags(CLASS_Deprecated | CLASS_Abstract | CLASS_NewerVersionExists)) && PotentialComponentClass->HasMetaData(FBlueprintMetadata::MD_BlueprintSpawnableComponent))
 					{
-						SCSEditor->AddNewComponent(BPClass->GeneratedClass, nullptr );
+						SCSEditor->AddNewComponent(PotentialComponentClass, nullptr);
 					}
 				}
 
