Fixed naming convention issues when adding new variables or event dispatchers to Blueprints. They will now be formatted like other items in Blueprints: New[Type]_[Number].

diff --git a/Engine/Source/Editor/Kismet/Private/BlueprintEditor.cpp b/Engine/Source/Editor/Kismet/Private/BlueprintEditor.cpp
--- a/Engine/Source/Editor/Kismet/Private/BlueprintEditor.cpp
+++ b/Engine/Source/Editor/Kismet/Private/BlueprintEditor.cpp
@@ -6273,20 +6273,10 @@ void FBlueprintEditor::OnAddNewVariable()
 {
 	const FScopedTransaction Transaction( LOCTEXT("AddVariable", "Add Variable") );
 
-	FString VarNameString = TEXT("NewVar");
-	FName VarName = FName(*VarNameString);
-
 	// Reset MyBlueprint item filter so new variable is visible
 	MyBlueprintWidget->OnResetItemFilter();
 
-	// Make sure the new name is valid
-	TSharedPtr<INameValidatorInterface> NameValidator = MakeShareable(new FKismetNameValidator(GetBlueprintObj()));
-	int32 Index = 0;
-	while (NameValidator->IsValid(VarName) != Ok)
-	{
-		VarName = FName(*FString::Printf(TEXT("%s%i"), *VarNameString, Index));
-		++Index;
-	}
+	FName VarName = FBlueprintEditorUtils::FindUniqueKismetName(GetBlueprintObj(), TEXT("NewVar"));
 
 	bool bSuccess = MyBlueprintWidget.IsValid() && FBlueprintEditorUtils::AddMemberVariable(GetBlueprintObj(), VarName, MyBlueprintWidget->GetLastPinTypeUsed());
 
@@ -6330,16 +6320,8 @@ void FBlueprintEditor::OnAddNewDelegate()
 	// Reset MyBlueprint item filter so new variable is visible
 	MyBlueprintWidget->OnResetItemFilter();
 
-	const FString NameString = TEXT("NewEventDispatcher");
-	FName Name = FName(*NameString);
-	TArray<FName> Variables;
-	FBlueprintEditorUtils::GetClassVariableList(Blueprint, Variables);
-	int32 Index = 0;
-	while (Variables.Contains(Name) || !FBlueprintEditorUtils::IsGraphNameUnique(Blueprint, Name))
-	{
-		Name = FName(*FString::Printf(TEXT("%s%i"), *NameString, Index));
-		++Index;
-	}
+	FName Name = FBlueprintEditorUtils::FindUniqueKismetName(GetBlueprintObj(), TEXT("NewEventDispatcher"));
+
 
 	const FScopedTransaction Transaction( LOCTEXT("AddNewDelegate", "Add New Event Dispatcher") ); 
 	Blueprint->Modify();
