Editing multiple Math Expression nodes will no longer cause a crash.
Changing the math expression will no longer mark the Blueprint structurally modified.

Kismet Inspector will now do a deferred refresh, waiting until the next tick to assign the new objects to the detail view.

diff --git a/Engine/Source/Editor/BlueprintGraph/Private/K2Node_MathExpression.cpp b/Engine/Source/Editor/BlueprintGraph/Private/K2Node_MathExpression.cpp
--- a/Engine/Source/Editor/BlueprintGraph/Private/K2Node_MathExpression.cpp
+++ b/Engine/Source/Editor/BlueprintGraph/Private/K2Node_MathExpression.cpp
@@ -2397,7 +2397,9 @@ void UK2Node_MathExpression::RebuildExpression(FString InExpression)
 
 		// finally, recompile
 		UBlueprint* Blueprint = FBlueprintEditorUtils::FindBlueprintForNodeChecked(this);
-		FBlueprintEditorUtils::MarkBlueprintAsStructurallyModified(Blueprint);
+		FBlueprintEditorUtils::MarkBlueprintAsModified(Blueprint);
+		// The UI needs a refresh, so notify any interested parties that the blueprint has changed
+		Blueprint->BroadcastChanged();
 	}
 }
 
diff --git a/Engine/Source/Editor/Kismet/Private/SKismetInspector.cpp b/Engine/Source/Editor/Kismet/Private/SKismetInspector.cpp
--- a/Engine/Source/Editor/Kismet/Private/SKismetInspector.cpp
+++ b/Engine/Source/Editor/Kismet/Private/SKismetInspector.cpp
@@ -36,6 +36,18 @@ public:
 //////////////////////////////////////////////////////////////////////////
 // SKismetInspector
 
+void SKismetInspector::Tick( const FGeometry& AllottedGeometry, const double InCurrentTime, const float InDeltaTime )
+{
+	if(bRefreshOnTick)
+	{
+		FKismetSelectionInfo SelectionInfo;
+		UpdateFromObjects(RefreshPropertyObjects, SelectionInfo, RefreshOptions);
+		RefreshPropertyObjects.Empty();
+
+		bRefreshOnTick = false;
+	}
+}
+
 TSharedRef<SWidget> SKismetInspector::MakeContextualEditingWidget(struct FKismetSelectionInfo& SelectionInfo, const FShowDetailsOptions& Options)
 {
 	TSharedRef< SVerticalBox > ContextualEditingWidget = SNew( SVerticalBox );
@@ -342,8 +354,10 @@ void SKismetInspector::ShowDetailsForObjects(const TArray<UObject*>& PropertyObj
 		bIsReentrant = false;
 	}
 
-	FKismetSelectionInfo SelectionInfo;
-	UpdateFromObjects(PropertyObjects, SelectionInfo, Options);
+	// Refresh is being deferred until the next tick, this prevents batch operations from bombarding the details view with calls to refresh
+	RefreshPropertyObjects = PropertyObjects;
+	RefreshOptions = Options;
+	bRefreshOnTick = true;
 }
 
 void SKismetInspector::AddPropertiesRecursive(UProperty* Property)
diff --git a/Engine/Source/Editor/Kismet/Public/SKismetInspector.h b/Engine/Source/Editor/Kismet/Public/SKismetInspector.h
--- a/Engine/Source/Editor/Kismet/Public/SKismetInspector.h
+++ b/Engine/Source/Editor/Kismet/Public/SKismetInspector.h
@@ -51,6 +51,10 @@ public:
 		{}
 	};
 
+	// SWidget interface
+	virtual void Tick( const FGeometry& AllottedGeometry, const double InCurrentTime, const float InDeltaTime ) override;
+	// End of SWidget interface
+
 	/** Update the inspector window to show information on the supplied object */
 	void ShowDetailsForSingleObject(UObject* Object, const FShowDetailsOptions& Options = FShowDetailsOptions());
 
@@ -120,6 +124,15 @@ protected:
 	/** User defined delegate for OnFinishedChangingProperties */
 	FOnFinishedChangingProperties::FDelegate UserOnFinishedChangingProperties;
 
+	/** When TRUE, the Kismet inspector needs to refresh the details view on Tick */
+	bool bRefreshOnTick;
+
+	/** Holds the property objects that need to be displayed by the inspector starting on the next tick */
+	TArray<UObject*> RefreshPropertyObjects;
+
+	/** Details options that are used by the inspector on the next refresh. */
+	FShowDetailsOptions RefreshOptions;
+
 protected:
 	/** Show properties of the selected object */
 	void SetPropertyWindowContents(TArray<UObject*> Objects);
