In the non-unified editor view, Local Variables now appear in the correct section.

diff --git a/Engine/Source/Editor/Kismet/Private/SMyBlueprint.cpp b/Engine/Source/Editor/Kismet/Private/SMyBlueprint.cpp
--- a/Engine/Source/Editor/Kismet/Private/SMyBlueprint.cpp
+++ b/Engine/Source/Editor/Kismet/Private/SMyBlueprint.cpp
@@ -903,10 +903,7 @@ void SMyBlueprint::GetLocalVariables(FGraphActionListBuilderBase& OutAllActions)
 				{
 					TSharedPtr<FEdGraphSchemaAction_K2LocalVar> NewVarAction = MakeShareable(new FEdGraphSchemaAction_K2LocalVar(Category, FText::FromName(Variable.VarName), TEXT(""), 0));
 					NewVarAction->SetVariableInfo(Variable.VarName, Func);
-					if ( GetDefault<UEditorExperimentalSettings>()->bUnifiedBlueprintEditor )
-					{
-						NewVarAction->SectionID = NodeSectionID::LOCAL_VARIABLE;
-					}
+					NewVarAction->SectionID = NodeSectionID::LOCAL_VARIABLE;
 					OutAllActions.AddAction(NewVarAction);
 				}
 			}
