Fixed stripping of FText information from graph description in graph search data in the find in blueprint manager.

diff --git a/Engine/Source/Editor/Kismet/Private/FindInBlueprintManager.cpp b/Engine/Source/Editor/Kismet/Private/FindInBlueprintManager.cpp
--- a/Engine/Source/Editor/Kismet/Private/FindInBlueprintManager.cpp
+++ b/Engine/Source/Editor/Kismet/Private/FindInBlueprintManager.cpp
@@ -514,10 +514,10 @@ namespace BlueprintSearchMetaDataHelpers
 					Graph->GetSchema()->GetGraphDisplayInformation(*Graph, DisplayInfo);
 					InWriter->WriteValue(FFindInBlueprintSearchTags::FiB_Name, DisplayInfo.PlainName);
 
-					FString GraphDescription = FBlueprintEditorUtils::GetGraphDescription(Graph).ToString();
+					FText GraphDescription = FBlueprintEditorUtils::GetGraphDescription(Graph);
 					if(!GraphDescription.IsEmpty())
 					{
-						InWriter->WriteValue(FFindInBlueprintSearchTags::FiB_Description, FText::FromString(GraphDescription));
+						InWriter->WriteValue(FFindInBlueprintSearchTags::FiB_Description, GraphDescription);
 					}
 					// All nodes will appear as children to the graph in search results
 					GatherNodesFromGraph(InWriter, Graph);
