Fix for the defaults details dissapearing or being incorrectly displayed when the blueprint editor refreshes.
#UE-8568 - Invalid display of class defaults after compiling
#Branch UE4
#Proj Kismet

#info FBlueprintEditor::StartEditingDefaults has  been modified to handle the Single Layout BP editor in addition to the legacy version and is now called in call sites that were previously handling defaults specifically for the Single Layout Editor. The problem previously was that refresh editors was calling StartEditingDefaults which failed to handle default editing in the single layout bp editor.

diff --git a/Engine/Source/Editor/Kismet/Private/BlueprintEditor.cpp b/Engine/Source/Editor/Kismet/Private/BlueprintEditor.cpp
--- a/Engine/Source/Editor/Kismet/Private/BlueprintEditor.cpp
+++ b/Engine/Source/Editor/Kismet/Private/BlueprintEditor.cpp
@@ -2535,47 +2535,7 @@ void FBlueprintEditor::EditClassDefaults_Clicked()
 	if ( Blueprint != nullptr )
 	{
 		SetUISelectionState(FBlueprintEditor::SelectionState_ClassDefaults);
-
-		const bool bForceRefresh = true;
-
-		if ( IsEditingSingleBlueprint() )
-		{
-			if ( GetBlueprintObj()->GeneratedClass != NULL )
-			{
-				UObject* DefaultObject = GetBlueprintObj()->GeneratedClass->GetDefaultObject();
-
-				// Update the details panel
-				FString Title;
-				DefaultObject->GetName(Title);
-				SKismetInspector::FShowDetailsOptions Options(FText::FromString(Title), true);
-				Options.bShowComponents = false;
-
-				Inspector->ShowDetailsForSingleObject(DefaultObject, Options);
-
-				TSharedPtr<SDockTab> OwnerTab = Inspector->GetOwnerTab();
-				if ( OwnerTab.IsValid() )
-				{
-					OwnerTab->ActivateInParent(ETabActivationCause::SetDirectly);
-					OwnerTab->FlashTab();
-				}
-			}
-		}
-		else if ( GetEditingObjects().Num() > 0 )
-		{
-			TArray<UObject*> DefaultObjects;
-			for ( int32 i = 0; i < GetEditingObjects().Num(); ++i )
-			{
-				auto Blueprint = (UBlueprint*)( GetEditingObjects()[i] );
-				if ( Blueprint->GeneratedClass )
-				{
-					DefaultObjects.Add(Blueprint->GeneratedClass->GetDefaultObject());
-				}
-			}
-			if ( DefaultObjects.Num() )
-			{
-				DefaultEditor->ShowDetailsForObjects(DefaultObjects);
-			}
-		}
+		StartEditingDefaults( true, true );
 	}
 }
 
@@ -6232,11 +6192,35 @@ void FBlueprintEditor::OnRepairCorruptedBlueprint()
 
 void FBlueprintEditor::StartEditingDefaults(bool bAutoFocus, bool bForceRefresh)
 {
+	const bool bSingleLayoutBPEditor = GetDefault<UEditorExperimentalSettings>()->bUnifiedBlueprintEditor && IsDetailsPanelEditingClassDefaults();
+
 	if (IsEditingSingleBlueprint())
 	{
 		if (GetBlueprintObj()->GeneratedClass != NULL)
 		{
-			DefaultEditor->ShowDetailsForSingleObject(GetBlueprintObj()->GeneratedClass->GetDefaultObject(), SKismetInspector::FShowDetailsOptions(DefaultEditString(), bForceRefresh));
+			if (bSingleLayoutBPEditor)
+			{
+				UObject* DefaultObject = GetBlueprintObj()->GeneratedClass->GetDefaultObject();
+
+				// Update the details panel
+				FString Title;
+				DefaultObject->GetName(Title);
+				SKismetInspector::FShowDetailsOptions Options(FText::FromString(Title), true);
+				Options.bShowComponents = false;
+
+				Inspector->ShowDetailsForSingleObject(DefaultObject, Options);
+
+				TSharedPtr<SDockTab> OwnerTab = Inspector->GetOwnerTab();
+				if (OwnerTab.IsValid())
+				{
+					OwnerTab->ActivateInParent(ETabActivationCause::SetDirectly);
+					OwnerTab->FlashTab();
+				}
+			}
+			else
+			{
+				DefaultEditor->ShowDetailsForSingleObject(GetBlueprintObj()->GeneratedClass->GetDefaultObject(), SKismetInspector::FShowDetailsOptions(DefaultEditString(), bForceRefresh));
+			}
 		}
 	}
 	else if (GetEditingObjects().Num() > 0)
