Force a refresh of the MyBlueprint window before sending the call to rename during various graph panel actions that create new actions in the MyBlueprint window.

diff --git a/Engine/Source/Editor/Kismet/Private/BlueprintEditor.cpp b/Engine/Source/Editor/Kismet/Private/BlueprintEditor.cpp
--- a/Engine/Source/Editor/Kismet/Private/BlueprintEditor.cpp
+++ b/Engine/Source/Editor/Kismet/Private/BlueprintEditor.cpp
@@ -4132,8 +4132,7 @@ void FBlueprintEditor::OnCollapseSelectionToFunction()
 
 		FBlueprintEditorUtils::MarkBlueprintAsStructurallyModified( BlueprintObj );
 
-		MyBlueprintWidget->SelectItemByName(FunctionGraph->GetFName(),ESelectInfo::OnMouseClick);
-		MyBlueprintWidget->OnRequestRenameOnActionNode();
+		RenameNewlyAddedAction(FunctionGraph->GetFName());
 	}
 }
 
@@ -4283,8 +4282,7 @@ void FBlueprintEditor::OnCollapseSelectionToMacro()
 
 		FBlueprintEditorUtils::MarkBlueprintAsStructurallyModified( BlueprintObj );
 
-		MyBlueprintWidget->SelectItemByName(MacroGraph->GetFName(),ESelectInfo::OnMouseClick);
-		MyBlueprintWidget->OnRequestRenameOnActionNode();
+		RenameNewlyAddedAction(MacroGraph->GetFName());
 	}
 }
 
@@ -6248,6 +6246,8 @@ void FBlueprintEditor::RenameNewlyAddedAction(FName InActionName)
 
 	if (MyBlueprintWidget.IsValid())
 	{
+		// Force a refresh immediately, the item has to be present in the list for the rename requests to be successful.
+		MyBlueprintWidget->Refresh();
 		MyBlueprintWidget->SelectItemByName(InActionName,ESelectInfo::OnMouseClick);
 		MyBlueprintWidget->OnRequestRenameOnActionNode();
 	}
