Added registration flag for SBlueprintSubPalette to prevent multiple refreshes in a single frame

diff --git a/Engine/Source/Editor/Kismet/Private/SBlueprintSubPalette.cpp b/Engine/Source/Editor/Kismet/Private/SBlueprintSubPalette.cpp
--- a/Engine/Source/Editor/Kismet/Private/SBlueprintSubPalette.cpp
+++ b/Engine/Source/Editor/Kismet/Private/SBlueprintSubPalette.cpp
@@ -195,6 +195,7 @@ SBlueprintSubPalette::~SBlueprintSubPalette()
 //------------------------------------------------------------------------------
 void SBlueprintSubPalette::Construct(FArguments const& InArgs, TWeakPtr<FBlueprintEditor> InBlueprintEditor)
 {
+	bIsActiveTimerRegistered = false;
 	BlueprintEditorPtr = InBlueprintEditor;
 
 	struct LocalUtils
@@ -259,6 +260,7 @@ void SBlueprintSubPalette::Construct(FArguments const& InArgs, TWeakPtr<FBluepri
 EActiveTimerReturnType SBlueprintSubPalette::TriggerRefreshActionsList(double InCurrentTime, float InDeltaTime)
 {
 	RefreshActionsList(true);
+	bIsActiveTimerRegistered = false;
 	return EActiveTimerReturnType::Stop;
 }
 
@@ -390,7 +392,11 @@ void SBlueprintSubPalette::GenerateContextMenuEntries(FMenuBuilder& MenuBuilder)
 //------------------------------------------------------------------------------
 void SBlueprintSubPalette::RequestRefreshActionsList()
 {
-	RegisterActiveTimer(0.f, FWidgetActiveTimerDelegate::CreateSP(this, &SBlueprintSubPalette::TriggerRefreshActionsList));
+	if (!bIsActiveTimerRegistered)
+	{
+		bIsActiveTimerRegistered = true;
+		RegisterActiveTimer(0.f, FWidgetActiveTimerDelegate::CreateSP(this, &SBlueprintSubPalette::TriggerRefreshActionsList));
+	}
 }
 
 //------------------------------------------------------------------------------
diff --git a/Engine/Source/Editor/Kismet/Private/SBlueprintSubPalette.h b/Engine/Source/Editor/Kismet/Private/SBlueprintSubPalette.h
--- a/Engine/Source/Editor/Kismet/Private/SBlueprintSubPalette.h
+++ b/Engine/Source/Editor/Kismet/Private/SBlueprintSubPalette.h
@@ -109,7 +109,11 @@ private:
 	 */
 	TSharedPtr<SWidget> ConstructContextMenuWidget() const;
 
+private:
 	/** Pointer to the command list created for this (so multiple sub-palettes can have their own bindings)*/
 	TSharedPtr<FUICommandList> CommandList;
+
+	/** Whether the active timer to refresh the actions list is currently registered */
+	bool bIsActiveTimerRegistered;
 };
 
