Selecting local variables will display details again.

Completely removed the defunct local variable menu which used to display local variables. Should reduce the likelihood of further local variable issues from occuring.

diff --git a/Engine/Source/Editor/Kismet/Private/SMyBlueprint.cpp b/Engine/Source/Editor/Kismet/Private/SMyBlueprint.cpp
--- a/Engine/Source/Editor/Kismet/Private/SMyBlueprint.cpp
+++ b/Engine/Source/Editor/Kismet/Private/SMyBlueprint.cpp
@@ -371,24 +371,9 @@ void SMyBlueprint::Construct(const FArguments& InArgs, TWeakPtr<FBlueprintEditor
 		+ SVerticalBox::Slot()
 		.FillHeight(1.0f)
 		[
-			SAssignNew(ActionMenuContainer, SSplitter)
-			.Orientation(Orient_Vertical)
-
-			+ SSplitter::Slot()
-			[
-				GraphActionMenu.ToSharedRef()
-			]
+			GraphActionMenu.ToSharedRef()
 		]
 	];
-
-	if (GetLocalActionsListVisibility() == EVisibility::Visible)
-	{
-		ActionMenuContainer->AddSlot()
-			.Value(0.33)
-		[
-			ConstructLocalActionPanel()
-		];
-	}
 	
 	ResetLastPinType();
 
@@ -717,30 +702,6 @@ void SMyBlueprint::Refresh()
 	bNeedsRefresh = false;
 
 	GraphActionMenu->RefreshAllActions(true);
-	
-	bool bLocalActionsAreVisible = (GetLocalActionsListVisibility() == EVisibility::Visible);
-
-	//if ( !GetDefault<UEditorExperimentalSettings>()->bUnifiedBlueprintEditor )
-	//{
-	//	if ( bLocalActionsAreVisible )
-	//	{
-	//		if ( !LocalGraphActionMenu.IsValid() )
-	//		{
-	//			ActionMenuContainer->AddSlot()
-	//				.Value(0.33f)
-	//				[
-	//					ConstructLocalActionPanel()
-	//				];
-	//		}
-	//		check(LocalGraphActionMenu.IsValid());
-	//		LocalGraphActionMenu->RefreshAllActions(true);
-	//	}
-	//	else if ( LocalGraphActionMenu.IsValid() && ensure(ActionMenuContainer->GetChildren()->Num() > 1) )
-	//	{
-	//		ActionMenuContainer->RemoveAt(1);
-	//		LocalGraphActionMenu = nullptr;
-	//	}
-	//}
 }
 
 TSharedRef<SWidget> SMyBlueprint::OnCreateWidgetForAction(FCreateWidgetForActionData* const InCreateData)
@@ -748,61 +709,6 @@ TSharedRef<SWidget> SMyBlueprint::OnCreateWidgetForAction(FCreateWidgetForAction
 	return BlueprintEditorPtr.IsValid() ? SNew(SBlueprintPaletteItem, InCreateData, BlueprintEditorPtr.Pin()) : SNew(SBlueprintPaletteItem, InCreateData, GetBlueprintObj());
 }
 
-TSharedRef<SWidget> SMyBlueprint::ConstructLocalActionPanel()
-{
-	return SNew(SVerticalBox)
-		+SVerticalBox::Slot()
-			.AutoHeight()
-		[
-			SNew( SBorder )
-				.Visibility(this, &SMyBlueprint::GetLocalActionsListVisibility)
-				.BorderImage( FEditorStyle::GetBrush("DetailsView.CategoryTop") )
-				.BorderBackgroundColor( FLinearColor( .6,.6,.6, 1.0f ) )
-			[
-				SNew(SHorizontalBox)
-				+SHorizontalBox::Slot()
-					.AutoWidth()
-					.VAlign(VAlign_Center)
-					.Padding(4.0f, 0.0f, 8.0f, 0.0f)
-				[
-					SNew(STextBlock)
-						.Text(LOCTEXT("LocalVariables", "Local Variables"))
-				]
-				+SHorizontalBox::Slot()
-					.HAlign(HAlign_Left)
-				[
-					SNew(SButton)
-						.OnClicked(this, &SMyBlueprint::OnAddNewLocalVariable)
-						.ButtonStyle( FEditorStyle::Get(), "ToggleButton")
-						.ToolTipText(LOCTEXT("AddNewLocalVar_Tooltip", "Adds a new local variable"))
-					[
-						SNew(SImage)
-							.Image(FEditorStyle::GetBrush("PropertyWindow.Button_AddToArray"))
-					]
-				]
-			]
-		]
-		+SVerticalBox::Slot()
-			.FillHeight(1.0f)
-		[
-			SAssignNew(LocalGraphActionMenu, SGraphActionMenu, false)
-				.Visibility(this, &SMyBlueprint::GetLocalActionsListVisibility)
-				.OnGetFilterText(this, &SMyBlueprint::GetFilterText)
-				.OnCreateWidgetForAction(this, &SMyBlueprint::OnCreateWidgetForAction)
-				.OnCollectAllActions(this, &SMyBlueprint::GetLocalVariables)
-				.OnActionDragged(this, &SMyBlueprint::OnActionDragged)
-				.OnCategoryDragged(this, &SMyBlueprint::OnCategoryDragged)
-				.OnActionSelected(this, &SMyBlueprint::OnLocalActionSelected)
-				.OnActionDoubleClicked(this, &SMyBlueprint::OnActionDoubleClicked)
-				.OnContextMenuOpening(this, &SMyBlueprint::OnContextMenuOpening)
-				.OnCategoryTextCommitted(this, &SMyBlueprint::OnCategoryNameCommitted)
-				.OnCanRenameSelectedAction(this, &SMyBlueprint::CanRequestRenameOnActionNode)
-				.OnGetSectionTitle(this,&SMyBlueprint::OnGetSectionTitle)
-				.AlphaSortItems(false)
-		];
-
-}
-
 void SMyBlueprint::GetChildGraphs(UEdGraph* EdGraph, FGraphActionListBuilderBase& OutAllActions, FString ParentCategory)
 {
 	// Grab subgraphs
@@ -1284,7 +1190,7 @@ void SMyBlueprint::CollectAllActions(FGraphActionListBuilderBase& OutAllActions)
 		GetChildEvents(Graph, NewFuncAction->SectionID, OutAllActions);
 	}
 
-	if ( GetLocalActionsListVisibility().IsVisible() )
+	if(GetLocalActionsListVisibility().IsVisible())
 	{
 		GetLocalVariables(OutAllActions);
 	}
@@ -1465,21 +1371,6 @@ FReply SMyBlueprint::OnCategoryDragged(const FString& InCategory, const FPointer
 
 void SMyBlueprint::OnGlobalActionSelected(const TArray< TSharedPtr<FEdGraphSchemaAction> >& InActions)
 {
-	// If an action is being selected, clear the LocalGraphActionMenu of any selection it has, this keeps it so that only one menu will ever have selection at a time
-	if(InActions.Num() && LocalGraphActionMenu.IsValid() && InActions.Num())
-	{
-		LocalGraphActionMenu->SelectItemByName(NAME_None);
-	}
-	OnActionSelected(InActions);
-}
-
-void SMyBlueprint::OnLocalActionSelected(const TArray< TSharedPtr<FEdGraphSchemaAction> >& InActions)
-{
-	// If an action is being selected, clear the GraphActionMenu of any selection it has, this keeps it so that only one menu will ever have selection at a time
-	if(InActions.Num() && InActions.Num())
-	{
-		GraphActionMenu->SelectItemByName(NAME_None);
-	}
 	OnActionSelected(InActions);
 }
 
@@ -1682,19 +1573,7 @@ FEdGraphSchemaAction_K2Var* SMyBlueprint::SelectionAsVar() const
 
 FEdGraphSchemaAction_K2LocalVar* SMyBlueprint::SelectionAsLocalVar() const
 {
-	if ( GetDefault<UEditorExperimentalSettings>()->bUnifiedBlueprintEditor )
-	{
-		return SelectionAsType<FEdGraphSchemaAction_K2LocalVar>(GraphActionMenu);
-	}
-	else
-	{
-		if ( LocalGraphActionMenu.IsValid() )
-		{
-			return SelectionAsType<FEdGraphSchemaAction_K2LocalVar>(LocalGraphActionMenu);
-		}
-	}
-
-	return NULL;
+	return SelectionAsType<FEdGraphSchemaAction_K2LocalVar>(GraphActionMenu);
 }
 
 FEdGraphSchemaAction_K2Delegate* SMyBlueprint::SelectionAsDelegate() const
@@ -1714,19 +1593,9 @@ bool SMyBlueprint::SelectionIsCategory() const
 
 bool SMyBlueprint::SelectionHasContextMenu() const
 {
-	bool bSelectionHasContextMenu = false;
-
-	// Do not consider the selection having a context menu if the menu is being brought up in the other action menu, that being if an item is selected in the GraphActionMenu and you right click in LocalGraphActionMenu
-	if(GraphActionMenu.IsValid() && GraphActionMenu->HasFocusedDescendants())
-	{
-		bSelectionHasContextMenu = SelectionAsGraph() || SelectionAsVar() || SelectionIsCategory() || SelectionAsDelegate() || SelectionAsEnum() || SelectionAsEvent() || SelectionAsStruct();
-	}
-	else if(LocalGraphActionMenu.IsValid() && LocalGraphActionMenu->HasFocusedDescendants())
-	{
-		bSelectionHasContextMenu = SelectionAsLocalVar() != NULL;
-	}
-
-	return bSelectionHasContextMenu;
+	TArray<TSharedPtr<FEdGraphSchemaAction> > SelectedActions;
+	GraphActionMenu->GetSelectedActions(SelectedActions);
+	return SelectedActions.Num() > 0;
 }
 
 FString SMyBlueprint::GetGraphCategory(UEdGraph* InGraph) const
@@ -2467,10 +2336,6 @@ FReply SMyBlueprint::OnAddNewLocalVariable()
 void SMyBlueprint::OnFilterTextChanged( const FText& InFilterText )
 {
 	GraphActionMenu->GenerateFilteredItems(false);
-	if (LocalGraphActionMenu.IsValid())
-	{
-		LocalGraphActionMenu->GenerateFilteredItems(false);
-	}
 }
 
 FText SMyBlueprint::GetFilterText() const
@@ -2482,10 +2347,6 @@ void SMyBlueprint::OnRequestRenameOnActionNode()
 {
 	// Attempt to rename in both menus, only one of them will have anything selected
 	GraphActionMenu->OnRequestRenameOnActionNode();
-	if (LocalGraphActionMenu.IsValid())
-	{
-		LocalGraphActionMenu->OnRequestRenameOnActionNode();
-	}
 }
 
 bool SMyBlueprint::CanRequestRenameOnActionNode() const
@@ -2498,10 +2359,6 @@ bool SMyBlueprint::CanRequestRenameOnActionNode() const
 	{
 		return GraphActionMenu->CanRequestRenameOnActionNode();
 	}
-	else if(LocalGraphActionMenu.IsValid())
-	{
-		return LocalGraphActionMenu->CanRequestRenameOnActionNode();
-	}
 	return false;
 }
 
@@ -2512,29 +2369,16 @@ void SMyBlueprint::SelectItemByName(const FName& ItemName, ESelectInfo::Type Sel
 	{
 		ClearGraphActionMenuSelection();
 	}
-	// Attempt to select the item in the main graph action menu, if that fails, attempt the same in LocalGraphActionMenu
-	else if( !GraphActionMenu->SelectItemByName(ItemName, SelectInfo, SectionId, bIsCategory) && LocalGraphActionMenu.IsValid() )
-	{
-		LocalGraphActionMenu->SelectItemByName(ItemName, SelectInfo, SectionId, bIsCategory);
-	}
 }
 
 void SMyBlueprint::ClearGraphActionMenuSelection()
 {
 	GraphActionMenu->SelectItemByName(NAME_None);
-	if (LocalGraphActionMenu.IsValid())
-	{
-		LocalGraphActionMenu->SelectItemByName(NAME_None);
-	}
 }
 
 void SMyBlueprint::ExpandCategory(const FString& CategoryName)
 {
 	GraphActionMenu->ExpandCategory(CategoryName);
-	if (LocalGraphActionMenu.IsValid())
-	{
-		LocalGraphActionMenu->ExpandCategory(CategoryName);
-	}
 }
 
 bool SMyBlueprint::MoveCategoryBeforeCategory( const FString& InCategoryToMove, const FString& InTargetCategory )
diff --git a/Engine/Source/Editor/Kismet/Private/SMyBlueprint.h b/Engine/Source/Editor/Kismet/Private/SMyBlueprint.h
--- a/Engine/Source/Editor/Kismet/Private/SMyBlueprint.h
+++ b/Engine/Source/Editor/Kismet/Private/SMyBlueprint.h
@@ -133,9 +133,6 @@ private:
 	/** Creates widgets for the graph schema actions */
 	TSharedRef<SWidget> OnCreateWidgetForAction(struct FCreateWidgetForActionData* const InCreateData);
 
-	/** Creates the local variable action list sub-widget */
-	TSharedRef<SWidget> ConstructLocalActionPanel();
-
 	/** Callback used to populate all actions list in SGraphActionMenu */
 	void CollectAllActions(FGraphActionListBuilderBase& OutAllActions);
 	void CollectStaticSections(TArray<int32>& StaticSectionIDs);
@@ -152,7 +149,6 @@ private:
 	void OnActionSelected(const TArray< TSharedPtr<FEdGraphSchemaAction> >& InActions);
 	static void OnActionSelectedHelper(TSharedPtr<FEdGraphSchemaAction> InAction, UBlueprint* Blueprint, TSharedRef<SKismetInspector> Inspector);
 	void OnGlobalActionSelected(const TArray< TSharedPtr<FEdGraphSchemaAction> >& InActions);
-	void OnLocalActionSelected(const TArray< TSharedPtr<FEdGraphSchemaAction> >& InActions);
 	void OnActionDoubleClicked(const TArray< TSharedPtr<FEdGraphSchemaAction> >& InActions);
 	TSharedPtr<SWidget> OnContextMenuOpening();
 
@@ -233,9 +229,6 @@ private:
 	/** Graph Action Menu for displaying all our variables and functions */
 	TSharedPtr<class SGraphActionMenu> GraphActionMenu;
 
-	/** Graph Action Menu for displaying all local variables */
-	TSharedPtr<class SGraphActionMenu> LocalGraphActionMenu;
-
 	/** The last pin type used (including the function editor last pin type) */
 	FEdGraphPinType LastPinType;
 	FEdGraphPinType LastFunctionPinType;
@@ -246,9 +239,6 @@ private:
 	/** The filter box that handles filtering for both graph action menus. */
 	TSharedPtr< SSearchBox > FilterBox;
 
-	/** Contains both the GraphActionMenu and LocalGraphActionMenu */
-	TSharedPtr< SSplitter > ActionMenuContainer;
-
 	/** Enums created from 'blueprint' level */
 	TArray<TWeakObjectPtr<UUserDefinedStruct>> StructsAddedToBlueprint;
 
