[INTEGRATE] Change 2425811 by [REDACTED_PERSON][REDACTED_HANDLE] on 2015/01/30 10:47:42

	BP - Changing overridable functions to be a bit more out of the way for new blueprints.  They're accessible in a lot of ways now though.

diff --git a/Engine/Source/Editor/Kismet/Private/SMyBlueprint.cpp b/Engine/Source/Editor/Kismet/Private/SMyBlueprint.cpp
--- a/Engine/Source/Editor/Kismet/Private/SMyBlueprint.cpp
+++ b/Engine/Source/Editor/Kismet/Private/SMyBlueprint.cpp
@@ -498,7 +498,14 @@ FText SMyBlueprint::OnGetSectionTitle( int32 InSectionID )
 		SeperatorTitle = NSLOCTEXT("GraphActionNode", "Components", "Components");
 		break;
 	case NodeSectionID::FUNCTION:
-		SeperatorTitle = NSLOCTEXT("GraphActionNode", "Functions", "Functions");
+		if ( OverridableFunctionActions.Num() > 0 )
+		{
+			SeperatorTitle = FText::Format( NSLOCTEXT("GraphActionNode", "FunctionsOverridableFormat", "Functions ({0} Overridable)"), FText::AsNumber(OverridableFunctionActions.Num()) );
+		}
+		else
+		{
+			SeperatorTitle = NSLOCTEXT("GraphActionNode", "Functions", "Functions");
+		}
 		break;
 	case NodeSectionID::FUNCTION_OVERRIDABLE:
 		SeperatorTitle = NSLOCTEXT("GraphActionNode", "OverridableFunctions", "Overridable Functions");
@@ -545,7 +552,36 @@ TSharedRef<SWidget> SMyBlueprint::OnGetSectionWidget(TSharedRef<SWidget> RowWidg
 		break;
 	case NodeSectionID::FUNCTION:
 		AddNewText = LOCTEXT("AddNewFunction", "Function");
-		break;
+		return SAssignNew(FunctionSectionButton, SComboButton)
+		.ButtonStyle(FEditorStyle::Get(), "RoundButton")
+		.ForegroundColor(FEditorStyle::GetSlateColor("DefaultForeground"))
+		.ContentPadding(FMargin(2, 0))
+		.OnGetMenuContent(this, &SMyBlueprint::OnGetFunctionListMenu)
+		.HasDownArrow(false)
+		.ButtonContent()
+		[
+			SNew(SHorizontalBox)
+
+			+ SHorizontalBox::Slot()
+			.AutoWidth()
+			.Padding(FMargin(0, 1))
+			[
+				SNew(SImage)
+				.Image(FEditorStyle::GetBrush("Plus"))
+			]
+
+			+ SHorizontalBox::Slot()
+			.VAlign(VAlign_Center)
+			.AutoWidth()
+			.Padding(FMargin(2,0,0,0))
+			[
+				SNew(STextBlock)
+				.Font(IDetailLayoutBuilder::GetDetailFontBold())
+				.Text(AddNewText)
+				.Visibility(this, &SMyBlueprint::OnGetSectionTextVisibility, WeakRowWidget, InSectionID)
+				.ShadowOffset(FVector2D(1,1))
+			]
+		];
 	case NodeSectionID::MACRO:
 		AddNewText = LOCTEXT("AddNewMacro", "Macro");
 		break;
@@ -565,10 +601,10 @@ TSharedRef<SWidget> SMyBlueprint::OnGetSectionWidget(TSharedRef<SWidget> RowWidg
 	return SNew(SButton)
 		.ButtonStyle(FEditorStyle::Get(), "RoundButton")
 		.ForegroundColor(FEditorStyle::GetSlateColor("DefaultForeground"))
+		.ContentPadding(FMargin(2, 0))
 		.OnClicked(this, &SMyBlueprint::OnAddButtonClickedOnSection, InSectionID)
 		.HAlign(HAlign_Center)
 		.VAlign(VAlign_Center)
-		.ContentPadding(FMargin(2, 0))
 		//.IsEnabled(this, &SMyBlueprint::OnAddButtonClickedOnSection)
 		[
 			SNew(SHorizontalBox)
@@ -589,7 +625,7 @@ TSharedRef<SWidget> SMyBlueprint::OnGetSectionWidget(TSharedRef<SWidget> RowWidg
 				SNew(STextBlock)
 				.Font(IDetailLayoutBuilder::GetDetailFontBold())
 				.Text(AddNewText)
-				.Visibility(this, &SMyBlueprint::OnGetSectionTextVisibility, WeakRowWidget)
+				.Visibility(this, &SMyBlueprint::OnGetSectionTextVisibility, WeakRowWidget, InSectionID)
 				.ShadowOffset(FVector2D(1,1))
 			]
 		];
@@ -622,9 +658,16 @@ FReply SMyBlueprint::OnAddButtonClickedOnSection(int32 InSectionID)
 	return FReply::Handled();
 }
 
-EVisibility SMyBlueprint::OnGetSectionTextVisibility(TWeakPtr<SWidget> RowWidget) const
+EVisibility SMyBlueprint::OnGetSectionTextVisibility(TWeakPtr<SWidget> RowWidget, int32 InSectionID) const
 {
-	if ( RowWidget.Pin()->IsHovered() )
+	bool ShowText = RowWidget.Pin()->IsHovered();
+	if ( InSectionID == NodeSectionID::FUNCTION && FunctionSectionButton->IsOpen() )
+	{
+		ShowText = true;
+	}
+
+	// If the row is currently hovered, or a menu is being displayed for a button, keep the button expanded.
+	if ( ShowText )
 	{
 		return EVisibility::SelfHitTestInvisible;
 	}
@@ -634,6 +677,43 @@ EVisibility SMyBlueprint::OnGetSectionTextVisibility(TWeakPtr<SWidget> RowWidget
 	}
 }
 
+TSharedRef<SWidget> SMyBlueprint::OnGetFunctionListMenu()
+{
+	const bool bShouldCloseWindowAfterMenuSelection = true;
+	FMenuBuilder MenuBuilder(bShouldCloseWindowAfterMenuSelection, BlueprintEditorPtr.Pin()->GetToolkitCommands());
+
+	MenuBuilder.BeginSection("AddNewItem", LOCTEXT("AddOperations", "Add New"));
+	{
+		MenuBuilder.AddMenuEntry(FBlueprintEditorCommands::Get().AddNewFunction);
+	}
+	MenuBuilder.EndSection();
+
+	BuildOverridableFunctionsMenu(MenuBuilder);
+
+	TSharedRef<SWidget> MenuWidget = MenuBuilder.MakeWidget();
+
+	return MenuWidget;
+}
+
+void SMyBlueprint::BuildOverridableFunctionsMenu(FMenuBuilder& MenuBuilder)
+{
+	MenuBuilder.BeginSection("OverrideFunction", LOCTEXT("OverrideFunction", "Override Function"));
+	{
+		for ( auto& OverrideAction : OverridableFunctionActions )
+		{
+			MenuBuilder.AddMenuEntry(
+				OverrideAction->MenuDescription,
+				FText::FromString(OverrideAction->TooltipDescription),
+				FSlateIcon(),
+				FUIAction(FExecuteAction::CreateSP(this, &SMyBlueprint::ImplementFunction, OverrideAction)),
+				NAME_None,
+				EUserInterfaceActionType::Button
+				);
+		}
+	}
+	MenuBuilder.EndSection();
+}
+
 bool SMyBlueprint::CanRequestRenameOnActionNode(TWeakPtr<FGraphActionNode> InSelectedNode) const
 {
 	bool bIsReadOnly = true;
@@ -930,7 +1010,8 @@ void SMyBlueprint::CollectAllActions(FGraphActionListBuilderBase& OutAllActions)
 	// Initialise action sorting instance
 	FGraphActionSort SortList( Blueprint->CategorySorting );
 	// List of names of functions we implement
-	TArray<FName> ImplementedFunctions;
+	ImplementedFunctionCache.Empty();
+
 	// Grab Variables
 	for (TFieldIterator<UProperty> PropertyIt(Blueprint->SkeletonGeneratedClass, FieldIteratorSuperFlag); PropertyIt; ++PropertyIt)
 	{
@@ -999,7 +1080,7 @@ void SMyBlueprint::CollectAllActions(FGraphActionListBuilderBase& OutAllActions)
 				{
 					NewFuncAction->EdGraph = Graph;
 				}
-				ImplementedFunctions.AddUnique(PropertyName);
+				ImplementedFunctionCache.Add(PropertyName);
 			}
 		}
 	}
@@ -1041,7 +1122,7 @@ void SMyBlueprint::CollectAllActions(FGraphActionListBuilderBase& OutAllActions)
 		GetChildGraphs(Graph, OutAllActions);
 		GetChildEvents(Graph, NewFuncAction->SectionID, OutAllActions);
 
-		ImplementedFunctions.AddUnique(Graph->GetFName());
+		ImplementedFunctionCache.Add(Graph->GetFName());
 	}
 	// Grab macros implemented by the blueprint
 	for (int32 i = 0; i < Blueprint->MacroGraphs.Num(); i++)
@@ -1065,20 +1146,22 @@ void SMyBlueprint::CollectAllActions(FGraphActionListBuilderBase& OutAllActions)
 		GetChildGraphs(Graph, OutAllActions);
 		GetChildEvents(Graph, NewFuncAction->SectionID, OutAllActions);
 
-		ImplementedFunctions.AddUnique(FunctionName);
+		ImplementedFunctionCache.Add(FunctionName);
 	}
 
-	// Show potentially overridable functions
-	for (TFieldIterator<UFunction> FunctionIt(Blueprint->ParentClass, EFieldIteratorFlags::IncludeSuper); FunctionIt; ++FunctionIt)
+	OverridableFunctionActions.Reset();
+
+	// Cache potentially overridable functions
+	for ( TFieldIterator<UFunction> FunctionIt(Blueprint->ParentClass, EFieldIteratorFlags::IncludeSuper); FunctionIt; ++FunctionIt )
 	{
 		const UFunction* Function = *FunctionIt;
 		const FName FunctionName = Function->GetFName();
 
-		if (UEdGraphSchema_K2::CanKismetOverrideFunction(Function) && !ImplementedFunctions.Contains(FunctionName) && !UEdGraphSchema_K2::FunctionCanBePlacedAsEvent(Function) && !FBlueprintEditorUtils::FindOverrideForFunction(Blueprint, CastChecked<UClass>(Function->GetOuter()), Function->GetFName()))
+		if ( UEdGraphSchema_K2::CanKismetOverrideFunction(Function) && !ImplementedFunctionCache.Contains(FunctionName) && !UEdGraphSchema_K2::FunctionCanBePlacedAsEvent(Function) && !FBlueprintEditorUtils::FindOverrideForFunction(Blueprint, CastChecked<UClass>(Function->GetOuter()), Function->GetFName()) )
 		{
 			FString FunctionTooltip = UK2Node_CallFunction::GetDefaultTooltipForFunction(Function);
 			FText FunctionDesc = FText::FromString(Function->GetMetaData(TEXT("FriendlyName")));
-			if (FunctionDesc.IsEmpty())
+			if ( FunctionDesc.IsEmpty() )
 			{
 				FunctionDesc = FText::FromString(Function->GetName());
 			}
@@ -1088,7 +1171,10 @@ void SMyBlueprint::CollectAllActions(FGraphActionListBuilderBase& OutAllActions)
 			TSharedPtr<FEdGraphSchemaAction_K2Graph> NewFuncAction = MakeShareable(new FEdGraphSchemaAction_K2Graph(EEdGraphSchemaAction_K2Graph::Function, FunctionCategory, FunctionDesc, FunctionTooltip, 1));
 			NewFuncAction->FuncName = FunctionName;
 			NewFuncAction->SectionID = NodeSectionID::FUNCTION_OVERRIDABLE;
-			OutAllActions.AddAction(NewFuncAction);
+
+			OverridableFunctionActions.Add(NewFuncAction);
+
+			//OutAllActions.AddAction(NewFuncAction);
 		}
 	}
 
@@ -1106,7 +1192,7 @@ void SMyBlueprint::CollectAllActions(FGraphActionListBuilderBase& OutAllActions)
 			FString FunctionDesc = FunctionName.ToString();
 			
 			FString FunctionCategory;
-			if(UFunction* Function = Blueprint->SkeletonGeneratedClass->FindFunctionByName(Graph->GetFName()))
+			if ( UFunction* Function = Blueprint->SkeletonGeneratedClass->FindFunctionByName(Graph->GetFName()) )
 			{
 				FunctionCategory = Function->GetMetaData(FBlueprintMetadata::MD_FunctionCategory);
 			}
@@ -1136,7 +1222,7 @@ void SMyBlueprint::CollectAllActions(FGraphActionListBuilderBase& OutAllActions)
 					const UFunction* Function = *FunctionIt;
 					const FName FunctionName = Function->GetFName();
 
-					if (UEdGraphSchema_K2::CanKismetOverrideFunction(Function) && !ImplementedFunctions.Contains(FunctionName) && !UEdGraphSchema_K2::FunctionCanBePlacedAsEvent(Function))
+					if ( UEdGraphSchema_K2::CanKismetOverrideFunction(Function) && !ImplementedFunctionCache.Contains(FunctionName) && !UEdGraphSchema_K2::FunctionCanBePlacedAsEvent(Function) )
 					{
 						FString FunctionTooltip = Function->GetToolTipText().ToString();
 						FString FunctionDesc = Function->GetMetaData(TEXT("FriendlyName"));
@@ -1470,15 +1556,19 @@ void SMyBlueprint::OnActionSelectedHelper(TSharedPtr<FEdGraphSchemaAction> InAct
 	}
 }
 
-void SMyBlueprint::OnActionDoubleClicked( const TArray< TSharedPtr<FEdGraphSchemaAction> >& InActions )
+void SMyBlueprint::OnActionDoubleClicked(const TArray< TSharedPtr<FEdGraphSchemaAction> >& InActions)
 {
-	if( !BlueprintEditorPtr.IsValid() )
+	if ( !BlueprintEditorPtr.IsValid() )
 	{
 		return;
 	}
 
-	TSharedPtr<FEdGraphSchemaAction> InAction( InActions.Num() > 0 ? InActions[0] : NULL );
+	TSharedPtr<FEdGraphSchemaAction> InAction(InActions.Num() > 0 ? InActions[0] : NULL);
+	ExecuteAction(InAction);
+}
 
+void SMyBlueprint::ExecuteAction(TSharedPtr<FEdGraphSchemaAction> InAction)
+{
 	UBlueprint* Blueprint = BlueprintEditorPtr.Pin()->GetBlueprintObj();
 	if(InAction.IsValid())
 	{
@@ -1544,6 +1634,8 @@ template<class SchemaActionType> SchemaActionType* SelectionAsType( const TShare
 	if ( SelectedAction.IsValid() &&
 		 SelectedAction->GetTypeId() == SchemaActionType::StaticGetTypeId() )
 	{
+		// TODO Why not? StaticCastSharedPtr<>()
+
 		Selection = (SchemaActionType*)SelectedActions[0].Get();
 	}
 
@@ -1723,6 +1815,17 @@ void SMyBlueprint::BuildAddNewMenu(FMenuBuilder& MenuBuilder)
 		MenuBuilder.AddMenuEntry(FBlueprintEditorCommands::Get().AddNewVariable);
 		MenuBuilder.AddMenuEntry(FBlueprintEditorCommands::Get().AddNewLocalVariable);
 		MenuBuilder.AddMenuEntry(FBlueprintEditorCommands::Get().AddNewFunction);
+
+		if ( OverridableFunctionActions.Num() > 0 )
+		{
+			MenuBuilder.AddSubMenu(
+				LOCTEXT("OverrideFunction", "Override Function"),
+				FText::GetEmpty(),
+				FNewMenuDelegate::CreateSP(this, &SMyBlueprint::BuildOverridableFunctionsMenu),
+				false,
+				FSlateIcon(FEditorStyle::GetStyleSetName(), "BlueprintEditor.AddNewFunction.Small"));
+		}
+
 		MenuBuilder.AddMenuEntry(FBlueprintEditorCommands::Get().AddNewMacroDeclaration);
 		MenuBuilder.AddMenuEntry(FBlueprintEditorCommands::Get().AddNewEventGraph);
 		MenuBuilder.AddMenuEntry(FBlueprintEditorCommands::Get().AddNewDelegate);
@@ -1801,38 +1904,48 @@ bool SMyBlueprint::CanImplementFunction() const
 
 void SMyBlueprint::OnImplementFunction()
 {
-	if (FEdGraphSchemaAction_K2Graph* GraphAction = SelectionAsGraph())
+	if ( FEdGraphSchemaAction_K2Graph* GraphAction = SelectionAsGraph() )
 	{
-		check(GetBlueprintObj()->SkeletonGeneratedClass);
-		UFunction* OverrideFunc = FindField<UFunction>(GetBlueprintObj()->SkeletonGeneratedClass, GraphAction->FuncName);
-		if (OverrideFunc == NULL)
+		ImplementFunction(GraphAction);
+	}
+}
+
+void SMyBlueprint::ImplementFunction(TSharedPtr<FEdGraphSchemaAction_K2Graph> GraphAction)
+{
+	ImplementFunction(GraphAction.Get());
+}
+
+void SMyBlueprint::ImplementFunction(FEdGraphSchemaAction_K2Graph* GraphAction)
+{
+	check(GetBlueprintObj()->SkeletonGeneratedClass);
+	UFunction* OverrideFunc = FindField<UFunction>(GetBlueprintObj()->SkeletonGeneratedClass, GraphAction->FuncName);
+	if (OverrideFunc == NULL)
+	{
+		// maybe it's from a native interface, check those too
+		for ( UClass* TempClass=GetBlueprintObj()->ParentClass; (NULL != TempClass) && (NULL == OverrideFunc); TempClass=TempClass->GetSuperClass() )
 		{
-			// maybe it's from a native interface, check those too
-			for ( UClass* TempClass=GetBlueprintObj()->ParentClass; (NULL != TempClass) && (NULL == OverrideFunc); TempClass=TempClass->GetSuperClass() )
+			for (int32 Idx=0; Idx<TempClass->Interfaces.Num(); ++Idx)
 			{
-				for (int32 Idx=0; Idx<TempClass->Interfaces.Num(); ++Idx)
+				FImplementedInterface const& I = TempClass->Interfaces[Idx];
+				if (!I.bImplementedByK2)
 				{
-					FImplementedInterface const& I = TempClass->Interfaces[Idx];
-					if (!I.bImplementedByK2)
+					OverrideFunc = FindField<UFunction>(I.Class, GraphAction->FuncName);
+					if (OverrideFunc)
 					{
-						OverrideFunc = FindField<UFunction>(I.Class, GraphAction->FuncName);
-						if (OverrideFunc)
-						{
-							// found it, done
-							break;
-						}
+						// found it, done
+						break;
 					}
 				}
 			}
 		}
-		check(OverrideFunc);
-		UClass* const OverrideFuncClass = CastChecked<UClass>(OverrideFunc->GetOuter());
-
-		// Implement the function graph
-		UEdGraph* const NewGraph = FBlueprintEditorUtils::CreateNewGraph(GetBlueprintObj(), GraphAction->FuncName, UEdGraph::StaticClass(), UEdGraphSchema_K2::StaticClass());
-		FBlueprintEditorUtils::AddFunctionGraph(GetBlueprintObj(), NewGraph, /*bIsUserCreated=*/ false, OverrideFuncClass);
-		BlueprintEditorPtr.Pin()->OpenDocument(NewGraph, FDocumentTracker::OpenNewDocument);
 	}
+	check(OverrideFunc);
+	UClass* const OverrideFuncClass = CastChecked<UClass>(OverrideFunc->GetOuter());
+
+	// Implement the function graph
+	UEdGraph* const NewGraph = FBlueprintEditorUtils::CreateNewGraph(GetBlueprintObj(), GraphAction->FuncName, UEdGraph::StaticClass(), UEdGraphSchema_K2::StaticClass());
+	FBlueprintEditorUtils::AddFunctionGraph(GetBlueprintObj(), NewGraph, /*bIsUserCreated=*/ false, OverrideFuncClass);
+	BlueprintEditorPtr.Pin()->OpenDocument(NewGraph, FDocumentTracker::OpenNewDocument);
 }
 
 void SMyBlueprint::OnFindEntry()
diff --git a/Engine/Source/Editor/Kismet/Private/SMyBlueprint.h b/Engine/Source/Editor/Kismet/Private/SMyBlueprint.h
--- a/Engine/Source/Editor/Kismet/Private/SMyBlueprint.h
+++ b/Engine/Source/Editor/Kismet/Private/SMyBlueprint.h
@@ -150,6 +150,7 @@ private:
 	static void OnActionSelectedHelper(TSharedPtr<FEdGraphSchemaAction> InAction, UBlueprint* Blueprint, TSharedRef<SKismetInspector> Inspector);
 	void OnGlobalActionSelected(const TArray< TSharedPtr<FEdGraphSchemaAction> >& InActions);
 	void OnActionDoubleClicked(const TArray< TSharedPtr<FEdGraphSchemaAction> >& InActions);
+	void ExecuteAction(TSharedPtr<FEdGraphSchemaAction> InAction);
 	TSharedPtr<SWidget> OnContextMenuOpening();
 
 	TSharedRef<SWidget> CreateAddNewMenuWidget();
@@ -159,7 +160,9 @@ private:
 	bool CanRequestRenameOnActionNode(TWeakPtr<struct FGraphActionNode> InSelectedNode) const;
 	FText OnGetSectionTitle( int32 InSectionID );
 	TSharedRef<SWidget> OnGetSectionWidget( TSharedRef<SWidget> RowWidget, int32 InSectionID );
-	EVisibility OnGetSectionTextVisibility(TWeakPtr<SWidget> RowWidget) const;
+	EVisibility OnGetSectionTextVisibility(TWeakPtr<SWidget> RowWidget, int32 InSectionID) const;
+	TSharedRef<SWidget> OnGetFunctionListMenu();
+	void BuildOverridableFunctionsMenu(FMenuBuilder& MenuBuilder);
 	FReply OnAddButtonClickedOnSection(int32 InSectionID);
 
 	/** Support functions for checkbox to manage displaying user variables only */
@@ -181,6 +184,8 @@ private:
 	void OnFocusNode();
 	void OnFocusNodeInNewTab();
 	void OnImplementFunction();
+	void ImplementFunction(TSharedPtr<FEdGraphSchemaAction_K2Graph> GraphAction);
+	void ImplementFunction(FEdGraphSchemaAction_K2Graph* GraphAction);
 	bool CanImplementFunction() const;
 	void OnFindEntry();
 	bool CanFindEntry() const;
@@ -229,6 +234,15 @@ private:
 	/** Graph Action Menu for displaying all our variables and functions */
 	TSharedPtr<class SGraphActionMenu> GraphActionMenu;
 
+	/** The +Function button in the function section */
+	TSharedPtr<SComboButton> FunctionSectionButton;
+
+	/** When we rebuild the view of members, we cache (but don't display) any overridable functions for user in popup menus. */
+	TArray< TSharedPtr<FEdGraphSchemaAction_K2Graph> > OverridableFunctionActions;
+
+	/** When we refresh the list of functions we cache off the implemented ones to ask questions for overridable functions. */
+	TSet<FName> ImplementedFunctionCache;
+
 	/** The last pin type used (including the function editor last pin type) */
 	FEdGraphPinType LastPinType;
 	FEdGraphPinType LastFunctionPinType;
