[INTEGRATE] Change 2425853 by [REDACTED_PERSON][REDACTED_HANDLE] on 2015/01/30 10:56:21

	BP - Continuing to itterate on layout based on the survey we've changed the default layout of the blueprint editor.  Additionally fixing the Default editor so that data only blueprints load in the data only default editor.

diff --git a/Engine/Source/Editor/Kismet/Private/BlueprintDetailsCustomization.cpp b/Engine/Source/Editor/Kismet/Private/BlueprintDetailsCustomization.cpp
--- a/Engine/Source/Editor/Kismet/Private/BlueprintDetailsCustomization.cpp
+++ b/Engine/Source/Editor/Kismet/Private/BlueprintDetailsCustomization.cpp
@@ -4246,7 +4246,7 @@ void FBlueprintGlobalOptionsDetails::CustomizeDetails(IDetailLayoutBuilder& Deta
 		for (TFieldIterator<UProperty> PropertyIt(Blueprint->GetClass(), EFieldIteratorFlags::IncludeSuper); PropertyIt; ++PropertyIt)
 		{
 			UProperty* Property = *PropertyIt;
-			if(FObjectEditorUtils::GetCategory(Property) != TEXT("BlueprintOption"))
+			if(FObjectEditorUtils::GetCategory(Property) != TEXT("BlueprintOptions"))
 			{
 				DetailLayout.HideProperty(DetailLayout.GetProperty(Property->GetFName()));
 			}
diff --git a/Engine/Source/Editor/Kismet/Private/BlueprintEditor.cpp b/Engine/Source/Editor/Kismet/Private/BlueprintEditor.cpp
--- a/Engine/Source/Editor/Kismet/Private/BlueprintEditor.cpp
+++ b/Engine/Source/Editor/Kismet/Private/BlueprintEditor.cpp
@@ -1608,10 +1608,21 @@ void FBlueprintEditor::RegisterApplicationModes(const TArray<UBlueprint*>& InBlu
 		{
 			if ( GetDefault<UEditorExperimentalSettings>()->bUnifiedBlueprintEditor )
 			{
-				AddApplicationMode(
-					FBlueprintEditorApplicationModes::StandardBlueprintEditorMode,
-					MakeShareable(new FBlueprintEditorUnifiedMode(SharedThis(this), FBlueprintEditorApplicationModes::StandardBlueprintEditorMode, FBlueprintEditorApplicationModes::GetLocalizedMode, CanAccessComponentsMode())));
-				SetCurrentMode(FBlueprintEditorApplicationModes::StandardBlueprintEditorMode);
+				if ( bShouldOpenInDefaultsMode )
+				{
+					// We either have no blueprints or many, open in the defaults mode for multi-editing
+					AddApplicationMode(
+						FBlueprintEditorApplicationModes::BlueprintDefaultsMode,
+						MakeShareable(new FBlueprintDefaultsApplicationMode(SharedThis(this))));
+					SetCurrentMode(FBlueprintEditorApplicationModes::BlueprintDefaultsMode);
+				}
+				else
+				{
+					AddApplicationMode(
+						FBlueprintEditorApplicationModes::StandardBlueprintEditorMode,
+						MakeShareable(new FBlueprintEditorUnifiedMode(SharedThis(this), FBlueprintEditorApplicationModes::StandardBlueprintEditorMode, FBlueprintEditorApplicationModes::GetLocalizedMode, CanAccessComponentsMode())));
+					SetCurrentMode(FBlueprintEditorApplicationModes::StandardBlueprintEditorMode);
+				}
 			}
 			else
 			{
@@ -2050,9 +2061,6 @@ void FBlueprintEditor::CreateDefaultTabContents(const TArray<UBlueprint*>& InBlu
 
 	if (InBlueprints.Num() > 0)
 	{
-		// Hide the inspector title for Persona, where we customize the defaults editor widget
-		const bool bIsPersona = InBlueprint ? InBlueprint->IsA<UAnimBlueprint>() : InBlueprints[0]->IsA<UAnimBlueprint>();
-		const bool bShowTitle = !bIsPersona;
 		const bool bShowPublicView = true;
 		const bool bHideNameArea = false;
 
@@ -2062,7 +2070,7 @@ void FBlueprintEditor::CreateDefaultTabContents(const TArray<UBlueprint*>& InBlu
 			. ViewIdentifier(FName("BlueprintDefaults"))
 			. IsEnabled(!bIsInterface)
 			. ShowPublicViewControl(bShowPublicView)
-			. ShowTitleArea(bShowTitle)
+			. ShowTitleArea(false)
 			. HideNameArea(bHideNameArea)
 			. IsPropertyEditingEnabledDelegate( IsPropertyEditingEnabledDelegate )
 			. OnFinishedChangingProperties( FOnFinishedChangingProperties::FDelegate::CreateSP( this, &FBlueprintEditor::OnFinishedChangingProperties ) );
@@ -6217,7 +6225,7 @@ void FBlueprintEditor::StartEditingDefaults(bool bAutoFocus, bool bForceRefresh)
 			}
 			else
 			{
-				DefaultEditor->ShowDetailsForSingleObject(GetBlueprintObj()->GeneratedClass->GetDefaultObject(), SKismetInspector::FShowDetailsOptions(DefaultEditString(), bForceRefresh));
+				DefaultEditor->ShowDetailsForSingleObject(GetBlueprintObj()->GeneratedClass->GetDefaultObject(), SKismetInspector::FShowDetailsOptions(FText::GetEmpty(), bForceRefresh));
 			}
 		}
 	}
@@ -6630,11 +6638,6 @@ void FBlueprintEditor::VariableListWasUpdated()
 	StartEditingDefaults(/*bAutoFocus=*/ false);
 }
 
-FText FBlueprintEditor::DefaultEditString() 
-{
-	return LOCTEXT("BlueprintEditingDefaults", "Editing defaults");
-}
-
 bool FBlueprintEditor::GetBoundsForSelectedNodes(class FSlateRect& Rect, float Padding)
 {
 	TSharedPtr<SGraphEditor> FocusedGraphEd = FocusedGraphEdPtr.Pin();
diff --git a/Engine/Source/Editor/Kismet/Private/BlueprintEditorModes.cpp b/Engine/Source/Editor/Kismet/Private/BlueprintEditorModes.cpp
--- a/Engine/Source/Editor/Kismet/Private/BlueprintEditorModes.cpp
+++ b/Engine/Source/Editor/Kismet/Private/BlueprintEditorModes.cpp
@@ -172,7 +172,7 @@ FBlueprintDefaultsApplicationMode::FBlueprintDefaultsApplicationMode(TSharedPtr<
 	BlueprintDefaultsTabFactories.RegisterFactory(MakeShareable(new FDefaultsEditorSummoner(InBlueprintEditor)));
 	BlueprintDefaultsTabFactories.RegisterFactory(MakeShareable(new FFindResultsSummoner(InBlueprintEditor)));
 
-	TabLayout = FTabManager::NewLayout( "Standalone_BlueprintDefaults_Layout_v4" )
+	TabLayout = FTabManager::NewLayout( "Standalone_BlueprintDefaults_Layout_v6" )
 		->AddArea
 		(
 			FTabManager::NewPrimaryArea()
@@ -187,7 +187,6 @@ FBlueprintDefaultsApplicationMode::FBlueprintDefaultsApplicationMode(TSharedPtr<
 			->Split
 			(
 				FTabManager::NewStack()
-				->SetHideTabWell(true)
 				->AddTab( FBlueprintEditorTabs::DefaultEditorID, ETabState::OpenedTab )
 			)
 		);
@@ -546,7 +545,7 @@ void FBlueprintMacroApplicationMode::PostActivateMode()
 
 ////////////////////////////////////////
 //
-FBlueprintEditorUnifiedMode::FBlueprintEditorUnifiedMode(TSharedPtr<class FBlueprintEditor> InBlueprintEditor, FName InModeName, FText(*GetLocalizedMode)( const FName ), const bool bRegisterViewport, const bool bRegisterDefaultsTab)
+FBlueprintEditorUnifiedMode::FBlueprintEditorUnifiedMode(TSharedPtr<class FBlueprintEditor> InBlueprintEditor, FName InModeName, FText(*GetLocalizedMode)( const FName ), const bool bRegisterViewport)
 	: FApplicationMode(InModeName, GetLocalizedMode)
 {
 	MyBlueprintEditor = InBlueprintEditor;
@@ -563,15 +562,12 @@ FBlueprintEditorUnifiedMode::FBlueprintEditorUnifiedMode(TSharedPtr<class FBluep
 		BlueprintEditorTabFactories.RegisterFactory(MakeShareable(new FConstructionScriptEditorSummoner(InBlueprintEditor)));
 		BlueprintEditorTabFactories.RegisterFactory(MakeShareable(new FSCSViewportSummoner(InBlueprintEditor)));
 	}
-	if( bRegisterDefaultsTab )
-	{
-		BlueprintEditorTabFactories.RegisterFactory(MakeShareable(new FDefaultsEditorSummoner(InBlueprintEditor)));
-	}
+
 	CoreTabFactories.RegisterFactory(MakeShareable(new FSelectionDetailsSummoner(InBlueprintEditor)));
 
 	if ( bRegisterViewport )
 	{
-		TabLayout = FTabManager::NewLayout( "Blueprints_Unified_Components_v2" )
+		TabLayout = FTabManager::NewLayout( "Blueprints_Unified_Components_v5" )
 		->AddArea
 		(
 			FTabManager::NewPrimaryArea()->SetOrientation(Orient_Vertical)
@@ -579,9 +575,26 @@ FBlueprintEditorUnifiedMode::FBlueprintEditorUnifiedMode(TSharedPtr<class FBluep
 			(
 				FTabManager::NewSplitter()->SetOrientation(Orient_Horizontal)
 				->Split
+				(
+					FTabManager::NewSplitter()->SetOrientation(Orient_Vertical)
+					->SetSizeCoefficient(0.15f)
+					->Split
+					(
+						FTabManager::NewStack()
+						->SetSizeCoefficient(0.30f)
+						->AddTab( FBlueprintEditorTabs::ConstructionScriptEditorID, ETabState::OpenedTab )
+					)
+					->Split
+					(
+						FTabManager::NewStack()
+						->SetSizeCoefficient(0.70f)
+						->AddTab( FBlueprintEditorTabs::MyBlueprintID, ETabState::OpenedTab )
+					)
+				)
+				->Split
 				(
 					FTabManager::NewSplitter()->SetOrientation( Orient_Vertical )
-					->SetSizeCoefficient(0.70f)
+					->SetSizeCoefficient(0.60f)
 					->Split
 					(
 						FTabManager::NewStack()
@@ -607,26 +620,13 @@ FBlueprintEditorUnifiedMode::FBlueprintEditorUnifiedMode(TSharedPtr<class FBluep
 				->Split
 				(
 					FTabManager::NewSplitter() ->SetOrientation( Orient_Vertical )
-					->SetSizeCoefficient(0.20f)
-					->Split
-					(
-						FTabManager::NewStack()
-						->SetSizeCoefficient(0.20f)
-						->AddTab( FBlueprintEditorTabs::ConstructionScriptEditorID, ETabState::OpenedTab )
-					)
-					->Split
-					(
-						FTabManager::NewStack()
-						->SetSizeCoefficient(0.30f)
-						->AddTab( FBlueprintEditorTabs::MyBlueprintID, ETabState::OpenedTab )
-					)
+					->SetSizeCoefficient(0.25f)
 					->Split
 					(
 						FTabManager::NewStack()
 						->SetSizeCoefficient(0.50f)
 						->AddTab( FBlueprintEditorTabs::DetailsID, ETabState::OpenedTab )
 						->AddTab( FBlueprintEditorTabs::PaletteID, ETabState::ClosedTab )
-						->AddTab( FBlueprintEditorTabs::DefaultEditorID, ETabState::ClosedTab )
 					)
 				)
 			)
@@ -634,17 +634,28 @@ FBlueprintEditorUnifiedMode::FBlueprintEditorUnifiedMode(TSharedPtr<class FBluep
 	}
 	else
 	{
-		TabLayout = FTabManager::NewLayout( "Blueprints_Unified_v2" )
+		TabLayout = FTabManager::NewLayout( "Blueprints_Unified_v4" )
 		->AddArea
 		(
-			FTabManager::NewPrimaryArea() ->SetOrientation(Orient_Vertical)
-			->Split
+		FTabManager::NewPrimaryArea()->SetOrientation(Orient_Vertical)
+		->Split
 			(
-				FTabManager::NewSplitter() ->SetOrientation(Orient_Horizontal)
+				FTabManager::NewSplitter()->SetOrientation(Orient_Horizontal)
+				->Split
+				(
+					FTabManager::NewSplitter()->SetOrientation(Orient_Vertical)
+					->SetSizeCoefficient(0.15f)
+					->Split
+					(
+						FTabManager::NewStack()
+						->SetSizeCoefficient(0.50f)
+						->AddTab( FBlueprintEditorTabs::MyBlueprintID, ETabState::OpenedTab )
+					)
+				)
 				->Split
 				(
 					FTabManager::NewSplitter()->SetOrientation( Orient_Vertical )
-					->SetSizeCoefficient(0.70f)
+					->SetSizeCoefficient(0.60f)
 					->Split
 					(
 						FTabManager::NewStack()
@@ -669,20 +680,13 @@ FBlueprintEditorUnifiedMode::FBlueprintEditorUnifiedMode(TSharedPtr<class FBluep
 				->Split
 				(
 					FTabManager::NewSplitter() ->SetOrientation( Orient_Vertical )
-					->SetSizeCoefficient(0.20f)
-					->Split
-					(
-						FTabManager::NewStack()
-						->SetSizeCoefficient(0.40f)
-						->AddTab( FBlueprintEditorTabs::MyBlueprintID, ETabState::OpenedTab )
-					)
+					->SetSizeCoefficient(0.25f)
 					->Split
 					(
 						FTabManager::NewStack()
 						->SetSizeCoefficient(0.60f)
 						->AddTab( FBlueprintEditorTabs::DetailsID, ETabState::OpenedTab )
 						->AddTab( FBlueprintEditorTabs::PaletteID, ETabState::ClosedTab )
-						->AddTab( FBlueprintEditorTabs::DefaultEditorID, ETabState::ClosedTab )
 					)
 				)
 			)
diff --git a/Engine/Source/Editor/Kismet/Private/BlueprintEditorTabFactories.cpp b/Engine/Source/Editor/Kismet/Private/BlueprintEditorTabFactories.cpp
--- a/Engine/Source/Editor/Kismet/Private/BlueprintEditorTabFactories.cpp
+++ b/Engine/Source/Editor/Kismet/Private/BlueprintEditorTabFactories.cpp
@@ -16,7 +16,8 @@
 #include "SMyBlueprint.h"
 #include "Engine/Blueprint.h"
 #include "Engine/TimelineTemplate.h"
-
+#include "EditorClassUtils.h"
+#include "SHyperlink.h"
 
 
 #define LOCTEXT_NAMESPACE "BlueprintEditor"
@@ -142,7 +143,82 @@ TSharedRef<SWidget> FDefaultsEditorSummoner::CreateTabBody(const FWorkflowTabSpa
 {
 	TSharedPtr<FBlueprintEditor> BlueprintEditorPtr = StaticCastSharedPtr<FBlueprintEditor>(HostingApp.Pin());
 
-	return BlueprintEditorPtr->GetDefaultEditor();
+	TSharedRef<SWidget> Message = CreateOptionalDataOnlyMessage();
+
+	return SNew(SVerticalBox)
+
+		+ SVerticalBox::Slot()
+		.AutoHeight()
+		.Padding(FMargin(0,0,0,1))
+		[
+			Message
+		]
+
+		+ SVerticalBox::Slot()
+		.FillHeight(1.0f)
+		[
+			BlueprintEditorPtr->GetDefaultEditor()
+		];
+}
+
+TSharedRef<SWidget> FDefaultsEditorSummoner::CreateOptionalDataOnlyMessage() const
+{
+	TSharedPtr<FBlueprintEditor> BlueprintEditorPtr = StaticCastSharedPtr<FBlueprintEditor>(HostingApp.Pin());
+
+	TSharedRef<SWidget> Message = SNullWidget::NullWidget;
+	if ( UBlueprint* Blueprint = BlueprintEditorPtr->GetBlueprintObj() )
+	{
+		if ( FBlueprintEditorUtils::IsDataOnlyBlueprint(Blueprint) )
+		{
+			Message = SNew(SBorder)
+				.Padding(FMargin(5))
+				.BorderImage(FEditorStyle::GetBrush("ToolPanel.GroupBorder"))
+				[
+					SNew(SWrapBox)
+					.UseAllottedWidth(true)
+
+					+ SWrapBox::Slot()
+					[
+						SNew(STextBlock)
+						.Font(FEditorStyle::GetFontStyle("BoldFont"))
+						.Text(LOCTEXT("DataOnlyMessage_Part1", "NOTE: This is a data only blueprint, so only the default values are shown.  It does not have any script or variables.  If you want to add some, "))
+					]
+
+					+ SWrapBox::Slot()
+					[
+						SNew(SHyperlink)
+						.Style(FEditorStyle::Get(), "EditBPHyperlink")
+						.TextStyle(FEditorStyle::Get(), "DetailsView.EditBlueprintHyperlinkStyle")
+						.OnNavigate(this, &FDefaultsEditorSummoner::OnChangeBlueprintToNotDataOnly)
+						.Text(LOCTEXT("FullEditor", "Open Full Blueprint Editor"))
+						.ToolTipText(LOCTEXT("FullEditorToolTip", "This opens the blueprint in the full editor."))
+					]
+				];
+		}
+	}
+
+	return Message;
+}
+
+void FDefaultsEditorSummoner::OnChangeBlueprintToNotDataOnly()
+{
+	UBlueprint* Blueprint = nullptr;
+
+	{
+		TSharedPtr<FBlueprintEditor> BlueprintEditorPtr = StaticCastSharedPtr<FBlueprintEditor>(HostingApp.Pin());
+		Blueprint = BlueprintEditorPtr->GetBlueprintObj();
+		if ( Blueprint )
+		{
+			BlueprintEditorPtr->CloseWindow();
+		}
+	}
+
+	if ( Blueprint )
+	{
+		Blueprint->bForceFullEditor = true;
+
+		GEditor->EditObject(Blueprint);
+	}
 }
 
 FConstructionScriptEditorSummoner::FConstructionScriptEditorSummoner(TSharedPtr<class FAssetEditorToolkit> InHostingApp)
diff --git a/Engine/Source/Editor/Kismet/Private/BlueprintEditorTabFactories.h b/Engine/Source/Editor/Kismet/Private/BlueprintEditorTabFactories.h
--- a/Engine/Source/Editor/Kismet/Private/BlueprintEditorTabFactories.h
+++ b/Engine/Source/Editor/Kismet/Private/BlueprintEditorTabFactories.h
@@ -72,6 +72,11 @@ public:
 	{
 		return LOCTEXT("DefaultsEditorTooltip", "The class defaults editor lets you set the default value for all variables in your Blueprint.");
 	}
+
+private:
+	TSharedRef<SWidget> CreateOptionalDataOnlyMessage() const;
+
+	void OnChangeBlueprintToNotDataOnly();
 };
 
 /////////////////////////////////////////////////////
diff --git a/Engine/Source/Editor/Kismet/Public/BlueprintEditor.h b/Engine/Source/Editor/Kismet/Public/BlueprintEditor.h
--- a/Engine/Source/Editor/Kismet/Public/BlueprintEditor.h
+++ b/Engine/Source/Editor/Kismet/Public/BlueprintEditor.h
@@ -856,9 +856,6 @@ protected:
 	/** Callback when properties have finished being handled */
 	void OnFinishedChangingProperties(const FPropertyChangedEvent& PropertyChangedEvent);
 
-	/** the string to show for edit defaults*/
-	static FText DefaultEditString();
-
 	/** On starting to rename node */
 	void OnRenameNode();
 	bool CanRenameNodes() const;
diff --git a/Engine/Source/Editor/Kismet/Public/BlueprintEditorModes.h b/Engine/Source/Editor/Kismet/Public/BlueprintEditorModes.h
--- a/Engine/Source/Editor/Kismet/Public/BlueprintEditorModes.h
+++ b/Engine/Source/Editor/Kismet/Public/BlueprintEditorModes.h
@@ -136,7 +136,7 @@ protected:
 class KISMET_API FBlueprintEditorUnifiedMode : public FApplicationMode
 {
 public:
-	FBlueprintEditorUnifiedMode(TSharedPtr<class FBlueprintEditor> InBlueprintEditor, FName InModeName, FText(*GetLocalizedMode)( const FName ), const bool bRegisterViewport = true, const bool bRegisterDefaultsTab = true);
+	FBlueprintEditorUnifiedMode(TSharedPtr<class FBlueprintEditor> InBlueprintEditor, FName InModeName, FText(*GetLocalizedMode)( const FName ), const bool bRegisterViewport = true);
 
 	virtual void RegisterTabFactories(TSharedPtr<FTabManager> InTabManager) override;
 	virtual void PreDeactivateMode() override;
diff --git a/Engine/Source/Editor/UnrealEd/Public/Kismet2/BlueprintEditorUtils.h b/Engine/Source/Editor/UnrealEd/Public/Kismet2/BlueprintEditorUtils.h
--- a/Engine/Source/Editor/UnrealEd/Public/Kismet2/BlueprintEditorUtils.h
+++ b/Engine/Source/Editor/UnrealEd/Public/Kismet2/BlueprintEditorUtils.h
@@ -343,7 +343,7 @@ public:
 	/** returns if a graph is an intermediate build product */
 	static bool IsGraphIntermediate(const UEdGraph* Graph);
 
-	/** Returns whether or not this blueprint is a data-only blueprint (has no code) */
+	/** @return true if the blueprint does not contain any special logic or variables or other elements that require a full compile. */
 	static bool IsDataOnlyBlueprint(const UBlueprint* Blueprint);
 
 	/** Returns whether or not the blueprint is const during execution */
@@ -1162,6 +1162,7 @@ protected:
 	 * @param InOutVarDesc		The variable description to validate
 	 */
 	static void PostSetupObjectPinType(UBlueprint* InBlueprint, FBPVariableDescription& InOutVarDesc);
+
 public:
 	static FName GetFunctionNameFromClassByGuid(const UClass* InClass, const FGuid FunctionGuid);
 	static bool GetFunctionGuidFromClassByFieldName(const UClass* InClass, const FName FunctionName, FGuid& FunctionGuid);
diff --git a/Engine/Source/Runtime/Engine/Classes/Engine/Blueprint.h b/Engine/Source/Runtime/Engine/Classes/Engine/Blueprint.h
--- a/Engine/Source/Runtime/Engine/Classes/Engine/Blueprint.h
+++ b/Engine/Source/Runtime/Engine/Classes/Engine/Blueprint.h
@@ -288,29 +288,30 @@ class ENGINE_API UBlueprint : public UBlueprintCore
 #if WITH_EDITORONLY_DATA
 	/** Whether or not this blueprint is newly created, and hasn't been opened in an editor yet */
 	UPROPERTY(transient)
-	uint32 bIsNewlyCreated:1;
+	uint32 bIsNewlyCreated : 1;
 
 	/** Whether to force opening the full (non data-only) editor for this blueprint. */
 	UPROPERTY(transient)
-	uint32 bForceFullEditor:1;
+	uint32 bForceFullEditor : 1;
+
 	/**whether or not you want to continuously rerun the construction script for an actor as you drag it in the editor, or only when the drag operation is complete*/
-	UPROPERTY(EditAnywhere, Category=BlueprintOption)
-	uint32 bRunConstructionScriptOnDrag:1;
+	UPROPERTY(EditAnywhere, Category=BlueprintOptions)
+	uint32 bRunConstructionScriptOnDrag : 1;
 
 	/** Whether or not this blueprint's class is a const class or not.  Should set CLASS_Const in the KismetCompiler. */
-	UPROPERTY(EditAnywhere, Category=BlueprintOption)
-	uint32 bGenerateConstClass:1;
+	UPROPERTY(EditAnywhere, Category=BlueprintOptions)
+	uint32 bGenerateConstClass : 1;
 
 	/**shows up in the content browser when the blueprint is hovered */
-	UPROPERTY(EditAnywhere, Category=BlueprintOption)
+	UPROPERTY(EditAnywhere, Category=BlueprintOptions)
 	FString BlueprintDescription;
 
 	/** The category of the Blueprint, used to organize this Blueprint class when displayed in palette windows */
-	UPROPERTY(EditAnywhere, Category=BlueprintOption)
+	UPROPERTY(EditAnywhere, Category=BlueprintOptions)
 	FString BlueprintCategory;
 
 	/** Additional HideCategories. The are added to HideCategories from parent. */
-	UPROPERTY(EditAnywhere, Category=BlueprintOption)
+	UPROPERTY(EditAnywhere, Category=BlueprintOptions)
 	TArray<FString> HideCategories;
 
 	/** TRUE to show a warning when attempting to start in PIE and there is a compiler error on this Blueprint */
@@ -319,7 +320,7 @@ class ENGINE_API UBlueprint : public UBlueprintCore
 
 	/** Guid key for finding searchable data for Blueprint in the DDC */
 	UPROPERTY()
-	FGuid  SearchGuid;
+	FGuid SearchGuid;
 
 	/** Deprecates the Blueprint, marking the generated class with the CLASS_Deprecated flag */
 	UPROPERTY(EditAnywhere, Category=BlueprintOption)
