[INTEGRATE] Change 2426110 by [REDACTED_PERSON][REDACTED_HANDLE] on 2015/01/30 13:50:20

	Build - Fixing non-unity build.

diff --git a/Engine/Source/Editor/Kismet/Private/BlueprintEditorTabFactories.cpp b/Engine/Source/Editor/Kismet/Private/BlueprintEditorTabFactories.cpp
--- a/Engine/Source/Editor/Kismet/Private/BlueprintEditorTabFactories.cpp
+++ b/Engine/Source/Editor/Kismet/Private/BlueprintEditorTabFactories.cpp
@@ -13,6 +13,7 @@
 #include "SSCSEditorViewport.h"
 #include "SBlueprintPalette.h"
 #include "FindInBlueprints.h"
+#include "Kismet2/BlueprintEditorUtils.h"
 #include "SMyBlueprint.h"
 #include "Engine/Blueprint.h"
 #include "Engine/TimelineTemplate.h"
