[INTEGRATE] Change 2426460 by [REDACTED_PERSON][REDACTED_HANDLE] on 2015/01/30 17:19:58

	BP - Tweaking the override function location so that the function button doesn't drop down.

diff --git a/Engine/Source/Editor/Kismet/Private/SMyBlueprint.cpp b/Engine/Source/Editor/Kismet/Private/SMyBlueprint.cpp
--- a/Engine/Source/Editor/Kismet/Private/SMyBlueprint.cpp
+++ b/Engine/Source/Editor/Kismet/Private/SMyBlueprint.cpp
@@ -498,14 +498,7 @@ FText SMyBlueprint::OnGetSectionTitle( int32 InSectionID )
 		SeperatorTitle = NSLOCTEXT("GraphActionNode", "Components", "Components");
 		break;
 	case NodeSectionID::FUNCTION:
-		if ( OverridableFunctionActions.Num() > 0 )
-		{
-			SeperatorTitle = FText::Format( NSLOCTEXT("GraphActionNode", "FunctionsOverridableFormat", "Functions ({0} Overridable)"), FText::AsNumber(OverridableFunctionActions.Num()) );
-		}
-		else
-		{
-			SeperatorTitle = NSLOCTEXT("GraphActionNode", "Functions", "Functions");
-		}
+		SeperatorTitle = NSLOCTEXT("GraphActionNode", "Functions", "Functions");
 		break;
 	case NodeSectionID::FUNCTION_OVERRIDABLE:
 		SeperatorTitle = NSLOCTEXT("GraphActionNode", "OverridableFunctions", "Overridable Functions");
@@ -552,36 +545,38 @@ TSharedRef<SWidget> SMyBlueprint::OnGetSectionWidget(TSharedRef<SWidget> RowWidg
 		break;
 	case NodeSectionID::FUNCTION:
 		AddNewText = LOCTEXT("AddNewFunction", "Function");
-		return SAssignNew(FunctionSectionButton, SComboButton)
-		.ButtonStyle(FEditorStyle::Get(), "RoundButton")
-		.ForegroundColor(FEditorStyle::GetSlateColor("DefaultForeground"))
-		.ContentPadding(FMargin(2, 0))
-		.OnGetMenuContent(this, &SMyBlueprint::OnGetFunctionListMenu)
-		.HasDownArrow(false)
-		.ButtonContent()
-		[
-			SNew(SHorizontalBox)
 
-			+ SHorizontalBox::Slot()
-			.AutoWidth()
-			.Padding(FMargin(0, 1))
-			[
-				SNew(SImage)
-				.Image(FEditorStyle::GetBrush("Plus"))
-			]
+		if ( OverridableFunctionActions.Num() > 0 )
+		{
+			return SNew(SHorizontalBox)
+				+ SHorizontalBox::Slot()
+				.AutoWidth()
+				[
+					SAssignNew(FunctionSectionButton, SComboButton)
+					.Visibility(this, &SMyBlueprint::OnGetSectionTextVisibility, WeakRowWidget, InSectionID)
+					.ButtonStyle(FEditorStyle::Get(), "RoundButton")
+					.ForegroundColor(FEditorStyle::GetSlateColor("DefaultForeground"))
+					.ContentPadding(FMargin(2, 0))
+					.OnGetMenuContent(this, &SMyBlueprint::OnGetFunctionListMenu)
+					.HasDownArrow(true)
+					.ButtonContent()
+					[
+						SNew(STextBlock)
+						.Font(IDetailLayoutBuilder::GetDetailFontBold())
+						.Text(LOCTEXT("Override", "Override"))
+						.ShadowOffset(FVector2D(1, 1))
+					]
+				]
 
-			+ SHorizontalBox::Slot()
-			.VAlign(VAlign_Center)
-			.AutoWidth()
-			.Padding(FMargin(2,0,0,0))
-			[
-				SNew(STextBlock)
-				.Font(IDetailLayoutBuilder::GetDetailFontBold())
-				.Text(AddNewText)
-				.Visibility(this, &SMyBlueprint::OnGetSectionTextVisibility, WeakRowWidget, InSectionID)
-				.ShadowOffset(FVector2D(1,1))
-			]
-		];
+				+ SHorizontalBox::Slot()
+				.AutoWidth()
+				.Padding(2, 0,0,0)
+				[
+					CreateAddToSectionButton(InSectionID, WeakRowWidget, AddNewText)
+				];
+		}
+
+		break;
 	case NodeSectionID::MACRO:
 		AddNewText = LOCTEXT("AddNewMacro", "Macro");
 		break;
@@ -598,6 +593,11 @@ TSharedRef<SWidget> SMyBlueprint::OnGetSectionWidget(TSharedRef<SWidget> RowWidg
 		return SNullWidget::NullWidget;
 	}
 
+	return CreateAddToSectionButton(InSectionID, WeakRowWidget, AddNewText);
+}
+
+TSharedRef<SWidget> SMyBlueprint::CreateAddToSectionButton(int32 InSectionID, TWeakPtr<SWidget> WeakRowWidget, FText AddNewText)
+{
 	return SNew(SButton)
 		.ButtonStyle(FEditorStyle::Get(), "RoundButton")
 		.ForegroundColor(FEditorStyle::GetSlateColor("DefaultForeground"))
@@ -605,7 +605,6 @@ TSharedRef<SWidget> SMyBlueprint::OnGetSectionWidget(TSharedRef<SWidget> RowWidg
 		.OnClicked(this, &SMyBlueprint::OnAddButtonClickedOnSection, InSectionID)
 		.HAlign(HAlign_Center)
 		.VAlign(VAlign_Center)
-		//.IsEnabled(this, &SMyBlueprint::OnAddButtonClickedOnSection)
 		[
 			SNew(SHorizontalBox)
 
@@ -661,7 +660,7 @@ FReply SMyBlueprint::OnAddButtonClickedOnSection(int32 InSectionID)
 EVisibility SMyBlueprint::OnGetSectionTextVisibility(TWeakPtr<SWidget> RowWidget, int32 InSectionID) const
 {
 	bool ShowText = RowWidget.Pin()->IsHovered();
-	if ( InSectionID == NodeSectionID::FUNCTION && FunctionSectionButton->IsOpen() )
+	if ( InSectionID == NodeSectionID::FUNCTION && FunctionSectionButton.IsValid() && FunctionSectionButton->IsOpen() )
 	{
 		ShowText = true;
 	}
@@ -682,12 +681,6 @@ TSharedRef<SWidget> SMyBlueprint::OnGetFunctionListMenu()
 	const bool bShouldCloseWindowAfterMenuSelection = true;
 	FMenuBuilder MenuBuilder(bShouldCloseWindowAfterMenuSelection, BlueprintEditorPtr.Pin()->GetToolkitCommands());
 
-	MenuBuilder.BeginSection("AddNewItem", LOCTEXT("AddOperations", "Add New"));
-	{
-		MenuBuilder.AddMenuEntry(FBlueprintEditorCommands::Get().AddNewFunction);
-	}
-	MenuBuilder.EndSection();
-
 	BuildOverridableFunctionsMenu(MenuBuilder);
 
 	TSharedRef<SWidget> MenuWidget = MenuBuilder.MakeWidget();
diff --git a/Engine/Source/Editor/Kismet/Private/SMyBlueprint.h b/Engine/Source/Editor/Kismet/Private/SMyBlueprint.h
--- a/Engine/Source/Editor/Kismet/Private/SMyBlueprint.h
+++ b/Engine/Source/Editor/Kismet/Private/SMyBlueprint.h
@@ -155,6 +155,7 @@ private:
 
 	TSharedRef<SWidget> CreateAddNewMenuWidget();
 	void BuildAddNewMenu(FMenuBuilder& MenuBuilder);
+	TSharedRef<SWidget> CreateAddToSectionButton(int32 InSectionID, TWeakPtr<SWidget> WeakRowWidget, FText AddNewText);
 
 	void OnCategoryNameCommitted(const FText& InNewText, ETextCommit::Type InTextCommit, TWeakPtr< struct FGraphActionNode > InAction );
 	bool CanRequestRenameOnActionNode(TWeakPtr<struct FGraphActionNode> InSelectedNode) const;
