[INTEGRATE] Change 2426499 by [REDACTED_PERSON][REDACTED_HANDLE] on 2015/01/30 17:31:35

	Editor: Replace ad-hoc _C suffix removal with two new utilities in FBlueprintEditorUtils
	- GetClassNameWithoutSuffix just removes _C (but now only for BP generated classes)
	- GetFriendlyClassDisplayName looks for DisplayName metadata and potentially pulls a nice localized display name if one is available, falling back to removing _C when necessary

diff --git a/Engine/Source/Editor/BehaviorTreeEditor/Private/BehaviorTreeEditor.cpp b/Engine/Source/Editor/BehaviorTreeEditor/Private/BehaviorTreeEditor.cpp
--- a/Engine/Source/Editor/BehaviorTreeEditor/Private/BehaviorTreeEditor.cpp
+++ b/Engine/Source/Editor/BehaviorTreeEditor/Private/BehaviorTreeEditor.cpp
@@ -2099,8 +2099,7 @@ void FBehaviorTreeEditor::HandleNewNodeClassPicked(UClass* InClass) const
 {
 	if(BehaviorTree != nullptr)
 	{
-		FString ClassName = InClass->GetName();
-		ClassName.RemoveFromEnd(TEXT("_C"));
+		FString ClassName = FBlueprintEditorUtils::GetClassNameWithoutSuffix(InClass);
 
 		FString PathName = BehaviorTree->GetOutermost()->GetPathName();
 		PathName = FPaths::GetPath(PathName);
diff --git a/Engine/Source/Editor/BlueprintGraph/Private/EdGraphSchema_K2.cpp b/Engine/Source/Editor/BlueprintGraph/Private/EdGraphSchema_K2.cpp
--- a/Engine/Source/Editor/BlueprintGraph/Private/EdGraphSchema_K2.cpp
+++ b/Engine/Source/Editor/BlueprintGraph/Private/EdGraphSchema_K2.cpp
@@ -407,14 +407,17 @@ FText UEdGraphSchema_K2::FPinTypeTreeInfo::GetDescription() const
 	}
 	else if (PinType.PinSubCategoryObject.IsValid())
 	{
-		FString DisplayName = PinType.PinSubCategoryObject->GetName();
-		const auto SubCategoryClass = Cast<UClass>(PinType.PinSubCategoryObject.Get());
-		if (SubCategoryClass && !SubCategoryClass->HasAnyClassFlags(CLASS_Native))
+		FText DisplayName;
+		if (UClass* SubCategoryClass = Cast<UClass>(PinType.PinSubCategoryObject.Get()))
 		{
-			DisplayName.RemoveFromEnd(TEXT("_C"));
+			DisplayName = FBlueprintEditorUtils::GetFriendlyClassDisplayName(SubCategoryClass);
+		}
+		else
+		{
+			DisplayName = FText::FromString(PinType.PinSubCategoryObject->GetName());
 		}
 
-		return FText::FromString(DisplayName);
+		return DisplayName;
 	}
 	else
 	{
diff --git a/Engine/Source/Editor/BlueprintGraph/Private/K2Node_CallFunction.cpp b/Engine/Source/Editor/BlueprintGraph/Private/K2Node_CallFunction.cpp
--- a/Engine/Source/Editor/BlueprintGraph/Private/K2Node_CallFunction.cpp
+++ b/Engine/Source/Editor/BlueprintGraph/Private/K2Node_CallFunction.cpp
@@ -443,15 +443,10 @@ FString UK2Node_CallFunction::GetFunctionContextString() const
 		{
 			TrueSelfClass = CurrentSelfClass->GetAuthoritativeClass();
 		}
-		FString TargetString = (TrueSelfClass != NULL) ? TrueSelfClass->GetName() : TEXT("None");
 
-		// This action won't be necessary once the new name convention is used.
-		if(TargetString.EndsWith(TEXT("_C")))
-		{
-			TargetString = TargetString.LeftChop(2);
-		}
+		const FText TargetText = FBlueprintEditorUtils::GetFriendlyClassDisplayName(TrueSelfClass);
 
-		ContextString = FString(TEXT("\n")) + FString::Printf(*LOCTEXT("CallFunctionOnDifferentContext", "Target is %s").ToString(), *TargetString);
+		ContextString = FText::Format(LOCTEXT("CallFunctionOnDifferentContext", "\nTarget is {0}"), TargetText).ToString();
 	}
 
 	return ContextString;
diff --git a/Engine/Source/Editor/BlueprintGraph/Private/K2Node_Event.cpp b/Engine/Source/Editor/BlueprintGraph/Private/K2Node_Event.cpp
--- a/Engine/Source/Editor/BlueprintGraph/Private/K2Node_Event.cpp
+++ b/Engine/Source/Editor/BlueprintGraph/Private/K2Node_Event.cpp
@@ -79,17 +79,11 @@ FText UK2Node_Event::GetNodeTitle(ENodeTitleType::Type TitleType) const
 
 		if(TitleType == ENodeTitleType::FullTitle && EventSignatureClass != NULL && EventSignatureClass->IsChildOf(UInterface::StaticClass()))
 		{
-			FString SourceString = EventSignatureClass->GetName();
-
-			// @todo: This action won't be necessary once the new name convention is used.
-			if(SourceString.EndsWith(TEXT("_C")))
-			{
-				SourceString = SourceString.LeftChop(2);
-			}
+			const FText SignatureClassAsText = FBlueprintEditorUtils::GetFriendlyClassDisplayName(EventSignatureClass);
 
 			FFormatNamedArguments FullTitleArgs;
 			FullTitleArgs.Add(TEXT("Title"), Title);
-			FullTitleArgs.Add(TEXT("InterfaceClass"), FText::FromString(SourceString));
+			FullTitleArgs.Add(TEXT("InterfaceClass"), SignatureClassAsText);
 
 			Title = FText::Format(LOCTEXT("EventFromInterface", "{Title}\nFrom {InterfaceClass}"), FullTitleArgs);
 		}
diff --git a/Engine/Source/Editor/Kismet/Private/SSCSEditor.cpp b/Engine/Source/Editor/Kismet/Private/SSCSEditor.cpp
--- a/Engine/Source/Editor/Kismet/Private/SSCSEditor.cpp
+++ b/Engine/Source/Editor/Kismet/Private/SSCSEditor.cpp
@@ -1628,13 +1628,7 @@ FText SSCS_RowWidget::GetIntroducedInToolTipText() const
 				}
 				else
 				{
-					FString ClassName = (BestClass != nullptr) ? BestClass->GetName() : TEXT("(unknown)");
-					if ((BestClass != nullptr) && (BestClass->ClassGeneratedBy != nullptr) && ClassName.EndsWith(TEXT("_C")))
-					{
-						ClassName.RemoveFromEnd(TEXT("_C"));
-					}
-
-					IntroducedInTooltip = FText::FromString(ClassName);
+					IntroducedInTooltip = FBlueprintEditorUtils::GetFriendlyClassDisplayName(BestClass);
 				}
 			}
 			else
@@ -2600,21 +2594,17 @@ FText SSCS_RowWidget::GetTooltipText() const
 	else
 	{
 		UClass* Class = ( NodePtr->GetComponentTemplate() != nullptr ) ? NodePtr->GetComponentTemplate()->GetClass() : nullptr;
-		FString ClassName = Class != nullptr ? Class->GetName() : TEXT( "(null)" );
-		if( Class != nullptr && Class->ClassGeneratedBy != nullptr && ClassName.EndsWith( TEXT( "_C" ) ) )
-		{
-			ClassName.RemoveFromEnd( TEXT( "_C" ) );
-		}
+		const FText ClassDisplayName = FBlueprintEditorUtils::GetFriendlyClassDisplayName(Class);
 
 		FFormatNamedArguments Args;
-		Args.Add( TEXT("ClassName"), FText::FromString( ClassName ));
-		Args.Add( TEXT("NodeName"), FText::FromString( NodePtr->GetDisplayString() ) );
+		Args.Add(TEXT("ClassName"), ClassDisplayName);
+		Args.Add(TEXT("NodeName"), FText::FromString(NodePtr->GetDisplayString()));
 
-		if ( NodePtr->IsNative() )
+		if (NodePtr->IsNative())
 		{
-			if(NodePtr->IsInstanced())
+			if (NodePtr->IsInstanced())
 			{
-				if ( NodePtr->GetComponentTemplate() != NULL )
+				if (NodePtr->GetComponentTemplate() != NULL)
 				{
 					return FText::Format( LOCTEXT("RegularToolTip", "{ClassName}"), Args );
 				}
@@ -2623,7 +2613,7 @@ FText SSCS_RowWidget::GetTooltipText() const
 			}
 			else
 			{
-				if(NodePtr->GetComponentTemplate() != NULL)
+				if (NodePtr->GetComponentTemplate() != NULL)
 				{
 					return FText::Format( LOCTEXT("NativeClassToolTip", "Native {ClassName}"), Args );
 				}
@@ -2635,7 +2625,7 @@ FText SSCS_RowWidget::GetTooltipText() const
 		{
 			if (NodePtr->IsInheritedSCS())
 			{
-				if ( NodePtr->GetComponentTemplate() != NULL )
+				if (NodePtr->GetComponentTemplate() != NULL)
 				{
 					return FText::Format( LOCTEXT("InheritedToolTip", "Inherited {ClassName}"), Args );
 				}
@@ -2644,7 +2634,7 @@ FText SSCS_RowWidget::GetTooltipText() const
 			}
 			else
 			{
-				if ( NodePtr->GetComponentTemplate() != NULL )
+				if (NodePtr->GetComponentTemplate() != NULL)
 				{
 					return FText::Format( LOCTEXT("RegularToolTip", "{ClassName}"), Args );
 				}
diff --git a/Engine/Source/Editor/UnrealEd/Private/Kismet2/BlueprintEditorUtils.cpp b/Engine/Source/Editor/UnrealEd/Private/Kismet2/BlueprintEditorUtils.cpp
--- a/Engine/Source/Editor/UnrealEd/Private/Kismet2/BlueprintEditorUtils.cpp
+++ b/Engine/Source/Editor/UnrealEd/Private/Kismet2/BlueprintEditorUtils.cpp
@@ -7169,5 +7169,35 @@ void FBlueprintEditorUtils::PostSetupObjectPinType(UBlueprint* InBlueprint, FBPV
 	}
 }
 
+FText FBlueprintEditorUtils::GetFriendlyClassDisplayName(const UClass* Class)
+{
+	if (Class != nullptr)
+	{
+		return Class->GetDisplayNameText();
+	}
+	else
+	{
+		return LOCTEXT("ClassIsNull", "None");
+	}
+}
+
+FString FBlueprintEditorUtils::GetClassNameWithoutSuffix(const UClass* Class)
+{
+	if (Class != nullptr)
+	{
+		FString Result = Class->GetName();
+		if (Class->ClassGeneratedBy != nullptr)
+		{
+			Result.RemoveFromEnd(TEXT("_C"), ESearchCase::CaseSensitive);
+		}
+
+		return Result;
+	}
+	else
+	{
+		return LOCTEXT("ClassIsNull", "None").ToString();
+	}
+}
+
 #undef LOCTEXT_NAMESPACE
 
diff --git a/Engine/Source/Editor/UnrealEd/Private/Kismet2/ComponentEditorUtils.cpp b/Engine/Source/Editor/UnrealEd/Private/Kismet2/ComponentEditorUtils.cpp
--- a/Engine/Source/Editor/UnrealEd/Private/Kismet2/ComponentEditorUtils.cpp
+++ b/Engine/Source/Editor/UnrealEd/Private/Kismet2/ComponentEditorUtils.cpp
@@ -7,7 +7,7 @@
 #include "Engine/SCS_Node.h"
 #include "Engine/SimpleConstructionScript.h"
 #include "ScopedTransaction.h"
-
+#include "BlueprintEditorUtils.h"
 
 USceneComponent* FComponentEditorUtils::GetSceneComponent( UObject* Object, UObject* SubObject /*= NULL*/ )
 {
@@ -122,12 +122,7 @@ FString FComponentEditorUtils::GenerateValidVariableName(TSubclassOf<UActorCompo
 {
 	check(ComponentOwner);
 
-	// Strip off "_C" suffix if it has one
-	FString ComponentTypeName = *ComponentClass->GetName();
-	if (ComponentClass->ClassGeneratedBy && ComponentTypeName.EndsWith(TEXT("_C")))
-	{
-		ComponentTypeName = ComponentTypeName.Left( ComponentTypeName.Len() - 2 );
-	}
+	FString ComponentTypeName = FBlueprintEditorUtils::GetClassNameWithoutSuffix(ComponentClass);
 
 	// Strip off 'Component' if the class ends with that.  It just looks better in the UI.
 	const FString SuffixToStrip( TEXT( "Component" ) );
diff --git a/Engine/Source/Editor/UnrealEd/Public/Kismet2/BlueprintEditorUtils.h b/Engine/Source/Editor/UnrealEd/Public/Kismet2/BlueprintEditorUtils.h
--- a/Engine/Source/Editor/UnrealEd/Public/Kismet2/BlueprintEditorUtils.h
+++ b/Engine/Source/Editor/UnrealEd/Public/Kismet2/BlueprintEditorUtils.h
@@ -1166,4 +1166,14 @@ protected:
 public:
 	static FName GetFunctionNameFromClassByGuid(const UClass* InClass, const FGuid FunctionGuid);
 	static bool GetFunctionGuidFromClassByFieldName(const UClass* InClass, const FName FunctionName, FGuid& FunctionGuid);
+
+	/**
+	 * Returns a friendly class display name for the specified class (removing things like _C from the end, may localize the class name).  Class can be nullptr.
+	 */
+	static FText GetFriendlyClassDisplayName(const UClass* Class);
+
+	/**
+	 * Returns a class name for the specified class that has no automatic suffixes, but is otherwise unmodified.  Class can be nullptr.
+	 */
+	static FString GetClassNameWithoutSuffix(const UClass* Class);
 };
diff --git a/Engine/Source/Runtime/Engine/Private/SimpleConstructionScript.cpp b/Engine/Source/Runtime/Engine/Private/SimpleConstructionScript.cpp
--- a/Engine/Source/Runtime/Engine/Private/SimpleConstructionScript.cpp
+++ b/Engine/Source/Runtime/Engine/Private/SimpleConstructionScript.cpp
@@ -761,11 +761,10 @@ USCS_Node* USimpleConstructionScript::CreateNode(UClass* NewComponentClass, FNam
 	ensure(NULL != Cast<UBlueprintGeneratedClass>(Blueprint->GeneratedClass));
 
 	FName NewComponentName(NAME_None);
-	if (NewComponentClass->ClassGeneratedBy && NewComponentClass->GetName().EndsWith(TEXT("_C")))
+	if (NewComponentClass->ClassGeneratedBy != nullptr)
 	{
-		const FString NewClassName = NewComponentClass->GetName();
-		const int32 NewStrLen = NewClassName.Len() - 2;
-		NewComponentName = MakeUniqueObjectName(Blueprint->GeneratedClass, NewComponentClass, FName(*NewClassName.Left(NewStrLen)));
+		const FString NewClassName = FBlueprintEditorUtils::GetClassNameWithoutSuffix(NewComponentClass);
+		NewComponentName = MakeUniqueObjectName(Blueprint->GeneratedClass, NewComponentClass, FName(*NewClassName));
 	}
 
 	auto NewComponentTemplate = NewObject<UActorComponent>(Blueprint->GeneratedClass, NewComponentClass, NewComponentName);
