[INTEGRATE] Change 2426508 by [REDACTED_PERSON][REDACTED_HANDLE] on 2015/01/30 17:36:02

	BP  - normal BPs now select the CDO by default.

diff --git a/Engine/Source/Editor/Kismet/Private/BlueprintEditor.cpp b/Engine/Source/Editor/Kismet/Private/BlueprintEditor.cpp
--- a/Engine/Source/Editor/Kismet/Private/BlueprintEditor.cpp
+++ b/Engine/Source/Editor/Kismet/Private/BlueprintEditor.cpp
@@ -1551,7 +1551,7 @@ void FBlueprintEditor::InitBlueprintEditor(const EToolkitMode::Type Mode, const
 			{
 				NewDocument_OnClick(CGT_NewMacroGraph);
 			}
-			else if ( Blueprint->BlueprintType == BPTYPE_Interface)
+			else if ( Blueprint->BlueprintType == BPTYPE_Interface )
 			{
 				NewDocument_OnClick(CGT_NewFunctionGraph);
 			}
@@ -1561,6 +1561,17 @@ void FBlueprintEditor::InitBlueprintEditor(const EToolkitMode::Type Mode, const
 			}
 		}
 	}
+
+	if ( UBlueprint* Blueprint = GetBlueprintObj() )
+	{
+		if ( Blueprint->GetClass() == UBlueprint::StaticClass() && Blueprint->BlueprintType == BPTYPE_Normal )
+		{
+			if ( !bShouldOpenInDefaultsMode )
+			{
+				GetToolkitCommands()->ExecuteAction(FFullBlueprintEditorCommands::Get().EditClassDefaults.ToSharedRef());
+			}
+		}
+	}
 }
 
 void FBlueprintEditor::RegisterApplicationModes(const TArray<UBlueprint*>& InBlueprints, bool bShouldOpenInDefaultsMode, bool bNewlyCreated/* = false*/)
