[INTEGRATE] Change 2426515 by [REDACTED_PERSON][REDACTED_HANDLE] on 2015/01/30 17:37:23

	New SCS look and feel on by default.

diff --git a/Engine/Source/Editor/Kismet/Private/SSCSEditor.cpp b/Engine/Source/Editor/Kismet/Private/SSCSEditor.cpp
--- a/Engine/Source/Editor/Kismet/Private/SSCSEditor.cpp
+++ b/Engine/Source/Editor/Kismet/Private/SSCSEditor.cpp
@@ -1343,15 +1343,6 @@ TSharedRef<SWidget> SSCS_RowWidget::GenerateWidgetForColumn( const FName& Column
 						.Image(ComponentIcon)
 						.ColorAndOpacity(this, &SSCS_RowWidget::GetColorTint)
 					]
-				+SHorizontalBox::Slot()
-					.AutoWidth()
-					.VAlign(VAlign_Center)
-					.Padding(4,0,4,0)
-					[
-						SNew(STextBlock)
-						.Text(LOCTEXT("RootLabel", "[ROOT]"))
-						.Visibility(this, &SSCS_RowWidget::GetRootLabelVisibility)
-					]
 				+SHorizontalBox::Slot()
 					.VAlign(VAlign_Center)
 					.Padding(2, 0, 0, 0)
@@ -1436,16 +1427,6 @@ void AddToToolTipInfoBox(const TSharedRef<SVerticalBox>& InfoBox, const FText& K
 
 TSharedRef<SToolTip> SSCS_RowWidget::CreateToolTipWidget() const
 {
-	bool bSingleLayoutBPEditor = GetDefault<UEditorExperimentalSettings>()->bUnifiedBlueprintEditor;
-	if (!bSingleLayoutBPEditor)
-	{
-		return IDocumentation::Get()->CreateToolTip(
-		TAttribute<FText>(this, &SSCS_RowWidget::GetTooltipText),
-		NULL,
-		GetDocumentationLink(),
-		GetDocumentationExcerptName());
-	}
-
 	// Create a box to hold every line of info in the body of the tooltip
 	TSharedRef<SVerticalBox> InfoBox = SNew(SVerticalBox);
 
@@ -1714,25 +1695,6 @@ FSlateColor SSCS_RowWidget::GetColorTint() const
 	}
 }
 
-EVisibility SSCS_RowWidget::GetRootLabelVisibility() const
-{
-	bool bSingleLayoutBPEditor = GetDefault<UEditorExperimentalSettings>()->bUnifiedBlueprintEditor;
-	if (bSingleLayoutBPEditor)
-	{
-		return EVisibility::Collapsed;
-	}
-
-	FSCSEditorTreeNodePtrType NodePtr = GetNode();
-	if(NodePtr.IsValid() && SCSEditor.IsValid() && NodePtr == SCSEditor.Pin()->SceneRootNodePtr)
-	{
-		return EVisibility::Visible;
-	}
-	else
-	{
-		return EVisibility::Collapsed;
-	}
-}
-
 TSharedPtr<SWidget> SSCS_RowWidget::BuildSceneRootDropActionMenu(FSCSEditorTreeNodePtrType DroppedNodePtr)
 {
 	check(SCSEditor.IsValid());
@@ -3021,50 +2983,15 @@ void SSCSEditor::Construct( const FArguments& InArgs )
 
 	FSlateBrush const* MobilityHeaderBrush = FEditorStyle::GetBrush(TEXT("ClassIcon.ComponentMobilityHeaderIcon"));
 	
-	TSharedPtr<SHeaderRow> HeaderRow;
-	
-	bool bSingleLayoutBPEditor = GetDefault<UEditorExperimentalSettings>()->bUnifiedBlueprintEditor;
-	if (bSingleLayoutBPEditor)
-	{
-		HeaderRow = SNew(SHeaderRow)
-
+	TSharedPtr<SHeaderRow> HeaderRow = SNew(SHeaderRow)
 		+ SHeaderRow::Column(SCS_ColumnName_ComponentClass)
 		.DefaultLabel(LOCTEXT("Class", "Class"))
 		.FillWidth(4);
-	}
-	else
-	{
-		HeaderRow = SNew(SHeaderRow)
-
-		+ SHeaderRow::Column(SCS_ColumnName_Mobility)
-		.DefaultLabel(LOCTEXT("MobilityColumnLabel", "Mobility"))
-		.FixedWidth(16.0f) // mobility icons are 16px (16 slate-units = 16px, when application scale == 1)
-		.HeaderContent()
-		[
-			SNew(SHorizontalBox)
-			.ToolTip(SNew(SToolTip).Text(LOCTEXT("MobilityColumnTooltip", "Mobility")))
-			+ SHorizontalBox::Slot()
-			.FillWidth(1.0f)
-			.VAlign(VAlign_Center)
-			.HAlign(HAlign_Center)
-			[
-				SNew(SImage).Image(MobilityHeaderBrush)
-			]
-		]
-
-		+ SHeaderRow::Column(SCS_ColumnName_ComponentClass)
-		.DefaultLabel(LOCTEXT("Class", "Class"))
-		.FillWidth(4)
-
-		+ SHeaderRow::Column(SCS_ColumnName_Asset)
-		.DefaultLabel(LOCTEXT("Asset", "Asset"))\
-		.FillWidth(3);
-	}
 	
 	SCSTreeWidget = SNew(SSCSTreeType)
 		.ToolTipText(LOCTEXT("DropAssetToAddComponent", "Drop asset here to add a component."))
 		.SCSEditor(this)
-		.TreeItemsSource(bSingleLayoutBPEditor ? &RootNodes : &RootComponentNodes)
+		.TreeItemsSource(&RootNodes)
 		.SelectionMode(ESelectionMode::Multi)
 		.OnGenerateRow(this, &SSCSEditor::MakeTableRowWidget)
 		.OnGetChildren(this, &SSCSEditor::OnGetChildrenForTree)
@@ -3079,207 +3006,162 @@ void SSCSEditor::Construct( const FArguments& InArgs )
 			HeaderRow
 		);
 
-	if (bSingleLayoutBPEditor)
-	{
-		SCSTreeWidget->GetHeaderRow()->SetVisibility(EVisibility::Collapsed);
-	}
+	SCSTreeWidget->GetHeaderRow()->SetVisibility(EVisibility::Collapsed);
 
 	TSharedPtr<SWidget> Contents;
 
-	if (bSingleLayoutBPEditor)
-	{
-		FMenuBuilder EditBlueprintMenuBuilder( true, NULL );
+	FMenuBuilder EditBlueprintMenuBuilder( true, NULL );
 
-		EditBlueprintMenuBuilder.BeginSection( NAME_None, LOCTEXT("EditBlueprintMenu_ExistingBlueprintHeader", "Existing Blueprint" ) );
+	EditBlueprintMenuBuilder.BeginSection( NAME_None, LOCTEXT("EditBlueprintMenu_ExistingBlueprintHeader", "Existing Blueprint" ) );
 
-		EditBlueprintMenuBuilder.AddMenuEntry
-		(
-			LOCTEXT("OpenBlueprintEditor", "Open Blueprint Editor"),
-			LOCTEXT("OpenBlueprintEditor_ToolTip", "Opens the blueprint editor for this asset"),
-			FSlateIcon(),
-			FUIAction(FExecuteAction::CreateSP(this, &SSCSEditor::OnOpenBlueprintEditor))
-		);
+	EditBlueprintMenuBuilder.AddMenuEntry
+	(
+		LOCTEXT("OpenBlueprintEditor", "Open Blueprint Editor"),
+		LOCTEXT("OpenBlueprintEditor_ToolTip", "Opens the blueprint editor for this asset"),
+		FSlateIcon(),
+		FUIAction(FExecuteAction::CreateSP(this, &SSCSEditor::OnOpenBlueprintEditor))
+	);
 
 
-		EditBlueprintMenuBuilder.AddMenuEntry
-		(
-			LOCTEXT("PushChangesToBlueprint", "Apply Instance Changes to Blueprint"),
-			TAttribute<FText>(this, &SSCSEditor::OnGetApplyChangesToBlueprintTooltip),
-			FSlateIcon(),
-			FUIAction(FExecuteAction::CreateSP(this, &SSCSEditor::OnApplyChangesToBlueprint))
-		);
+	EditBlueprintMenuBuilder.AddMenuEntry
+	(
+		LOCTEXT("PushChangesToBlueprint", "Apply Instance Changes to Blueprint"),
+		TAttribute<FText>(this, &SSCSEditor::OnGetApplyChangesToBlueprintTooltip),
+		FSlateIcon(),
+		FUIAction(FExecuteAction::CreateSP(this, &SSCSEditor::OnApplyChangesToBlueprint))
+	);
 
-		EditBlueprintMenuBuilder.AddMenuEntry
-		(
-			LOCTEXT("ResetToDefault", "Reset Instance Changes to Blueprint Default"),
-			TAttribute<FText>(this, &SSCSEditor::OnGetResetToBlueprintDefaultsTooltip),
-			FSlateIcon(),
-			FUIAction(FExecuteAction::CreateSP(this, &SSCSEditor::OnResetToBlueprintDefaults))
-		);
+	EditBlueprintMenuBuilder.AddMenuEntry
+	(
+		LOCTEXT("ResetToDefault", "Reset Instance Changes to Blueprint Default"),
+		TAttribute<FText>(this, &SSCSEditor::OnGetResetToBlueprintDefaultsTooltip),
+		FSlateIcon(),
+		FUIAction(FExecuteAction::CreateSP(this, &SSCSEditor::OnResetToBlueprintDefaults))
+	);
 
-		EditBlueprintMenuBuilder.BeginSection( NAME_None, LOCTEXT("EditBlueprintMenu_NewHeader", "Create New" ) );
-		//EditBlueprintMenuBuilder.AddMenuSeparator();
+	EditBlueprintMenuBuilder.BeginSection( NAME_None, LOCTEXT("EditBlueprintMenu_NewHeader", "Create New" ) );
+	//EditBlueprintMenuBuilder.AddMenuSeparator();
 
-		EditBlueprintMenuBuilder.AddMenuEntry
-		(
-			LOCTEXT("PromoteToBlueprint", "Convert to Class Blueprint"),
-			LOCTEXT("PromoteToBluerprintTooltip","Converts the existing Blueprint into a new SubClass Blueprint" ),
-			FSlateIcon(),
-			FUIAction(FExecuteAction::CreateSP(this, &SSCSEditor::PromoteToBlueprint))
-		);
+	EditBlueprintMenuBuilder.AddMenuEntry
+	(
+		LOCTEXT("PromoteToBlueprint", "Convert to Class Blueprint"),
+		LOCTEXT("PromoteToBluerprintTooltip","Converts the existing Blueprint into a new SubClass Blueprint" ),
+		FSlateIcon(),
+		FUIAction(FExecuteAction::CreateSP(this, &SSCSEditor::PromoteToBlueprint))
+	);
+
+	bool bHideComponentClassCombo = InArgs._HideComponentClassCombo.Get();
 
-		Contents = SNew(SVerticalBox)
+	Contents = SNew(SVerticalBox)
+	+ SVerticalBox::Slot()
+	.Padding(0.0f)
+	[
+		SNew(SVerticalBox)
 		+ SVerticalBox::Slot()
-		.Padding(0.0f)
+		.AutoHeight()
+		.VAlign(VAlign_Top)
+		.Padding(0)
 		[
-			SNew(SVerticalBox)
-			+ SVerticalBox::Slot()
-			.AutoHeight()
-			.VAlign(VAlign_Top)
-			.Padding(0)
+			SNew(SBorder)
+			.Padding(2.0f)
+			.BorderImage(FEditorStyle::GetBrush("DetailsView.CategoryTop"))
+			.AddMetaData<FTagMetaData>(FTagMetaData(TEXT("ComponentsPanel")))
+			.BorderBackgroundColor( FLinearColor( .6,.6,.6, 1.0f ) )
 			[
-				SNew(SBorder)
-				.Padding(2.0f)
-				.BorderImage(FEditorStyle::GetBrush("DetailsView.CategoryTop"))
-				.AddMetaData<FTagMetaData>(FTagMetaData(TEXT("ComponentsPanel")))
-				.BorderBackgroundColor( FLinearColor( .6,.6,.6, 1.0f ) )
+				SNew(SHorizontalBox)
+				+ SHorizontalBox::Slot()
+				.HAlign(HAlign_Left)
+				.AutoWidth()
 				[
-					SNew(SHorizontalBox)
-					+ SHorizontalBox::Slot()
-					.HAlign(HAlign_Left)
-					.AutoWidth()
-					[
-						SNew(SComponentClassCombo)
-						.OnComponentClassSelected(this, &SSCSEditor::PerformComboAddClass)
-						.ToolTipText(LOCTEXT("AddComponent_Tooltip", "Adds a new component to this actor"))
-					]
-					+ SHorizontalBox::Slot()
-					.Padding( 4.f, 0.0f )
-					.HAlign(HAlign_Right)
+					SNew(SComponentClassCombo)
+					.Visibility(bHideComponentClassCombo ? EVisibility::Hidden : EVisibility::Visible)
+					.OnComponentClassSelected(this, &SSCSEditor::PerformComboAddClass)
+					.ToolTipText(LOCTEXT("AddComponent_Tooltip", "Adds a new component to this actor"))
+				]
+				+ SHorizontalBox::Slot()
+				.Padding( 4.f, 0.0f )
+				.HAlign(HAlign_Right)
+				[
+					SNew( SButton )
+					.Visibility( this, &SSCSEditor::GetPromoteToBlueprintButtonVisibility )
+					.OnClicked( this, &SSCSEditor::OnPromoteToBlueprintClicked )
+					.ButtonStyle(FEditorStyle::Get(), "ToggleButton")
+					.ContentPadding(0)
+					.ToolTip(IDocumentation::Get()->CreateToolTip(
+						LOCTEXT("PromoteToBluerprintTooltip","Converts this actor into a reusable Class Blueprint that can have script behavior" ),
+						NULL,
+						TEXT("Shared/LevelEditor"),
+						TEXT("ConvertToBlueprint")))
 					[
-						SNew( SButton )
-						.Visibility( this, &SSCSEditor::GetPromoteToBlueprintButtonVisibility )
-						.OnClicked( this, &SSCSEditor::OnPromoteToBlueprintClicked )
-						.ButtonStyle(FEditorStyle::Get(), "ToggleButton")
-						.ContentPadding(0)
-						.ToolTip(IDocumentation::Get()->CreateToolTip(
-							LOCTEXT("PromoteToBluerprintTooltip","Converts this actor into a reusable Class Blueprint that can have script behavior" ),
-							NULL,
-							TEXT("Shared/LevelEditor"),
-							TEXT("ConvertToBlueprint")))
+						SNew(SHorizontalBox)
+						+SHorizontalBox::Slot()
+						.VAlign(VAlign_Center)
+						.AutoWidth()
+						.Padding(2.f, 1.f)
+						[
+							SNew(SImage)
+							.Image( FEditorStyle::Get().GetBrush( "ClassIcon.BlueprintCore" ) )
+						]
+						+ SHorizontalBox::Slot()
+						.VAlign(VAlign_Center)
+						.Padding(1.f)
 						[
-							SNew(SHorizontalBox)
-							+SHorizontalBox::Slot()
-							.VAlign(VAlign_Center)
-							.AutoWidth()
-							.Padding(2.f, 1.f)
-							[
-								SNew(SImage)
-								.Image( FEditorStyle::Get().GetBrush( "ClassIcon.BlueprintCore" ) )
-							]
-							+ SHorizontalBox::Slot()
-							.VAlign(VAlign_Center)
-							.Padding(1.f)
-							[
-								SNew(STextBlock)
-								.TextStyle(FEditorStyle::Get(), "ContentBrowser.TopBar.Font")
-								.Text( LOCTEXT("PromoteToBlueprint", "Convert to Class Blueprint") )
-							]
+							SNew(STextBlock)
+							.TextStyle(FEditorStyle::Get(), "ContentBrowser.TopBar.Font")
+							.Text( LOCTEXT("PromoteToBlueprint", "Convert to Class Blueprint") )
 						]
 					]
-					+ SHorizontalBox::Slot()
-					.Padding(4.f, 0.0f)
-					.HAlign(HAlign_Right)
+				]
+				+ SHorizontalBox::Slot()
+				.Padding(4.f, 0.0f)
+				.HAlign(HAlign_Right)
+				[
+					SNew(SComboButton)
+					.Visibility(this, &SSCSEditor::GetEditBlueprintButtonVisibility)
+					.ContentPadding(FMargin(0))
+					.ComboButtonStyle(FEditorStyle::Get(), "ContentBrowser.NewAsset.Style")
+					.ForegroundColor(FLinearColor::White)
+					.ButtonContent()
 					[
-						SNew(SComboButton)
-						.Visibility(this, &SSCSEditor::GetEditBlueprintButtonVisibility)
-						.ContentPadding(FMargin(0))
-						.ComboButtonStyle(FEditorStyle::Get(), "ContentBrowser.NewAsset.Style")
-						.ForegroundColor(FLinearColor::White)
-						.ButtonContent()
+						SNew( SHorizontalBox )
+						+ SHorizontalBox::Slot()
+						.AutoWidth()
+						.HAlign(HAlign_Center)
+						.VAlign(VAlign_Center)
+						.Padding( 2.0f, 0.0f )
 						[
-							SNew( SHorizontalBox )
-							+ SHorizontalBox::Slot()
-							.AutoWidth()
-							.HAlign(HAlign_Center)
-							.VAlign(VAlign_Center)
-							.Padding( 2.0f, 0.0f )
-							[
-								SNew( SImage )
-								.Image( FEditorStyle::Get().GetBrush( "ClassIcon.BlueprintCore" ) )
-							]
-							.Padding( 2.0f, 0.0f )
-							+ SHorizontalBox::Slot()
-							[
-								SNew(STextBlock)
-								.TextStyle(FEditorStyle::Get(), "ContentBrowser.TopBar.Font")
-								.Text(LOCTEXT("EditBlueprint", "Edit Blueprint"))
-							]
-
+							SNew( SImage )
+							.Image( FEditorStyle::Get().GetBrush( "ClassIcon.BlueprintCore" ) )
 						]
-						.MenuContent()
+						.Padding( 2.0f, 0.0f )
+						+ SHorizontalBox::Slot()
 						[
-							EditBlueprintMenuBuilder.MakeWidget()
+							SNew(STextBlock)
+							.TextStyle(FEditorStyle::Get(), "ContentBrowser.TopBar.Font")
+							.Text(LOCTEXT("EditBlueprint", "Edit Blueprint"))
 						]
-					]
-				]
-			]
 
-			+ SVerticalBox::Slot()
-			.Padding(0.0f, 0.0f)
-			[
-				SNew(SBorder)
-				.Padding(2.0f)
-				.BorderImage(FEditorStyle::GetBrush("SCSEditor.TreePanel"))
-				.AddMetaData<FTagMetaData>(FTagMetaData(TEXT("ComponentsPanel")))
-				[
-					SCSTreeWidget.ToSharedRef()
+					]
+					.MenuContent()
+					[
+						EditBlueprintMenuBuilder.MakeWidget()
+					]
 				]
 			]
-		];
-	}
-	else if( InArgs._HideComponentClassCombo.Get() )
-	{
-		Contents = SNew(SBorder)
-
-		.Padding(2.0f)
-		.BorderImage(FEditorStyle::GetBrush("ToolPanel.GroupBorder"))
-		.AddMetaData<FTagMetaData>(FTagMetaData(TEXT("ComponentsPanel")))
-		[
-			SCSTreeWidget.ToSharedRef()
-		];
-	}
-	else
-	{
-		Contents = SNew(SBorder)
+		]
 
-		.Padding(2.0f)
-		.BorderImage(FEditorStyle::GetBrush("ToolPanel.GroupBorder"))
-		.AddMetaData<FTagMetaData>(FTagMetaData(TEXT("ComponentsPanel")))
+		+ SVerticalBox::Slot()
+		.Padding(0.0f, 0.0f)
 		[
-			SNew(SVerticalBox)
-
-			// Component picker
-			+ SVerticalBox::Slot()
-			.Padding(1.f)
-			.AutoHeight()
-			[
-				SNew(SHorizontalBox)
-				+ SHorizontalBox::Slot()
-				.AutoWidth()
-				[
-					SNew(SComponentClassCombo)
-					.OnComponentClassSelected(this, &SSCSEditor::PerformComboAddClass)
-				]
-			]
-			// Tree
-			+ SVerticalBox::Slot()
-			.Padding(0.f, 0.f, 0.f, 2.f)
+			SNew(SBorder)
+			.Padding(2.0f)
+			.BorderImage(FEditorStyle::GetBrush("SCSEditor.TreePanel"))
+			.AddMetaData<FTagMetaData>(FTagMetaData(TEXT("ComponentsPanel")))
 			[
 				SCSTreeWidget.ToSharedRef()
 			]
-		];
-	}
+		]
+	];
 
 	this->ChildSlot
 	[
@@ -4043,19 +3925,10 @@ void SSCSEditor::UpdateTree(bool bRegenerateTreeNodes)
 		// Reset the scene root node
 		SceneRootNodePtr.Reset();
 
-		TSharedPtr<FSCSEditorTreeNode> ActorTreeNode;
-		bool bSingleLayoutBPEditor = GetDefault<UEditorExperimentalSettings>()->bUnifiedBlueprintEditor;
-		if (bSingleLayoutBPEditor)
-		{
-			ActorTreeNode = MakeShareable(new FSCSEditorTreeNodeRootActor());
-			
-			//RootTreeNode = ActorTreeNode
-			//RootNodes.Add(RootTreeNode);
-			//SCSTreeWidget->SetItemExpansion(RootTreeNode, true);
+		TSharedPtr<FSCSEditorTreeNode> ActorTreeNode  = MakeShareable(new FSCSEditorTreeNodeRootActor());
 
-			RootNodes.Add(ActorTreeNode);
-			RootNodes.Add(MakeShareable(new FSCSEditorTreeNodeSeparator()));
-		}
+		RootNodes.Add(ActorTreeNode);
+		RootNodes.Add(MakeShareable(new FSCSEditorTreeNodeSeparator()));
 
 		// Build the tree data source according to what mode we're in
 		if (EditorMode == EComponentEditorMode::BlueprintSCS)
@@ -4099,7 +3972,7 @@ void SSCSEditor::UpdateTree(bool bRegenerateTreeNodes)
 					UActorComponent* ActorComp = *CompIter;
 					if (!ActorComp->IsA<USceneComponent>())
 					{
-						if (bSingleLayoutBPEditor && !bSeparatorAdded)
+						if (!bSeparatorAdded)
 						{
 							bSeparatorAdded = true;
 							RootNodes.Add(MakeShareable(new FSCSEditorTreeNodeSeparator()));
@@ -4166,7 +4039,7 @@ void SSCSEditor::UpdateTree(bool bRegenerateTreeNodes)
 					UActorComponent* ActorComp = *CompIter;
 					if (!ActorComp->IsA<USceneComponent>() && !ActorComp->IsEditorOnly())
 					{
-						if (bSingleLayoutBPEditor && !bSeparatorAdded)
+						if (!bSeparatorAdded)
 						{
 							bSeparatorAdded = true;
 							RootNodes.Add(MakeShareable(new FSCSEditorTreeNode(FSCSEditorTreeNode::SeparatorNode)));
@@ -4235,13 +4108,7 @@ void SSCSEditor::UpdateTree(bool bRegenerateTreeNodes)
 
 const TArray<FSCSEditorTreeNodePtrType>& SSCSEditor::GetRootNodes() const
 {
-	bool bSingleLayoutBPEditor = GetDefault<UEditorExperimentalSettings>()->bUnifiedBlueprintEditor;
-	if (bSingleLayoutBPEditor)
-	{
-		return RootNodes;
-	}
-	
-	return RootComponentNodes;
+	return RootNodes;
 }
 
 TSharedPtr<FSCSEditorTreeNode> SSCSEditor::AddRootComponentTreeNode(UActorComponent* ActorComp)
diff --git a/Engine/Source/Editor/Kismet/Public/SSCSEditor.h b/Engine/Source/Editor/Kismet/Public/SSCSEditor.h
--- a/Engine/Source/Editor/Kismet/Public/SSCSEditor.h
+++ b/Engine/Source/Editor/Kismet/Public/SSCSEditor.h
@@ -484,8 +484,6 @@ public:
 	FText GetAssetPath() const;
 	EVisibility GetAssetVisibility() const;
 
-	EVisibility GetRootLabelVisibility() const;
-
 	/* Get the node used by the row Widget */
 	virtual FSCSEditorTreeNodePtrType GetNode() const { return TreeNodePtr; };
 
