[INTEGRATE] Change 2426876 by [REDACTED_PERSON][REDACTED_HANDLE] on 2015/01/30 21:33:48

	Component tree: Overhauled the component tree tooltip to display better information consistently across both in-level and BP editor
	- Better display of component 'source' (inherited or not, etc...)
	- Added display of component class tooltip (A static mesh component is a...)
	- Minor formatting tweaks here and there, moved several widgets to reference the editor style
	- Colorized component text based on the type (introduced here, inherited native, inherited BP)
	[UE-8703]

diff --git a/Engine/Source/Editor/EditorStyle/Private/SlateEditorStyle.cpp b/Engine/Source/Editor/EditorStyle/Private/SlateEditorStyle.cpp
--- a/Engine/Source/Editor/EditorStyle/Private/SlateEditorStyle.cpp
+++ b/Engine/Source/Editor/EditorStyle/Private/SlateEditorStyle.cpp
@@ -3926,13 +3926,24 @@ void FSlateEditorStyle::FStyle::SetupGraphEditorStyles()
 			}
 		}
 
-		// SCSEditor
+		// Timeline Editor
+		{
+			Set( "TimelineEditor.AddFloatTrack", new IMAGE_BRUSH( "Icons/icon_TrackAddFloat_36x24px", Icon36x24, FLinearColor::Black ) );
+			Set( "TimelineEditor.AddVectorTrack", new IMAGE_BRUSH( "Icons/icon_TrackAddVector_36x24px", Icon36x24, FLinearColor::Black ) );
+			Set( "TimelineEditor.AddEventTrack", new IMAGE_BRUSH( "Icons/icon_TrackAddEvent_36x24px", Icon36x24, FLinearColor::Black ) );
+			Set( "TimelineEditor.AddColorTrack", new IMAGE_BRUSH( "Icons/icon_TrackAddColor_36x24px", Icon36x24, FLinearColor::Black ) );
+			Set( "TimelineEditor.AddCurveAssetTrack", new IMAGE_BRUSH( "Icons/icon_TrackAddCurve_36x24px", Icon36x24, FLinearColor::Black ) );
+			Set( "TimelineEditor.DeleteTrack", new IMAGE_BRUSH( "Icons/icon_TrackDelete_36x24px", Icon36x24, FLinearColor::Black ) );
+		}
+	}
 
-		Set("SCSEditor.ToggleComponentEditing" , new IMAGE_BRUSH( "Icons/icon_translate_40x", Icon40x40) );
-		Set("SCSEditor.ToggleComponentEditing.Small" , new IMAGE_BRUSH( "Icons/icon_translate_40x", Icon20x20) );
+	// SCSEditor
+	{
+		Set("SCSEditor.ToggleComponentEditing", new IMAGE_BRUSH("Icons/icon_translate_40x", Icon40x40));
+		Set("SCSEditor.ToggleComponentEditing.Small", new IMAGE_BRUSH("Icons/icon_translate_40x", Icon20x20));
 		Set("SCSEditor.TileViewTooltip.NonContentBorder", new BOX_BRUSH("/Docking/TabContentArea", FMargin(4 / 16.0f)));
 
-		Set("SCSEditor.PromoteToBlueprintIcon", new IMAGE_BRUSH("Icons/AssetIcons/Blueprint_16x", Icon16x16) );
+		Set("SCSEditor.PromoteToBlueprintIcon", new IMAGE_BRUSH("Icons/AssetIcons/Blueprint_16x", Icon16x16));
 
 		Set("SCSEditor.TopBar.Font", FTextBlockStyle(NormalText)
 			.SetFont(TTF_CORE_FONT("Fonts/Roboto-Bold", 10))
@@ -3941,18 +3952,52 @@ void FSlateEditorStyle::FStyle::SetupGraphEditorStyles()
 			.SetShadowOffset(FVector2D(1, 1))
 			.SetShadowColorAndOpacity(FLinearColor(0, 0, 0, 0.9f)));
 
-		Set( "SCSEditor.TreePanel", new BOX_BRUSH( "Common/GroupBorder_FlatTop", FMargin(4.0f/16.0f) ) );
+		Set("SCSEditor.TreePanel", new BOX_BRUSH("Common/GroupBorder_FlatTop", FMargin(4.0f / 16.0f)));
 
+		//
 
-		// Timeline Editor
-		{
-			Set( "TimelineEditor.AddFloatTrack", new IMAGE_BRUSH( "Icons/icon_TrackAddFloat_36x24px", Icon36x24, FLinearColor::Black ) );
-			Set( "TimelineEditor.AddVectorTrack", new IMAGE_BRUSH( "Icons/icon_TrackAddVector_36x24px", Icon36x24, FLinearColor::Black ) );
-			Set( "TimelineEditor.AddEventTrack", new IMAGE_BRUSH( "Icons/icon_TrackAddEvent_36x24px", Icon36x24, FLinearColor::Black ) );
-			Set( "TimelineEditor.AddColorTrack", new IMAGE_BRUSH( "Icons/icon_TrackAddColor_36x24px", Icon36x24, FLinearColor::Black ) );
-			Set( "TimelineEditor.AddCurveAssetTrack", new IMAGE_BRUSH( "Icons/icon_TrackAddCurve_36x24px", Icon36x24, FLinearColor::Black ) );
-			Set( "TimelineEditor.DeleteTrack", new IMAGE_BRUSH( "Icons/icon_TrackDelete_36x24px", Icon36x24, FLinearColor::Black ) );
-		}
+		Set("SCSEditor.ComponentTooltip.Title",
+			FTextBlockStyle(NormalText)
+			.SetFont(TTF_CORE_FONT("Fonts/Roboto-Regular", 12))
+			.SetColorAndOpacity(FLinearColor::Black)
+			);
+
+		Set("SCSEditor.ComponentTooltip.Label",
+			FTextBlockStyle(NormalText)
+			.SetColorAndOpacity(FLinearColor(0.2f, 0.2f, 0.2f))
+			.SetShadowOffset(FVector2D(1.0f, 1.0f))
+			.SetShadowColorAndOpacity(FLinearColor(0.8f, 0.8f, 0.8f))
+		);
+		Set("SCSEditor.ComponentTooltip.ImportantLabel",
+			FTextBlockStyle(NormalText)
+			.SetColorAndOpacity(FLinearColor(0.2f, 0.2f, 0.2f))
+			.SetShadowOffset(FVector2D(1.0f, 1.0f))
+			.SetShadowColorAndOpacity(FLinearColor(0.8f, 0.8f, 0.8f))
+		);
+
+
+		Set("SCSEditor.ComponentTooltip.Value", 
+			FTextBlockStyle(NormalText)
+ 			.SetFont(TTF_CORE_FONT("Fonts/Roboto-Bold", 10))
+			.SetColorAndOpacity(FLinearColor(0.0f, 0.0f, 0.0f))
+ 			.SetShadowOffset(FVector2D(1.0f, 1.0f))
+ 			.SetShadowColorAndOpacity(FLinearColor(0.8f, 0.8f, 0.8f))
+		);
+		Set("SCSEditor.ComponentTooltip.ImportantValue",
+			FTextBlockStyle(NormalText)
+			.SetFont(TTF_CORE_FONT("Fonts/Roboto-Bold", 10))
+			.SetColorAndOpacity(FLinearColor(0.3f, 0.0f, 0.0f))
+			.SetShadowOffset(FVector2D(1.0f, 1.0f))
+			.SetShadowColorAndOpacity(FLinearColor(0.8f, 0.8f, 0.8f))
+		);
+
+		Set("SCSEditor.ComponentTooltip.ClassDescription",
+			FTextBlockStyle(NormalText)
+			.SetFont(TTF_CORE_FONT("Fonts/Roboto-Italic", 10))
+			.SetColorAndOpacity(FLinearColor(0.2f, 0.2f, 0.2f))
+			.SetShadowOffset(FVector2D(1.0f, 1.0f))
+			.SetShadowColorAndOpacity(FLinearColor(0.8f, 0.8f, 0.8f))
+		);
 	}
 
 	// Notify editor
diff --git a/Engine/Source/Editor/Kismet/Private/SSCSEditor.cpp b/Engine/Source/Editor/Kismet/Private/SSCSEditor.cpp
--- a/Engine/Source/Editor/Kismet/Private/SSCSEditor.cpp
+++ b/Engine/Source/Editor/Kismet/Private/SSCSEditor.cpp
@@ -1321,7 +1321,8 @@ TSharedRef<SWidget> SSCS_RowWidget::GenerateWidgetForColumn( const FName& Column
 				.OnVerifyTextChanged( this, &SSCS_RowWidget::OnNameTextVerifyChanged )
 				.OnTextCommitted( this, &SSCS_RowWidget::OnNameTextCommit )
 				.IsSelected( this, &SSCS_RowWidget::IsSelectedExclusively )
-				.IsReadOnly( !NodePtr->CanRename() || (SCSEditor.IsValid() && !SCSEditor.Pin()->IsEditingAllowed()) );
+				.ColorAndOpacity( this, &SSCS_RowWidget::GetColorTintForText )
+				.IsReadOnly(!NodePtr->CanRename() || (SCSEditor.IsValid() && !SCSEditor.Pin()->IsEditingAllowed()));
 
 		NodePtr->SetRenameRequestedDelegate(FSCSEditorTreeNode::FOnRenameRequested::CreateSP(InlineWidget.Get(), &SInlineEditableTextBlock::EnterEditingMode));
 		
@@ -1341,7 +1342,7 @@ TSharedRef<SWidget> SSCS_RowWidget::GenerateWidgetForColumn( const FName& Column
 					[
 						SNew(SImage)
 						.Image(ComponentIcon)
-						.ColorAndOpacity(this, &SSCS_RowWidget::GetColorTint)
+						.ColorAndOpacity(this, &SSCS_RowWidget::GetColorTintForIcon)
 					]
 				+SHorizontalBox::Slot()
 					.VAlign(VAlign_Center)
@@ -1394,7 +1395,7 @@ TSharedRef<SWidget> SSCS_RowWidget::GenerateWidgetForColumn( const FName& Column
 	}
 }
 
-void AddToToolTipInfoBox(const TSharedRef<SVerticalBox>& InfoBox, const FText& Key, TSharedRef<SWidget> ValueIcon, const TAttribute<FText>& Value, bool bImportant)
+void SSCS_RowWidget::AddToToolTipInfoBox(const TSharedRef<SVerticalBox>& InfoBox, const FText& Key, TSharedRef<SWidget> ValueIcon, const TAttribute<FText>& Value, bool bImportant)
 {
 	InfoBox->AddSlot()
 		.AutoHeight()
@@ -1406,8 +1407,9 @@ void AddToToolTipInfoBox(const TSharedRef<SVerticalBox>& InfoBox, const FText& K
 			.AutoWidth()
 			.Padding(0, 0, 4, 0)
 			[
-				SNew(STextBlock).Text(FText::Format(LOCTEXT("AssetViewTooltipFormat", "{0}:"), Key))
-				.ColorAndOpacity(bImportant ? FLinearColor( 0.05f, 0.05f, 0.05f ) : FLinearColor( 0.15f, 0.15f, 0.15f ) )
+				SNew(STextBlock)
+				.TextStyle(FEditorStyle::Get(), bImportant ? "SCSEditor.ComponentTooltip.ImportantLabel" : "SCSEditor.ComponentTooltip.Label")
+				.Text(FText::Format(LOCTEXT("AssetViewTooltipFormat", "{0}:"), Key))
 			]
 
 			+ SHorizontalBox::Slot()
@@ -1419,8 +1421,9 @@ void AddToToolTipInfoBox(const TSharedRef<SVerticalBox>& InfoBox, const FText& K
 			+ SHorizontalBox::Slot()
 			.AutoWidth()
 			[
-				SNew(STextBlock).Text(Value)
-				.ColorAndOpacity(bImportant ? FLinearColor( 0.0f, 0.0f, 0.0f ) : FLinearColor( 0.1f, 0.1f, 0.1f ) )
+				SNew(STextBlock)
+				.TextStyle(FEditorStyle::Get(), bImportant ? "SCSEditor.ComponentTooltip.ImportantValue" : "SCSEditor.ComponentTooltip.Value")
+				.Text(Value)
 			]
 		];
 }
@@ -1430,17 +1433,31 @@ TSharedRef<SToolTip> SSCS_RowWidget::CreateToolTipWidget() const
 	// Create a box to hold every line of info in the body of the tooltip
 	TSharedRef<SVerticalBox> InfoBox = SNew(SVerticalBox);
 
-	// Add asset if applicable to this node
-	if (GetAssetVisibility() == EVisibility::Visible)
-	{
-		AddToToolTipInfoBox(InfoBox, LOCTEXT("TooltipAsset", "Asset"), SNullWidget::NullWidget, TAttribute<FText>(this, &SSCS_RowWidget::GetAssetName), false);
-	}
-
+	// 
 	if (FSCSEditorTreeNode* TreeNode = GetNode().Get())
 	{
 		if (TreeNode->GetNodeType() == FSCSEditorTreeNode::ComponentNode)
 		{
+			// Add the tooltip
+			if (UActorComponent* Template = TreeNode->GetComponentTemplate())
+			{
+				UClass* TemplateClass = Template->GetClass();
+				FText ClassTooltip = TemplateClass->GetToolTipText();
+
+				InfoBox->AddSlot()
+					.AutoHeight()
+					.HAlign(HAlign_Center)
+					.Padding(FMargin(0, 2, 0, 4))
+					[
+						SNew(STextBlock)
+						.TextStyle(FEditorStyle::Get(), "SCSEditor.ComponentTooltip.ClassDescription")
+						.Text(ClassTooltip)
+						.WrapTextAt(400.0f)
+					];
+			}
+
 			// Add introduction point
+			AddToToolTipInfoBox(InfoBox, LOCTEXT("TooltipAddType", "Source"), SNullWidget::NullWidget, TAttribute<FText>::Create(TAttribute<FText>::FGetter::CreateSP(this, &SSCS_RowWidget::GetComponentAddSourceToolTipText)), false);
 			if (TreeNode->IsInherited())
 			{
 				AddToToolTipInfoBox(InfoBox, LOCTEXT("TooltipIntroducedIn", "Introduced in"), SNullWidget::NullWidget, TAttribute<FText>::Create(TAttribute<FText>::FGetter::CreateSP(this, &SSCS_RowWidget::GetIntroducedInToolTipText)), false);
@@ -1449,6 +1466,13 @@ TSharedRef<SToolTip> SSCS_RowWidget::CreateToolTipWidget() const
 			// Add mobility
 			TSharedRef<SImage> MobilityIcon = SNew(SImage).Image(this, &SSCS_RowWidget::GetMobilityIconImage);
 			AddToToolTipInfoBox(InfoBox, LOCTEXT("TooltipMobility", "Mobility"), MobilityIcon, TAttribute<FText>::Create(TAttribute<FText>::FGetter::CreateSP(this, &SSCS_RowWidget::GetMobilityToolTipText)), false);
+
+			// Add asset if applicable to this node
+			if (GetAssetVisibility() == EVisibility::Visible)
+			{
+				InfoBox->AddSlot()[SNew(SSpacer).Size(FVector2D(1.0f, 8.0f))];
+				AddToToolTipInfoBox(InfoBox, LOCTEXT("TooltipAsset", "Asset"), SNullWidget::NullWidget, TAttribute<FText>(this, &SSCS_RowWidget::GetAssetName), false);
+			}
 		}
 	}
 
@@ -1472,12 +1496,11 @@ TSharedRef<SToolTip> SSCS_RowWidget::CreateToolTipWidget() const
 					+ SHorizontalBox::Slot()
 					.AutoWidth()
 					.VAlign(VAlign_Center)
-					.Padding(4)
+					.Padding(2)
 					[
 						SNew(STextBlock)
+						.TextStyle(FEditorStyle::Get(), "SCSEditor.ComponentTooltip.Title")
 						.Text(this, &SSCS_RowWidget::GetTooltipText)
-						.ColorAndOpacity( FLinearColor::Black )
-						.Font(FEditorStyle::GetFontStyle("ContentBrowser.TileViewTooltip.NameFont"))
 					]
 				]
 			]
@@ -1487,7 +1510,7 @@ TSharedRef<SToolTip> SSCS_RowWidget::CreateToolTipWidget() const
 			[
 				SNew(SBorder)
 				.BorderImage(FEditorStyle::GetBrush("NoBorder"))
-				.Padding(4)
+				.Padding(2)
 				[
 					InfoBox
 				]
@@ -1531,15 +1554,15 @@ FText SSCS_RowWidget::GetMobilityToolTipText() const
 		{
 			if (SceneComponentTemplate->Mobility == EComponentMobility::Movable)
 			{
-				MobilityToolTip = LOCTEXT("MovableMobilityTooltip", "Movable component");
+				MobilityToolTip = LOCTEXT("MovableMobilityTooltip", "Movable");
 			}
 			else if (SceneComponentTemplate->Mobility == EComponentMobility::Stationary)
 			{
-				MobilityToolTip = LOCTEXT("StationaryMobilityTooltip", "Stationary component");
+				MobilityToolTip = LOCTEXT("StationaryMobilityTooltip", "Stationary");
 			}
 			else if (SceneComponentTemplate->Mobility == EComponentMobility::Static)
 			{
-				MobilityToolTip = LOCTEXT("StaticMobilityTooltip", "Static component");
+				MobilityToolTip = LOCTEXT("StaticMobilityTooltip", "Static");
 			}
 			else
 			{
@@ -1557,9 +1580,42 @@ FText SSCS_RowWidget::GetMobilityToolTipText() const
 	return MobilityToolTip;
 }
 
+FText SSCS_RowWidget::GetComponentAddSourceToolTipText() const
+{
+	FText NodeType;
+	
+	if (FSCSEditorTreeNode* TreeNode = TreeNodePtr.Get())
+	{
+		if (TreeNode->IsInherited())
+		{
+			if (TreeNode->IsNative())
+			{
+				NodeType = LOCTEXT("InheritedNativeComponent", "Inherited (C++)");
+			}
+			else
+			{
+				NodeType = LOCTEXT("InheritedBlueprintComponent", "Inherited (Blueprint)");
+			}
+		}
+		else
+		{
+			if (TreeNode->IsInstanced())
+			{
+				NodeType = LOCTEXT("ThisInstanceAddedComponent", "This actor instance");
+			}
+			else
+			{
+				NodeType = LOCTEXT("ThisBlueprintAddedComponent", "This Blueprint");
+			}
+		}
+	}
+
+	return NodeType;
+}
+
 FText SSCS_RowWidget::GetIntroducedInToolTipText() const
 {
-	FText IntroducedInTooltip = LOCTEXT("IntroducedHereTooltip", "this class");
+	FText IntroducedInTooltip = LOCTEXT("IntroducedInThisBPTooltip", "this class");
 
 	if (FSCSEditorTreeNode* TreeNode = TreeNodePtr.Get())
 	{
@@ -1617,6 +1673,10 @@ FText SSCS_RowWidget::GetIntroducedInToolTipText() const
 				IntroducedInTooltip = LOCTEXT("IntroducedInNoTemplateError", "[no component template found]");
 			}
 		}
+		else if (TreeNode->IsInstanced())
+		{
+			IntroducedInTooltip = LOCTEXT("IntroducedInThisActorInstanceTooltip", "this actor instance");
+		}
 	}
 
 	return IntroducedInTooltip;
@@ -1671,10 +1731,40 @@ EVisibility SSCS_RowWidget::GetAssetVisibility() const
 	}
 }
 
-FSlateColor SSCS_RowWidget::GetColorTint() const
+FSlateColor SSCS_RowWidget::GetColorTintForIcon() const
 {
 	const FLinearColor InheritedBlueprintComponentColor(0.08f, 0.35f, 0.6f);
-	const FLinearColor InheritedNativeComponentColor(0.08f, 0.15f, 0.6f);
+	const FLinearColor InstancedInheritedBlueprintComponentColor(0.08f, 0.35f, 0.6f);
+	const FLinearColor InheritedNativeComponentColor(0.7f, 0.9f, 0.7f);
+	const FLinearColor IntroducedHereColor(FLinearColor::White);
+
+	FSCSEditorTreeNodePtrType NodePtr = GetNode();
+	if (NodePtr->IsInherited())
+	{
+		if (NodePtr->IsNative())
+		{
+			return InheritedNativeComponentColor;
+		}
+		else if (NodePtr->IsInstanced())
+		{
+			return InstancedInheritedBlueprintComponentColor;
+		}
+		else
+		{
+			return InheritedBlueprintComponentColor;
+		}
+	}
+	else
+	{
+		return IntroducedHereColor;
+	}
+}
+
+FSlateColor SSCS_RowWidget::GetColorTintForText() const
+{
+	const FLinearColor InheritedBlueprintComponentColor(0.6f, 0.6f, 0.8f);
+	const FLinearColor InstancedInheritedBlueprintComponentColor(InheritedBlueprintComponentColor);
+	const FLinearColor InheritedNativeComponentColor(0.7f, 0.9f, 0.7f);
 	const FLinearColor IntroducedHereColor(FLinearColor::White);
 
 	FSCSEditorTreeNodePtrType NodePtr = GetNode();
@@ -1684,6 +1774,10 @@ FSlateColor SSCS_RowWidget::GetColorTint() const
 		{
 			return InheritedNativeComponentColor;
 		}
+		else if (NodePtr->IsInstanced())
+		{
+			return InstancedInheritedBlueprintComponentColor;
+		}
 		else
 		{
 			return InheritedBlueprintComponentColor;
@@ -2557,53 +2651,14 @@ FText SSCS_RowWidget::GetTooltipText() const
 	{
 		UClass* Class = ( NodePtr->GetComponentTemplate() != nullptr ) ? NodePtr->GetComponentTemplate()->GetClass() : nullptr;
 		const FText ClassDisplayName = FBlueprintEditorUtils::GetFriendlyClassDisplayName(Class);
+		const FText ComponentDisplayName = NodePtr->GetDisplayName();
+
 
 		FFormatNamedArguments Args;
 		Args.Add(TEXT("ClassName"), ClassDisplayName);
 		Args.Add(TEXT("NodeName"), FText::FromString(NodePtr->GetDisplayString()));
 
-		if (NodePtr->IsNative())
-		{
-			if (NodePtr->IsInstanced())
-			{
-				if (NodePtr->GetComponentTemplate() != NULL)
-				{
-					return FText::Format( LOCTEXT("RegularToolTip", "{ClassName}"), Args );
-				}
-
-				return FText::Format( LOCTEXT("MissingRegularComponentToolTip", "MISSING!! {NodeName}"), Args );
-			}
-			else
-			{
-				if (NodePtr->GetComponentTemplate() != NULL)
-				{
-					return FText::Format( LOCTEXT("NativeClassToolTip", "Native {ClassName}"), Args );
-				}
-
-				return FText::Format( LOCTEXT("MissingNativeComponentToolTip", "MISSING!! Native {NodeName}"), Args );
-			}
-		}
-		else
-		{
-			if (NodePtr->IsInheritedSCS())
-			{
-				if (NodePtr->GetComponentTemplate() != NULL)
-				{
-					return FText::Format( LOCTEXT("InheritedToolTip", "Inherited {ClassName}"), Args );
-				}
-
-				return FText::Format( LOCTEXT("MissingInheritedComponentToolTip", "MISSING!! Inherited {NodeName}"), Args );
-			}
-			else
-			{
-				if (NodePtr->GetComponentTemplate() != NULL)
-				{
-					return FText::Format( LOCTEXT("RegularToolTip", "{ClassName}"), Args );
-				}
-
-				return FText::Format( LOCTEXT("MissingRegularComponentToolTip", "MISSING!! {NodeName}"), Args );
-			}
-		}
+		return FText::Format(LOCTEXT("ComponentTooltip", "{NodeName} ({ClassName})"), Args);
 	}
 }
 
@@ -2790,9 +2845,8 @@ TSharedRef<SToolTip> SSCS_RowWidget_ActorRoot::CreateToolTipWidget() const
 					.Padding(4)
 					[
 						SNew(STextBlock)
+						.TextStyle(FEditorStyle::Get(), "SCSEditor.ComponentTooltip.Title")
 						.Text(this, &SSCS_RowWidget_ActorRoot::GetActorDisplayText)
-						.ColorAndOpacity( FLinearColor::Black )
-						.Font(FEditorStyle::GetFontStyle("ContentBrowser.TileViewTooltip.NameFont"))
 					]
 				]
 			]
diff --git a/Engine/Source/Editor/Kismet/Public/SSCSEditor.h b/Engine/Source/Editor/Kismet/Public/SSCSEditor.h
--- a/Engine/Source/Editor/Kismet/Public/SSCSEditor.h
+++ b/Engine/Source/Editor/Kismet/Public/SSCSEditor.h
@@ -476,7 +476,8 @@ public:
 
 	FText GetNameLabel() const;
 	FText GetTooltipText() const;
-	FSlateColor GetColorTint() const;
+	FSlateColor GetColorTintForIcon() const;
+	FSlateColor GetColorTintForText() const;
 	FString GetDocumentationLink() const;
 	FString GetDocumentationExcerptName() const;
 	
@@ -490,6 +491,8 @@ public:
 protected:
 	virtual ESelectionMode::Type GetSelectionMode() const override;
 
+	static void AddToToolTipInfoBox(const TSharedRef<SVerticalBox>& InfoBox, const FText& Key, TSharedRef<SWidget> ValueIcon, const TAttribute<FText>& Value, bool bImportant);
+
 private:
 	/** Verifies the name of the component when changing it */
 	bool OnNameTextVerifyChanged(const FText& InNewText, FText& OutErrorMessage);
@@ -543,6 +546,14 @@ private:
 	 * @returns An FText object containing a description of when the component was first introduced
 	 */
 	FText GetIntroducedInToolTipText() const;
+
+	/**
+	 * Retrieves tooltip text describing how the component was introduced
+	 * 
+	 * @returns An FText object containing a description of when the component was first introduced
+	 */
+	FText GetComponentAddSourceToolTipText() const;
+
 public:
 	/** Pointer back to owning SCSEditor 2 tool */
 	TWeakPtr<SSCSEditor> SCSEditor;
