[INTEGRATE] Change 2426915 by [REDACTED_PERSON][REDACTED_HANDLE] on 2015/01/30 22:52:50

	Blueprint Editor: Allowed drag-dropping components and component classes into the BP editor 3D viewport, which behaves the same as dropping them in the components tree, creating new components and attaching them to the root where it makes sense [UE-8357]

diff --git a/Engine/Source/Editor/Kismet/Private/SSCSEditor.cpp b/Engine/Source/Editor/Kismet/Private/SSCSEditor.cpp
--- a/Engine/Source/Editor/Kismet/Private/SSCSEditor.cpp
+++ b/Engine/Source/Editor/Kismet/Private/SSCSEditor.cpp
@@ -118,65 +118,72 @@ FReply SSCSEditorDragDropTree::OnDragOver( const FGeometry& MyGeometry, const FD
 	return Handled;
 }
 
-FReply SSCSEditorDragDropTree::OnDrop( const FGeometry& MyGeometry, const FDragDropEvent& DragDropEvent ) 
+FReply SSCSEditor::TryHandleAssetDragDropOperation(const FDragDropEvent& DragDropEvent)
 {
-	FReply Handled = FReply::Unhandled();
-
-	if (SCSEditor != nullptr)
+	TSharedPtr<FDragDropOperation> Operation = DragDropEvent.GetOperation();
+	if (Operation.IsValid() && (Operation->IsOfType<FExternalDragOperation>() || Operation->IsOfType<FAssetDragDropOp>()))
 	{
-		TSharedPtr<FDragDropOperation> Operation = DragDropEvent.GetOperation();
-		if (Operation.IsValid() && (Operation->IsOfType<FExternalDragOperation>() || Operation->IsOfType<FAssetDragDropOp>()))
+		TArray< FAssetData > DroppedAssetData = AssetUtil::ExtractAssetDataFromDrag(DragDropEvent);
+		const int32 NumAssets = DroppedAssetData.Num();
+
+		if (NumAssets > 0)
 		{
-			TArray< FAssetData > DroppedAssetData = AssetUtil::ExtractAssetDataFromDrag( DragDropEvent );
-			const int32 NumAssets = DroppedAssetData.Num();
+			GWarn->BeginSlowTask(LOCTEXT("LoadingComponents", "Loading Component(s)"), true);
 
-			if (NumAssets > 0)
+			for (int32 DroppedAssetIdx = 0; DroppedAssetIdx < NumAssets; ++DroppedAssetIdx)
 			{
-				GWarn->BeginSlowTask( LOCTEXT("LoadingComponents", "Loading Component(s)"), true );
+				const FAssetData& AssetData = DroppedAssetData[DroppedAssetIdx];
 
-				for (int32 DroppedAssetIdx = 0; DroppedAssetIdx < NumAssets; ++DroppedAssetIdx)
+				if (!AssetData.IsAssetLoaded())
 				{
-					const FAssetData& AssetData = DroppedAssetData[DroppedAssetIdx];
-
-					if (!AssetData.IsAssetLoaded())
-					{
-						GWarn->StatusUpdate(DroppedAssetIdx, NumAssets, FText::Format(LOCTEXT("LoadingAsset", "Loading Asset {0}"), FText::FromName(AssetData.AssetName)));
-					}
+					GWarn->StatusUpdate(DroppedAssetIdx, NumAssets, FText::Format(LOCTEXT("LoadingAsset", "Loading Asset {0}"), FText::FromName(AssetData.AssetName)));
+				}
 
-					UClass* AssetClass = AssetData.GetClass();
-					UObject* Asset = AssetData.GetAsset();
+				UClass* AssetClass = AssetData.GetClass();
+				UObject* Asset = AssetData.GetAsset();
 
-					UBlueprint* BPClass = Cast<UBlueprint>(Asset);
-					UClass* PotentialComponentClass = nullptr;
-					
-					if ((BPClass != nullptr) && (BPClass->GeneratedClass != nullptr) && (BPClass->GeneratedClass->IsChildOf(UActorComponent::StaticClass())))
-					{
-						PotentialComponentClass = BPClass->GeneratedClass;
-					}
-					else if (AssetClass->IsChildOf(UClass::StaticClass()) && CastChecked<UClass>(Asset)->IsChildOf(UActorComponent::StaticClass()))
-					{
-						PotentialComponentClass = Cast<UClass>(Asset);
-					}
+				UBlueprint* BPClass = Cast<UBlueprint>(Asset);
+				UClass* PotentialComponentClass = nullptr;
 
-					TSubclassOf<UActorComponent>  MatchingComponentClassForAsset = FComponentAssetBrokerage::GetPrimaryComponentForAsset( AssetClass );
-					if (MatchingComponentClassForAsset != nullptr)
-					{
-						SCSEditor->AddNewComponent(MatchingComponentClassForAsset, Asset);
-					}
-					else if ((PotentialComponentClass != nullptr) && (!PotentialComponentClass->HasAnyClassFlags(CLASS_Deprecated | CLASS_Abstract | CLASS_NewerVersionExists)) && PotentialComponentClass->HasMetaData(FBlueprintMetadata::MD_BlueprintSpawnableComponent))
-					{
-						SCSEditor->AddNewComponent(PotentialComponentClass, nullptr);
-					}
+				if ((BPClass != nullptr) && (BPClass->GeneratedClass != nullptr) && (BPClass->GeneratedClass->IsChildOf(UActorComponent::StaticClass())))
+				{
+					PotentialComponentClass = BPClass->GeneratedClass;
+				}
+				else if (AssetClass->IsChildOf(UClass::StaticClass()) && CastChecked<UClass>(Asset)->IsChildOf(UActorComponent::StaticClass()))
+				{
+					PotentialComponentClass = Cast<UClass>(Asset);
 				}
 
-				GWarn->EndSlowTask();
+				TSubclassOf<UActorComponent>  MatchingComponentClassForAsset = FComponentAssetBrokerage::GetPrimaryComponentForAsset(AssetClass);
+				if (MatchingComponentClassForAsset != nullptr)
+				{
+					AddNewComponent(MatchingComponentClassForAsset, Asset);
+				}
+				else if ((PotentialComponentClass != nullptr) && (!PotentialComponentClass->HasAnyClassFlags(CLASS_Deprecated | CLASS_Abstract | CLASS_NewerVersionExists)) && PotentialComponentClass->HasMetaData(FBlueprintMetadata::MD_BlueprintSpawnableComponent))
+				{
+					AddNewComponent(PotentialComponentClass, nullptr);
+				}
 			}
 
-			Handled = FReply::Handled();
+			GWarn->EndSlowTask();
 		}
+
+		return FReply::Handled();
 	}
 
-	return Handled;
+	return FReply::Unhandled();
+}
+
+FReply SSCSEditorDragDropTree::OnDrop( const FGeometry& MyGeometry, const FDragDropEvent& DragDropEvent ) 
+{
+	if (SCSEditor != nullptr)
+	{
+		return SCSEditor->TryHandleAssetDragDropOperation(DragDropEvent);
+	}
+	else
+	{
+		return FReply::Unhandled();
+	}
 }
 
 
diff --git a/Engine/Source/Editor/Kismet/Private/SSCSEditorViewport.cpp b/Engine/Source/Editor/Kismet/Private/SSCSEditorViewport.cpp
--- a/Engine/Source/Editor/Kismet/Private/SSCSEditorViewport.cpp
+++ b/Engine/Source/Editor/Kismet/Private/SSCSEditorViewport.cpp
@@ -350,6 +350,12 @@ bool SSCSEditorViewport::GetIsSimulateEnabled()
 	return ViewportClient->GetIsSimulateEnabled();
 }
 
+FReply SSCSEditorViewport::OnDrop(const FGeometry& MyGeometry, const FDragDropEvent& DragDropEvent)
+{
+	TSharedPtr<SSCSEditor> SCSEditor = BlueprintEditorPtr.Pin()->GetSCSEditor();
+	return SCSEditor->TryHandleAssetDragDropOperation(DragDropEvent);
+}
+
 EActiveTimerReturnType SSCSEditorViewport::DeferredUpdatePreview(double InCurrentTime, float InDeltaTime, bool bResetCamera)
 {
 	if (ViewportClient.IsValid())
diff --git a/Engine/Source/Editor/Kismet/Private/SSCSEditorViewport.h b/Engine/Source/Editor/Kismet/Private/SSCSEditorViewport.h
--- a/Engine/Source/Editor/Kismet/Private/SSCSEditorViewport.h
+++ b/Engine/Source/Editor/Kismet/Private/SSCSEditorViewport.h
@@ -46,6 +46,10 @@ public:
 	 */
 	void RequestRefresh(bool bResetCamera = false, bool bRefreshNow = false);
 
+	// SWidget interface
+	virtual FReply OnDrop(const FGeometry& MyGeometry, const FDragDropEvent& DragDropEvent) override;
+	// End of SWidget interface
+
 	/**
 	 * Called when the selected component changes in the SCS editor.
 	 */
diff --git a/Engine/Source/Editor/Kismet/Public/SSCSEditor.h b/Engine/Source/Editor/Kismet/Public/SSCSEditor.h
--- a/Engine/Source/Editor/Kismet/Public/SSCSEditor.h
+++ b/Engine/Source/Editor/Kismet/Public/SSCSEditor.h
@@ -869,6 +869,9 @@ public:
 
 	/** @return The current editor mode (editing live actors or editing blueprints) */
 	EComponentEditorMode::Type GetEditorMode() const { return EditorMode; }
+
+	/** Try to handle a drag-drop operation */
+	FReply TryHandleAssetDragDropOperation(const FDragDropEvent& DragDropEvent);
 protected:
 	FString GetSelectedClassText() const;
 
