Batched blueprint searches to speed up global find in blueprints.

diff --git a/Engine/Source/Editor/Kismet/Private/FindInBlueprints.cpp b/Engine/Source/Editor/Kismet/Private/FindInBlueprints.cpp
--- a/Engine/Source/Editor/Kismet/Private/FindInBlueprints.cpp
+++ b/Engine/Source/Editor/Kismet/Private/FindInBlueprints.cpp
@@ -884,40 +884,7 @@ public:
 		return true;
 	}
 
-	virtual uint32 Run() override
-	{
-		FFindInBlueprintSearchManager::Get().BeginSearchQuery(this);
-
-		FSearchData QueryResult;
-
-		// Searching comes to an end if it is requested using the StopTaskCounter or continuing the search query yields no results
-		while(StopTaskCounter.GetValue() == 0 && FFindInBlueprintSearchManager::Get().ContinueSearchQuery(this, QueryResult))
-		{
-			if(QueryResult.BlueprintPath.Len())
-			{
-				//Each blueprints acts as a category
-				FSearchResult BlueprintCategory = FSearchResult(new FFindInBlueprintsResult(FText::FromString(QueryResult.BlueprintPath)));
-				FindInBlueprintsHelpers::Extract(QueryResult.Value, SearchTokens, BlueprintCategory);
-
-				// If there are children, add the item to the search results
-				if(BlueprintCategory->Children.Num() > 0)
-				{
-					FScopeLock ScopeLock(&SearchCriticalSection);
-					ItemsFound.Add(BlueprintCategory);
-				}
-			}
-		}
-
-		if(StopTaskCounter.GetValue() > 0)
-		{
-			// Ensure that the FiB Manager knows that we are done searching
-			FFindInBlueprintSearchManager::Get().EnsureSearchQueryEnds(this);
-		}
-
-		bThreadCompleted = true;
-
-		return 0;
-	}
+	virtual uint32 Run() override;
 
 	virtual void Stop() override
 	{
@@ -930,6 +897,8 @@ public:
 	}
 	/** End FRunnable Interface */
 
+	TArray<FSearchData> MakeExtractionBatch() const;
+
 	/** Brings the thread to a safe stop before continuing. */
 	void EnsureCompletion()
 	{
@@ -1476,4 +1445,135 @@ void SFindInBlueprints::OnCopyAction()
 	FPlatformMisc::ClipboardCopy( *SelectedText );
 }
 
+////////////////////////////////////
+// FExtractionTask
+
+/**
+ * Task to search a batch of blueprints
+ *
+ * Each small batch is non-abandonable, but if the search is cancelled, no more batches are queued
+ */
+class FExtractionTask : public FNonAbandonableTask
+{
+public:
+	/** Async task function */
+	void DoWork();
+
+	/** Batch of blueprints and associated data to search */
+	TArray<FSearchData> QueryResultBatch;
+
+	/** Pointer to array of tokens to search for */
+	const TArray<FString>* Tokens;
+
+	/** Function to call once the search result is ready */
+	TFunction<void(const FSearchResult&)> OnResultReady;
+};
+
+TArray<FSearchData> FStreamSearch::MakeExtractionBatch() const
+{
+	// Batch sizes between 10 and 30 seem fastest
+	static const int BatchSize = 15;
+
+	TArray<FSearchData> Batch;
+	do
+	{
+		if (StopTaskCounter.GetValue())
+			return TArray<FSearchData>();
+
+		FSearchData QueryResult;
+		if (!FFindInBlueprintSearchManager::Get().ContinueSearchQuery(this, QueryResult))
+			return Batch;
+
+		Batch.Push(QueryResult);
+	}
+	while (Batch.Num() <= BatchSize);
+
+	return Batch;
+}
+
+uint32 FStreamSearch::Run()
+{
+	FFindInBlueprintSearchManager::Get().BeginSearchQuery(this);
+
+	// Seems to go fastest at or slightly above one thread per physical core
+	TArray<FAsyncTask<FExtractionTask>> ExtractionTasks;
+	ExtractionTasks.SetNum(FPlatformMisc::NumberOfCores());
+
+	TFunction<void(const FSearchResult&)> OnResultReady = [this](const FSearchResult& Result) {
+		FScopeLock ScopeLock(&SearchCriticalSection);
+		ItemsFound.Add(Result);
+	};
+
+	// Searching comes to an end if it is requested using the StopTaskCounter or continuing the search query yields no results
+	TArray<FSearchData> QueryResultBatch;
+	for (;;)
+	{
+		if (QueryResultBatch.Num() == 0)
+		{
+			QueryResultBatch = MakeExtractionBatch();
+		}
+		else
+		{
+			// Spinning until a task becomes available
+			FPlatformProcess::Sleep(0);
+		}
+
+		if (QueryResultBatch.Num() == 0)
+		{
+			// All search work has been dispatched
+			break;
+		}
+
+		for (auto& Task: ExtractionTasks)
+		{
+			if (Task.IsWorkDone())
+			{
+				Task.EnsureCompletion();
+
+				auto& worker = Task.GetTask();
+				worker.QueryResultBatch = MoveTemp(QueryResultBatch);
+				worker.Tokens = &SearchTokens;
+				worker.OnResultReady = OnResultReady;
+
+				Task.StartBackgroundTask();
+
+				// Ensure batch array is empty (probably already is due to move above)
+				QueryResultBatch.Reset();
+				break;
+			}
+		}
+	}
+
+	// Wait for all work to finish
+	for (auto& Task: ExtractionTasks)
+	{
+		Task.EnsureCompletion();
+	}
+
+	if (StopTaskCounter.GetValue())
+	{
+		// Ensure that the FiB Manager knows that we are done searching
+		FFindInBlueprintSearchManager::Get().EnsureSearchQueryEnds(this);
+	}
+
+	bThreadCompleted = true;
+
+	return 0;
+}
+
+void FExtractionTask::DoWork()
+{
+	for (const auto& QueryResult: QueryResultBatch)
+	{
+		FSearchResult BlueprintCategory = FSearchResult(new FFindInBlueprintsResult(FText::FromString(QueryResult.BlueprintPath)));
+		FindInBlueprintsHelpers::Extract(QueryResult.Value, *Tokens, BlueprintCategory);
+
+		// If there are children, add the item to the search results
+		if(BlueprintCategory->Children.Num() != 0)
+		{
+			OnResultReady(BlueprintCategory);
+		}
+	}
+}
+
 #undef LOCTEXT_NAMESPACE
