[INTEGRATE] Change 2427157 by [REDACTED_PERSON][REDACTED_HANDLE] on 2015/01/31 15:14:04

	Refactored component and actor renaming to work well and as the user expects from bor the viewport and the component tree. UE-8213

diff --git a/Engine/Source/Editor/Kismet/Private/SSCSEditor.cpp b/Engine/Source/Editor/Kismet/Private/SSCSEditor.cpp
--- a/Engine/Source/Editor/Kismet/Private/SSCSEditor.cpp
+++ b/Engine/Source/Editor/Kismet/Private/SSCSEditor.cpp
@@ -1175,6 +1175,27 @@ UActorComponent* FSCSEditorTreeNodeComponent::INTERNAL_GetOverridenComponentTemp
 	return OverridenComponent;
 }
 
+//////////////////////////////////////////////////////////////////////////
+// FSCSEditorTreeNodeRootActor
+
+FName FSCSEditorTreeNodeRootActor::GetVariableName() const
+{
+	if (Actor)
+	{
+		return Actor->GetFName();
+	}
+	return NAME_None;
+}
+
+void FSCSEditorTreeNodeRootActor::OnCompleteRename(const FText& InNewName)
+{
+	if (Actor && Actor->IsActorLabelEditable() && !InNewName.ToString().Equals(Actor->GetActorLabel(), ESearchCase::CaseSensitive))
+	{
+		const FScopedTransaction Transaction(LOCTEXT("SCSEditorRenameActorTransaction", "Rename Actor"));
+		Actor->SetActorLabel(InNewName.ToString());
+	}
+}
+
 //////////////////////////////////////////////////////////////////////////
 // FSCSEditorComponentObjectTextFactory
 
@@ -2789,9 +2810,22 @@ void SSCS_RowWidget::OnNameTextCommit(const FText& InNewName, ETextCommit::Type
 
 TSharedRef<SWidget> SSCS_RowWidget_ActorRoot::GenerateWidgetForColumn(const FName& ColumnName)
 {
+	FSCSEditorTreeNodePtrType NodePtr = GetNode();
+
 	// We've removed the other columns for now,  implement them for the root actor if necessary
 	ensure(ColumnName == SCS_ColumnName_ComponentClass);
 
+	// Create the name field
+	TSharedPtr<SInlineEditableTextBlock> InlineWidget =
+		SNew(SInlineEditableTextBlock)
+		.Text(this, &SSCS_RowWidget_ActorRoot::GetActorDisplayText)
+		.OnVerifyTextChanged(this, &SSCS_RowWidget_ActorRoot::OnVerifyActorLabelChanged)
+		.OnTextCommitted(this, &SSCS_RowWidget_ActorRoot::OnNameTextCommit)
+		.IsSelected(this, &SSCS_RowWidget_ActorRoot::IsSelectedExclusively)
+		.IsReadOnly(!NodePtr->CanRename() || (SCSEditor.IsValid() && !SCSEditor.Pin()->IsEditingAllowed()));
+
+	NodePtr->SetRenameRequestedDelegate(FSCSEditorTreeNode::FOnRenameRequested::CreateSP(InlineWidget.Get(), &SInlineEditableTextBlock::EnterEditingMode));
+
 	return SNew(SHorizontalBox)
 		.ToolTip(CreateToolTipWidget())
 
@@ -2805,12 +2839,21 @@ TSharedRef<SWidget> SSCS_RowWidget_ActorRoot::GenerateWidgetForColumn(const FNam
 		]
 
 	+ SHorizontalBox::Slot()
+		.AutoWidth()
+		.HAlign(HAlign_Left)
+		.VAlign(VAlign_Center)
+		.Padding(0.0f, 0.0f)
+		[
+			InlineWidget.ToSharedRef()
+		]
+
+	+SHorizontalBox::Slot()
 		.HAlign(HAlign_Left)
 		.VAlign(VAlign_Center)
 		.Padding(0.0f, 0.0f)
 		[
 			SNew(STextBlock)
-			.Text(this, &SSCS_RowWidget_ActorRoot::GetActorDisplayText)
+			.Text(this, &SSCS_RowWidget_ActorRoot::GetActorContextText)
 			.ColorAndOpacity(FSlateColor::UseForeground())
 		];
 }
@@ -2873,6 +2916,27 @@ TSharedRef<SToolTip> SSCS_RowWidget_ActorRoot::CreateToolTipWidget() const
 	return IDocumentation::Get()->CreateToolTip(TAttribute<FText>(this, &SSCS_RowWidget_ActorRoot::GetActorDisplayText), TooltipContent, InfoBox, TEXT(""), TEXT(""));
 }
 
+bool SSCS_RowWidget_ActorRoot::OnVerifyActorLabelChanged(const FText& InLabel, FText& OutErrorMessage)
+{
+	FText TrimmedLabel = FText::TrimPrecedingAndTrailing(InLabel);
+
+	if (TrimmedLabel.IsEmpty())
+	{
+		OutErrorMessage = LOCTEXT("RenameFailed_LeftBlank", "Names cannot be left blank");
+		return false;
+	}
+
+	if (TrimmedLabel.ToString().Len() >= NAME_SIZE)
+	{
+		FFormatNamedArguments Arguments;
+		Arguments.Add(TEXT("CharCount"), NAME_SIZE);
+		OutErrorMessage = FText::Format(LOCTEXT("RenameFailed_TooLong", "Names must be less than {CharCount} characters long."), Arguments);
+		return false;
+	}
+
+	return true;
+}
+
 const FSlateBrush* SSCS_RowWidget_ActorRoot::GetActorIcon() const
 {
 	if (SCSEditor.IsValid())
@@ -2896,18 +2960,37 @@ FText SSCS_RowWidget_ActorRoot::GetActorDisplayText() const
 			AActor* DefaultActor = SCSEditorPtr->ActorContext.Get();
 			if( DefaultActor )
 			{
+				FString Name;
 				UBlueprint* Blueprint = UBlueprint::GetBlueprintFromClass(DefaultActor->GetClass());
 				if(Blueprint != nullptr)
 				{
-					FString Name;
 					Blueprint->GetName(Name);
-					return FText::Format(LOCTEXT("DefaultActor_Name", "{0} (self)"), FText::FromString(Name));
 				}
 				else
 				{
-					FString Name = DefaultActor->GetActorLabel();
-					return FText::Format(LOCTEXT("DefaultActor_Name", "{0} (Instance)"), FText::FromString(Name));
+					Name = DefaultActor->GetActorLabel();
 				}
+				return FText::FromString(Name);
+			}
+		}
+	}
+	return FText::GetEmpty();
+}
+
+FText SSCS_RowWidget_ActorRoot::GetActorContextText() const
+{
+	if (SCSEditor.IsValid())
+	{
+		TSharedPtr<SSCSEditor> SCSEditorPtr = SCSEditor.Pin();
+		if (AActor* DefaultActor = SCSEditorPtr->GetActorContext())
+		{
+			if (UBlueprint* Blueprint = UBlueprint::GetBlueprintFromClass(DefaultActor->GetClass()))
+			{
+				return LOCTEXT("ActorContext_self", " (self)");
+			}
+			else
+			{
+				return LOCTEXT("ActorContext_Instance", " (Instance)");
 			}
 		}
 	}
@@ -3231,9 +3314,26 @@ void SSCSEditor::Construct( const FArguments& InArgs )
 
 	// Refresh the tree widget
 	UpdateTree();
+
+	if (EditorMode == EComponentEditorMode::ActorInstance)
+	{
+		GEngine->OnLevelComponentRequestRename().AddSP(this, &SSCSEditor::OnLevelComponentRequestRename);
+	}
 }
 END_SLATE_FUNCTION_BUILD_OPTIMIZATION
 
+
+void SSCSEditor::OnLevelComponentRequestRename(const UActorComponent* InComponent)
+{
+	TArray< FSCSEditorTreeNodePtrType > SelectedItems = SCSTreeWidget->GetSelectedItems();
+	
+	FSCSEditorTreeNodePtrType Node = GetNodeFromActorComponent(InComponent);
+	if (SelectedItems.Contains(Node) && CanRenameComponent())
+	{
+		OnRenameComponent(true);
+	}
+}
+
 UBlueprint* SSCSEditor::GetBlueprint() const
 {
 	if (AActor* Actor = GetActorContext())
@@ -3986,7 +4086,7 @@ void SSCSEditor::UpdateTree(bool bRegenerateTreeNodes)
 		// Reset the scene root node
 		SceneRootNodePtr.Reset();
 
-		TSharedPtr<FSCSEditorTreeNode> ActorTreeNode  = MakeShareable(new FSCSEditorTreeNodeRootActor());
+		TSharedPtr<FSCSEditorTreeNode> ActorTreeNode = MakeShareable(new FSCSEditorTreeNodeRootActor(GetActorContext(),EditorMode == EComponentEditorMode::ActorInstance));
 
 		RootNodes.Add(ActorTreeNode);
 		RootNodes.Add(MakeShareable(new FSCSEditorTreeNodeSeparator()));
diff --git a/Engine/Source/Editor/Kismet/Public/SSCSEditor.h b/Engine/Source/Editor/Kismet/Public/SSCSEditor.h
--- a/Engine/Source/Editor/Kismet/Public/SSCSEditor.h
+++ b/Engine/Source/Editor/Kismet/Public/SSCSEditor.h
@@ -432,10 +432,22 @@ private:
 class KISMET_API FSCSEditorTreeNodeRootActor : public FSCSEditorTreeNode
 {
 public:
-	FSCSEditorTreeNodeRootActor()
+	FSCSEditorTreeNodeRootActor(AActor* InActor, bool bInAllowRename)
 		: FSCSEditorTreeNode(FSCSEditorTreeNode::RootActorNode)
+		, Actor(InActor)
+		, bAllowRename(bInAllowRename)
 	{
 	}
+
+	// FSCSEditorTreeNode public interface
+	virtual FName GetVariableName() const override;
+	virtual bool CanRename() const override { return bAllowRename; }
+	virtual void OnCompleteRename(const FText& InNewName) override;
+	// End of FSCSEditorTreeNode public interface
+
+private:
+	AActor* Actor;
+	bool bAllowRename;
 };
 
 class KISMET_API FSCSEditorTreeNodeSeparator : public FSCSEditorTreeNode
@@ -493,13 +505,13 @@ protected:
 
 	static void AddToToolTipInfoBox(const TSharedRef<SVerticalBox>& InfoBox, const FText& Key, TSharedRef<SWidget> ValueIcon, const TAttribute<FText>& Value, bool bImportant);
 
+	/** Commits the new name of the component */
+	void OnNameTextCommit(const FText& InNewName, ETextCommit::Type InTextCommit);
+
 private:
 	/** Verifies the name of the component when changing it */
 	bool OnNameTextVerifyChanged(const FText& InNewText, FText& OutErrorMessage);
 
-	/** Commits the new name of the component */
-	void OnNameTextCommit(const FText& InNewName, ETextCommit::Type InTextCommit);
-
 	/** Builds a context menu popup for dropping a child node onto the scene root node */
 	TSharedPtr<SWidget> BuildSceneRootDropActionMenu(FSCSEditorTreeNodePtrType DroppedNodePtr);
 
@@ -575,9 +587,13 @@ private:
 	/** Creates a tooltip for this row */
 	TSharedRef<SToolTip> CreateToolTipWidget() const;
 
+	/** Called to validate the actor name */
+	bool OnVerifyActorLabelChanged(const FText& InLabel, FText& OutErrorMessage);
+
 	/** Data accessors */
 	const FSlateBrush* GetActorIcon() const;
 	FText GetActorDisplayText() const;
+	FText GetActorContextText() const;
 	FText GetActorClassNameText() const;
 	FText GetActorSuperClassNameText() const;
 	FText GetActorMobilityText() const;
@@ -881,6 +897,9 @@ protected:
 	/** Called to display context menu when right clicking on the widget */
 	TSharedPtr< SWidget > CreateContextMenu();
 
+	/** Called when the level editor requests a component to be renamed. */
+	void OnLevelComponentRequestRename(const UActorComponent* InComponent);
+
 	/**
 	 * Requests a rename on the selected component
 	 * @param bTransactional Whether or not the rename should be transactional (i.e. undoable)
diff --git a/Engine/Source/Editor/LevelEditor/Private/LevelEditorActions.cpp b/Engine/Source/Editor/LevelEditor/Private/LevelEditorActions.cpp
--- a/Engine/Source/Editor/LevelEditor/Private/LevelEditorActions.cpp
+++ b/Engine/Source/Editor/LevelEditor/Private/LevelEditorActions.cpp
@@ -1395,10 +1395,18 @@ bool FLevelEditorActionCallbacks::Delete_CanExecute()
 
 void FLevelEditorActionCallbacks::Rename_Execute()
 {
-	AActor* Actor = Cast<AActor>( *GEditor->GetSelectedActorIterator() );
-	if(Actor)
+	UActorComponent* Component = Cast<UActorComponent>(*GEditor->GetSelectedComponentIterator());
+	if (Component)
 	{
-		GEditor->BroadcastLevelActorRequestRename(Actor);
+		GEditor->BroadcastLevelComponentRequestRename(Component);
+	}
+	else
+	{
+		AActor* Actor = Cast<AActor>(*GEditor->GetSelectedActorIterator());
+		if (Actor)
+		{
+			GEditor->BroadcastLevelActorRequestRename(Actor);
+		}
 	}
 }
 
diff --git a/Engine/Source/Runtime/Engine/Classes/Engine/Engine.h b/Engine/Source/Runtime/Engine/Classes/Engine/Engine.h
--- a/Engine/Source/Runtime/Engine/Classes/Engine/Engine.h
+++ b/Engine/Source/Runtime/Engine/Classes/Engine/Engine.h
@@ -1590,6 +1590,13 @@ public:
 	/** Called by internal engine systems after a level actor has been requested to be renamed */
 	void BroadcastLevelActorRequestRename(const AActor* InActor) { LevelActorRequestRenameEvent.Broadcast(InActor); }
 
+	/** Editor-only event triggered when actors are being requested to be renamed */
+	DECLARE_EVENT_OneParam(UEngine, FLevelComponentRequestRenameEvent, const UActorComponent*);
+	FLevelComponentRequestRenameEvent& OnLevelComponentRequestRename() { return LevelComponentRequestRenameEvent; }
+
+	/** Called by internal engine systems after a level actor has been requested to be renamed */
+	void BroadcastLevelComponentRequestRename(const UActorComponent* InComponent) { LevelComponentRequestRenameEvent.Broadcast(InComponent); }
+
 #endif // #if WITH_EDITOR
 
 	/** Event triggered after a server travel failure of any kind has occurred */
@@ -2213,6 +2220,9 @@ private:
 	/** Broadcasts whenever an actor is being renamed */
 	FLevelActorRequestRenameEvent LevelActorRequestRenameEvent;
 
+	/** Broadcasts whenever a component is being renamed */
+	FLevelComponentRequestRenameEvent LevelComponentRequestRenameEvent;
+
 	/** Broadcasts after an actor has been moved, rotated or scaled */
 	FOnActorMovedEvent		OnActorMovedEvent;
 
