[INTEGRATE] Change 2427169 by [REDACTED_PERSON][REDACTED_HANDLE] on 2015/01/31 15:38:18

	Fixed multiple undo issues in the SCSEditor
	- After dropping actor component types not being able to redo
	- Fixed rentrant renames causing unnesssary transactions
	- Fixed agressive rebuilding of the scs tree, preventing renames from occuring

diff --git a/Engine/Source/Editor/Kismet/Private/SSCSEditor.cpp b/Engine/Source/Editor/Kismet/Private/SSCSEditor.cpp
--- a/Engine/Source/Editor/Kismet/Private/SSCSEditor.cpp
+++ b/Engine/Source/Editor/Kismet/Private/SSCSEditor.cpp
@@ -2801,7 +2801,7 @@ void SSCS_RowWidget::OnNameTextCommit(const FText& InNewName, ETextCommit::Type
 	TSharedPtr<SSCSEditor> PinnedEditor = SCSEditor.Pin();
 	if (PinnedEditor.IsValid() && PinnedEditor->GetEditorMode() == EComponentEditorMode::ActorInstance)
 	{
-		PinnedEditor->UpdateTree();
+		PinnedEditor->UpdateTree(false);
 	}
 }
 
@@ -3416,12 +3416,6 @@ TSharedRef<ITableRow> SSCSEditor::MakeTableRowWidget( FSCSEditorTreeNodePtrType
 				ItemName = ComponentTemplateOrInstance->GetFName();
 			}
 		}
-
-		if(DeferredRenameRequest == ItemName)
-		{
-			SCSTreeWidget->SetSelection(InNodePtr);
-			OnRenameComponent(false);
-		}
 	}
 
 	// Setup a meta tag for this node
@@ -4598,7 +4592,8 @@ UActorComponent* SSCSEditor::AddNewNodeForInstancedComponent(UActorComponent* Ne
 		OnRenameComponent(false);
 	}
 
-	UpdateTree();
+	UpdateTree(false);
+
 
 	return NewInstanceComponent;
 }
@@ -5033,6 +5028,7 @@ void SSCSEditor::RemoveComponentNode(FSCSEditorTreeNodePtrType InNodePtr)
 			}
 
 			// Destroy the component instance
+			ComponentInstance->Modify();
 			ComponentInstance->DestroyComponent(true);
 		}
 	}
diff --git a/Engine/Source/Editor/UnrealEd/Private/EditorActor.cpp b/Engine/Source/Editor/UnrealEd/Private/EditorActor.cpp
--- a/Engine/Source/Editor/UnrealEd/Private/EditorActor.cpp
+++ b/Engine/Source/Editor/UnrealEd/Private/EditorActor.cpp
@@ -556,10 +556,11 @@ bool UUnrealEdEngine::edactDeleteSelected( UWorld* InWorld, bool bVerifyDeletion
 		// Delete selected components if they are an Instance component
 		for (FSelectionIterator It(GetSelectedComponentIterator()); It; ++It)
 		{
-			UActorComponent* CurActor = CastChecked<UActorComponent>(*It);
-			if (CurActor->CreationMethod == EComponentCreationMethod::Instance)
+			UActorComponent* CurActorComp = CastChecked<UActorComponent>(*It);
+			if (CurActorComp->CreationMethod == EComponentCreationMethod::Instance)
 			{
-				CurActor->DestroyComponent(true);
+				CurActorComp->Modify();
+				CurActorComp->DestroyComponent(true);
 				bComponentWasDeleted = true;
 			}
 		}
