[INTEGRATE] Change 2427730 by [REDACTED_PERSON][REDACTED_HANDLE] on 2015/02/01 22:53:58

	Fixed check failure that ocurred when closing and reopneing the blueprint editor, stale pointer hanging around in a delegate..

diff --git a/Engine/Source/Editor/Kismet/Private/SMyBlueprint.cpp b/Engine/Source/Editor/Kismet/Private/SMyBlueprint.cpp
--- a/Engine/Source/Editor/Kismet/Private/SMyBlueprint.cpp
+++ b/Engine/Source/Editor/Kismet/Private/SMyBlueprint.cpp
@@ -395,6 +395,11 @@ void SMyBlueprint::Construct(const FArguments& InArgs, TWeakPtr<FBlueprintEditor
 	FCoreUObjectDelegates::OnObjectPropertyChanged.AddRaw(this, &SMyBlueprint::OnObjectPropertyChanged);
 }
 
+SMyBlueprint::~SMyBlueprint()
+{
+	FCoreUObjectDelegates::OnObjectPropertyChanged.RemoveAll(this);
+}
+
 void SMyBlueprint::Tick(const FGeometry& AllottedGeometry, const double InCurrentTime, const float InDeltaTime)
 {
 	SCompoundWidget::Tick(AllottedGeometry, InCurrentTime, InDeltaTime);
diff --git a/Engine/Source/Editor/Kismet/Private/SMyBlueprint.h b/Engine/Source/Editor/Kismet/Private/SMyBlueprint.h
--- a/Engine/Source/Editor/Kismet/Private/SMyBlueprint.h
+++ b/Engine/Source/Editor/Kismet/Private/SMyBlueprint.h
@@ -55,6 +55,8 @@ public:
 	SLATE_END_ARGS()
 
 	void Construct(const FArguments& InArgs, TWeakPtr<FBlueprintEditor> InBlueprintEditor, const UBlueprint* InBlueprint = nullptr);
+	~SMyBlueprint();
+
 	void SetInspector( TSharedPtr<SKismetInspector> InInspector ) { Inspector = InInspector ; }
 
 	/* SWidget interface */
