[INTEGRATE] Change 2428808 by [REDACTED_PERSON][REDACTED_HANDLE] on 2015/02/02 17:28:10

	BP - Invalidating the preview of the viewport client now forces a recreation in the blue editor.  In serveral places additional changes are applied to a blueprint that

diff --git a/Engine/Source/Editor/Kismet/Private/BlueprintEditor.cpp b/Engine/Source/Editor/Kismet/Private/BlueprintEditor.cpp
--- a/Engine/Source/Editor/Kismet/Private/BlueprintEditor.cpp
+++ b/Engine/Source/Editor/Kismet/Private/BlueprintEditor.cpp
@@ -700,8 +700,22 @@ void FBlueprintEditor::UpdateSCSPreview(bool bUpdateNow)
 	// refresh widget
 	if(SCSViewport.IsValid())
 	{
-		// Ignore 'bUpdateNow' if "Components" mode is not current. Otherwise the preview actor might be spawned in as a result, which can lead to a few odd behaviors if the mode is not current.
-		SCSViewport->RequestRefresh(false, bUpdateNow && IsModeCurrent(FBlueprintEditorApplicationModes::BlueprintComponentsMode));
+		if ( GetDefault<UEditorExperimentalSettings>()->bUnifiedBlueprintEditor )
+		{
+			TSharedPtr<SDockTab> OwnerTab = Inspector->GetOwnerTab();
+			if ( OwnerTab.IsValid() )
+			{
+				bUpdateNow &= OwnerTab->IsForeground();
+			}
+
+			// Only request a refresh immediately if the viewport tab is in the foreground.
+			SCSViewport->RequestRefresh(false, bUpdateNow);
+		}
+		else
+		{
+			// Ignore 'bUpdateNow' if "Components" mode is not current. Otherwise the preview actor might be spawned in as a result, which can lead to a few odd behaviors if the mode is not current.
+			SCSViewport->RequestRefresh(false, bUpdateNow && IsModeCurrent(FBlueprintEditorApplicationModes::BlueprintComponentsMode));
+		}
 	}
 }
 
@@ -2834,6 +2848,8 @@ void FBlueprintEditor::OnBlueprintChangedImpl(UBlueprint* InBlueprint, bool bIsJ
 {
 	if (InBlueprint)
 	{
+		DestroyPreview();
+
 		// Refresh the graphs
 		ERefreshBlueprintEditorReason::Type Reason = bIsJustBeingCompiled ? ERefreshBlueprintEditorReason::BlueprintCompiled : ERefreshBlueprintEditorReason::UnknownReason;
 		RefreshEditors(Reason);
@@ -7070,11 +7086,10 @@ void FBlueprintEditor::OnFindInstancesCustomEvent()
 
 AActor* FBlueprintEditor::GetPreviewActor() const
 {
-	//return SCSViewport.IsValid() ? SCSViewport->GetPreviewActor() : nullptr;
-
 	UBlueprint* PreviewBlueprint = GetBlueprintObj();
 
-	// Note: The weak ptr can become stale if the actor is reinstanced due to a Blueprint change, etc. In that case we look to see if we can find the new instance in the preview world and then update the weak ptr.
+	// Note: The weak ptr can become stale if the actor is reinstanced due to a Blueprint change, etc. In that 
+	// case we look to see if we can find the new instance in the preview world and then update the weak ptr.
 	if ( PreviewActorPtr.IsStale(true) && PreviewBlueprint )
 	{
 		UWorld* PreviewWorld = PreviewScene.GetWorld();
@@ -7213,6 +7228,8 @@ void FBlueprintEditor::DestroyPreview()
 
 		PreviewBlueprint = nullptr;
 	}
+
+	PreviewActorPtr = nullptr;
 }
 
 
diff --git a/Engine/Source/Editor/Kismet/Private/BlueprintEditorTabFactories.cpp b/Engine/Source/Editor/Kismet/Private/BlueprintEditorTabFactories.cpp
--- a/Engine/Source/Editor/Kismet/Private/BlueprintEditorTabFactories.cpp
+++ b/Engine/Source/Editor/Kismet/Private/BlueprintEditorTabFactories.cpp
@@ -276,6 +276,16 @@ TSharedRef<SWidget> FSCSViewportSummoner::CreateTabBody(const FWorkflowTabSpawnI
 	}
 }
 
+TSharedRef<SDockTab> FSCSViewportSummoner::SpawnTab(const FWorkflowTabSpawnInfo& Info) const
+{
+	TSharedRef<SDockTab> Tab = FWorkflowTabFactory::SpawnTab(Info);
+
+	TSharedPtr<FBlueprintEditor> BlueprintEditorPtr = StaticCastSharedPtr<FBlueprintEditor>(HostingApp.Pin());
+	BlueprintEditorPtr->GetSCSViewport()->SetOwnerTab(Tab);
+
+	return Tab;
+}
+
 FPaletteSummoner::FPaletteSummoner(TSharedPtr<class FAssetEditorToolkit> InHostingApp)
 	: FWorkflowTabFactory(FBlueprintEditorTabs::PaletteID, InHostingApp)
 {
diff --git a/Engine/Source/Editor/Kismet/Private/BlueprintEditorTabFactories.h b/Engine/Source/Editor/Kismet/Private/BlueprintEditorTabFactories.h
--- a/Engine/Source/Editor/Kismet/Private/BlueprintEditorTabFactories.h
+++ b/Engine/Source/Editor/Kismet/Private/BlueprintEditorTabFactories.h
@@ -206,6 +206,7 @@ public:
 	FSCSViewportSummoner(TSharedPtr<class FAssetEditorToolkit> InHostingApp);
 
 	virtual TSharedRef<SWidget> CreateTabBody(const FWorkflowTabSpawnInfo& Info) const override;
+	virtual TSharedRef<SDockTab> SpawnTab(const FWorkflowTabSpawnInfo& Info) const override;
 };
 
 //////////////////////////////////////////////////////////////////////////
diff --git a/Engine/Source/Editor/Kismet/Private/SCSEditorViewportClient.cpp b/Engine/Source/Editor/Kismet/Private/SCSEditorViewportClient.cpp
--- a/Engine/Source/Editor/Kismet/Private/SCSEditorViewportClient.cpp
+++ b/Engine/Source/Editor/Kismet/Private/SCSEditorViewportClient.cpp
@@ -184,9 +184,9 @@ void FSCSEditorViewportClient::Tick(float DeltaSeconds)
 		InvalidatePreview(false);
 	}
 
-	if ( PreviewActor != LastPreviewActor.Get() || PreviewActor == nullptr )
+	if ( PreviewActor != LastPreviewActor.Get() || PreviewActor == nullptr || IsRealtime() )
 	{
-//		PreviewActor = PreviewActor;
+		LastPreviewActor = PreviewActor;
 
 		Invalidate();
 		RefreshPreviewBounds();
@@ -673,6 +673,9 @@ void FSCSEditorViewportClient::InvalidatePreview(bool bResetCamera)
 	{
 		return;
 	}
+
+	// Destroy the preview
+	BlueprintEditorPtr.Pin()->DestroyPreview();
 	
 	if( bResetCamera )
 	{
diff --git a/Engine/Source/Editor/Kismet/Private/SSCSEditor.cpp b/Engine/Source/Editor/Kismet/Private/SSCSEditor.cpp
--- a/Engine/Source/Editor/Kismet/Private/SSCSEditor.cpp
+++ b/Engine/Source/Editor/Kismet/Private/SSCSEditor.cpp
@@ -123,7 +123,7 @@ FReply SSCSEditor::TryHandleAssetDragDropOperation(const FDragDropEvent& DragDro
 	TSharedPtr<FDragDropOperation> Operation = DragDropEvent.GetOperation();
 	if (Operation.IsValid() && (Operation->IsOfType<FExternalDragOperation>() || Operation->IsOfType<FAssetDragDropOp>()))
 	{
-		TArray< FAssetData > DroppedAssetData = AssetUtil::ExtractAssetDataFromDrag(DragDropEvent);
+		TArray< FAssetData > DroppedAssetData = AssetUtil::ExtractAssetDataFromDrag(Operation);
 		const int32 NumAssets = DroppedAssetData.Num();
 
 		if (NumAssets > 0)
@@ -154,7 +154,7 @@ FReply SSCSEditor::TryHandleAssetDragDropOperation(const FDragDropEvent& DragDro
 					PotentialComponentClass = Cast<UClass>(Asset);
 				}
 
-				TSubclassOf<UActorComponent>  MatchingComponentClassForAsset = FComponentAssetBrokerage::GetPrimaryComponentForAsset(AssetClass);
+				TSubclassOf<UActorComponent> MatchingComponentClassForAsset = FComponentAssetBrokerage::GetPrimaryComponentForAsset(AssetClass);
 				if (MatchingComponentClassForAsset != nullptr)
 				{
 					AddNewComponent(MatchingComponentClassForAsset, Asset);
diff --git a/Engine/Source/Editor/Kismet/Private/SSCSEditorViewport.cpp b/Engine/Source/Editor/Kismet/Private/SSCSEditorViewport.cpp
--- a/Engine/Source/Editor/Kismet/Private/SSCSEditorViewport.cpp
+++ b/Engine/Source/Editor/Kismet/Private/SSCSEditorViewport.cpp
@@ -337,7 +337,7 @@ void SSCSEditorViewport::RequestRefresh(bool bResetCamera, bool bRefreshNow)
 void SSCSEditorViewport::OnComponentSelectionChanged()
 {
 	// When the component selection changes, make sure to invalidate hit proxies to sync with the current selection
-	SceneViewport->InvalidateHitProxy();
+	SceneViewport->Invalidate();
 }
 
 void SSCSEditorViewport::OnFocusViewportToSelection()
@@ -350,6 +350,16 @@ bool SSCSEditorViewport::GetIsSimulateEnabled()
 	return ViewportClient->GetIsSimulateEnabled();
 }
 
+void SSCSEditorViewport::SetOwnerTab(TSharedRef<SDockTab> Tab)
+{
+	OwnerTab = Tab;
+}
+
+TSharedPtr<SDockTab> SSCSEditorViewport::GetOwnerTab() const
+{
+	return OwnerTab.Pin();
+}
+
 FReply SSCSEditorViewport::OnDrop(const FGeometry& MyGeometry, const FDragDropEvent& DragDropEvent)
 {
 	TSharedPtr<SSCSEditor> SCSEditor = BlueprintEditorPtr.Pin()->GetSCSEditor();
diff --git a/Engine/Source/Editor/Kismet/Private/SSCSEditorViewport.h b/Engine/Source/Editor/Kismet/Private/SSCSEditorViewport.h
--- a/Engine/Source/Editor/Kismet/Private/SSCSEditorViewport.h
+++ b/Engine/Source/Editor/Kismet/Private/SSCSEditorViewport.h
@@ -65,6 +65,10 @@ public:
 	 */
 	bool GetIsSimulateEnabled();
 
+	void SetOwnerTab(TSharedRef<SDockTab> Tab);
+
+	TSharedPtr<SDockTab> GetOwnerTab() const;
+
 protected:
 	/**
 	 * Determines if the viewport widget is visible.
@@ -95,4 +99,7 @@ private:
 
 	/** Whether the active timer (for updating the preview) is registered */
 	bool bIsActiveTimerRegistered;
+
+	/** The owner dock tab for this viewport. */
+	TWeakPtr<SDockTab> OwnerTab;
 };
