[INTEGRATE] Change 2429206 by [REDACTED_PERSON][REDACTED_HANDLE] on 2015/02/02 22:06:28

	Editor - Removing some old code that causes crashes in the editor when attempting to click on viewport menu anchors under low framerates, or when the frame rate is articially low, like coming back from the window not being focused.  The viewport toolbar no longer attempts to collapse when the frame rate is low, or the window is not focused.

diff --git a/Engine/Source/Editor/Cascade/Private/SCascadePreviewViewport.cpp b/Engine/Source/Editor/Cascade/Private/SCascadePreviewViewport.cpp
--- a/Engine/Source/Editor/Cascade/Private/SCascadePreviewViewport.cpp
+++ b/Engine/Source/Editor/Cascade/Private/SCascadePreviewViewport.cpp
@@ -35,11 +35,6 @@ bool SCascadePreviewViewport::IsVisible() const
 	return ViewportWidget.IsValid() && (!ParentTab.IsValid() || ParentTab.Pin()->IsForeground());
 }
 
-EVisibility SCascadePreviewViewport::GetWidgetVisibility() const
-{
-	return IsVisible()? EVisibility::Visible: EVisibility::Collapsed;
-}
-
 TSharedPtr<FSceneViewport> SCascadePreviewViewport::GetViewport() const
 {
 	return SceneViewport;
@@ -72,7 +67,6 @@ TSharedPtr<SWidget> SCascadePreviewViewport::MakeViewportToolbar()
 	return
 	SNew(SCascadePreviewViewportToolBar)
 		.CascadePtr(CascadePtr)
-		.Visibility(this, &SCascadePreviewViewport::GetWidgetVisibility)
 		.IsEnabled(FSlateApplication::Get().GetNormalExecutionAttribute());
 }
 
diff --git a/Engine/Source/Editor/Cascade/Private/SCascadePreviewViewport.h b/Engine/Source/Editor/Cascade/Private/SCascadePreviewViewport.h
--- a/Engine/Source/Editor/Cascade/Private/SCascadePreviewViewport.h
+++ b/Engine/Source/Editor/Cascade/Private/SCascadePreviewViewport.h
@@ -32,7 +32,6 @@ public:
 
 	/** Returns true if the viewport is visible */
 	bool IsVisible() const;
-	EVisibility GetWidgetVisibility() const;
 
 	/** Accessors */
 	TSharedPtr<FSceneViewport> GetViewport() const;
diff --git a/Engine/Source/Editor/Kismet/Private/SSCSEditorViewport.cpp b/Engine/Source/Editor/Kismet/Private/SSCSEditorViewport.cpp
--- a/Engine/Source/Editor/Kismet/Private/SSCSEditorViewport.cpp
+++ b/Engine/Source/Editor/Kismet/Private/SSCSEditorViewport.cpp
@@ -209,19 +209,12 @@ SSCSEditorViewport::~SSCSEditorViewport()
 	}
 }
 
-
-
 bool SSCSEditorViewport::IsVisible() const
 {
 	// We consider the viewport to be visible if the reference is valid
 	return ViewportWidget.IsValid() && SEditorViewport::IsVisible();
 }
 
-EVisibility SSCSEditorViewport::GetWidgetVisibility() const
-{
-	return IsVisible()? EVisibility::Visible: EVisibility::Collapsed;
-}
-
 TSharedRef<FEditorViewportClient> SSCSEditorViewport::MakeEditorViewportClient()
 {
 	FPreviewScene* PreviewScene = BlueprintEditorPtr.Pin()->GetPreviewScene();
@@ -240,7 +233,6 @@ TSharedPtr<SWidget> SSCSEditorViewport::MakeViewportToolbar()
 	return 
 		SNew(SSCSEditorViewportToolBar)
 		.EditorViewport(SharedThis(this))
-		.Visibility(this, &SSCSEditorViewport::GetWidgetVisibility)
 		.IsEnabled(FSlateApplication::Get().GetNormalExecutionAttribute());
 }
 
diff --git a/Engine/Source/Editor/Kismet/Private/SSCSEditorViewport.h b/Engine/Source/Editor/Kismet/Private/SSCSEditorViewport.h
--- a/Engine/Source/Editor/Kismet/Private/SSCSEditorViewport.h
+++ b/Engine/Source/Editor/Kismet/Private/SSCSEditorViewport.h
@@ -76,7 +76,6 @@ protected:
 	 * @return true if the viewport is visible; false otherwise.
 	 */
 	bool IsVisible() const;
-	EVisibility GetWidgetVisibility() const;
 
 	/** Called when the simulation toggle command is fired */
 	void ToggleIsSimulateEnabled();
diff --git a/Engine/Source/Editor/PhAT/Private/SPhATPreviewViewport.cpp b/Engine/Source/Editor/PhAT/Private/SPhATPreviewViewport.cpp
--- a/Engine/Source/Editor/PhAT/Private/SPhATPreviewViewport.cpp
+++ b/Engine/Source/Editor/PhAT/Private/SPhATPreviewViewport.cpp
@@ -44,11 +44,6 @@ bool SPhATPreviewViewport::IsVisible() const
 	return ViewportWidget.IsValid() && (!ParentTab.IsValid() || ParentTab.Pin()->IsForeground());
 }
 
-EVisibility SPhATPreviewViewport::GetWidgetVisibility() const
-{
-	return IsVisible()? EVisibility::Visible: EVisibility::Collapsed;
-}
-
 TSharedPtr<FSceneViewport> SPhATPreviewViewport::GetViewport() const
 {
 	return SceneViewport;
diff --git a/Engine/Source/Editor/PhAT/Private/SPhATPreviewViewport.h b/Engine/Source/Editor/PhAT/Private/SPhATPreviewViewport.h
--- a/Engine/Source/Editor/PhAT/Private/SPhATPreviewViewport.h
+++ b/Engine/Source/Editor/PhAT/Private/SPhATPreviewViewport.h
@@ -30,7 +30,6 @@ public:
 
 	/** Returns true if the viewport is visible */
 	bool IsVisible() const;
-	EVisibility GetWidgetVisibility() const;
 
 	/** Accessors */
 	TSharedPtr<FSceneViewport> GetViewport() const;
