[INTEGRATE] Change 2429294 by [REDACTED_PERSON][REDACTED_HANDLE] on 2015/02/03 00:05:34

	Blueprints: Include Actor CDO in search scope for rename conflicts when renaming a component in a Blueprint [UE-8936]

diff --git a/Engine/Source/Editor/Kismet/Private/SSCSEditor.cpp b/Engine/Source/Editor/Kismet/Private/SSCSEditor.cpp
--- a/Engine/Source/Editor/Kismet/Private/SSCSEditor.cpp
+++ b/Engine/Source/Editor/Kismet/Private/SSCSEditor.cpp
@@ -2746,24 +2746,30 @@ ESelectionMode::Type SSCS_RowWidget::GetSelectionMode() const
 bool SSCS_RowWidget::OnNameTextVerifyChanged(const FText& InNewText, FText& OutErrorMessage)
 {
 	FSCSEditorTreeNodePtrType NodePtr = GetNode();
+	UBlueprint* Blueprint = GetBlueprint();
 
-	if(!InNewText.IsEmpty())
+	if (!InNewText.IsEmpty())
 	{
+		AActor* ExistingNameSearchScope = NodePtr->GetComponentTemplate()->GetOwner();
+		if ((ExistingNameSearchScope == nullptr) && (Blueprint != nullptr))
+		{
+			ExistingNameSearchScope = Cast<AActor>(Blueprint->GeneratedClass->GetDefaultObject());
+		}
+
 		if (!FComponentEditorUtils::IsValidVariableNameString(NodePtr->GetComponentTemplate(), InNewText.ToString()))
 		{
 			OutErrorMessage = LOCTEXT("RenameFailed_EngineReservedName", "This name is reserved for engine use.");
 			return false;
 		}
-		else if (!FComponentEditorUtils::IsComponentNameAvailable(InNewText.ToString(), NodePtr->GetComponentTemplate()->GetOwner(), NodePtr->GetComponentTemplate()))
+		else if (!FComponentEditorUtils::IsComponentNameAvailable(InNewText.ToString(), ExistingNameSearchScope, NodePtr->GetComponentTemplate()))
 		{
 			OutErrorMessage = LOCTEXT("RenameFailed_ExistingName", "Another component already has the same name.");
 			return false;
 		}
 	}
 
-	UBlueprint* Blueprint = GetBlueprint();
 	TSharedPtr<INameValidatorInterface> NameValidator;
-	if(Blueprint != nullptr)
+	if (Blueprint != nullptr)
 	{
 		NameValidator = MakeShareable(new FKismetNameValidator(GetBlueprint(), NodePtr->GetVariableName()));
 	}
@@ -2773,20 +2779,20 @@ bool SSCS_RowWidget::OnNameTextVerifyChanged(const FText& InNewText, FText& OutE
 	}
 
 	EValidatorResult ValidatorResult = NameValidator->IsValid(InNewText.ToString());
-	if(ValidatorResult == EValidatorResult::AlreadyInUse)
+	if (ValidatorResult == EValidatorResult::AlreadyInUse)
 	{
 		OutErrorMessage = FText::Format(LOCTEXT("RenameFailed_InUse", "{0} is in use by another variable or function!"), InNewText);
 	}
-	else if(ValidatorResult == EValidatorResult::EmptyName)
+	else if (ValidatorResult == EValidatorResult::EmptyName)
 	{
 		OutErrorMessage = LOCTEXT("RenameFailed_LeftBlank", "Names cannot be left blank!");
 	}
-	else if(ValidatorResult == EValidatorResult::TooLong)
+	else if (ValidatorResult == EValidatorResult::TooLong)
 	{
 		OutErrorMessage = LOCTEXT("RenameFailed_NameTooLong", "Names must have fewer than 100 characters!");
 	}
 
-	if(OutErrorMessage.IsEmpty())
+	if (OutErrorMessage.IsEmpty())
 	{
 		return true;
 	}
