[INTEGRATE] Change 2429779 by [REDACTED_PERSON][REDACTED_HANDLE] on 2015/02/03 08:23:03

	Renamed Class Blueprint to Blueprint Class in a few places where that regressed

diff --git a/Engine/Source/Editor/Kismet/Private/SSCSEditor.cpp b/Engine/Source/Editor/Kismet/Private/SSCSEditor.cpp
--- a/Engine/Source/Editor/Kismet/Private/SSCSEditor.cpp
+++ b/Engine/Source/Editor/Kismet/Private/SSCSEditor.cpp
@@ -3196,7 +3196,7 @@ void SSCSEditor::Construct( const FArguments& InArgs )
 
 	EditBlueprintMenuBuilder.AddMenuEntry
 	(
-		LOCTEXT("PromoteToBlueprint", "Convert to Class Blueprint"),
+		LOCTEXT("PromoteToBlueprint", "Convert to Blueprint Class"),
 		LOCTEXT("PromoteToBluerprintTooltip","Converts the existing Blueprint into a new SubClass Blueprint" ),
 		FSlateIcon(),
 		FUIAction(FExecuteAction::CreateSP(this, &SSCSEditor::PromoteToBlueprint))
@@ -3240,7 +3240,7 @@ void SSCSEditor::Construct( const FArguments& InArgs )
 					.ButtonStyle(FEditorStyle::Get(), "ToggleButton")
 					.ContentPadding(0)
 					.ToolTip(IDocumentation::Get()->CreateToolTip(
-						LOCTEXT("PromoteToBluerprintTooltip","Converts this actor into a reusable Class Blueprint that can have script behavior" ),
+						LOCTEXT("PromoteToBluerprintTooltip","Converts this actor into a reusable Blueprint Class that can have script behavior" ),
 						NULL,
 						TEXT("Shared/LevelEditor"),
 						TEXT("ConvertToBlueprint")))
@@ -3260,7 +3260,7 @@ void SSCSEditor::Construct( const FArguments& InArgs )
 						[
 							SNew(STextBlock)
 							.TextStyle(FEditorStyle::Get(), "ContentBrowser.TopBar.Font")
-							.Text( LOCTEXT("PromoteToBlueprint", "Convert to Class Blueprint") )
+							.Text( LOCTEXT("PromoteToBlueprint", "Convert to Blueprint Class") )
 						]
 					]
 				]
@@ -3566,7 +3566,7 @@ TSharedPtr< SWidget > SSCSEditor::CreateContextMenu()
 							{
 								MenuBuilder.AddMenuEntry(
 									FText::Format(LOCTEXT("GoToBlueprintForComponent", "Edit {0}"), FText::FromString(ComponentTemplate->GetClass()->ClassGeneratedBy->GetName())),
-									LOCTEXT("EditBlueprintForComponent_ToolTip", "Edits the class blueprint that defines this component."),
+									LOCTEXT("EditBlueprintForComponent_ToolTip", "Edits the Blueprint Class that defines this component."),
 									FSlateIcon(FEditorStyle::GetStyleSetName(), FClassIconFinder::FindIconNameForClass(ComponentTemplate->GetClass())),
 									FUIAction(
 										FExecuteAction::CreateSP(this, &SSCSEditor::OnEditBlueprint, ComponentTemplate->GetClass()->ClassGeneratedBy),
