[INTEGRATE] Change 2429936 by [REDACTED_PERSON][REDACTED_HANDLE] on 2015/02/03 10:30:18

	Fix issue were you could not rename a new root node added because the SCSEditor tree would agressively update the entire tree.

diff --git a/Engine/Source/Editor/Kismet/Private/SSCSEditor.cpp b/Engine/Source/Editor/Kismet/Private/SSCSEditor.cpp
--- a/Engine/Source/Editor/Kismet/Private/SSCSEditor.cpp
+++ b/Engine/Source/Editor/Kismet/Private/SSCSEditor.cpp
@@ -4583,9 +4583,8 @@ UActorComponent* SSCSEditor::AddNewNodeForInstancedComponent(UActorComponent* Ne
 		if(SceneRootNodePtr.IsValid() && SceneRootNodePtr->IsDefaultSceneRoot())
 		{
 			RemoveComponentNode(SceneRootNodePtr);
+			RootNodes.Remove( SceneRootNodePtr );
 			SceneRootNodePtr.Reset();
-
-			UpdateTree();
 		}
 	}
 	else
