[INTEGRATE] Change 2430273 by [REDACTED_PERSON][REDACTED_HANDLE] on 2015/02/03 14:31:46

	Fix crash duplicating components in the SCSEditor on actor instances

diff --git a/Engine/Source/Editor/Kismet/Private/SSCSEditor.cpp b/Engine/Source/Editor/Kismet/Private/SSCSEditor.cpp
--- a/Engine/Source/Editor/Kismet/Private/SSCSEditor.cpp
+++ b/Engine/Source/Editor/Kismet/Private/SSCSEditor.cpp
@@ -3778,8 +3778,11 @@ void SSCSEditor::OnDuplicateComponent()
 				USceneComponent* NewSceneComponent = Cast<USceneComponent>(CloneComponent);
 				if(NewSceneComponent != NULL)
 				{
-					// Ensure that any native attachment relationship inherited from the original copy is removed (to prevent a GLEO assertion)
-					NewSceneComponent->DetachFromParent(true);
+					if( EditorMode == EComponentEditorMode::BlueprintSCS )
+					{
+						// Ensure that any native attachment relationship inherited from the original copy is removed (to prevent a GLEO assertion)
+						NewSceneComponent->DetachFromParent(true);
+					}
 
 					// Attempt to locate the original node in the SCS tree
 					FSCSEditorTreeNodePtrType OriginalNodePtr = FindTreeNode(OriginalComponent);
