[INTEGRATE] Change 2430292 by [REDACTED_PERSON][REDACTED_HANDLE] on 2015/02/03 14:41:58

	Fixed root display in the components tree for native roots that aren't set to the first defined native component.

diff --git a/Engine/Source/Editor/Kismet/Private/SSCSEditor.cpp b/Engine/Source/Editor/Kismet/Private/SSCSEditor.cpp
--- a/Engine/Source/Editor/Kismet/Private/SSCSEditor.cpp
+++ b/Engine/Source/Editor/Kismet/Private/SSCSEditor.cpp
@@ -4120,8 +4120,16 @@ void SSCSEditor::UpdateTree(bool bRegenerateTreeNodes)
 				
 				TInlineComponentArray<UActorComponent*> Components;
 				CDO->GetComponents(Components);
+
+				// Add the native root component
+				USceneComponent* RootComponent = CDO->GetRootComponent();
+				if(RootComponent != NULL)
+				{
+					Components.Remove(RootComponent);
+					AddTreeNodeFromComponent(RootComponent);
+				}
 				
-				// Add the native base class SceneComponent hierarchy
+				// Add the rest of the native base class SceneComponent hierarchy
 				for(auto CompIter = Components.CreateIterator(); CompIter; ++CompIter)
 				{
 					USceneComponent* SceneComp = Cast<USceneComponent>(*CompIter);
@@ -4187,8 +4195,16 @@ void SSCSEditor::UpdateTree(bool bRegenerateTreeNodes)
 				// Get the full set of instanced components
 				TInlineComponentArray<UActorComponent*> Components;
 				ActorInstance->GetComponents(Components);
+
+				// Add the root component first (it may not be the first one)
+				USceneComponent* RootComponent = ActorInstance->GetRootComponent();
+				if(RootComponent != nullptr)
+				{
+					Components.Remove(RootComponent);
+					AddTreeNodeFromComponent(RootComponent);
+				}
 				
-				// Now add the instanced scene component hierarchy
+				// Now add the rest of the instanced scene component hierarchy
 				for(auto CompIter = Components.CreateIterator(); CompIter; ++CompIter)
 				{
 					USceneComponent* SceneComp = Cast<USceneComponent>(*CompIter);
