[INTEGRATE] Change 2430415 by [REDACTED_PERSON][REDACTED_HANDLE] on 2015/02/03 15:35:22

	Refactored the defered renaming to use a GetNodeID() function to remove the dependency on GetVariableName().

diff --git a/Engine/Source/Editor/Kismet/Private/SSCSEditor.cpp b/Engine/Source/Editor/Kismet/Private/SSCSEditor.cpp
--- a/Engine/Source/Editor/Kismet/Private/SSCSEditor.cpp
+++ b/Engine/Source/Editor/Kismet/Private/SSCSEditor.cpp
@@ -234,6 +234,20 @@ FSCSEditorTreeNode::FSCSEditorTreeNode(FSCSEditorTreeNode::ENodeType InNodeType)
 {
 }
 
+FName FSCSEditorTreeNode::GetNodeID() const
+{
+	FName ItemName = GetVariableName();
+	if (ItemName == NAME_None)
+	{
+		UActorComponent* ComponentTemplateOrInstance = GetComponentTemplate();
+		if (ComponentTemplateOrInstance != nullptr)
+		{
+			ItemName = ComponentTemplateOrInstance->GetFName();
+		}
+	}
+	return ItemName;
+}
+
 FName FSCSEditorTreeNode::GetVariableName() const
 {
 	return NAME_None;
@@ -1180,7 +1194,7 @@ UActorComponent* FSCSEditorTreeNodeComponent::INTERNAL_GetOverridenComponentTemp
 //////////////////////////////////////////////////////////////////////////
 // FSCSEditorTreeNodeRootActor
 
-FName FSCSEditorTreeNodeRootActor::GetVariableName() const
+FName FSCSEditorTreeNodeRootActor::GetNodeID() const
 {
 	if (Actor)
 	{
@@ -3413,19 +3427,6 @@ FReply SSCSEditor::OnKeyDown( const FGeometry& MyGeometry, const FKeyEvent& InKe
 
 TSharedRef<ITableRow> SSCSEditor::MakeTableRowWidget( FSCSEditorTreeNodePtrType InNodePtr, const TSharedRef<STableViewBase>& OwnerTable )
 {
-	if(DeferredRenameRequest != NAME_None)
-	{
-		FName ItemName = InNodePtr->GetVariableName();
-		if(ItemName == NAME_None)
-		{
-			UActorComponent* ComponentTemplateOrInstance = InNodePtr->GetComponentTemplate();
-			if (ComponentTemplateOrInstance != nullptr)
-			{
-				ItemName = ComponentTemplateOrInstance->GetFName();
-			}
-		}
-	}
-
 	// Setup a meta tag for this node
 	FGraphNodeMetaData TagMeta(TEXT("TableRow"));
 	if (InNodePtr.IsValid() && InNodePtr->GetComponentTemplate() != NULL )
@@ -3433,6 +3434,7 @@ TSharedRef<ITableRow> SSCSEditor::MakeTableRowWidget( FSCSEditorTreeNodePtrType
 		TagMeta.FriendlyName = FString::Printf(TEXT("TableRow,%s,0"), *InNodePtr->GetComponentTemplate()->GetReadableName());
 	}
 
+	// Create the node of the appropriate type
 	if (InNodePtr->GetNodeType() == FSCSEditorTreeNode::RootActorNode)
 	{
 		return SNew(SSCS_RowWidget_ActorRoot, SharedThis(this), InNodePtr, OwnerTable)
@@ -5341,13 +5343,7 @@ FSCSEditorTreeNodePtrType SSCSEditor::FindTreeNode(const FName& InVariableOrInst
 
 		if(InStartNodePtr.IsValid())
 		{
-			FName ItemName = InStartNodePtr->GetVariableName();
-			if(ItemName == NAME_None)
-			{
-				UActorComponent* ComponentTemplateOrInstance = InStartNodePtr->GetComponentTemplate();
-				check(ComponentTemplateOrInstance != nullptr);
-				ItemName = ComponentTemplateOrInstance->GetFName();
-			}
+			FName ItemName = InStartNodePtr->GetNodeID();
 
 			// Check to see if the given name matches the item name
 			if(InVariableOrInstanceName == ItemName)
@@ -5376,14 +5372,7 @@ void SSCSEditor::OnItemScrolledIntoView( FSCSEditorTreeNodePtrType InItem, const
 {
 	if(DeferredRenameRequest != NAME_None)
 	{
-		FName ItemName = InItem->GetVariableName();
-		if(ItemName == NAME_None)
-		{
-			UActorComponent* ComponentTemplateOrInstance = InItem->GetComponentTemplate();
-			check(ComponentTemplateOrInstance != nullptr);
-			ItemName = ComponentTemplateOrInstance->GetFName();
-		}
-
+		FName ItemName = InItem->GetNodeID();
 		if(DeferredRenameRequest == ItemName)
 		{
 			DeferredRenameRequest = NAME_None;
@@ -5405,16 +5394,10 @@ void SSCSEditor::OnRenameComponent(bool bTransactional)
 	// Should already be prevented from making it here.
 	check(SelectedItems.Num() == 1);
 
-	SCSTreeWidget->RequestScrollIntoView(SelectedItems[0]);
-	DeferredRenameRequest = SelectedItems[0]->GetVariableName();
-	if(DeferredRenameRequest == NAME_None)
-	{
-		UActorComponent* ComponentTemplateOrInstance = SelectedItems[0]->GetComponentTemplate();
-		check(ComponentTemplateOrInstance != nullptr);
-		DeferredRenameRequest = ComponentTemplateOrInstance->GetFName();
-	}
-
+	DeferredRenameRequest = SelectedItems[0]->GetNodeID();
 	bIsDeferredRenameRequestTransactional = bTransactional;
+
+	SCSTreeWidget->RequestScrollIntoView(SelectedItems[0]);
 }
 
 bool SSCSEditor::CanRenameComponent() const
diff --git a/Engine/Source/Editor/Kismet/Public/SSCSEditor.h b/Engine/Source/Editor/Kismet/Public/SSCSEditor.h
--- a/Engine/Source/Editor/Kismet/Public/SSCSEditor.h
+++ b/Engine/Source/Editor/Kismet/Public/SSCSEditor.h
@@ -35,6 +35,11 @@ public:
 	 */
 	FSCSEditorTreeNode(FSCSEditorTreeNode::ENodeType InNodeType);
 
+	/**
+	* @return The name to identify this node.
+	*/
+	virtual FName GetNodeID() const;
+
 	/**
 	 * @return The name of the variable represented by this node.
 	 */
@@ -440,7 +445,7 @@ public:
 	}
 
 	// FSCSEditorTreeNode public interface
-	virtual FName GetVariableName() const override;
+	virtual FName GetNodeID() const override;
 	virtual bool CanRename() const override { return bAllowRename; }
 	virtual void OnCompleteRename(const FText& InNewName) override;
 	// End of FSCSEditorTreeNode public interface
