[INTEGRATE] Change 2430908 by [REDACTED_PERSON][REDACTED_HANDLE] on 2015/02/03 20:14:21

	Metadata was only intended for the component rows. Minor cleanup.

diff --git a/Engine/Source/Editor/Kismet/Private/SSCSEditor.cpp b/Engine/Source/Editor/Kismet/Private/SSCSEditor.cpp
--- a/Engine/Source/Editor/Kismet/Private/SSCSEditor.cpp
+++ b/Engine/Source/Editor/Kismet/Private/SSCSEditor.cpp
@@ -3437,13 +3437,11 @@ TSharedRef<ITableRow> SSCSEditor::MakeTableRowWidget( FSCSEditorTreeNodePtrType
 	// Create the node of the appropriate type
 	if (InNodePtr->GetNodeType() == FSCSEditorTreeNode::RootActorNode)
 	{
-		return SNew(SSCS_RowWidget_ActorRoot, SharedThis(this), InNodePtr, OwnerTable)
-			.AddMetaData<FTutorialMetaData>(TagMeta);
+		return SNew(SSCS_RowWidget_ActorRoot, SharedThis(this), InNodePtr, OwnerTable);
 	}
 	else if (InNodePtr->GetNodeType() == FSCSEditorTreeNode::SeparatorNode)
 	{
-		return SNew(SSCS_RowWidget_Separator, SharedThis(this), InNodePtr, OwnerTable)
-			.AddMetaData<FTutorialMetaData>(TagMeta);
+		return SNew(SSCS_RowWidget_Separator, SharedThis(this), InNodePtr, OwnerTable);
 	}
 
 	return SNew(SSCS_RowWidget, SharedThis(this), InNodePtr, OwnerTable)
