[INTEGRATE] Change 2431101 by [REDACTED_PERSON][REDACTED_HANDLE] on 2015/02/03 22:19:52

	Component tree: Inherited native components now scan class variables to determine a variable name if a direct match doesn't exist (e.g., cases like AStaticMeshActor::StaticMeshComponent which points to an object named "StaticMeshComponent0")
	- This allow drag-drop operations from the tree to a Blueprint graph to drop a working component reference in more circumstances
	- Display of these name mismatch components in the tree now uses the found variable name instead of the object FName or NAME_None, fixing broken re-attachment drop tooltips and improving consistency as the name observed in the tree matches the name observed in My Blueprints and the graph editor
	[UE-8842][UE-9054]

diff --git a/Engine/Source/Editor/Kismet/Private/SSCSEditor.cpp b/Engine/Source/Editor/Kismet/Private/SSCSEditor.cpp
--- a/Engine/Source/Editor/Kismet/Private/SSCSEditor.cpp
+++ b/Engine/Source/Editor/Kismet/Private/SSCSEditor.cpp
@@ -189,7 +189,7 @@ FReply SSCSEditorDragDropTree::OnDrop( const FGeometry& MyGeometry, const FDragD
 
 
 //////////////////////////////////////////////////////////////////////////
-//
+// FSCSRowDragDropOp - The drag-drop operation triggered when dragging a row in the components tree
 
 class FSCSRowDragDropOp : public FKismetVariableDragDropAction
 {
@@ -656,33 +656,8 @@ FName FSCSEditorTreeNodeComponentBase::GetVariableName() const
 	}
 	else if (ComponentTemplate != NULL)
 	{
-		// If the owner class is a Blueprint class, see if there's a corresponding object property that contains the component template
-		check(ComponentTemplate->GetOwner());
-		UClass* OwnerClass = ComponentTemplate->GetOwner()->GetActorClass();
-		if (OwnerClass != NULL)
-		{
-			UBlueprint* Blueprint = UBlueprint::GetBlueprintFromClass(OwnerClass);
-			if (Blueprint != NULL)
-			{
-				UClass* StartingClass = IsInstanced() ? Blueprint->GeneratedClass : Blueprint->ParentClass;
-				if (StartingClass != nullptr)
-				{
-					if (UObjectProperty* ObjectProp = FindField<UObjectProperty>(StartingClass, ComponentTemplate->GetFName()))
-					{
-						UObject* CDO = StartingClass->GetDefaultObject();
-						UObject* Object = ObjectProp->GetObjectPropertyValue(ObjectProp->ContainerPtrToValuePtr<void>(CDO));
-
-						if (Object)
-						{
-							if (Object->GetClass() == ComponentTemplate->GetClass())
-							{
-								VariableName = ComponentTemplate->GetFName();
-							}
-						}
-					}
-				}
-			}
-		}
+		// Try to find the component anchor variable name (first looks for an exact match then scans for any matching variable that points to the archetype in the CDO)
+		VariableName = FComponentEditorUtils::FindVariableNameGivenComponentInstance(ComponentTemplate);
 	}
 
 	return VariableName;
@@ -1904,7 +1879,9 @@ FReply SSCS_RowWidget::OnDragDetected( const FGeometry& MyGeometry, const FPoint
 		{
 			UBlueprint* Blueprint = FirstNode->GetBlueprint();
 			const FName VariableName = FirstNode->GetVariableName();
-			TSharedRef<FSCSRowDragDropOp> Operation = FSCSRowDragDropOp::New(VariableName, Blueprint != nullptr ? Blueprint->SkeletonGeneratedClass : nullptr, FNodeCreationAnalytic());
+ 			UStruct* VariableScope = (Blueprint != nullptr) ? Blueprint->SkeletonGeneratedClass : nullptr;
+
+			TSharedRef<FSCSRowDragDropOp> Operation = FSCSRowDragDropOp::New(VariableName, VariableScope, FNodeCreationAnalytic());
 			Operation->SetCtrlDrag(true); // Always put a getter
 			Operation->PendingDropAction = FSCSRowDragDropOp::DropAction_None;
 			Operation->SourceNodes = SelectedNodePtrs;
diff --git a/Engine/Source/Editor/UnrealEd/Private/Kismet2/ComponentEditorUtils.cpp b/Engine/Source/Editor/UnrealEd/Private/Kismet2/ComponentEditorUtils.cpp
--- a/Engine/Source/Editor/UnrealEd/Private/Kismet2/ComponentEditorUtils.cpp
+++ b/Engine/Source/Editor/UnrealEd/Private/Kismet2/ComponentEditorUtils.cpp
@@ -327,4 +327,50 @@ void FComponentEditorUtils::PropagateTransformInner(class USceneComponent* Insta
 	{
 		PropagateTransformPropertyChange(InstancedSceneComponent, RelativeScale3DProperty, OldDefaultTransform.Scale, NewDefaultTransform.Scale, UpdatedComponents);
 	}
-}
\ No newline at end of file
+}
+
+FName FComponentEditorUtils::FindVariableNameGivenComponentInstance(UActorComponent* ComponentInstance)
+{
+	check(ComponentInstance != nullptr);
+
+	// First see if the name just works
+	if (AActor* OwnerActor = ComponentInstance->GetOwner())
+	{
+		UClass* OwnerActorClass = OwnerActor->GetClass();
+		if (UObjectProperty* TestProperty = FindField<UObjectProperty>(OwnerActorClass, ComponentInstance->GetFName()))
+		{
+			if (ComponentInstance->GetClass()->IsChildOf(TestProperty->PropertyClass))
+			{
+				return TestProperty->GetFName();
+			}
+		}
+	}
+
+	// Name mismatch, try finding a differently named variable pointing to the the component (the mismatch should only be possible for native components)
+	if (UActorComponent* Archetype = Cast<UActorComponent>(ComponentInstance->GetArchetype()))
+	{
+		if (AActor* OwnerActor = Archetype->GetOwner())
+		{
+			UClass* OwnerClass = OwnerActor->GetClass();
+			AActor* OwnerCDO = CastChecked<AActor>(OwnerClass->GetDefaultObject());
+			check(OwnerCDO->HasAnyFlags(RF_ClassDefaultObject));
+
+			for (TFieldIterator<UObjectProperty> PropIt(OwnerClass, EFieldIteratorFlags::IncludeSuper); PropIt; ++PropIt)
+			{
+				UObjectProperty* TestProperty = *PropIt;
+				if (Archetype->GetClass()->IsChildOf(TestProperty->PropertyClass))
+				{
+					void* TestPropertyInstanceAddress = TestProperty->ContainerPtrToValuePtr<void>(OwnerCDO);
+					UObject* ObjectPointedToByProperty = TestProperty->GetObjectPropertyValue(TestPropertyInstanceAddress);
+					if (ObjectPointedToByProperty == Archetype)
+					{
+						// This property points to the component archetype, so it's an anchor even if it was named wrong
+						return TestProperty->GetFName();
+					}
+				}
+			}
+		}
+	}
+
+	return NAME_None;
+}
diff --git a/Engine/Source/Editor/UnrealEd/Public/Kismet2/ComponentEditorUtils.h b/Engine/Source/Editor/UnrealEd/Public/Kismet2/ComponentEditorUtils.h
--- a/Engine/Source/Editor/UnrealEd/Public/Kismet2/ComponentEditorUtils.h
+++ b/Engine/Source/Editor/UnrealEd/Public/Kismet2/ComponentEditorUtils.h
@@ -119,6 +119,9 @@ public:
 			UpdatedComponents.Add(InSceneComponent);
 		}
 	}
+
+	// Try to find the correct variable name for a given native component template or instance (which can have a mismatch)
+	static FName FindVariableNameGivenComponentInstance(UActorComponent* ComponentInstance);
 private:
 	static void PropagateTransformInner(class USceneComponent* InstancedSceneComponent, const FTransformData& OldDefaultTransform, const FTransformData& NewDefaultTransform, TSet<class USceneComponent*>& UpdatedComponents);
 
