[INTEGRATE] Change 2431212 by [REDACTED_PERSON][REDACTED_HANDLE] on 2015/02/04 00:45:49

	Component tree: Allow drag-dropping C++ and Blueprint actor classes into the tree to create new child actor components [UE-8458][UE-8459]
	- Note: ChildActorComponents never spawn their actor in the instanced case [UE-9063] but they work fine for the Blueprint case

diff --git a/Engine/Source/Editor/Kismet/Private/SSCSEditor.cpp b/Engine/Source/Editor/Kismet/Private/SSCSEditor.cpp
--- a/Engine/Source/Editor/Kismet/Private/SSCSEditor.cpp
+++ b/Engine/Source/Editor/Kismet/Private/SSCSEditor.cpp
@@ -128,7 +128,7 @@ FReply SSCSEditor::TryHandleAssetDragDropOperation(const FDragDropEvent& DragDro
 
 		if (NumAssets > 0)
 		{
-			GWarn->BeginSlowTask(LOCTEXT("LoadingComponents", "Loading Component(s)"), true);
+			GWarn->BeginSlowTask(LOCTEXT("LoadingAssets", "Loading Asset(s)"), true);
 
 			for (int32 DroppedAssetIdx = 0; DroppedAssetIdx < NumAssets; ++DroppedAssetIdx)
 			{
@@ -144,14 +144,30 @@ FReply SSCSEditor::TryHandleAssetDragDropOperation(const FDragDropEvent& DragDro
 
 				UBlueprint* BPClass = Cast<UBlueprint>(Asset);
 				UClass* PotentialComponentClass = nullptr;
+				UClass* PotentialActorClass = nullptr;
 
-				if ((BPClass != nullptr) && (BPClass->GeneratedClass != nullptr) && (BPClass->GeneratedClass->IsChildOf(UActorComponent::StaticClass())))
+				if ((BPClass != nullptr) && (BPClass->GeneratedClass != nullptr))
 				{
-					PotentialComponentClass = BPClass->GeneratedClass;
+					if (BPClass->GeneratedClass->IsChildOf(UActorComponent::StaticClass()))
+					{
+						PotentialComponentClass = BPClass->GeneratedClass;
+					}
+					else if (BPClass->GeneratedClass->IsChildOf(AActor::StaticClass()))
+					{
+						PotentialActorClass = BPClass->GeneratedClass;
+					}
 				}
-				else if (AssetClass->IsChildOf(UClass::StaticClass()) && CastChecked<UClass>(Asset)->IsChildOf(UActorComponent::StaticClass()))
+				else if (AssetClass->IsChildOf(UClass::StaticClass()))
 				{
-					PotentialComponentClass = Cast<UClass>(Asset);
+					UClass* AssetAsClass = CastChecked<UClass>(Asset);
+					if (AssetAsClass->IsChildOf(UActorComponent::StaticClass()))
+					{
+						PotentialComponentClass = AssetAsClass;
+					}
+					else if (AssetAsClass->IsChildOf(AActor::StaticClass()))
+					{
+						PotentialActorClass = AssetAsClass;
+					}
 				}
 
 				TSubclassOf<UActorComponent> MatchingComponentClassForAsset = FComponentAssetBrokerage::GetPrimaryComponentForAsset(AssetClass);
@@ -159,9 +175,19 @@ FReply SSCSEditor::TryHandleAssetDragDropOperation(const FDragDropEvent& DragDro
 				{
 					AddNewComponent(MatchingComponentClassForAsset, Asset);
 				}
-				else if ((PotentialComponentClass != nullptr) && (!PotentialComponentClass->HasAnyClassFlags(CLASS_Deprecated | CLASS_Abstract | CLASS_NewerVersionExists)) && PotentialComponentClass->HasMetaData(FBlueprintMetadata::MD_BlueprintSpawnableComponent))
+				else if ((PotentialComponentClass != nullptr) && !PotentialComponentClass->HasAnyClassFlags(CLASS_Deprecated | CLASS_Abstract | CLASS_NewerVersionExists))
 				{
-					AddNewComponent(PotentialComponentClass, nullptr);
+					if (PotentialComponentClass->HasMetaData(FBlueprintMetadata::MD_BlueprintSpawnableComponent))
+					{
+						AddNewComponent(PotentialComponentClass, nullptr);
+					}
+				}
+				else if ((PotentialActorClass != nullptr) && !PotentialActorClass->HasAnyClassFlags(CLASS_Deprecated | CLASS_Abstract | CLASS_NewerVersionExists))
+				{
+					if (UChildActorComponent* ChildActorComponent = Cast<UChildActorComponent>(AddNewComponent(UChildActorComponent::StaticClass(), nullptr)))
+					{
+						ChildActorComponent->ChildActorClass = PotentialActorClass;
+					}
 				}
 			}
 
