[INTEGRATE] Change 2431887 by [REDACTED_PERSON][REDACTED_HANDLE] on 2015/02/04 12:02:10

	BP - Fixing several selection issues.  Selecting the CDO now highlights the button or the root actor node no matter which you you select.  Creating new variables now refreshes the editors in the right order so that you can rename the variable after creating it.

diff --git a/Engine/Source/Editor/GraphEditor/Private/SGraphActionMenu.cpp b/Engine/Source/Editor/GraphEditor/Private/SGraphActionMenu.cpp
--- a/Engine/Source/Editor/GraphEditor/Private/SGraphActionMenu.cpp
+++ b/Engine/Source/Editor/GraphEditor/Private/SGraphActionMenu.cpp
@@ -974,7 +974,7 @@ TSharedRef<ITableRow> SGraphActionMenu::MakeWidget( TSharedPtr<FGraphActionNode>
 
 		if(OnCreateWidgetForAction.IsBound())
 		{
-			CreateData.IsRowSelectedDelegate = FIsSelected::CreateSP( TableRow.Get(), &STableRow< TSharedPtr<FGraphActionNode> >::IsSelectedExclusively );
+			CreateData.IsRowSelectedDelegate = FIsSelected::CreateSP( TableRow.Get(), &STableRow< TSharedPtr<FGraphActionNode> >::IsSelected );
 			CreateData.bIsReadOnly = bIsReadOnly;
 			CreateData.bHandleMouseButtonDown = false;		//Default to NOT using the delegate. OnCreateWidgetForAction can set to true if we need it
 			RowContent = OnCreateWidgetForAction.Execute( &CreateData );
diff --git a/Engine/Source/Editor/Kismet/Private/BlueprintEditor.cpp b/Engine/Source/Editor/Kismet/Private/BlueprintEditor.cpp
--- a/Engine/Source/Editor/Kismet/Private/BlueprintEditor.cpp
+++ b/Engine/Source/Editor/Kismet/Private/BlueprintEditor.cpp
@@ -2559,15 +2559,24 @@ void FBlueprintEditor::EditGlobalOptions_Clicked()
 
 bool FBlueprintEditor::IsDetailsPanelEditingClassDefaults() const
 {
-	return CurrentUISelection == FBlueprintEditor::SelectionState_ClassDefaults;
+	UBlueprint* Blueprint = GetBlueprintObj();
+	if ( Blueprint != nullptr )
+	{
+		if ( Blueprint->GeneratedClass != nullptr )
+		{
+			UObject* DefaultObject = GetBlueprintObj()->GeneratedClass->GetDefaultObject();
+			return Inspector->IsSelected(DefaultObject);
+		}
+	}
+
+	return false;
 }
 
 void FBlueprintEditor::EditClassDefaults_Clicked()
 {
-	UBlueprint* Blueprint = GetBlueprintObj();
-	if ( Blueprint != nullptr )
+	if ( IsEditingSingleBlueprint() )
 	{
-		SetUISelectionState(FBlueprintEditor::SelectionState_ClassDefaults);
+		UBlueprint* Blueprint = GetBlueprintObj();
 		StartEditingDefaults( true, true );
 	}
 }
@@ -2850,16 +2859,16 @@ void FBlueprintEditor::OnBlueprintChangedImpl(UBlueprint* InBlueprint, bool bIsJ
 	{
 		DestroyPreview();
 
-		// Refresh the graphs
-		ERefreshBlueprintEditorReason::Type Reason = bIsJustBeingCompiled ? ERefreshBlueprintEditorReason::BlueprintCompiled : ERefreshBlueprintEditorReason::UnknownReason;
-		RefreshEditors(Reason);
-
 		// Notify that the blueprint has been changed (update Content browser, etc)
 		InBlueprint->PostEditChange();
 
 		// Call PostEditChange() on any Actors that are based on this Blueprint
 		FBlueprintEditorUtils::PostEditChangeBlueprintActors(InBlueprint);
 
+		// Refresh the graphs
+		ERefreshBlueprintEditorReason::Type Reason = bIsJustBeingCompiled ? ERefreshBlueprintEditorReason::BlueprintCompiled : ERefreshBlueprintEditorReason::UnknownReason;
+		RefreshEditors(Reason);
+
 		// In case objects were deleted, which should close the tab
 		if (GetCurrentMode() == FBlueprintEditorApplicationModes::StandardBlueprintEditorMode)
 		{
@@ -6225,29 +6234,38 @@ void FBlueprintEditor::OnRepairCorruptedBlueprint()
 
 void FBlueprintEditor::StartEditingDefaults(bool bAutoFocus, bool bForceRefresh)
 {
-	const bool bSingleLayoutBPEditor = GetDefault<UEditorExperimentalSettings>()->bUnifiedBlueprintEditor && IsDetailsPanelEditingClassDefaults();
+	SetUISelectionState(FBlueprintEditor::SelectionState_ClassDefaults);
+
+	const bool bSingleLayoutBPEditor = GetDefault<UEditorExperimentalSettings>()->bUnifiedBlueprintEditor;
 
 	if (IsEditingSingleBlueprint())
 	{
-		if (GetBlueprintObj()->GeneratedClass != NULL)
+		if (GetBlueprintObj()->GeneratedClass != nullptr)
 		{
 			if (bSingleLayoutBPEditor)
 			{
-				UObject* DefaultObject = GetBlueprintObj()->GeneratedClass->GetDefaultObject();
+				if ( SCSEditor.IsValid() && GetBlueprintObj()->GeneratedClass->IsChildOf<AActor>() )
+				{
+					SCSEditor->SelectRoot();
+				}
+				else
+				{
+					UObject* DefaultObject = GetBlueprintObj()->GeneratedClass->GetDefaultObject();
 
-				// Update the details panel
-				FString Title;
-				DefaultObject->GetName(Title);
-				SKismetInspector::FShowDetailsOptions Options(FText::FromString(Title), true);
-				Options.bShowComponents = false;
+					// Update the details panel
+					FString Title;
+					DefaultObject->GetName(Title);
+					SKismetInspector::FShowDetailsOptions Options(FText::FromString(Title), true);
+					Options.bShowComponents = false;
 
-				Inspector->ShowDetailsForSingleObject(DefaultObject, Options);
+					Inspector->ShowDetailsForSingleObject(DefaultObject, Options);
 
-				TSharedPtr<SDockTab> OwnerTab = Inspector->GetOwnerTab();
-				if (OwnerTab.IsValid())
-				{
-					OwnerTab->ActivateInParent(ETabActivationCause::SetDirectly);
-					OwnerTab->FlashTab();
+					TSharedPtr<SDockTab> OwnerTab = Inspector->GetOwnerTab();
+					if ( OwnerTab.IsValid() )
+					{
+						OwnerTab->ActivateInParent(ETabActivationCause::SetDirectly);
+						OwnerTab->FlashTab();
+					}
 				}
 			}
 			else
diff --git a/Engine/Source/Editor/Kismet/Private/SKismetInspector.cpp b/Engine/Source/Editor/Kismet/Private/SKismetInspector.cpp
--- a/Engine/Source/Editor/Kismet/Private/SKismetInspector.cpp
+++ b/Engine/Source/Editor/Kismet/Private/SKismetInspector.cpp
@@ -126,6 +126,19 @@ TSharedPtr<SDockTab> SKismetInspector::GetOwnerTab() const
 	return OwnerTab.Pin();
 }
 
+bool SKismetInspector::IsSelected(UObject* Object) const
+{
+	for ( const TWeakObjectPtr<UObject>& SelectedObject : SelectedObjects )
+	{
+		if ( SelectedObject.Get() == Object )
+		{
+			return true;
+		}
+	}
+
+	return false;
+}
+
 FText SKismetInspector::GetContextualEditingWidgetTitle() const
 {
 	FText Title = PropertyViewTitle;
diff --git a/Engine/Source/Editor/Kismet/Private/SMyBlueprint.cpp b/Engine/Source/Editor/Kismet/Private/SMyBlueprint.cpp
--- a/Engine/Source/Editor/Kismet/Private/SMyBlueprint.cpp
+++ b/Engine/Source/Editor/Kismet/Private/SMyBlueprint.cpp
@@ -2047,6 +2047,11 @@ UEdGraph* SMyBlueprint::GetFocusedGraph() const
 	return EdGraph;
 }
 
+void SMyBlueprint::OnObjectPropertyChanged(UObject* InObject, FPropertyChangedEvent& InPropertyChangedEvent)
+{
+	bNeedsRefresh = ( InObject == Blueprint );
+}
+
 void SMyBlueprint::OnDeleteDelegate(FEdGraphSchemaAction_K2Delegate* InDelegateAction)
 {
 	UEdGraph* EdGraph = InDelegateAction->EdGraph;
diff --git a/Engine/Source/Editor/Kismet/Private/SMyBlueprint.h b/Engine/Source/Editor/Kismet/Private/SMyBlueprint.h
--- a/Engine/Source/Editor/Kismet/Private/SMyBlueprint.h
+++ b/Engine/Source/Editor/Kismet/Private/SMyBlueprint.h
@@ -226,10 +226,7 @@ private:
 	UEdGraph* GetFocusedGraph() const;
 
 	/** Delegate to hook us into non-structural Blueprint object post-change events */
-	void OnObjectPropertyChanged(UObject* InObject, FPropertyChangedEvent& InPropertyChangedEvent)
-	{
-		bNeedsRefresh = (InObject == Blueprint);
-	}
+	void OnObjectPropertyChanged(UObject* InObject, FPropertyChangedEvent& InPropertyChangedEvent);
 private:
 	/** Pointer back to the blueprint editor that owns us */
 	TWeakPtr<FBlueprintEditor> BlueprintEditorPtr;
diff --git a/Engine/Source/Editor/Kismet/Public/SKismetInspector.h b/Engine/Source/Editor/Kismet/Public/SKismetInspector.h
--- a/Engine/Source/Editor/Kismet/Public/SKismetInspector.h
+++ b/Engine/Source/Editor/Kismet/Public/SKismetInspector.h
@@ -72,6 +72,9 @@ public:
 	void SetOwnerTab(TSharedRef<SDockTab> Tab);
 	TSharedPtr<SDockTab> GetOwnerTab() const;
 
+	/** @return true if the object is in the selection set. */
+	bool IsSelected(UObject* Object) const;
+
 protected:
 	/** Update the inspector window to show information on the supplied objects */
 	void UpdateFromObjects(const TArray<UObject*>& PropertyObjects, struct FKismetSelectionInfo& SelectionInfo, const FShowDetailsOptions& Options);
diff --git a/Engine/Source/Runtime/Slate/Public/Widgets/Views/STableRow.h b/Engine/Source/Runtime/Slate/Public/Widgets/Views/STableRow.h
--- a/Engine/Source/Runtime/Slate/Public/Widgets/Views/STableRow.h
+++ b/Engine/Source/Runtime/Slate/Public/Widgets/Views/STableRow.h
@@ -785,7 +785,7 @@ public:
 	}
 
 	/** 
-	 * Callback to determine if the row is selected or not
+	 * Callback to determine if the row is selected singularly and has keyboard focus or not
 	 *
 	 * @return		true if selected by owning widget.
 	 */
@@ -802,6 +802,19 @@ public:
 		return OwnerWidget->Private_IsItemSelected( *MyItem );
 	}
 
+	/**
+	 * Callback to determine if the row is selected or not
+	 *
+	 * @return		true if selected by owning widget.
+	 */
+	bool IsSelected() const
+	{
+		TSharedPtr< ITypedTableView< ItemType > > OwnerWidget = OwnerTablePtr.Pin();
+
+		const ItemType* MyItem = OwnerWidget->Private_ItemFromWidget(this);
+		return OwnerWidget->Private_IsItemSelected(*MyItem);
+	}
+
 	/** Protected constructor; SWidgets should only be instantiated via declarative syntax. */
 	STableRow()
 		: IndexInList(0)
