[INTEGRATE] Change 2432164 by [REDACTED_PERSON][REDACTED_HANDLE] on 2015/02/04 14:37:16

	[UE-8809] Hiding the "operate on selected actor(s)" section of the context menu when you're dragging off a pin.

diff --git a/Engine/Source/Editor/Kismet/Private/BlueprintActionMenuUtils.cpp b/Engine/Source/Editor/Kismet/Private/BlueprintActionMenuUtils.cpp
--- a/Engine/Source/Editor/Kismet/Private/BlueprintActionMenuUtils.cpp
+++ b/Engine/Source/Editor/Kismet/Private/BlueprintActionMenuUtils.cpp
@@ -431,7 +431,7 @@ void FBlueprintActionMenuUtils::MakeContextMenu(FBlueprintActionContext const& C
 
 	const UBlueprintEditorSettings* BlueprintSettings = GetDefault<UBlueprintEditorSettings>();
 	bool const bAddTargetContext  = bIsContextSensitive && BlueprintSettings->bUseTargetContextForNodeMenu;
-	bool bCanOperateOnLevelActors = bIsContextSensitive;
+	bool bCanOperateOnLevelActors = bIsContextSensitive && (Context.Pins.Num() == 0);
 	bool bCanHaveActorComponents  = bIsContextSensitive;
 	// determine if we can operate on certain object selections (level actors, 
 	// components, etc.)
