Merging

//depot/UE4-Releases/4.7/Engine/Source/Editor/Kismet/Private/BlueprintEditorTabFactories.cpp

to //depot/UE4/Engine/Source/Editor/Kismet/Private/BlueprintEditorTabFactories.cpp

diff --git a/Engine/Source/Editor/Kismet/Private/BlueprintEditorTabFactories.cpp b/Engine/Source/Editor/Kismet/Private/BlueprintEditorTabFactories.cpp
--- a/Engine/Source/Editor/Kismet/Private/BlueprintEditorTabFactories.cpp
+++ b/Engine/Source/Editor/Kismet/Private/BlueprintEditorTabFactories.cpp
@@ -43,9 +43,9 @@ void FGraphEditorSummoner::SaveState(TSharedPtr<SDockTab> Tab, TSharedPtr<FTabPa
 	float ZoomAmount;
 	GraphEditor->GetViewLocation(ViewLocation, ZoomAmount);
 
-	UEdGraph* Graph = FTabPayload_UObject::CastChecked<UEdGraph>(Payload);
+	UEdGraph* Graph = Payload->IsValid() ? FTabPayload_UObject::CastChecked<UEdGraph>(Payload) : nullptr;
 
-	if (BlueprintEditorPtr.Pin()->IsGraphInCurrentBlueprint(Graph))
+	if (Graph && BlueprintEditorPtr.Pin()->IsGraphInCurrentBlueprint(Graph))
 	{
 		// Don't save references to external graphs.
 		BlueprintEditorPtr.Pin()->GetBlueprintObj()->LastEditedDocuments.Add(FEditedDocumentInfo(Graph, ViewLocation, ZoomAmount));
