[INTEGRATE] Change 2433325 by [REDACTED_PERSON][REDACTED_HANDLE] on 2015/02/05 05:46:56

	UE-9157 Crash PIEing after editing MovementComponent

diff --git a/Engine/Source/Editor/Kismet/Private/BlueprintDetailsCustomization.cpp b/Engine/Source/Editor/Kismet/Private/BlueprintDetailsCustomization.cpp
--- a/Engine/Source/Editor/Kismet/Private/BlueprintDetailsCustomization.cpp
+++ b/Engine/Source/Editor/Kismet/Private/BlueprintDetailsCustomization.cpp
@@ -4336,9 +4336,12 @@ void FBlueprintComponentDetails::CustomizeDetails(IDetailLayoutBuilder& DetailLa
 	check(Blueprint != NULL);
 
 	TArray<FSCSEditorTreeNodePtrType> Nodes = Editor->GetSelectedNodes();
-	check(Nodes.Num() > 0);
 
-	if (Nodes.Num() == 1)
+	if (!Nodes.Num())
+	{
+		CachedNodePtr = nullptr;
+	}
+	else if (Nodes.Num() == 1)
 	{
 		CachedNodePtr = Nodes[0];
 	}
@@ -4539,7 +4542,7 @@ UClass* FBlueprintComponentDetails::FindCommonBaseClassFromSelected() const
 				SelectionClasses.Add( Node->GetComponentTemplate()->GetClass() );
 			}
 		}
-		Result = UClass::FindCommonBase( SelectionClasses );
+		Result = SelectionClasses.Num() ? UClass::FindCommonBase(SelectionClasses) : nullptr;
 	}
 	return Result;
 }
