[INTEGRATE] Change 2433441 by [REDACTED_PERSON][REDACTED_HANDLE] on 2015/02/05 08:33:14

	BP - The function outputs section properly refreshes now upon the first addition of an output pin.  Fixing several selection bugs in the editor.  Additionally changing the section for a few blueprint options and making a few flags advanced in the blueprint settings.

diff --git a/Engine/Source/Editor/Kismet/Private/BlueprintDetailsCustomization.cpp b/Engine/Source/Editor/Kismet/Private/BlueprintDetailsCustomization.cpp
--- a/Engine/Source/Editor/Kismet/Private/BlueprintDetailsCustomization.cpp
+++ b/Engine/Source/Editor/Kismet/Private/BlueprintDetailsCustomization.cpp
@@ -3903,7 +3903,7 @@ FReply FBlueprintGraphActionDetails::OnAddNewOutputClicked()
 
 	AttemptToCreateResultNode();
 
-	UK2Node_EditablePinBase * FunctionResultNode = FunctionResultNodePtr.Get();
+	UK2Node_EditablePinBase* FunctionResultNode = FunctionResultNodePtr.Get();
 	if( FunctionResultNode )
 	{
 		FunctionResultNode->Modify();
@@ -3922,6 +3922,11 @@ FReply FBlueprintGraphActionDetails::OnAddNewOutputClicked()
 		FunctionResultNode->CreateUserDefinedPin(NewPinName, PinType);
 		
 		OnParamsChanged(FunctionResultNode, true);
+
+		if ( !PreviousResultNode )
+		{
+			DetailsLayoutPtr->ForceRefreshDetails();
+		}
 	}
 	else
 	{
@@ -4246,15 +4251,17 @@ void FBlueprintGlobalOptionsDetails::CustomizeDetails(IDetailLayoutBuilder& Deta
 		for (TFieldIterator<UProperty> PropertyIt(Blueprint->GetClass(), EFieldIteratorFlags::IncludeSuper); PropertyIt; ++PropertyIt)
 		{
 			UProperty* Property = *PropertyIt;
-			if(FObjectEditorUtils::GetCategory(Property) != TEXT("BlueprintOptions"))
+			FString Category = FObjectEditorUtils::GetCategory(Property);
+
+			if ( Category != TEXT("BlueprintOptions") && Category != TEXT("ClassOptions") )
 			{
 				DetailLayout.HideProperty(DetailLayout.GetProperty(Property->GetFName()));
 			}
 		}
 
 		// Display the parent class and set up the menu for reparenting
-		IDetailCategoryBuilder& Category = DetailLayout.EditCategory("Globals", LOCTEXT("BlueprintGlobalDetailsCategory", "Globals"));
-		Category.AddCustomRow( LOCTEXT("BlueprintGlobalDetailsCategory", "Globals") )
+		IDetailCategoryBuilder& Category = DetailLayout.EditCategory("ClassOptions", LOCTEXT("ClassOptions", "Class Options"));
+		Category.AddCustomRow( LOCTEXT("ClassOptions", "Class Options") )
 		.NameContent()
 		[
 			SNew(STextBlock)
@@ -4305,7 +4312,7 @@ void FBlueprintGlobalOptionsDetails::CustomizeDetails(IDetailLayoutBuilder& Deta
 		else
 		{
 			// Only display the ability to deprecate a Blueprint on non-level Blueprints.
-			Category.AddCustomRow( LOCTEXT("DeprecateLabel", "Deprecate") )
+			Category.AddCustomRow( LOCTEXT("DeprecateLabel", "Deprecate"), true )
 				.NameContent()
 				[
 					SNew(STextBlock)
@@ -4316,10 +4323,10 @@ void FBlueprintGlobalOptionsDetails::CustomizeDetails(IDetailLayoutBuilder& Deta
 				.ValueContent()
 				[
 					SNew(SCheckBox)
-						.IsEnabled( this, &FBlueprintGlobalOptionsDetails::CanDeprecateBlueprint )
-						.IsChecked( this, &FBlueprintGlobalOptionsDetails::IsDeprecatedBlueprint )
-						.OnCheckStateChanged( this, &FBlueprintGlobalOptionsDetails::OnDeprecateBlueprint )
-						.ToolTipText( this, &FBlueprintGlobalOptionsDetails::GetDeprecatedTooltip )
+					.IsEnabled( this, &FBlueprintGlobalOptionsDetails::CanDeprecateBlueprint )
+					.IsChecked( this, &FBlueprintGlobalOptionsDetails::IsDeprecatedBlueprint )
+					.OnCheckStateChanged( this, &FBlueprintGlobalOptionsDetails::OnDeprecateBlueprint )
+					.ToolTipText( this, &FBlueprintGlobalOptionsDetails::GetDeprecatedTooltip )
 				];
 		}
 	}
diff --git a/Engine/Source/Editor/Kismet/Private/BlueprintEditor.cpp b/Engine/Source/Editor/Kismet/Private/BlueprintEditor.cpp
--- a/Engine/Source/Editor/Kismet/Private/BlueprintEditor.cpp
+++ b/Engine/Source/Editor/Kismet/Private/BlueprintEditor.cpp
@@ -630,12 +630,13 @@ void FBlueprintEditor::RefreshEditors(ERefreshBlueprintEditorReason::Type Reason
 		UpdateSCSPreview();
 	}
 
+	// Only force a refresh of the defaults view if we have a standalone defaults viewer.
 	// Note: There is an optimization inside of ShowDetailsForSingleObject() that skips the refresh if the object being selected is the same as the previous object.
 	// The SKismetInspector class is shared between both Defaults mode and Components mode, but in Defaults mode the object selected is always going to be the CDO. Given
 	// that the selection does not really change, we force it to refresh and skip the optimization. Otherwise, some things may not work correctly in Defaults mode. For
 	// example, transform details are customized and the rotation value is cached at customization time; if we don't force refresh here, then after an undo of a previous
 	// rotation edit, transform details won't be re-customized and thus the cached rotation value will be stale, resulting in an invalid rotation value on the next edit.
-	StartEditingDefaults(/*bAutoFocus=*/false, /*bForceRefresh=*/true);
+	RefreshStandAloneDefaultsEditor();
 
 	// Update associated controls like the function editor
 	BroadcastRefresh();
@@ -1497,7 +1498,11 @@ void FBlueprintEditor::SetCurrentMode(FName NewMode)
 	FWorkflowCentricApplication::SetCurrentMode(NewMode);
 }
 
-void FBlueprintEditor::InitBlueprintEditor(const EToolkitMode::Type Mode, const TSharedPtr< IToolkitHost >& InitToolkitHost, const TArray<UBlueprint*>& InBlueprints, bool bShouldOpenInDefaultsMode)
+void FBlueprintEditor::InitBlueprintEditor(
+	const EToolkitMode::Type Mode,
+	const TSharedPtr< IToolkitHost >& InitToolkitHost,
+	const TArray<UBlueprint*>& InBlueprints,
+	bool bShouldOpenInDefaultsMode)
 {
 	check(InBlueprints.Num() == 1 || bShouldOpenInDefaultsMode);
 
@@ -2031,6 +2036,7 @@ FBlueprintEditor::~FBlueprintEditor()
 
 void FBlueprintEditor::FocusInspectorOnGraphSelection(const FGraphPanelSelectionSet& NewSelection, bool bForceRefresh)
 {
+	// If this graph has selected nodes update the details panel to match.
 	if ( NewSelection.Array().Num() > 0 )
 	{
 		SetUISelectionState(FBlueprintEditor::SelectionState_Graph);
@@ -2084,7 +2090,7 @@ void FBlueprintEditor::CreateDefaultTabContents(const TArray<UBlueprint*>& InBlu
 		. IsPropertyEditingEnabledDelegate( IsPropertyEditingEnabledDelegate )
 		. OnFinishedChangingProperties( FOnFinishedChangingProperties::FDelegate::CreateSP(this, &FBlueprintEditor::OnFinishedChangingProperties) );
 
-	if (InBlueprints.Num() > 0)
+	if ( InBlueprints.Num() > 0 )
 	{
 		const bool bShowPublicView = true;
 		const bool bHideNameArea = false;
@@ -2729,6 +2735,13 @@ void FBlueprintEditor::OnActiveTabChanged( TSharedPtr<SDockTab> PreviouslyActive
 	{
 		//UE_LOG(LogBlueprint, Log, TEXT("OnActiveTabChanged: %s"), *NewlyActivated->GetLayoutIdentifier().ToString());
 	}
+
+	if ( NewlyActivated->GetTabRole() == ETabRole::DocumentTab )
+	{
+		FocusedGraphEdPtr = nullptr;
+	}
+
+	MyBlueprintWidget->Refresh();
 }
 
 void FBlueprintEditor::OnGraphEditorFocused(const TSharedRef<SGraphEditor>& InGraphEditor)
@@ -2738,16 +2751,31 @@ void FBlueprintEditor::OnGraphEditorFocused(const TSharedRef<SGraphEditor>& InGr
 	InGraphEditor->SetPinVisibility(PinVisibility);
 
 	// Update the inspector as well, to show selection from the focused graph editor
-	FocusInspectorOnGraphSelection(GetSelectedNodes());
+	FGraphPanelSelectionSet SelectedNodes = GetSelectedNodes();
+	FocusInspectorOnGraphSelection(SelectedNodes);
 
 	// During undo, garbage graphs can be temporarily brought into focus, ensure that before a refresh of the MyBlueprint window that the graph is owned by a Blueprint
-	if(FocusedGraphEdPtr.IsValid() && MyBlueprintWidget.IsValid())
+	if ( FocusedGraphEdPtr.IsValid() && MyBlueprintWidget.IsValid() )
 	{
 		// The focused graph can be garbage as well
-		TWeakObjectPtr< UEdGraph > FocusedGraph = FocusedGraphEdPtr.Pin()->GetCurrentGraph();
-		if(FocusedGraph.IsValid() && FBlueprintEditorUtils::FindBlueprintForGraph(FocusedGraph.Get()))
+		TWeakObjectPtr< UEdGraph > FocusedGraphPtr = FocusedGraphEdPtr.Pin()->GetCurrentGraph();
+		UEdGraph* FocusedGraph = FocusedGraphPtr.Get();
+
+		if ( FocusedGraph != nullptr )
 		{
-			MyBlueprintWidget->Refresh();
+			if ( FBlueprintEditorUtils::FindBlueprintForGraph(FocusedGraph) )
+			{
+				MyBlueprintWidget->Refresh();
+			}
+
+			// If this graph doesn't have any selected nodes, but the current selection state is from another graph,
+			// lets just set the focused graph to be the selected one.  We don't want the wrong function context to 
+			// seem relevant when you're looking at a different function.
+			if ( SelectedNodes.Num() == 0 )
+			{
+				SetUISelectionState(FBlueprintEditor::SelectionState_Graph);
+				Inspector->ShowDetailsForSingleObject(FocusedGraph);
+			}
 		}
 	}
 }
@@ -6255,7 +6283,7 @@ void FBlueprintEditor::StartEditingDefaults(bool bAutoFocus, bool bForceRefresh)
 					// Update the details panel
 					FString Title;
 					DefaultObject->GetName(Title);
-					SKismetInspector::FShowDetailsOptions Options(FText::FromString(Title), true);
+					SKismetInspector::FShowDetailsOptions Options(FText::FromString(Title), bForceRefresh);
 					Options.bShowComponents = false;
 
 					Inspector->ShowDetailsForSingleObject(DefaultObject, Options);
@@ -6268,28 +6296,31 @@ void FBlueprintEditor::StartEditingDefaults(bool bAutoFocus, bool bForceRefresh)
 					}
 				}
 			}
-			else
-			{
-				DefaultEditor->ShowDetailsForSingleObject(GetBlueprintObj()->GeneratedClass->GetDefaultObject(), SKismetInspector::FShowDetailsOptions(FText::GetEmpty(), bForceRefresh));
-			}
 		}
 	}
-	else if (GetEditingObjects().Num() > 0)
+	
+	RefreshStandAloneDefaultsEditor();
+}
+
+void FBlueprintEditor::RefreshStandAloneDefaultsEditor()
+{
+	// Update the details panel
+	SKismetInspector::FShowDetailsOptions Options(FText::GetEmpty(), true);
+
+	TArray<UObject*> DefaultObjects;
+	for ( int32 i = 0; i < GetEditingObjects().Num(); ++i )
 	{
-		TArray<UObject*> DefaultObjects;
-		for (int32 i = 0; i < GetEditingObjects().Num(); ++i)
+		auto Blueprint = (UBlueprint*)( GetEditingObjects()[i] );
+		if ( Blueprint->GeneratedClass )
 		{
-			auto Blueprint = (UBlueprint*)(GetEditingObjects()[i]);
-			if (Blueprint->GeneratedClass)
-			{
-				DefaultObjects.Add(Blueprint->GeneratedClass->GetDefaultObject());
-			}
-		}
-		if (DefaultObjects.Num())
-		{
-			DefaultEditor->ShowDetailsForObjects(DefaultObjects);
+			DefaultObjects.Add(Blueprint->GeneratedClass->GetDefaultObject());
 		}
 	}
+
+	if ( DefaultObjects.Num() )
+	{
+		DefaultEditor->ShowDetailsForObjects(DefaultObjects);
+	}
 }
 
 void FBlueprintEditor::RenameNewlyAddedAction(FName InActionName)
diff --git a/Engine/Source/Editor/Kismet/Public/BlueprintEditor.h b/Engine/Source/Editor/Kismet/Public/BlueprintEditor.h
--- a/Engine/Source/Editor/Kismet/Public/BlueprintEditor.h
+++ b/Engine/Source/Editor/Kismet/Public/BlueprintEditor.h
@@ -530,6 +530,9 @@ public:
 protected:
 	/** Called during initialization of the blueprint editor to register any application modes. */
 	virtual void RegisterApplicationModes(const TArray<UBlueprint*>& InBlueprints, bool bShouldOpenInDefaultsMode, bool bNewlyCreated = false);
+
+	// Updates the selected object used by the stand lone defaults editor widget.
+	void RefreshStandAloneDefaultsEditor();
 	
 	// Zooming to fit the entire graph
 	void ZoomToWindow_Clicked();
diff --git a/Engine/Source/Runtime/Engine/Classes/Engine/Blueprint.h b/Engine/Source/Runtime/Engine/Classes/Engine/Blueprint.h
--- a/Engine/Source/Runtime/Engine/Classes/Engine/Blueprint.h
+++ b/Engine/Source/Runtime/Engine/Classes/Engine/Blueprint.h
@@ -272,7 +272,7 @@ class ENGINE_API UBlueprint : public UBlueprintCore
 
 	/** Pointer to the parent class that the generated class should derive from */
 	UPROPERTY(AssetRegistrySearchable)
-	TSubclassOf<class UObject>  ParentClass;
+	TSubclassOf<class UObject> ParentClass;
 
 	UPROPERTY(transient)
 	UObject* PRIVATE_InnermostPreviousCDO;
@@ -299,7 +299,7 @@ class ENGINE_API UBlueprint : public UBlueprintCore
 	uint32 bRunConstructionScriptOnDrag : 1;
 
 	/** Whether or not this blueprint's class is a const class or not.  Should set CLASS_Const in the KismetCompiler. */
-	UPROPERTY(EditAnywhere, Category=BlueprintOptions)
+	UPROPERTY(EditAnywhere, Category=ClassOptions, AdvancedDisplay)
 	uint32 bGenerateConstClass : 1;
 
 	/**shows up in the content browser when the blueprint is hovered */
@@ -323,7 +323,7 @@ class ENGINE_API UBlueprint : public UBlueprintCore
 	FGuid SearchGuid;
 
 	/** Deprecates the Blueprint, marking the generated class with the CLASS_Deprecated flag */
-	UPROPERTY(EditAnywhere, Category=BlueprintOption)
+	UPROPERTY(EditAnywhere, Category=ClassOptions, AdvancedDisplay)
 	bool bDeprecate;
 #endif //WITH_EDITORONLY_DATA
 
