[INTEGRATE] Change 2433482 by [REDACTED_PERSON][REDACTED_HANDLE] on 2015/02/05 09:11:04

	[UE-8315] - Duplicated components appear in the proper location in BP editor
	[UE-8617] - Alt+LMB drag to duplicate component in world works
	[UE-8802] - Fixed issues with location of duplicated component
	[UE-9193] - Substantially reorganized cut, copy, paste, duplicate, and delete to work in the world. Moved most of the implementations to FComponentEditorUtils so it could be shared by the SCSEditor (in ActorInstance mode) and the level editor

	Checking in to keep from getting too unwieldly/large.
	Known issues:
	- Paste doesn't always show up in the right place
	- Paste always attaches to root in world
	- Can crash when undoing transform manipulation with multiple duplicated components selected

diff --git a/Engine/Source/Editor/Kismet/Private/SSCSEditor.cpp b/Engine/Source/Editor/Kismet/Private/SSCSEditor.cpp
--- a/Engine/Source/Editor/Kismet/Private/SSCSEditor.cpp
+++ b/Engine/Source/Editor/Kismet/Private/SSCSEditor.cpp
@@ -16,7 +16,6 @@
 
 #include "AssetSelection.h"
 #include "Editor/SceneOutliner/Private/SSocketChooser.h"
-#include "Factories.h"
 #include "ScopedTransaction.h"
 
 #include "DragAndDrop/AssetDragDropOp.h"
@@ -26,7 +25,6 @@
 
 #include "IDocumentation.h"
 #include "Kismet2NameValidators.h"
-#include "UnrealExporter.h"
 #include "TutorialMetaData.h"
 #include "SInlineEditableTextBlock.h"
 #include "GenericCommands.h"
@@ -1213,102 +1211,6 @@ void FSCSEditorTreeNodeRootActor::OnCompleteRename(const FText& InNewName)
 	}
 }
 
-//////////////////////////////////////////////////////////////////////////
-// FSCSEditorComponentObjectTextFactory
-
-struct FSCSEditorComponentObjectTextFactory : public FCustomizableTextObjectFactory
-{
-	// Child->Parent name map
-	TMap<FName, FName> ParentMap;
-
-	// Name->Instance object mapping
-	TMap<FName, UActorComponent*> NewObjectMap;
-
-	// Determine whether or not scene components in the new object set can be attached to the given scene root component
-	bool CanAttachComponentsTo(USceneComponent* InRootComponent)
-	{
-		check(InRootComponent);
-
-		// For each component in the set, check against the given root component and break if we fail to validate
-		bool bCanAttachToRoot = true;
-		for(auto NewComponentIt = NewObjectMap.CreateConstIterator(); NewComponentIt && bCanAttachToRoot; ++NewComponentIt)
-		{
-			// If this is a scene component, and it does not already have a parent within the set
-			USceneComponent* SceneComponent = Cast<USceneComponent>(NewComponentIt->Value);
-			if(SceneComponent != NULL && !ParentMap.Contains(SceneComponent->GetFName()))
-			{
-				// Determine if we are allowed to attach the scene component to the given root component
-				bCanAttachToRoot = InRootComponent->CanAttachAsChild(SceneComponent, NAME_None)
-					&& SceneComponent->Mobility >= InRootComponent->Mobility
-					&& (!InRootComponent->IsEditorOnly() || SceneComponent->IsEditorOnly());
-			}
-		}
-
-		return bCanAttachToRoot;
-	}
-
-	// Constructs a new object factory from the given text buffer
-	static TSharedRef<FSCSEditorComponentObjectTextFactory> Get(const EComponentEditorMode::Type EditorMode, const FString& InTextBuffer)
-	{
-		// Construct a new instance
-		TSharedPtr<FSCSEditorComponentObjectTextFactory> FactoryPtr = MakeShareable(new FSCSEditorComponentObjectTextFactory());
-		check(FactoryPtr.IsValid());
-
-		// Create new objects if we're allowed to
-		if(FactoryPtr->CanCreateObjectsFromText(InTextBuffer))
-		{
-			EObjectFlags ObjectFlags = RF_Transactional;
-			if (EditorMode == EComponentEditorMode::BlueprintSCS)
-			{
-				ObjectFlags |= RF_ArchetypeObject;
-			}
-
-			// Use the transient package initially for creating the objects, since the variable name is used when copying
-			FactoryPtr->ProcessBuffer(GetTransientPackage(), ObjectFlags, InTextBuffer);
-		}
-
-		return FactoryPtr.ToSharedRef();
-	}
-
-	virtual ~FSCSEditorComponentObjectTextFactory() {}
-
-protected:	
-	// Constructor; protected to only allow this class to instance itself
-	FSCSEditorComponentObjectTextFactory()
-		:FCustomizableTextObjectFactory(GWarn)
-	{
-	}
-
-	// FCustomizableTextObjectFactory implementation
-
-	virtual bool CanCreateClass(UClass* ObjectClass) const override
-	{
-		// Only allow actor component types to be created
-		return ObjectClass->IsChildOf(UActorComponent::StaticClass());
-	}
-
-	virtual void ProcessConstructedObject(UObject* NewObject) override
-	{
-		check(NewObject);
-
-		// Add it to the new object map
-		NewObjectMap.Add(NewObject->GetFName(), Cast<UActorComponent>(NewObject));
-
-		// If this is a scene component and it has a parent
-		USceneComponent* SceneComponent = Cast<USceneComponent>(NewObject);
-		if(SceneComponent && SceneComponent->AttachParent)
-		{
-			// Add an entry to the child->parent name map
-			ParentMap.Add(NewObject->GetFName(), SceneComponent->AttachParent->GetFName());
-
-			// Clear this so it isn't used when constructing the new SCS node
-			SceneComponent->AttachParent = NULL;
-		}
-	}
-
-	// FCustomizableTextObjectFactory (end)
-};
-
 //////////////////////////////////////////////////////////////////////////
 // SSCS_RowWidget
 
@@ -3781,12 +3683,12 @@ void SSCSEditor::OnDuplicateComponent()
 				USceneComponent* NewSceneComponent = Cast<USceneComponent>(CloneComponent);
 				if(NewSceneComponent != NULL)
 				{
-					if( EditorMode == EComponentEditorMode::BlueprintSCS )
+					if (EditorMode == EComponentEditorMode::BlueprintSCS)
 					{
 						// Ensure that any native attachment relationship inherited from the original copy is removed (to prevent a GLEO assertion)
 						NewSceneComponent->DetachFromParent(true);
 					}
-
+					
 					// Attempt to locate the original node in the SCS tree
 					FSCSEditorTreeNodePtrType OriginalNodePtr = FindTreeNode(OriginalComponent);
 					if(OriginalNodePtr.IsValid())
@@ -4492,16 +4394,24 @@ UActorComponent* SSCSEditor::AddNewComponent( UClass* NewComponentClass, UObject
 			TArray<uint8> SavedProperties;
 			FObjectWriter Writer(ComponentTemplate, SavedProperties);
 			FObjectReader(NewComponent, SavedProperties);
+			NewComponent->UpdateComponentToWorld();
 		}
 	}
 	else    // EComponentEditorMode::ActorInstance
 	{
-		if (AActor* ActorInstance = GetActorContext())
+		if (ComponentTemplate)
+		{
+			// Create a duplicate of the provided template
+			NewComponent = AddNewNodeForInstancedComponent(FComponentEditorUtils::DuplicateComponent(ComponentTemplate), nullptr);
+		}
+		else if (AActor* ActorInstance = GetActorContext())
 		{
+			// No template, so create a wholly new component
 			ActorInstance->Modify();
 
+			// Create an appropriate name for the new component
 			FName NewComponentName = NAME_None;
-			if (Asset != nullptr)
+			if (Asset)
 			{
 				NewComponentName = *FComponentEditorUtils::GenerateValidVariableNameFromAsset(Asset, ActorInstance);
 			}
@@ -4509,7 +4419,33 @@ UActorComponent* SSCSEditor::AddNewComponent( UClass* NewComponentClass, UObject
 			{
 				NewComponentName = *FComponentEditorUtils::GenerateValidVariableName(NewComponentClass, ActorInstance);
 			}
-			NewComponent = AddNewNodeForInstancedComponent(NewObject<UActorComponent>(ActorInstance, NewComponentClass, NewComponentName, RF_Transactional, ComponentTemplate), Asset, true);
+
+			// Construct the new component and attach as needed
+			auto NewInstanceComponent = NewObject<UActorComponent>(ActorInstance, NewComponentClass, NewComponentName, RF_Transactional);
+			if (USceneComponent* NewSceneComponent = Cast<USceneComponent>(NewInstanceComponent))
+			{
+				USceneComponent* RootComponent = ActorInstance->GetRootComponent();
+				if (RootComponent)
+				{
+					NewSceneComponent->AttachTo(RootComponent, NAME_None, EAttachLocation::KeepRelativeOffset);
+				}
+				else
+				{
+					ActorInstance->SetRootComponent(NewSceneComponent);
+				}
+			}
+
+			// If the component was created from/for a particular asset, assign it now
+			if (Asset)
+			{
+				FComponentAssetBrokerage::AssignAssetToComponent(NewInstanceComponent, Asset);
+			}
+
+			// Add to SerializedComponents array so it gets saved
+			ActorInstance->AddInstanceComponent(NewInstanceComponent);
+			NewInstanceComponent->RegisterComponent();
+
+			NewComponent = AddNewNodeForInstancedComponent(NewInstanceComponent, Asset);
 		}
 	}
 
@@ -4569,32 +4505,6 @@ UActorComponent* SSCSEditor::AddNewNodeForInstancedComponent(UActorComponent* Ne
 {
 	check(NewInstanceComponent != nullptr);
 
-	AActor* ActorInstance = ActorContext.Get();
-	check(ActorInstance != nullptr);
-
-	if (USceneComponent* NewSceneComponent = Cast<USceneComponent>(NewInstanceComponent))
-	{
-		USceneComponent* RootComponent = ActorInstance->GetRootComponent();
-		if (RootComponent)
-		{
-			NewSceneComponent->AttachTo(RootComponent);
-		}
-		else
-		{
-			ActorInstance->SetRootComponent(NewSceneComponent);
-		}
-	}
-
-	// Add to SerializedComponents array so it gets saved
-	ActorInstance->AddInstanceComponent(NewInstanceComponent);
-
-	if(Asset)
-	{
-		FComponentAssetBrokerage::AssignAssetToComponent(NewInstanceComponent, Asset);
-	}
-
-	NewInstanceComponent->RegisterComponent();
-
 	FSCSEditorTreeNodePtrType NewNodePtr;
 
 	// Add the new node to the editor tree
@@ -4690,16 +4600,9 @@ bool SSCSEditor::CanCopyNodes() const
 
 void SSCSEditor::CopySelectedNodes()
 {
-	FStringOutputDevice Archive;
-	const FExportObjectInnerContext Context;
+	// Distill the selected nodes into a list of components to copy
+	TArray<UActorComponent*> ComponentsToCopy;
 	TArray<FSCSEditorTreeNodePtrType> SelectedNodes = GetSelectedNodes();
-
-	// Clear the mark state for saving.
-	UnMarkAllObjects(EObjectMark(OBJECTMARK_TagExp|OBJECTMARK_TagImp));
-
-	// Duplicate the selected component templates into temporary objects that we can modify
-	TMap<FName, FName> ParentMap;
-	TMap<FName, UActorComponent*> ObjectMap;
 	for (int32 i = 0; i < SelectedNodes.Num(); ++i)
 	{
 		// Get the current selected node reference
@@ -4707,58 +4610,15 @@ void SSCSEditor::CopySelectedNodes()
 		check(SelectedNodePtr.IsValid());
 
 		// Get the component template associated with the selected node
-		UObject* ObjectToCopy = SelectedNodePtr->GetComponentTemplate();
-		if(ObjectToCopy)
-		{
-			// If valid, duplicate the component template into a temporary object
-			ObjectToCopy = StaticDuplicateObject(ObjectToCopy, GetTransientPackage(), *SelectedNodePtr->GetVariableName().ToString(), RF_AllFlags & ~RF_ArchetypeObject);
-			if(ObjectToCopy)
-			{
-				// Get the closest parent node of the current node selection within the selected set
-				FSCSEditorTreeNodePtrType ParentNodePtr = SelectedNodePtr->FindClosestParent(SelectedNodes);
-				if(ParentNodePtr.IsValid())
-				{
-					// If valid, record the parent node's variable name into the node->parent map
-					ParentMap.Add(SelectedNodePtr->GetVariableName(), ParentNodePtr->GetVariableName());
-				}
-
-				// Record the temporary object into the name->object map
-				ObjectMap.Add(SelectedNodePtr->GetVariableName(), CastChecked<UActorComponent>(ObjectToCopy));
-			}
-		}
-	}
-
-	// Export the component object(s) to text for copying
-	for(auto ObjectIt = ObjectMap.CreateIterator(); ObjectIt; ++ObjectIt)
-	{
-		// Get the component object to be copied
-		UActorComponent* ComponentToCopy = ObjectIt->Value;
-		check(ComponentToCopy);
-
-		// If this component object had a parent within the selected set
-		if(ParentMap.Contains(ComponentToCopy->GetFName()))
+		UActorComponent* ComponentTemplate = SelectedNodePtr->GetComponentTemplate();
+		if (ComponentTemplate)
 		{
-			// Get the name of the parent component
-			FName ParentName = ParentMap[ComponentToCopy->GetFName()];
-			if(ObjectMap.Contains(ParentName))
-			{
-				// Ensure that this component is a scene component
-				USceneComponent* SceneComponent = Cast<USceneComponent>(ComponentToCopy);
-				if(SceneComponent)
-				{
-					// Set the attach parent to the matching parent object in the temporary set. This allows us to preserve hierarchy in the copied set.
-					SceneComponent->AttachParent = Cast<USceneComponent>(ObjectMap[ParentName]);
-				}
-			}
+			ComponentsToCopy.Add(ComponentTemplate);
 		}
-
-		// Export the component object to the given string
-		UExporter::ExportToOutputDevice(&Context, ComponentToCopy, NULL, Archive, TEXT("copy"), 0, PPF_ExportsNotFullyQualified|PPF_Copy|PPF_Delimited, false, ComponentToCopy->GetOuter());
 	}
 
-	// Copy text to clipboard
-	FString ExportedText = Archive;
-	FPlatformMisc::ClipboardCopy(*ExportedText);
+	// Copy the components to the clipboard
+	FComponentEditorUtils::CopyComponents(ComponentsToCopy);
 }
 
 bool SSCSEditor::CanPasteNodes() const
@@ -4768,118 +4628,118 @@ bool SSCSEditor::CanPasteNodes() const
 		return false;
 	}
 
-	FString ClipboardContent;
-	FPlatformMisc::ClipboardPaste(ClipboardContent);
-
-	// Obtain the component object text factory for the clipboard content and return whether or not we can use it
-	TSharedRef<FSCSEditorComponentObjectTextFactory> Factory = FSCSEditorComponentObjectTextFactory::Get(EditorMode, ClipboardContent);
-	return Factory->NewObjectMap.Num() > 0
-		&& (SceneRootNodePtr->IsDefaultSceneRoot() || Factory->CanAttachComponentsTo(Cast<USceneComponent>(SceneRootNodePtr->GetComponentTemplate())));
+	return FComponentEditorUtils::CanPasteComponents(Cast<USceneComponent>(SceneRootNodePtr->GetComponentTemplate()), SceneRootNodePtr->IsDefaultSceneRoot());
 }
 
 void SSCSEditor::PasteNodes()
 {
 	const FScopedTransaction Transaction(LOCTEXT("PasteComponents", "Paste Component(s)"));
 
-	// Get the text from the clipboard
-	FString TextToImport;
-	FPlatformMisc::ClipboardPaste(TextToImport);
-
-	// Get a new component object factory for the clipboard content
-	TSharedRef<FSCSEditorComponentObjectTextFactory> Factory = FSCSEditorComponentObjectTextFactory::Get(EditorMode, TextToImport);
-
-	// Clear the current selection
-	SCSTreeWidget->ClearSelection();
-
-	// Get the object that's being edited
-	UBlueprint* Blueprint = nullptr;
-	AActor* ActorInstance = nullptr;
 	if (EditorMode == EComponentEditorMode::BlueprintSCS)
 	{
-		Blueprint = GetBlueprint();
+		// Get the components to paste from the clipboard
+		TMap<FName, FName> ParentMap;
+		TMap<FName, UActorComponent*> NewObjectMap;
+		FComponentEditorUtils::GetComponentsFromClipboard(ParentMap, NewObjectMap, true);
+		
+		// Clear the current selection
+		SCSTreeWidget->ClearSelection();
+
+		// Get the blueprint that's being edited
+		UBlueprint* Blueprint = GetBlueprint();
 		check(Blueprint != nullptr && Blueprint->SimpleConstructionScript != nullptr);
 
 		Blueprint->Modify();
 		SaveSCSCurrentState(Blueprint->SimpleConstructionScript);
-	}
-	else    // EComponentEditorMode::ActorInstance
-	{
-		ActorInstance = GetActorContext();
-		if (ActorInstance != nullptr)
-		{
-			ActorInstance->Modify();
-		}
-	}
-
-	// Create a new tree node for each new (pasted) component
-	TMap<FName, FSCSEditorTreeNodePtrType> NewNodeMap;
-	for(auto NewObjectIt = Factory->NewObjectMap.CreateIterator(); NewObjectIt; ++NewObjectIt)
-	{
-		// Get the component object instance
-		UActorComponent* NewActorComponent = NewObjectIt->Value;
-		check(NewActorComponent);
 
-		if(Blueprint != nullptr)
+		// Create a new tree node for each new (pasted) component
+		TMap<FName, FSCSEditorTreeNodePtrType> NewNodeMap;
+		for (auto NewObjectIt = NewObjectMap.CreateIterator(); NewObjectIt; ++NewObjectIt)
 		{
+			// Get the component object instance
+			UActorComponent* NewActorComponent = NewObjectIt->Value;
+			check(NewActorComponent);
+
 			// Relocate the instance from the transient package to the BPGC and assign it a unique object name
-			NewActorComponent->Rename(NULL, Blueprint->GeneratedClass, REN_DontCreateRedirectors|REN_DoNotDirty);
+			NewActorComponent->Rename(NULL, Blueprint->GeneratedClass, REN_DontCreateRedirectors | REN_DoNotDirty);
 
 			// Create a new SCS node to contain the new component and add it to the tree
 			NewActorComponent = AddNewNode(Blueprint->SimpleConstructionScript->CreateNode(NewActorComponent), NULL, false, false);
-		}
-		else
-		{
-			// Relocate the instance from the transient package to the Actor and assign it a unique object name
-			check(ActorInstance != nullptr); //@TODO: This doesn't look justified, but the code below was already assuming it
-			FString NewComponentName = FComponentEditorUtils::GenerateValidVariableName(NewActorComponent->GetClass(), ActorInstance);
-			NewActorComponent->Rename(*NewComponentName, ActorInstance, REN_DontCreateRedirectors | REN_DoNotDirty);
 
-			// Create a new node to contain the new component instance and add it to the tree
-			NewActorComponent = AddNewNodeForInstancedComponent(NewActorComponent, NULL, false);
-		}
-
-		if(NewActorComponent)
-		{
-			// Locate the node that corresponds to the new component template or instance
-			FSCSEditorTreeNodePtrType NewNodePtr = FindTreeNode(NewActorComponent);
-			if(NewNodePtr.IsValid())
+			if (NewActorComponent)
 			{
-				// Add the new node to the node map
-				NewNodeMap.Add(NewObjectIt->Key, NewNodePtr);
+				// Locate the node that corresponds to the new component template or instance
+				FSCSEditorTreeNodePtrType NewNodePtr = FindTreeNode(NewActorComponent);
+				if (NewNodePtr.IsValid())
+				{
+					// Add the new node to the node map
+					NewNodeMap.Add(NewObjectIt->Key, NewNodePtr);
 
-				// Update the selection to include the new node
-				SCSTreeWidget->SetItemSelection(NewNodePtr, true);
+					// Update the selection to include the new node
+					SCSTreeWidget->SetItemSelection(NewNodePtr, true);
+				}
 			}
 		}
-	}
 
-	// Restore the node hierarchy from the original copy
-	for(auto NodeIt = NewNodeMap.CreateConstIterator(); NodeIt; ++NodeIt)
-	{
-		// If an entry exists in the set of known parent nodes for the current node
-		if(Factory->ParentMap.Contains(NodeIt->Key))
+		// Restore the node hierarchy from the original copy
+		for (auto NodeIt = NewNodeMap.CreateConstIterator(); NodeIt; ++NodeIt)
 		{
-			// Get the parent node name
-			FName ParentName = Factory->ParentMap[NodeIt->Key];
-			if(NewNodeMap.Contains(ParentName))
+			// If an entry exists in the set of known parent nodes for the current node
+			if (ParentMap.Contains(NodeIt->Key))
 			{
-				// Reattach the current node to the parent node (this will also handle detachment from the scene root node)
-				NewNodeMap[ParentName]->AddChild(NodeIt->Value);
+				// Get the parent node name
+				FName ParentName = ParentMap[NodeIt->Key];
+				if (NewNodeMap.Contains(ParentName))
+				{
+					// Reattach the current node to the parent node (this will also handle detachment from the scene root node)
+					NewNodeMap[ParentName]->AddChild(NodeIt->Value);
 
-				// Ensure that the new node is expanded to show the child node(s)
-				SCSTreeWidget->SetItemExpansion(NewNodeMap[ParentName], true);
+					// Ensure that the new node is expanded to show the child node(s)
+					SCSTreeWidget->SetItemExpansion(NewNodeMap[ParentName], true);
+				}
 			}
 		}
-	}
 
-	if(Blueprint != nullptr)
-	{
 		// Modify the Blueprint generated class structure (this will also call UpdateTree() as a result)
 		FBlueprintEditorUtils::MarkBlueprintAsStructurallyModified(Blueprint);
 	}
-	else
+	else    // EComponentEditorMode::ActorInstance
 	{
-		UpdateTree();
+		// Determine where in the hierarchy to paste
+		USceneComponent* TargetComponent = nullptr;
+		auto SelectedNodes = GetSelectedNodes();
+		for (int32 i = 0; i < SelectedNodes.Num(); ++i)
+		{
+			FSCSEditorTreeNodePtrType SelectedNodePtr = SelectedNodes[i];
+			check(SelectedNodePtr.IsValid());
+
+			if (auto SceneComponent = Cast<USceneComponent>(SelectedNodePtr->GetComponentTemplate()))
+			{
+				TargetComponent = SceneComponent;
+				break;
+			}
+		}
+
+		// Paste the components
+		TArray<UActorComponent*> PastedComponents;
+		FComponentEditorUtils::PasteComponents(PastedComponents, GetActorContext(), TargetComponent);
+
+		if (PastedComponents.Num() > 0)
+		{
+			// We only want the pasted node(s) to be selected
+			SCSTreeWidget->ClearSelection();
+			UpdateTree();
+
+			// Select the nodes that correspond to the pasted components
+			for (auto PastedComponent : PastedComponents)
+			{
+				FSCSEditorTreeNodePtrType PastedNode = GetNodeFromActorComponent(PastedComponent);
+				if (PastedNode.IsValid())
+				{
+					SCSTreeWidget->SetItemSelection(PastedNode, true);
+				}
+			}
+		}
 	}
 }
 
@@ -4900,7 +4760,7 @@ bool SSCSEditor::CanDeleteNodes() const
 
 void SSCSEditor::OnDeleteNodes()
 {
-	const FScopedTransaction Transaction( LOCTEXT("RemoveComponent", "Remove Component") );
+	const FScopedTransaction Transaction( LOCTEXT("RemoveComponents", "Remove Components") );
 
 	if (EditorMode == EComponentEditorMode::BlueprintSCS)
 	{
@@ -4956,57 +4816,38 @@ void SSCSEditor::OnDeleteNodes()
 			ActorInstance->Modify();
 		}
 
-		FSCSEditorTreeNodePtrType NewSelection;
-		TArray<FSCSEditorTreeNodePtrType> SelectedNodes = SCSTreeWidget->GetSelectedItems();
+		TArray<UActorComponent*> ComponentsToDelete;
+		TArray<FSCSEditorTreeNodePtrType> SelectedNodes = GetSelectedNodes();
 		for (int32 i = 0; i < SelectedNodes.Num(); ++i)
 		{
-			auto Node = SelectedNodes[i];
+			// Get the current selected node reference
+			FSCSEditorTreeNodePtrType SelectedNodePtr = SelectedNodes[i];
+			check(SelectedNodePtr.IsValid());
 
-			// Find an appropriate node to select after removal
-			if(!NewSelection.IsValid() || NewSelection == Node)
+			// Get the component template associated with the selected node
+			UActorComponent* ComponentTemplate = SelectedNodePtr->GetComponentTemplate();
+			if (ComponentTemplate)
 			{
-				// Default to the parent node
-				NewSelection = Node->GetParent();
-				if (NewSelection.IsValid() && NewSelection->GetNodeType() != FSCSEditorTreeNode::RootActorNode)
-				{
-					// If we have sibling nodes, find the one that immediately precedes the one being removed
-					const TArray<FSCSEditorTreeNodePtrType>& ChildNodes = NewSelection->GetChildren();
-					for (int32 ChildIndex = 0; ChildIndex < ChildNodes.Num() && Node != ChildNodes[ChildIndex]; ++ChildIndex)
-					{
-						NewSelection = ChildNodes[ChildIndex];
-					}
-				}
-				else
-				{
-					const TArray<FSCSEditorTreeNodePtrType>& Nodes = GetRootNodes();
-					int32 NodeIndex = Nodes.Find(Node);
-					if (NodeIndex != INDEX_NONE && NodeIndex > 0)
-					{
-						while (NodeIndex > 0 && Nodes[--NodeIndex]->GetNodeType() == FSCSEditorTreeNode::SeparatorNode)
-						{
-							--NodeIndex;
-						}
-						NewSelection = Nodes[NodeIndex];
-					}
-					else
-					{
-						NewSelection.Reset();
-					}
-				}
+				ComponentsToDelete.Add(ComponentTemplate);
 			}
-
-			// This will clear the current selection
-			RemoveComponentNode(Node);
 		}
 
-		// Reset the selection
-		if(NewSelection.IsValid())
+		UActorComponent* ComponentToSelect = nullptr;
+		int32 NumDeletedComponents = FComponentEditorUtils::DeleteComponents(ComponentsToDelete, ComponentToSelect);
+		if (NumDeletedComponents > 0)
 		{
-			SCSTreeWidget->SetItemSelection(NewSelection, true);
-		}
+			if (ComponentToSelect)
+			{
+				FSCSEditorTreeNodePtrType NodeToSelect = GetNodeFromActorComponent(ComponentToSelect);
+				if (NodeToSelect.IsValid())
+				{
+					SCSTreeWidget->SetSelection(NodeToSelect);
+				}
+			}
 
-		// Rebuild the tree view to reflect the new component hierarchy
-		UpdateTree();
+			// Rebuild the tree view to reflect the new component hierarchy
+			UpdateTree();
+		}
 	}
 
 	// Do this AFTER marking the Blueprint as modified
diff --git a/Engine/Source/Editor/Kismet/Public/SSCSEditor.h b/Engine/Source/Editor/Kismet/Public/SSCSEditor.h
--- a/Engine/Source/Editor/Kismet/Public/SSCSEditor.h
+++ b/Engine/Source/Editor/Kismet/Public/SSCSEditor.h
@@ -758,7 +758,7 @@ public:
 
 	/** Adds a new component instance node to the component Table
 		@param NewInstanceComponent	(In) The component being added to the actor instance
-		@param Asset	(In) Optional asset to assign to the component
+		@param Asset (In) Optional asset to assign to the component
 		@param bSetFocusToNewItem (In) Select the new item and activate the inline rename widget (default is true)
 		@return The reference of the newly created ActorComponent */
 	UActorComponent* AddNewNodeForInstancedComponent(UActorComponent* NewInstanceComponent, UObject* Asset, bool bSetFocusToNewItem = true);
@@ -781,6 +781,10 @@ public:
 	void PasteNodes();
 	bool CanPasteNodes() const;
 
+	/** Callbacks to duplicate the selected component */
+	bool CanDuplicateComponent() const;
+	void OnDuplicateComponent();
+
 	/** Removes existing selected component nodes from the SCS */
 	void OnDeleteNodes();
 	bool CanDeleteNodes() const;
@@ -905,15 +909,14 @@ protected:
 	/** Called when the level editor requests a component to be renamed. */
 	void OnLevelComponentRequestRename(const UActorComponent* InComponent);
 
+	/** Checks to see if renaming is allowed on the selected component */
+	bool CanRenameComponent() const;
 	/**
 	 * Requests a rename on the selected component
 	 * @param bTransactional Whether or not the rename should be transactional (i.e. undoable)
 	 */
 	void OnRenameComponent(bool bTransactional);
 
-	/** Checks to see if renaming is allowed on the selected component */
-	bool CanRenameComponent() const;
-
 	/**
 	 * Function to create events for the current selection
 	 * @param Blueprint						the active blueprint context
@@ -939,10 +942,6 @@ protected:
 	 */
 	static void ViewEvent(UBlueprint* Blueprint, const FName EventName, const FComponentEventConstructionData EventData);
 
-	/** Callbacks to duplicate the selected component */
-	bool CanDuplicateComponent() const;
-	void OnDuplicateComponent();
-
 	/** Helper method to add a tree node for the given SCS node */
 	FSCSEditorTreeNodePtrType AddTreeNode(USCS_Node* InSCSNode, FSCSEditorTreeNodePtrType InParentNodePtr, const bool bIsInheritedSCS);
 
diff --git a/Engine/Source/Editor/LevelEditor/Private/LevelEditor.cpp b/Engine/Source/Editor/LevelEditor/Private/LevelEditor.cpp
--- a/Engine/Source/Editor/LevelEditor/Private/LevelEditor.cpp
+++ b/Engine/Source/Editor/LevelEditor/Private/LevelEditor.cpp
@@ -392,6 +392,11 @@ void FLevelEditorModule::BroadcastMapChanged( UWorld* World, EMapChangeType::Typ
 	MapChangedEvent.Broadcast( World, MapChangeType );
 }
 
+void FLevelEditorModule::BroadcastComponentsEdited()
+{
+	ComponentsEditedEvent.Broadcast();
+}
+
 const FLevelEditorCommands& FLevelEditorModule::GetLevelEditorCommands() const
 {
 	return FLevelEditorCommands::Get();
diff --git a/Engine/Source/Editor/LevelEditor/Private/LevelEditorActions.cpp b/Engine/Source/Editor/LevelEditor/Private/LevelEditorActions.cpp
--- a/Engine/Source/Editor/LevelEditor/Private/LevelEditorActions.cpp
+++ b/Engine/Source/Editor/LevelEditor/Private/LevelEditorActions.cpp
@@ -43,6 +43,7 @@
 #include "ReferenceViewer.h"
 #include "Developer/MeshUtilities/Public/MeshUtilities.h"
 #include "EditorClassUtils.h"
+#include "ComponentEditorUtils.h"
 
 #include "EditorActorFolders.h"
 #include "ActorPickerMode.h"
@@ -1461,7 +1462,20 @@ bool FLevelEditorActionCallbacks::Paste_CanExecute()
 			return false;
 		}
 	}
-	return GUnrealEd->CanPasteSelectedActorsFromClipboard( GetWorld() );
+
+	bool bCanPaste = false;
+	if (GEditor->GetSelectedComponentCount() > 0)
+	{
+		check(GEditor->GetSelectedActorCount() == 1);
+		auto SelectedActor = CastChecked<AActor>(*GEditor->GetSelectedActorIterator());
+		bCanPaste = FComponentEditorUtils::CanPasteComponents(SelectedActor->GetRootComponent());
+	}
+	else
+	{
+		bCanPaste = GUnrealEd->CanPasteSelectedActorsFromClipboard(GetWorld());
+	}
+
+	return bCanPaste;
 }
 
 bool FLevelEditorActionCallbacks::PasteHere_CanExecute()
diff --git a/Engine/Source/Editor/LevelEditor/Private/SActorDetails.cpp b/Engine/Source/Editor/LevelEditor/Private/SActorDetails.cpp
--- a/Engine/Source/Editor/LevelEditor/Private/SActorDetails.cpp
+++ b/Engine/Source/Editor/LevelEditor/Private/SActorDetails.cpp
@@ -13,8 +13,10 @@ void SActorDetails::Construct(const FArguments& InArgs, const FName TabIdentifie
 	bSelectionGuard = false;
 	bShowingRootActorNodeSelected = false;
 
-	// Event subscriptions
 	USelection::SelectionChangedEvent.AddRaw(this, &SActorDetails::OnEditorSelectionChanged);
+	
+	FLevelEditorModule& LevelEditor = FModuleManager::GetModuleChecked<FLevelEditorModule>("LevelEditor");
+	LevelEditor.OnComponentsEdited().AddRaw(this, &SActorDetails::OnComponentsEditedInWorld);
 
 	FPropertyEditorModule& PropPlugin = FModuleManager::LoadModuleChecked<FPropertyEditorModule>("PropertyEditor");
 	FDetailsViewArgs DetailsViewArgs;
@@ -134,6 +136,12 @@ SActorDetails::~SActorDetails()
 {
 	GEditor->UnregisterForUndo(this);
 	USelection::SelectionChangedEvent.RemoveAll(this);
+	
+	auto LevelEditor = FModuleManager::GetModulePtr<FLevelEditorModule>("LevelEditor");
+	if (LevelEditor != nullptr)
+	{
+		LevelEditor->OnComponentsEdited().RemoveAll(this);
+	}
 }
 
 void SActorDetails::SetObjects(const TArray<UObject*>& InObjects)
@@ -172,6 +180,12 @@ void SActorDetails::PostUndo(bool bSuccess)
 	// Refresh the tree and update the selection to match the world
 	SCSEditor->UpdateTree();
 	UpdateComponentTreeFromEditorSelection();
+
+	auto SelectedActor = GetSelectedActorInEditor();
+	if (SelectedActor)
+	{
+		GUnrealEd->SetActorSelectionFlags(SelectedActor);
+	}
 }
 
 void SActorDetails::PostRedo(bool bSuccess)
@@ -179,6 +193,18 @@ void SActorDetails::PostRedo(bool bSuccess)
 	PostUndo(bSuccess);
 }
 
+void SActorDetails::OnComponentsEditedInWorld()
+{
+	if (GetSelectedActorInEditor() == GetActorContext())
+	{
+		// The component composition of the observed actor has changed, so rebuild the node tree
+		TGuardValue<bool> SelectionGuard(bSelectionGuard, true);
+
+		// Refresh the tree and update the selection to match the world
+		SCSEditor->UpdateTree();
+	}
+}
+
 void SActorDetails::OnEditorSelectionChanged(UObject* Object)
 {
 	if(!bSelectionGuard && SCSEditor.IsValid())
diff --git a/Engine/Source/Editor/LevelEditor/Private/SActorDetails.h b/Engine/Source/Editor/LevelEditor/Private/SActorDetails.h
--- a/Engine/Source/Editor/LevelEditor/Private/SActorDetails.h
+++ b/Engine/Source/Editor/LevelEditor/Private/SActorDetails.h
@@ -28,15 +28,17 @@ public:
 
 private:
 	AActor* GetSelectedActorInEditor() const;
+	AActor* GetActorContext() const;
 
+	void OnComponentsEditedInWorld();
 	void OnEditorSelectionChanged(UObject* Object);
-	AActor* GetActorContext() const;
 	void OnSCSEditorRootSelected(AActor* Actor);
 	void OnSCSEditorTreeViewSelectionChanged(const TArray<TSharedPtr<class FSCSEditorTreeNode> >& SelectedNodes);
 	void UpdateComponentTreeFromEditorSelection();
 	bool IsPropertyEditingEnabled() const;
 	void OnBlueprintWarningHyperlinkClicked( const FSlateHyperlinkRun::FMetadata& Metadata );
 	EVisibility GetBlueprintComponentWarningVisibility() const;
+
 private:
 	TSharedPtr<SSplitter> DetailsSplitter;
 	TSharedPtr<class IDetailsView> DetailsView;
diff --git a/Engine/Source/Editor/LevelEditor/Public/LevelEditor.h b/Engine/Source/Editor/LevelEditor/Public/LevelEditor.h
--- a/Engine/Source/Editor/LevelEditor/Public/LevelEditor.h
+++ b/Engine/Source/Editor/LevelEditor/Public/LevelEditor.h
@@ -183,6 +183,9 @@ public:
 	 */
 	virtual void BroadcastMapChanged( UWorld* World, EMapChangeType::Type MapChangeType );
 
+	/** Called when an edit command is executed on one or more components in the world */
+	virtual void BroadcastComponentsEdited();
+
 	/** Called when actor selection changes */
 	DECLARE_EVENT_OneParam( FLevelEditorModule, FActorSelectionChangedEvent, const TArray<UObject*>& );
 	virtual FActorSelectionChangedEvent& OnActorSelectionChanged() { return ActorSelectionChangedEvent; }
@@ -195,6 +198,10 @@ public:
 	DECLARE_EVENT_TwoParams( FLevelEditorModule, FMapChangedEvent, UWorld*, EMapChangeType::Type );
 	virtual FMapChangedEvent& OnMapChanged() { return MapChangedEvent; }
 
+	/** Called when an edit command is executed on components in the world */
+	DECLARE_EVENT(FLevelEditorModule, FComponentsEditedEvent);
+	virtual FComponentsEditedEvent& OnComponentsEdited() { return ComponentsEditedEvent; }
+
 	/** Delegates to be called to extend the level viewport menus */
 	DECLARE_DELEGATE_RetVal_OneParam( TSharedRef<FExtender>, FLevelEditorMenuExtender, const TSharedRef<FUICommandList>);
 	DECLARE_DELEGATE_RetVal_TwoParams( TSharedRef<FExtender>, FLevelViewportMenuExtender_SelectedObjects, const TSharedRef<FUICommandList>, const TArray<UObject*>);
@@ -264,6 +271,9 @@ private:
 	/** Multicast delegate executed when viewports should be redrawn */
 	FRedrawLevelEditingViewportsEvent RedrawLevelEditingViewportsEvent;
 
+	/** Multicast delegate executed after components are edited in the world */
+	FComponentsEditedEvent ComponentsEditedEvent;
+
 	/** Multicast delegate executed when viewports should be redrawn */
 	FTakeHighResScreenShotsEvent TakeHighResScreenShotsEvent;
 
diff --git a/Engine/Source/Editor/UnrealEd/Private/EditorActor.cpp b/Engine/Source/Editor/UnrealEd/Private/EditorActor.cpp
--- a/Engine/Source/Editor/UnrealEd/Private/EditorActor.cpp
+++ b/Engine/Source/Editor/UnrealEd/Private/EditorActor.cpp
@@ -27,6 +27,8 @@
 #include "Engine/Light.h"
 #include "GameFramework/Pawn.h"
 #include "Engine/StaticMeshActor.h"
+#include "ComponentEditorUtils.h"
+#include "LevelEditor.h"
 
 #define LOCTEXT_NAMESPACE "UnrealEd.EditorActor"
 
@@ -102,48 +104,63 @@ UUnrealEdEngine::UUnrealEdEngine(const FObjectInitializer& ObjectInitializer)
 
 void UUnrealEdEngine::edactCopySelected( UWorld* InWorld, FString* DestinationData )
 {
-	// Before copying, deselect:
-	//		- Actors belonging to prefabs unless all actors in the prefab are selected.
-	//		- Builder brushes.
-	TArray<AActor*> ActorsToDeselect;
+	if (GetSelectedComponentCount() > 0)
+	{
+		// Copy components
+		TArray<UActorComponent*> SelectedComponents;
+		for (FSelectedEditableComponentIterator It(GetSelectedEditableComponentIterator()); It; ++It)
+		{
+			SelectedComponents.Add(CastChecked<UActorComponent>(*It));
+		}
 
-	bool bSomeSelectedActorsNotInCurrentLevel = false;
-	for ( FSelectionIterator It( GetSelectedActorIterator() ) ; It ; ++It )
+		FComponentEditorUtils::CopyComponents(SelectedComponents);
+	}
+	else
 	{
-		AActor* Actor = static_cast<AActor*>( *It );
-		checkSlow( Actor->IsA(AActor::StaticClass()) );
+		// Copy Actors
+		// Before copying, deselect:
+		//		- Actors belonging to prefabs unless all actors in the prefab are selected.
+		//		- Builder brushes.
+		TArray<AActor*> ActorsToDeselect;
 
-		// Deselect any selected builder brushes.
-		ABrush* Brush = Cast< ABrush >( Actor );
-		const bool bActorIsBuilderBrush = (Brush && FActorEditorUtils::IsABuilderBrush(Brush));
-		if( bActorIsBuilderBrush )
+		bool bSomeSelectedActorsNotInCurrentLevel = false;
+		for (FSelectionIterator It(GetSelectedActorIterator()); It; ++It)
 		{
-			ActorsToDeselect.Add(Actor);
+			AActor* Actor = static_cast<AActor*>( *It );
+			checkSlow(Actor->IsA(AActor::StaticClass()));
+
+			// Deselect any selected builder brushes.
+			ABrush* Brush = Cast< ABrush >(Actor);
+			const bool bActorIsBuilderBrush = ( Brush && FActorEditorUtils::IsABuilderBrush(Brush) );
+			if (bActorIsBuilderBrush)
+			{
+				ActorsToDeselect.Add(Actor);
+			}
+
+			// If any selected actors are not in the current level, warn the user that some actors will not be copied.
+			if (!bSomeSelectedActorsNotInCurrentLevel && !Actor->GetLevel()->IsCurrentLevel())
+			{
+				bSomeSelectedActorsNotInCurrentLevel = true;
+				FMessageDialog::Open(EAppMsgType::Ok, NSLOCTEXT("UnrealEd", "CopySelectedActorsInNonCurrentLevel", "Some selected actors are not in the current level and will not be copied."));
+			}
 		}
 
-		// If any selected actors are not in the current level, warn the user that some actors will not be copied.
-		if ( !bSomeSelectedActorsNotInCurrentLevel && !Actor->GetLevel()->IsCurrentLevel() )
+		const FScopedBusyCursor BusyCursor;
+		for (int32 ActorIndex = 0; ActorIndex < ActorsToDeselect.Num(); ++ActorIndex)
 		{
-			bSomeSelectedActorsNotInCurrentLevel = true;
-			FMessageDialog::Open( EAppMsgType::Ok, NSLOCTEXT("UnrealEd", "CopySelectedActorsInNonCurrentLevel", "Some selected actors are not in the current level and will not be copied.") );
+			AActor* Actor = ActorsToDeselect[ActorIndex];
+			GetSelectedActors()->Deselect(Actor);
 		}
-	}
 
-	const FScopedBusyCursor BusyCursor;
-	for( int32 ActorIndex = 0 ; ActorIndex < ActorsToDeselect.Num() ; ++ActorIndex )
-	{
-		AActor* Actor = ActorsToDeselect[ActorIndex];
-		GetSelectedActors()->Deselect( Actor );
-	}
-
-	// Export the actors.
-	FStringOutputDevice Ar;
-	const FSelectedActorExportObjectInnerContext Context;
-	UExporter::ExportToOutputDevice( &Context, InWorld, NULL, Ar, TEXT("copy"), 0, PPF_DeepCompareInstances | PPF_ExportsNotFullyQualified);
-	FPlatformMisc::ClipboardCopy( *Ar );
-	if( DestinationData )
-	{
-		*DestinationData = MoveTemp(Ar);
+		// Export the actors.
+		FStringOutputDevice Ar;
+		const FSelectedActorExportObjectInnerContext Context;
+		UExporter::ExportToOutputDevice(&Context, InWorld, NULL, Ar, TEXT("copy"), 0, PPF_DeepCompareInstances | PPF_ExportsNotFullyQualified);
+		FPlatformMisc::ClipboardCopy(*Ar);
+		if (DestinationData)
+		{
+			*DestinationData = MoveTemp(Ar);
+		}
 	}
 }
 
@@ -237,74 +254,103 @@ void UUnrealEdEngine::edactPasteSelected(UWorld* InWorld, bool bDuplicate, bool
 		return;
 	}
 	
+	if (GetSelectedComponentCount() > 0)
+	{
+		auto SelectedActor = CastChecked<AActor>(*GetSelectedActorIterator());
 
-	const FScopedBusyCursor BusyCursor;
+		TArray<UActorComponent*> PastedComponents;
+		FComponentEditorUtils::PasteComponents(PastedComponents, SelectedActor, SelectedActor->GetRootComponent());
 
-	// Create a location offset.
-	const FVector LocationOffset = CreateLocationOffset( bDuplicate, bOffsetLocations );
+		if (PastedComponents.Num() > 0)
+		{
+			// Make sure all the SCS trees have a chance to update
+			FLevelEditorModule& LevelEditor = FModuleManager::LoadModuleChecked<FLevelEditorModule>("LevelEditor");
+			LevelEditor.BroadcastComponentsEdited();
 
-	FCachedActorLabels ActorLabels(InWorld);
+			// Select the new clones
+			USelection* ComponentSelection = GetSelectedComponents();
+			ComponentSelection->Modify(false);
+			ComponentSelection->BeginBatchSelectOperation();
+			ComponentSelection->DeselectAll();
 
-	// Transact the current selection set.
-	USelection* SelectedActors = GetSelectedActors();
-	SelectedActors->Modify();
+			for (auto PastedComp : PastedComponents)
+			{
+				GEditor->SelectComponent(PastedComp, true, false);
+			}
 
-	// Get pasted text.
-	FString PasteString;
-	if( SourceData )
-	{
-		PasteString = *SourceData;
+			ComponentSelection->EndBatchSelectOperation(true);
+		}
 	}
 	else
 	{
-		FPlatformMisc::ClipboardPaste(PasteString);
-	}
-	const TCHAR* Paste = *PasteString;
+		const FScopedBusyCursor BusyCursor;
 
-	// Import the actors.
-	auto Factory = NewObject<ULevelFactory>();
-	Factory->FactoryCreateText( ULevel::StaticClass(), InWorld->GetCurrentLevel(), InWorld->GetCurrentLevel()->GetFName(), RF_Transactional, NULL, bDuplicate ? TEXT("move") : TEXT("paste"), Paste, Paste+FCString::Strlen(Paste), GWarn );
+		// Create a location offset.
+		const FVector LocationOffset = CreateLocationOffset(bDuplicate, bOffsetLocations);
 
-	// Fire ULevel::LevelDirtiedEvent when falling out of scope.
-	FScopedLevelDirtied			LevelDirtyCallback;
+		FCachedActorLabels ActorLabels(InWorld);
 
-	// Update the actors' locations and update the global list of visible layers.
-	for ( FSelectionIterator It( GetSelectedActorIterator() ) ; It ; ++It )
-	{
-		AActor* Actor = static_cast<AActor*>( *It );
-		checkSlow( Actor->IsA(AActor::StaticClass()) );
+		// Transact the current selection set.
+		USelection* SelectedActors = GetSelectedActors();
+		SelectedActors->Modify();
 
-		// We only want to offset the location if this actor is the root of a selected attachment hierarchy
-		// Offsetting children of an attachment hierarchy would cause them to drift away from the node they're attached to
-		// as the offset would effectively get applied twice
-		const AActor* const ParentActor = Actor->GetAttachParentActor();
-		const FVector& ActorLocationOffset = (ParentActor && ParentActor->IsSelected()) ? FVector::ZeroVector : LocationOffset;
+		// Get pasted text.
+		FString PasteString;
+		if (SourceData)
+		{
+			PasteString = *SourceData;
+		}
+		else
+		{
+			FPlatformMisc::ClipboardPaste(PasteString);
+		}
+		const TCHAR* Paste = *PasteString;
+
+		// Import the actors.
+		auto Factory = NewObject<ULevelFactory>();
+		Factory->FactoryCreateText(ULevel::StaticClass(), InWorld->GetCurrentLevel(), InWorld->GetCurrentLevel()->GetFName(), RF_Transactional, NULL, bDuplicate ? TEXT("move") : TEXT("paste"), Paste, Paste + FCString::Strlen(Paste), GWarn);
+
+		// Fire ULevel::LevelDirtiedEvent when falling out of scope.
+		FScopedLevelDirtied			LevelDirtyCallback;
 
-		// Offset the actor's location.
-		Actor->TeleportTo(Actor->GetActorLocation() + ActorLocationOffset, Actor->GetActorRotation(), false, true);
+		// Update the actors' locations and update the global list of visible layers.
+		for (FSelectionIterator It(GetSelectedActorIterator()); It; ++It)
+		{
+			AActor* Actor = static_cast<AActor*>( *It );
+			checkSlow(Actor->IsA(AActor::StaticClass()));
 
-		// Re-label duplicated actors so that labels become unique
-		GEditor->SetActorLabelUnique(Actor, Actor->GetActorLabel(), &ActorLabels);
-		ActorLabels.Add(Actor->GetActorLabel());
+			// We only want to offset the location if this actor is the root of a selected attachment hierarchy
+			// Offsetting children of an attachment hierarchy would cause them to drift away from the node they're attached to
+			// as the offset would effectively get applied twice
+			const AActor* const ParentActor = Actor->GetAttachParentActor();
+			const FVector& ActorLocationOffset = ( ParentActor && ParentActor->IsSelected() ) ? FVector::ZeroVector : LocationOffset;
 
-		GEditor->Layers->InitializeNewActorLayers( Actor );
+			// Offset the actor's location.
+			Actor->TeleportTo(Actor->GetActorLocation() + ActorLocationOffset, Actor->GetActorRotation(), false, true);
+
+			// Re-label duplicated actors so that labels become unique
+			GEditor->SetActorLabelUnique(Actor, Actor->GetActorLabel(), &ActorLabels);
+			ActorLabels.Add(Actor->GetActorLabel());
+
+			GEditor->Layers->InitializeNewActorLayers(Actor);
 
 			// Ensure any layers this actor belongs to are visible
-		GEditor->Layers->SetLayersVisibility( Actor->Layers, true );
+			GEditor->Layers->SetLayersVisibility(Actor->Layers, true);
 
-		Actor->CheckDefaultSubobjects();
-		Actor->InvalidateLightingCache();
-		// Call PostEditMove to update components, etc.
-		Actor->PostEditMove( true );
-		Actor->PostDuplicate(false);
-		Actor->CheckDefaultSubobjects();
+			Actor->CheckDefaultSubobjects();
+			Actor->InvalidateLightingCache();
+			// Call PostEditMove to update components, etc.
+			Actor->PostEditMove(true);
+			Actor->PostDuplicate(false);
+			Actor->CheckDefaultSubobjects();
 
-		// Request saves/refreshes.
-		Actor->MarkPackageDirty();
-		LevelDirtyCallback.Request();
+			// Request saves/refreshes.
+			Actor->MarkPackageDirty();
+			LevelDirtyCallback.Request();
+		}
+		// Note the selection change.  This will also redraw level viewports and update the pivot.
+		NoteSelectionChange();
 	}
-	// Note the selection change.  This will also redraw level viewports and update the pivot.
-	NoteSelectionChange();
 }
 
 namespace DuplicateSelectedActors {
@@ -376,103 +422,146 @@ public:
 
 void UUnrealEdEngine::edactDuplicateSelected( ULevel* InLevel, bool bOffsetLocations )
 {
-	using namespace DuplicateSelectedActors;
+	auto NumSelectedComponents = GetSelectedComponentCount();
+	if (NumSelectedComponents > 0)
+	{
+		TArray<UActorComponent*> NewComponentClones;
+		NewComponentClones.Reserve(NumSelectedComponents);
 
-	const FScopedBusyCursor BusyCursor;
-	GetSelectedActors()->Modify();
+		// Duplicate selected components if they are an Instance component
+		for (FSelectedEditableComponentIterator It(GetSelectedEditableComponentIterator()); It; ++It)
+		{
+			auto Component = CastChecked<UActorComponent>(*It);
+			if (Component->CreationMethod == EComponentCreationMethod::Instance)
+			{
+				UActorComponent* Clone = FComponentEditorUtils::DuplicateComponent(Component);
+				if (Clone)
+				{
+					NewComponentClones.Add(Clone);
+				}
+			}
+		}
 
-	// Create per-level job lists.
-	typedef TMap<ULevel*, FDuplicateJob*>	DuplicateJobMap;
-	DuplicateJobMap							DuplicateJobs;
+		if (NewComponentClones.Num() > 0)
+		{
+			// Make sure all the SCS trees have a chance to update
+			FLevelEditorModule& LevelEditor = FModuleManager::LoadModuleChecked<FLevelEditorModule>("LevelEditor");
+			LevelEditor.BroadcastComponentsEdited();
 
-	// Build set of selected actors before duplication
-	TArray<AActor*> PreDuplicateSelection;
+			// Select the new clones
+			USelection* ComponentSelection = GetSelectedComponents();
+			ComponentSelection->Modify(false);
+			ComponentSelection->BeginBatchSelectOperation();
+			ComponentSelection->DeselectAll();
 
-	// Add selected actors to the per-level job lists.
-	bool bHaveActorLocation = false;
-	FVector AnyActorLocation = FVector::ZeroVector;
-	for ( FSelectionIterator It( GetSelectedActorIterator() ) ; It ; ++It )
+			for (auto Clone : NewComponentClones)
+			{
+				GEditor->SelectComponent(Clone, true, false);
+			}
+
+			ComponentSelection->EndBatchSelectOperation(true);
+		}
+	}
+	else
 	{
-		AActor* Actor = static_cast<AActor*>( *It );
-		checkSlow( Actor->IsA(AActor::StaticClass()) );
+		using namespace DuplicateSelectedActors;
 
-		if( !bHaveActorLocation )
+		const FScopedBusyCursor BusyCursor;
+		GetSelectedActors()->Modify();
+
+		// Create per-level job lists.
+		typedef TMap<ULevel*, FDuplicateJob*>	DuplicateJobMap;
+		DuplicateJobMap							DuplicateJobs;
+
+		// Build set of selected actors before duplication
+		TArray<AActor*> PreDuplicateSelection;
+
+		// Add selected actors to the per-level job lists.
+		bool bHaveActorLocation = false;
+		FVector AnyActorLocation = FVector::ZeroVector;
+		for (FSelectionIterator It(GetSelectedActorIterator()); It; ++It)
 		{
-			bHaveActorLocation = true;
-			AnyActorLocation = Actor->GetActorLocation();
+			AActor* Actor = static_cast<AActor*>( *It );
+			checkSlow(Actor->IsA(AActor::StaticClass()));
+
+			if (!bHaveActorLocation)
+			{
+				bHaveActorLocation = true;
+				AnyActorLocation = Actor->GetActorLocation();
+			}
+
+			PreDuplicateSelection.Add(Actor);
+
+			ULevel* OldLevel = Actor->GetLevel();
+			FDuplicateJob** Job = DuplicateJobs.Find(OldLevel);
+			if (Job)
+			{
+				( *Job )->Actors.Add(Actor);
+			}
+			else
+			{
+				// Allocate a new job for the level.
+				FDuplicateJob* NewJob = new FDuplicateJob;
+				NewJob->SrcLevel = OldLevel;
+				NewJob->Actors.Add(Actor);
+				DuplicateJobs.Add(OldLevel, NewJob);
+			}
 		}
 
-		PreDuplicateSelection.Add(Actor);
+		UWorld* World = InLevel->OwningWorld;
+		ULevel* DesiredLevel = InLevel;
 
-		ULevel* OldLevel = Actor->GetLevel();
-		FDuplicateJob** Job = DuplicateJobs.Find( OldLevel );
-		if ( Job )
+		USelection* SelectedActors = GetSelectedActors();
+		SelectedActors->BeginBatchSelectOperation();
+		SelectedActors->Modify();
+
+		// For each level, select the actors in that level and copy-paste into the destination level.
+		TArray<AActor*>	NewActors;
+		for (DuplicateJobMap::TIterator It(DuplicateJobs); It; ++It)
 		{
-			(*Job)->Actors.Add( Actor );
+			FDuplicateJob* Job = It.Value();
+			check(Job);
+			Job->DuplicateActorsToLevel(NewActors, InLevel, bOffsetLocations);
 		}
-		else
+
+		// Select any newly created actors and prefabs.
+		SelectNone(false, true);
+		for (int32 ActorIndex = 0; ActorIndex < NewActors.Num(); ++ActorIndex)
 		{
-			// Allocate a new job for the level.
-			FDuplicateJob* NewJob = new FDuplicateJob;
-			NewJob->SrcLevel = OldLevel;
-			NewJob->Actors.Add( Actor );
-			DuplicateJobs.Add( OldLevel, NewJob );
+			AActor* Actor = NewActors[ActorIndex];
+			SelectActor(Actor, true, false);
 		}
-	}
-
-	UWorld* World = InLevel->OwningWorld;
-	ULevel* DesiredLevel = InLevel;
-
-	USelection* SelectedActors = GetSelectedActors();
-	SelectedActors->BeginBatchSelectOperation();
-	SelectedActors->Modify();
+		SelectedActors->EndBatchSelectOperation();
+		NoteSelectionChange();
 
-	// For each level, select the actors in that level and copy-paste into the destination level.
-	TArray<AActor*>	NewActors;
-	for ( DuplicateJobMap::TIterator It( DuplicateJobs ) ; It ; ++It )
-	{
-		FDuplicateJob* Job = It.Value();
-		check( Job );
-		Job->DuplicateActorsToLevel( NewActors, InLevel, bOffsetLocations );
-	}
+		// Finally, cleanup.
+		for (DuplicateJobMap::TIterator It(DuplicateJobs); It; ++It)
+		{
+			FDuplicateJob* Job = It.Value();
+			delete Job;
+		}
 
-	// Select any newly created actors and prefabs.
-	SelectNone( false, true );
-	for ( int32 ActorIndex = 0 ; ActorIndex < NewActors.Num() ; ++ActorIndex )
-	{
-		AActor* Actor = NewActors[ ActorIndex ];
-		SelectActor( Actor, true, false );
-	}
-	SelectedActors->EndBatchSelectOperation();
-	NoteSelectionChange();
+		// Build set of selected actors after duplication
+		TArray<AActor*> PostDuplicateSelection;
+		for (FSelectionIterator It(GetSelectedActorIterator()); It; ++It)
+		{
+			AActor* Actor = static_cast<AActor*>( *It );
+			checkSlow(Actor->IsA(AActor::StaticClass()));
 
-	// Finally, cleanup.
-	for ( DuplicateJobMap::TIterator It( DuplicateJobs ) ; It ; ++It )
-	{
-		FDuplicateJob* Job = It.Value();
-		delete Job;
-	}
-	
-	// Build set of selected actors after duplication
-	TArray<AActor*> PostDuplicateSelection;
-	for ( FSelectionIterator It( GetSelectedActorIterator() ) ; It ; ++It )
-	{
-		AActor* Actor = static_cast<AActor*>( *It );
-		checkSlow( Actor->IsA(AActor::StaticClass()) );
+			// We generate new seeds when we duplicate
+			Actor->SeedAllRandomStreams();
 
-		// We generate new seeds when we duplicate
-		Actor->SeedAllRandomStreams();
+			PostDuplicateSelection.Add(Actor);
+		}
 
-		PostDuplicateSelection.Add(Actor);
-	}
-	
-	TArray<FEdMode*> ActiveModes;
-	GLevelEditorModeTools().GetActiveModes( ActiveModes );
+		TArray<FEdMode*> ActiveModes;
+		GLevelEditorModeTools().GetActiveModes(ActiveModes);
 
-	for( int32 ModeIndex = 0; ModeIndex < ActiveModes.Num(); ++ModeIndex )
-	{
-		// Tell the tools about the duplication
-		ActiveModes[ModeIndex]->ActorsDuplicatedNotify(PreDuplicateSelection, PostDuplicateSelection, bOffsetLocations);
+		for (int32 ModeIndex = 0; ModeIndex < ActiveModes.Num(); ++ModeIndex)
+		{
+			// Tell the tools about the duplication
+			ActiveModes[ModeIndex]->ActorsDuplicatedNotify(PreDuplicateSelection, PostDuplicateSelection, bOffsetLocations);
+		}
 	}
 }
 
@@ -548,29 +637,55 @@ bool UUnrealEdEngine::edactDeleteSelected( UWorld* InWorld, bool bVerifyDeletion
 
 	const double StartSeconds = FPlatformTime::Seconds();
 
-	GetSelectedActors()->Modify();
-
 	if (GetSelectedComponentCount() > 0)
 	{
-		bool bComponentWasDeleted = false;
-		// Delete selected components if they are an Instance component
+		TArray<UActorComponent*> SelectedEditableComponents;
 		for (FSelectedEditableComponentIterator It(GetSelectedEditableComponentIterator()); It; ++It)
 		{
-			UActorComponent* CurActorComp = CastChecked<UActorComponent>(*It);
-			if (CurActorComp->CreationMethod == EComponentCreationMethod::Instance)
-			{
-				CurActorComp->Modify();
-				CurActorComp->DestroyComponent(true);
-				bComponentWasDeleted = true;
-			}
+			SelectedEditableComponents.Add(CastChecked<UActorComponent>(*It));
 		}
-		if (bComponentWasDeleted)
+
+		if (SelectedEditableComponents.Num() > 0)
 		{
-			NoteSelectionChange();
+			// Modify the actor that owns the selected components
+			check(GetSelectedActorCount() == 1);
+			(*GetSelectedActorIterator())->Modify();
+			
+			// Delete the selected components
+			UActorComponent* ComponentToSelect = nullptr;
+			int32 NumDeletedComponents = FComponentEditorUtils::DeleteComponents(SelectedEditableComponents, ComponentToSelect);
+
+			if (NumDeletedComponents > 0)
+			{
+				// Make sure all the SCS trees have a chance to rebuild
+				FLevelEditorModule& LevelEditor = FModuleManager::LoadModuleChecked<FLevelEditorModule>("LevelEditor");
+				LevelEditor.BroadcastComponentsEdited();
+
+				// Update the editor component selection if possible
+				if (ComponentToSelect)
+				{
+					USelection* ComponentSelection = GetSelectedComponents();
+					ComponentSelection->Modify(false);
+					ComponentSelection->BeginBatchSelectOperation();
+					ComponentSelection->DeselectAll();
+
+					GEditor->SelectComponent(ComponentToSelect, true, false);
+
+					// Make sure the selection changed event fires so the SCS trees can update their selection
+					ComponentSelection->MarkBatchDirty();
+					ComponentSelection->EndBatchSelectOperation(true);
+				}
+
+				UE_LOG(LogEditorActor, Log, TEXT("Deleted %d Components (%3.3f secs)"), NumDeletedComponents, FPlatformTime::Seconds() - StartSeconds);
+				return true;
+			}
 		}
-		return bComponentWasDeleted;
+
+		return false;
 	}
 
+	GetSelectedActors()->Modify();
+
 	// Fire ULevel::LevelDirtiedEvent when falling out of scope.
 	FScopedLevelDirtied			LevelDirtyCallback;
 
diff --git a/Engine/Source/Editor/UnrealEd/Private/Kismet2/ComponentEditorUtils.cpp b/Engine/Source/Editor/UnrealEd/Private/Kismet2/ComponentEditorUtils.cpp
--- a/Engine/Source/Editor/UnrealEd/Private/Kismet2/ComponentEditorUtils.cpp
+++ b/Engine/Source/Editor/UnrealEd/Private/Kismet2/ComponentEditorUtils.cpp
@@ -8,6 +8,102 @@
 #include "Engine/SimpleConstructionScript.h"
 #include "ScopedTransaction.h"
 #include "BlueprintEditorUtils.h"
+#include "Factories.h"
+#include "UnrealExporter.h"
+
+// Text object factory for pasting components
+struct FComponentObjectTextFactory : public FCustomizableTextObjectFactory
+{
+	// Child->Parent name map
+	TMap<FName, FName> ParentMap;
+
+	// Name->Instance object mapping
+	TMap<FName, UActorComponent*> NewObjectMap;
+
+	// Determine whether or not scene components in the new object set can be attached to the given scene root component
+	bool CanAttachComponentsTo(USceneComponent* InRootComponent)
+	{
+		check(InRootComponent);
+
+		// For each component in the set, check against the given root component and break if we fail to validate
+		bool bCanAttachToRoot = true;
+		for (auto NewComponentIt = NewObjectMap.CreateConstIterator(); NewComponentIt && bCanAttachToRoot; ++NewComponentIt)
+		{
+			// If this is a scene component, and it does not already have a parent within the set
+			USceneComponent* SceneComponent = Cast<USceneComponent>(NewComponentIt->Value);
+			if (SceneComponent != NULL && !ParentMap.Contains(SceneComponent->GetFName()))
+			{
+				// Determine if we are allowed to attach the scene component to the given root component
+				bCanAttachToRoot = InRootComponent->CanAttachAsChild(SceneComponent, NAME_None)
+					&& SceneComponent->Mobility >= InRootComponent->Mobility
+					&& ( !InRootComponent->IsEditorOnly() || SceneComponent->IsEditorOnly() );
+			}
+		}
+
+		return bCanAttachToRoot;
+	}
+
+	// Constructs a new object factory from the given text buffer
+	static TSharedRef<FComponentObjectTextFactory> Get(const FString& InTextBuffer, bool bPasteAsArchetypes = false)
+	{
+		// Construct a new instance
+		TSharedPtr<FComponentObjectTextFactory> FactoryPtr = MakeShareable(new FComponentObjectTextFactory());
+		check(FactoryPtr.IsValid());
+
+		// Create new objects if we're allowed to
+		if (FactoryPtr->CanCreateObjectsFromText(InTextBuffer))
+		{
+			EObjectFlags ObjectFlags = RF_Transactional;
+			if (bPasteAsArchetypes)
+			{
+				ObjectFlags |= RF_ArchetypeObject;
+			}
+
+			// Use the transient package initially for creating the objects, since the variable name is used when copying
+			FactoryPtr->ProcessBuffer(GetTransientPackage(), ObjectFlags, InTextBuffer);
+		}
+
+		return FactoryPtr.ToSharedRef();
+	}
+
+	virtual ~FComponentObjectTextFactory() {}
+
+protected:
+	// Constructor; protected to only allow this class to instance itself
+	FComponentObjectTextFactory()
+		:FCustomizableTextObjectFactory(GWarn)
+	{
+	}
+
+	// FCustomizableTextObjectFactory implementation
+
+	virtual bool CanCreateClass(UClass* ObjectClass) const override
+	{
+		// Only allow actor component types to be created
+		return ObjectClass->IsChildOf(UActorComponent::StaticClass());
+	}
+
+	virtual void ProcessConstructedObject(UObject* NewObject) override
+	{
+		check(NewObject);
+
+		// Add it to the new object map
+		NewObjectMap.Add(NewObject->GetFName(), Cast<UActorComponent>(NewObject));
+
+		// If this is a scene component and it has a parent
+		USceneComponent* SceneComponent = Cast<USceneComponent>(NewObject);
+		if (SceneComponent && SceneComponent->AttachParent)
+		{
+			// Add an entry to the child->parent name map
+			ParentMap.Add(NewObject->GetFName(), SceneComponent->AttachParent->GetFName());
+
+			// Clear this so it isn't used when constructing the new SCS node
+			SceneComponent->AttachParent = NULL;
+		}
+	}
+
+	// FCustomizableTextObjectFactory (end)
+};
 
 USceneComponent* FComponentEditorUtils::GetSceneComponent( UObject* Object, UObject* SubObject /*= NULL*/ )
 {
@@ -161,6 +257,295 @@ FString FComponentEditorUtils::GenerateValidVariableNameFromAsset(UObject* Asset
 	return ComponentInstanceName;
 }
 
+USceneComponent* FComponentEditorUtils::FindClosestParentInList(UActorComponent* ChildComponent, const TArray<UActorComponent*>& ComponentList)
+{
+	USceneComponent* ClosestParentComponent = nullptr;
+	for (auto Component : ComponentList)
+	{
+		auto ChildAsScene = CastChecked<USceneComponent>(ChildComponent);
+		auto SceneComponent = CastChecked<USceneComponent>(Component);
+		if (ChildAsScene && SceneComponent)
+		{
+			// Check to see if any parent is also in the list
+			USceneComponent* Parent = ChildAsScene->GetAttachParent();
+			while (Parent != nullptr)
+			{
+				if (ComponentList.Contains(Parent))
+				{
+					ClosestParentComponent = SceneComponent;
+					break;
+				}
+
+				Parent = Parent->GetAttachParent();
+			}
+		}
+	}
+
+	return ClosestParentComponent;
+}
+
+void FComponentEditorUtils::CopyComponents(const TArray<UActorComponent*>& ComponentsToCopy)
+{
+	FStringOutputDevice Archive;
+	const FExportObjectInnerContext Context;
+
+	// Clear the mark state for saving.
+	UnMarkAllObjects(EObjectMark(OBJECTMARK_TagExp | OBJECTMARK_TagImp));
+
+	// Duplicate the selected component templates into temporary objects that we can modify
+	TMap<FName, FName> ParentMap;
+	TMap<FName, UActorComponent*> ObjectMap;
+	for (auto Component : ComponentsToCopy)
+	{
+		// Duplicate the component into a temporary object
+		UObject* DuplicatedComponent = StaticDuplicateObject(Component, GetTransientPackage(), *Component->GetName(), RF_AllFlags & ~RF_ArchetypeObject);
+		if (DuplicatedComponent)
+		{
+			// Find the closest parent component of the current component within the list of components to copy
+			USceneComponent* ClosestSelectedParent = FindClosestParentInList(Component, ComponentsToCopy);
+			if (ClosestSelectedParent)
+			{
+				// If the parent is included in the list, record it into the node->parent map
+				ParentMap.Add(Component->GetFName(), ClosestSelectedParent->GetFName());
+			}
+
+			// Record the temporary object into the name->object map
+			ObjectMap.Add(Component->GetFName(), CastChecked<UActorComponent>(DuplicatedComponent));
+		}
+	}
+
+	// Export the component object(s) to text for copying
+	for (auto ObjectIt = ObjectMap.CreateIterator(); ObjectIt; ++ObjectIt)
+	{
+		// Get the component object to be copied
+		UActorComponent* ComponentToCopy = ObjectIt->Value;
+		check(ComponentToCopy);
+
+		// If this component object had a parent within the selected set
+		if (ParentMap.Contains(ComponentToCopy->GetFName()))
+		{
+			// Get the name of the parent component
+			FName ParentName = ParentMap[ComponentToCopy->GetFName()];
+			if (ObjectMap.Contains(ParentName))
+			{
+				// Ensure that this component is a scene component
+				USceneComponent* SceneComponent = Cast<USceneComponent>(ComponentToCopy);
+				if (SceneComponent)
+				{
+					// Set the attach parent to the matching parent object in the temporary set. This allows us to preserve hierarchy in the copied set.
+					SceneComponent->AttachParent = Cast<USceneComponent>(ObjectMap[ParentName]);
+				}
+			}
+		}
+
+		// Export the component object to the given string
+		UExporter::ExportToOutputDevice(&Context, ComponentToCopy, NULL, Archive, TEXT("copy"), 0, PPF_ExportsNotFullyQualified | PPF_Copy | PPF_Delimited, false, ComponentToCopy->GetOuter());
+	}
+
+	// Copy text to clipboard
+	FString ExportedText = Archive;
+	FPlatformMisc::ClipboardCopy(*ExportedText);
+}
+
+bool FComponentEditorUtils::CanPasteComponents(USceneComponent* RootComponent, bool bOverrideCanAttach)
+{
+	FString ClipboardContent;
+	FPlatformMisc::ClipboardPaste(ClipboardContent);
+
+	// Obtain the component object text factory for the clipboard content and return whether or not we can use it
+	TSharedRef<FComponentObjectTextFactory> Factory = FComponentObjectTextFactory::Get(ClipboardContent);
+	return Factory->NewObjectMap.Num() > 0 && ( bOverrideCanAttach || Factory->CanAttachComponentsTo(RootComponent) );
+}
+
+void FComponentEditorUtils::PasteComponents(TArray<UActorComponent*>& OutPastedComponents, AActor* TargetActor, USceneComponent* TargetComponent)
+{
+	check(TargetActor);
+
+	// Get the text from the clipboard
+	FString TextToImport;
+	FPlatformMisc::ClipboardPaste(TextToImport);
+
+	// Get a new component object factory for the clipboard content
+	TSharedRef<FComponentObjectTextFactory> Factory = FComponentObjectTextFactory::Get(TextToImport);
+
+	TargetActor->Modify();
+
+	USceneComponent* TargetParent = TargetComponent ? TargetComponent->GetAttachParent() : nullptr;
+	for (auto& NewObjectPair : Factory->NewObjectMap)
+	{
+		// Get the component object instance
+		UActorComponent* NewActorComponent = NewObjectPair.Value;
+		check(NewActorComponent);
+
+		// Relocate the instance from the transient package to the Actor and assign it a unique object name
+		FString NewComponentName = FComponentEditorUtils::GenerateValidVariableName(NewActorComponent->GetClass(), TargetActor);
+		NewActorComponent->Rename(*NewComponentName, TargetActor, REN_DontCreateRedirectors | REN_DoNotDirty);
+
+		if (auto NewSceneComponent = Cast<USceneComponent>(NewActorComponent))
+		{
+			// Default to attaching to the target component's parent if possible, otherwise attach to the root
+			USceneComponent* NewComponentParent = TargetParent ? TargetParent : TargetActor->GetRootComponent();
+			
+			// Check to see if there's an entry for the current component in the set of parent components
+			if (Factory->ParentMap.Contains(NewObjectPair.Key))
+			{
+				// Get the parent component name
+				FName ParentName = Factory->ParentMap[NewObjectPair.Key];
+				if (Factory->NewObjectMap.Contains(ParentName))
+				{
+					// The parent should by definition be a scene component
+					NewComponentParent = CastChecked<USceneComponent>(Factory->NewObjectMap[ParentName]);
+				}
+			}
+
+			//So, if we're pasting a component that was the root, we want to attach keeping the world location
+			// Otherwise, we want to attach keeping the relative transform
+			// So how can we tell if a pasted component was a root?
+
+			//NewSceneComponent->UpdateComponentToWorld();
+			if (NewComponentParent)
+			{
+				// Reattach the current node to the parent node
+				NewSceneComponent->AttachTo(NewComponentParent, NAME_None/*, EAttachLocation::KeepWorldPosition*/);
+			}
+			else
+			{
+				// There is no root component and this component isn't the child of another component in the map, so make it the root
+				TargetActor->SetRootComponent(NewSceneComponent);
+			}
+		}
+
+		TargetActor->AddInstanceComponent(NewActorComponent);
+		NewActorComponent->RegisterComponent();
+
+		OutPastedComponents.Add(NewActorComponent);
+	}
+}
+
+void FComponentEditorUtils::GetComponentsFromClipboard(TMap<FName, FName>& OutParentMap, TMap<FName, UActorComponent*>& OutNewObjectMap, bool bGetComponentsAsArchetypes)
+{
+	// Get the text from the clipboard
+	FString TextToImport;
+	FPlatformMisc::ClipboardPaste(TextToImport);
+
+	// Get a new component object factory for the clipboard content
+	TSharedRef<FComponentObjectTextFactory> Factory = FComponentObjectTextFactory::Get(TextToImport, bGetComponentsAsArchetypes);
+
+	// Return the created component mappings
+	OutParentMap = Factory->ParentMap;
+	OutNewObjectMap = Factory->NewObjectMap;
+}
+
+int32 FComponentEditorUtils::DeleteComponents(const TArray<UActorComponent*>& ComponentsToDelete, UActorComponent*& OutComponentToSelect)
+{
+	int32 NumDeletedComponents = 0;
+
+	for (auto ComponentToDelete : ComponentsToDelete)
+	{
+		if (ComponentToDelete->CreationMethod != EComponentCreationMethod::Instance)
+		{
+			// We can only delete instance components, so retain selection on the un-deletable component
+			OutComponentToSelect = ComponentToDelete;
+			continue;
+		}
+
+		// If necessary, determine the component that should be selected following the deletion of the indicated component
+		if (!OutComponentToSelect || ComponentToDelete == OutComponentToSelect)
+		{
+			USceneComponent* RootComponent = ComponentToDelete->GetOwner()->GetRootComponent();
+			if (RootComponent != ComponentToDelete)
+			{
+				// Worst-case, the root can be selected
+				OutComponentToSelect = RootComponent;
+
+				if (auto ComponentToDeleteAsSceneComp = Cast<USceneComponent>(ComponentToDelete))
+				{
+					if (USceneComponent* ParentComponent = ComponentToDeleteAsSceneComp->GetAttachParent())
+					{
+						// The component to delete has a parent, so we select that in the absence of an appropriate sibling
+						OutComponentToSelect = ParentComponent;
+
+						// Try to select the sibling that immediately precedes the component to delete
+						TArray<USceneComponent*> Siblings;
+						ParentComponent->GetChildrenComponents(false, Siblings);
+						for (int32 i = 0; i < Siblings.Num() && ComponentToDelete != Siblings[i]; ++i)
+						{
+							if (!Siblings[i]->IsPendingKill())
+							{
+								OutComponentToSelect = Siblings[i];
+							}
+						}
+					}
+				}
+				else
+				{
+					// For a non-scene component, try to select the preceding non-scene component
+					TInlineComponentArray<UActorComponent*> ActorComponents;
+					ComponentToDelete->GetOwner()->GetComponents(ActorComponents);
+					for (int32 i = 0; i < ActorComponents.Num() && ComponentToDelete != ActorComponents[i]; ++i)
+					{
+						if (!ActorComponents[i]->IsA(USceneComponent::StaticClass()))
+						{
+							OutComponentToSelect = ActorComponents[i];
+						}
+					}
+				}
+			}
+			else
+			{
+				OutComponentToSelect = nullptr;
+			}
+		}
+
+		// Actually delete the component
+		ComponentToDelete->Modify();
+		ComponentToDelete->DestroyComponent(true);
+		NumDeletedComponents++;
+	}
+
+	return NumDeletedComponents;
+}
+
+UActorComponent* FComponentEditorUtils::DuplicateComponent(UActorComponent* TemplateComponent)
+{
+	check(TemplateComponent);
+
+	UActorComponent* NewCloneComponent = nullptr;
+	AActor* Actor = TemplateComponent->GetOwner();
+	if (!TemplateComponent->IsEditorOnly() && Actor)
+	{
+		Actor->Modify();
+		UClass* ComponentClass = TemplateComponent->GetClass();
+		FName NewComponentName = *FComponentEditorUtils::GenerateValidVariableName(ComponentClass, Actor);
+
+		bool bKeepWorldLocationOnAttach = false;
+		NewCloneComponent = ConstructObject<UActorComponent>(ComponentClass, Actor, NewComponentName, RF_Transactional, TemplateComponent);
+		
+		// ComponentToWorld is not a UPROPERTY, so make sure the clone has calculated it properly
+		NewCloneComponent->UpdateComponentToWorld();
+
+		// If the clone is a scene component without an attach parent, attach it to the root (can happen when duplicating the root component)
+		auto NewSceneComponent = Cast<USceneComponent>(NewCloneComponent);
+		if (NewSceneComponent && !NewSceneComponent->GetAttachParent())
+		{
+			USceneComponent* RootComponent = Actor->GetRootComponent();
+
+			// There should be no situation in which a scene component is duplicated when a root doesn't exist
+			check(RootComponent);
+
+			NewSceneComponent->AttachTo(RootComponent, NAME_None, EAttachLocation::KeepWorldPosition);
+		}
+
+		// Add to SerializedComponents array so it gets saved
+		Actor->AddInstanceComponent(NewCloneComponent);
+		
+		// Register the new component
+		NewCloneComponent->RegisterComponent();
+	}
+
+	return NewCloneComponent;
+}
+
 void FComponentEditorUtils::AdjustComponentDelta(USceneComponent* Component, FVector& Drag, FRotator& Rotation)
 {
 	USceneComponent* ParentSceneComp = Component->GetAttachParent();
diff --git a/Engine/Source/Editor/UnrealEd/Private/LevelEditorViewport.cpp b/Engine/Source/Editor/UnrealEd/Private/LevelEditorViewport.cpp
--- a/Engine/Source/Editor/UnrealEd/Private/LevelEditorViewport.cpp
+++ b/Engine/Source/Editor/UnrealEd/Private/LevelEditorViewport.cpp
@@ -1480,7 +1480,7 @@ FLevelEditorViewportClient::FLevelEditorViewportClient(const TSharedPtr<SLevelVi
 	, bEnableColorScaling(false)
 	, bEditorCameraCut(false)
 	, bDrawBaseInfo(false)
-	, bDuplicateActorsOnNextDrag( false )
+	, bDuplicateOnNextDrag( false )
 	, bDuplicateActorsInProgress( false )
 	, bIsTrackingBrushModification( false )
 	, bLockedCameraView(true)
@@ -2178,16 +2178,15 @@ bool FLevelEditorViewportClient::InputWidgetDelta(FViewport* Viewport, EAxisList
 				// If duplicate dragging . . .
 				if ( IsAltPressed() && (LeftMouseButtonDown || RightMouseButtonDown) )
 				{
-					// The widget has been offset, so check if we should duplicate actors.
-					if ( bDuplicateActorsOnNextDrag )
+					// The widget has been offset, so check if we should duplicate the selection.
+					if ( bDuplicateOnNextDrag )
 					{
 						// Only duplicate if we're translating or rotating.
 						if ( !Drag.IsNearlyZero() || !Rot.IsZero() )
 						{
-							// Actors haven't been dragged since ALT+LMB went down.
-							bDuplicateActorsOnNextDrag = false;
-
-							GEditor->edactDuplicateSelected( GetWorld()->GetCurrentLevel(), false );
+							// Widget hasn't been dragged since ALT+LMB went down.
+							bDuplicateOnNextDrag = false;
+							GEditor->edactDuplicateSelected(GetWorld()->GetCurrentLevel(), false);
 						}
 					}
 				}
@@ -2421,13 +2420,13 @@ void FLevelEditorViewportClient::TrackingStarted( const FInputEventState& InInpu
 			if(Event == IE_Pressed && (Key == EKeys::LeftMouseButton || Key == EKeys::RightMouseButton) && !bDuplicateActorsInProgress)
 			{
 				// Set the flag so that the actors actors will be duplicated as soon as the widget is displaced.
-				bDuplicateActorsOnNextDrag = true;
+				bDuplicateOnNextDrag = true;
 				bDuplicateActorsInProgress = true;
 			}
 		}
 		else
 		{
-			bDuplicateActorsOnNextDrag = false;
+			bDuplicateOnNextDrag = false;
 		}
 	}
 
@@ -2549,7 +2548,7 @@ void FLevelEditorViewportClient::TrackingStopped()
 	const bool MiddleMouseButtonDown = Viewport->KeyState(EKeys::MiddleMouseButton);
 
 	// Only disable the duplicate on next drag flag if we actually dragged the mouse.
-	bDuplicateActorsOnNextDrag = false;
+	bDuplicateOnNextDrag = false;
 
 	// here we check to see if anything of worth actually changed when ending our MouseMovement
 	// If the TransCount > 0 (we changed something of value) so we need to call PostEditMove() on stuff
diff --git a/Engine/Source/Editor/UnrealEd/Private/UnrealEdSrv.cpp b/Engine/Source/Editor/UnrealEd/Private/UnrealEdSrv.cpp
--- a/Engine/Source/Editor/UnrealEd/Private/UnrealEdSrv.cpp
+++ b/Engine/Source/Editor/UnrealEd/Private/UnrealEdSrv.cpp
@@ -49,6 +49,8 @@
 #include "Components/SphereComponent.h"
 #include "Engine/StaticMeshActor.h"
 #include "EngineUtils.h"
+#include "LevelEditor.h"
+#include "ComponentEditorUtils.h"
 
 
 DEFINE_LOG_CATEGORY_STATIC(LogUnrealEdSrv, Log, All);
@@ -1609,37 +1611,56 @@ bool UUnrealEdEngine::DoDirtyPackagesNeedCheckout() const
 
 bool UUnrealEdEngine::Exec_Edit( UWorld* InWorld, const TCHAR* Str, FOutputDevice& Ar )
 {
+	const bool bComponentsSelected = GetSelectedComponentCount() > 0;
+
 	if( FParse::Command(&Str,TEXT("CUT")) )
 	{
-		TArray<FEdMode*> ActiveModes; 
-		GLevelEditorModeTools().GetActiveModes( ActiveModes );
-		for( int32 ModeIndex = 0; ModeIndex < ActiveModes.Num(); ++ModeIndex )
+		TArray<FEdMode*> ActiveModes;
+		GLevelEditorModeTools().GetActiveModes(ActiveModes);
+		for (int32 ModeIndex = 0; ModeIndex < ActiveModes.Num(); ++ModeIndex)
 		{
 			if (ActiveModes[ModeIndex]->ProcessEditCut())
 			{
 				return true;
 			}
 		}
-		CopySelectedActorsToClipboard( InWorld, true );
+
+		if (bComponentsSelected)
+		{
+			edactCopySelected(InWorld);
+			edactDeleteSelected(InWorld);
+		}
+		else
+		{
+			CopySelectedActorsToClipboard(InWorld, true);
+		}
 	}
 	else if( FParse::Command(&Str,TEXT("COPY")) )
 	{
-		TArray<FEdMode*> ActiveModes; 
-		GLevelEditorModeTools().GetActiveModes( ActiveModes );
-		for( int32 ModeIndex = 0; ModeIndex < ActiveModes.Num(); ++ModeIndex )
+		TArray<FEdMode*> ActiveModes;
+		GLevelEditorModeTools().GetActiveModes(ActiveModes);
+		for (int32 ModeIndex = 0; ModeIndex < ActiveModes.Num(); ++ModeIndex)
 		{
 			if (ActiveModes[ModeIndex]->ProcessEditCopy())
 			{
 				return true;
 			}
 		}
-		CopySelectedActorsToClipboard( InWorld, false );
+
+		if (bComponentsSelected)
+		{
+			edactCopySelected(InWorld);
+		}
+		else
+		{
+			CopySelectedActorsToClipboard(InWorld, false);
+		}
 	}
 	else if( FParse::Command(&Str,TEXT("PASTE")) )
 	{
-		TArray<FEdMode*> ActiveModes; 
-		GLevelEditorModeTools().GetActiveModes( ActiveModes );
-		for( int32 ModeIndex = 0; ModeIndex < ActiveModes.Num(); ++ModeIndex )
+		TArray<FEdMode*> ActiveModes;
+		GLevelEditorModeTools().GetActiveModes(ActiveModes);
+		for (int32 ModeIndex = 0; ModeIndex < ActiveModes.Num(); ++ModeIndex)
 		{
 			if (ActiveModes[ModeIndex]->ProcessEditPaste())
 			{
@@ -1647,27 +1668,35 @@ bool UUnrealEdEngine::Exec_Edit( UWorld* InWorld, const TCHAR* Str, FOutputDevic
 			}
 		}
 
-		// How should this paste be handled
-		EPasteTo PasteTo = PT_OriginalLocation;
-		FText TransDescription = NSLOCTEXT("UnrealEd", "Paste", "Paste");
-		if( FParse::Value( Str, TEXT("TO="), TempStr, 15 ) )
+		if (bComponentsSelected)
 		{
-			if( !FCString::Strcmp( TempStr, TEXT("HERE") ) )
-			{
-				PasteTo = PT_Here;
-				TransDescription = NSLOCTEXT("UnrealEd", "PasteHere", "Paste Here");
-			}
-			else
+			const FScopedTransaction Transaction(NSLOCTEXT("UnrealEd", "PasteComponents", "Paste Components"));
+			edactPasteSelected(InWorld, false, false, true);
+		}
+		else
+		{
+			// How should this paste be handled
+			EPasteTo PasteTo = PT_OriginalLocation;
+			FText TransDescription = NSLOCTEXT("UnrealEd", "Paste", "Paste");
+			if (FParse::Value(Str, TEXT("TO="), TempStr, 15))
 			{
-				if( !FCString::Strcmp( TempStr, TEXT("ORIGIN") ) )
+				if (!FCString::Strcmp(TempStr, TEXT("HERE")))
+				{
+					PasteTo = PT_Here;
+					TransDescription = NSLOCTEXT("UnrealEd", "PasteHere", "Paste Here");
+				}
+				else
 				{
-					PasteTo = PT_WorldOrigin;
-					TransDescription = NSLOCTEXT("UnrealEd", "PasteToWorldOrigin", "Paste To World Origin");
+					if (!FCString::Strcmp(TempStr, TEXT("ORIGIN")))
+					{
+						PasteTo = PT_WorldOrigin;
+						TransDescription = NSLOCTEXT("UnrealEd", "PasteToWorldOrigin", "Paste To World Origin");
+					}
 				}
 			}
-		}
 
-		PasteSelectedActorsFromClipboard( InWorld, TransDescription, PasteTo );
+			PasteSelectedActorsFromClipboard(InWorld, TransDescription, PasteTo);
+		}
 	}
 
 	return false;
diff --git a/Engine/Source/Editor/UnrealEd/Public/Kismet2/ComponentEditorUtils.h b/Engine/Source/Editor/UnrealEd/Public/Kismet2/ComponentEditorUtils.h
--- a/Engine/Source/Editor/UnrealEd/Public/Kismet2/ComponentEditorUtils.h
+++ b/Engine/Source/Editor/UnrealEd/Public/Kismet2/ComponentEditorUtils.h
@@ -29,6 +29,52 @@ public:
 	/** Generate a valid variable name string for the given component instance based on the name of the asset referenced by the component */
 	static FString GenerateValidVariableNameFromAsset(UObject* Asset, AActor* ComponentOwner);
 
+	/**
+	 * Copies the selected components to the clipboard
+	 * @param ComponentsToCopy The list of components to copy
+	 */
+	static void CopyComponents(const TArray<UActorComponent*>& ComponentsToCopy);
+
+	/**
+	 * Determines whether the current contents of the clipboard contain paste-able component information
+	 * @param RootComponent The root component of the actor being pasted on
+	 * @param bOverrideCanAttach Optional override declaring that components can be attached and a check is not needed
+	 * @return Whether components can be pasted
+	 */
+	static bool CanPasteComponents(USceneComponent* RootComponent, bool bOverrideCanAttach = false);
+
+	/**
+	 * Attempts to paste components from the clipboard as siblings of the target component
+	 * @param OutPastedComponents List of all the components that were pasted
+	 * @param TargetActor The actor to attach the pasted components to
+	 * @param TargetComponent The component the paste is targeting (will attempt to paste components as siblings). If null, will attach pasted components to the root.
+	 */
+	static void PasteComponents(TArray<UActorComponent*>& OutPastedComponents, AActor* TargetActor, USceneComponent* TargetComponent = nullptr);
+
+	/**
+	 * Gets the copied components from the clipboard without attempting to paste/apply them in any way
+	 * @param OutParentMap Contains the child->parent name relationships of the copied components
+	 * @param OutNewObjectMap Contains the name->instance object mapping of the copied components
+	 */
+	static void GetComponentsFromClipboard(TMap<FName, FName>& OutParentMap, TMap<FName, UActorComponent*>& OutNewObjectMap, bool bGetComponentsAsArchetypes);
+
+	/**
+	 * Deletes the indicated components and identifies the component that should be selected following the operation.
+	 * Note: Does not take care of the actual selection of a new component. It only identifies which component should be selected.
+	 * 
+	 * @param ComponentsToDelete The list of components to delete
+	 * @param OutComponentToSelect The component that should be selected after the deletion
+	 * @return The number of components that were actually deleted
+	 */
+	static int32 DeleteComponents(const TArray<UActorComponent*>& ComponentsToDelete, UActorComponent*& OutComponentToSelect);
+
+	/** 
+	 * Duplicates a component instance and takes care of attachment and registration.
+	 * @param TemplateComponent The component to duplicate
+	 * @return The created clone of the provided TemplateComponent. nullptr if the duplication failed.
+	 */
+	static UActorComponent* DuplicateComponent(UActorComponent* TemplateComponent);
+
 	/**
 	 * Ensures that the selection override delegate is properly bound for the supplied component
 	* This includes any attached editor-only primitive components (such as billboard visualizers)
@@ -124,5 +170,6 @@ public:
 	static FName FindVariableNameGivenComponentInstance(UActorComponent* ComponentInstance);
 private:
 	static void PropagateTransformInner(class USceneComponent* InstancedSceneComponent, const FTransformData& OldDefaultTransform, const FTransformData& NewDefaultTransform, TSet<class USceneComponent*>& UpdatedComponents);
-
+	
+	static USceneComponent* FindClosestParentInList(UActorComponent* ChildComponent, const TArray<UActorComponent*>& ComponentList);
 };
\ No newline at end of file
diff --git a/Engine/Source/Editor/UnrealEd/Public/LevelEditorViewport.h b/Engine/Source/Editor/UnrealEd/Public/LevelEditorViewport.h
--- a/Engine/Source/Editor/UnrealEd/Public/LevelEditorViewport.h
+++ b/Engine/Source/Editor/UnrealEd/Public/LevelEditorViewport.h
@@ -648,10 +648,10 @@ public:
 	bool bDrawBaseInfo;
 
 	/**
-	 * Used for actor drag duplication.  Set to true on Alt+LMB so that the selected
-	 * actors will be duplicated as soon as the widget is displaced.
+	 * Used for drag duplication. Set to true on Alt+LMB so that the selected
+	 * objects (components or actors) will be duplicated as soon as the widget is displaced.
 	 */
-	bool					bDuplicateActorsOnNextDrag;
+	bool					bDuplicateOnNextDrag;
 
 	/**
 	* bDuplicateActorsOnNextDrag will not be set again while bDuplicateActorsInProgress is true.
diff --git a/Engine/Source/Runtime/Engine/Classes/Components/SceneComponent.h b/Engine/Source/Runtime/Engine/Classes/Components/SceneComponent.h
--- a/Engine/Source/Runtime/Engine/Classes/Components/SceneComponent.h
+++ b/Engine/Source/Runtime/Engine/Classes/Components/SceneComponent.h
@@ -382,11 +382,15 @@ public:
 	UFUNCTION(BlueprintCallable, Category="Components")
 	int32 GetNumChildrenComponents() const;
 
-	/** Gets the number of attached children components */
+	/** Gets the attached child component at the specified location */
 	UFUNCTION(BlueprintCallable, Category="Components")
 	class USceneComponent* GetChildComponent(int32 ChildIndex) const;
 
-	/** Gets the number of attached children components */
+	/** 
+	 * Gets all the attached child components
+	 * @param bIncludeAllDescendants Whether to include all descendants in the list of children (i.e. grandchildren, great grandchildren, etc.)
+	 * @param Children The list of attached child components
+	 */
 	UFUNCTION(BlueprintCallable, Category="Components")
 	void GetChildrenComponents(bool bIncludeAllDescendants, TArray<USceneComponent*>& Children) const;
 
diff --git a/Engine/Source/Runtime/Engine/Classes/Engine/Selection.h b/Engine/Source/Runtime/Engine/Classes/Engine/Selection.h
--- a/Engine/Source/Runtime/Engine/Classes/Engine/Selection.h
+++ b/Engine/Source/Runtime/Engine/Classes/Engine/Selection.h
@@ -472,7 +472,7 @@ public:
 };
 
 /**
- * An iterator used to iterate through selected components which can be editable (not created in a blueprint)
+ * An iterator used to iterate through selected components that are editable (i.e. not created in a blueprint)
  */
 class FSelectedEditableComponentIterator : public TSelectionIterator<FSelectedEditableComponentFilter>
 {
