Integrate CL2437278 to fix Persona crash caused by 2441284
Merging using UE4-Releases-4.7-To-UE4

diff --git a/Engine/Source/Editor/Kismet/Private/BlueprintEditor.cpp b/Engine/Source/Editor/Kismet/Private/BlueprintEditor.cpp
--- a/Engine/Source/Editor/Kismet/Private/BlueprintEditor.cpp
+++ b/Engine/Source/Editor/Kismet/Private/BlueprintEditor.cpp
@@ -2741,7 +2741,10 @@ void FBlueprintEditor::OnActiveTabChanged( TSharedPtr<SDockTab> PreviouslyActive
 		FocusedGraphEdPtr = nullptr;
 	}
 
-	MyBlueprintWidget->Refresh();
+	if (MyBlueprintWidget.IsValid() == true)
+	{
+		MyBlueprintWidget->Refresh();
+	}
 }
 
 void FBlueprintEditor::OnGraphEditorFocused(const TSharedRef<SGraphEditor>& InGraphEditor)
@@ -6342,7 +6345,10 @@ void FBlueprintEditor::OnAddNewVariable()
 	const FScopedTransaction Transaction( LOCTEXT("AddVariable", "Add Variable") );
 
 	// Reset MyBlueprint item filter so new variable is visible
-	MyBlueprintWidget->OnResetItemFilter();
+	if (MyBlueprintWidget.IsValid())
+	{
+		MyBlueprintWidget->OnResetItemFilter();
+	}
 
 	FName VarName = FBlueprintEditorUtils::FindUniqueKismetName(GetBlueprintObj(), TEXT("NewVar"));
 
@@ -6386,7 +6392,10 @@ void FBlueprintEditor::OnAddNewDelegate()
 	check(NULL != Blueprint);
 
 	// Reset MyBlueprint item filter so new variable is visible
-	MyBlueprintWidget->OnResetItemFilter();
+	if (MyBlueprintWidget.IsValid())
+	{
+		MyBlueprintWidget->OnResetItemFilter();
+	}
 
 	FName Name = FBlueprintEditorUtils::FindUniqueKismetName(GetBlueprintObj(), TEXT("NewEventDispatcher"));
 
