Integrate 2437251 To fix crash on opening anim blueprint caused by 2441284
Merging using UE4-Releases-4.7-To-UE4

diff --git a/Engine/Source/Editor/Kismet/Private/BlueprintEditor.cpp b/Engine/Source/Editor/Kismet/Private/BlueprintEditor.cpp
--- a/Engine/Source/Editor/Kismet/Private/BlueprintEditor.cpp
+++ b/Engine/Source/Editor/Kismet/Private/BlueprintEditor.cpp
@@ -6313,8 +6313,8 @@ void FBlueprintEditor::RefreshStandAloneDefaultsEditor()
 	TArray<UObject*> DefaultObjects;
 	for ( int32 i = 0; i < GetEditingObjects().Num(); ++i )
 	{
-		auto Blueprint = (UBlueprint*)( GetEditingObjects()[i] );
-		if ( Blueprint->GeneratedClass )
+		UBlueprintCore* Blueprint = Cast<UBlueprintCore>(GetEditingObjects()[i]);
+		if ( Blueprint && Blueprint->GeneratedClass )
 		{
 			DefaultObjects.Add(Blueprint->GeneratedClass->GetDefaultObject());
 		}
