Changelist number in Blueprint diff tool will no longer display with commas.

diff --git a/Engine/Source/Editor/Kismet/Private/DiffUtils.cpp b/Engine/Source/Editor/Kismet/Private/DiffUtils.cpp
--- a/Engine/Source/Editor/Kismet/Private/DiffUtils.cpp
+++ b/Engine/Source/Editor/Kismet/Private/DiffUtils.cpp
@@ -546,9 +546,11 @@ FText DiffViewUtils::GetPanelLabel(const UBlueprint* Blueprint, const FRevisionI
 		
 		if(ISourceControlModule::Get().GetProvider().UsesChangelists())
 		{
+			FNumberFormattingOptions NumberOptions;
+			NumberOptions.SetUseGrouping(false);
 			RevisionData = FText::Format(NSLOCTEXT("DiffViewUtils", "RevisionData", "Revision {0} - CL {1} - {2}")
 				, FText::FromString(Revision.Revision)
-				, FText::AsNumber(Revision.Changelist)
+				, FText::AsNumber(Revision.Changelist, &NumberOptions)
 				, FText::FromString(Revision.Date.ToString(TEXT("%m/%d/%Y"))));
 		}
 		else
