[INTEGRATE] Change 2434198 by [REDACTED_PERSON][REDACTED_HANDLE] on 2015/02/05 16:29:35

	Component tree: Switched to using short tooltips for component descriptions

diff --git a/Engine/Source/Editor/Kismet/Private/SSCSEditor.cpp b/Engine/Source/Editor/Kismet/Private/SSCSEditor.cpp
--- a/Engine/Source/Editor/Kismet/Private/SSCSEditor.cpp
+++ b/Engine/Source/Editor/Kismet/Private/SSCSEditor.cpp
@@ -1389,7 +1389,7 @@ TSharedRef<SToolTip> SSCS_RowWidget::CreateToolTipWidget() const
 			if (UActorComponent* Template = TreeNode->GetComponentTemplate())
 			{
 				UClass* TemplateClass = Template->GetClass();
-				FText ClassTooltip = TemplateClass->GetToolTipText();
+				FText ClassTooltip = TemplateClass->GetToolTipText(/*bShortTooltip=*/ true);
 
 				InfoBox->AddSlot()
 					.AutoHeight()
