[INTEGRATE] Change 2434573 by [REDACTED_PERSON][REDACTED_HANDLE] on 2015/02/05 19:52:14

	BP - Fixing a crash in the blueprint editor when the activated tab is not valid.

diff --git a/Engine/Source/Editor/Kismet/Private/BlueprintEditor.cpp b/Engine/Source/Editor/Kismet/Private/BlueprintEditor.cpp
--- a/Engine/Source/Editor/Kismet/Private/BlueprintEditor.cpp
+++ b/Engine/Source/Editor/Kismet/Private/BlueprintEditor.cpp
@@ -2736,7 +2736,7 @@ void FBlueprintEditor::OnActiveTabChanged( TSharedPtr<SDockTab> PreviouslyActive
 		//UE_LOG(LogBlueprint, Log, TEXT("OnActiveTabChanged: %s"), *NewlyActivated->GetLayoutIdentifier().ToString());
 	}
 
-	if ( NewlyActivated->GetTabRole() == ETabRole::DocumentTab )
+	if ( NewlyActivated.IsValid() && NewlyActivated->GetTabRole() == ETabRole::DocumentTab )
 	{
 		FocusedGraphEdPtr = nullptr;
 	}
