[INTEGRATE] Change 2434746 by [REDACTED_PERSON][REDACTED_HANDLE] on 2015/02/05 22:35:48

	BP - Tweaking some more selection related things in the blueprint editor.

diff --git a/Engine/Source/Editor/Kismet/Private/BlueprintEditor.cpp b/Engine/Source/Editor/Kismet/Private/BlueprintEditor.cpp
--- a/Engine/Source/Editor/Kismet/Private/BlueprintEditor.cpp
+++ b/Engine/Source/Editor/Kismet/Private/BlueprintEditor.cpp
@@ -630,13 +630,19 @@ void FBlueprintEditor::RefreshEditors(ERefreshBlueprintEditorReason::Type Reason
 		UpdateSCSPreview();
 	}
 
-	// Only force a refresh of the defaults view if we have a standalone defaults viewer.
 	// Note: There is an optimization inside of ShowDetailsForSingleObject() that skips the refresh if the object being selected is the same as the previous object.
 	// The SKismetInspector class is shared between both Defaults mode and Components mode, but in Defaults mode the object selected is always going to be the CDO. Given
 	// that the selection does not really change, we force it to refresh and skip the optimization. Otherwise, some things may not work correctly in Defaults mode. For
 	// example, transform details are customized and the rotation value is cached at customization time; if we don't force refresh here, then after an undo of a previous
 	// rotation edit, transform details won't be re-customized and thus the cached rotation value will be stale, resulting in an invalid rotation value on the next edit.
-	RefreshStandAloneDefaultsEditor();
+	if ( CurrentUISelection == FBlueprintEditor::SelectionState_ClassDefaults )
+	{
+		StartEditingDefaults(true, true);
+	}
+	else
+	{
+		RefreshStandAloneDefaultsEditor();
+	}
 
 	// Update associated controls like the function editor
 	BroadcastRefresh();
@@ -1652,6 +1658,11 @@ void FBlueprintEditor::RegisterApplicationModes(const TArray<UBlueprint*>& InBlu
 						FBlueprintEditorApplicationModes::StandardBlueprintEditorMode,
 						MakeShareable(new FBlueprintEditorUnifiedMode(SharedThis(this), FBlueprintEditorApplicationModes::StandardBlueprintEditorMode, FBlueprintEditorApplicationModes::GetLocalizedMode, CanAccessComponentsMode())));
 					SetCurrentMode(FBlueprintEditorApplicationModes::StandardBlueprintEditorMode);
+
+					if ( bShouldOpenInComponentsMode && CanAccessComponentsMode() )
+					{
+						TabManager->InvokeTab(FBlueprintEditorTabs::SCSViewportID);
+					}
 				}
 			}
 			else
@@ -2037,7 +2048,7 @@ FBlueprintEditor::~FBlueprintEditor()
 void FBlueprintEditor::FocusInspectorOnGraphSelection(const FGraphPanelSelectionSet& NewSelection, bool bForceRefresh)
 {
 	// If this graph has selected nodes update the details panel to match.
-	if ( NewSelection.Array().Num() > 0 )
+	if ( NewSelection.Num() > 0 || CurrentUISelection == FBlueprintEditor::SelectionState_Graph )
 	{
 		SetUISelectionState(FBlueprintEditor::SelectionState_Graph);
 
@@ -2770,15 +2781,6 @@ void FBlueprintEditor::OnGraphEditorFocused(const TSharedRef<SGraphEditor>& InGr
 			{
 				MyBlueprintWidget->Refresh();
 			}
-
-			// If this graph doesn't have any selected nodes, but the current selection state is from another graph,
-			// lets just set the focused graph to be the selected one.  We don't want the wrong function context to 
-			// seem relevant when you're looking at a different function.
-			if ( SelectedNodes.Num() == 0 )
-			{
-				SetUISelectionState(FBlueprintEditor::SelectionState_Graph);
-				Inspector->ShowDetailsForSingleObject(FocusedGraph);
-			}
 		}
 	}
 }
