[INTEGRATE] Change 2435305 by [REDACTED_PERSON][REDACTED_HANDLE] on 2015/02/06 09:13:06

	UE-9225 Events not native to a blueprint do not show up in My Blueprint panel

	UK2Node_InputKey and UK2Node_InputAction nodes are listed in MyBlueprint

diff --git a/Engine/Source/Editor/BlueprintGraph/Classes/EdGraphSchema_K2_Actions.h b/Engine/Source/Editor/BlueprintGraph/Classes/EdGraphSchema_K2_Actions.h
--- a/Engine/Source/Editor/BlueprintGraph/Classes/EdGraphSchema_K2_Actions.h
+++ b/Engine/Source/Editor/BlueprintGraph/Classes/EdGraphSchema_K2_Actions.h
@@ -714,6 +714,56 @@ struct BLUEPRINTGRAPH_API FEdGraphSchemaAction_K2Event : public FEdGraphSchemaAc
 	// End of FEdGraphSchemaAction interface
 };
 
+/*******************************************************************************
+* FEdGraphSchemaAction_K2InputAction
+*******************************************************************************/
+
+/**
+* A reference to a specific event (living inside a Blueprint graph)... intended
+* to be used the 'docked' palette only.
+*/
+USTRUCT()
+struct BLUEPRINTGRAPH_API FEdGraphSchemaAction_K2InputAction : public FEdGraphSchemaAction_K2TargetNode
+{
+	GENERATED_USTRUCT_BODY()
+
+	/**
+	 * An empty default constructor (stubbed out for arrays and other containers
+	 * that need to allocate default instances of this struct)
+	 */
+	FEdGraphSchemaAction_K2InputAction()
+		: FEdGraphSchemaAction_K2TargetNode()
+	{}
+
+	/**
+	* The primary constructor, used to customize the Super FEdGraphSchemaAction.
+	*
+	* @param   Category		The tree parent header (or path, delimited by '|') you want to sort this action under.
+	* @param   MenuDescription	The string you want displayed in the tree, corresponding to this action.
+	* @param   Tooltip			A string to display when hovering over this action entry.
+	* @param   Grouping		Used to override list ordering (actions with the same number get grouped together, higher numbers get sorted first).
+	*/
+	FEdGraphSchemaAction_K2InputAction(FString const& Category, FText const& MenuDescription, FString const& Tooltip, int32 const Grouping)
+		: FEdGraphSchemaAction_K2TargetNode(Category, MenuDescription, Tooltip, Grouping)
+	{}
+
+	/**
+	* Provides a set identifier for all FEdGraphSchemaAction_K2InputAction actions.
+	*
+	* @return A static identifier for actions of this type.
+	*/
+	static FName StaticGetTypeId()
+	{
+		static FName Type("FEdGraphSchemaAction_K2InputAction");
+		return Type;
+	}
+
+	// FEdGraphSchemaAction interface
+	virtual FName GetTypeId() const override { return StaticGetTypeId(); }
+	virtual bool IsParentable() const override { return true; }
+	// End of FEdGraphSchemaAction interface
+};
+
 /*******************************************************************************
 * FEdGraphSchemaAction_K2Delegate
 *******************************************************************************/
diff --git a/Engine/Source/Editor/Kismet/Private/BPFunctionDragDropAction.cpp b/Engine/Source/Editor/Kismet/Private/BPFunctionDragDropAction.cpp
--- a/Engine/Source/Editor/Kismet/Private/BPFunctionDragDropAction.cpp
+++ b/Engine/Source/Editor/Kismet/Private/BPFunctionDragDropAction.cpp
@@ -146,7 +146,8 @@ bool FKismetDragDropAction::ActionWillShowExistingNode() const
 	if (ActionNode.IsValid() && (HoveredGraph != NULL))
 	{
 		bWillFocusOnExistingNode = (ActionNode->GetTypeId() == FEdGraphSchemaAction_K2TargetNode::StaticGetTypeId()) ||
-		                           (ActionNode->GetTypeId() == FEdGraphSchemaAction_K2Event::StaticGetTypeId());
+			(ActionNode->GetTypeId() == FEdGraphSchemaAction_K2Event::StaticGetTypeId()) ||
+			(ActionNode->GetTypeId() == FEdGraphSchemaAction_K2InputAction::StaticGetTypeId());
 
 		if (!bWillFocusOnExistingNode && (ActionNode->GetTypeId() == FEdGraphSchemaAction_K2AddEvent::StaticGetTypeId()))
 		{
diff --git a/Engine/Source/Editor/Kismet/Private/BlueprintActionMenuUtils.cpp b/Engine/Source/Editor/Kismet/Private/BlueprintActionMenuUtils.cpp
--- a/Engine/Source/Editor/Kismet/Private/BlueprintActionMenuUtils.cpp
+++ b/Engine/Source/Editor/Kismet/Private/BlueprintActionMenuUtils.cpp
@@ -727,7 +727,8 @@ const UK2Node* FBlueprintActionMenuUtils::ExtractNodeTemplateFromAction(TSharedP
 			ActionId == FEdGraphSchemaAction_K2TargetNode::StaticGetTypeId() ||
 			ActionId == FEdGraphSchemaAction_K2PasteHere::StaticGetTypeId() ||
 			ActionId == FEdGraphSchemaAction_K2Event::StaticGetTypeId() || 
-			ActionId == FEdGraphSchemaAction_K2AddEvent::StaticGetTypeId())
+			ActionId == FEdGraphSchemaAction_K2AddEvent::StaticGetTypeId() ||
+			ActionId == FEdGraphSchemaAction_K2InputAction::StaticGetTypeId())
 		{
 			FEdGraphSchemaAction_K2NewNode* NewNodeAction = (FEdGraphSchemaAction_K2NewNode*)PaletteAction.Get();
 			TemplateNode = NewNodeAction->NodeTemplate;
diff --git a/Engine/Source/Editor/Kismet/Private/SMyBlueprint.cpp b/Engine/Source/Editor/Kismet/Private/SMyBlueprint.cpp
--- a/Engine/Source/Editor/Kismet/Private/SMyBlueprint.cpp
+++ b/Engine/Source/Editor/Kismet/Private/SMyBlueprint.cpp
@@ -824,6 +824,28 @@ void SMyBlueprint::GetChildGraphs(UEdGraph* EdGraph, FGraphActionListBuilderBase
 	}
 }
 
+struct FCreateEdGraphSchemaActionHelper
+{
+	template<class ActionType, class NodeType>
+	static void CreateAll(UEdGraph const* EdGraph, int32 SectionId, FGraphActionListBuilderBase& OutAllActions, const FString& ActionCategory)
+	{
+		TArray<NodeType*> EventNodes;
+		EdGraph->GetNodesOfClass<NodeType>(EventNodes);
+		for (auto It = EventNodes.CreateConstIterator(); It; ++It)
+		{
+			NodeType* const EventNode = (*It);
+
+			FString const Tooltip = EventNode->GetTooltipText().ToString();
+			FText const Description = EventNode->GetNodeTitle(ENodeTitleType::EditableTitle);
+
+			TSharedPtr<ActionType> EventNodeAction = MakeShareable(new ActionType(ActionCategory, Description, Tooltip, 0));
+			EventNodeAction->NodeTemplate = EventNode;
+			EventNodeAction->SectionID = SectionId;
+			OutAllActions.AddAction(EventNodeAction);
+		}
+	}
+};
+
 void SMyBlueprint::GetChildEvents(UEdGraph const* EdGraph, int32 const SectionId, FGraphActionListBuilderBase& OutAllActions) const
 {
 	if (!ensure(EdGraph != NULL))
@@ -836,23 +858,11 @@ void SMyBlueprint::GetChildEvents(UEdGraph const* EdGraph, int32 const SectionId
 	FGraphDisplayInfo EdGraphDisplayInfo;
 	Schema->GetGraphDisplayInformation(*EdGraph, EdGraphDisplayInfo);
 	FText const EdGraphDisplayName = EdGraphDisplayInfo.DisplayName;
-
-	TArray<UK2Node_Event*> EventNodes;
-	EdGraph->GetNodesOfClass<UK2Node_Event>(EventNodes);
-
 	FString const ActionCategory = EdGraphDisplayName.ToString();
-	for (auto It = EventNodes.CreateConstIterator(); It; ++It)
-	{
-		UK2Node_Event* const EventNode = (*It);
-
-		FString const Tooltip = EventNode->GetTooltipText().ToString();
-		FText const Description = EventNode->GetNodeTitle(ENodeTitleType::EditableTitle);
 
-		TSharedPtr<FEdGraphSchemaAction_K2Event> EventNodeAction = MakeShareable(new FEdGraphSchemaAction_K2Event(ActionCategory, Description, Tooltip, 0));
-		EventNodeAction->NodeTemplate = EventNode;
-		EventNodeAction->SectionID = SectionId;
-		OutAllActions.AddAction(EventNodeAction);
-	}
+	FCreateEdGraphSchemaActionHelper::CreateAll<FEdGraphSchemaAction_K2Event, UK2Node_Event>(EdGraph, SectionId, OutAllActions, ActionCategory);
+	FCreateEdGraphSchemaActionHelper::CreateAll<FEdGraphSchemaAction_K2InputAction, UK2Node_InputKey>(EdGraph, SectionId, OutAllActions, ActionCategory);
+	FCreateEdGraphSchemaActionHelper::CreateAll<FEdGraphSchemaAction_K2InputAction, UK2Node_InputAction>(EdGraph, SectionId, OutAllActions, ActionCategory);
 }
 
 void SMyBlueprint::GetLocalVariables(FGraphActionListBuilderBase& OutAllActions) const
@@ -1482,7 +1492,7 @@ void SMyBlueprint::OnActionSelected( const TArray< TSharedPtr<FEdGraphSchemaActi
 
 void SMyBlueprint::OnActionSelectedHelper(TSharedPtr<FEdGraphSchemaAction> InAction, UBlueprint* Blueprint, TSharedRef<SKismetInspector> Inspector)
 {
-	if(InAction.IsValid())
+	if (InAction.IsValid())
 	{
 		if (InAction->GetTypeId() == FEdGraphSchemaAction_K2Graph::StaticGetTypeId())
 		{
@@ -1537,7 +1547,9 @@ void SMyBlueprint::OnActionSelectedHelper(TSharedPtr<FEdGraphSchemaAction> InAct
 
 			Inspector->ShowDetailsForSingleObject(StructAction->Struct, Options);
 		}
-		else if (InAction->GetTypeId() == FEdGraphSchemaAction_K2TargetNode::StaticGetTypeId() || InAction->GetTypeId() == FEdGraphSchemaAction_K2Event::StaticGetTypeId())
+		else if (InAction->GetTypeId() == FEdGraphSchemaAction_K2TargetNode::StaticGetTypeId() ||
+			InAction->GetTypeId() == FEdGraphSchemaAction_K2Event::StaticGetTypeId() ||
+			InAction->GetTypeId() == FEdGraphSchemaAction_K2InputAction::StaticGetTypeId())
 		{
 			FEdGraphSchemaAction_K2TargetNode* TargetNodeAction = (FEdGraphSchemaAction_K2TargetNode*)InAction.Get();
 			SKismetInspector::FShowDetailsOptions Options(TargetNodeAction->NodeTemplate->GetNodeTitle(ENodeTitleType::EditableTitle));
@@ -1613,7 +1625,8 @@ void SMyBlueprint::ExecuteAction(TSharedPtr<FEdGraphSchemaAction> InAction)
 			FEdGraphSchemaAction_K2Event* EventNodeAction = (FEdGraphSchemaAction_K2Event*)InAction.Get();
 			FKismetEditorUtilities::BringKismetToFocusAttentionOnObject(EventNodeAction->NodeTemplate);
 		}
-		else if (InAction->GetTypeId() == FEdGraphSchemaAction_K2TargetNode::StaticGetTypeId())
+		else if (InAction->GetTypeId() == FEdGraphSchemaAction_K2TargetNode::StaticGetTypeId() || 
+			InAction->GetTypeId() == FEdGraphSchemaAction_K2InputAction::StaticGetTypeId())
 		{
 			FEdGraphSchemaAction_K2TargetNode* TargetNodeAction = (FEdGraphSchemaAction_K2TargetNode*)InAction.Get();
 			FKismetEditorUtilities::BringKismetToFocusAttentionOnObject(TargetNodeAction->NodeTemplate);
@@ -1676,6 +1689,11 @@ FEdGraphSchemaAction_K2Event* SMyBlueprint::SelectionAsEvent() const
 	return SelectionAsType<FEdGraphSchemaAction_K2Event>( GraphActionMenu );
 }
 
+FEdGraphSchemaAction_K2InputAction* SMyBlueprint::SelectionAsInputAction() const
+{
+	return SelectionAsType<FEdGraphSchemaAction_K2InputAction>(GraphActionMenu);
+}
+
 bool SMyBlueprint::SelectionIsCategory() const
 {
 	return !GraphActionMenu->GetSelectedCategoryName().IsEmpty();
@@ -1856,6 +1874,10 @@ void SMyBlueprint::OpenGraph(FDocumentTracker::EOpenDocumentCause InCause)
 	{
 		EdGraph = EventAction->NodeTemplate->GetGraph();
 	}
+	else if (FEdGraphSchemaAction_K2InputAction* InputAction = SelectionAsInputAction())
+	{
+		EdGraph = InputAction->NodeTemplate->GetGraph();
+	}
 	
 	if (EdGraph)
 	{
@@ -1877,14 +1899,18 @@ void SMyBlueprint::OnOpenGraphInNewTab()
 bool SMyBlueprint::CanFocusOnNode() const
 {
 	FEdGraphSchemaAction_K2Event const* const EventAction = SelectionAsEvent();
-	return (EventAction != NULL) && (EventAction->NodeTemplate != NULL);
+	FEdGraphSchemaAction_K2InputAction const* const InputAction = SelectionAsInputAction();
+	return (EventAction && EventAction->NodeTemplate) || (InputAction && InputAction->NodeTemplate);
 }
 
 void SMyBlueprint::OnFocusNode()
 {
-	if (FEdGraphSchemaAction_K2Event* EventAction = SelectionAsEvent())
+	FEdGraphSchemaAction_K2Event* EventAction = SelectionAsEvent();
+	FEdGraphSchemaAction_K2InputAction* InputAction = SelectionAsInputAction();
+	if (EventAction || InputAction)
 	{
-		FKismetEditorUtilities::BringKismetToFocusAttentionOnObject(EventAction->NodeTemplate);
+		auto Node = EventAction ? EventAction->NodeTemplate : InputAction->NodeTemplate;
+		FKismetEditorUtilities::BringKismetToFocusAttentionOnObject(Node);
 	}
 }
 
@@ -1977,6 +2003,12 @@ void SMyBlueprint::OnFindEntry()
 	{
 		SearchTerm = EventAction->MenuDescription.ToString();
 	}
+	else if (FEdGraphSchemaAction_K2InputAction* InputAction = SelectionAsInputAction())
+	{
+		SearchTerm = InputAction->NodeTemplate ? 
+			InputAction->NodeTemplate->GetNodeTitle(ENodeTitleType::FullTitle).ToString() :
+			InputAction->MenuDescription.ToString();
+	}
 
 	if(!SearchTerm.IsEmpty())
 	{
diff --git a/Engine/Source/Editor/Kismet/Private/SMyBlueprint.h b/Engine/Source/Editor/Kismet/Private/SMyBlueprint.h
--- a/Engine/Source/Editor/Kismet/Private/SMyBlueprint.h
+++ b/Engine/Source/Editor/Kismet/Private/SMyBlueprint.h
@@ -86,6 +86,9 @@ public:
 
 	/** Accessor for getting the current selection as a K2 event */
 	FEdGraphSchemaAction_K2Event* SelectionAsEvent() const;
+	
+	/** Accessor for getting the current selection as a K2 Input Action */
+	FEdGraphSchemaAction_K2InputAction* SelectionAsInputAction() const;
 
 	/** Accessor for getting the current selection as a K2 local var */
 	FEdGraphSchemaAction_K2LocalVar* SelectionAsLocalVar() const;
diff --git a/Engine/Source/Editor/UnrealEd/Private/Kismet2/BlueprintEditorUtils.cpp b/Engine/Source/Editor/UnrealEd/Private/Kismet2/BlueprintEditorUtils.cpp
--- a/Engine/Source/Editor/UnrealEd/Private/Kismet2/BlueprintEditorUtils.cpp
+++ b/Engine/Source/Editor/UnrealEd/Private/Kismet2/BlueprintEditorUtils.cpp
@@ -6874,6 +6874,11 @@ bool FBlueprintEditorUtils::IsPaletteActionReadOnly(TSharedPtr<FEdGraphSchemaAct
 
 		bIsReadOnly = (AssociatedNode == NULL) || (!AssociatedNode->bCanRenameNode);	
 	}
+	else if (ActionIn->GetTypeId() == FEdGraphSchemaAction_K2InputAction::StaticGetTypeId())
+	{
+		bIsReadOnly = true;
+	}
+
 
 	return bIsReadOnly;
 }
