[INTEGRATE] Change 2436667 by [REDACTED_PERSON][REDACTED_HANDLE] on 2015/02/07 02:41:14

	Component Tree: Rename the 'create child blueprint' action to better match what it does, and add a duplicate 'Edit Blueprint' option to show 'Add Script' as a keyword that always goes to the event graph

diff --git a/Engine/Source/Editor/Kismet/Private/SSCSEditor.cpp b/Engine/Source/Editor/Kismet/Private/SSCSEditor.cpp
--- a/Engine/Source/Editor/Kismet/Private/SSCSEditor.cpp
+++ b/Engine/Source/Editor/Kismet/Private/SSCSEditor.cpp
@@ -3093,9 +3093,18 @@ void SSCSEditor::Construct( const FArguments& InArgs )
 		LOCTEXT("OpenBlueprintEditor", "Open Blueprint Editor"),
 		LOCTEXT("OpenBlueprintEditor_ToolTip", "Opens the blueprint editor for this asset"),
 		FSlateIcon(),
-		FUIAction(FExecuteAction::CreateSP(this, &SSCSEditor::OnOpenBlueprintEditor))
+		FUIAction(FExecuteAction::CreateSP(this, &SSCSEditor::OnOpenBlueprintEditor, /*bForceCodeEditing=*/ false))
 	);
 
+	EditBlueprintMenuBuilder.AddMenuEntry
+	(
+		LOCTEXT("OpenBlueprintEditorScriptMode", "Add or Edit Script"),
+		LOCTEXT("OpenBlueprintEditorScriptMode_ToolTip", "Opens the blueprint editor for this asset, showing the event graph"),
+		FSlateIcon(),
+		FUIAction(FExecuteAction::CreateSP(this, &SSCSEditor::OnOpenBlueprintEditor, /*bForceCodeEditing=*/ true))
+	);
+
+	EditBlueprintMenuBuilder.BeginSection(NAME_None, LOCTEXT("EditBlueprintMenu_InstanceHeader", "Instance modifications"));
 
 	EditBlueprintMenuBuilder.AddMenuEntry
 	(
@@ -3118,8 +3127,8 @@ void SSCSEditor::Construct( const FArguments& InArgs )
 
 	EditBlueprintMenuBuilder.AddMenuEntry
 	(
-		LOCTEXT("PromoteToBlueprint", "Convert to Blueprint Class"),
-		LOCTEXT("PromoteToBluerprintTooltip","Converts the existing Blueprint into a new SubClass Blueprint" ),
+		LOCTEXT("CreateChildBlueprint", "Create Child Blueprint Class"),
+		LOCTEXT("CreateChildBlueprintTooltip", "Creates a Child Blueprint Class based on the current Blueprint, allowing you to create variants easily.  This replaces the current actor instance with a new one based on the new Child Blueprint Class." ),
 		FSlateIcon(),
 		FUIAction(FExecuteAction::CreateSP(this, &SSCSEditor::PromoteToBlueprint))
 	);
@@ -5343,12 +5352,22 @@ FText SSCSEditor::OnGetResetToBlueprintDefaultsTooltip() const
 	}
 }
 
-void SSCSEditor::OnOpenBlueprintEditor() const
+void SSCSEditor::OnOpenBlueprintEditor(bool bForceCodeEditing) const
 {
-	AActor* ActorInstance = GetActorContext();
-	UBlueprint* Blueprint = (ActorInstance != nullptr) ? Cast<UBlueprint>( ActorInstance->GetClass()->ClassGeneratedBy ) : nullptr;
-
-	FAssetEditorManager::Get().OpenEditorForAsset(Blueprint);
+	if (AActor* ActorInstance = GetActorContext())
+	{
+		if (UBlueprint* Blueprint = Cast<UBlueprint>(ActorInstance->GetClass()->ClassGeneratedBy))
+		{
+			if (bForceCodeEditing && (Blueprint->UbergraphPages.Num() > 0))
+			{
+				FKismetEditorUtilities::BringKismetToFocusAttentionOnObject(Blueprint->UbergraphPages[0]);
+			}
+			else
+			{
+				FAssetEditorManager::Get().OpenEditorForAsset(Blueprint);
+			}
+		}
+	}
 }
 
 void SSCSEditor::OnApplyChangesToBlueprint() const
diff --git a/Engine/Source/Editor/Kismet/Public/SSCSEditor.h b/Engine/Source/Editor/Kismet/Public/SSCSEditor.h
--- a/Engine/Source/Editor/Kismet/Public/SSCSEditor.h
+++ b/Engine/Source/Editor/Kismet/Public/SSCSEditor.h
@@ -979,7 +979,7 @@ protected:
 	FText OnGetResetToBlueprintDefaultsTooltip() const;
 
 	/** Opens the blueprint editor for the blueprint being viewed by the scseditor */
-	void OnOpenBlueprintEditor() const;
+	void OnOpenBlueprintEditor(bool bForceCodeEditing) const;
 
 	/** Propagates instance changes to the blueprint */
 	void OnApplyChangesToBlueprint() const;
