[INTEGRATE] Change 2438562 by [REDACTED_PERSON][REDACTED_HANDLE] on 2015/02/09 14:50:15

	BP - Fixing a bug when navigating to a node, the tab containing the node isn't always activated and brought to the front even when it exists.

diff --git a/Engine/Source/Editor/Kismet/Private/WorkflowOrientedApp/WorkflowTabManager.cpp b/Engine/Source/Editor/Kismet/Private/WorkflowOrientedApp/WorkflowTabManager.cpp
--- a/Engine/Source/Editor/Kismet/Private/WorkflowOrientedApp/WorkflowTabManager.cpp
+++ b/Engine/Source/Editor/Kismet/Private/WorkflowOrientedApp/WorkflowTabManager.cpp
@@ -530,6 +530,8 @@ TSharedPtr<SDockTab> FDocumentTracker::NavigateCurrentTab(TSharedPtr<FTabPayload
 			TSharedPtr<FDocumentTabFactory> Factory = FindSupportingFactory(InPayload.ToSharedRef());
 			// If doing a Quick navigate of the document, do not save history data as it's likely still at the default values. The object is always saved
 			LastEditedTabInfo.Pin()->AddTabHistory(Factory->CreateTabHistoryNode(InPayload), InNavigateCause != QuickNavigateCurrentDocument);
+			// Ensure that the tab appears if the tab isn't currently in the foreground.
+			LastEditedTabInfo.Pin()->GetTab().Pin()->ActivateInParent(ETabActivationCause::SetDirectly);
 		}
 		else if(InNavigateCause == NavigateBackwards)
 		{
