[INTEGRATE] Change 2438564 by [REDACTED_PERSON][REDACTED_HANDLE] on 2015/02/09 14:52:10

	UE-9374 - Editor - The + Buttons in the MyBlueprint View now disable themselves during PIE.

diff --git a/Engine/Source/Editor/Kismet/Private/SMyBlueprint.cpp b/Engine/Source/Editor/Kismet/Private/SMyBlueprint.cpp
--- a/Engine/Source/Editor/Kismet/Private/SMyBlueprint.cpp
+++ b/Engine/Source/Editor/Kismet/Private/SMyBlueprint.cpp
@@ -539,6 +539,8 @@ FText SMyBlueprint::OnGetSectionTitle( int32 InSectionID )
 
 TSharedRef<SWidget> SMyBlueprint::OnGetSectionWidget(TSharedRef<SWidget> RowWidget, int32 InSectionID)
 {
+	TSharedPtr<FBlueprintEditor> BlueprintEditor = BlueprintEditorPtr.Pin();
+
 	TWeakPtr<SWidget> WeakRowWidget = RowWidget;
 
 	FText AddNewText;
@@ -558,6 +560,7 @@ TSharedRef<SWidget> SMyBlueprint::OnGetSectionWidget(TSharedRef<SWidget> RowWidg
 				.AutoWidth()
 				[
 					SAssignNew(FunctionSectionButton, SComboButton)
+					.IsEnabled(BlueprintEditor.Get(), &FBlueprintEditor::InEditingMode)
 					.Visibility(this, &SMyBlueprint::OnGetSectionTextVisibility, WeakRowWidget, InSectionID)
 					.ButtonStyle(FEditorStyle::Get(), "RoundButton")
 					.ForegroundColor(FEditorStyle::GetSlateColor("DefaultForeground"))
@@ -603,11 +606,14 @@ TSharedRef<SWidget> SMyBlueprint::OnGetSectionWidget(TSharedRef<SWidget> RowWidg
 
 TSharedRef<SWidget> SMyBlueprint::CreateAddToSectionButton(int32 InSectionID, TWeakPtr<SWidget> WeakRowWidget, FText AddNewText)
 {
+	TSharedPtr<FBlueprintEditor> BlueprintEditor = BlueprintEditorPtr.Pin();
+
 	return SNew(SButton)
 		.ButtonStyle(FEditorStyle::Get(), "RoundButton")
 		.ForegroundColor(FEditorStyle::GetSlateColor("DefaultForeground"))
 		.ContentPadding(FMargin(2, 0))
 		.OnClicked(this, &SMyBlueprint::OnAddButtonClickedOnSection, InSectionID)
+		.IsEnabled(BlueprintEditor.Get(), &FBlueprintEditor::InEditingMode)
 		.HAlign(HAlign_Center)
 		.VAlign(VAlign_Center)
 		[
@@ -637,22 +643,24 @@ TSharedRef<SWidget> SMyBlueprint::CreateAddToSectionButton(int32 InSectionID, TW
 
 FReply SMyBlueprint::OnAddButtonClickedOnSection(int32 InSectionID)
 {
+	TSharedPtr<FBlueprintEditor> BlueprintEditor = BlueprintEditorPtr.Pin();
+
 	switch ( InSectionID )
 	{
 	case NodeSectionID::VARIABLE:
-		BlueprintEditorPtr.Pin()->GetToolkitCommands()->ExecuteAction(FBlueprintEditorCommands::Get().AddNewVariable.ToSharedRef());
+		BlueprintEditor->GetToolkitCommands()->ExecuteAction(FBlueprintEditorCommands::Get().AddNewVariable.ToSharedRef());
 		break;
 	case NodeSectionID::FUNCTION:
-		BlueprintEditorPtr.Pin()->GetToolkitCommands()->ExecuteAction(FBlueprintEditorCommands::Get().AddNewFunction.ToSharedRef());
+		BlueprintEditor->GetToolkitCommands()->ExecuteAction(FBlueprintEditorCommands::Get().AddNewFunction.ToSharedRef());
 		break;
 	case NodeSectionID::MACRO:
-		BlueprintEditorPtr.Pin()->GetToolkitCommands()->ExecuteAction(FBlueprintEditorCommands::Get().AddNewMacroDeclaration.ToSharedRef());
+		BlueprintEditor->GetToolkitCommands()->ExecuteAction(FBlueprintEditorCommands::Get().AddNewMacroDeclaration.ToSharedRef());
 		break;
 	case NodeSectionID::DELEGATE:
-		BlueprintEditorPtr.Pin()->GetToolkitCommands()->ExecuteAction(FBlueprintEditorCommands::Get().AddNewDelegate.ToSharedRef());
+		BlueprintEditor->GetToolkitCommands()->ExecuteAction(FBlueprintEditorCommands::Get().AddNewDelegate.ToSharedRef());
 		break;
 	case NodeSectionID::GRAPH:
-		BlueprintEditorPtr.Pin()->GetToolkitCommands()->ExecuteAction(FBlueprintEditorCommands::Get().AddNewEventGraph.ToSharedRef());
+		BlueprintEditor->GetToolkitCommands()->ExecuteAction(FBlueprintEditorCommands::Get().AddNewEventGraph.ToSharedRef());
 		break;
 	case NodeSectionID::LOCAL_VARIABLE:
 		OnAddNewLocalVariable();
