[INTEGRATE] Change 2438572 by [REDACTED_PERSON][REDACTED_HANDLE] on 2015/02/09 14:56:14

	Editor - Several visual improvements to the buttons in the editor that we want users to be drawn to clicking.  Using font awesome in a few places to give the icons matching size and formatting to the text in the same button.  Additionally fixing several focus issues with buttons, and things being refreshed at the wrong time, that was causing buttons to not be clickable in the My Blueprints view.

diff --git a/Engine/Source/Editor/ContentBrowser/Private/SContentBrowser.cpp b/Engine/Source/Editor/ContentBrowser/Private/SContentBrowser.cpp
--- a/Engine/Source/Editor/ContentBrowser/Private/SContentBrowser.cpp
+++ b/Engine/Source/Editor/ContentBrowser/Private/SContentBrowser.cpp
@@ -109,7 +109,7 @@ void SContentBrowser::Construct( const FArguments& InArgs, const FName& InInstan
 						.Padding(0,0,4,0)
 						[
 							SNew( SComboButton )
-							.ComboButtonStyle( FEditorStyle::Get(), "ContentBrowser.NewAsset.Style" )
+							.ComboButtonStyle( FEditorStyle::Get(), "ToolbarComboButton" )
 							.ForegroundColor(FLinearColor::White)
 							.ContentPadding(0)
 							.OnGetMenuContent_Lambda( [this]{ return MakeAddNewContextMenu( true, false ); } )
@@ -490,9 +490,25 @@ void SContentBrowser::Construct( const FArguments& InArgs, const FName& InInstan
 							.AddMetaData<FTagMetaData>(FTagMetaData(TEXT("ContentBrowserFiltersCombo")))
 							.ButtonContent()
 							[
-								SNew( STextBlock )
-								.TextStyle( FEditorStyle::Get(), "ContentBrowser.Filters.Text" )
-								.Text( LOCTEXT( "Filters", "Filters" ) )
+								SNew(SHorizontalBox)
+
+								+ SHorizontalBox::Slot()
+								.AutoWidth()
+								[
+									SNew(STextBlock)
+									.TextStyle(FEditorStyle::Get(), "ContentBrowser.Filters.Text")
+									.Font(FEditorStyle::Get().GetFontStyle("FontAwesome.9"))
+									.Text(FString(TEXT("\xf0b0")) /*fa-filter*/)
+								]
+
+								+ SHorizontalBox::Slot()
+								.AutoWidth()
+								.Padding(2,0,0,0)
+								[
+									SNew(STextBlock)
+									.TextStyle(FEditorStyle::Get(), "ContentBrowser.Filters.Text")
+									.Text(LOCTEXT("Filters", "Filters"))
+								]
 							]
 						]
 
diff --git a/Engine/Source/Editor/Kismet/Private/BlueprintEditor.cpp b/Engine/Source/Editor/Kismet/Private/BlueprintEditor.cpp
--- a/Engine/Source/Editor/Kismet/Private/BlueprintEditor.cpp
+++ b/Engine/Source/Editor/Kismet/Private/BlueprintEditor.cpp
@@ -798,6 +798,15 @@ void FBlueprintEditor::OnSelectionUpdated(const TArray<FSCSEditorTreeNodePtrType
 						InspectorTitle = FText::FromString(NodePtr->GetDisplayString());
 						InspectorObjects.Add(EditableComponent);
 					}
+
+					if ( SCSViewport.IsValid() )
+					{
+						TSharedPtr<SDockTab> OwnerTab = SCSViewport->GetOwnerTab();
+						if ( OwnerTab.IsValid() )
+						{
+							OwnerTab->FlashTab();
+						}
+					}
 				}
 			}
 		}
@@ -2747,14 +2756,16 @@ void FBlueprintEditor::OnActiveTabChanged( TSharedPtr<SDockTab> PreviouslyActive
 		//UE_LOG(LogBlueprint, Log, TEXT("OnActiveTabChanged: %s"), *NewlyActivated->GetLayoutIdentifier().ToString());
 	}
 
+	// If the newly active document tab isn't a graph we want to make sure we clear the focused graph pointer.
+	// Several other UI reads that, like the MyBlueprints view uses it to determine if it should show the "Local Variable" section.
 	if ( NewlyActivated.IsValid() && NewlyActivated->GetTabRole() == ETabRole::DocumentTab )
 	{
 		FocusedGraphEdPtr = nullptr;
-	}
 
-	if (MyBlueprintWidget.IsValid() == true)
-	{
-		MyBlueprintWidget->Refresh();
+		if ( MyBlueprintWidget.IsValid() == true )
+		{
+			MyBlueprintWidget->Refresh();
+		}
 	}
 }
 
diff --git a/Engine/Source/Editor/Kismet/Private/SSCSEditor.cpp b/Engine/Source/Editor/Kismet/Private/SSCSEditor.cpp
--- a/Engine/Source/Editor/Kismet/Private/SSCSEditor.cpp
+++ b/Engine/Source/Editor/Kismet/Private/SSCSEditor.cpp
@@ -1,4 +1,4 @@
-// Copyright 1998-2015 Epic Games, Inc. All Rights Reserved.
+﻿// Copyright 1998-2015 Epic Games, Inc. All Rights Reserved.
 
 
 #include "BlueprintEditorPrivatePCH.h"
@@ -1806,7 +1806,7 @@ FReply SSCS_RowWidget::OnDragDetected( const FGeometry& MyGeometry, const FPoint
 		{
 			UBlueprint* Blueprint = FirstNode->GetBlueprint();
 			const FName VariableName = FirstNode->GetVariableName();
- 			UStruct* VariableScope = (Blueprint != nullptr) ? Blueprint->SkeletonGeneratedClass : nullptr;
+			UStruct* VariableScope = (Blueprint != nullptr) ? Blueprint->SkeletonGeneratedClass : nullptr;
 
 			TSharedRef<FSCSRowDragDropOp> Operation = FSCSRowDragDropOp::New(VariableName, VariableScope, FNodeCreationAnalytic());
 			Operation->SetCtrlDrag(true); // Always put a getter
@@ -3142,30 +3142,34 @@ void SSCSEditor::Construct( const FArguments& InArgs )
 		.Padding(0)
 		[
 			SNew(SBorder)
-			.Padding(2.0f)
+			.Padding(0)
 			.BorderImage(FEditorStyle::GetBrush("DetailsView.CategoryTop"))
 			.AddMetaData<FTagMetaData>(FTagMetaData(TEXT("ComponentsPanel")))
 			.BorderBackgroundColor( FLinearColor( .6,.6,.6, 1.0f ) )
 			[
 				SNew(SHorizontalBox)
+				
 				+ SHorizontalBox::Slot()
-				.HAlign(HAlign_Left)
 				.AutoWidth()
+				.HAlign(HAlign_Left)
 				[
 					SNew(SComponentClassCombo)
+					.AddMetaData<FTagMetaData>(FTagMetaData(TEXT("Actor.AddComponent")))
 					.Visibility(bHideComponentClassCombo ? EVisibility::Hidden : EVisibility::Visible)
 					.OnComponentClassSelected(this, &SSCSEditor::PerformComboAddClass)
 					.ToolTipText(LOCTEXT("AddComponent_Tooltip", "Adds a new component to this actor"))
 				]
+
 				+ SHorizontalBox::Slot()
-				.Padding( 4.f, 0.0f )
+				.FillWidth(1.0f)
 				.HAlign(HAlign_Right)
 				[
 					SNew( SButton )
+					.AddMetaData<FTagMetaData>(FTagMetaData(TEXT("Actor.ConvertToBlueprint")))
 					.Visibility( this, &SSCSEditor::GetPromoteToBlueprintButtonVisibility )
 					.OnClicked( this, &SSCSEditor::OnPromoteToBlueprintClicked )
-					.ButtonStyle(FEditorStyle::Get(), "ToggleButton")
-					.ContentPadding(0)
+					.ButtonStyle(FEditorStyle::Get(), "FlatButton.Primary")
+					.ContentPadding(FMargin(10,0))
 					.ToolTip(IDocumentation::Get()->CreateToolTip(
 						LOCTEXT("PromoteToBluerprintTooltip","Converts this actor into a reusable Blueprint Class that can have script behavior" ),
 						NULL,
@@ -3173,32 +3177,49 @@ void SSCSEditor::Construct( const FArguments& InArgs )
 						TEXT("ConvertToBlueprint")))
 					[
 						SNew(SHorizontalBox)
-						+SHorizontalBox::Slot()
+						
+						//+SHorizontalBox::Slot()
+						//.VAlign(VAlign_Center)
+						//.AutoWidth()
+						//.Padding(2.f, 1.f)
+						//[
+						//	SNew(SImage)
+						//	.Image( FEditorStyle::Get().GetBrush( "ClassIcon.BlueprintCore" ) )
+						//]
+
+						+ SHorizontalBox::Slot()
 						.VAlign(VAlign_Center)
+						.Padding(3.f)
 						.AutoWidth()
-						.Padding(2.f, 1.f)
 						[
-							SNew(SImage)
-							.Image( FEditorStyle::Get().GetBrush( "ClassIcon.BlueprintCore" ) )
+							SNew(STextBlock)
+							.TextStyle(FEditorStyle::Get(), "ContentBrowser.TopBar.Font")
+							.Font( FEditorStyle::Get().GetFontStyle( "FontAwesome.10" ) )
+							.Text(FString(TEXT("\xf085")) /*fa-cogs*/)
 						]
+
 						+ SHorizontalBox::Slot()
 						.VAlign(VAlign_Center)
-						.Padding(1.f)
+						.Padding(3.f)
+						.AutoWidth()
 						[
 							SNew(STextBlock)
 							.TextStyle(FEditorStyle::Get(), "ContentBrowser.TopBar.Font")
-							.Text( LOCTEXT("PromoteToBlueprint", "Convert to Blueprint Class") )
+							//.Text( LOCTEXT("PromoteToBlueprint", "Add Script") )
+							.Text(LOCTEXT("PromoteToBlueprint", "Convert to Blueprint Class"))
 						]
 					]
 				]
 				+ SHorizontalBox::Slot()
-				.Padding(4.f, 0.0f)
+				.FillWidth(1.0f)
 				.HAlign(HAlign_Right)
 				[
 					SNew(SComboButton)
+					.AddMetaData<FTagMetaData>(FTagMetaData(TEXT("Actor.EditBlueprint")))
 					.Visibility(this, &SSCSEditor::GetEditBlueprintButtonVisibility)
 					.ContentPadding(FMargin(0))
-					.ComboButtonStyle(FEditorStyle::Get(), "ContentBrowser.NewAsset.Style")
+					.ComboButtonStyle(FEditorStyle::Get(), "ToolbarComboButton")
+					.ButtonStyle(FEditorStyle::Get(), "FlatButton.Primary")
 					.ForegroundColor(FLinearColor::White)
 					.ButtonContent()
 					[
diff --git a/Engine/Source/Editor/Kismet/Public/BlueprintEditor.h b/Engine/Source/Editor/Kismet/Public/BlueprintEditor.h
--- a/Engine/Source/Editor/Kismet/Public/BlueprintEditor.h
+++ b/Engine/Source/Editor/Kismet/Public/BlueprintEditor.h
@@ -68,7 +68,7 @@ struct FComponentEventConstructionData
 	// The name of the event handler to create.
 	FName VariableName;
 	// The template component that the handler applies to.
-	TWeakObjectPtr<UActorComponent> Component;
+	TWeakObjectPtr<UObject> Component;
 };
 
 /** The delegate that the caller must supply to BuildComponentActionsSubMenu that returns the currently selected items */
diff --git a/Engine/Source/Editor/UnrealEd/Private/Kismet2/Kismet2.cpp b/Engine/Source/Editor/UnrealEd/Private/Kismet2/Kismet2.cpp
--- a/Engine/Source/Editor/UnrealEd/Private/Kismet2/Kismet2.cpp
+++ b/Engine/Source/Editor/UnrealEd/Private/Kismet2/Kismet2.cpp
@@ -1548,21 +1548,29 @@ void FKismetEditorUtilities::CreateNewBoundEventForActor(AActor* Actor, FName Ev
 	}
 }
 
-void FKismetEditorUtilities::CreateNewBoundEventForComponent(UActorComponent* Component, FName EventName, UBlueprint* Blueprint, UObjectProperty* ComponentProperty)
+void FKismetEditorUtilities::CreateNewBoundEventForComponent(UObject* Component, FName EventName, UBlueprint* Blueprint, UObjectProperty* ComponentProperty)
 {
-	if ((Component != NULL) && (EventName != NAME_None) && (Blueprint != NULL) && (ComponentProperty != NULL))
+	if ( Component != nullptr )
+	{
+		CreateNewBoundEventForClass(Component->GetClass(), EventName, Blueprint, ComponentProperty);
+	}
+}
+
+void FKismetEditorUtilities::CreateNewBoundEventForClass(UClass* Class, FName EventName, UBlueprint* Blueprint, UObjectProperty* ComponentProperty)
+{
+	if ( ( Class != nullptr ) && ( EventName != NAME_None ) && ( Blueprint != nullptr ) && ( ComponentProperty != nullptr ) )
 	{
 		// First, find the property we want to bind to
-		UMulticastDelegateProperty* DelegateProperty = FindField<UMulticastDelegateProperty>(Component->GetClass(), EventName);
-		if(DelegateProperty != NULL)
+		UMulticastDelegateProperty* DelegateProperty = FindField<UMulticastDelegateProperty>(Class, EventName);
+		if ( DelegateProperty != nullptr )
 		{
-			UEdGraph* TargetGraph = NULL;
+			UEdGraph* TargetGraph = nullptr;
 			if(Blueprint->UbergraphPages.Num() > 0)
 			{
 				TargetGraph = Blueprint->UbergraphPages[0]; // Just use the first graph
 			}
 
-			if(TargetGraph != NULL)
+			if ( TargetGraph != nullptr )
 			{
 				// Figure out a decent place to stick the node
 				const FVector2D NewNodePos = TargetGraph->GetGoodPlaceForNewNode();
@@ -1574,7 +1582,7 @@ void FKismetEditorUtilities::CreateNewBoundEventForComponent(UActorComponent* Co
 				UK2Node_ComponentBoundEvent* EventNode = FEdGraphSchemaAction_K2NewNode::SpawnNodeFromTemplate<UK2Node_ComponentBoundEvent>(TargetGraph, EventNodeTemplate, NewNodePos);
 
 				// Finally, bring up kismet and jump to the new node
-				if (EventNode != NULL)
+				if ( EventNode != nullptr )
 				{
 					BringKismetToFocusAttentionOnObject(EventNode);
 				}
@@ -1610,14 +1618,14 @@ const UK2Node_ActorBoundEvent* FKismetEditorUtilities::FindBoundEventForActor(AA
 const UK2Node_ComponentBoundEvent* FKismetEditorUtilities::FindBoundEventForComponent(const UBlueprint* Blueprint, FName EventName, FName PropertyName)
 {
 	const UK2Node_ComponentBoundEvent* Node = NULL;
-	if (Blueprint && EventName != NAME_None && PropertyName != NAME_None)
+	if ( Blueprint && EventName != NAME_None && PropertyName != NAME_None )
 	{
 		TArray<UK2Node_ComponentBoundEvent*> EventNodes;
 		FBlueprintEditorUtils::GetAllNodesOfClass(Blueprint, EventNodes);
-		for(auto NodeIter = EventNodes.CreateIterator(); NodeIter; ++NodeIter)
+		for ( auto NodeIter = EventNodes.CreateIterator(); NodeIter; ++NodeIter )
 		{
 			UK2Node_ComponentBoundEvent* BoundEvent = *NodeIter;
-			if ((BoundEvent->ComponentPropertyName == PropertyName) && (BoundEvent->DelegatePropertyName == EventName))
+			if ( ( BoundEvent->ComponentPropertyName == PropertyName ) && ( BoundEvent->DelegatePropertyName == EventName ) )
 			{
 				Node = *NodeIter;
 				break;
diff --git a/Engine/Source/Editor/UnrealEd/Private/SComponentClassCombo.cpp b/Engine/Source/Editor/UnrealEd/Private/SComponentClassCombo.cpp
--- a/Engine/Source/Editor/UnrealEd/Private/SComponentClassCombo.cpp
+++ b/Engine/Source/Editor/UnrealEd/Private/SComponentClassCombo.cpp
@@ -50,10 +50,12 @@ void SComponentClassCombo::Construct(const FArguments& InArgs)
 		+SHorizontalBox::Slot()
 		.VAlign(VAlign_Center)
 		.AutoWidth()
-		.Padding(2.f,1.f)
+		.Padding(1.f,1.f)
 		[
-			SNew(SImage)
-			.Image(FEditorStyle::GetBrush(TEXT("Plus")))
+			SNew(STextBlock)
+			.TextStyle(FEditorStyle::Get(), "ContentBrowser.TopBar.Font")
+			.Font(FEditorStyle::Get().GetFontStyle("FontAwesome.10"))
+			.Text(FString(TEXT("\xf067")) /*fa-plus*/)
 		]
 		+ SHorizontalBox::Slot()
 		.VAlign(VAlign_Center)
@@ -94,8 +96,9 @@ void SComponentClassCombo::Construct(const FArguments& InArgs)
 		]
 	]
 	.IsFocusable(true)
-	.ContentPadding(FMargin(0))
-	.ComboButtonStyle(FEditorStyle::Get(), "ContentBrowser.NewAsset.Style")
+	.ContentPadding(FMargin(5, 0))
+	.ComboButtonStyle(FEditorStyle::Get(), "ToolbarComboButton")
+	.ButtonStyle(FEditorStyle::Get(), "FlatButton.Success")
 	.ForegroundColor(FLinearColor::White)
 	.OnComboBoxOpened(this, &SComponentClassCombo::ClearSelection);
 
diff --git a/Engine/Source/Editor/UnrealEd/Public/Kismet2/KismetEditorUtilities.h b/Engine/Source/Editor/UnrealEd/Public/Kismet2/KismetEditorUtilities.h
--- a/Engine/Source/Editor/UnrealEd/Public/Kismet2/KismetEditorUtilities.h
+++ b/Engine/Source/Editor/UnrealEd/Public/Kismet2/KismetEditorUtilities.h
@@ -151,7 +151,10 @@ public:
 	static void CreateNewBoundEventForActor(AActor* Actor, FName EventName);
 
 	/** Create a new event node in the  blueprint, for the supplied component, event name and blueprint */
-	static void CreateNewBoundEventForComponent(UActorComponent* Component, FName EventName, UBlueprint* Blueprint, UObjectProperty* ComponentProperty);
+	static void CreateNewBoundEventForComponent(UObject* Component, FName EventName, UBlueprint* Blueprint, UObjectProperty* ComponentProperty);
+
+	/** Create a new event node in the  blueprint, for the supplied class, event name and blueprint */
+	static void CreateNewBoundEventForClass(UClass* Class, FName EventName, UBlueprint* Blueprint, UObjectProperty* ComponentProperty);
 
 	/** Can we paste to this graph? */
 	static bool CanPasteNodes(const class UEdGraph* Graph);
