[INTEGRATE] Change 2438863 by [REDACTED_PERSON][REDACTED_HANDLE].7_Mac on 2015/02/09 17:36:01

	Fixed not being able to programatically select variables in the my blueprint tab

	The code that actually selected items was erroneously removed in CL#2432451
	This addresses UE-9415 - The inline rename mode after adding a variable frequently ends up editing the wrong variable

diff --git a/Engine/Source/Editor/Kismet/Private/SMyBlueprint.cpp b/Engine/Source/Editor/Kismet/Private/SMyBlueprint.cpp
--- a/Engine/Source/Editor/Kismet/Private/SMyBlueprint.cpp
+++ b/Engine/Source/Editor/Kismet/Private/SMyBlueprint.cpp
@@ -2525,6 +2525,11 @@ void SMyBlueprint::SelectItemByName(const FName& ItemName, ESelectInfo::Type Sel
 	{
 		ClearGraphActionMenuSelection();
 	}
+	else
+	{
+		// Attempt to select the item in the main graph action menu
+		GraphActionMenu->SelectItemByName(ItemName, SelectInfo, SectionId, bIsCategory);
+	}
 }
 
 void SMyBlueprint::ClearGraphActionMenuSelection()
