[INTEGRATE] Change 2439080 by [REDACTED_PERSON][REDACTED_HANDLE] on 2015/02/09 19:39:12

	Fixed crash when undoing a function add. UE-9369

diff --git a/Engine/Source/Editor/Kismet/Private/BlueprintDetailsCustomization.cpp b/Engine/Source/Editor/Kismet/Private/BlueprintDetailsCustomization.cpp
--- a/Engine/Source/Editor/Kismet/Private/BlueprintDetailsCustomization.cpp
+++ b/Engine/Source/Editor/Kismet/Private/BlueprintDetailsCustomization.cpp
@@ -3874,13 +3874,16 @@ FReply FBaseBlueprintGraphActionDetails::OnAddNewInputClicked()
 
 EVisibility FBlueprintGraphActionDetails::GetAddNewInputOutputVisibility() const
 {
-	UK2Node_EditablePinBase * FunctionEntryNode = FunctionEntryNodePtr.Get();
-	if(UEdGraph* Graph = FunctionEntryNode->GetGraph())
+	UK2Node_EditablePinBase* FunctionEntryNode = FunctionEntryNodePtr.Get();
+	if (FunctionEntryNodePtr.IsValid())
 	{
-		// Math expression graphs are read only, do not allow adding or removing of pins
-		if(Cast<UK2Node_MathExpression>(Graph->GetOuter()))
+		if(UEdGraph* Graph = FunctionEntryNode->GetGraph())
 		{
-			return EVisibility::Collapsed;
+			// Math expression graphs are read only, do not allow adding or removing of pins
+			if(Cast<UK2Node_MathExpression>(Graph->GetOuter()))
+			{
+				return EVisibility::Collapsed;
+			}
 		}
 	}
 	return EVisibility::Visible;
