[INTEGRATE] Change 2439240 by [REDACTED_PERSON][REDACTED_HANDLE] on 2015/02/09 21:54:16

	[UE-9494] Fix for crash that occured when diffing blueprints, we have no blueprint editor in this case

diff --git a/Engine/Source/Editor/Kismet/Private/SMyBlueprint.cpp b/Engine/Source/Editor/Kismet/Private/SMyBlueprint.cpp
--- a/Engine/Source/Editor/Kismet/Private/SMyBlueprint.cpp
+++ b/Engine/Source/Editor/Kismet/Private/SMyBlueprint.cpp
@@ -486,7 +486,7 @@ void SMyBlueprint::OnCategoryNameCommitted(const FText& InNewText, ETextCommit::
 		}
 		Refresh();
 		FBlueprintEditorUtils::MarkBlueprintAsModified(GetBlueprintObj());
-		BlueprintEditorPtr.Pin()->GetMyBlueprintWidget()->SelectItemByName(FName(*CategoryName.ToString()), ESelectInfo::OnMouseClick, InAction.Pin()->SectionID, true);
+		SelectItemByName(FName(*CategoryName.ToString()), ESelectInfo::OnMouseClick, InAction.Pin()->SectionID, true);
 	}
 }
 
@@ -539,8 +539,6 @@ FText SMyBlueprint::OnGetSectionTitle( int32 InSectionID )
 
 TSharedRef<SWidget> SMyBlueprint::OnGetSectionWidget(TSharedRef<SWidget> RowWidget, int32 InSectionID)
 {
-	TSharedPtr<FBlueprintEditor> BlueprintEditor = BlueprintEditorPtr.Pin();
-
 	TWeakPtr<SWidget> WeakRowWidget = RowWidget;
 
 	FText AddNewText;
@@ -560,7 +558,7 @@ TSharedRef<SWidget> SMyBlueprint::OnGetSectionWidget(TSharedRef<SWidget> RowWidg
 				.AutoWidth()
 				[
 					SAssignNew(FunctionSectionButton, SComboButton)
-					.IsEnabled(BlueprintEditor.Get(), &FBlueprintEditor::InEditingMode)
+					.IsEnabled(this, &SMyBlueprint::IsEditingMode)
 					.Visibility(this, &SMyBlueprint::OnGetSectionTextVisibility, WeakRowWidget, InSectionID)
 					.ButtonStyle(FEditorStyle::Get(), "RoundButton")
 					.ForegroundColor(FEditorStyle::GetSlateColor("DefaultForeground"))
@@ -606,14 +604,12 @@ TSharedRef<SWidget> SMyBlueprint::OnGetSectionWidget(TSharedRef<SWidget> RowWidg
 
 TSharedRef<SWidget> SMyBlueprint::CreateAddToSectionButton(int32 InSectionID, TWeakPtr<SWidget> WeakRowWidget, FText AddNewText)
 {
-	TSharedPtr<FBlueprintEditor> BlueprintEditor = BlueprintEditorPtr.Pin();
-
 	return SNew(SButton)
 		.ButtonStyle(FEditorStyle::Get(), "RoundButton")
 		.ForegroundColor(FEditorStyle::GetSlateColor("DefaultForeground"))
 		.ContentPadding(FMargin(2, 0))
 		.OnClicked(this, &SMyBlueprint::OnAddButtonClickedOnSection, InSectionID)
-		.IsEnabled(BlueprintEditor.Get(), &FBlueprintEditor::InEditingMode)
+		.IsEnabled(this, &SMyBlueprint::IsEditingMode)
 		.HAlign(HAlign_Center)
 		.VAlign(VAlign_Center)
 		[
@@ -780,7 +776,7 @@ bool SMyBlueprint::CanRequestRenameOnActionNode(TWeakPtr<FGraphActionNode> InSel
 		bIsReadOnly = FBlueprintEditorUtils::IsPaletteActionReadOnly(InSelectedNode.Pin()->Actions[0], BlueprintEditorPtr.Pin());
 	}
 
-	return BlueprintEditorPtr.Pin()->InEditingMode() && !bIsReadOnly;
+	return IsEditingMode() && !bIsReadOnly;
 }
 
 void SMyBlueprint::Refresh()
@@ -1303,24 +1299,25 @@ void SMyBlueprint::CollectStaticSections(TArray<int32>& StaticSectionIDs)
 	if ( IsShowingEmptySections() )
 	{
 		TSharedPtr<FBlueprintEditor> BlueprintEditor = BlueprintEditorPtr.Pin();
-		
-		if ( BlueprintEditor->NewDocument_IsVisibleForType(FBlueprintEditor::CGT_NewEventGraph) )
+		const bool bIsEditor = BlueprintEditor.IsValid();
+
+		if (!bIsEditor || BlueprintEditor->NewDocument_IsVisibleForType(FBlueprintEditor::CGT_NewEventGraph))
 		{
 			StaticSectionIDs.Add(NodeSectionID::GRAPH);
 		}
-		if ( BlueprintEditor->NewDocument_IsVisibleForType(FBlueprintEditor::CGT_NewMacroGraph) )
+		if (!bIsEditor || BlueprintEditor->NewDocument_IsVisibleForType(FBlueprintEditor::CGT_NewMacroGraph))
 		{
 			StaticSectionIDs.Add(NodeSectionID::MACRO);
 		}
-		if ( BlueprintEditor->NewDocument_IsVisibleForType(FBlueprintEditor::CGT_NewFunctionGraph) )
+		if (!bIsEditor || BlueprintEditor->NewDocument_IsVisibleForType(FBlueprintEditor::CGT_NewFunctionGraph))
 		{
 			StaticSectionIDs.Add(NodeSectionID::FUNCTION);
 		}
-		if ( BlueprintEditor->NewDocument_IsVisibleForType(FBlueprintEditor::CGT_NewVariable) )
+		if (!bIsEditor || BlueprintEditor->NewDocument_IsVisibleForType(FBlueprintEditor::CGT_NewVariable))
 		{
 			StaticSectionIDs.Add(NodeSectionID::VARIABLE);
 		}
-		if ( BlueprintEditor->FBlueprintEditor::AddNewDelegateIsVisible() )
+		if (!bIsEditor || BlueprintEditor->FBlueprintEditor::AddNewDelegateIsVisible())
 		{
 			StaticSectionIDs.Add(NodeSectionID::DELEGATE);
 		}
@@ -1808,7 +1805,7 @@ TSharedPtr<SWidget> SMyBlueprint::OnContextMenuOpening()
 											LOCTEXT("AddEventSubMenu_ToolTip", "Add Event"), 
 											FNewMenuDelegate::CreateStatic(	&SSCSEditor::BuildMenuEventsSection,
 												BlueprintEditor->GetBlueprintObj(), ComponentProperty->PropertyClass, 
-												FCanExecuteAction::CreateSP(BlueprintEditor.Get(), &FBlueprintEditor::InEditingMode),
+												FCanExecuteAction::CreateRaw(this, &SMyBlueprint::IsEditingMode),
 												FGetSelectedObjectsDelegate::CreateSP(this, &SMyBlueprint::GetSelectedItemsForContextMenu)));
 				}
 			}
@@ -2092,6 +2089,12 @@ void SMyBlueprint::OnObjectPropertyChanged(UObject* InObject, FPropertyChangedEv
 	bNeedsRefresh = ( InObject == Blueprint );
 }
 
+bool SMyBlueprint::IsEditingMode() const
+{
+	auto BlueprintEditorSPtr = BlueprintEditorPtr.Pin();
+	return BlueprintEditorSPtr.IsValid() && BlueprintEditorSPtr->InEditingMode();
+}
+
 void SMyBlueprint::OnDeleteDelegate(FEdGraphSchemaAction_K2Delegate* InDelegateAction)
 {
 	UEdGraph* EdGraph = InDelegateAction->EdGraph;
@@ -2279,7 +2282,7 @@ struct FDeleteEntryHelper
 bool SMyBlueprint::CanDeleteEntry() const
 {
 	// Cannot delete entries while not in editing mode
-	if(!BlueprintEditorPtr.Pin()->InEditingMode())
+	if(!IsEditingMode())
 	{
 		return false;
 	}
diff --git a/Engine/Source/Editor/Kismet/Private/SMyBlueprint.h b/Engine/Source/Editor/Kismet/Private/SMyBlueprint.h
--- a/Engine/Source/Editor/Kismet/Private/SMyBlueprint.h
+++ b/Engine/Source/Editor/Kismet/Private/SMyBlueprint.h
@@ -230,6 +230,9 @@ private:
 
 	/** Delegate to hook us into non-structural Blueprint object post-change events */
 	void OnObjectPropertyChanged(UObject* InObject, FPropertyChangedEvent& InPropertyChangedEvent);
+
+	/** Helper function indicating whehter we're in editing mode, and can modify the target blueprint */
+	bool IsEditingMode() const;
 private:
 	/** Pointer back to the blueprint editor that owns us */
 	TWeakPtr<FBlueprintEditor> BlueprintEditorPtr;
