Initializing SKismetInspector::bRefreshOnTick, it is assigned before we tick, but we don't want to rely on that kind of coincidence...

diff --git a/Engine/Source/Editor/Kismet/Private/SKismetInspector.cpp b/Engine/Source/Editor/Kismet/Private/SKismetInspector.cpp
--- a/Engine/Source/Editor/Kismet/Private/SKismetInspector.cpp
+++ b/Engine/Source/Editor/Kismet/Private/SKismetInspector.cpp
@@ -224,6 +224,7 @@ void SKismetInspector::Construct(const FArguments& InArgs)
 	bShowInspectorPropertyView = true;
 	PublicViewState = ECheckBoxState::Unchecked;
 	bComponenetDetailsCustomizationEnabled = false;
+	bRefreshOnTick = false;
 
 	BlueprintEditorPtr = InArgs._Kismet2;
 	bShowPublicView = InArgs._ShowPublicViewControl;
