[INTEGRATE] Change 2439709 by [REDACTED_PERSON][REDACTED_HANDLE] on 2015/02/10 08:56:08

	UE-9201 - BP - Selecting a component in the SCS editor now adds it to the context sensitive right click menu for blueprints, just as if you'd selected it in the my blueprints view.

diff --git a/Engine/Source/Editor/Kismet/Private/SBlueprintActionMenu.cpp b/Engine/Source/Editor/Kismet/Private/SBlueprintActionMenu.cpp
--- a/Engine/Source/Editor/Kismet/Private/SBlueprintActionMenu.cpp
+++ b/Engine/Source/Editor/Kismet/Private/SBlueprintActionMenu.cpp
@@ -405,11 +405,23 @@ void SBlueprintActionMenu::CollectAllActions(FGraphActionListBuilderBase& OutAll
 	{
 		FilterContext.Pins = DraggedFromPins;
 		
+		// Get selection from the "My Blueprint" view.
 		FEdGraphSchemaAction_K2Var* SelectedVar = BlueprintEditor->GetMyBlueprintWidget()->SelectionAsVar();
 		if ((SelectedVar != nullptr) && (SelectedVar->GetProperty() != nullptr))
 		{
 			FilterContext.SelectedObjects.Add(SelectedVar->GetProperty());
 		}
+		// If the selection come from the SCS editor, add it to the filter context.
+		else if ( Blueprint->SkeletonGeneratedClass )
+		{
+			TArray<FSCSEditorTreeNodePtrType> Nodes = BlueprintEditor->GetSCSEditor()->GetSelectedNodes();
+			if ( Nodes.Num() == 1 && Nodes[0]->GetNodeType() == FSCSEditorTreeNode::ComponentNode )
+			{
+				FName PropertyName = Nodes[0]->GetVariableName();
+				UObjectProperty* VariableProperty = FindField<UObjectProperty>(Blueprint->SkeletonGeneratedClass, PropertyName);
+				FilterContext.SelectedObjects.Add(VariableProperty);
+			}
+		}
 	}
 	
 	FBlueprintActionMenuBuilder MenuBuilder(EditorPtr);
