[INTEGRATE] Change 2439710 by [REDACTED_PERSON][REDACTED_HANDLE] on 2015/02/10 08:58:13

	Editor - Working on the blueprint details customization to present events a lot differently.  Now they are a list of buttons that can be used to add events for a component, instead of a dropdown combobox.  The expansion presents a much more eye catching bit of UI.

diff --git a/Engine/Source/Editor/Kismet/Private/BlueprintDetailsCustomization.cpp b/Engine/Source/Editor/Kismet/Private/BlueprintDetailsCustomization.cpp
--- a/Engine/Source/Editor/Kismet/Private/BlueprintDetailsCustomization.cpp
+++ b/Engine/Source/Editor/Kismet/Private/BlueprintDetailsCustomization.cpp
@@ -1,7 +1,7 @@
 // Copyright 1998-2015 Epic Games, Inc. All Rights Reserved.
 
-
 #include "BlueprintEditorPrivatePCH.h"
+
 #include "BlueprintUtilities.h"
 #include "ScopedTransaction.h"
 #include "GraphEditor.h"
@@ -20,6 +20,7 @@
 #include "SGraphActionMenu.h"
 #include "SPinTypeSelector.h"
 #include "Kismet2NameValidators.h"
+#include "SWidgetSwitcher.h"
 
 #include "ComponentAssetBroker.h"
 #include "PropertyCustomizationHelpers.h"
@@ -43,6 +44,129 @@
 
 #define LOCTEXT_NAMESPACE "BlueprintDetailsCustomization"
 
+void FBlueprintDetails::AddEventsCategory(IDetailLayoutBuilder& DetailBuilder, UProperty* VariableProperty)
+{
+	UBlueprint* Blueprint = GetBlueprintObj();
+	check(Blueprint);
+
+	if ( UObjectProperty* ComponentProperty = Cast<UObjectProperty>(VariableProperty) )
+	{
+		UClass* PropertyClass = ComponentProperty->PropertyClass;
+
+		// Check for Ed Graph vars that can generate events
+		if ( PropertyClass && Blueprint->AllowsDynamicBinding() )
+		{
+			if ( FBlueprintEditorUtils::CanClassGenerateEvents(PropertyClass) )
+			{
+				for ( TFieldIterator<UMulticastDelegateProperty> PropertyIt(PropertyClass, EFieldIteratorFlags::IncludeSuper); PropertyIt; ++PropertyIt )
+				{
+					UProperty* Property = *PropertyIt;
+
+					FName PropertyName = ComponentProperty->GetFName();
+
+					// Check for multicast delegates that we can safely assign
+					if ( !Property->HasAnyPropertyFlags(CPF_Parm) && Property->HasAllPropertyFlags(CPF_BlueprintAssignable) )
+					{
+						FName EventName = Property->GetFName();
+						FText EventText = Property->GetDisplayNameText();
+
+						IDetailCategoryBuilder& EventCategory = DetailBuilder.EditCategory(TEXT("Events"), LOCTEXT("Events", "Events"), ECategoryPriority::Uncommon);
+
+						EventCategory.AddCustomRow(EventText)
+						.NameContent()
+						[
+							SNew(SHorizontalBox)
+							.ToolTipText(Property->GetToolTipText())
+
+							+ SHorizontalBox::Slot()
+							.AutoWidth()
+							.VAlign(VAlign_Center)
+							.Padding(0, 0, 5, 0)
+							[
+								SNew(SImage)
+								.Image(FEditorStyle::GetBrush("GraphEditor.Event_16x"))
+							]
+
+							+ SHorizontalBox::Slot()
+							.VAlign(VAlign_Center)
+							[
+								SNew(STextBlock)
+								.Font(IDetailLayoutBuilder::GetDetailFont())
+								.Text(EventText)
+							]
+						]
+						.ValueContent()
+						.MinDesiredWidth(150)
+						.MaxDesiredWidth(200)
+						[
+							SNew(SButton)
+							.ButtonStyle(FEditorStyle::Get(), "FlatButton.Success")
+							.HAlign(HAlign_Center)
+							.OnClicked(this, &FBlueprintVarActionDetails::HandleAddOrViewEventForVariable, EventName, PropertyName, TWeakObjectPtr<UClass>(PropertyClass))
+							.ForegroundColor(FSlateColor::UseForeground())
+							[
+								SNew(SWidgetSwitcher)
+								.WidgetIndex(this, &FBlueprintVarActionDetails::HandleAddOrViewIndexForButton, EventName, PropertyName)
+
+								+ SWidgetSwitcher::Slot()
+								[
+									SNew(STextBlock)
+									.Font(FEditorStyle::GetFontStyle(TEXT("BoldFont")))
+									.Text(LOCTEXT("ViewEvent", "View"))
+								]
+
+								+ SWidgetSwitcher::Slot()
+								[
+									SNew(SImage)
+									.Image(FEditorStyle::GetBrush("Plus"))
+								]
+							]
+						];
+					}
+				}
+			}
+		}
+	}
+}
+
+FReply FBlueprintDetails::HandleAddOrViewEventForVariable(const FName EventName, FName PropertyName, TWeakObjectPtr<UClass> PropertyClass)
+{
+	UBlueprint* Blueprint = GetBlueprintObj();
+
+	// Find the corresponding variable property in the Blueprint
+	UObjectProperty* VariableProperty = FindField<UObjectProperty>(Blueprint->SkeletonGeneratedClass, PropertyName);
+
+	if ( VariableProperty )
+	{
+		if ( !FKismetEditorUtilities::FindBoundEventForComponent(Blueprint, EventName, VariableProperty->GetFName()) )
+		{
+			FKismetEditorUtilities::CreateNewBoundEventForClass(PropertyClass.Get(), EventName, Blueprint, VariableProperty);
+		}
+		else
+		{
+			const UK2Node_ComponentBoundEvent* ExistingNode = FKismetEditorUtilities::FindBoundEventForComponent(Blueprint, EventName, VariableProperty->GetFName());
+			if ( ExistingNode )
+			{
+				FKismetEditorUtilities::BringKismetToFocusAttentionOnObject(ExistingNode);
+			}
+		}
+	}
+
+	return FReply::Handled();
+}
+
+int32 FBlueprintDetails::HandleAddOrViewIndexForButton(const FName EventName, FName PropertyName) const
+{
+	UBlueprint* Blueprint = GetBlueprintObj();
+
+	if ( FKismetEditorUtilities::FindBoundEventForComponent(Blueprint, EventName, PropertyName) )
+	{
+		return 0; // View
+	}
+
+	return 1; // Add
+}
+
 // UProperty Detail Customization
 BEGIN_SLATE_FUNCTION_BUILD_OPTIMIZATION
 void FBlueprintVarActionDetails::CustomizeDetails( IDetailLayoutBuilder& DetailLayout )
@@ -403,67 +527,15 @@ void FBlueprintVarActionDetails::CustomizeDetails( IDetailLayoutBuilder& DetailL
 		.ToolTip(ReplicationTooltip)
 	];
 
-		// Handle event generation
-	if(FBlueprintEditorUtils::DoesSupportEventGraphs(GetBlueprintObj()))
-	{
-		if (UObjectProperty* ComponentProperty = Cast<UObjectProperty>(VariableProperty))
-		{
-			// Check for Ed Graph vars that can generate events
-			if (ComponentProperty->PropertyClass && ComponentProperty->PropertyClass->IsChildOf(UActorComponent::StaticClass()))
-			{
-				// Add Events Section property class that has valid events
-				if( FBlueprintEditorUtils::CanClassGenerateEvents( ComponentProperty->PropertyClass ))
-				{
-					IDetailCategoryBuilder& EventCategory = DetailLayout.EditCategory("Component", LOCTEXT("ComponentDetailsCategory", "Events"));
-
-					FText NodeName = FText::FromName( VarName );
-					const FText AddEventLabel = FText::Format( LOCTEXT( "ScriptingEvents_AddEvent", "Add Event For {0}" ), NodeName );
-					const FText AddEventToolTip = FText::Format( LOCTEXT( "ScriptingEvents_AddOrView", "Adds or views events for {0} in the Component Blueprint" ), NodeName );
-
-					TSharedPtr<SToolTip> EventsTooltip = IDocumentation::Get()->CreateToolTip(AddEventToolTip, NULL, DocLink, TEXT("Events"));
+	UBlueprint* Blueprint = GetBlueprintObj();
 
-					EventCategory.AddCustomRow( LOCTEXT("AddEventHeader", "Add Event"))
-						.NameContent()
-						[
-							SNew(SHorizontalBox)
-							+SHorizontalBox::Slot()
-							.Padding( FMargin( 2.0f, 0, 0, 0 ) )
-							.VAlign( VAlign_Center )
-							.AutoWidth()
-							[
-								SNew( STextBlock )
-								.Text( LOCTEXT("ScriptingEvents_Label", "Add Event"))
-								.ToolTip(EventsTooltip)
-								.Font( IDetailLayoutBuilder::GetDetailFont() )
-							]
-						]
-					.ValueContent()
-						[
-							SNew(SHorizontalBox)
-							+SHorizontalBox::Slot()
-							.FillWidth( 1 )
-							.VAlign( VAlign_Center )
-							[
-								SNew( SComboButton )
-								.HAlign(HAlign_Fill)
-								.VAlign(VAlign_Center)
-								.OnGetMenuContent( this, &FBlueprintVarActionDetails::BuildEventsMenuForVariable )
-								.ButtonContent()
-								[
-									SNew( STextBlock )
-									.Text( AddEventLabel ) 
-									.ToolTip( EventsTooltip )
-									.Font( IDetailLayoutBuilder::GetDetailFont() )
-								]
-							]
-						];
-				}
-			}
-		}
+	// Handle event generation
+	if ( FBlueprintEditorUtils::DoesSupportEventGraphs(Blueprint) )
+	{
+		AddEventsCategory(DetailLayout, VariableProperty);
 	}
 
 	// Add in default value editing for properties that can be edited, local properties cannot be edited
-	UBlueprint* Blueprint = GetBlueprintObj();
 	if ((Blueprint != NULL) && (Blueprint->GeneratedClass != NULL))
 	{
 		if (VariableProperty != NULL)
@@ -663,6 +735,7 @@ void FBlueprintVarActionDetails::CustomizeDetails( IDetailLayoutBuilder& DetailL
 		RefreshPropertyFlags();
 	}
 }
+
 END_SLATE_FUNCTION_BUILD_OPTIMIZATION
 
 void FBlueprintVarActionDetails::RefreshPropertyFlags()
@@ -4335,8 +4408,6 @@ void FBlueprintGlobalOptionsDetails::CustomizeDetails(IDetailLayoutBuilder& Deta
 	}
 }
 
-
-
 void FBlueprintComponentDetails::CustomizeDetails(IDetailLayoutBuilder& DetailLayout)
 {
 	check( BlueprintEditorPtr.IsValid() );
@@ -4480,51 +4551,14 @@ void FBlueprintComponentDetails::CustomizeDetails(IDetailLayoutBuilder& DetailLa
 			]
 		];
 	}
-	// Add Events Section if we have a common base class that has valid events
-	UClass* CommonEventsClass = FindCommonBaseClassFromSelected();
 
-	if( FBlueprintEditorUtils::CanClassGenerateEvents( CommonEventsClass ))
+	// Handle event generation
+	if ( FBlueprintEditorUtils::DoesSupportEventGraphs(Blueprint) && Nodes.Num() == 1 )
 	{
-		IDetailCategoryBuilder& EventCategory = DetailLayout.EditCategory("Component", LOCTEXT("ComponentDetailsCategory", "Events"), ECategoryPriority::Important);
+		FName PropertyName = CachedNodePtr->GetVariableName();
+		UObjectProperty* VariableProperty = FindField<UObjectProperty>(Blueprint->SkeletonGeneratedClass, PropertyName);
 
-		FText NodeName = ( Nodes.Num() > 1 ) ? LOCTEXT( "ScriptingEvents_SelectedComponents", "Selected Components" ) : FText::FromName( CachedNodePtr->GetVariableName() );
-		const FText AddEventLabel = FText::Format( LOCTEXT( "ScriptingEvents_AddEvent", "Add Event For {0}" ), NodeName );
-		const FText AddEventToolTip = FText::Format( LOCTEXT( "ScriptingEvents_AddOrView", "Adds or views events for {0} in the Component Blueprint" ), NodeName );
-
-		EventCategory.AddCustomRow( LOCTEXT("AddEventHeader", "Add Event"))
-			.NameContent()
-			[
-				SNew(SHorizontalBox)
-				+SHorizontalBox::Slot()
-				.Padding( FMargin( 2.0f, 0, 0, 0 ) )
-				.VAlign( VAlign_Center )
-				.AutoWidth()
-				[
-					SNew( STextBlock )
-						.Text( LOCTEXT("ScriptingEvents_Label", "Add Event")  )
-						.Font( IDetailLayoutBuilder::GetDetailFont() )
-				]
-			]
-			.ValueContent()
-			[
-				SNew(SHorizontalBox)
-				+SHorizontalBox::Slot()
-				.FillWidth( 1 )
-				.VAlign( VAlign_Center )
-				[
-					SNew( SComboButton )
-						.HAlign(HAlign_Fill)
-						.VAlign(VAlign_Center)
-						.OnGetMenuContent( this, &FBlueprintComponentDetails::BuildEventsMenuForComponents)
-						.ButtonContent()
-					[
-						SNew( STextBlock )
-							.Text( AddEventLabel ) 
-							.ToolTipText( AddEventToolTip )
-							.Font( IDetailLayoutBuilder::GetDetailFont() )
-					]
-				]
-			];
+		AddEventsCategory(DetailLayout, VariableProperty);
 	}
 
 	// Don't show tick properties for components in the blueprint details
@@ -4532,54 +4566,6 @@ void FBlueprintComponentDetails::CustomizeDetails(IDetailLayoutBuilder& DetailLa
 	PrimaryTickProperty->MarkHiddenByCustomization();
 }
 
-UClass* FBlueprintComponentDetails::FindCommonBaseClassFromSelected() const
-{
-	UClass* Result = NULL;
-
-	if( BlueprintEditorPtr.IsValid() )
-	{
-		TSharedPtr<SSCSEditor> Editor = BlueprintEditorPtr.Pin()->GetSCSEditor();
-		check( Editor.IsValid() );
-		// Find common component class from selected
-		TArray<UClass*> SelectionClasses;
-		TArray<FSCSEditorTreeNodePtrType> SelectedNodes = Editor->GetSelectedNodes();
-
-		for (auto NodeIter = SelectedNodes.CreateConstIterator(); NodeIter; ++NodeIter)
-		{
-			auto Node = *NodeIter;
-			if( Node.IsValid() && Node->GetComponentTemplate() )
-			{
-				SelectionClasses.Add( Node->GetComponentTemplate()->GetClass() );
-			}
-		}
-		Result = SelectionClasses.Num() ? UClass::FindCommonBase(SelectionClasses) : nullptr;
-	}
-	return Result;
-}
-
-TSharedRef<SWidget> FBlueprintComponentDetails::BuildEventsMenuForComponents() const
-{
-	// See if the component class has event properties
-	if( BlueprintEditorPtr.IsValid() )
-	{
-		TSharedPtr<SSCSEditor> Editor = BlueprintEditorPtr.Pin()->GetSCSEditor();
-
-		if( Editor.IsValid() )
-		{
-			UClass* CommonComponentClass = FindCommonBaseClassFromSelected();
-			if( CommonComponentClass )
-			{
-				FMenuBuilder MenuBuilder( true, NULL );
-				Editor->BuildMenuEventsSection( MenuBuilder, BlueprintEditorPtr.Pin()->GetBlueprintObj(), CommonComponentClass, 
-												FCanExecuteAction::CreateSP(BlueprintEditorPtr.Pin().Get(), &FBlueprintEditor::InEditingMode),
-												FGetSelectedObjectsDelegate::CreateSP(Editor.Get(), &SSCSEditor::GetSelectedItemsForContextMenu));
-				return MenuBuilder.MakeWidget();
-			}
-		}
-	}
-	return SNullWidget::NullWidget;
-}
-
 FText FBlueprintComponentDetails::OnGetVariableText() const
 {
 	check(CachedNodePtr.IsValid());
diff --git a/Engine/Source/Editor/Kismet/Private/BlueprintDetailsCustomization.h b/Engine/Source/Editor/Kismet/Private/BlueprintDetailsCustomization.h
--- a/Engine/Source/Editor/Kismet/Private/BlueprintDetailsCustomization.h
+++ b/Engine/Source/Editor/Kismet/Private/BlueprintDetailsCustomization.h
@@ -33,8 +33,34 @@ namespace EVariableReplication
 	};
 }
 
+class FBlueprintDetails : public IDetailCustomization
+{
+public:
+	FBlueprintDetails(TWeakPtr<SMyBlueprint> InMyBlueprint)
+		: Blueprint(InMyBlueprint.Pin()->GetBlueprintObj())
+	{
+	}
+
+	FBlueprintDetails(TWeakPtr<FBlueprintEditor> InBlueprintEditorPtr)
+		: Blueprint(InBlueprintEditorPtr.Pin()->GetBlueprintObj())
+	{
+	}
+
+protected:
+
+	UBlueprint* GetBlueprintObj() const { return Blueprint.Get(); }
+
+	void AddEventsCategory(IDetailLayoutBuilder& DetailBuilder, UProperty* VariableProperty);
+	FReply HandleAddOrViewEventForVariable(const FName EventName, FName PropertyName, TWeakObjectPtr<UClass> PropertyClass);
+	int32 HandleAddOrViewIndexForButton(const FName EventName, FName PropertyName) const;
+
+private:
+	/** Pointer back to my parent tab */
+	TWeakObjectPtr<UBlueprint> Blueprint;
+};
+
 /** Details customization for variables selected in the MyBlueprint panel */
-class FBlueprintVarActionDetails : public IDetailCustomization
+class FBlueprintVarActionDetails : public FBlueprintDetails
 {
 public:
 	/** Makes a new instance of this detail layout class for a specific detail view requesting it */
@@ -44,7 +70,8 @@ public:
 	}
 
 	FBlueprintVarActionDetails(TWeakPtr<SMyBlueprint> InMyBlueprint)
-		: MyBlueprint(InMyBlueprint)
+		: FBlueprintDetails(InMyBlueprint)
+		, MyBlueprint(InMyBlueprint)
 		, bIsVarNameInvalid(false)
 	{
 	}
@@ -56,7 +83,6 @@ public:
 
 private:
 	/** Accessors passed to parent */
-	UBlueprint* GetBlueprintObj() const {return MyBlueprint.Pin()->GetBlueprintObj();}
 	FEdGraphSchemaAction_K2Var* MyBlueprintSelectionAsVar() const {return MyBlueprint.Pin()->SelectionAsVar();}
 	FEdGraphSchemaAction_K2LocalVar* MyBlueprintSelectionAsLocalVar() const {return MyBlueprint.Pin()->SelectionAsLocalVar();}
 	UK2Node_Variable* EdGraphSelectionAsVar() const;
@@ -646,12 +672,13 @@ private:
 
 
 /** Details customization for Blueprint Component settings */
-class FBlueprintComponentDetails : public IDetailCustomization
+class FBlueprintComponentDetails : public FBlueprintDetails
 {
 public:
 	/** Constructor */
 	FBlueprintComponentDetails(TWeakPtr<FBlueprintEditor> InBlueprintEditorPtr)
-		: BlueprintEditorPtr(InBlueprintEditorPtr)
+		: FBlueprintDetails(InBlueprintEditorPtr)
+		, BlueprintEditorPtr(InBlueprintEditorPtr)
 		, bIsVariableNameInvalid(false)
 	{
 
@@ -667,9 +694,6 @@ public:
 	virtual void CustomizeDetails( IDetailLayoutBuilder& DetailLayout ) override;
 
 protected:
-	/** Accessors passed to parent */
-	UBlueprint* GetBlueprintObj() const {return BlueprintEditorPtr.Pin()->GetBlueprintObj();}
-
 	/** Callbacks for widgets */
 	FText OnGetVariableText() const;
 	void OnVariableTextChanged(const FText& InNewText);
@@ -682,12 +706,6 @@ protected:
 	void OnVariableCategorySelectionChanged(TSharedPtr<FString> ProposedSelection, ESelectInfo::Type /*SelectInfo*/);
 	TSharedRef<ITableRow> MakeVariableCategoryViewWidget(TSharedPtr<FString> Item, const TSharedRef< STableViewBase >& OwnerTable);
 
-	/** Find common base class from current selection */
-	UClass* FindCommonBaseClassFromSelected() const;
-
-	/** Build Event Menu for currently selected components */
-	TSharedRef<SWidget> BuildEventsMenuForComponents() const;
-
 	FText GetSocketName() const;
 	void OnBrowseSocket();
 	void OnClearSocket();
