[INTEGRATE] Change 2439825 by [REDACTED_PERSON][REDACTED_HANDLE] on 2015/02/10 10:47:19

	- Tweaked title of "Convert to Blueprint" button in an effort to describe what it does.
	- Added padding to the Add Component, Add script buttons.

diff --git a/Engine/Source/Editor/Kismet/Private/SSCSEditor.cpp b/Engine/Source/Editor/Kismet/Private/SSCSEditor.cpp
--- a/Engine/Source/Editor/Kismet/Private/SSCSEditor.cpp
+++ b/Engine/Source/Editor/Kismet/Private/SSCSEditor.cpp
@@ -3150,6 +3150,7 @@ void SSCSEditor::Construct( const FArguments& InArgs )
 				SNew(SHorizontalBox)
 				
 				+ SHorizontalBox::Slot()
+				.Padding( 3.0f, 3.0f )
 				.AutoWidth()
 				.HAlign(HAlign_Left)
 				[
@@ -3163,6 +3164,7 @@ void SSCSEditor::Construct( const FArguments& InArgs )
 				+ SHorizontalBox::Slot()
 				.FillWidth(1.0f)
 				.HAlign(HAlign_Right)
+				.Padding( 3.0f, 3.0f )
 				[
 					SNew( SButton )
 					.AddMetaData<FTagMetaData>(FTagMetaData(TEXT("Actor.ConvertToBlueprint")))
@@ -3206,7 +3208,7 @@ void SSCSEditor::Construct( const FArguments& InArgs )
 							SNew(STextBlock)
 							.TextStyle(FEditorStyle::Get(), "ContentBrowser.TopBar.Font")
 							//.Text( LOCTEXT("PromoteToBlueprint", "Add Script") )
-							.Text(LOCTEXT("PromoteToBlueprint", "Convert to Blueprint Class"))
+							.Text(LOCTEXT("PromoteToBlueprint", "Blueprint/Add Script"))
 						]
 					]
 				]
