[INTEGRATE] Change 2439860 by [REDACTED_PERSON][REDACTED_HANDLE] on 2015/02/10 11:21:04

	BP - Improving the styling of the debug object drop down menu a bit.

diff --git a/Engine/Source/Editor/Kismet/Private/SBlueprintEditorSelectedDebugObjectWidget.cpp b/Engine/Source/Editor/Kismet/Private/SBlueprintEditorSelectedDebugObjectWidget.cpp
--- a/Engine/Source/Editor/Kismet/Private/SBlueprintEditorSelectedDebugObjectWidget.cpp
+++ b/Engine/Source/Editor/Kismet/Private/SBlueprintEditorSelectedDebugObjectWidget.cpp
@@ -40,6 +40,7 @@ void SBlueprintEditorSelectedDebugObjectWidget::Construct(const FArguments& InAr
 	LastObjectObserved = DebugObjects[0];
 
 	DebugWorldsComboBox = SNew(STextComboBox)
+		.ButtonStyle(FEditorStyle::Get(), "FlatButton.Light")
 		.ToolTip(IDocumentation::Get()->CreateToolTip(
 		LOCTEXT("BlueprintDebugWorldTooltip", "Select a world to debug"),
 		nullptr,
@@ -52,6 +53,7 @@ void SBlueprintEditorSelectedDebugObjectWidget::Construct(const FArguments& InAr
 		.OnSelectionChanged(this, &SBlueprintEditorSelectedDebugObjectWidget::DebugWorldSelectionChanged);
 
 	DebugObjectsComboBox = SNew(STextComboBox)
+		.ButtonStyle(FEditorStyle::Get(), "FlatButton.Light")
 		.ToolTip(IDocumentation::Get()->CreateToolTip(
 		LOCTEXT("BlueprintDebugObjectTooltip", "Select an object to debug"),
 		nullptr,
