[INTEGRATE] Change 2440438 by [REDACTED_PERSON][REDACTED_HANDLE] on 2015/02/10 16:54:59

	Fixed crash getting context menu items from the selected component in the SCS editor when the blueprint does not have an SCS editor (not actor based)

diff --git a/Engine/Source/Editor/Kismet/Private/SBlueprintActionMenu.cpp b/Engine/Source/Editor/Kismet/Private/SBlueprintActionMenu.cpp
--- a/Engine/Source/Editor/Kismet/Private/SBlueprintActionMenu.cpp
+++ b/Engine/Source/Editor/Kismet/Private/SBlueprintActionMenu.cpp
@@ -412,7 +412,7 @@ void SBlueprintActionMenu::CollectAllActions(FGraphActionListBuilderBase& OutAll
 			FilterContext.SelectedObjects.Add(SelectedVar->GetProperty());
 		}
 		// If the selection come from the SCS editor, add it to the filter context.
-		else if ( Blueprint->SkeletonGeneratedClass )
+		else if ( Blueprint->SkeletonGeneratedClass && BlueprintEditor->GetSCSEditor().IsValid() )
 		{
 			TArray<FSCSEditorTreeNodePtrType> Nodes = BlueprintEditor->GetSCSEditor()->GetSelectedNodes();
 			if ( Nodes.Num() == 1 && Nodes[0]->GetNodeType() == FSCSEditorTreeNode::ComponentNode )
