[INTEGRATE] Change 2440497 by [REDACTED_PERSON][REDACTED_HANDLE] on 2015/02/10 17:25:29

	UE-9549 - The SCSEditor is no longer portrayed in the API as if it's guaranteed to exist.
	UE-9368 - We're now broadcasting foregrounding tab changes from the global tab manager.  These changes are now watched by the workflow tab manager to then relay that information to anyone who cares.  The blueprint editor now uses this information to have graph tabs properly clear the active graph ptr when a graph tab is no longer in the foreground.

diff --git a/Engine/Source/Editor/EditorStyle/Private/SlateEditorStyle.cpp b/Engine/Source/Editor/EditorStyle/Private/SlateEditorStyle.cpp
--- a/Engine/Source/Editor/EditorStyle/Private/SlateEditorStyle.cpp
+++ b/Engine/Source/Editor/EditorStyle/Private/SlateEditorStyle.cpp
@@ -207,6 +207,23 @@ void FSlateEditorStyle::FStyle::SetupGeneralStyles()
 
 		Set("NormalText", NormalText);
 
+		Set("NormalText.Subdued", FTextBlockStyle(NormalText)
+			.SetColorAndOpacity(FSlateColor::UseSubduedForeground()));
+
+		Set("SmallText", FTextBlockStyle(NormalText)
+			.SetFont(TTF_CORE_FONT("Fonts/Roboto-Regular", 9)));
+
+		Set("SmallText.Subdued", FTextBlockStyle(NormalText)
+			.SetFont(TTF_CORE_FONT("Fonts/Roboto-Regular", 9))
+			.SetColorAndOpacity(FSlateColor::UseSubduedForeground()));
+
+		Set("TinyText", FTextBlockStyle(NormalText)
+			.SetFont(TTF_CORE_FONT("Fonts/Roboto-Regular", 8)));
+
+		Set("TinyText.Subdued", FTextBlockStyle(NormalText)
+			.SetFont(TTF_CORE_FONT("Fonts/Roboto-Regular", 8))
+			.SetColorAndOpacity(FSlateColor::UseSubduedForeground()));
+
 		Set("LargeText", FTextBlockStyle(NormalText)
 			.SetFont(TTF_CORE_FONT("Fonts/Roboto-Bold", 11))
 			.SetColorAndOpacity(FLinearColor(1.0f, 1.0f, 1.0f))
@@ -3098,6 +3115,12 @@ void FSlateEditorStyle::FStyle::SetupPropertyEditorStyles()
 		Set( "DetailsView.AdvancedDropdownBorder.Open", new IMAGE_BRUSH( "Common/ScrollBoxShadowTop", FVector2D(64,8) ) );
 		Set( "DetailsView.CategoryFontStyle", TTF_CORE_FONT( "Fonts/Roboto-Bold", 10 ) );
 
+		Set( "DetailsView.CategoryTextStyle", 
+			FTextBlockStyle(NormalText)
+			.SetFont(GetFontStyle("DetailsView.CategoryFontStyle"))
+			.SetShadowOffset(FVector2D(1.0f, 1.0f))
+		);
+
 		Set( "DetailsView.TreeView.TableRow", FTableRowStyle()
 			.SetEvenRowBackgroundBrush( FSlateNoResource() )
 			.SetEvenRowBackgroundHoveredBrush( FSlateNoResource() )
diff --git a/Engine/Source/Editor/GraphEditor/Private/SGraphActionMenu.cpp b/Engine/Source/Editor/GraphEditor/Private/SGraphActionMenu.cpp
--- a/Engine/Source/Editor/GraphEditor/Private/SGraphActionMenu.cpp
+++ b/Engine/Source/Editor/GraphEditor/Private/SGraphActionMenu.cpp
@@ -1056,10 +1056,10 @@ TSharedRef<ITableRow> SGraphActionMenu::MakeWidget( TSharedPtr<FGraphActionNode>
 			+ SHorizontalBox::Slot()
 			.VAlign(VAlign_Center)
 			[
-				SNew(STextBlock)
+				SNew(SRichTextBlock)
 				.Text(SectionTitle)
-				.Font(FEditorStyle::GetFontStyle("DetailsView.CategoryFontStyle"))
-				.ShadowOffset(FVector2D(1.0f, 1.0f))
+				.DecoratorStyleSet(&FEditorStyle::Get())
+				.TextStyle(FEditorStyle::Get(), "DetailsView.CategoryTextStyle")
 			]
 
 			+ SHorizontalBox::Slot()
diff --git a/Engine/Source/Editor/Kismet/Private/BlueprintDetailsCustomization.cpp b/Engine/Source/Editor/Kismet/Private/BlueprintDetailsCustomization.cpp
--- a/Engine/Source/Editor/Kismet/Private/BlueprintDetailsCustomization.cpp
+++ b/Engine/Source/Editor/Kismet/Private/BlueprintDetailsCustomization.cpp
@@ -1728,7 +1728,7 @@ TSharedRef<SWidget> FBlueprintVarActionDetails::BuildEventsMenuForVariable() con
 		TWeakPtr<FBlueprintEditor> BlueprintEditorPtr = MyBlueprintPtr->GetBlueprintEditor();
 		if( BlueprintEditorPtr.IsValid() && ComponentProperty )
 		{
-			TSharedRef<SSCSEditor> Editor =  BlueprintEditorPtr.Pin()->GetSCSEditor();
+			TSharedPtr<SSCSEditor> Editor =  BlueprintEditorPtr.Pin()->GetSCSEditor();
 			FMenuBuilder MenuBuilder( true, NULL );
 			Editor->BuildMenuEventsSection( MenuBuilder, BlueprintEditorPtr.Pin()->GetBlueprintObj(), ComponentProperty->PropertyClass, 
 											FCanExecuteAction::CreateSP(BlueprintEditorPtr.Pin().Get(), &FBlueprintEditor::InEditingMode),
diff --git a/Engine/Source/Editor/Kismet/Private/BlueprintEditor.cpp b/Engine/Source/Editor/Kismet/Private/BlueprintEditor.cpp
--- a/Engine/Source/Editor/Kismet/Private/BlueprintEditor.cpp
+++ b/Engine/Source/Editor/Kismet/Private/BlueprintEditor.cpp
@@ -2744,29 +2744,8 @@ bool FBlueprintEditor::CanRedoGraphAction() const
 	return !InDebuggingMode();
 }
 
-void FBlueprintEditor::OnActiveTabChanged( TSharedPtr<SDockTab> PreviouslyActive, TSharedPtr<SDockTab> NewlyActivated )
+void FBlueprintEditor::OnActiveTabChanged(TSharedPtr<SDockTab> PreviouslyActive, TSharedPtr<SDockTab> NewlyActivated)
 {
-	if (NewlyActivated.IsValid() == false)
-	{
-		SetUISelectionState(NAME_None);
-		//UE_LOG(LogBlueprint, Log, TEXT("OnActiveTabChanged: NONE"));
-	}
-	else
-	{
-		//UE_LOG(LogBlueprint, Log, TEXT("OnActiveTabChanged: %s"), *NewlyActivated->GetLayoutIdentifier().ToString());
-	}
-
-	// If the newly active document tab isn't a graph we want to make sure we clear the focused graph pointer.
-	// Several other UI reads that, like the MyBlueprints view uses it to determine if it should show the "Local Variable" section.
-	if ( NewlyActivated.IsValid() && NewlyActivated->GetTabRole() == ETabRole::DocumentTab )
-	{
-		FocusedGraphEdPtr = nullptr;
-
-		if ( MyBlueprintWidget.IsValid() == true )
-		{
-			MyBlueprintWidget->Refresh();
-		}
-	}
 }
 
 void FBlueprintEditor::OnGraphEditorFocused(const TSharedRef<SGraphEditor>& InGraphEditor)
@@ -2796,6 +2775,18 @@ void FBlueprintEditor::OnGraphEditorFocused(const TSharedRef<SGraphEditor>& InGr
 	}
 }
 
+void FBlueprintEditor::OnGraphEditorBackgrounded(const TSharedRef<SGraphEditor>& InGraphEditor)
+{
+	// If the newly active document tab isn't a graph we want to make sure we clear the focused graph pointer.
+	// Several other UI reads that, like the MyBlueprints view uses it to determine if it should show the "Local Variable" section.
+	FocusedGraphEdPtr = nullptr;
+
+	if ( MyBlueprintWidget.IsValid() == true )
+	{
+		MyBlueprintWidget->Refresh();
+	}
+}
+
 void FBlueprintEditor::OnGraphEditorDropActor(const TArray< TWeakObjectPtr<AActor> >& Actors, UEdGraph* Graph, const FVector2D& DropLocation)
 {
 	// We need to check that the dropped actor is in the right sublevel for the reference
diff --git a/Engine/Source/Editor/Kismet/Private/BlueprintEditorTabFactories.cpp b/Engine/Source/Editor/Kismet/Private/BlueprintEditorTabFactories.cpp
--- a/Engine/Source/Editor/Kismet/Private/BlueprintEditorTabFactories.cpp
+++ b/Engine/Source/Editor/Kismet/Private/BlueprintEditorTabFactories.cpp
@@ -29,6 +29,12 @@ void FGraphEditorSummoner::OnTabActivated(TSharedPtr<SDockTab> Tab) const
 	BlueprintEditorPtr.Pin()->OnGraphEditorFocused(GraphEditor);
 }
 
+void FGraphEditorSummoner::OnTabBackgrounded(TSharedPtr<SDockTab> Tab) const
+{
+	TSharedRef<SGraphEditor> GraphEditor = StaticCastSharedRef<SGraphEditor>(Tab->GetContent());
+	BlueprintEditorPtr.Pin()->OnGraphEditorBackgrounded(GraphEditor);
+}
+
 void FGraphEditorSummoner::OnTabRefreshed(TSharedPtr<SDockTab> Tab) const
 {
 	TSharedRef<SGraphEditor> GraphEditor = StaticCastSharedRef<SGraphEditor>(Tab->GetContent());
@@ -237,7 +243,7 @@ TSharedRef<SWidget> FConstructionScriptEditorSummoner::CreateTabBody(const FWork
 {
 	TSharedPtr<FBlueprintEditor> BlueprintEditorPtr = StaticCastSharedPtr<FBlueprintEditor>(HostingApp.Pin());
 
-	return BlueprintEditorPtr->GetSCSEditor();
+	return BlueprintEditorPtr->GetSCSEditor().ToSharedRef();
 }
 
 FSCSViewportSummoner::FSCSViewportSummoner(TSharedPtr<class FAssetEditorToolkit> InHostingApp)
diff --git a/Engine/Source/Editor/Kismet/Private/BlueprintEditorTabFactories.h b/Engine/Source/Editor/Kismet/Private/BlueprintEditorTabFactories.h
--- a/Engine/Source/Editor/Kismet/Private/BlueprintEditorTabFactories.h
+++ b/Engine/Source/Editor/Kismet/Private/BlueprintEditorTabFactories.h
@@ -124,14 +124,12 @@ public:
 
 	virtual void OnTabActivated(TSharedPtr<SDockTab> Tab) const override;
 
+	virtual void OnTabBackgrounded(TSharedPtr<SDockTab> Tab) const override;
+
 	virtual void OnTabRefreshed(TSharedPtr<SDockTab> Tab) const override;
 
 	virtual void SaveState(TSharedPtr<SDockTab> Tab, TSharedPtr<FTabPayload> Payload) const override;
 
-//	virtual void OnTabClosed(TSharedRef<SDockTab> Tab, TSharedPtr<FTabPayload> Payload) const override
-//	{
-//		BlueprintEditorPtr.Pin()->RequestSaveEditedObjectState();
-//	}
 protected:
 	virtual TAttribute<FText> ConstructTabNameForObject(UEdGraph* DocumentID) const override
 	{
diff --git a/Engine/Source/Editor/Kismet/Private/SCSEditorViewportClient.cpp b/Engine/Source/Editor/Kismet/Private/SCSEditorViewportClient.cpp
--- a/Engine/Source/Editor/Kismet/Private/SCSEditorViewportClient.cpp
+++ b/Engine/Source/Editor/Kismet/Private/SCSEditorViewportClient.cpp
@@ -162,7 +162,7 @@ void FSCSEditorViewportClient::Tick(float DeltaSeconds)
 	FEditorViewportClient::Tick(DeltaSeconds);
 
 	// Register the selection override delegate for the preview actor's components
-	TSharedRef<SSCSEditor> SCSEditor = BlueprintEditorPtr.Pin()->GetSCSEditor();
+	TSharedPtr<SSCSEditor> SCSEditor = BlueprintEditorPtr.Pin()->GetSCSEditor();
 	AActor* PreviewActor = GetPreviewActor();
 	if (PreviewActor != nullptr)
 	{
@@ -757,7 +757,7 @@ void FSCSEditorViewportClient::ToggleIsSimulateEnabled()
 	PreviewScene->GetWorld()->bShouldSimulatePhysics = bIsSimulateEnabled;
 
 	AActor* PreviewActor = GetPreviewActor();
-	TSharedRef<SWidget> SCSEditor = BlueprintEditorPtr.Pin()->GetSCSEditor();
+	TSharedPtr<SWidget> SCSEditor = BlueprintEditorPtr.Pin()->GetSCSEditor();
 	TSharedRef<SWidget> Inspector = BlueprintEditorPtr.Pin()->GetInspector();
 
 	// When simulate is enabled, we don't want to allow the user to modify the components
diff --git a/Engine/Source/Editor/Kismet/Private/SMyBlueprint.cpp b/Engine/Source/Editor/Kismet/Private/SMyBlueprint.cpp
--- a/Engine/Source/Editor/Kismet/Private/SMyBlueprint.cpp
+++ b/Engine/Source/Editor/Kismet/Private/SMyBlueprint.cpp
@@ -503,7 +503,15 @@ FText SMyBlueprint::OnGetSectionTitle( int32 InSectionID )
 		SeperatorTitle = NSLOCTEXT("GraphActionNode", "Components", "Components");
 		break;
 	case NodeSectionID::FUNCTION:
-		SeperatorTitle = NSLOCTEXT("GraphActionNode", "Functions", "Functions");
+		if ( OverridableFunctionActions.Num() > 0 )
+		{
+			SeperatorTitle = FText::Format(NSLOCTEXT("GraphActionNode", "FunctionsOverridableFormat", "Functions <TinyText.Subdued>({0} Overridable)</>"), FText::AsNumber(OverridableFunctionActions.Num()));
+		}
+		else
+		{
+			SeperatorTitle = NSLOCTEXT("GraphActionNode", "Functions", "Functions");
+		}
+
 		break;
 	case NodeSectionID::FUNCTION_OVERRIDABLE:
 		SeperatorTitle = NSLOCTEXT("GraphActionNode", "OverridableFunctions", "Overridable Functions");
@@ -524,8 +532,15 @@ FText SMyBlueprint::OnGetSectionTitle( int32 InSectionID )
 		SeperatorTitle = NSLOCTEXT("GraphActionNode", "Userenums", "User Enums");
 		break;	
 	case NodeSectionID::LOCAL_VARIABLE:
-		SeperatorTitle = NSLOCTEXT("GraphActionNode", "LocalVariables", "Local Variables");
-		break;	
+		if ( GetFocusedGraph() )
+		{
+			SeperatorTitle = FText::Format(NSLOCTEXT("GraphActionNode", "LocalVariables", "Local Variables <TinyText.Subdued>({0})</>"), FText::FromName(GetFocusedGraph()->GetFName()));
+		}
+		else
+		{
+			SeperatorTitle = NSLOCTEXT("GraphActionNode", "LocalVariables", "Local Variables");
+		}
+		break;
 	case NodeSectionID::USER_STRUCT:
 		SeperatorTitle = NSLOCTEXT("GraphActionNode", "Userstructs", "User Structs");
 		break;	
diff --git a/Engine/Source/Editor/Kismet/Private/WorkflowOrientedApp/WorkflowTabManager.cpp b/Engine/Source/Editor/Kismet/Private/WorkflowOrientedApp/WorkflowTabManager.cpp
--- a/Engine/Source/Editor/Kismet/Private/WorkflowOrientedApp/WorkflowTabManager.cpp
+++ b/Engine/Source/Editor/Kismet/Private/WorkflowOrientedApp/WorkflowTabManager.cpp
@@ -372,11 +372,13 @@ FDocumentTracker::FDocumentTracker()
 {
 	// Make sure we know when tabs become active
 	OnActiveTabChangedDelegateHandle = FGlobalTabmanager::Get()->OnActiveTabChanged_Subscribe( FOnActiveTabChanged::FDelegate::CreateRaw( this, &FDocumentTracker::OnActiveTabChanged ) );
+	TabForegroundedDelegateHandle = FGlobalTabmanager::Get()->OnTabForegrounded_Subscribe(FOnActiveTabChanged::FDelegate::CreateRaw(this, &FDocumentTracker::OnTabForegrounded));
 }
 
 FDocumentTracker::~FDocumentTracker()
 {
-	FGlobalTabmanager::Get()->OnActiveTabChanged_Unsubscribe( OnActiveTabChangedDelegateHandle );
+	FGlobalTabmanager::Get()->OnActiveTabChanged_Unsubscribe(OnActiveTabChangedDelegateHandle);
+	FGlobalTabmanager::Get()->OnTabForegrounded_Unsubscribe(TabForegroundedDelegateHandle);
 }
 
 // Called by the global active tab changed callback; dispatches to individually registered callbacks
@@ -399,6 +401,43 @@ void FDocumentTracker::OnActiveTabChanged(TSharedPtr<SDockTab> PreviouslyActive,
 	}
 }
 
+void FDocumentTracker::OnTabForegrounded(TSharedPtr<SDockTab> ForegroundedTab, TSharedPtr<SDockTab> BackgroundedTab)
+{
+	TSharedPtr<SDockTab> OwnedForeground, OwnedBackground;
+	TSharedPtr<FDocumentTabFactory> ForegroundFactory, BackgroundFactory;
+
+	FTabList& List = GetSpawnedList();
+	for ( auto ListIt = List.CreateIterator(); ListIt; ++ListIt )
+	{
+		// Get the factory (can't fail; the tabs had to come from somewhere; failing means a tab survived a mode transition to a mode where it is not allowed!)
+		TSharedPtr<FDocumentTabFactory> Factory = ( *ListIt )->GetFactory().Pin();
+		if ( ensure(Factory.IsValid()) )
+		{
+			TSharedPtr<SDockTab> Tab = ( *ListIt )->GetTab().Pin();
+			if ( Tab == ForegroundedTab )
+			{
+				OwnedForeground = Tab;
+				ForegroundFactory = Factory;
+			}
+			else if ( Tab == BackgroundedTab )
+			{
+				OwnedBackground = Tab;
+				BackgroundFactory = Factory;
+			}
+		}
+	}
+
+	if ( OwnedBackground.IsValid() )
+	{
+		BackgroundFactory->OnTabBackgrounded(OwnedBackground);
+	}
+
+	if ( OwnedForeground.IsValid() )
+	{
+		ForegroundFactory->OnTabForegrounded(OwnedForeground);
+	}
+}
+
 FDocumentTracker::FTabList& FDocumentTracker::GetSpawnedList()
 {
 	// Remove any closed tabs
diff --git a/Engine/Source/Editor/Kismet/Public/BlueprintEditor.h b/Engine/Source/Editor/Kismet/Public/BlueprintEditor.h
--- a/Engine/Source/Editor/Kismet/Public/BlueprintEditor.h
+++ b/Engine/Source/Editor/Kismet/Public/BlueprintEditor.h
@@ -197,7 +197,9 @@ public:
 	TSharedRef<class SBlueprintPalette> GetPalette() const {return Palette.ToSharedRef();}
 	TSharedRef<class SWidget> GetCompilerResults() const {return CompilerResults.ToSharedRef();}
 	TSharedRef<class SFindInBlueprints> GetFindResults() const {return FindResults.ToSharedRef();}
-	TSharedRef<class SSCSEditor> GetSCSEditor() const {return SCSEditor.ToSharedRef();}
+
+	/** Getters for the various optional Kismet2 widgets */
+	TSharedPtr<class SSCSEditor> GetSCSEditor() const {return SCSEditor;}
 	TSharedPtr<class SSCSEditorViewport> GetSCSViewport() const {return SCSViewport;}
 	TSharedPtr<class SMyBlueprint> GetMyBlueprintWidget() const {return MyBlueprintWidget;}
 
@@ -379,6 +381,9 @@ public:
 	/** Called when graph editor focus is changed */
 	virtual void OnGraphEditorFocused(const TSharedRef<class SGraphEditor>& InGraphEditor);
 
+	/** Called when the graph editor tab is backgrounded */
+	virtual void OnGraphEditorBackgrounded(const TSharedRef<SGraphEditor>& InGraphEditor);
+
 	/** Enable/disable the SCS editor preview viewport */
 	void EnableSCSPreview(bool bEnable);
 
diff --git a/Engine/Source/Editor/Kismet/Public/WorkflowOrientedApp/WorkflowTabFactory.h b/Engine/Source/Editor/Kismet/Public/WorkflowOrientedApp/WorkflowTabFactory.h
--- a/Engine/Source/Editor/Kismet/Public/WorkflowOrientedApp/WorkflowTabFactory.h
+++ b/Engine/Source/Editor/Kismet/Public/WorkflowOrientedApp/WorkflowTabFactory.h
@@ -150,6 +150,12 @@ public:
 	// Called when a tab created from this factory is activated
 	virtual void OnTabActivated(TSharedPtr<SDockTab> Tab) const {}
 
+	// Called when a tab created from this factory is brought to the foreground
+	virtual void OnTabForegrounded(TSharedPtr<SDockTab> Tab) const {}
+
+	// Called when a tab created from this factory is sent to the background
+	virtual void OnTabBackgrounded(TSharedPtr<SDockTab> Tab) const {}
+
 	// Called when a tab created from this factory should be refreshed
 	virtual void OnTabRefreshed(TSharedPtr<SDockTab> Tab) const {}
 
diff --git a/Engine/Source/Editor/Kismet/Public/WorkflowOrientedApp/WorkflowTabManager.h b/Engine/Source/Editor/Kismet/Public/WorkflowOrientedApp/WorkflowTabManager.h
--- a/Engine/Source/Editor/Kismet/Public/WorkflowOrientedApp/WorkflowTabManager.h
+++ b/Engine/Source/Editor/Kismet/Public/WorkflowOrientedApp/WorkflowTabManager.h
@@ -296,6 +296,9 @@ private:
 	/** Handle to the registered OnActiveTabChanged delegate */
 	FDelegateHandle OnActiveTabChangedDelegateHandle;
 
+	/** Handle to the registered OnTabForegrounded delegate */
+	FDelegateHandle TabForegroundedDelegateHandle;
+
 private:
 	// Clean the spawned list 
 	FTabList& GetSpawnedList();
@@ -303,6 +306,9 @@ private:
 	// Called by the global active tab changed callback; dispatches to individually registered callbacks
 	void OnActiveTabChanged(TSharedPtr<SDockTab> PreviouslyActive, TSharedPtr<SDockTab> NewlyActivated);
 
+	// Called by the global tab manager foregrounds a tab; dispatches to individually registered callbacks
+	void OnTabForegrounded(TSharedPtr<SDockTab> PreviouslyActive, TSharedPtr<SDockTab> NewlyActivated);
+
 	/** Callback for navigating tabs with buttons */
 	FReply OnNavigateTab(FDocumentTracker::EOpenDocumentCause InCause);
 
diff --git a/Engine/Source/Runtime/Slate/Private/Framework/Docking/SDockingTabWell.cpp b/Engine/Source/Runtime/Slate/Private/Framework/Docking/SDockingTabWell.cpp
--- a/Engine/Source/Runtime/Slate/Private/Framework/Docking/SDockingTabWell.cpp
+++ b/Engine/Source/Runtime/Slate/Private/Framework/Docking/SDockingTabWell.cpp
@@ -459,6 +459,8 @@ void SDockingTabWell::BringTabToFront( int32 TabIndexToActivate )
 				: Tabs[ForegroundTabIndex];
 			
 			MyDockArea->GetTabManager()->GetPrivateApi().OnTabForegrounded(NewForegroundTab, PreviousForegroundTab);
+
+			FGlobalTabmanager::Get()->GetPrivateApi().OnTabForegrounded(NewForegroundTab, PreviousForegroundTab);
 		}
 	}
 
diff --git a/Engine/Source/Runtime/Slate/Private/Framework/Docking/TabManager.cpp b/Engine/Source/Runtime/Slate/Private/Framework/Docking/TabManager.cpp
--- a/Engine/Source/Runtime/Slate/Private/Framework/Docking/TabManager.cpp
+++ b/Engine/Source/Runtime/Slate/Private/Framework/Docking/TabManager.cpp
@@ -1570,27 +1570,31 @@ FDelegateHandle FGlobalTabmanager::OnActiveTabChanged_Subscribe( const FOnActive
 	return OnActiveTabChanged.Add( InDelegate );
 }
 
-
-
 void FGlobalTabmanager::OnActiveTabChanged_Unsubscribe( const FOnActiveTabChanged::FDelegate& InDelegate )
 {
 	OnActiveTabChanged.DEPRECATED_Remove( InDelegate );
 }
 
-
 void FGlobalTabmanager::OnActiveTabChanged_Unsubscribe( FDelegateHandle Handle )
 {
 	OnActiveTabChanged.Remove( Handle );
 }
 
+FDelegateHandle FGlobalTabmanager::OnTabForegrounded_Subscribe(const FOnActiveTabChanged::FDelegate& InDelegate)
+{
+	return TabForegrounded.Add(InDelegate);
+}
+
+void FGlobalTabmanager::OnTabForegrounded_Unsubscribe(FDelegateHandle Handle)
+{
+	TabForegrounded.Remove(Handle);
+}
 
 TSharedPtr<class SDockTab> FGlobalTabmanager::GetActiveTab() const
 {
 	return ActiveTabPtr.Pin();
 }
 
-
-
 void FGlobalTabmanager::SetActiveTab( const TSharedPtr<class SDockTab>& NewActiveTab )
 {
 	const bool bShouldApplyChange =
@@ -1769,6 +1773,8 @@ void FGlobalTabmanager::OnTabForegrounded( const TSharedPtr<SDockTab>& NewForegr
 			BackgroundedTabManager->GetPrivateApi().HideWindows();
 		}
 	}
+
+	TabForegrounded.Broadcast(NewForegroundTab, BackgroundedTab);
 }
 
 void FGlobalTabmanager::OnTabRelocated( const TSharedRef<SDockTab>& RelocatedTab, const TSharedPtr<SWindow>& NewOwnerWindow )
diff --git a/Engine/Source/Runtime/Slate/Public/Framework/Docking/TabManager.h b/Engine/Source/Runtime/Slate/Public/Framework/Docking/TabManager.h
--- a/Engine/Source/Runtime/Slate/Public/Framework/Docking/TabManager.h
+++ b/Engine/Source/Runtime/Slate/Public/Framework/Docking/TabManager.h
@@ -762,6 +762,12 @@ public:
 	/** Unsubscribe to notifications about the active tab changing */
 	void OnActiveTabChanged_Unsubscribe( FDelegateHandle Handle );
 
+	/** Subscribe to notifications about a foreground tab changing */
+	FDelegateHandle OnTabForegrounded_Subscribe(const FOnActiveTabChanged::FDelegate& InDelegate);
+
+	/** Unsubscribe to notifications about a foreground tab changing */
+	void OnTabForegrounded_Unsubscribe(FDelegateHandle Handle);
+
 	/** @return the currently active tab; NULL pointer if there is no active tab */
 	TSharedPtr<SDockTab> GetActiveTab() const;
 
@@ -888,6 +894,8 @@ private:
 
 	FOnActiveTabChanged OnActiveTabChanged;
 
+	FOnActiveTabChanged TabForegrounded;
+
 	FText AppTitle;
 
 	/** A window under which all of the windows in this application will nest. */
