[INTEGRATE] Change 2440652 by [REDACTED_PERSON][REDACTED_HANDLE] on 2015/02/10 18:37:40

	Renamed bUseLegacyMenuSystem to bForceLegacyMenuSystem, we think some users may be 'stuck' with the old menuing system and if not we need to know why they're sticking with the old code path.

diff --git a/Engine/Source/Editor/BlueprintGraph/Private/BlueprintEditorSettings.cpp b/Engine/Source/Editor/BlueprintGraph/Private/BlueprintEditorSettings.cpp
--- a/Engine/Source/Editor/BlueprintGraph/Private/BlueprintEditorSettings.cpp
+++ b/Engine/Source/Editor/BlueprintGraph/Private/BlueprintEditorSettings.cpp
@@ -17,7 +17,7 @@ UBlueprintEditorSettings::UBlueprintEditorSettings(const FObjectInitializer& Obj
 	, bFlattenFavoritesMenus(true)
 	, bFavorPureCastNodes(false)
 	, bAutoCastObjectConnections(false)
-	, bUseLegacyMenuingSystem(false)
+	, bForceLegacyMenuingSystem(false)
 	, bShowViewportOnSimulate(false)
 	, bShowInheritedVariables(false)
 	, bShowEmptySections(true)
diff --git a/Engine/Source/Editor/BlueprintGraph/Public/BlueprintEditorSettings.h b/Engine/Source/Editor/BlueprintGraph/Public/BlueprintEditorSettings.h
--- a/Engine/Source/Editor/BlueprintGraph/Public/BlueprintEditorSettings.h
+++ b/Engine/Source/Editor/BlueprintGraph/Public/BlueprintEditorSettings.h
@@ -57,7 +57,7 @@ public:
 
 	/** If set, then the new refactored menu system will be replaced with the old (legacy) system (as a fallback, in case the new system has unforeseen problems)*/
 	UPROPERTY(EditAnywhere, AdvancedDisplay, config, Category=Workflow)
-	bool bUseLegacyMenuingSystem;
+	bool bForceLegacyMenuingSystem;
 
 	/** If true will show the viewport tab when simulate is clicked. */
 	UPROPERTY(EditAnywhere, config, Category=Workflow)
diff --git a/Engine/Source/Editor/Kismet/Private/BlueprintActionMenuBuilder.cpp b/Engine/Source/Editor/Kismet/Private/BlueprintActionMenuBuilder.cpp
--- a/Engine/Source/Editor/Kismet/Private/BlueprintActionMenuBuilder.cpp
+++ b/Engine/Source/Editor/Kismet/Private/BlueprintActionMenuBuilder.cpp
@@ -16,7 +16,7 @@
 #include "SMyBlueprint.h"				// for SelectionAsVar()
 #include "BlueprintEditorUtils.h"		// for FindBlueprintForGraphChecked()
 #include "BlueprintEditor.h"			// for GetFocusedGraph()
-#include "BlueprintEditorSettings.h"	// for bUseLegacyMenuingSystem
+#include "BlueprintEditorSettings.h"	// for bForceLegacyMenuingSystem
 
 #define LOCTEXT_NAMESPACE "BlueprintActionMenuBuilder"
 
@@ -542,7 +542,7 @@ void FBlueprintActionMenuBuilder::RebuildActionList()
 	}
 	
 	const UBlueprintEditorSettings* BlueprintSettings = GetDefault<UBlueprintEditorSettings>();
-	if (!BlueprintSettings->bUseLegacyMenuingSystem)
+	if (!BlueprintSettings->bForceLegacyMenuingSystem)
 	{
 		FBlueprintActionDatabase::FActionRegistry const& ActionDatabase = FBlueprintActionDatabase::Get().GetAllActions();
 		for (auto const& ActionEntry : ActionDatabase)
diff --git a/Engine/Source/Editor/Kismet/Private/BlueprintActionMenuUtils.cpp b/Engine/Source/Editor/Kismet/Private/BlueprintActionMenuUtils.cpp
--- a/Engine/Source/Editor/Kismet/Private/BlueprintActionMenuUtils.cpp
+++ b/Engine/Source/Editor/Kismet/Private/BlueprintActionMenuUtils.cpp
@@ -618,7 +618,7 @@ void FBlueprintActionMenuUtils::MakeContextMenu(FBlueprintActionContext const& C
 
 	MenuOut.RebuildActionList();
 
-	if (!BlueprintSettings->bUseLegacyMenuingSystem)
+	if (!BlueprintSettings->bForceLegacyMenuingSystem)
 	{
 		for (UEdGraph const* Graph : Context.Graphs)
 		{
@@ -655,7 +655,7 @@ void FBlueprintActionMenuUtils::MakeFavoritesMenu(FBlueprintActionContext const&
 	MenuOut.Empty();
 
 	const UBlueprintEditorSettings* BlueprintSettings = GetDefault<UBlueprintEditorSettings>();
-	if (!BlueprintSettings->bUseLegacyMenuingSystem)
+	if (!BlueprintSettings->bForceLegacyMenuingSystem)
 	{
 		FBlueprintActionFilter MenuFilter;
 		MenuFilter.Context = Context;
