[INTEGRATE] Change 2440775 by [REDACTED_PERSON][REDACTED_HANDLE] on 2015/02/10 20:01:26

	Potential non-unity compile fixes (missing includes, etc...)

diff --git a/Engine/Source/Editor/Kismet/Private/SBlueprintActionMenu.cpp b/Engine/Source/Editor/Kismet/Private/SBlueprintActionMenu.cpp
--- a/Engine/Source/Editor/Kismet/Private/SBlueprintActionMenu.cpp
+++ b/Engine/Source/Editor/Kismet/Private/SBlueprintActionMenu.cpp
@@ -15,6 +15,7 @@
 #include "BlueprintActionMenuUtils.h"
 #include "BlueprintPaletteFavorites.h"
 #include "IDocumentation.h"
+#include "SSCSEditor.h"
 
 #define LOCTEXT_NAMESPACE "SBlueprintGraphContextMenu"
 
