[INTEGRATE] Change 2441414 by [REDACTED_PERSON][REDACTED_HANDLE].7 on 2015/02/11 07:01:38

	Fixed not being able to select inherited components in the scs editor

	This addresses UE-9377

diff --git a/Engine/Source/Editor/Kismet/Private/SSCSEditor.cpp b/Engine/Source/Editor/Kismet/Private/SSCSEditor.cpp
--- a/Engine/Source/Editor/Kismet/Private/SSCSEditor.cpp
+++ b/Engine/Source/Editor/Kismet/Private/SSCSEditor.cpp
@@ -3780,24 +3780,32 @@ FSCSEditorTreeNodePtrType SSCSEditor::GetNodeFromActorComponent(const UActorComp
 				// If the given component is one that's created during Blueprint construction
 				if (ActorComponent->IsCreatedByConstructionScript())
 				{
-					// Get the Blueprint object associated with the owner's class
-					UBlueprint* Blueprint = UBlueprint::GetBlueprintFromClass(OwnerClass);
-					if (Blueprint && Blueprint->SimpleConstructionScript)
+					TArray<UBlueprint*> ParentBPStack;
+
+					// Check the entire Class hierarchy for the node
+					UBlueprint::GetBlueprintHierarchyFromClass(OwnerClass, ParentBPStack);
+
+					for(int32 StackIndex = ParentBPStack.Num() - 1; StackIndex >= 0; --StackIndex)
 					{
-						// Attempt to locate an SCS node with a variable name that matches the name of the given component
-						TArray<USCS_Node*> AllNodes = Blueprint->SimpleConstructionScript->GetAllNodes();
-						for (int32 i = 0; i < AllNodes.Num(); ++i)
+						if(ParentBPStack[StackIndex]->SimpleConstructionScript)
 						{
-							USCS_Node* SCS_Node = AllNodes[i];
-
-							check(SCS_Node != NULL);
-							if (SCS_Node->VariableName == ActorComponent->GetFName())
+							// Attempt to locate an SCS node with a variable name that matches the name of the given component
+							TArray<USCS_Node*> AllNodes = ParentBPStack[StackIndex]->SimpleConstructionScript->GetAllNodes();
+							for (int32 i = 0; i < AllNodes.Num(); ++i)
 							{
-								// We found a match; redirect to the component archetype instance that may be associated with a tree node
-								ActorComponent = SCS_Node->ComponentTemplate;
-								break;
+								USCS_Node* SCS_Node = AllNodes[i];
+
+								check(SCS_Node != NULL);
+								if (SCS_Node->VariableName == ActorComponent->GetFName())
+								{
+									// We found a match; redirect to the component archetype instance that may be associated with a tree node
+									ActorComponent = SCS_Node->ComponentTemplate;
+									break;
+								}
 							}
+
 						}
+
 					}
 				}
 				else
