[INTEGRATE] Change 2441426 by [REDACTED_PERSON][REDACTED_HANDLE] on 2015/02/11 07:35:11

	UMG - Adding metadata to the section buttons so that they can be tagged for tutorials.

diff --git a/Engine/Source/Editor/Kismet/Private/SMyBlueprint.cpp b/Engine/Source/Editor/Kismet/Private/SMyBlueprint.cpp
--- a/Engine/Source/Editor/Kismet/Private/SMyBlueprint.cpp
+++ b/Engine/Source/Editor/Kismet/Private/SMyBlueprint.cpp
@@ -557,14 +557,17 @@ TSharedRef<SWidget> SMyBlueprint::OnGetSectionWidget(TSharedRef<SWidget> RowWidg
 	TWeakPtr<SWidget> WeakRowWidget = RowWidget;
 
 	FText AddNewText;
+	FName MetaDataTag;
 
 	switch ( InSectionID )
 	{
 	case NodeSectionID::VARIABLE:
 		AddNewText = LOCTEXT("AddNewVariable", "Variable");
+		MetaDataTag = TEXT("AddNewVariable");
 		break;
 	case NodeSectionID::FUNCTION:
 		AddNewText = LOCTEXT("AddNewFunction", "Function");
+		MetaDataTag = TEXT("AddNewFunction");
 
 		if ( OverridableFunctionActions.Num() > 0 )
 		{
@@ -593,31 +596,35 @@ TSharedRef<SWidget> SMyBlueprint::OnGetSectionWidget(TSharedRef<SWidget> RowWidg
 				.AutoWidth()
 				.Padding(2, 0,0,0)
 				[
-					CreateAddToSectionButton(InSectionID, WeakRowWidget, AddNewText)
+					CreateAddToSectionButton(InSectionID, WeakRowWidget, AddNewText, MetaDataTag)
 				];
 		}
 
 		break;
 	case NodeSectionID::MACRO:
 		AddNewText = LOCTEXT("AddNewMacro", "Macro");
+		MetaDataTag = TEXT("AddNewMacro");
 		break;
 	case NodeSectionID::DELEGATE:
 		AddNewText = LOCTEXT("AddNewDelegate", "Event Dispatcher");
+		MetaDataTag = TEXT("AddNewDelegate");
 		break;
 	case NodeSectionID::GRAPH:
 		AddNewText = LOCTEXT("AddNewGraph", "New Graph");
+		MetaDataTag = TEXT("AddNewGraph");
 		break;
 	case NodeSectionID::LOCAL_VARIABLE:
 		AddNewText = LOCTEXT("AddNewLocalVariable", "Local Variable");
+		MetaDataTag = TEXT("AddNewLocalVariable");
 		break;
 	default:
 		return SNullWidget::NullWidget;
 	}
 
-	return CreateAddToSectionButton(InSectionID, WeakRowWidget, AddNewText);
+	return CreateAddToSectionButton(InSectionID, WeakRowWidget, AddNewText, MetaDataTag);
 }
 
-TSharedRef<SWidget> SMyBlueprint::CreateAddToSectionButton(int32 InSectionID, TWeakPtr<SWidget> WeakRowWidget, FText AddNewText)
+TSharedRef<SWidget> SMyBlueprint::CreateAddToSectionButton(int32 InSectionID, TWeakPtr<SWidget> WeakRowWidget, FText AddNewText, FName MetaDataTag)
 {
 	return SNew(SButton)
 		.ButtonStyle(FEditorStyle::Get(), "RoundButton")
@@ -627,6 +634,7 @@ TSharedRef<SWidget> SMyBlueprint::CreateAddToSectionButton(int32 InSectionID, TW
 		.IsEnabled(this, &SMyBlueprint::IsEditingMode)
 		.HAlign(HAlign_Center)
 		.VAlign(VAlign_Center)
+		.AddMetaData<FTagMetaData>(FTagMetaData(MetaDataTag))
 		[
 			SNew(SHorizontalBox)
 
diff --git a/Engine/Source/Editor/Kismet/Private/SMyBlueprint.h b/Engine/Source/Editor/Kismet/Private/SMyBlueprint.h
--- a/Engine/Source/Editor/Kismet/Private/SMyBlueprint.h
+++ b/Engine/Source/Editor/Kismet/Private/SMyBlueprint.h
@@ -160,7 +160,7 @@ private:
 
 	TSharedRef<SWidget> CreateAddNewMenuWidget();
 	void BuildAddNewMenu(FMenuBuilder& MenuBuilder);
-	TSharedRef<SWidget> CreateAddToSectionButton(int32 InSectionID, TWeakPtr<SWidget> WeakRowWidget, FText AddNewText);
+	TSharedRef<SWidget> CreateAddToSectionButton(int32 InSectionID, TWeakPtr<SWidget> WeakRowWidget, FText AddNewText, FName MetaDataTag);
 
 	void OnCategoryNameCommitted(const FText& InNewText, ETextCommit::Type InTextCommit, TWeakPtr< struct FGraphActionNode > InAction );
 	bool CanRequestRenameOnActionNode(TWeakPtr<struct FGraphActionNode> InSelectedNode) const;
