[INTEGRATE] Change 2441435 by [REDACTED_PERSON][REDACTED_HANDLE] on 2015/02/11 07:46:05

	The class wizard now filters out non-Blueprintable component types when adding a new BP component

	UE-9604 - Add Blueprint Component dialog ignores Blueprintable metadata, allowing you to create a blueprint from any component at all

	This vastly limits the selection until UE-8946 has been addressed.

	This change also makes the class wizard verify that its featured classes pass the active class filter (as some featured component classes aren't currently Blueprintable).

diff --git a/Engine/Source/Editor/ClassViewer/Private/SClassViewer.cpp b/Engine/Source/Editor/ClassViewer/Private/SClassViewer.cpp
--- a/Engine/Source/Editor/ClassViewer/Private/SClassViewer.cpp
+++ b/Engine/Source/Editor/ClassViewer/Private/SClassViewer.cpp
@@ -2996,4 +2996,9 @@ void SClassViewer::Tick( const FGeometry& AllottedGeometry, const double InCurre
 	}
 }
 
+bool SClassViewer::IsClassAllowed(const UClass* InClass) const
+{
+	return ClassViewer::Helpers::IsClassAllowed(InitOptions, InClass);
+}
+
 #undef LOCTEXT_NAMESPACE
diff --git a/Engine/Source/Editor/ClassViewer/Private/SClassViewer.h b/Engine/Source/Editor/ClassViewer/Private/SClassViewer.h
--- a/Engine/Source/Editor/ClassViewer/Private/SClassViewer.h
+++ b/Engine/Source/Editor/ClassViewer/Private/SClassViewer.h
@@ -42,6 +42,9 @@ public:
 	virtual FReply OnFocusReceived( const FGeometry& MyGeometry, const FFocusEvent& InFocusEvent ) override;
 	virtual bool SupportsKeyboardFocus() const override;
 
+	/** Test to see whether the given class would be allowed by this class viewer */
+	virtual bool IsClassAllowed(const UClass* InClass) const;
+
 	/** Destroys the internal Class Hierarchy database */
 	static void DestroyClassHierarchy();
 private:
diff --git a/Engine/Source/Editor/GameProjectGeneration/Private/SNewClassDialog.cpp b/Engine/Source/Editor/GameProjectGeneration/Private/SNewClassDialog.cpp
--- a/Engine/Source/Editor/GameProjectGeneration/Private/SNewClassDialog.cpp
+++ b/Engine/Source/Editor/GameProjectGeneration/Private/SNewClassDialog.cpp
@@ -179,7 +179,18 @@ void SNewClassDialog::Construct( const FArguments& InArgs )
 
 	ClassViewer = StaticCastSharedRef<SClassViewer>(FModuleManager::LoadModuleChecked<FClassViewerModule>("ClassViewer").CreateClassViewer(Options, FOnClassPicked::CreateSP(this, &SNewClassDialog::OnAdvancedClassSelected)));
 
-	SetupParentClassItems(InArgs._FeaturedClasses);
+	// Make sure the featured classes all pass the active class filter
+	TArray<FNewClassInfo> ValidatedFeaturedClasses;
+	ValidatedFeaturedClasses.Reserve(InArgs._FeaturedClasses.Num());
+	for (const FNewClassInfo& FeaturedClassInfo : InArgs._FeaturedClasses)
+	{
+		if (FeaturedClassInfo.ClassType != FNewClassInfo::EClassType::UObject || ClassViewer->IsClassAllowed(FeaturedClassInfo.BaseClass))
+		{
+			ValidatedFeaturedClasses.Add(FeaturedClassInfo);
+		}
+	}
+
+	SetupParentClassItems(ValidatedFeaturedClasses);
 	UpdateInputValidity();
 
 	TSharedRef<SWidget> DocWidget = IDocumentation::Get()->CreateAnchor(TAttribute<FString>(this, &SNewClassDialog::GetSelectedParentDocLink));
diff --git a/Engine/Source/Editor/Kismet/Private/SSCSEditor.cpp b/Engine/Source/Editor/Kismet/Private/SSCSEditor.cpp
--- a/Engine/Source/Editor/Kismet/Private/SSCSEditor.cpp
+++ b/Engine/Source/Editor/Kismet/Private/SSCSEditor.cpp
@@ -4188,10 +4188,10 @@ TSharedPtr<FSCSEditorTreeNode> SSCSEditor::AddRootComponentTreeNode(UActorCompon
 	return NewTreeNode;
 }
 
-class FClassParentFilter : public IClassViewerFilter
+class FComponentClassParentFilter : public IClassViewerFilter
 {
 public:
-	FClassParentFilter(const TSubclassOf<UActorComponent>& InComponentClass) : ComponentClass(InComponentClass) {}
+	FComponentClassParentFilter(const TSubclassOf<UActorComponent>& InComponentClass) : ComponentClass(InComponentClass) {}
 
 	virtual bool IsClassAllowed(const FClassViewerInitializationOptions& InInitOptions, const UClass* InClass, TSharedRef< FClassViewerFilterFuncs > InFilterFuncs ) override
 	{
@@ -4206,6 +4206,19 @@ public:
 	TSubclassOf<UActorComponent> ComponentClass;
 };
 
+typedef FComponentClassParentFilter FNativeComponentClassParentFilter;
+
+class FBlueprintComponentClassParentFilter : public FComponentClassParentFilter
+{
+public:
+	FBlueprintComponentClassParentFilter(const TSubclassOf<UActorComponent>& InComponentClass) : FComponentClassParentFilter(InComponentClass) {}
+
+	virtual bool IsClassAllowed(const FClassViewerInitializationOptions& InInitOptions, const UClass* InClass, TSharedRef< FClassViewerFilterFuncs > InFilterFuncs ) override
+	{
+		return FComponentClassParentFilter::IsClassAllowed(InInitOptions, InClass, InFilterFuncs) && FKismetEditorUtilities::CanCreateBlueprintOfClass(InClass);
+	}
+};
+
 UClass* SSCSEditor::CreateNewCPPComponent( TSubclassOf<UActorComponent> ComponentClass )
 {
 	TSharedPtr<SWindow> ParentWindow = FSlateApplication::Get().FindWidgetWindow(SharedThis(this));
@@ -4226,7 +4239,7 @@ UClass* SSCSEditor::CreateNewCPPComponent( TSubclassOf<UActorComponent> Componen
 		.Modal()
 		.OnAddedToProject(FOnAddedToProject::CreateLambda(OnCodeAddedToProject))
 		.FeatureComponentClasses()
-		.AllowableParents(MakeShareable( new FClassParentFilter(ComponentClass) ))
+		.AllowableParents(MakeShareable( new FNativeComponentClassParentFilter(ComponentClass) ))
 		.DefaultClassPrefix(TEXT("New"))
 	);
 
@@ -4264,7 +4277,7 @@ UClass* SSCSEditor::CreateNewBPComponent(TSubclassOf<UActorComponent> ComponentC
 		.WindowTitle(LOCTEXT("AddNewBlueprintComponent", "Add Blueprint Component"))
 		.ParentWindow(FSlateApplication::Get().FindWidgetWindow(SharedThis(this)))
 		.Modal()
-		.AllowableParents(MakeShareable( new FClassParentFilter(ComponentClass) ))
+		.AllowableParents(MakeShareable( new FBlueprintComponentClassParentFilter(ComponentClass) ))
 		.FeatureComponentClasses()
 		.OnAddedToProject(FOnAddedToProject::CreateLambda(OnAddedToProject))
 		.DefaultClassPrefix(TEXT("New"))
