[INTEGRATE] Change 2442111 by [REDACTED_PERSON][REDACTED_HANDLE] on 2015/02/11 14:41:43

	UE-8144 - Now restoring the details view if it's hidden when you select class defaults or class settings.

diff --git a/Engine/Source/Editor/Kismet/Private/BlueprintEditor.cpp b/Engine/Source/Editor/Kismet/Private/BlueprintEditor.cpp
--- a/Engine/Source/Editor/Kismet/Private/BlueprintEditor.cpp
+++ b/Engine/Source/Editor/Kismet/Private/BlueprintEditor.cpp
@@ -789,6 +789,15 @@ void FBlueprintEditor::OnSelectionUpdated(const TArray<FSCSEditorTreeNodePtrType
 					DefaultActor->GetName(Title);
 					InspectorTitle = FText::FromString(Title);
 					bShowComponents = false;
+
+					// Show the details panel if it doesn't exist.
+					TabManager->InvokeTab(FBlueprintEditorTabs::DetailsID);
+
+					TSharedPtr<SDockTab> OwnerTab = Inspector->GetOwnerTab();
+					if ( OwnerTab.IsValid() )
+					{
+						OwnerTab->FlashTab();
+					}
 				}
 				else
 				{
@@ -2574,10 +2583,12 @@ void FBlueprintEditor::EditGlobalOptions_Clicked()
 		// Show details for the Blueprint instance we're editing
 		Inspector->ShowDetailsForSingleObject(Blueprint);
 
+		// Show the details panel if it doesn't exist.
+		TabManager->InvokeTab(FBlueprintEditorTabs::DetailsID);
+
 		TSharedPtr<SDockTab> OwnerTab = Inspector->GetOwnerTab();
 		if ( OwnerTab.IsValid() )
 		{
-			OwnerTab->ActivateInParent(ETabActivationCause::SetDirectly);
 			OwnerTab->FlashTab();
 		}
 	}
@@ -6295,10 +6306,12 @@ void FBlueprintEditor::StartEditingDefaults(bool bAutoFocus, bool bForceRefresh)
 
 					Inspector->ShowDetailsForSingleObject(DefaultObject, Options);
 
+					// Show the details panel if it doesn't exist.
+					TabManager->InvokeTab(FBlueprintEditorTabs::DetailsID);
+
 					TSharedPtr<SDockTab> OwnerTab = Inspector->GetOwnerTab();
 					if ( OwnerTab.IsValid() )
 					{
-						OwnerTab->ActivateInParent(ETabActivationCause::SetDirectly);
 						OwnerTab->FlashTab();
 					}
 				}
