[INTEGRATE] Change 2442156 by [REDACTED_PERSON][REDACTED_HANDLE] on 2015/02/11 15:07:59

	UE-8144 - Now restoring the details view if it's hidden when you select class defaults or class settings.  Fixing a bug that was causing the insepector to come up at the wrong time, we now clear selection state when switching modes.

diff --git a/Engine/Source/Editor/Kismet/Private/BlueprintEditor.cpp b/Engine/Source/Editor/Kismet/Private/BlueprintEditor.cpp
--- a/Engine/Source/Editor/Kismet/Private/BlueprintEditor.cpp
+++ b/Engine/Source/Editor/Kismet/Private/BlueprintEditor.cpp
@@ -1518,6 +1518,9 @@ void FBlueprintEditor::RegisterTabSpawners(const TSharedRef<class FTabManager>&
 
 void FBlueprintEditor::SetCurrentMode(FName NewMode)
 {
+	// Clear the selection state when the mode changes.
+	SetUISelectionState(NAME_None);
+
 	OnModeSetData.Broadcast( NewMode );
 	FWorkflowCentricApplication::SetCurrentMode(NewMode);
 }
