[INTEGRATE] Change 2442262 by [REDACTED_PERSON][REDACTED_HANDLE] on 2015/02/11 16:13:30

	- Removed colored text from fields in the SCS editor and replaced with "(inherited)"
	https://jira.ol.epicgames.net/browse/UE-9502

diff --git a/Engine/Source/Editor/Kismet/Private/SSCSEditor.cpp b/Engine/Source/Editor/Kismet/Private/SSCSEditor.cpp
--- a/Engine/Source/Editor/Kismet/Private/SSCSEditor.cpp
+++ b/Engine/Source/Editor/Kismet/Private/SSCSEditor.cpp
@@ -1265,13 +1265,12 @@ TSharedRef<SWidget> SSCS_RowWidget::GenerateWidgetForColumn( const FName& Column
 			ComponentIcon = FClassIconFinder::FindIconForClass( NodePtr->GetComponentTemplate()->GetClass(), TEXT("SCS.Component") );
 		}
 
-		TSharedPtr<SInlineEditableTextBlock> InlineWidget = 
+		InlineWidget =
 			SNew(SInlineEditableTextBlock)
 				.Text(this, &SSCS_RowWidget::GetNameLabel)
 				.OnVerifyTextChanged( this, &SSCS_RowWidget::OnNameTextVerifyChanged )
 				.OnTextCommitted( this, &SSCS_RowWidget::OnNameTextCommit )
 				.IsSelected( this, &SSCS_RowWidget::IsSelectedExclusively )
-				.ColorAndOpacity( this, &SSCS_RowWidget::GetColorTintForText )
 				.IsReadOnly(!NodePtr->CanRename() || (SCSEditor.IsValid() && !SCSEditor.Pin()->IsEditingAllowed()));
 
 		NodePtr->SetRenameRequestedDelegate(FSCSEditorTreeNode::FOnRenameRequested::CreateSP(InlineWidget.Get(), &SInlineEditableTextBlock::EnterEditingMode));
@@ -1710,31 +1709,6 @@ FSlateColor SSCS_RowWidget::GetColorTintForIcon() const
 	}
 }
 
-FSlateColor SSCS_RowWidget::GetColorTintForText() const
-{
-	const FSlateColor SelectedTextColor(FLinearColor::Black);
-
-	FSCSEditorTreeNodePtrType NodePtr = GetNode();
-	if (NodePtr->IsInherited())
-	{
-		if (NodePtr->IsNative())
-		{
-			static FSlateColor NativeColor = FEditorStyle::GetSlateColor("Common.InheritedFromNativeTextColor");
-			return IsSelected() ? SelectedTextColor : NativeColor;
-		}
-		else
-		{
-			static FSlateColor BlueprintColor = FEditorStyle::GetSlateColor("Common.InheritedFromBlueprintTextColor");
-			return IsSelected() ? SelectedTextColor : BlueprintColor;
-		}
-	}
-	else
-	{
-		static FSlateColor ThisInstanceColor = FEditorStyle::GetSlateColor("Common.IntroducedInThisInstanceTextColor");
-		return IsSelected() ? SelectedTextColor : ThisInstanceColor;
-	}
-}
-
 TSharedPtr<SWidget> SSCS_RowWidget::BuildSceneRootDropActionMenu(FSCSEditorTreeNodePtrType DroppedNodePtr)
 {
 	check(SCSEditor.IsValid());
@@ -2576,6 +2550,15 @@ void SSCS_RowWidget::PostDragDropAction(bool bRegenerateTreeNodes)
 
 FText SSCS_RowWidget::GetNameLabel() const
 {
+	if( InlineWidget.IsValid() && !InlineWidget->IsInEditMode() )
+	{
+		FSCSEditorTreeNodePtrType NodePtr = GetNode();
+		if(NodePtr->IsInherited())
+		{
+			return FText::Format(LOCTEXT("NativeComponentFormatString","{0} (Inherited)"), FText::FromString(GetNode()->GetDisplayString()));
+		}
+	}
+
 	// NOTE: Whatever this returns also becomes the variable name
 	return FText::FromString(GetNode()->GetDisplayString());
 }
@@ -3155,6 +3138,8 @@ void SSCSEditor::Construct( const FArguments& InArgs )
 				.HAlign(HAlign_Left)
 				[
 					SNew(SComponentClassCombo)
+					//** Hack 4.7 Child actor components are not allowed to be added via the actor details panel */
+					//.AllowChildActorComponent(EditorMode==EComponentEditorMode::BlueprintSCS)
 					.AddMetaData<FTagMetaData>(FTagMetaData(TEXT("Actor.AddComponent")))
 					.Visibility(bHideComponentClassCombo ? EVisibility::Hidden : EVisibility::Visible)
 					.OnComponentClassSelected(this, &SSCSEditor::PerformComboAddClass)
@@ -3214,6 +3199,7 @@ void SSCSEditor::Construct( const FArguments& InArgs )
 				]
 				+ SHorizontalBox::Slot()
 				.FillWidth(1.0f)
+				.Padding( 3.0f, 3.0f )
 				.HAlign(HAlign_Right)
 				[
 					SNew(SComboButton)
diff --git a/Engine/Source/Editor/Kismet/Public/SSCSEditor.h b/Engine/Source/Editor/Kismet/Public/SSCSEditor.h
--- a/Engine/Source/Editor/Kismet/Public/SSCSEditor.h
+++ b/Engine/Source/Editor/Kismet/Public/SSCSEditor.h
@@ -574,7 +574,7 @@ private:
 public:
 	/** Pointer back to owning SCSEditor 2 tool */
 	TWeakPtr<SSCSEditor> SCSEditor;
-
+	TSharedPtr<SInlineEditableTextBlock> InlineWidget;
 private:
 	/** Pointer to node we represent */
 	FSCSEditorTreeNodePtrType TreeNodePtr;
diff --git a/Engine/Source/Editor/PropertyEditor/Private/SDetailNameArea.cpp b/Engine/Source/Editor/PropertyEditor/Private/SDetailNameArea.cpp
--- a/Engine/Source/Editor/PropertyEditor/Private/SDetailNameArea.cpp
+++ b/Engine/Source/Editor/PropertyEditor/Private/SDetailNameArea.cpp
@@ -201,14 +201,28 @@ void SDetailNameArea::BuildObjectNameAreaSelectionLabel( TSharedRef< SHorizontal
 					FEditorClassUtils::GetDocumentationLinkWidget(ObjectClass)
 				];
 
-			SelectionLabelBox->AddSlot()
-				.AutoWidth()
-				.VAlign( VAlign_Center )
-				.HAlign( HAlign_Left )
-				.Padding( 6.0f, 1.0f, 0.0f, 0.0f )
-				[
-					FEditorClassUtils::GetSourceLink(ObjectClass, ObjectWeakPtr)
-				];
+
+			if( ObjectClass && ObjectClass->ClassGeneratedBy == nullptr && ObjectClass->GetOutermost() )
+			{
+				const FString ModuleName = FPackageName::GetShortName(ObjectClass->GetOutermost()->GetFName());
+
+				FModuleStatus PackageModuleStatus;
+				if(FModuleManager::Get().QueryModule(*ModuleName, PackageModuleStatus))
+				{
+					if( PackageModuleStatus.bIsGameModule ) 
+					{
+						SelectionLabelBox->AddSlot()
+						.AutoWidth()
+						.VAlign(VAlign_Center)
+						.HAlign(HAlign_Left)
+						.Padding(6.0f, 1.0f, 0.0f, 0.0f)
+						[
+							FEditorClassUtils::GetSourceLink(ObjectClass, ObjectWeakPtr)
+						];
+					}
+				}
+			}
+		
 		}
 	}
 	else
