Details panel for variables now cache the UProperty and the variable name and will not submit changes to newly selected variables.

diff --git a/Engine/Source/Editor/Kismet/Private/BlueprintDetailsCustomization.cpp b/Engine/Source/Editor/Kismet/Private/BlueprintDetailsCustomization.cpp
--- a/Engine/Source/Editor/Kismet/Private/BlueprintDetailsCustomization.cpp
+++ b/Engine/Source/Editor/Kismet/Private/BlueprintDetailsCustomization.cpp
@@ -167,18 +167,38 @@ int32 FBlueprintDetails::HandleAddOrViewIndexForButton(const FName EventName, FN
 	return 1; // Add
 }
 
+FBlueprintVarActionDetails::~FBlueprintVarActionDetails()
+{
+	// Remove the callback delegate we registered for
+	TWeakPtr<FBlueprintEditor> BlueprintEditor = MyBlueprint.Pin()->GetBlueprintEditor();
+	if( BlueprintEditor.IsValid() )
+	{
+		BlueprintEditor.Pin()->OnRefresh().RemoveAll(this);
+	}
+}
+
 // UProperty Detail Customization
 BEGIN_SLATE_FUNCTION_BUILD_OPTIMIZATION
 void FBlueprintVarActionDetails::CustomizeDetails( IDetailLayoutBuilder& DetailLayout )
 {
-	UProperty* VariableProperty = SelectionAsProperty();
+	CachedVariableProperty = SelectionAsProperty();
 
-	if(!VariableProperty)
+	if(!CachedVariableProperty.IsValid())
 	{
 		return;
 	}
 
-	FName VarName =	GetVariableName();
+	CachedVariableName = GetVariableName();
+
+	TWeakPtr<FBlueprintEditor> BlueprintEditor = MyBlueprint.Pin()->GetBlueprintEditor();
+	if( BlueprintEditor.IsValid() )
+	{
+		BlueprintEditor.Pin()->OnRefresh().AddSP(this, &FBlueprintVarActionDetails::OnPostEditorRefresh);
+	}
+
+
+	UProperty* VariableProperty = CachedVariableProperty.Get();
+
 	const UEdGraphSchema_K2* Schema = GetDefault<UEdGraphSchema_K2>();
 	
 	IDetailCategoryBuilder& Category = DetailLayout.EditCategory("Variable", LOCTEXT("VariableDetailsCategory", "Variable"));
@@ -265,7 +285,7 @@ void FBlueprintVarActionDetails::CustomizeDetails( IDetailLayoutBuilder& DetailL
 		SNew(SEditableTextBox)
 		.Text( this, &FBlueprintVarActionDetails::OnGetTooltipText )
 		.ToolTip(ToolTipTooltip)
-		.OnTextCommitted( this, &FBlueprintVarActionDetails::OnTooltipTextCommitted, VarName )
+		.OnTextCommitted( this, &FBlueprintVarActionDetails::OnTooltipTextCommitted, CachedVariableName )
 		.Font( DetailFontInfo )
 	];
 
@@ -394,7 +414,7 @@ void FBlueprintVarActionDetails::CustomizeDetails( IDetailLayoutBuilder& DetailL
 			[
 				SNew(SEditableTextBox)
 					.Text(this, &FBlueprintVarActionDetails::OnGetCategoryText)
-					.OnTextCommitted(this, &FBlueprintVarActionDetails::OnCategoryTextCommitted, VarName )
+					.OnTextCommitted(this, &FBlueprintVarActionDetails::OnCategoryTextCommitted, CachedVariableName )
 					.ToolTip(CategoryTooltip)
 					.SelectAllTextWhenFocused(true)
 					.RevertTextOnEscape(true)
@@ -649,7 +669,6 @@ void FBlueprintVarActionDetails::CustomizeDetails( IDetailLayoutBuilder& DetailL
 					}
 				}
 
-				TWeakPtr<FBlueprintEditor> BlueprintEditor = MyBlueprint.Pin()->GetBlueprintEditor();
 				if(BlueprintEditor.IsValid())
 				{
 					TSharedPtr< IDetailsView > DetailsView  = BlueprintEditor.Pin()->GetInspector()->GetPropertyView();
@@ -740,15 +759,17 @@ END_SLATE_FUNCTION_BUILD_OPTIMIZATION
 
 void FBlueprintVarActionDetails::RefreshPropertyFlags()
 {
-	UProperty* Property = SelectionAsProperty();
-
-	PropertyFlags.Empty();
-	for( const TCHAR* PropertyFlag : ParsePropertyFlags(Property->PropertyFlags) )
+	UProperty* VariableProperty = CachedVariableProperty.Get();
+	if(VariableProperty)
 	{
-		PropertyFlags.Add(MakeShareable<FString>(new FString(PropertyFlag)));
-	}
+		PropertyFlags.Empty();
+		for( const TCHAR* PropertyFlag : ParsePropertyFlags(VariableProperty->PropertyFlags) )
+		{
+			PropertyFlags.Add(MakeShareable<FString>(new FString(PropertyFlag)));
+		}
 
-	PropertyFlagWidget.Pin()->RequestListRefresh();
+		PropertyFlagWidget.Pin()->RequestListRefresh();
+	}
 }
 
 TSharedRef<ITableRow> FBlueprintVarActionDetails::OnGenerateWidgetForPropertyList( TSharedPtr< FString > Item, const TSharedRef<STableViewBase>& OwnerTable )
@@ -812,20 +833,20 @@ bool FBlueprintVarActionDetails::GetVariableNameChangeEnabled() const
 	UBlueprint* Blueprint = GetBlueprintObj();
 	check(Blueprint != NULL);
 
-	UProperty* VariableProperty = SelectionAsProperty();
-	if(VariableProperty != NULL)
+	UProperty* VariableProperty = CachedVariableProperty.Get();
+	if(VariableProperty != nullptr)
 	{
-		if(FBlueprintEditorUtils::FindNewVariableIndex(Blueprint, GetVariableName()) != INDEX_NONE)
+		if(FBlueprintEditorUtils::FindNewVariableIndex(Blueprint, CachedVariableName) != INDEX_NONE)
 		{
 			bIsReadOnly = false;
 		}
-		else if(Blueprint->FindTimelineTemplateByVariableName(GetVariableName()))
+		else if(Blueprint->FindTimelineTemplateByVariableName(CachedVariableName))
 		{
 			bIsReadOnly = false;
 		}
 		else if(IsAComponentVariable(VariableProperty) && Blueprint->SimpleConstructionScript != NULL)
 		{
-			if (USCS_Node* Node = Blueprint->SimpleConstructionScript->FindSCSNode(GetVariableName()))
+			if (USCS_Node* Node = Blueprint->SimpleConstructionScript->FindSCSNode(CachedVariableName))
 			{
 				bIsReadOnly = !FComponentEditorUtils::IsValidVariableNameString(Node->ComponentTemplate, Node->VariableName.ToString());
 			}
@@ -841,7 +862,7 @@ bool FBlueprintVarActionDetails::GetVariableNameChangeEnabled() const
 
 FText FBlueprintVarActionDetails::OnGetVarName() const
 {
-	return FText::FromName(GetVariableName());
+	return FText::FromName(CachedVariableName);
 }
 
 void FBlueprintVarActionDetails::OnVarNameChanged(const FText& InNewText)
@@ -851,14 +872,14 @@ void FBlueprintVarActionDetails::OnVarNameChanged(const FText& InNewText)
 	UBlueprint* Blueprint = GetBlueprintObj();
 	check(Blueprint != NULL);
 
-	UProperty* VariableProperty = SelectionAsProperty();
+	UProperty* VariableProperty = CachedVariableProperty.Get();
 	if(VariableProperty && IsAComponentVariable(VariableProperty) && Blueprint->SimpleConstructionScript != NULL)
 	{
 		TArray<USCS_Node*> Nodes = Blueprint->SimpleConstructionScript->GetAllNodes();
 		for (TArray<USCS_Node*>::TConstIterator NodeIt(Nodes); NodeIt; ++NodeIt)
 		{
 			USCS_Node* Node = *NodeIt;
-			if (Node->VariableName == GetVariableName() && !FComponentEditorUtils::IsValidVariableNameString(Node->ComponentTemplate, InNewText.ToString()))
+			if (Node->VariableName == CachedVariableName && !FComponentEditorUtils::IsValidVariableNameString(Node->ComponentTemplate, InNewText.ToString()))
 			{
 				VarNameEditableTextBox->SetError(LOCTEXT("ComponentVariableRenameFailed_NotValid", "This name is reserved for engine use."));
 				return;
@@ -866,7 +887,7 @@ void FBlueprintVarActionDetails::OnVarNameChanged(const FText& InNewText)
 		}
 	}
 
-	TSharedPtr<INameValidatorInterface> NameValidator = MakeShareable(new FKismetNameValidator(Blueprint, GetVariableName(), GetLocalVariableScope(VariableProperty)));
+	TSharedPtr<INameValidatorInterface> NameValidator = MakeShareable(new FKismetNameValidator(Blueprint, CachedVariableName, GetLocalVariableScope(VariableProperty)));
 
 	EValidatorResult ValidatorResult = NameValidator->IsValid(InNewText.ToString());
 	if(ValidatorResult == EValidatorResult::AlreadyInUse)
@@ -894,7 +915,7 @@ void FBlueprintVarActionDetails::OnVarNameChanged(const FText& InNewText)
 
 void FBlueprintVarActionDetails::OnVarNameCommitted(const FText& InNewText, ETextCommit::Type InTextCommit)
 {
-	if((InTextCommit == ETextCommit::OnEnter || InTextCommit == ETextCommit::OnUserMovedFocus) && !bIsVarNameInvalid)
+	if(InTextCommit != ETextCommit::OnCleared && !bIsVarNameInvalid)
 	{
 		const FScopedTransaction Transaction( LOCTEXT( "RenameVariable", "Rename Variable" ) );
 
@@ -902,7 +923,8 @@ void FBlueprintVarActionDetails::OnVarNameCommitted(const FText& InNewText, ETex
 
 		// Double check we're not renaming a timeline disguised as a variable
 		bool bIsTimeline = false;
-		UProperty* VariableProperty = SelectionAsProperty();
+
+		UProperty* VariableProperty = CachedVariableProperty.Get();
 		if (VariableProperty != NULL)
 		{
 			// Don't allow removal of timeline properties - you need to remove the timeline node for that
@@ -911,25 +933,25 @@ void FBlueprintVarActionDetails::OnVarNameCommitted(const FText& InNewText, ETex
 			{
 				bIsTimeline = true;
 			}
-		}
 
-		// Rename as a timeline if required
-		if (bIsTimeline)
-		{
-			FBlueprintEditorUtils::RenameTimeline(GetBlueprintObj(), GetVariableName(), NewVarName);
-		}
-		else if(IsALocalVariable(VariableProperty))
-		{
-			UFunction* LocalVarScope = Cast<UFunction>(VariableProperty->GetOuter());
-			FBlueprintEditorUtils::RenameLocalVariable(GetBlueprintObj(), LocalVarScope, GetVariableName(), NewVarName);
-		}
-		else
-		{
-			FBlueprintEditorUtils::RenameMemberVariable(GetBlueprintObj(), GetVariableName(), NewVarName);
-		}
+			// Rename as a timeline if required
+			if (bIsTimeline)
+			{
+				FBlueprintEditorUtils::RenameTimeline(GetBlueprintObj(), CachedVariableName, NewVarName);
+			}
+			else if(IsALocalVariable(VariableProperty))
+			{
+				UFunction* LocalVarScope = Cast<UFunction>(VariableProperty->GetOuter());
+				FBlueprintEditorUtils::RenameLocalVariable(GetBlueprintObj(), LocalVarScope, CachedVariableName, NewVarName);
+			}
+			else
+			{
+				FBlueprintEditorUtils::RenameMemberVariable(GetBlueprintObj(), CachedVariableName, NewVarName);
+			}
 
-		check(MyBlueprint.IsValid());
-		MyBlueprint.Pin()->SelectItemByName(NewVarName, ESelectInfo::OnMouseClick);
+			check(MyBlueprint.IsValid());
+			MyBlueprint.Pin()->SelectItemByName(NewVarName, ESelectInfo::OnMouseClick);
+		}
 	}
 
 	bIsVarNameInvalid = false;
@@ -938,7 +960,7 @@ void FBlueprintVarActionDetails::OnVarNameCommitted(const FText& InNewText, ETex
 
 bool FBlueprintVarActionDetails::GetVariableTypeChangeEnabled() const
 {
-	UProperty* VariableProperty = SelectionAsProperty();
+	UProperty* VariableProperty = CachedVariableProperty.Get();
 	if(VariableProperty && !IsALocalVariable(VariableProperty))
 	{
 		if(GetBlueprintObj()->SkeletonGeneratedClass->GetAuthoritativeClass() != VariableProperty->GetOwnerClass()->GetAuthoritativeClass())
@@ -946,7 +968,7 @@ bool FBlueprintVarActionDetails::GetVariableTypeChangeEnabled() const
 			return false;
 		}
 		// If the variable belongs to this class and cannot be found in the member variable list, it is not editable (it may be a component)
-		if (FBlueprintEditorUtils::FindNewVariableIndex(GetBlueprintObj(), GetVariableName()) == INDEX_NONE)
+		if (FBlueprintEditorUtils::FindNewVariableIndex(GetBlueprintObj(), CachedVariableName) == INDEX_NONE)
 		{
 			return false;
 		}
@@ -956,15 +978,15 @@ bool FBlueprintVarActionDetails::GetVariableTypeChangeEnabled() const
 
 bool FBlueprintVarActionDetails::GetVariableCategoryChangeEnabled() const
 {
-	UProperty* VariableProp = SelectionAsProperty();
-	if(VariableProp)
+	UProperty* VariableProperty = CachedVariableProperty.Get();
+	if(VariableProperty)
 	{
-		if(UClass* VarSourceClass = Cast<UClass>(VariableProp->GetOuter()))
+		if(UClass* VarSourceClass = Cast<UClass>(VariableProperty->GetOuter()))
 		{
 			// If the variable's source class is the same as the current blueprint's class then it was created in this blueprint and it's category can be changed.
 			return VarSourceClass == GetBlueprintObj()->SkeletonGeneratedClass;
 		}
-		else if(IsALocalVariable(VariableProp))
+		else if(IsALocalVariable(VariableProperty))
 		{
 			return true;
 		}
@@ -975,13 +997,12 @@ bool FBlueprintVarActionDetails::GetVariableCategoryChangeEnabled() const
 
 FEdGraphPinType FBlueprintVarActionDetails::OnGetVarType() const
 {
-	UProperty* VariableProp = SelectionAsProperty();
-
-	if(VariableProp)
+	UProperty* VariableProperty = CachedVariableProperty.Get();
+	if (VariableProperty)
 	{
 		const UEdGraphSchema_K2* K2Schema = GetDefault<UEdGraphSchema_K2>();
 		FEdGraphPinType Type;
-		K2Schema->ConvertPropertyToPinType(VariableProp, Type);
+		K2Schema->ConvertPropertyToPinType(VariableProperty, Type);
 		return Type;
 	}
 	return FEdGraphPinType();
@@ -991,21 +1012,24 @@ void FBlueprintVarActionDetails::OnVarTypeChanged(const FEdGraphPinType& NewPinT
 {
 	if (FBlueprintEditorUtils::IsPinTypeValid(NewPinType))
 	{
-		FName VarName = GetVariableName();
+		FName VarName = CachedVariableName;
 
 		if (VarName != NAME_None)
 		{
 			// Set the MyBP tab's last pin type used as this, for adding lots of variables of the same type
 			MyBlueprint.Pin()->GetLastPinTypeUsed() = NewPinType;
 
-			UProperty* VariableProperty = SelectionAsProperty();
-			if(IsALocalVariable(VariableProperty))
+			UProperty* VariableProperty = CachedVariableProperty.Get();
+			if(VariableProperty)
 			{
-				FBlueprintEditorUtils::ChangeLocalVariableType(GetBlueprintObj(), GetLocalVariableScope(VariableProperty), VarName, NewPinType);
-			}
-			else
-			{
-				FBlueprintEditorUtils::ChangeMemberVariableType(GetBlueprintObj(), VarName, NewPinType);
+				if(IsALocalVariable(VariableProperty))
+				{
+					FBlueprintEditorUtils::ChangeLocalVariableType(GetBlueprintObj(), GetLocalVariableScope(VariableProperty), VarName, NewPinType);
+				}
+				else
+				{
+					FBlueprintEditorUtils::ChangeMemberVariableType(GetBlueprintObj(), VarName, NewPinType);
+				}
 			}
 		}
 	}
@@ -1013,11 +1037,11 @@ void FBlueprintVarActionDetails::OnVarTypeChanged(const FEdGraphPinType& NewPinT
 
 FText FBlueprintVarActionDetails::OnGetTooltipText() const
 {
-	FName VarName = GetVariableName();
+	FName VarName = CachedVariableName;
 	if (VarName != NAME_None)
 	{
 		FString Result;
-		FBlueprintEditorUtils::GetBlueprintVariableMetaData(GetBlueprintObj(), VarName, GetLocalVariableScope(SelectionAsProperty()), TEXT("tooltip"), Result);
+		FBlueprintEditorUtils::GetBlueprintVariableMetaData(GetBlueprintObj(), VarName, GetLocalVariableScope(CachedVariableProperty.Get()), TEXT("tooltip"), Result);
 		return FText::FromString(Result);
 	}
 	return FText();
@@ -1025,7 +1049,7 @@ FText FBlueprintVarActionDetails::OnGetTooltipText() const
 
 void FBlueprintVarActionDetails::OnTooltipTextCommitted(const FText& NewText, ETextCommit::Type InTextCommit, FName VarName)
 {
-	FBlueprintEditorUtils::SetBlueprintVariableMetaData(GetBlueprintObj(), VarName, GetLocalVariableScope(SelectionAsProperty()), TEXT("tooltip"), NewText.ToString() );
+	FBlueprintEditorUtils::SetBlueprintVariableMetaData(GetBlueprintObj(), VarName, GetLocalVariableScope(CachedVariableProperty.Get()), TEXT("tooltip"), NewText.ToString() );
 }
 
 void FBlueprintVarActionDetails::PopulateCategories(SMyBlueprint* MyBlueprint, TArray<TSharedPtr<FString>>& CategorySource)
@@ -1225,12 +1249,12 @@ FName FBlueprintVarActionDetails::GetVariableName() const
 
 FText FBlueprintVarActionDetails::OnGetCategoryText() const
 {
-	FName VarName = GetVariableName();
+	FName VarName = CachedVariableName;
 	if (VarName != NAME_None)
 	{
 		const UEdGraphSchema_K2* K2Schema = GetDefault<UEdGraphSchema_K2>();
 
-		FName Category = FBlueprintEditorUtils::GetBlueprintVariableCategory(GetBlueprintObj(), VarName, GetLocalVariableScope(SelectionAsProperty()));
+		FName Category = FBlueprintEditorUtils::GetBlueprintVariableCategory(GetBlueprintObj(), VarName, GetLocalVariableScope(CachedVariableProperty.Get()));
 
 		// Older blueprints will have their name as the default category and whenever it is the same as the default category, display localized text
 		if( Category == GetBlueprintObj()->GetFName() || Category == K2Schema->VR_DefaultCategory )
@@ -1256,7 +1280,7 @@ void FBlueprintVarActionDetails::OnCategoryTextCommitted(const FText& NewText, E
 		FString NewCategory = CategoryName.ToString();
 		if(NewCategory.Len() <= NAME_SIZE)
 		{
-			FBlueprintEditorUtils::SetBlueprintVariableCategory(GetBlueprintObj(), VarName, GetLocalVariableScope(SelectionAsProperty()), FName( *NewCategory ));
+			FBlueprintEditorUtils::SetBlueprintVariableCategory(GetBlueprintObj(), VarName, GetLocalVariableScope(CachedVariableProperty.Get()), FName( *NewCategory ));
 			check(MyBlueprint.IsValid());
 			PopulateCategories(MyBlueprint.Pin().Get(), CategorySource);
 			MyBlueprint.Pin()->ExpandCategory(NewCategory);
@@ -1274,12 +1298,12 @@ TSharedRef< ITableRow > FBlueprintVarActionDetails::MakeCategoryViewWidget( TSha
 
 void FBlueprintVarActionDetails::OnCategorySelectionChanged( TSharedPtr<FString> ProposedSelection, ESelectInfo::Type /*SelectInfo*/ )
 {
-	FName VarName = GetVariableName();
+	FName VarName = CachedVariableName;
 	if (ProposedSelection.IsValid() && VarName != NAME_None)
 	{
 		FString NewCategory = *ProposedSelection.Get();
 
-		FBlueprintEditorUtils::SetBlueprintVariableCategory(GetBlueprintObj(), VarName, GetLocalVariableScope(SelectionAsProperty()), FName( *NewCategory ));
+		FBlueprintEditorUtils::SetBlueprintVariableCategory(GetBlueprintObj(), VarName, GetLocalVariableScope(CachedVariableProperty.Get()), FName( *NewCategory ));
 		CategoryListView.Pin()->ClearSelection();
 		CategoryComboButton.Pin()->SetIsOpen(false);
 		MyBlueprint.Pin()->ExpandCategory(NewCategory);
@@ -1288,7 +1312,7 @@ void FBlueprintVarActionDetails::OnCategorySelectionChanged( TSharedPtr<FString>
 
 EVisibility FBlueprintVarActionDetails::ShowEditableCheckboxVisibilty() const
 {
-	UProperty* VariableProperty = SelectionAsProperty();
+	UProperty* VariableProperty = CachedVariableProperty.Get();
 	if (VariableProperty)
 	{
 		if (IsABlueprintVariable(VariableProperty) && !IsAComponentVariable(VariableProperty))
@@ -1301,7 +1325,7 @@ EVisibility FBlueprintVarActionDetails::ShowEditableCheckboxVisibilty() const
 
 ECheckBoxState FBlueprintVarActionDetails::OnEditableCheckboxState() const
 {
-	UProperty* VariableProperty = SelectionAsProperty();
+	UProperty* VariableProperty = CachedVariableProperty.Get();
 	if (VariableProperty)
 	{
 		return VariableProperty->HasAnyPropertyFlags(CPF_DisableEditOnInstance) ? ECheckBoxState::Unchecked : ECheckBoxState::Checked;
@@ -1311,7 +1335,7 @@ ECheckBoxState FBlueprintVarActionDetails::OnEditableCheckboxState() const
 
 void FBlueprintVarActionDetails::OnEditableChanged(ECheckBoxState InNewState)
 {
-	FName VarName = GetVariableName();
+	FName VarName = CachedVariableName;
 
 	// Toggle the flag on the blueprint's version of the variable description, based on state
 	const bool bVariableIsExposed = InNewState == ECheckBoxState::Checked;
@@ -1322,7 +1346,7 @@ void FBlueprintVarActionDetails::OnEditableChanged(ECheckBoxState InNewState)
 
 ECheckBoxState FBlueprintVarActionDetails::OnCreateWidgetCheckboxState() const
 {
-	UProperty* Property = SelectionAsProperty();
+	UProperty* Property = CachedVariableProperty.Get();
 	if (Property)
 	{
 		bool bMakingWidget = FEdMode::ShouldCreateWidgetForProperty(Property);
@@ -1334,23 +1358,23 @@ ECheckBoxState FBlueprintVarActionDetails::OnCreateWidgetCheckboxState() const
 
 void FBlueprintVarActionDetails::OnCreateWidgetChanged(ECheckBoxState InNewState)
 {
-	const FName VarName = GetVariableName();
+	const FName VarName = CachedVariableName;
 	if (VarName != NAME_None)
 	{
 		if (InNewState == ECheckBoxState::Checked)
 		{
-			FBlueprintEditorUtils::SetBlueprintVariableMetaData(GetBlueprintObj(), VarName, GetLocalVariableScope(SelectionAsProperty()), FEdMode::MD_MakeEditWidget, TEXT("true"));
+			FBlueprintEditorUtils::SetBlueprintVariableMetaData(GetBlueprintObj(), VarName, GetLocalVariableScope(CachedVariableProperty.Get()), FEdMode::MD_MakeEditWidget, TEXT("true"));
 		}
 		else
 		{
-			FBlueprintEditorUtils::RemoveBlueprintVariableMetaData(GetBlueprintObj(), VarName, GetLocalVariableScope(SelectionAsProperty()), FEdMode::MD_MakeEditWidget);
+			FBlueprintEditorUtils::RemoveBlueprintVariableMetaData(GetBlueprintObj(), VarName, GetLocalVariableScope(CachedVariableProperty.Get()), FEdMode::MD_MakeEditWidget);
 		}
 	}
 }
 
 EVisibility FBlueprintVarActionDetails::Show3DWidgetVisibility() const
 {
-	UProperty* VariableProperty = SelectionAsProperty();
+	UProperty* VariableProperty = CachedVariableProperty.Get();
 	if (VariableProperty)
 	{
 		if (IsABlueprintVariable(VariableProperty) && FEdMode::CanCreateWidgetForProperty(VariableProperty))
@@ -1363,7 +1387,7 @@ EVisibility FBlueprintVarActionDetails::Show3DWidgetVisibility() const
 
 bool FBlueprintVarActionDetails::Is3DWidgetEnabled()
 {
-	UProperty* VariableProperty = SelectionAsProperty();
+	UProperty* VariableProperty = CachedVariableProperty.Get();
 	if (VariableProperty)
 	{
 		return ( VariableProperty && !VariableProperty->HasAnyPropertyFlags(CPF_DisableEditOnInstance) ) ;
@@ -1373,7 +1397,7 @@ bool FBlueprintVarActionDetails::Is3DWidgetEnabled()
 
 ECheckBoxState FBlueprintVarActionDetails::OnGetExposedToSpawnCheckboxState() const
 {
-	UProperty* Property = SelectionAsProperty();
+	UProperty* Property = CachedVariableProperty.Get();
 	if (Property)
 	{
 		return (Property && Property->GetBoolMetaData(FBlueprintMetadata::MD_ExposeOnSpawn) != false) ? ECheckBoxState::Checked : ECheckBoxState::Unchecked;
@@ -1383,7 +1407,7 @@ ECheckBoxState FBlueprintVarActionDetails::OnGetExposedToSpawnCheckboxState() co
 
 void FBlueprintVarActionDetails::OnExposedToSpawnChanged(ECheckBoxState InNewState)
 {
-	const FName VarName = GetVariableName();
+	const FName VarName = CachedVariableName;
 	if (VarName != NAME_None)
 	{
 		const bool bExposeOnSpawn = (InNewState == ECheckBoxState::Checked);
@@ -1400,7 +1424,7 @@ void FBlueprintVarActionDetails::OnExposedToSpawnChanged(ECheckBoxState InNewSta
 
 EVisibility FBlueprintVarActionDetails::ExposeOnSpawnVisibility() const
 {
-	UProperty* VariableProperty = SelectionAsProperty();
+	UProperty* VariableProperty = CachedVariableProperty.Get();
 	if (VariableProperty)
 	{
 		const UEdGraphSchema_K2* K2Schema = GetDefault<UEdGraphSchema_K2>();
@@ -1418,7 +1442,7 @@ EVisibility FBlueprintVarActionDetails::ExposeOnSpawnVisibility() const
 
 ECheckBoxState FBlueprintVarActionDetails::OnGetPrivateCheckboxState() const
 {
-	UProperty* Property = SelectionAsProperty();
+	UProperty* Property = CachedVariableProperty.Get();
 	if (Property)
 	{
 		return (Property && Property->GetBoolMetaData(FBlueprintMetadata::MD_Private) != false) ? ECheckBoxState::Checked : ECheckBoxState::Unchecked;
@@ -1428,7 +1452,7 @@ ECheckBoxState FBlueprintVarActionDetails::OnGetPrivateCheckboxState() const
 
 void FBlueprintVarActionDetails::OnPrivateChanged(ECheckBoxState InNewState)
 {
-	const FName VarName = GetVariableName();
+	const FName VarName = CachedVariableName;
 	if (VarName != NAME_None)
 	{
 		const bool bExposeOnSpawn = (InNewState == ECheckBoxState::Checked);
@@ -1445,7 +1469,7 @@ void FBlueprintVarActionDetails::OnPrivateChanged(ECheckBoxState InNewState)
 
 EVisibility FBlueprintVarActionDetails::ExposePrivateVisibility() const
 {
-	UProperty* Property = SelectionAsProperty();
+	UProperty* Property = CachedVariableProperty.Get();
 	if (Property)
 	{
 		if (IsABlueprintVariable(Property) && !IsAComponentVariable(Property))
@@ -1458,7 +1482,7 @@ EVisibility FBlueprintVarActionDetails::ExposePrivateVisibility() const
 
 ECheckBoxState FBlueprintVarActionDetails::OnGetExposedToMatineeCheckboxState() const
 {
-	UProperty* Property = SelectionAsProperty();
+	UProperty* Property = CachedVariableProperty.Get();
 	if (Property)
 	{
 		return Property && Property->HasAnyPropertyFlags(CPF_Interp) ? ECheckBoxState::Checked : ECheckBoxState::Unchecked;
@@ -1471,7 +1495,7 @@ void FBlueprintVarActionDetails::OnExposedToMatineeChanged(ECheckBoxState InNewS
 	// Toggle the flag on the blueprint's version of the variable description, based on state
 	const bool bExposeToMatinee = (InNewState == ECheckBoxState::Checked);
 	
-	const FName VarName = GetVariableName();
+	const FName VarName = CachedVariableName;
 	if (VarName != NAME_None)
 	{
 		FBlueprintEditorUtils::SetInterpFlag(GetBlueprintObj(), VarName, bExposeToMatinee);
@@ -1480,7 +1504,7 @@ void FBlueprintVarActionDetails::OnExposedToMatineeChanged(ECheckBoxState InNewS
 
 EVisibility FBlueprintVarActionDetails::ExposeToMatineeVisibility() const
 {
-	UProperty* VariableProperty = SelectionAsProperty();
+	UProperty* VariableProperty = CachedVariableProperty.Get();
 	if (VariableProperty && !IsALocalVariable(VariableProperty))
 	{
 		const bool bIsInteger = VariableProperty->IsA(UIntProperty::StaticClass());
@@ -1502,7 +1526,7 @@ EVisibility FBlueprintVarActionDetails::ExposeToMatineeVisibility() const
 ECheckBoxState FBlueprintVarActionDetails::OnGetConfigVariableCheckboxState() const
 {
 	UBlueprint* Blueprint = GetBlueprintObj();
-	const FName VarName = GetVariableName();
+	const FName VarName = CachedVariableName;
 	ECheckBoxState CheckboxValue = ECheckBoxState::Unchecked;
 
 	if( Blueprint && VarName != NAME_None )
@@ -1520,7 +1544,7 @@ ECheckBoxState FBlueprintVarActionDetails::OnGetConfigVariableCheckboxState() co
 void FBlueprintVarActionDetails::OnSetConfigVariableState( ECheckBoxState InNewState )
 {
 	UBlueprint* Blueprint = GetBlueprintObj();
-	const FName VarName = GetVariableName();
+	const FName VarName = CachedVariableName;
 
 	if( Blueprint && VarName != NAME_None )
 	{
@@ -1543,7 +1567,7 @@ void FBlueprintVarActionDetails::OnSetConfigVariableState( ECheckBoxState InNewS
 
 EVisibility FBlueprintVarActionDetails::ExposeConfigVisibility() const
 {
-	UProperty* Property = SelectionAsProperty();
+	UProperty* Property = CachedVariableProperty.Get();
 	if (Property)
 	{
 		if (IsABlueprintVariable(Property) && !IsAComponentVariable(Property))
@@ -1556,11 +1580,11 @@ EVisibility FBlueprintVarActionDetails::ExposeConfigVisibility() const
 
 FText FBlueprintVarActionDetails::OnGetMetaKeyValue(FName Key) const
 {
-	FName VarName = GetVariableName();
+	FName VarName = CachedVariableName;
 	if (VarName != NAME_None)
 	{
 		FString Result;
-		FBlueprintEditorUtils::GetBlueprintVariableMetaData(GetBlueprintObj(), VarName, GetLocalVariableScope(SelectionAsProperty()), Key, /*out*/ Result);
+		FBlueprintEditorUtils::GetBlueprintVariableMetaData(GetBlueprintObj(), VarName, GetLocalVariableScope(CachedVariableProperty.Get()), Key, /*out*/ Result);
 
 		return FText::FromString(Result);
 	}
@@ -1569,19 +1593,19 @@ FText FBlueprintVarActionDetails::OnGetMetaKeyValue(FName Key) const
 
 void FBlueprintVarActionDetails::OnMetaKeyValueChanged(const FText& NewMinValue, ETextCommit::Type CommitInfo, FName Key)
 {
-	FName VarName = GetVariableName();
+	FName VarName = CachedVariableName;
 	if (VarName != NAME_None)
 	{
 		if ((CommitInfo == ETextCommit::OnEnter) || (CommitInfo == ETextCommit::OnUserMovedFocus))
 		{
-			FBlueprintEditorUtils::SetBlueprintVariableMetaData(GetBlueprintObj(), VarName, GetLocalVariableScope(SelectionAsProperty()), Key, NewMinValue.ToString());
+			FBlueprintEditorUtils::SetBlueprintVariableMetaData(GetBlueprintObj(), VarName, GetLocalVariableScope(CachedVariableProperty.Get()), Key, NewMinValue.ToString());
 		}
 	}
 }
 
 EVisibility FBlueprintVarActionDetails::RangeVisibility() const
 {
-	UProperty* VariableProperty = SelectionAsProperty();
+	UProperty* VariableProperty = CachedVariableProperty.Get();
 	if (VariableProperty)
 	{
 		const bool bIsInteger = VariableProperty->IsA(UIntProperty::StaticClass());
@@ -1601,7 +1625,7 @@ TSharedPtr<FString> FBlueprintVarActionDetails::GetVariableReplicationType() con
 	EVariableReplication::Type VariableReplication = EVariableReplication::None;
 	
 	uint64 PropFlags = 0;
-	UProperty* VariableProperty = SelectionAsProperty();
+	UProperty* VariableProperty = CachedVariableProperty.Get();
 
 	if (VariableProperty)
 	{
@@ -1640,7 +1664,7 @@ void FBlueprintVarActionDetails::OnChangeReplication(TSharedPtr<FString> ItemSel
 
 	EVariableReplication::Type VariableReplication = (EVariableReplication::Type)NewSelection;
 	
-	UProperty* VariableProperty = SelectionAsProperty();
+	UProperty* VariableProperty = CachedVariableProperty.Get();
 
 	if (VariableProperty)
 	{
@@ -1684,7 +1708,7 @@ void FBlueprintVarActionDetails::OnChangeReplication(TSharedPtr<FString> ItemSel
 void FBlueprintVarActionDetails::ReplicationOnRepFuncChanged(const FString& NewOnRepFunc) const
 {
 	FName NewFuncName = FName(*NewOnRepFunc);
-	UProperty* VariableProperty = SelectionAsProperty();
+	UProperty* VariableProperty = CachedVariableProperty.Get();
 
 	if (VariableProperty)
 	{
@@ -1707,7 +1731,7 @@ void FBlueprintVarActionDetails::ReplicationOnRepFuncChanged(const FString& NewO
 
 EVisibility FBlueprintVarActionDetails::ReplicationVisibility() const
 {
-	UProperty* VariableProperty = SelectionAsProperty();
+	UProperty* VariableProperty = CachedVariableProperty.Get();
 	if(VariableProperty)
 	{
 		if (!IsAComponentVariable(VariableProperty) && IsABlueprintVariable(VariableProperty))
@@ -1739,9 +1763,15 @@ TSharedRef<SWidget> FBlueprintVarActionDetails::BuildEventsMenuForVariable() con
 	return SNullWidget::NullWidget;
 }
 
+void FBlueprintVarActionDetails::OnPostEditorRefresh()
+{
+	CachedVariableProperty = SelectionAsProperty();
+	CachedVariableName = GetVariableName();
+}
+
 EVisibility FBlueprintVarActionDetails::GetTransientVisibility() const
 {
-	UProperty* VariableProperty = SelectionAsProperty();
+	UProperty* VariableProperty = CachedVariableProperty.Get();
 	if (VariableProperty)
 	{
 		if (IsABlueprintVariable(VariableProperty) && !IsAComponentVariable(VariableProperty))
@@ -1754,7 +1784,7 @@ EVisibility FBlueprintVarActionDetails::GetTransientVisibility() const
 
 ECheckBoxState FBlueprintVarActionDetails::OnGetTransientCheckboxState() const
 {
-	UProperty* Property = SelectionAsProperty();
+	UProperty* Property = CachedVariableProperty.Get();
 	if (Property)
 	{
 		return (Property && Property->HasAnyPropertyFlags(CPF_Transient)) ? ECheckBoxState::Checked : ECheckBoxState::Unchecked;
@@ -1764,7 +1794,7 @@ ECheckBoxState FBlueprintVarActionDetails::OnGetTransientCheckboxState() const
 
 void FBlueprintVarActionDetails::OnTransientChanged(ECheckBoxState InNewState)
 {
-	UProperty* Property = SelectionAsProperty();
+	UProperty* Property = CachedVariableProperty.Get();
 	if (Property)
 	{
 		const bool bTransientFlag = (InNewState == ECheckBoxState::Checked);
@@ -1774,7 +1804,7 @@ void FBlueprintVarActionDetails::OnTransientChanged(ECheckBoxState InNewState)
 
 EVisibility FBlueprintVarActionDetails::GetSaveGameVisibility() const
 {
-	UProperty* VariableProperty = SelectionAsProperty();
+	UProperty* VariableProperty = CachedVariableProperty.Get();
 	if (VariableProperty)
 	{
 		if (IsABlueprintVariable(VariableProperty) && !IsAComponentVariable(VariableProperty))
@@ -1787,7 +1817,7 @@ EVisibility FBlueprintVarActionDetails::GetSaveGameVisibility() const
 
 ECheckBoxState FBlueprintVarActionDetails::OnGetSaveGameCheckboxState() const
 {
-	UProperty* Property = SelectionAsProperty();
+	UProperty* Property = CachedVariableProperty.Get();
 	if (Property)
 	{
 		return (Property && Property->HasAnyPropertyFlags(CPF_SaveGame)) ? ECheckBoxState::Checked : ECheckBoxState::Unchecked;
@@ -1797,7 +1827,7 @@ ECheckBoxState FBlueprintVarActionDetails::OnGetSaveGameCheckboxState() const
 
 void FBlueprintVarActionDetails::OnSaveGameChanged(ECheckBoxState InNewState)
 {
-	UProperty* Property = SelectionAsProperty();
+	UProperty* Property = CachedVariableProperty.Get();
 	if (Property)
 	{
 		const bool bSaveGameFlag = (InNewState == ECheckBoxState::Checked);
@@ -1807,7 +1837,7 @@ void FBlueprintVarActionDetails::OnSaveGameChanged(ECheckBoxState InNewState)
 
 EVisibility FBlueprintVarActionDetails::IsTooltipEditVisible() const
 {
-	UProperty* VariableProperty = SelectionAsProperty();
+	UProperty* VariableProperty = CachedVariableProperty.Get();
 	if (VariableProperty)
 	{
 		if ((IsABlueprintVariable(VariableProperty) && !IsAComponentVariable(VariableProperty)) || IsALocalVariable(VariableProperty))
diff --git a/Engine/Source/Editor/Kismet/Private/BlueprintDetailsCustomization.h b/Engine/Source/Editor/Kismet/Private/BlueprintDetailsCustomization.h
--- a/Engine/Source/Editor/Kismet/Private/BlueprintDetailsCustomization.h
+++ b/Engine/Source/Editor/Kismet/Private/BlueprintDetailsCustomization.h
@@ -75,6 +75,8 @@ public:
 		, bIsVarNameInvalid(false)
 	{
 	}
+
+	~FBlueprintVarActionDetails();
 	
 	/** IDetailCustomization interface */
 	virtual void CustomizeDetails( IDetailLayoutBuilder& DetailLayout ) override;
@@ -176,6 +178,9 @@ private:
 	/** Delegate to build variable events droplist menu */
 	TSharedRef<SWidget> BuildEventsMenuForVariable() const;
 	
+	/** Refreshes cached data that changes after a Blueprint recompile */
+	void OnPostEditorRefresh();
+
 private:
 	/** Pointer back to my parent tab */
 	TWeakPtr<SMyBlueprint> MyBlueprint;
@@ -200,6 +205,12 @@ private:
 
 	/** The listview widget for displaying property flags */
 	TWeakPtr< SListView< TSharedPtr< FString > > > PropertyFlagWidget;
+
+	/** Cached property for the variable we are affecting */
+	TWeakObjectPtr<UProperty> CachedVariableProperty;
+
+	/** Cached name for the variable we are affecting */
+	FName CachedVariableName;
 };
 
 class FBaseBlueprintGraphActionDetails : public IDetailCustomization
