[INTEGRATE] Change 2444090 by [REDACTED_PERSON][REDACTED_HANDLE] on 2015/02/12 20:23:56

	Blueprints: Prevent the 'Add Component' and 'Add Variable|etc...' buttons from being pressed when not in an editing situation (e.g., PIE or stopped at a breakpoint) [UE-9374]

diff --git a/Engine/Source/Editor/Kismet/Private/SMyBlueprint.cpp b/Engine/Source/Editor/Kismet/Private/SMyBlueprint.cpp
--- a/Engine/Source/Editor/Kismet/Private/SMyBlueprint.cpp
+++ b/Engine/Source/Editor/Kismet/Private/SMyBlueprint.cpp
@@ -229,6 +229,7 @@ void SMyBlueprint::Construct(const FArguments& InArgs, TWeakPtr<FBlueprintEditor
 			.HasDownArrow(true)
 			.ContentPadding(FMargin(1, 0, 2, 0))
 			.AddMetaData<FTagMetaData>(FTagMetaData(TEXT("MyBlueprintAddNewCombo")))
+			.IsEnabled(this, &SMyBlueprint::IsEditingMode)
 			.ButtonContent()
 			[
 				SNew(SHorizontalBox)
diff --git a/Engine/Source/Editor/Kismet/Private/SSCSEditor.cpp b/Engine/Source/Editor/Kismet/Private/SSCSEditor.cpp
--- a/Engine/Source/Editor/Kismet/Private/SSCSEditor.cpp
+++ b/Engine/Source/Editor/Kismet/Private/SSCSEditor.cpp
@@ -3142,6 +3142,7 @@ void SSCSEditor::Construct( const FArguments& InArgs )
 					.Visibility(bHideComponentClassCombo ? EVisibility::Hidden : EVisibility::Visible)
 					.OnComponentClassSelected(this, &SSCSEditor::PerformComboAddClass)
 					.ToolTipText(LOCTEXT("AddComponent_Tooltip", "Adds a new component to this actor"))
+					.IsEnabled(AllowEditing)
 				]
 
 				+ SHorizontalBox::Slot()
