[INTEGRATE] Change 2444868 by [REDACTED_PERSON][REDACTED_HANDLE] on 2015/02/13 11:52:11

	Fix child actor component not spawning child in cases of drag and drop in the IWCE window

diff --git a/Engine/Source/Editor/Kismet/Private/SSCSEditor.cpp b/Engine/Source/Editor/Kismet/Private/SSCSEditor.cpp
--- a/Engine/Source/Editor/Kismet/Private/SSCSEditor.cpp
+++ b/Engine/Source/Editor/Kismet/Private/SSCSEditor.cpp
@@ -188,6 +188,7 @@ FReply SSCSEditor::TryHandleAssetDragDropOperation(const FDragDropEvent& DragDro
 					if (UChildActorComponent* ChildActorComponent = Cast<UChildActorComponent>(AddNewComponent(UChildActorComponent::StaticClass(), nullptr)))
 					{
 						ChildActorComponent->ChildActorClass = PotentialActorClass;
+						ChildActorComponent->CreateChildActor();
 					}
 				}
 			}
@@ -4383,6 +4384,7 @@ UActorComponent* SSCSEditor::AddNewComponent( UClass* NewComponentClass, UObject
 
 			// Add to SerializedComponents array so it gets saved
 			ActorInstance->AddInstanceComponent(NewInstanceComponent);
+			NewInstanceComponent->OnComponentCreated();
 			NewInstanceComponent->RegisterComponent();
 
 			NewComponent = AddNewNodeForInstancedComponent(NewInstanceComponent, Asset);
diff --git a/Engine/Source/Runtime/Engine/Classes/Components/ChildActorComponent.h b/Engine/Source/Runtime/Engine/Classes/Components/ChildActorComponent.h
--- a/Engine/Source/Runtime/Engine/Classes/Components/ChildActorComponent.h
+++ b/Engine/Source/Runtime/Engine/Classes/Components/ChildActorComponent.h
@@ -45,7 +45,7 @@ class UChildActorComponent : public USceneComponent
 	void ApplyComponentInstanceData(class FChildActorComponentInstanceData* ComponentInstanceData);
 
 	/** Create the child actor */
-	void CreateChildActor();
+	ENGINE_API void CreateChildActor();
 
 	/** Kill any currently present child actor */
 	void DestroyChildActor();
