[INTEGRATE] Change 2445940 by [REDACTED_PERSON][REDACTED_HANDLE] on 2015/02/13 23:26:50

	UE-9756 - BP - We now verify that before the details panel is spawned, we ensure the bp editor is being used otherwise we will accidentally force the blueprint editor into the foreground.

diff --git a/Engine/Source/Editor/Kismet/Private/BlueprintEditor.cpp b/Engine/Source/Editor/Kismet/Private/BlueprintEditor.cpp
--- a/Engine/Source/Editor/Kismet/Private/BlueprintEditor.cpp
+++ b/Engine/Source/Editor/Kismet/Private/BlueprintEditor.cpp
@@ -633,7 +633,7 @@ void FBlueprintEditor::RefreshEditors(ERefreshBlueprintEditorReason::Type Reason
 	// rotation edit, transform details won't be re-customized and thus the cached rotation value will be stale, resulting in an invalid rotation value on the next edit.
 	if ( CurrentUISelection == FBlueprintEditor::SelectionState_ClassDefaults )
 	{
-		StartEditingDefaults(true, true);
+		StartEditingDefaults(/*bAutoFocus=*/ false, true);
 	}
 	else
 	{
@@ -6291,7 +6291,10 @@ void FBlueprintEditor::StartEditingDefaults(bool bAutoFocus, bool bForceRefresh)
 
 					Inspector->ShowDetailsForSingleObject(DefaultObject, Options);
 
-					TryInvokingDetailsTab();
+					if ( bAutoFocus )
+					{
+						TryInvokingDetailsTab();
+					}
 				}
 			}
 		}
@@ -7532,15 +7535,22 @@ void FBlueprintEditor::TryInvokingDetailsTab(bool bFlash)
 {
 	if ( TabManager->CanSpawnTab(FBlueprintEditorTabs::DetailsID) )
 	{
-		// Show the details panel if it doesn't exist.
-		TabManager->InvokeTab(FBlueprintEditorTabs::DetailsID);
+		TSharedPtr<SDockTab> BlueprintTab = FGlobalTabmanager::Get()->GetMajorTabForTabManager(TabManager.ToSharedRef());
 
-		if ( bFlash )
+		// We don't want to force this tab into existance when the blueprint editor isn't in the foreground and actively
+		// being interacted with.  So we make sure the window it's in is focused and the tab is in the foreground.
+		if ( BlueprintTab.IsValid() && BlueprintTab->IsForeground() && BlueprintTab->GetParentWindow()->HasFocusedDescendants() )
 		{
-			TSharedPtr<SDockTab> OwnerTab = Inspector->GetOwnerTab();
-			if ( OwnerTab.IsValid() )
+			// Show the details panel if it doesn't exist.
+			TabManager->InvokeTab(FBlueprintEditorTabs::DetailsID);
+
+			if ( bFlash )
 			{
-				OwnerTab->FlashTab();
+				TSharedPtr<SDockTab> OwnerTab = Inspector->GetOwnerTab();
+				if ( OwnerTab.IsValid() )
+				{
+					OwnerTab->FlashTab();
+				}
 			}
 		}
 	}
diff --git a/Engine/Source/Runtime/Slate/Private/Framework/Docking/TabManager.cpp b/Engine/Source/Runtime/Slate/Private/Framework/Docking/TabManager.cpp
--- a/Engine/Source/Runtime/Slate/Private/Framework/Docking/TabManager.cpp
+++ b/Engine/Source/Runtime/Slate/Private/Framework/Docking/TabManager.cpp
@@ -1656,17 +1656,26 @@ bool FGlobalTabmanager::CanCloseManager( const TSet< TSharedRef<SDockTab> >& Tab
 	}
 
 	return bCanCloseManager;
+}
 
+TSharedPtr<SDockTab> FGlobalTabmanager::GetMajorTabForTabManager(const TSharedRef<FTabManager>& ChildManager)
+{
+	const int32 MajorTabIndex = SubTabManagers.IndexOfByPredicate(FindByManager(ChildManager));
+	if ( MajorTabIndex != INDEX_NONE )
+	{
+		return SubTabManagers[MajorTabIndex].MajorTab.Pin();
+	}
+
+	return TSharedPtr<SDockTab>();
 }
 
 void FGlobalTabmanager::DrawAttentionToTabManager( const TSharedRef<FTabManager>& ChildManager )
 {
-	const int32 MajorTabIndex = SubTabManagers.IndexOfByPredicate(FindByManager(ChildManager));
-	if (MajorTabIndex != INDEX_NONE)
+	TSharedPtr<SDockTab> Tab = GetMajorTabForTabManager(ChildManager);
+	if ( Tab.IsValid() )
 	{
-		this->DrawAttention( SubTabManagers[MajorTabIndex].MajorTab.Pin().ToSharedRef() );
+		this->DrawAttention(Tab.ToSharedRef());
 	}
-	
 }
 
 TSharedRef<FTabManager> FGlobalTabmanager::NewTabManager( const TSharedRef<SDockTab>& InOwnerTab )
diff --git a/Engine/Source/Runtime/Slate/Public/Framework/Docking/TabManager.h b/Engine/Source/Runtime/Slate/Public/Framework/Docking/TabManager.h
--- a/Engine/Source/Runtime/Slate/Public/Framework/Docking/TabManager.h
+++ b/Engine/Source/Runtime/Slate/Public/Framework/Docking/TabManager.h
@@ -792,6 +792,9 @@ public:
 
 	virtual bool CanCloseManager( const TSet< TSharedRef<SDockTab> >& TabsToIgnore = TSet< TSharedRef<SDockTab> >()) override;
 
+	/** Gets the major tab for the manager */
+	TSharedPtr<SDockTab> GetMajorTabForTabManager(const TSharedRef<FTabManager>& ChildManager);
+
 	/** Draw the user's attention to a child tab manager */
 	void DrawAttentionToTabManager( const TSharedRef<FTabManager>& ChildManager );
 
