[INTEGRATE] Change 2445973 by [REDACTED_PERSON][REDACTED_HANDLE] on 2015/02/14 00:31:56

	Blueprints: Fix a number of selection issues across a BP compile [UE-9862]
	- Fix a bug where compiling with a custom event or input event selected in the My Blueprints graph would deselect it
	- Fix a bug where "Multiple values" would be shown after recompiling with a variable entry is selected in My Blueprint
	- Fix a bug where "Multiple values" would be shown after recompiling with a variable get/set selected in a graph
	- Fix a bug where the standalone defaults editor wouldn't be updated properly when something else was selected during compilation
	- Fixed that SMyBlueprint::OnActionSelected could clear the details panel without taking ownership when InAction was invalid (now always takes ownership first - long term we need to combine the two so that selecting something includes taking ownership)

diff --git a/Engine/Source/Editor/GraphEditor/Private/SGraphActionMenu.cpp b/Engine/Source/Editor/GraphEditor/Private/SGraphActionMenu.cpp
--- a/Engine/Source/Editor/GraphEditor/Private/SGraphActionMenu.cpp
+++ b/Engine/Source/Editor/GraphEditor/Private/SGraphActionMenu.cpp
@@ -102,7 +102,11 @@ namespace GraphActionMenuHelpers
 			((FEdGraphSchemaAction_K2Struct*)InGraphAction)->GetPathName() == ItemName);
 		bCheck |= (InGraphAction->GetTypeId() == FEdGraphSchemaAction_K2Delegate::StaticGetTypeId() &&
 			((FEdGraphSchemaAction_K2Delegate*)InGraphAction)->GetDelegateName() == ItemName);
-		bCheck |= (InGraphAction->GetTypeId() == FEdGraphSchemaAction_K2TargetNode::StaticGetTypeId() &&
+
+		const bool bIsTargetNodeSubclass = (InGraphAction->GetTypeId() == FEdGraphSchemaAction_K2TargetNode::StaticGetTypeId()) ||
+			(InGraphAction->GetTypeId() == FEdGraphSchemaAction_K2Event::StaticGetTypeId()) ||
+			(InGraphAction->GetTypeId() == FEdGraphSchemaAction_K2InputAction::StaticGetTypeId());
+		bCheck |= (bIsTargetNodeSubclass &&
 			((FEdGraphSchemaAction_K2TargetNode*)InGraphAction)->NodeTemplate->GetNodeTitle(ENodeTitleType::EditableTitle).ToString() == ItemName.ToString());
 
 		return bCheck;
diff --git a/Engine/Source/Editor/Kismet/Private/BlueprintEditor.cpp b/Engine/Source/Editor/Kismet/Private/BlueprintEditor.cpp
--- a/Engine/Source/Editor/Kismet/Private/BlueprintEditor.cpp
+++ b/Engine/Source/Editor/Kismet/Private/BlueprintEditor.cpp
@@ -603,14 +603,50 @@ void FBlueprintEditor::AnalyticsTrackCompileEvent( UBlueprint* Blueprint, int32
 
 void FBlueprintEditor::RefreshEditors(ERefreshBlueprintEditorReason::Type Reason)
 {
-	if( Reason != ERefreshBlueprintEditorReason::BlueprintCompiled )
+	bool bForceFocusOnSelectedNodes = false;
+
+	if (CurrentUISelection == SelectionState_MyBlueprint)
+	{
+		// Handled below, here to avoid tripping the ensure
+	}
+	else if (CurrentUISelection == SelectionState_Components)
+	{
+		if (SCSEditor.IsValid())
+		{
+			SCSEditor->RefreshSelectionDetails();
+		}
+	}
+	else if (CurrentUISelection == SelectionState_Graph)
+	{
+		bForceFocusOnSelectedNodes = true;
+	}
+	else if (CurrentUISelection == SelectionState_ClassSettings)
+	{
+		// No need for a refresh, the Blueprint object didn't change
+	}
+	else if (CurrentUISelection == SelectionState_ClassDefaults)
+	{
+		StartEditingDefaults(/*bAutoFocus=*/ false, true);
+	}
+	else
+	{
+		ensureMsgf(false, TEXT("Unknown UI selection state in FBlueprintEditor::RefreshEditors"));
+	}
+
+	//@TODO: Should determine when we need to do the invalid/refresh business and if the graph node selection change
+	// under non-compiles is necessary (except when the selection mode is appropriate, as already detected above)
+	if (Reason != ERefreshBlueprintEditorReason::BlueprintCompiled)
 	{
 		DocumentManager->CleanInvalidTabs();
 
 		DocumentManager->RefreshAllTabs();
 
-		// The workflow manager only tracks document tabs.
-		FocusInspectorOnGraphSelection(GetSelectedNodes(), true);
+		bForceFocusOnSelectedNodes = true;
+	}
+
+	if (bForceFocusOnSelectedNodes)
+	{
+		FocusInspectorOnGraphSelection(GetSelectedNodes(), /*bForceRefresh=*/ true);
 	}
 
 	if (MyBlueprintWidget.IsValid())
@@ -618,7 +654,7 @@ void FBlueprintEditor::RefreshEditors(ERefreshBlueprintEditorReason::Type Reason
 		MyBlueprintWidget->Refresh();
 	}
 
-	if(SCSEditor.IsValid())
+	if (SCSEditor.IsValid())
 	{
 		SCSEditor->UpdateTree();
 		
@@ -631,14 +667,8 @@ void FBlueprintEditor::RefreshEditors(ERefreshBlueprintEditorReason::Type Reason
 	// that the selection does not really change, we force it to refresh and skip the optimization. Otherwise, some things may not work correctly in Defaults mode. For
 	// example, transform details are customized and the rotation value is cached at customization time; if we don't force refresh here, then after an undo of a previous
 	// rotation edit, transform details won't be re-customized and thus the cached rotation value will be stale, resulting in an invalid rotation value on the next edit.
-	if ( CurrentUISelection == FBlueprintEditor::SelectionState_ClassDefaults )
-	{
-		StartEditingDefaults(/*bAutoFocus=*/ false, true);
-	}
-	else
-	{
-		RefreshStandAloneDefaultsEditor();
-	}
+	//@TODO: Probably not always necessary
+	RefreshStandAloneDefaultsEditor();
 
 	// Update associated controls like the function editor
 	BroadcastRefresh();
diff --git a/Engine/Source/Editor/Kismet/Private/SMyBlueprint.cpp b/Engine/Source/Editor/Kismet/Private/SMyBlueprint.cpp
--- a/Engine/Source/Editor/Kismet/Private/SMyBlueprint.cpp
+++ b/Engine/Source/Editor/Kismet/Private/SMyBlueprint.cpp
@@ -807,7 +807,7 @@ void SMyBlueprint::Refresh()
 {
 	bNeedsRefresh = false;
 
-	GraphActionMenu->RefreshAllActions(true);
+	GraphActionMenu->RefreshAllActions(/*bPreserveExpansion=*/ true);
 }
 
 TSharedRef<SWidget> SMyBlueprint::OnCreateWidgetForAction(FCreateWidgetForActionData* const InCreateData)
@@ -1505,18 +1505,15 @@ void SMyBlueprint::OnActionSelected( const TArray< TSharedPtr<FEdGraphSchemaActi
 
 	TSharedPtr<FBlueprintEditor> BlueprintEditor = BlueprintEditorPtr.Pin();
 
-	if ( BlueprintEditor.IsValid() )
+	if (BlueprintEditor.IsValid())
 	{
-		if ( InAction.IsValid() )
-		{
-			BlueprintEditor->SetUISelectionState(FBlueprintEditor::SelectionState_MyBlueprint);
-		}
+		BlueprintEditor->SetUISelectionState(FBlueprintEditor::SelectionState_MyBlueprint);
 
 		CurrentBlueprint = BlueprintEditor->GetBlueprintObj();
 		CurrentInspector = BlueprintEditor->GetInspector();
-	}
 
-	OnActionSelectedHelper( InAction, Blueprint, CurrentInspector.ToSharedRef() );
+		OnActionSelectedHelper(InAction, Blueprint, CurrentInspector.ToSharedRef());
+	}
 }
 
 void SMyBlueprint::OnActionSelectedHelper(TSharedPtr<FEdGraphSchemaAction> InAction, UBlueprint* Blueprint, TSharedRef<SKismetInspector> Inspector)
@@ -1547,6 +1544,7 @@ void SMyBlueprint::OnActionSelectedHelper(TSharedPtr<FEdGraphSchemaAction> InAct
 			FEdGraphSchemaAction_K2Var* VarAction = (FEdGraphSchemaAction_K2Var*)InAction.Get();
 
 			SKismetInspector::FShowDetailsOptions Options(FText::FromName(VarAction->GetVariableName()));
+			Options.bForceRefresh = true;
 
 			Inspector->ShowDetailsForSingleObject(VarAction->GetProperty(), Options);
 		}
