[INTEGRATE] Change 2446422 by [REDACTED_PERSON][REDACTED_HANDLE] on 2015/02/15 17:13:05

	Make sure the created child component is part of the add component transaction

diff --git a/Engine/Source/Editor/Kismet/Private/SSCSEditor.cpp b/Engine/Source/Editor/Kismet/Private/SSCSEditor.cpp
--- a/Engine/Source/Editor/Kismet/Private/SSCSEditor.cpp
+++ b/Engine/Source/Editor/Kismet/Private/SSCSEditor.cpp
@@ -185,6 +185,7 @@ FReply SSCSEditor::TryHandleAssetDragDropOperation(const FDragDropEvent& DragDro
 				}
 				else if ((PotentialActorClass != nullptr) && !PotentialActorClass->HasAnyClassFlags(CLASS_Deprecated | CLASS_Abstract | CLASS_NewerVersionExists))
 				{
+					const FScopedTransaction Transaction( LOCTEXT("AddComponent", "Add Component") );
 					if (UChildActorComponent* ChildActorComponent = Cast<UChildActorComponent>(AddNewComponent(UChildActorComponent::StaticClass(), nullptr)))
 					{
 						ChildActorComponent->ChildActorClass = PotentialActorClass;
