[INTEGRATE] Change 2446424 by [REDACTED_PERSON][REDACTED_HANDLE] on 2015/02/15 17:19:41

	[UE-9901] Fix for ensure that fires when we refersh the blueprint editor when the selection state is 'None'

diff --git a/Engine/Source/Editor/Kismet/Private/BlueprintEditor.cpp b/Engine/Source/Editor/Kismet/Private/BlueprintEditor.cpp
--- a/Engine/Source/Editor/Kismet/Private/BlueprintEditor.cpp
+++ b/Engine/Source/Editor/Kismet/Private/BlueprintEditor.cpp
@@ -628,7 +628,7 @@ void FBlueprintEditor::RefreshEditors(ERefreshBlueprintEditorReason::Type Reason
 	{
 		StartEditingDefaults(/*bAutoFocus=*/ false, true);
 	}
-	else
+	else if (CurrentUISelection != FName())
 	{
 		ensureMsgf(false, TEXT("Unknown UI selection state in FBlueprintEditor::RefreshEditors"));
 	}
