[INTEGRATE] Change 2447217 by [REDACTED_PERSON][REDACTED_HANDLE] on 2015/02/16 11:15:50

	Use asset broker to properly set up a child actor via AddNewComponent

diff --git a/Engine/Source/Editor/Kismet/Private/SSCSEditor.cpp b/Engine/Source/Editor/Kismet/Private/SSCSEditor.cpp
--- a/Engine/Source/Editor/Kismet/Private/SSCSEditor.cpp
+++ b/Engine/Source/Editor/Kismet/Private/SSCSEditor.cpp
@@ -185,12 +185,7 @@ FReply SSCSEditor::TryHandleAssetDragDropOperation(const FDragDropEvent& DragDro
 				}
 				else if ((PotentialActorClass != nullptr) && !PotentialActorClass->HasAnyClassFlags(CLASS_Deprecated | CLASS_Abstract | CLASS_NewerVersionExists))
 				{
-					const FScopedTransaction Transaction( LOCTEXT("AddComponent", "Add Component") );
-					if (UChildActorComponent* ChildActorComponent = Cast<UChildActorComponent>(AddNewComponent(UChildActorComponent::StaticClass(), nullptr)))
-					{
-						ChildActorComponent->ChildActorClass = PotentialActorClass;
-						ChildActorComponent->CreateChildActor();
-					}
+					AddNewComponent(UChildActorComponent::StaticClass(), PotentialActorClass);
 				}
 			}
 
diff --git a/Engine/Source/Editor/UnrealEd/Private/ComponentAssetBroker.cpp b/Engine/Source/Editor/UnrealEd/Private/ComponentAssetBroker.cpp
--- a/Engine/Source/Editor/UnrealEd/Private/ComponentAssetBroker.cpp
+++ b/Engine/Source/Editor/UnrealEd/Private/ComponentAssetBroker.cpp
@@ -215,10 +215,17 @@ public:
 	{
 		if (UChildActorComponent* ChildActorComp = Cast<UChildActorComponent>(InComponent))
 		{
-			UBlueprint* BP = Cast<UBlueprint>(InAsset);
-			if (BP != NULL)
+			UClass* Class = Cast<UClass>(InAsset);
+			if (Class == nullptr)
 			{
-				ChildActorComp->ChildActorClass = TSubclassOf<AActor>(*(BP->GeneratedClass));
+				if (UBlueprint* BP = Cast<UBlueprint>(InAsset))
+				{
+					Class = *(BP->GeneratedClass);
+				}
+			}
+			if (Class && Class->IsChildOf<AActor>())
+			{
+				ChildActorComp->ChildActorClass = Class;
 				return true;
 			}
 		}
