[INTEGRATE] Change 2447405 by [REDACTED_PERSON][REDACTED_HANDLE] on 2015/02/16 13:25:24

	[UE-9901] Code review indicates that this ensure is just invalid.

diff --git a/Engine/Source/Editor/Kismet/Private/BlueprintEditor.cpp b/Engine/Source/Editor/Kismet/Private/BlueprintEditor.cpp
--- a/Engine/Source/Editor/Kismet/Private/BlueprintEditor.cpp
+++ b/Engine/Source/Editor/Kismet/Private/BlueprintEditor.cpp
@@ -628,10 +628,6 @@ void FBlueprintEditor::RefreshEditors(ERefreshBlueprintEditorReason::Type Reason
 	{
 		StartEditingDefaults(/*bAutoFocus=*/ false, true);
 	}
-	else if (CurrentUISelection != FName())
-	{
-		ensureMsgf(false, TEXT("Unknown UI selection state in FBlueprintEditor::RefreshEditors"));
-	}
 
 	//@TODO: Should determine when we need to do the invalid/refresh business and if the graph node selection change
 	// under non-compiles is necessary (except when the selection mode is appropriate, as already detected above)
