Added view mode menu to the SCS Editor Viewport.

diff --git a/Engine/Source/Editor/Kismet/Private/SSCSEditorViewport.cpp b/Engine/Source/Editor/Kismet/Private/SSCSEditorViewport.cpp
--- a/Engine/Source/Editor/Kismet/Private/SSCSEditorViewport.cpp
+++ b/Engine/Source/Editor/Kismet/Private/SSCSEditorViewport.cpp
@@ -62,6 +62,17 @@ public:
 					.OnGetMenuContent(this, &SSCSEditorViewportToolBar::GenerateCameraMenu)
 				]
 				+ SHorizontalBox::Slot()
+				.AutoWidth()
+				.Padding(2.0f, 2.0f)
+				[
+					SNew( SEditorViewportToolbarMenu )
+					.ParentToolBar( SharedThis( this ) )
+					.Cursor( EMouseCursor::Default )
+					.Label(this, &SSCSEditorViewportToolBar::GetViewMenuLabel)
+					.LabelIcon(this, &SSCSEditorViewportToolBar::GetViewMenuLabelIcon)
+					.OnGetMenuContent(this, &SSCSEditorViewportToolBar::GenerateViewMenu)
+				]
+				+ SHorizontalBox::Slot()
 				.Padding( 3.0f, 1.0f )
 				.HAlign( HAlign_Right )
 				[
@@ -132,6 +143,11 @@ public:
 
 	const FSlateBrush* GetCameraMenuLabelIcon() const
 	{
+		static FName PerspectiveIconName("EditorViewport.Perspective");
+		static FName TopIconName("EditorViewport.Top");
+		static FName SideIconName("EditorViewport.Side");
+		static FName FrontIconName("EditorViewport.Front");
+
 		FName Icon = NAME_None;
 
 		if(EditorViewport.IsValid())
@@ -139,19 +155,19 @@ public:
 			switch(EditorViewport.Pin()->GetViewportClient()->GetViewportType())
 			{
 			case LVT_Perspective:
-				Icon = FName("EditorViewport.Perspective");
+				Icon = PerspectiveIconName;
 				break;
 
 			case LVT_OrthoXY:
-				Icon = FName("EditorViewport.Top");
+				Icon = TopIconName;
 				break;
 
 			case LVT_OrthoYZ:
-				Icon = FName( "EditorViewport.Side");
+				Icon = SideIconName;
 				break;
 
 			case LVT_OrthoXZ:
-				Icon = FName("EditorViewport.Front");
+				Icon = FrontIconName;
 				break;
 			}
 		}
@@ -177,6 +193,74 @@ public:
 		return CameraMenuBuilder.MakeWidget();
 	}
 
+	FText GetViewMenuLabel() const
+	{
+		FText Label = NSLOCTEXT("BlueprintEditor", "ViewMenuTitle_Default", "View");
+
+		if (EditorViewport.IsValid())
+		{
+			switch (EditorViewport.Pin()->GetViewportClient()->GetViewMode())
+			{
+			case VMI_Lit:
+				Label = NSLOCTEXT("BlueprintEditor", "ViewMenuTitle_Lit", "Lit");
+				break;
+
+			case VMI_Unlit:
+				Label = NSLOCTEXT("BlueprintEditor", "ViewMenuTitle_Unlit", "Unlit");
+				break;
+
+			case VMI_BrushWireframe:
+				Label = NSLOCTEXT("BlueprintEditor", "ViewMenuTitle_Wireframe", "Wireframe");
+				break;
+			}
+		}
+
+		return Label;
+	}
+
+	const FSlateBrush* GetViewMenuLabelIcon() const
+	{
+		static FName LitModeIconName("EditorViewport.LitMode");
+		static FName UnlitModeIconName("EditorViewport.UnlitMode");
+		static FName WireframeModeIconName("EditorViewport.WireframeMode");
+
+		FName Icon = NAME_None;
+
+		if (EditorViewport.IsValid())
+		{
+			switch (EditorViewport.Pin()->GetViewportClient()->GetViewMode())
+			{
+			case VMI_Lit:
+				Icon = LitModeIconName;
+				break;
+
+			case VMI_Unlit:
+				Icon = UnlitModeIconName;
+				break;
+
+			case VMI_BrushWireframe:
+				Icon = WireframeModeIconName;
+				break;
+			}
+		}
+
+		return FEditorStyle::GetBrush(Icon);
+	}
+
+	TSharedRef<SWidget> GenerateViewMenu() const
+	{
+		TSharedPtr<const FUICommandList> CommandList = EditorViewport.IsValid() ? EditorViewport.Pin()->GetCommandList() : nullptr;
+
+		const bool bInShouldCloseWindowAfterMenuSelection = true;
+		FMenuBuilder ViewMenuBuilder(bInShouldCloseWindowAfterMenuSelection, CommandList);
+
+		ViewMenuBuilder.AddMenuEntry(FEditorViewportCommands::Get().LitMode, NAME_None, NSLOCTEXT("BlueprintEditor", "LitModeMenuOption", "Lit"));
+		ViewMenuBuilder.AddMenuEntry(FEditorViewportCommands::Get().UnlitMode, NAME_None, NSLOCTEXT("BlueprintEditor", "UnlitModeMenuOption", "Unlit"));
+		ViewMenuBuilder.AddMenuEntry(FEditorViewportCommands::Get().WireframeMode, NAME_None, NSLOCTEXT("BlueprintEditor", "WireframeModeMenuOption", "Wireframe"));
+
+		return ViewMenuBuilder.MakeWidget();
+	}
+
 private:
 	/** Reference to the parent viewport */
 	TWeakPtr<SSCSEditorViewport> EditorViewport;
