PR #826: Assert for components from C++ (SSCSEditor.cpp:1609) (Contributed by [REDACTED_PERSON])

https://github.com/EpicGames/UnrealEngine/pull/826

diff --git a/Engine/Source/Editor/Kismet/Private/SSCSEditor.cpp b/Engine/Source/Editor/Kismet/Private/SSCSEditor.cpp
--- a/Engine/Source/Editor/Kismet/Private/SSCSEditor.cpp
+++ b/Engine/Source/Editor/Kismet/Private/SSCSEditor.cpp
@@ -1606,8 +1606,14 @@ FText SSCS_RowWidget::GetIntroducedInToolTipText() const
 
 				if (BestClass == nullptr)
 				{
-					ensure(ComponentTemplate->IsCreatedByConstructionScript());
-					IntroducedInTooltip = LOCTEXT("IntroducedInUnknownError", "Unknown Blueprint Class (via an Add Component call)");
+					if (ComponentTemplate->IsCreatedByConstructionScript()) 
+					{
+						IntroducedInTooltip = LOCTEXT("IntroducedInUnknownError", "Unknown Blueprint Class (via an Add Component call)");
+					} 
+					else 
+					{
+						IntroducedInTooltip = LOCTEXT("IntroducedInNativeError", "Unknown native source (via C++ code)");
+					}
 				}
 				else
 				{
