Fixes a crash when closing a Blueprint Editor while a variable is selected.

diff --git a/Engine/Source/Editor/Kismet/Private/BlueprintDetailsCustomization.cpp b/Engine/Source/Editor/Kismet/Private/BlueprintDetailsCustomization.cpp
--- a/Engine/Source/Editor/Kismet/Private/BlueprintDetailsCustomization.cpp
+++ b/Engine/Source/Editor/Kismet/Private/BlueprintDetailsCustomization.cpp
@@ -169,11 +169,14 @@ int32 FBlueprintDetails::HandleAddOrViewIndexForButton(const FName EventName, FN
 
 FBlueprintVarActionDetails::~FBlueprintVarActionDetails()
 {
-	// Remove the callback delegate we registered for
-	TWeakPtr<FBlueprintEditor> BlueprintEditor = MyBlueprint.Pin()->GetBlueprintEditor();
-	if( BlueprintEditor.IsValid() )
+	if(MyBlueprint.IsValid())
 	{
-		BlueprintEditor.Pin()->OnRefresh().RemoveAll(this);
+		// Remove the callback delegate we registered for
+		TWeakPtr<FBlueprintEditor> BlueprintEditor = MyBlueprint.Pin()->GetBlueprintEditor();
+		if( BlueprintEditor.IsValid() )
+		{
+			BlueprintEditor.Pin()->OnRefresh().RemoveAll(this);
+		}
 	}
 }
 
