No longer given the option to override functions in macro libraries in the right click menu of the My Blueprint window.

diff --git a/Engine/Source/Editor/Kismet/Private/SMyBlueprint.cpp b/Engine/Source/Editor/Kismet/Private/SMyBlueprint.cpp
--- a/Engine/Source/Editor/Kismet/Private/SMyBlueprint.cpp
+++ b/Engine/Source/Editor/Kismet/Private/SMyBlueprint.cpp
@@ -1859,7 +1859,8 @@ void SMyBlueprint::BuildAddNewMenu(FMenuBuilder& MenuBuilder)
 		MenuBuilder.AddMenuEntry(FBlueprintEditorCommands::Get().AddNewLocalVariable);
 		MenuBuilder.AddMenuEntry(FBlueprintEditorCommands::Get().AddNewFunction);
 
-		if ( OverridableFunctionActions.Num() > 0 )
+		// If we cannot handle Function Graphs, we cannot handle function overrides
+		if ( OverridableFunctionActions.Num() > 0 && BlueprintEditorPtr.Pin()->NewDocument_IsVisibleForType(FBlueprintEditor::CGT_NewFunctionGraph) )
 		{
 			MenuBuilder.AddSubMenu(
 				LOCTEXT("OverrideFunction", "Override Function"),
