Dragging local variables over other actions in the My Blueprint window should give an error stating that it is invalid to do.

diff --git a/Engine/Source/Editor/Kismet/Private/BPVariableDragDropAction.cpp b/Engine/Source/Editor/Kismet/Private/BPVariableDragDropAction.cpp
--- a/Engine/Source/Editor/Kismet/Private/BPVariableDragDropAction.cpp
+++ b/Engine/Source/Editor/Kismet/Private/BPVariableDragDropAction.cpp
@@ -54,8 +54,8 @@ void FKismetVariableDragDropAction::HoverTargetChanged()
 
 	// Icon/text to draw on tooltip
 	FSlateColor IconColor = FLinearColor::White;
-	const FSlateBrush* StatusSymbol = FEditorStyle::GetBrush(TEXT("NoBrush")); 
-	FText Message;
+	const FSlateBrush* StatusSymbol = FEditorStyle::GetBrush(TEXT("Graph.ConnectorFeedback.Error"));
+	FText Message = LOCTEXT("InvalidDropTarget", "Invalid drop target!");
 
 	UEdGraphPin* PinUnderCursor = GetHoveredPin();
 
