Fixed animation blueprints having the wrong corner text

UE-10299 - Animation Blueprints have the wrong corner text

There were two versions of GetGraphAppearance, one which took a parameter, and one which did not. Only the version which didn't take the parameter was being overridden by Persona, which is the version that was never being used.

I renamed the parameterless version to GetCurrentGraphAppearance (now non-virtual), and then made sure that all derived BP editors were overriding the correct version (now virtual).

diff --git a/Engine/Source/Editor/Kismet/Private/BlueprintEditor.cpp b/Engine/Source/Editor/Kismet/Private/BlueprintEditor.cpp
--- a/Engine/Source/Editor/Kismet/Private/BlueprintEditor.cpp
+++ b/Engine/Source/Editor/Kismet/Private/BlueprintEditor.cpp
@@ -1201,7 +1201,7 @@ TSharedRef<SGraphEditor> FBlueprintEditor::CreateGraphEditorWidget(TSharedRef<FT
 	return Editor;
 }
 
-FGraphAppearanceInfo FBlueprintEditor::GetGraphAppearance() const
+FGraphAppearanceInfo FBlueprintEditor::GetCurrentGraphAppearance() const
 {
 	return GetGraphAppearance(GetFocusedGraph());
 }
diff --git a/Engine/Source/Editor/Kismet/Public/BlueprintEditor.h b/Engine/Source/Editor/Kismet/Public/BlueprintEditor.h
--- a/Engine/Source/Editor/Kismet/Public/BlueprintEditor.h
+++ b/Engine/Source/Editor/Kismet/Public/BlueprintEditor.h
@@ -879,10 +879,11 @@ protected:
 	virtual void	PostRedo(bool bSuccess) override;
 	// End of FEditorUndoClient
 
-	/** Get graph appearance */
-	virtual FGraphAppearanceInfo GetGraphAppearance() const;
-	/** Used to get the apperance of a specific graph, GetGraphAppearance() uses the currently focused graph. */
-	FGraphAppearanceInfo GetGraphAppearance(class UEdGraph* InGraph) const;
+	/** Get the graph appearance of the currently focused graph */
+	FGraphAppearanceInfo GetCurrentGraphAppearance() const;
+
+	/** Get the graph appearance of a specific graph, GetCurrentGraphAppearance() uses the currently focused graph. */
+	virtual FGraphAppearanceInfo GetGraphAppearance(class UEdGraph* InGraph) const;
 
 	/** Attempts to invoke the details tab if it's currently possible to. */
 	void TryInvokingDetailsTab(bool bFlash = true);
diff --git a/Engine/Source/Editor/Persona/Private/Persona.cpp b/Engine/Source/Editor/Persona/Private/Persona.cpp
--- a/Engine/Source/Editor/Persona/Private/Persona.cpp
+++ b/Engine/Source/Editor/Persona/Private/Persona.cpp
@@ -1958,9 +1958,9 @@ void FPersona::CreateDefaultTabContents(const TArray<UBlueprint*>& InBlueprints)
 	PreviewEditor = SNew(SPersonaPreviewPropertyEditor, SharedThis(this));
 }
 
-FGraphAppearanceInfo FPersona::GetGraphAppearance() const
+FGraphAppearanceInfo FPersona::GetGraphAppearance(UEdGraph* InGraph) const
 {
-	FGraphAppearanceInfo AppearanceInfo = FBlueprintEditor::GetGraphAppearance();
+	FGraphAppearanceInfo AppearanceInfo = FBlueprintEditor::GetGraphAppearance(InGraph);
 
 	if ( GetBlueprintObj()->IsA(UAnimBlueprint::StaticClass()) )
 	{
diff --git a/Engine/Source/Editor/Persona/Private/Persona.h b/Engine/Source/Editor/Persona/Private/Persona.h
--- a/Engine/Source/Editor/Persona/Private/Persona.h
+++ b/Engine/Source/Editor/Persona/Private/Persona.h
@@ -226,7 +226,7 @@ protected:
 	virtual void OnOpenRelatedAsset() override;
 	virtual void GetCustomDebugObjects(TArray<FCustomDebugObject>& DebugList) const override;
 	virtual void CreateDefaultTabContents(const TArray<UBlueprint*>& InBlueprints) override;
-	virtual FGraphAppearanceInfo GetGraphAppearance() const override;
+	virtual FGraphAppearanceInfo GetGraphAppearance(class UEdGraph* InGraph) const override;
 	virtual bool IsEditable(UEdGraph* InGraph) const override;
 	virtual FText GetGraphDecorationString(UEdGraph* InGraph) const override;
 	// End of FBlueprintEditor interface
diff --git a/Engine/Source/Editor/UMGEditor/Private/WidgetBlueprintEditor.cpp b/Engine/Source/Editor/UMGEditor/Private/WidgetBlueprintEditor.cpp
--- a/Engine/Source/Editor/UMGEditor/Private/WidgetBlueprintEditor.cpp
+++ b/Engine/Source/Editor/UMGEditor/Private/WidgetBlueprintEditor.cpp
@@ -670,9 +670,9 @@ void FWidgetBlueprintEditor::UpdatePreview(UBlueprint* InBlueprint, bool bInForc
 	OnWidgetPreviewUpdated.Broadcast();
 }
 
-FGraphAppearanceInfo FWidgetBlueprintEditor::GetGraphAppearance() const
+FGraphAppearanceInfo FWidgetBlueprintEditor::GetGraphAppearance(UEdGraph* InGraph) const
 {
-	FGraphAppearanceInfo AppearanceInfo = FBlueprintEditor::GetGraphAppearance();
+	FGraphAppearanceInfo AppearanceInfo = FBlueprintEditor::GetGraphAppearance(InGraph);
 
 	if ( GetBlueprintObj()->IsA(UWidgetBlueprint::StaticClass()) )
 	{
diff --git a/Engine/Source/Editor/UMGEditor/Private/WidgetBlueprintEditor.h b/Engine/Source/Editor/UMGEditor/Private/WidgetBlueprintEditor.h
--- a/Engine/Source/Editor/UMGEditor/Private/WidgetBlueprintEditor.h
+++ b/Engine/Source/Editor/UMGEditor/Private/WidgetBlueprintEditor.h
@@ -143,7 +143,7 @@ public:
 protected:
 	// Begin FBlueprintEditor
 	virtual void RegisterApplicationModes(const TArray<UBlueprint*>& InBlueprints, bool bShouldOpenInDefaultsMode, bool bNewlyCreated = false) override;
-	virtual FGraphAppearanceInfo GetGraphAppearance() const override;
+	virtual FGraphAppearanceInfo GetGraphAppearance(class UEdGraph* InGraph) const override;
 	// End FBlueprintEditor
 
 private:
