Put in safety checks to prevent a crash when undoing pasting a Custom Event node.

diff --git a/Engine/Source/Editor/Kismet/Private/BlueprintDetailsCustomization.cpp b/Engine/Source/Editor/Kismet/Private/BlueprintDetailsCustomization.cpp
--- a/Engine/Source/Editor/Kismet/Private/BlueprintDetailsCustomization.cpp
+++ b/Engine/Source/Editor/Kismet/Private/BlueprintDetailsCustomization.cpp
@@ -2692,16 +2692,21 @@ FText FBlueprintGraphActionDetails::GetCurrentReplicatedEventString() const
 
 	uint32 const ReplicatedNetMask = (FUNC_NetMulticast | FUNC_NetServer | FUNC_NetClient);
 
-	uint32 NetFlags = CustomEvent->FunctionFlags & ReplicatedNetMask;
-	if (CustomEvent->IsOverride())
+	FText ReplicationText;
+
+	if(CustomEvent)
 	{
-		UFunction* SuperFunction = FindField<UFunction>(CustomEvent->GetBlueprint()->ParentClass, CustomEvent->CustomFunctionName);
-		check(SuperFunction != NULL);
+		uint32 NetFlags = CustomEvent->FunctionFlags & ReplicatedNetMask;
+		if (CustomEvent->IsOverride())
+		{
+			UFunction* SuperFunction = FindField<UFunction>(CustomEvent->GetBlueprint()->ParentClass, CustomEvent->CustomFunctionName);
+			check(SuperFunction != NULL);
 
-		NetFlags = SuperFunction->FunctionFlags & ReplicatedNetMask;
+			NetFlags = SuperFunction->FunctionFlags & ReplicatedNetMask;
+		}
+		ReplicationText = ReplicationSpecifierProperName(NetFlags);
 	}
-
-	return ReplicationSpecifierProperName(NetFlags);
+	return ReplicationText;
 }
 
 bool FBaseBlueprintGraphActionDetails::ConditionallyCleanUpResultNode()
