[UE-10081] - Right-clicking an actor in the scene outliner always deselects components (even if the actor was already selected), so the proper context menu shows
- Removed the actor context menu options from the actor root in the components tree; can now only paste a component onto it
    -> Necessary because multiple details panels can be open and locked to different actors, so the context menu options would not necessarily pertain to the right-clicked actor root
- Fixed ProcesFCustomizableTextObjectFactory::ProcessBuffer attempting to load an archetype when none had been parsed from T3D

diff --git a/Engine/Source/Editor/Kismet/Private/SSCSEditor.cpp b/Engine/Source/Editor/Kismet/Private/SSCSEditor.cpp
--- a/Engine/Source/Editor/Kismet/Private/SSCSEditor.cpp
+++ b/Engine/Source/Editor/Kismet/Private/SSCSEditor.cpp
@@ -3003,7 +3003,6 @@ void SSCSEditor::Construct( const FArguments& InArgs )
 	ActorContext = InArgs._ActorContext;
 	AllowEditing = InArgs._AllowEditing;
 	PreviewActor = InArgs._PreviewActor;
-	ActorMenuExtender = InArgs._ActorMenuExtender;
 	OnSelectionUpdated = InArgs._OnSelectionUpdated;
 	OnItemDoubleClicked = InArgs._OnItemDoubleClicked;
 	OnHighlightPropertyInDetailsView = InArgs._OnHighlightPropertyInDetailsView;
@@ -3385,24 +3384,15 @@ TSharedPtr< SWidget > SSCSEditor::CreateContextMenu()
 	if (SelectedNodes.Num() > 0 || CanPasteNodes())
 	{
 		const bool CloseAfterSelection = true;
-		FMenuBuilder MenuBuilder( CloseAfterSelection, CommandList, ActorMenuExtender );
+		FMenuBuilder MenuBuilder( CloseAfterSelection, CommandList );
 
-		bool bShowJustPasteOption = false;
+		bool bOnlyShowPasteOption = false;
 
 		if (SelectedNodes.Num() > 0)
 		{
 			if (SelectedNodes.Num() == 1 && SelectedNodes[0]->GetNodeType() == FSCSEditorTreeNode::RootActorNode)
 			{
-				if (EditorMode == EComponentEditorMode::BlueprintSCS)
-				{
-					bShowJustPasteOption = false;
-				}
-				else
-				{
-					// Display the Actor menu
-					MenuBuilder.BeginSection("MainSection");
-					MenuBuilder.EndSection();
-				}
+				bOnlyShowPasteOption = true;
 			}
 			else
 			{
@@ -3410,11 +3400,11 @@ TSharedPtr< SWidget > SSCSEditor::CreateContextMenu()
 				{
 					if (SelectedNode->GetNodeType() != FSCSEditorTreeNode::ComponentNode)
 					{
-						bShowJustPasteOption = true;
+						bOnlyShowPasteOption = true;
 						break;
 					}
 				}
-				if (!bShowJustPasteOption)
+				if (!bOnlyShowPasteOption)
 				{
 					TArray<UActorComponent*> SelectedComponents;
 					TArray<FSCSEditorTreeNodePtrType> SelectedNodes = GetSelectedNodes();
@@ -3467,12 +3457,12 @@ TSharedPtr< SWidget > SSCSEditor::CreateContextMenu()
 		}
 		else
 		{
-			bShowJustPasteOption = true;
+			bOnlyShowPasteOption = true;
 		}
 
-		if (bShowJustPasteOption)
+		if (bOnlyShowPasteOption)
 		{
-			MenuBuilder.BeginSection("ComponentActions", LOCTEXT("ComponentContextMenu", "Component Actions") );
+			MenuBuilder.BeginSection("PasteComponent", LOCTEXT("EditComponentHeading", "Edit") );
 			{
 				MenuBuilder.AddMenuEntry( FGenericCommands::Get().Paste );
 			}
@@ -4534,7 +4524,7 @@ bool SSCSEditor::CanCopyNodes() const
 		}
 	}
 
-	// Copy the components to the clipboard
+	// Verify that the components can be copied
 	return FComponentEditorUtils::CanCopyComponents(ComponentsToCopy);
 }
 
@@ -4568,7 +4558,7 @@ bool SSCSEditor::CanPasteNodes() const
 		return false;
 	}
 
-	return FComponentEditorUtils::CanPasteComponents(Cast<USceneComponent>(SceneRootNodePtr->GetComponentTemplate()), SceneRootNodePtr->IsDefaultSceneRoot());
+	return FComponentEditorUtils::CanPasteComponents(Cast<USceneComponent>(SceneRootNodePtr->GetComponentTemplate()), SceneRootNodePtr->IsDefaultSceneRoot(), true);
 }
 
 void SSCSEditor::PasteNodes()
@@ -4645,8 +4635,8 @@ void SSCSEditor::PasteNodes()
 	}
 	else    // EComponentEditorMode::ActorInstance
 	{
-		// Determine where in the hierarchy to paste
-		USceneComponent* TargetComponent = nullptr;
+		// Determine where in the hierarchy to paste (default to the root)
+		USceneComponent* TargetComponent = GetActorContext()->GetRootComponent();
 		auto SelectedNodes = GetSelectedNodes();
 		for (int32 i = 0; i < SelectedNodes.Num(); ++i)
 		{
diff --git a/Engine/Source/Editor/Kismet/Public/SSCSEditor.h b/Engine/Source/Editor/Kismet/Public/SSCSEditor.h
--- a/Engine/Source/Editor/Kismet/Public/SSCSEditor.h
+++ b/Engine/Source/Editor/Kismet/Public/SSCSEditor.h
@@ -716,7 +716,6 @@ public:
 		{}
 
 		SLATE_ARGUMENT(EComponentEditorMode::Type, EditorMode)
-		SLATE_ARGUMENT(TSharedPtr<FExtender>, ActorMenuExtender)
 		SLATE_ATTRIBUTE(class AActor*, ActorContext)
 		SLATE_ATTRIBUTE(class AActor*, PreviewActor)
 		SLATE_ATTRIBUTE(bool, AllowEditing)
diff --git a/Engine/Source/Editor/LevelEditor/Private/SActorDetails.cpp b/Engine/Source/Editor/LevelEditor/Private/SActorDetails.cpp
--- a/Engine/Source/Editor/LevelEditor/Private/SActorDetails.cpp
+++ b/Engine/Source/Editor/LevelEditor/Private/SActorDetails.cpp
@@ -59,8 +59,7 @@ void SActorDetails::Construct(const FArguments& InArgs, const FName TabIdentifie
 		.EditorMode(EComponentEditorMode::ActorInstance)
 		.AllowEditing(this, &SActorDetails::GetAllowComponentTreeEditing)
 		.ActorContext(this, &SActorDetails::GetActorContext)
-		.OnSelectionUpdated(this, &SActorDetails::OnSCSEditorTreeViewSelectionChanged)
-		.ActorMenuExtender(InArgs._ActorMenuExtender);
+		.OnSelectionUpdated(this, &SActorDetails::OnSCSEditorTreeViewSelectionChanged);
 		
 	ComponentsBox->SetContent(SCSEditor.ToSharedRef());
 
diff --git a/Engine/Source/Editor/LevelEditor/Private/SActorDetails.h b/Engine/Source/Editor/LevelEditor/Private/SActorDetails.h
--- a/Engine/Source/Editor/LevelEditor/Private/SActorDetails.h
+++ b/Engine/Source/Editor/LevelEditor/Private/SActorDetails.h
@@ -9,7 +9,6 @@ class SActorDetails : public SCompoundWidget, public FEditorUndoClient
 {
 public:
 	SLATE_BEGIN_ARGS(SActorDetails) {}
-		SLATE_ARGUMENT( TSharedPtr<FExtender>, ActorMenuExtender )
 	SLATE_END_ARGS()
 
 	void Construct(const FArguments& InArgs, const FName TabIdentifier);
diff --git a/Engine/Source/Editor/LevelEditor/Private/SLevelEditor.cpp b/Engine/Source/Editor/LevelEditor/Private/SLevelEditor.cpp
--- a/Engine/Source/Editor/LevelEditor/Private/SLevelEditor.cpp
+++ b/Engine/Source/Editor/LevelEditor/Private/SLevelEditor.cpp
@@ -36,7 +36,7 @@
 #include "TutorialMetaData.h"
 #include "SDockTab.h"
 #include "SActorDetails.h"
-
+#include "ScopedTransaction.h"
 
 
 static const FName LevelEditorBuildAndSubmitTab("LevelEditorBuildAndSubmit");
@@ -511,7 +511,7 @@ TSharedRef<FTabManager> SLevelEditor::GetTabManager() const
 
 
 
-TSharedRef<SDockTab> SLevelEditor::SummonDetailsPanel( FName TabIdentifier, TSharedPtr<FExtender> ActorMenuExtender )
+TSharedRef<SDockTab> SLevelEditor::SummonDetailsPanel( FName TabIdentifier )
 {
 	TSharedPtr<SActorDetails> ActorDetails;
 
@@ -526,7 +526,6 @@ TSharedRef<SDockTab> SLevelEditor::SummonDetailsPanel( FName TabIdentifier, TSha
 			.AddMetaData<FTutorialMetaData>(FTutorialMetaData(TEXT("ActorDetails"), TEXT("LevelEditorSelectionDetails")))
 			[
 				SAssignNew( ActorDetails, SActorDetails, TabIdentifier )
-					.ActorMenuExtender(ActorMenuExtender)
 			]
 		];
 
@@ -574,20 +573,7 @@ TSharedRef<SDockTab> SLevelEditor::SpawnLevelEditorTab( const FSpawnTabArgs& Arg
 	}
 	else if( TabIdentifier == TEXT("LevelEditorSelectionDetails") || TabIdentifier == TEXT("LevelEditorSelectionDetails2") || TabIdentifier == TEXT("LevelEditorSelectionDetails3") || TabIdentifier == TEXT("LevelEditorSelectionDetails4") )
 	{
-		TWeakPtr<SLevelEditor> WeakLevelEditor = SharedThis(this);
-		TSharedPtr<FExtender> MenuExtender = MakeShareable(new FExtender);
-		MenuExtender->AddMenuExtension(
-			"MainSection", EExtensionHook::Before, GetLevelEditorActions(),
-			FMenuExtensionDelegate::CreateStatic([](FMenuBuilder& MenuBuilder, TWeakPtr<SLevelEditor> InWeakLevelEditor){
-				// Only extend the menu if we have actors selected
-				if (GEditor->GetSelectedActors()->Num())
-				{
-					FLevelEditorContextMenu::FillMenu(MenuBuilder, InWeakLevelEditor, LevelEditorMenuContext::NonViewport, TSharedPtr<FExtender>());
-				}
-			}, WeakLevelEditor)
-		);
-
-		TSharedRef<SDockTab> DetailsPanel = SummonDetailsPanel( TabIdentifier, MenuExtender );
+		TSharedRef<SDockTab> DetailsPanel = SummonDetailsPanel( TabIdentifier );
 		GUnrealEd->UpdateFloatingPropertyWindows();
 		return DetailsPanel;
 	}
@@ -636,7 +622,7 @@ TSharedRef<SDockTab> SLevelEditor::SpawnLevelEditorTab( const FSpawnTabArgs& Arg
 				"MainSection", EExtensionHook::Before, GetLevelEditorActions(),
 				FMenuExtensionDelegate::CreateStatic([](FMenuBuilder& MenuBuilder, TWeakPtr<SLevelEditor> InWeakLevelEditor){
 					// Only extend the menu if we have actors selected
-					if (GEditor->GetSelectedActors()->Num())
+					if (GEditor->GetSelectedActorCount() > 0)
 					{
 						FLevelEditorContextMenu::FillMenu(MenuBuilder, InWeakLevelEditor, LevelEditorMenuContext::NonViewport, TSharedPtr<FExtender>());
 					}
diff --git a/Engine/Source/Editor/LevelEditor/Private/SLevelEditor.h b/Engine/Source/Editor/LevelEditor/Private/SLevelEditor.h
--- a/Engine/Source/Editor/LevelEditor/Private/SLevelEditor.h
+++ b/Engine/Source/Editor/LevelEditor/Private/SLevelEditor.h
@@ -127,7 +127,7 @@ private:
 	
 	TSharedRef<SDockTab> SpawnLevelEditorTab(const FSpawnTabArgs& Args, FName TabIdentifier, FString InitializationPayload);
 	//TSharedRef<SDockTab> SpawnLevelEditorModeTab(const FSpawnTabArgs& Args, FEdMode* EditorMode);
-	TSharedRef<SDockTab> SummonDetailsPanel( FName Identifier, TSharedPtr<FExtender> ActorMenuExtender = nullptr );
+	TSharedRef<SDockTab> SummonDetailsPanel( FName Identifier );
 
 	/**
 	 * Binds UI commands to actions for the level editor                   
diff --git a/Engine/Source/Editor/SceneOutliner/Private/SSceneOutliner.cpp b/Engine/Source/Editor/SceneOutliner/Private/SSceneOutliner.cpp
--- a/Engine/Source/Editor/SceneOutliner/Private/SSceneOutliner.cpp
+++ b/Engine/Source/Editor/SceneOutliner/Private/SSceneOutliner.cpp
@@ -1056,6 +1056,19 @@ namespace SceneOutliner
 
 		if (SharedData->Mode == ESceneOutlinerMode::ActorBrowsing)
 		{
+			// Make sure that no components are selected
+			if (GEditor->GetSelectedComponentCount() > 0)
+			{
+				// We want to be able to undo to regain the previous component selection
+				const FScopedTransaction Transaction(NSLOCTEXT("UnrealEd", "ClickingOnActorsContextMenu", "Clicking on Actors (context menu)"));
+				USelection* ComponentSelection = GEditor->GetSelectedComponents();
+				ComponentSelection->Modify(false);
+				ComponentSelection->DeselectAll();
+
+				GUnrealEd->UpdatePivotLocationForSelection();
+				GEditor->RedrawLevelEditingViewports(false);
+			}
+
 			return BuildDefaultContextMenu();
 		}
 
diff --git a/Engine/Source/Editor/UnrealEd/Private/Factories/EditorFactories.cpp b/Engine/Source/Editor/UnrealEd/Private/Factories/EditorFactories.cpp
--- a/Engine/Source/Editor/UnrealEd/Private/Factories/EditorFactories.cpp
+++ b/Engine/Source/Editor/UnrealEd/Private/Factories/EditorFactories.cpp
@@ -5300,10 +5300,12 @@ void FCustomizableTextObjectFactory::ProcessBuffer(UObject* InParent, EObjectFla
 				FParse::Value( Str, TEXT("NAME="), ObjName );
 
 				// Setup archetype
+				UObject* ObjArchetype = nullptr;
 				FString ObjArchetypeName;
-				FParse::Value( Str, TEXT("ARCHETYPE="), ObjArchetypeName );
-				UObject* ObjArchetype;
-				ObjArchetype = LoadObject<UObject>(nullptr, *ObjArchetypeName, nullptr, LOAD_None, nullptr);
+				if (FParse::Value(Str, TEXT("ARCHETYPE="), ObjArchetypeName))
+				{
+					ObjArchetype = LoadObject<UObject>(nullptr, *ObjArchetypeName, nullptr, LOAD_None, nullptr);
+				}
 
 				// Make sure this name is not used by anything else. Will rename other stuff if necessary
 				ClearObjectNameUsage(InParent, ObjName);
diff --git a/Engine/Source/Editor/UnrealEd/Private/Kismet2/ComponentEditorUtils.cpp b/Engine/Source/Editor/UnrealEd/Private/Kismet2/ComponentEditorUtils.cpp
--- a/Engine/Source/Editor/UnrealEd/Private/Kismet2/ComponentEditorUtils.cpp
+++ b/Engine/Source/Editor/UnrealEd/Private/Kismet2/ComponentEditorUtils.cpp
@@ -266,12 +266,18 @@ void FComponentEditorUtils::CopyComponents(const TArray<UActorComponent*>& Compo
 	// Duplicate the selected component templates into temporary objects that we can modify
 	TMap<FName, FName> ParentMap;
 	TMap<FName, UActorComponent*> ObjectMap;
-	for (auto Component : ComponentsToCopy)
+	for (UActorComponent* Component : ComponentsToCopy)
 	{
 		// Duplicate the component into a temporary object
 		UObject* DuplicatedComponent = StaticDuplicateObject(Component, GetTransientPackage(), *Component->GetName(), RF_AllFlags & ~RF_ArchetypeObject);
 		if (DuplicatedComponent)
 		{
+			// If the duplicated component is a scene component, wipe its attach parent (to prevent log warnings for referencing a private object in an external package)
+			if (auto DuplicatedCompAsSceneComp = Cast<USceneComponent>(DuplicatedComponent))
+			{
+				DuplicatedCompAsSceneComp->AttachParent = nullptr;
+			}
+
 			// Find the closest parent component of the current component within the list of components to copy
 			USceneComponent* ClosestSelectedParent = FindClosestParentInList(Component, ComponentsToCopy);
 			if (ClosestSelectedParent)
@@ -318,13 +324,13 @@ void FComponentEditorUtils::CopyComponents(const TArray<UActorComponent*>& Compo
 	FPlatformMisc::ClipboardCopy(*ExportedText);
 }
 
-bool FComponentEditorUtils::CanPasteComponents(USceneComponent* RootComponent, bool bOverrideCanAttach)
+bool FComponentEditorUtils::CanPasteComponents(USceneComponent* RootComponent, bool bOverrideCanAttach, bool bPasteAsArchetypes)
 {
 	FString ClipboardContent;
 	FPlatformMisc::ClipboardPaste(ClipboardContent);
 
 	// Obtain the component object text factory for the clipboard content and return whether or not we can use it
-	TSharedRef<FComponentObjectTextFactory> Factory = FComponentObjectTextFactory::Get(ClipboardContent);
+	TSharedRef<FComponentObjectTextFactory> Factory = FComponentObjectTextFactory::Get(ClipboardContent, bPasteAsArchetypes);
 	return Factory->NewObjectMap.Num() > 0 && ( bOverrideCanAttach || Factory->CanAttachComponentsTo(RootComponent) );
 }
 
diff --git a/Engine/Source/Editor/UnrealEd/Public/Kismet2/ComponentEditorUtils.h b/Engine/Source/Editor/UnrealEd/Public/Kismet2/ComponentEditorUtils.h
--- a/Engine/Source/Editor/UnrealEd/Public/Kismet2/ComponentEditorUtils.h
+++ b/Engine/Source/Editor/UnrealEd/Public/Kismet2/ComponentEditorUtils.h
@@ -43,7 +43,7 @@ public:
 	 * @param bOverrideCanAttach Optional override declaring that components can be attached and a check is not needed
 	 * @return Whether components can be pasted
 	 */
-	static bool CanPasteComponents(USceneComponent* RootComponent, bool bOverrideCanAttach = false);
+	static bool CanPasteComponents(USceneComponent* RootComponent, bool bOverrideCanAttach = false, bool bPasteAsArchetypes = false);
 
 	/**
 	 * Attempts to paste components from the clipboard as siblings of the target component
