PR #672: Show functions in parent Blueprint without out parameters as overridable function in child (Contributed by [REDACTED_PERSON])

https://github.com/EpicGames/UnrealEngine/pull/672

Various small modifications made because of changes made for 4.7

diff --git a/Engine/Source/Editor/Kismet/Private/SMyBlueprint.cpp b/Engine/Source/Editor/Kismet/Private/SMyBlueprint.cpp
--- a/Engine/Source/Editor/Kismet/Private/SMyBlueprint.cpp
+++ b/Engine/Source/Editor/Kismet/Private/SMyBlueprint.cpp
@@ -1193,7 +1193,7 @@ void SMyBlueprint::CollectAllActions(FGraphActionListBuilderBase& OutAllActions)
 		const UFunction* Function = *FunctionIt;
 		const FName FunctionName = Function->GetFName();
 
-		if ( UEdGraphSchema_K2::CanKismetOverrideFunction(Function) && !ImplementedFunctionCache.Contains(FunctionName) && !UEdGraphSchema_K2::FunctionCanBePlacedAsEvent(Function) && !FBlueprintEditorUtils::FindOverrideForFunction(Blueprint, CastChecked<UClass>(Function->GetOuter()), Function->GetFName()) )
+		if ( UEdGraphSchema_K2::CanKismetOverrideFunction(Function) && !ImplementedFunctionCache.Contains(FunctionName) && !FBlueprintEditorUtils::FindOverrideForFunction(Blueprint, CastChecked<UClass>(Function->GetOuter()), Function->GetFName()) )
 		{
 			FString FunctionTooltip = UK2Node_CallFunction::GetDefaultTooltipForFunction(Function);
 			FText FunctionDesc = FText::FromString(Function->GetMetaData(TEXT("FriendlyName")));
@@ -1970,11 +1970,29 @@ void SMyBlueprint::ImplementFunction(TSharedPtr<FEdGraphSchemaAction_K2Graph> Gr
 void SMyBlueprint::ImplementFunction(FEdGraphSchemaAction_K2Graph* GraphAction)
 {
 	check(GetBlueprintObj()->SkeletonGeneratedClass);
-	UFunction* OverrideFunc = FindField<UFunction>(GetBlueprintObj()->SkeletonGeneratedClass, GraphAction->FuncName);
-	if (OverrideFunc == NULL)
+	const UFunction* OverrideFunc = FindField<UFunction>(GetBlueprintObj()->SkeletonGeneratedClass, GraphAction->FuncName);
+
+	// search up the class hiearchy, we want to find the original declaration of the function to match FBlueprintEventNodeSpawner.
+	// Doing so ensures that we can find the existing node if there is one:
+	const UClass* Iter = GetBlueprintObj()->SkeletonGeneratedClass->GetSuperClass();
+	while (Iter != nullptr)
+	{
+		if (const UFunction* F = Iter->FindFunctionByName(GraphAction->FuncName))
+		{
+			OverrideFunc = F;
+		}
+		else
+		{
+			break;
+		}
+		Iter = Iter->GetSuperClass();
+	}
+
+
+	if (OverrideFunc == nullptr)
 	{
 		// maybe it's from a native interface, check those too
-		for ( UClass* TempClass=GetBlueprintObj()->ParentClass; (NULL != TempClass) && (NULL == OverrideFunc); TempClass=TempClass->GetSuperClass() )
+		for ( UClass* TempClass=GetBlueprintObj()->ParentClass; (nullptr != TempClass) && (nullptr == OverrideFunc); TempClass=TempClass->GetSuperClass() )
 		{
 			for (int32 Idx=0; Idx<TempClass->Interfaces.Num(); ++Idx)
 			{
@@ -1994,10 +2012,38 @@ void SMyBlueprint::ImplementFunction(FEdGraphSchemaAction_K2Graph* GraphAction)
 	check(OverrideFunc);
 	UClass* const OverrideFuncClass = CastChecked<UClass>(OverrideFunc->GetOuter());
 
-	// Implement the function graph
-	UEdGraph* const NewGraph = FBlueprintEditorUtils::CreateNewGraph(GetBlueprintObj(), GraphAction->FuncName, UEdGraph::StaticClass(), UEdGraphSchema_K2::StaticClass());
-	FBlueprintEditorUtils::AddFunctionGraph(GetBlueprintObj(), NewGraph, /*bIsUserCreated=*/ false, OverrideFuncClass);
-	BlueprintEditorPtr.Pin()->OpenDocument(NewGraph, FDocumentTracker::OpenNewDocument);
+	if (UEdGraphSchema_K2::FunctionCanBePlacedAsEvent(OverrideFunc))
+	{
+		// Add to event graph
+		FName EventName = OverrideFunc->GetFName();
+		UK2Node_Event* ExistingNode = FBlueprintEditorUtils::FindOverrideForFunction(GetBlueprintObj(), OverrideFuncClass, EventName);
+
+		if (ExistingNode)
+		{
+			FKismetEditorUtilities::BringKismetToFocusAttentionOnObject(ExistingNode);
+		}
+		else
+		{
+			UEdGraph* EventGraph = FBlueprintEditorUtils::FindEventGraph(GetBlueprintObj());
+			UK2Node_Event* NewEventNodeTemplate = NewObject<UK2Node_Event>();
+			NewEventNodeTemplate->EventReference.SetExternalMember(EventName, OverrideFuncClass);
+			NewEventNodeTemplate->bOverrideFunction = true;
+
+			const FVector2D NewNodePos = EventGraph->GetGoodPlaceForNewNode();
+			UK2Node_Event* NewEventNode = FEdGraphSchemaAction_K2NewNode::SpawnNodeFromTemplate<UK2Node_Event>(EventGraph, NewEventNodeTemplate, NewNodePos);
+			if (NewEventNode)
+			{
+				FKismetEditorUtilities::BringKismetToFocusAttentionOnObject(NewEventNode);
+			}
+		}
+	}
+	else
+	{
+		// Implement the function graph
+		UEdGraph* const NewGraph = FBlueprintEditorUtils::CreateNewGraph(GetBlueprintObj(), GraphAction->FuncName, UEdGraph::StaticClass(), UEdGraphSchema_K2::StaticClass());
+		FBlueprintEditorUtils::AddFunctionGraph(GetBlueprintObj(), NewGraph, /*bIsUserCreated=*/ false, OverrideFuncClass);
+		BlueprintEditorPtr.Pin()->OpenDocument(NewGraph, FDocumentTracker::OpenNewDocument);
+	}
 }
 
 void SMyBlueprint::OnFindEntry()
