UE-10319 Crash attempting to edit a Struct Transform's rotation value as a slider

diff --git a/Engine/Source/Editor/Kismet/Private/UserDefinedStructureEditor.cpp b/Engine/Source/Editor/Kismet/Private/UserDefinedStructureEditor.cpp
--- a/Engine/Source/Editor/Kismet/Private/UserDefinedStructureEditor.cpp
+++ b/Engine/Source/Editor/Kismet/Private/UserDefinedStructureEditor.cpp
@@ -44,30 +44,32 @@ public:
 
 	void OnFinishedChangingProperties(const FPropertyChangedEvent& PropertyChangedEvent)
 	{
-		check(PropertyChangedEvent.MemberProperty
-			&& PropertyChangedEvent.MemberProperty->GetOwnerStruct()
-			&& (PropertyChangedEvent.MemberProperty->GetOwnerStruct() == GetUserDefinedStruct())
-			&& PropertyChangedEvent.MemberProperty->GetOwnerStruct()->IsA<UUserDefinedStruct>());
+		check(PropertyChangedEvent.MemberProperty && PropertyChangedEvent.MemberProperty->GetOwnerStruct());
 
-		const UProperty* DirectProperty = PropertyChangedEvent.MemberProperty;
-		while (!Cast<const UUserDefinedStruct>(DirectProperty->GetOuter()))
+		const bool bPropertyIsActual = PropertyChangedEvent.MemberProperty->GetOwnerStruct() == GetUserDefinedStruct();
+		const UProperty* DirectProperty = bPropertyIsActual ? PropertyChangedEvent.MemberProperty : nullptr;
+		while (DirectProperty && !Cast<const UUserDefinedStruct>(DirectProperty->GetOuter()))
 		{
-			DirectProperty = CastChecked<const UProperty>(DirectProperty->GetOuter());
+			DirectProperty = Cast<const UProperty>(DirectProperty->GetOuter());
 		}
+		ensure(nullptr != DirectProperty);
 
-		FString DefaultValueString;
-		bool bDefaultValueSet = false;
+		if (DirectProperty)
 		{
-			if (StructData.IsValid() && StructData->IsValid())
+			FString DefaultValueString;
+			bool bDefaultValueSet = false;
 			{
-				bDefaultValueSet = FBlueprintEditorUtils::PropertyValueToString(DirectProperty, StructData->GetStructMemory(), DefaultValueString);
+				if (StructData.IsValid() && StructData->IsValid())
+				{
+					bDefaultValueSet = FBlueprintEditorUtils::PropertyValueToString(DirectProperty, StructData->GetStructMemory(), DefaultValueString);
+				}
 			}
-		}
 
-		const FGuid VarGuid = FStructureEditorUtils::GetGuidForProperty(DirectProperty);
-		if (bDefaultValueSet && VarGuid.IsValid())
-		{
-			FStructureEditorUtils::ChangeVariableDefaultValue(GetUserDefinedStruct(), VarGuid, DefaultValueString);
+			const FGuid VarGuid = FStructureEditorUtils::GetGuidForProperty(DirectProperty);
+			if (bDefaultValueSet && VarGuid.IsValid())
+			{
+				FStructureEditorUtils::ChangeVariableDefaultValue(GetUserDefinedStruct(), VarGuid, DefaultValueString);
+			}
 		}
 	}
 
