Simple fix to adjust the placement of parent nodes
#UE-10523 - Parent nodes spawned through the context menu overlap existing nodes
#Branch UE4
#Proj kismet

diff --git a/Engine/Source/Editor/Kismet/Private/BlueprintEditor.cpp b/Engine/Source/Editor/Kismet/Private/BlueprintEditor.cpp
--- a/Engine/Source/Editor/Kismet/Private/BlueprintEditor.cpp
+++ b/Engine/Source/Editor/Kismet/Private/BlueprintEditor.cpp
@@ -3932,8 +3932,13 @@ void FBlueprintEditor::OnAddParentNode()
 				ParentFunctionNode->SetFromFunction(ValidParent);
 				ParentFunctionNode->AllocateDefaultPins();
 
-				ParentFunctionNode->NodePosX = FunctionFromNode.Node->NodePosX + FunctionFromNode.Node->NodeWidth + 20;
-				ParentFunctionNode->NodePosY = FunctionFromNode.Node->NodePosY + 15;
+				int32 NodeSizeY = 15;
+				if( UK2Node* Node = Cast<UK2Node>(SelectedObj))
+				{
+					NodeSizeY += Node->NodeWidget.IsValid() ? static_cast<int32>(Node->NodeWidget.Pin()->GetDesiredSize().Y) : 0;
+				}
+				ParentFunctionNode->NodePosX = FunctionFromNode.Node->NodePosX;
+				ParentFunctionNode->NodePosY = FunctionFromNode.Node->NodePosY + NodeSizeY;
 				FunctionNodeCreator.Finalize();
 			}
 		}
