[UE-9763] Refresh the selected actor's component tree when a custom component type changes its parent class.

change summary:
- modified SSCSEditor::Tick() to include additional checks for component type relative to its position in the overall hierarchy

diff --git a/Engine/Source/Editor/Kismet/Private/SSCSEditor.cpp b/Engine/Source/Editor/Kismet/Private/SSCSEditor.cpp
--- a/Engine/Source/Editor/Kismet/Private/SSCSEditor.cpp
+++ b/Engine/Source/Editor/Kismet/Private/SSCSEditor.cpp
@@ -3302,15 +3302,29 @@ void SSCSEditor::Tick( const FGeometry& AllottedGeometry, const double InCurrent
 		TFunction<bool (const TArray<FSCSEditorTreeNodePtrType>&, int32&)> AreAnyNodesInvalidLambda = [&](const TArray<FSCSEditorTreeNodePtrType>& InNodes, int32& OutNumValidNodes) -> bool
 		{
 			bool bFoundInvalidNode = false;
-			for(auto NodeIt = InNodes.CreateConstIterator(); NodeIt && !bFoundInvalidNode; ++NodeIt)
+			for(auto NodeIt = InNodes.CreateConstIterator(); NodeIt && !bFoundInvalidNode && !AreAnyNodesInvalidLambda((*NodeIt)->GetChildren(), OutNumValidNodes); ++NodeIt)
 			{
 				bool bIsComponent = (*NodeIt)->GetNodeType() == FSCSEditorTreeNode::ComponentNode;
-
-				const UActorComponent* InstancedComponent = (*NodeIt)->GetComponentTemplate();
-				bFoundInvalidNode = (bIsComponent && (!InstancedComponent || InstancedComponent->IsPendingKill())) || AreAnyNodesInvalidLambda((*NodeIt)->GetChildren(), OutNumValidNodes);
 				if (bIsComponent)
 				{
-					++OutNumValidNodes;
+					// A component node is considered to be invalid if one of the following conditions holds true:
+					//	(a) The component instance is invalid (i.e. NULL).
+					//	(b) The component instance is pending garbage collection.
+					//	(c) The component instance is a USceneComponent derivative and either is not the scene "root" node or otherwise is not contained within the scene component node hierarchy.
+					//	(d) The component instance is not a USceneComponent derivative and either matches the scene "root" node or is otherwise contained within the scene component node hierarchy.
+					//		=> (c) and (d) could happen if the user changes the parent class of a custom component's Blueprint class, for example.
+					//
+					FSCSEditorTreeNodePtrType ParentNodePtr = (*NodeIt)->GetParent();
+					const UActorComponent* InstancedComponent = (*NodeIt)->GetComponentTemplate();
+					bFoundInvalidNode = !InstancedComponent
+						|| InstancedComponent->IsPendingKill()
+						|| (InstancedComponent->IsA<USceneComponent>() && (*NodeIt != SceneRootNodePtr && (!ParentNodePtr.IsValid() || ParentNodePtr->GetNodeType() != FSCSEditorTreeNode::ComponentNode)))
+						|| (!InstancedComponent->IsA<USceneComponent>() && (*NodeIt == SceneRootNodePtr || (ParentNodePtr.IsValid() && ParentNodePtr->GetNodeType() == FSCSEditorTreeNode::ComponentNode)));
+					
+					if(!bFoundInvalidNode)
+					{
+						++OutNumValidNodes;
+					}
 				}
 			}
 
