Modified MyBlueprint code to ensure that the correct class signature (Authoritive) is used to create overridden functions.
#UE-9408 - Overridden Interface Functions in Child BPs lose array output pins after closing/reopening the project and become invalid
#Branch UE4
#Proj Kismet

#info previously overrides were being added using skeleton functions which suffered from problems with the signature class not being serialised out. if any functions have been created with the broken code there is no way to patch this up other than to recreate the override because the signature data has been lost.

diff --git a/Engine/Source/Editor/Kismet/Private/SMyBlueprint.cpp b/Engine/Source/Editor/Kismet/Private/SMyBlueprint.cpp
--- a/Engine/Source/Editor/Kismet/Private/SMyBlueprint.cpp
+++ b/Engine/Source/Editor/Kismet/Private/SMyBlueprint.cpp
@@ -2010,7 +2010,7 @@ void SMyBlueprint::ImplementFunction(FEdGraphSchemaAction_K2Graph* GraphAction)
 		}
 	}
 	check(OverrideFunc);
-	UClass* const OverrideFuncClass = CastChecked<UClass>(OverrideFunc->GetOuter());
+	UClass* const OverrideFuncClass = CastChecked<UClass>(OverrideFunc->GetOuter())->GetAuthoritativeClass();
 
 	if (UEdGraphSchema_K2::FunctionCanBePlacedAsEvent(OverrideFunc))
 	{
