Can drag and drop multiple components variables into a Blueprint to have them all be placed as nodes.

diff --git a/Engine/Source/Editor/Kismet/Private/SSCSEditor.cpp b/Engine/Source/Editor/Kismet/Private/SSCSEditor.cpp
--- a/Engine/Source/Editor/Kismet/Private/SSCSEditor.cpp
+++ b/Engine/Source/Editor/Kismet/Private/SSCSEditor.cpp
@@ -230,6 +230,12 @@ public:
 		DropAction_AttachToOrMakeNewRoot
 	};
 
+	// FGraphEditorDragDropAction interface
+	virtual void HoverTargetChanged() override;
+	virtual FReply DroppedOnNode(FVector2D ScreenPosition, FVector2D GraphPosition) override;
+	virtual FReply DroppedOnPanel(const TSharedRef< class SWidget >& Panel, FVector2D ScreenPosition, FVector2D GraphPosition, UEdGraph& Graph) override;
+	// End of FGraphEditorDragDropAction
+
 	/** Node(s) that we started the drag from */
 	TArray<FSCSEditorTreeNodePtrType> SourceNodes;
 
@@ -249,6 +255,75 @@ TSharedRef<FSCSRowDragDropOp> FSCSRowDragDropOp::New(FName InVariableName, UStru
 	return Operation.ToSharedRef();
 }
 
+void FSCSRowDragDropOp::HoverTargetChanged()
+{
+	bool bHoverHandled = false;
+
+	if(SourceNodes.Num() > 1)
+	{
+		// Display an error message if attempting to drop multiple source items onto a node
+		UEdGraphNode* VarNodeUnderCursor = Cast<UK2Node_Variable>(GetHoveredNode());
+		if (VarNodeUnderCursor != NULL)
+		{
+			// Icon/text to draw on tooltip
+			FSlateColor IconColor = FLinearColor::White;
+			const FSlateBrush* StatusSymbol = FEditorStyle::GetBrush(TEXT("Graph.ConnectorFeedback.Error"));
+			FText Message = LOCTEXT("InvalidMultiDropTarget", "Cannot replace node with multiple nodes");
+			SetSimpleFeedbackMessage(StatusSymbol, IconColor, Message);
+
+			bHoverHandled = true;
+		}
+	}
+
+	if(!bHoverHandled)
+	{
+		FKismetVariableDragDropAction::HoverTargetChanged();
+	}
+}
+
+FReply FSCSRowDragDropOp::DroppedOnNode(FVector2D ScreenPosition, FVector2D GraphPosition)
+{
+	// Only allow dropping on another node if there is only a single source item
+	if(SourceNodes.Num() == 1)
+	{
+		FKismetVariableDragDropAction::DroppedOnNode(ScreenPosition, GraphPosition);
+	}
+	return FReply::Handled();
+}
+
+FReply FSCSRowDragDropOp::DroppedOnPanel(const TSharedRef< class SWidget >& Panel, FVector2D ScreenPosition, FVector2D GraphPosition, UEdGraph& Graph)
+{
+	const FScopedTransaction Transaction(LOCTEXT("SCSEditorAddMultipleNodes", "Add Component Nodes"));
+
+	TArray<UK2Node_VariableGet*> OriginalVariableNodes;
+	Graph.GetNodesOfClass<UK2Node_VariableGet>(OriginalVariableNodes);
+
+	// Add source items to the graph in turn
+	for (FSCSEditorTreeNodePtrType& SourceNode : SourceNodes)
+	{
+		VariableName = SourceNode->GetVariableName();
+		FKismetVariableDragDropAction::DroppedOnPanel(Panel, ScreenPosition, GraphPosition, Graph);
+
+		GraphPosition.Y += 50;
+	}
+
+	TArray<UK2Node_VariableGet*> ResultVariableNodes;
+	Graph.GetNodesOfClass<UK2Node_VariableGet>(ResultVariableNodes);
+
+	if (ResultVariableNodes.Num() - OriginalVariableNodes.Num() > 1)
+	{
+		TSet<const UEdGraphNode*> NodeSelection;
+
+		// Because there is more than one new node, lets grab all the nodes at the bottom of the list and add them to a set for selection
+		for (int32 NodeIdx = ResultVariableNodes.Num() - 1; NodeIdx >= OriginalVariableNodes.Num(); --NodeIdx)
+		{
+			NodeSelection.Add(ResultVariableNodes[NodeIdx]);
+		}
+		Graph.SelectNodeSet(NodeSelection);
+	}
+	return FReply::Handled();
+}
+
 //////////////////////////////////////////////////////////////////////////
 // FSCSEditorTreeNode
 
