Undo/Redo actions that do not affect a UBlueprint's package will no longer dirty the Blueprint's package if the Blueprint is open.

FBlueprintEditor::PostUndo and ::PostRedo will not examine the transaction buffer's objects and see if any relate to the Blueprint, if none do then nothing further will be handled by the BP Editor.

diff --git a/Engine/Source/Editor/Kismet/Private/BlueprintEditor.cpp b/Engine/Source/Editor/Kismet/Private/BlueprintEditor.cpp
--- a/Engine/Source/Editor/Kismet/Private/BlueprintEditor.cpp
+++ b/Engine/Source/Editor/Kismet/Private/BlueprintEditor.cpp
@@ -2709,18 +2709,38 @@ void FBlueprintEditor::PostUndo(bool bSuccess)
 	// Clear selection, to avoid holding refs to nodes that go away
 	if (bSuccess && GetBlueprintObj())
 	{
-		SetUISelectionState(NAME_None);
+		bool bAffectsBlueprint = false;
+		const UPackage* BlueprintOutermost = GetBlueprintObj()->GetOutermost();
 
-		// Will cause a call to RefreshEditors()
-		FBlueprintEditorUtils::MarkBlueprintAsStructurallyModified(GetBlueprintObj());
-
-		TSharedPtr<class SSCSEditorViewport> ViewportPtr = GetSCSViewport();
-		if (ViewportPtr.IsValid())
+		// Look at the transaction this function is responding to, see if any object in it has an outermost of the Blueprint
+		const FTransaction* Transaction = GEditor->Trans->GetTransaction(GEditor->Trans->GetQueueLength() - GEditor->Trans->GetUndoCount());
+		TArray<UObject*> TransactionObjects;
+		Transaction->GetTransactionObjects(TransactionObjects);
+		for (UObject* Object : TransactionObjects)
 		{
-			ViewportPtr->Invalidate();
+			if (Object->GetOutermost() == BlueprintOutermost)
+			{
+				bAffectsBlueprint = true;
+				break;
+			}
 		}
 
-		FSlateApplication::Get().DismissAllMenus();
+		// Transaction affects the Blueprint this editor handles, so react as necessary
+		if(bAffectsBlueprint)
+		{
+			SetUISelectionState(NAME_None);
+
+			// Will cause a call to RefreshEditors()
+			FBlueprintEditorUtils::MarkBlueprintAsStructurallyModified(GetBlueprintObj());
+
+			TSharedPtr<class SSCSEditorViewport> ViewportPtr = GetSCSViewport();
+			if (ViewportPtr.IsValid())
+			{
+				ViewportPtr->Invalidate();
+			}
+
+			FSlateApplication::Get().DismissAllMenus();
+		}
 	}
 }
 
@@ -2729,10 +2749,30 @@ void FBlueprintEditor::PostRedo(bool bSuccess)
 	UBlueprint* BlueprintObj = GetBlueprintObj();
 	if (BlueprintObj && bSuccess)
 	{
-		// Will cause a call to RefreshEditors()
-		FBlueprintEditorUtils::MarkBlueprintAsStructurallyModified( BlueprintObj );
+		bool bAffectsBlueprint = false;
+		const UPackage* BlueprintOutermost = GetBlueprintObj()->GetOutermost();
 
-		FSlateApplication::Get().DismissAllMenus();
+		// Look at the transaction this function is responding to, see if any object in it has an outermost of the Blueprint
+		const FTransaction* Transaction = GEditor->Trans->GetTransaction(GEditor->Trans->GetQueueLength() - GEditor->Trans->GetUndoCount() - 1);
+		TArray<UObject*> TransactionObjects;
+		Transaction->GetTransactionObjects(TransactionObjects);
+		for (UObject* Object : TransactionObjects)
+		{
+			if (Object->GetOutermost() == BlueprintOutermost)
+			{
+				bAffectsBlueprint = true;
+				break;
+			}
+		}
+
+		// Transaction affects the Blueprint this editor handles, so react as necessary
+		if(bAffectsBlueprint)
+		{
+			// Will cause a call to RefreshEditors()
+			FBlueprintEditorUtils::MarkBlueprintAsStructurallyModified( BlueprintObj );
+
+			FSlateApplication::Get().DismissAllMenus();
+		}
 	}
 }
 
@@ -2748,17 +2788,7 @@ bool FBlueprintEditor::CanUndoGraphAction() const
 
 void FBlueprintEditor::RedoGraphAction()
 {
-	// Clear selection, to avoid holding refs to nodes that go away
-	if (GetBlueprintObj())
-	{
-		SetUISelectionState(NAME_None);
-
-		GEditor->RedoTransaction();
-
-		RefreshEditors();
-
-		FBlueprintEditorUtils::MarkBlueprintAsStructurallyModified(GetBlueprintObj());
-	}
+	GEditor->RedoTransaction();
 }
 
 bool FBlueprintEditor::CanRedoGraphAction() const
diff --git a/Engine/Source/Editor/UnrealEd/Classes/Editor/Transactor.h b/Engine/Source/Editor/UnrealEd/Classes/Editor/Transactor.h
--- a/Engine/Source/Editor/UnrealEd/Classes/Editor/Transactor.h
+++ b/Engine/Source/Editor/UnrealEd/Classes/Editor/Transactor.h
@@ -367,7 +367,7 @@ public:
 	 * Get all the objects that are part of this transaction.
 	 * @param	Objects		[out] Receives the object list.  Previous contents are cleared.
 	 */
-	void GetTransactionObjects(TArray<UObject*>& Objects);
+	void GetTransactionObjects(TArray<UObject*>& Objects) const;
 	void RemoveRecords( int32 Count = 1 );
 	int32 GetRecordCount() const;
 
diff --git a/Engine/Source/Editor/UnrealEd/Private/EditorTransaction.cpp b/Engine/Source/Editor/UnrealEd/Private/EditorTransaction.cpp
--- a/Engine/Source/Editor/UnrealEd/Private/EditorTransaction.cpp
+++ b/Engine/Source/Editor/UnrealEd/Private/EditorTransaction.cpp
@@ -319,7 +319,7 @@ SIZE_T FTransaction::DataSize() const
  * Get all the objects that are part of this transaction.
  * @param	Objects		[out] Receives the object list.  Previous contents are cleared.
  */
-void FTransaction::GetTransactionObjects(TArray<UObject*>& Objects)
+void FTransaction::GetTransactionObjects(TArray<UObject*>& Objects) const
 {
 	Objects.Empty(); // Just in case.
 
