Ensured that the separator line between scene and non-scene components is added when adding a new component

UE-9454 - IWCE: SCS Tree requires manual refresh to draw separator line between scene components and behavior components

diff --git a/Engine/Source/Editor/Kismet/Private/SSCSEditor.cpp b/Engine/Source/Editor/Kismet/Private/SSCSEditor.cpp
--- a/Engine/Source/Editor/Kismet/Private/SSCSEditor.cpp
+++ b/Engine/Source/Editor/Kismet/Private/SSCSEditor.cpp
@@ -3090,6 +3090,7 @@ void SSCSEditor::Construct( const FArguments& InArgs )
 	OnItemDoubleClicked = InArgs._OnItemDoubleClicked;
 	OnHighlightPropertyInDetailsView = InArgs._OnHighlightPropertyInDetailsView;
 	bUpdatingSelection = false;
+	bHasAddedSceneAndBehaviorComponentSeparator = false;
 
 	CommandList = MakeShareable( new FUICommandList );
 	CommandList->MapAction( FGenericCommands::Get().Cut,
@@ -4033,6 +4034,8 @@ void SSCSEditor::UpdateTree(bool bRegenerateTreeNodes)
 		RootNodes.Empty();
 		RootComponentNodes.Empty();
 
+		bHasAddedSceneAndBehaviorComponentSeparator = false;
+
 		// Reset the scene root node
 		SceneRootNodePtr.Reset();
 
@@ -4085,15 +4088,14 @@ void SSCSEditor::UpdateTree(bool bRegenerateTreeNodes)
 				}
 
 				// Add native ActorComponent nodes that aren't SceneComponents
-				bool bSeparatorAdded = false;
 				for(auto CompIter = Components.CreateIterator(); CompIter; ++CompIter)
 				{
 					UActorComponent* ActorComp = *CompIter;
 					if (!ActorComp->IsA<USceneComponent>())
 					{
-						if (!bSeparatorAdded)
+						if (!bHasAddedSceneAndBehaviorComponentSeparator)
 						{
-							bSeparatorAdded = true;
+							bHasAddedSceneAndBehaviorComponentSeparator = true;
 							RootNodes.Add(MakeShareable(new FSCSEditorTreeNodeSeparator()));
 						}
 						AddRootComponentTreeNode(ActorComp);
@@ -4160,15 +4162,14 @@ void SSCSEditor::UpdateTree(bool bRegenerateTreeNodes)
 				}
 
 				// Add all non-scene component instances to the root set first
-				bool bSeparatorAdded = false;
 				for(auto CompIter = Components.CreateIterator(); CompIter; ++CompIter)
 				{
 					UActorComponent* ActorComp = *CompIter;
 					if (!ActorComp->IsA<USceneComponent>() && !ActorComp->IsEditorOnly())
 					{
-						if (!bSeparatorAdded)
+						if (!bHasAddedSceneAndBehaviorComponentSeparator)
 						{
-							bSeparatorAdded = true;
+							bHasAddedSceneAndBehaviorComponentSeparator = true;
 							RootNodes.Add(MakeShareable(new FSCSEditorTreeNode(FSCSEditorTreeNode::SeparatorNode)));
 						}
 						AddRootComponentTreeNode(ActorComp);
@@ -4557,6 +4558,13 @@ UActorComponent* SSCSEditor::AddNewNodeForInstancedComponent(UActorComponent* Ne
 	}
 	else
 	{
+		// Make sure we've added the separator between scene and behavior components
+		if (!bHasAddedSceneAndBehaviorComponentSeparator)
+		{
+			bHasAddedSceneAndBehaviorComponentSeparator = true;
+			RootNodes.Add(MakeShareable(new FSCSEditorTreeNode(FSCSEditorTreeNode::SeparatorNode)));
+		}
+
 		NewNodePtr = AddRootComponentTreeNode(NewInstanceComponent);
 	}
 
@@ -4569,7 +4577,6 @@ UActorComponent* SSCSEditor::AddNewNodeForInstancedComponent(UActorComponent* Ne
 
 	UpdateTree(false);
 
-
 	return NewInstanceComponent;
 }
 
diff --git a/Engine/Source/Editor/Kismet/Public/SSCSEditor.h b/Engine/Source/Editor/Kismet/Public/SSCSEditor.h
--- a/Engine/Source/Editor/Kismet/Public/SSCSEditor.h
+++ b/Engine/Source/Editor/Kismet/Public/SSCSEditor.h
@@ -1067,4 +1067,6 @@ private:
 	/** Gate to prevent changing the selection while selection change is being broadcast. */
 	bool bUpdatingSelection;
 
+	/** true if we've added the separator between the scene and behavior components to the root nodes */
+	bool bHasAddedSceneAndBehaviorComponentSeparator;
 };
