Drag and drop in the SCS editor now expands the drop target node so you can see the new child

UE-9583 - Iwce: don't collapse component tree when rerooting

diff --git a/Engine/Source/Editor/Kismet/Private/SSCSEditor.cpp b/Engine/Source/Editor/Kismet/Private/SSCSEditor.cpp
--- a/Engine/Source/Editor/Kismet/Private/SSCSEditor.cpp
+++ b/Engine/Source/Editor/Kismet/Private/SSCSEditor.cpp
@@ -2575,6 +2575,9 @@ void SSCS_RowWidget::OnMakeNewRootDropAction(FSCSEditorTreeNodePtrType DroppedNo
 		if(OldSceneRootNodePtr.IsValid())
 		{
 			SceneRootNodePtr->AddChild(OldSceneRootNodePtr);
+
+			// Expand the new scene root as we've just added a child to it
+			SCSEditorPtr->SetNodeExpansionState(SceneRootNodePtr, true);
 		}
 	}
 	else    // EComponentEditorMode::ActorInstance
@@ -2596,6 +2599,9 @@ void SSCS_RowWidget::OnMakeNewRootDropAction(FSCSEditorTreeNodePtrType DroppedNo
 		// Set old root as child of new root
 		check(OldSceneRootNodePtr.IsValid());
 		SceneRootNodePtr->AddChild(OldSceneRootNodePtr);
+
+		// Expand the new scene root as we've just added a child to it
+		SCSEditorPtr->SetNodeExpansionState(SceneRootNodePtr, true);
 	}
 
 	PostDragDropAction(true);
@@ -3950,6 +3956,14 @@ void SSCSEditor::SelectNode(FSCSEditorTreeNodePtrType InNodeToSelect, bool IsCnt
 	}
 }
 
+void SSCSEditor::SetNodeExpansionState(FSCSEditorTreeNodePtrType InNodeToChange, const bool bIsExpanded)
+{
+	if(SCSTreeWidget.IsValid() && InNodeToChange.IsValid())
+	{
+		SCSTreeWidget->SetItemExpansion(InNodeToChange, bIsExpanded);
+	}
+}
+
 static FSCSEditorTreeNode* FindRecursive( FSCSEditorTreeNode* Node, FName Name )
 {
 	if (Node->GetVariableName() == Name)
diff --git a/Engine/Source/Editor/Kismet/Public/SSCSEditor.h b/Engine/Source/Editor/Kismet/Public/SSCSEditor.h
--- a/Engine/Source/Editor/Kismet/Public/SSCSEditor.h
+++ b/Engine/Source/Editor/Kismet/Public/SSCSEditor.h
@@ -840,6 +840,14 @@ public:
 	/** Select the given tree node */
 	void SelectNode(FSCSEditorTreeNodePtrType InNodeToSelect, bool IsCntrlDown);
 
+	/**
+	 * Set the expansion state of a node
+	 *
+	 * @param InNodeToChange	The node to be expanded/collapsed
+	 * @param bIsExpanded		True to expand the node, false to collapse it
+	 */
+	void SetNodeExpansionState(FSCSEditorTreeNodePtrType InNodeToChange, const bool bIsExpanded);
+
 	/**
 	 * Highlight a tree node and, optionally, a property with in it
 	 *
