Overridable functions will not appear multiple times in the MyBlueprint context menu.

diff --git a/Engine/Source/Editor/Kismet/Private/SMyBlueprint.cpp b/Engine/Source/Editor/Kismet/Private/SMyBlueprint.cpp
--- a/Engine/Source/Editor/Kismet/Private/SMyBlueprint.cpp
+++ b/Engine/Source/Editor/Kismet/Private/SMyBlueprint.cpp
@@ -1187,13 +1187,16 @@ void SMyBlueprint::CollectAllActions(FGraphActionListBuilderBase& OutAllActions)
 
 	OverridableFunctionActions.Reset();
 
+	// Fill with functions names we've already collected for rename, to ensure we do not add the same function multiple times.
+	TArray<FName> OverridableFunctionNames;
+
 	// Cache potentially overridable functions
 	for ( TFieldIterator<UFunction> FunctionIt(Blueprint->ParentClass, EFieldIteratorFlags::IncludeSuper); FunctionIt; ++FunctionIt )
 	{
 		const UFunction* Function = *FunctionIt;
 		const FName FunctionName = Function->GetFName();
 
-		if ( UEdGraphSchema_K2::CanKismetOverrideFunction(Function) && !ImplementedFunctionCache.Contains(FunctionName) && !FBlueprintEditorUtils::FindOverrideForFunction(Blueprint, CastChecked<UClass>(Function->GetOuter()), Function->GetFName()) )
+		if ( UEdGraphSchema_K2::CanKismetOverrideFunction(Function) && !OverridableFunctionNames.Contains(FunctionName) && !ImplementedFunctionCache.Contains(FunctionName) && !FBlueprintEditorUtils::FindOverrideForFunction(Blueprint, CastChecked<UClass>(Function->GetOuter()), Function->GetFName()) )
 		{
 			FString FunctionTooltip = UK2Node_CallFunction::GetDefaultTooltipForFunction(Function);
 			FText FunctionDesc = FText::FromString(Function->GetMetaData(TEXT("FriendlyName")));
@@ -1209,8 +1212,7 @@ void SMyBlueprint::CollectAllActions(FGraphActionListBuilderBase& OutAllActions)
 			NewFuncAction->SectionID = NodeSectionID::FUNCTION_OVERRIDABLE;
 
 			OverridableFunctionActions.Add(NewFuncAction);
-
-			//OutAllActions.AddAction(NewFuncAction);
+			OverridableFunctionNames.Add(FunctionName);
 		}
 	}
 
