Submitting missing file from CL# 2460257 for UE-7582

diff --git a/Engine/Source/Editor/Kismet/Private/SMyBlueprint.cpp b/Engine/Source/Editor/Kismet/Private/SMyBlueprint.cpp
--- a/Engine/Source/Editor/Kismet/Private/SMyBlueprint.cpp
+++ b/Engine/Source/Editor/Kismet/Private/SMyBlueprint.cpp
@@ -910,12 +910,13 @@ void SMyBlueprint::GetLocalVariables(FGraphActionListBuilderBase& OutAllActions)
 
 		// Search in all FunctionEntry nodes for their local variables
 		FString ActionCategory;
+		const UEdGraphSchema_K2* K2Schema = GetDefault<UEdGraphSchema_K2>();
 		for (UK2Node_FunctionEntry* const FunctionEntry : FunctionEntryNodes)
 		{
 			for( const FBPVariableDescription& Variable : FunctionEntry->LocalVariables )
 			{
 				FString Category = Variable.Category.ToString();
-				if (Variable.Category == GetDefault<UEdGraphSchema_K2>()->VR_DefaultCategory)
+				if (Variable.Category == K2Schema->VR_DefaultCategory)
 				{
 					Category = FString();
 				}
@@ -924,7 +925,7 @@ void SMyBlueprint::GetLocalVariables(FGraphActionListBuilderBase& OutAllActions)
 				if (Func)
 				{
 					TSharedPtr<FEdGraphSchemaAction_K2LocalVar> NewVarAction = MakeShareable(new FEdGraphSchemaAction_K2LocalVar(Category, FText::FromName(Variable.VarName), TEXT(""), 0));
-					NewVarAction->SetVariableInfo(Variable.VarName, Func);
+					NewVarAction->SetVariableInfo(Variable.VarName, Func, Variable.VarType.PinCategory == K2Schema->PC_Boolean);
 					NewVarAction->SectionID = NodeSectionID::LOCAL_VARIABLE;
 					OutAllActions.AddAction(NewVarAction);
 				}
@@ -1091,7 +1092,9 @@ void SMyBlueprint::CollectAllActions(FGraphActionListBuilderBase& OutAllActions)
 			}
 
 			TSharedPtr<FEdGraphSchemaAction_K2Var> NewVarAction = MakeShareable(new FEdGraphSchemaAction_K2Var(PropertyCategory, PropertyDesc, PropertyTooltip.ToString(), 0));
-			NewVarAction->SetVariableInfo(PropertyName, Blueprint->SkeletonGeneratedClass);
+			const UArrayProperty* ArrayProperty = Cast<const UArrayProperty>(Property);
+			const UProperty* TestProperty = ArrayProperty ? ArrayProperty->Inner : Property;
+			NewVarAction->SetVariableInfo(PropertyName, Blueprint->SkeletonGeneratedClass, Cast<UBoolProperty>(TestProperty) != nullptr);
 			NewVarAction->SectionID = NodeSectionID::VARIABLE;
 			SortList.AddAction( UserCategoryName, NewVarAction );
 		}
@@ -2057,11 +2060,11 @@ void SMyBlueprint::OnFindEntry()
 	}
 	else if (FEdGraphSchemaAction_K2Var* VarAction = SelectionAsVar())
 	{
-		SearchTerm = VarAction->GetVariableName().ToString();
+		SearchTerm = VarAction->GetFriendlyVariableName();
 	}
 	else if (FEdGraphSchemaAction_K2LocalVar* LocalVarAction = SelectionAsLocalVar())
 	{
-		SearchTerm = LocalVarAction->GetVariableName().ToString();
+		SearchTerm = LocalVarAction->GetFriendlyVariableName();
 	}
 	else if (FEdGraphSchemaAction_K2Delegate* DelegateAction = SelectionAsDelegate())
 	{
