Removed "Edit Tunnel" context menu option from collapsed graph nodes as it no longer did anything.

diff --git a/Engine/Source/Editor/BlueprintGraph/Private/EdGraphSchema_K2.cpp b/Engine/Source/Editor/BlueprintGraph/Private/EdGraphSchema_K2.cpp
--- a/Engine/Source/Editor/BlueprintGraph/Private/EdGraphSchema_K2.cpp
+++ b/Engine/Source/Editor/BlueprintGraph/Private/EdGraphSchema_K2.cpp
@@ -1294,12 +1294,6 @@ void UEdGraphSchema_K2::GetContextMenuActions(const UEdGraph* CurrentGraph, cons
 					MenuBuilder->AddMenuEntry( FGraphEditorCommands::Get().ReconstructNodes );
 					MenuBuilder->AddMenuEntry( FGraphEditorCommands::Get().BreakNodeLinks );
 
-					// tunnel nodes have option to open function editor
-					if (InGraphNode->IsA(UK2Node_Tunnel::StaticClass()))
-					{
-						MenuBuilder->AddMenuEntry(FGraphEditorCommands::Get().EditTunnel );
-					}
-
 					// Conditionally add the action to add an execution pin, if this is an execution node
 					if( InGraphNode->IsA(UK2Node_ExecutionSequence::StaticClass()) || InGraphNode->IsA(UK2Node_Switch::StaticClass()) )
 					{
diff --git a/Engine/Source/Editor/GraphEditor/Private/GraphEditorActions.cpp b/Engine/Source/Editor/GraphEditor/Private/GraphEditorActions.cpp
--- a/Engine/Source/Editor/GraphEditor/Private/GraphEditorActions.cpp
+++ b/Engine/Source/Editor/GraphEditor/Private/GraphEditorActions.cpp
@@ -69,8 +69,6 @@ void FGraphEditorCommandsImpl::RegisterCommands()
 
 	UI_COMMAND( OpenRelatedAsset, "Open Asset", "Opens the asset related to this node", EUserInterfaceActionType::Button, FInputGesture() )
 
-	UI_COMMAND( EditTunnel, "Edit Tunnel", "Edit input and output pins for tunnel", EUserInterfaceActionType::Button, FInputGesture() )
-
 	UI_COMMAND( CreateComment, "Create Comment", "Create a comment box", EUserInterfaceActionType::Button, FInputGesture(EKeys::C))
 
 	UI_COMMAND( FindInstancesOfCustomEvent, "Find Instances of Event", "Find the instances of this custom event", EUserInterfaceActionType::Button, FInputGesture() )
diff --git a/Engine/Source/Editor/GraphEditor/Public/GraphEditorActions.h b/Engine/Source/Editor/GraphEditor/Public/GraphEditorActions.h
--- a/Engine/Source/Editor/GraphEditor/Public/GraphEditorActions.h
+++ b/Engine/Source/Editor/GraphEditor/Public/GraphEditorActions.h
@@ -99,9 +99,6 @@ public:
 	// option for opening the asset related to the graph node
 	TSharedPtr< FUICommandInfo > OpenRelatedAsset;
 
-	//open function editor
-	TSharedPtr< FUICommandInfo > EditTunnel;
-
 	//create a comment node
 	TSharedPtr< FUICommandInfo > CreateComment;
 
diff --git a/Engine/Source/Editor/Kismet/Private/BlueprintEditor.cpp b/Engine/Source/Editor/Kismet/Private/BlueprintEditor.cpp
--- a/Engine/Source/Editor/Kismet/Private/BlueprintEditor.cpp
+++ b/Engine/Source/Editor/Kismet/Private/BlueprintEditor.cpp
@@ -1090,10 +1090,6 @@ TSharedRef<SGraphEditor> FBlueprintEditor::CreateGraphEditorWidget(TSharedRef<FT
 				FExecuteAction::CreateSP( this, &FBlueprintEditor::OnOpenRelatedAsset )
 				);
 
-			GraphEditorCommands->MapAction( FGraphEditorCommands::Get().EditTunnel,
-				FExecuteAction::CreateSP( this, &FBlueprintEditor::OnEditTunnel )
-				);
-
 			GraphEditorCommands->MapAction( FGraphEditorCommands::Get().CreateComment,
 				FExecuteAction::CreateSP( this, &FBlueprintEditor::OnCreateComment )
 				);
@@ -6999,19 +6995,6 @@ bool FBlueprintEditor::CanRecompileModules()
 	return !IHotReloadModule::Get().IsCurrentlyCompiling();
 }
 
-void FBlueprintEditor::OnEditTunnel()
-{
-	auto GraphEditor = FocusedGraphEdPtr.Pin();
-	if (GraphEditor.IsValid())
-	{
-		TSharedPtr<SWidget> TitleBar = GraphEditor->GetTitleBar();
-		if (TitleBar.IsValid())
-		{
-			StaticCastSharedPtr<SGraphTitleBar>(TitleBar)->BeginEditing();
-		}
-	}
-}
-
 void FBlueprintEditor::OnCreateComment()
 {
 	auto GraphEditor = FocusedGraphEdPtr.Pin();
diff --git a/Engine/Source/Editor/Kismet/Private/SGraphTitleBar.cpp b/Engine/Source/Editor/Kismet/Private/SGraphTitleBar.cpp
--- a/Engine/Source/Editor/Kismet/Private/SGraphTitleBar.cpp
+++ b/Engine/Source/Editor/Kismet/Private/SGraphTitleBar.cpp
@@ -76,8 +76,6 @@ void SGraphTitleBar::Construct( const FArguments& InArgs )
 	Kismet2Ptr = InArgs._Kismet2;
 	check(Kismet2Ptr.IsValid());
 
-	bEditingFunction = false;
-
 	// Set-up shared breadcrumb defaults
 	FMargin BreadcrumbTrailPadding = FMargin(4.f, 2.f);
 	const FSlateBrush* BreadcrumbButtonImage = FEditorStyle::GetBrush("BreadcrumbTrail.Delimiter");
@@ -207,11 +205,6 @@ void SGraphTitleBar::RebuildBreadcrumbTrail()
 	}
 }
 
-void SGraphTitleBar::BeginEditing()
-{
-	bEditingFunction = true;
-}
-
 UEdGraph* SGraphTitleBar::GetOuterGraph( UObject* Obj )
 {
 	if( Obj )
diff --git a/Engine/Source/Editor/Kismet/Private/SGraphTitleBar.h b/Engine/Source/Editor/Kismet/Private/SGraphTitleBar.h
--- a/Engine/Source/Editor/Kismet/Private/SGraphTitleBar.h
+++ b/Engine/Source/Editor/Kismet/Private/SGraphTitleBar.h
@@ -26,9 +26,6 @@ public:
 
 	void Construct(const FArguments& InArgs);
 
-	/** Enable function editing */
-	void BeginEditing();
-
 	/** Refresh the toolbar */
 	void Refresh();
 
@@ -39,9 +36,6 @@ protected:
 	/** Edited graph */
 	UEdGraph* EdGraphObj;
 
-	/** If we are editing the function */
-	bool bEditingFunction;
-
 	/** Pointer to the function editor widget */
 	TWeakPtr<class SFunctionEditor>	FuncEditorPtr;
 
