missed a couple files for last submit

diff --git a/Engine/Source/Editor/BlueprintGraph/Private/BlueprintEditorSettings.cpp b/Engine/Source/Editor/BlueprintGraph/Private/BlueprintEditorSettings.cpp
--- a/Engine/Source/Editor/BlueprintGraph/Private/BlueprintEditorSettings.cpp
+++ b/Engine/Source/Editor/BlueprintGraph/Private/BlueprintEditorSettings.cpp
@@ -22,6 +22,7 @@ UBlueprintEditorSettings::UBlueprintEditorSettings(const FObjectInitializer& Obj
 	, bShowInheritedVariables(false)
 	, bShowEmptySections(true)
 	, bSpawnDefaultBlueprintNodes(true)
+	, bHideConstructionScriptComponentsInDetailsView(true)
 	// Compiler Settings
 	, SaveOnCompile(SoC_Never)
 	, bJumpToNodeErrors(false)
diff --git a/Engine/Source/Editor/Kismet/Private/SSCSEditor.cpp b/Engine/Source/Editor/Kismet/Private/SSCSEditor.cpp
--- a/Engine/Source/Editor/Kismet/Private/SSCSEditor.cpp
+++ b/Engine/Source/Editor/Kismet/Private/SSCSEditor.cpp
@@ -50,6 +50,7 @@
 #include "SCreateAssetFromObject.h"
 
 #include "SourceCodeNavigation.h"
+#include "Editor/BlueprintGraph/Public/BlueprintEditorSettings.h"
 
 #define LOCTEXT_NAMESPACE "SSCSEditor"
 
@@ -1677,6 +1678,13 @@ FText SSCS_RowWidget::GetIntroducedInToolTipText() const
 							BestClass = OwningBP->GeneratedClass;
 						}
 					}
+					else if (OwningActor != nullptr)
+					{
+						if (UBlueprint* OwningBP = UBlueprint::GetBlueprintFromClass(OwningActor->GetClass()))
+						{
+							BestClass = OwningBP->GeneratedClass;
+						}
+					}
 				}
 
 				if (BestClass == nullptr)
@@ -1690,6 +1698,14 @@ FText SSCS_RowWidget::GetIntroducedInToolTipText() const
 						IntroducedInTooltip = LOCTEXT("IntroducedInNativeError", "Unknown native source (via C++ code)");
 					}
 				}
+				else if (!TreeNode->IsNative() && ComponentTemplate->CreationMethod == EComponentCreationMethod::Native)
+				{
+					IntroducedInTooltip = FText::Format(LOCTEXT("IntroducedInCPPErrorFmt", "{0} (via C++ code)"), FBlueprintEditorUtils::GetFriendlyClassDisplayName(BestClass));
+				}
+				else if (ComponentTemplate->CreationMethod == EComponentCreationMethod::UserConstructionScript)
+				{
+					IntroducedInTooltip = FText::Format(LOCTEXT("IntroducedInUCSErrorFmt", "{0} (via an Add Component call)"), FBlueprintEditorUtils::GetFriendlyClassDisplayName(BestClass));
+				}
 				else
 				{
 					IntroducedInTooltip = FBlueprintEditorUtils::GetFriendlyClassDisplayName(BestClass);
@@ -2090,6 +2106,16 @@ void SSCS_RowWidget::OnDragEnter( const FGeometry& MyGeometry, const FDragDropEv
 					// Can't attach Static components to mobile ones
 					Message = LOCTEXT("DropActionToolTip_Error_CannotAttachStationary", "Cannot attach Stationary components to movable ones.");
 				}
+				else if ((!NodePtr->IsNative() && HoveredTemplate->CreationMethod == EComponentCreationMethod::Native))
+				{
+					// Can't attach to post-construction C++-added components as they exist outside of the CDO and are not known at SCS execution time
+					Message = LOCTEXT("DropActionToolTip_Error_CannotAttachCPPAdded", "Cannot attach to components added in post-construction C++ code.");
+				}
+				else if ((HoveredTemplate->CreationMethod == EComponentCreationMethod::UserConstructionScript))
+				{
+					// Can't attach to UCS-added components as they exist outside of the CDO and are not known at SCS execution time
+					Message = LOCTEXT("DropActionToolTip_Error_CannotAttachUCSAdded", "Cannot attach to components added in the Construction Script.");
+				}
 				else if (HoveredTemplate->CanAttachAsChild(DraggedTemplate, NAME_None))
 				{
 					// Attach the dragged node(s) to this node
@@ -3418,8 +3444,9 @@ void SSCSEditor::Tick( const FGeometry& AllottedGeometry, const double InCurrent
 			int32 NumComponentInstances = 0;
 			for (auto CompIt = ActorInstance->GetComponents().CreateConstIterator(); CompIt; ++CompIt)
 			{
-				// Don't count editor-only components, because we don't show them
-				if (!(*CompIt)->IsEditorOnly())
+				// Don't count editor-only components, because we don't show them. Also, don't count UCS-added components if the option to hide them is enabled.
+				if (!(*CompIt)->IsEditorOnly()
+					&& ((*CompIt)->CreationMethod != EComponentCreationMethod::UserConstructionScript || !GetDefault<UBlueprintEditorSettings>()->bHideConstructionScriptComponentsInDetailsView))
 				{
 					++NumComponentInstances;
 				}
@@ -4169,7 +4196,8 @@ void SSCSEditor::UpdateTree(bool bRegenerateTreeNodes)
 				for(auto CompIter = Components.CreateIterator(); CompIter; ++CompIter)
 				{
 					USceneComponent* SceneComp = Cast<USceneComponent>(*CompIter);
-					if(SceneComp != nullptr && !SceneComp->IsEditorOnly())
+					if(SceneComp != nullptr && !SceneComp->IsEditorOnly()
+						&& (SceneComp->CreationMethod != EComponentCreationMethod::UserConstructionScript || !GetDefault<UBlueprintEditorSettings>()->bHideConstructionScriptComponentsInDetailsView))
 					{
 						AddTreeNodeFromComponent(SceneComp);
 					}
@@ -4179,7 +4207,8 @@ void SSCSEditor::UpdateTree(bool bRegenerateTreeNodes)
 				for(auto CompIter = Components.CreateIterator(); CompIter; ++CompIter)
 				{
 					UActorComponent* ActorComp = *CompIter;
-					if (!ActorComp->IsA<USceneComponent>() && !ActorComp->IsEditorOnly())
+					if (!ActorComp->IsA<USceneComponent>() && !ActorComp->IsEditorOnly()
+						&& (ActorComp->CreationMethod != EComponentCreationMethod::UserConstructionScript || !GetDefault<UBlueprintEditorSettings>()->bHideConstructionScriptComponentsInDetailsView))
 					{
 						if (!bHasAddedSceneAndBehaviorComponentSeparator)
 						{
