CIS fix (extraneous parentheses).

diff --git a/Engine/Source/Editor/Kismet/Private/SSCSEditor.cpp b/Engine/Source/Editor/Kismet/Private/SSCSEditor.cpp
--- a/Engine/Source/Editor/Kismet/Private/SSCSEditor.cpp
+++ b/Engine/Source/Editor/Kismet/Private/SSCSEditor.cpp
@@ -2111,7 +2111,7 @@ void SSCS_RowWidget::OnDragEnter( const FGeometry& MyGeometry, const FDragDropEv
 					// Can't attach to post-construction C++-added components as they exist outside of the CDO and are not known at SCS execution time
 					Message = LOCTEXT("DropActionToolTip_Error_CannotAttachCPPAdded", "Cannot attach to components added in post-construction C++ code.");
 				}
-				else if ((HoveredTemplate->CreationMethod == EComponentCreationMethod::UserConstructionScript))
+				else if (HoveredTemplate->CreationMethod == EComponentCreationMethod::UserConstructionScript)
 				{
 					// Can't attach to UCS-added components as they exist outside of the CDO and are not known at SCS execution time
 					Message = LOCTEXT("DropActionToolTip_Error_CannotAttachUCSAdded", "Cannot attach to components added in the Construction Script.");
