Fix a regression in which the CB thumbnail doesn't get cleared when all visualizable component nodes are deleted in the BP editor.

- introduced in CL# 2410040

diff --git a/Engine/Source/Editor/Kismet/Private/SSCSEditor.cpp b/Engine/Source/Editor/Kismet/Private/SSCSEditor.cpp
--- a/Engine/Source/Editor/Kismet/Private/SSCSEditor.cpp
+++ b/Engine/Source/Editor/Kismet/Private/SSCSEditor.cpp
@@ -4841,9 +4841,13 @@ void SSCSEditor::OnDeleteNodes()
 
 	if (EditorMode == EComponentEditorMode::BlueprintSCS)
 	{
-		// Remove node from SCS
 		UBlueprint* Blueprint = GetBlueprint();
-		FThumbnailRenderingInfo* RenderInfo = nullptr;
+		check(Blueprint != nullptr);
+
+		// Get the current render info for the blueprint. If this is NULL then the blueprint is not currently visualizable (no visible primitive components)
+		FThumbnailRenderingInfo* RenderInfo = GUnrealEd->GetThumbnailManager()->GetRenderingInfo( Blueprint );
+
+		// Remove node(s) from SCS
 		TArray<FSCSEditorTreeNodePtrType> SelectedNodes = SCSTreeWidget->GetSelectedItems();
 		for (int32 i = 0; i < SelectedNodes.Num(); ++i)
 		{
@@ -4855,18 +4859,9 @@ void SSCSEditor::OnDeleteNodes()
 				USimpleConstructionScript* SCS = SCS_Node->GetSCS();
 				check(SCS != nullptr && Blueprint == SCS->GetBlueprint());
 
-				if(Blueprint == nullptr)
-				{
-					Blueprint = SCS->GetBlueprint();
-					check(Blueprint != nullptr);
-
-					// Get the current render info for the blueprint. If this is NULL then the blueprint is not currently visualizable (no visible primitive components)
-					FThumbnailRenderingInfo* RenderInfo = GUnrealEd->GetThumbnailManager()->GetRenderingInfo( Blueprint );
-
-					// Saving objects for restoring purpose.
-					Blueprint->Modify();
-					SaveSCSCurrentState( SCS );
-				}
+				// Saving objects for restoring purpose.
+				Blueprint->Modify();
+				SaveSCSCurrentState( SCS );
 			}
 
 			RemoveComponentNode(Node);
