Fix shadowed variables

diff --git a/Engine/Source/Developer/GameplayDebugger/Classes/GameplayDebuggingHUDComponent.h b/Engine/Source/Developer/GameplayDebugger/Classes/GameplayDebuggingHUDComponent.h
--- a/Engine/Source/Developer/GameplayDebugger/Classes/GameplayDebuggingHUDComponent.h
+++ b/Engine/Source/Developer/GameplayDebugger/Classes/GameplayDebuggingHUDComponent.h
@@ -55,7 +55,7 @@ public:
 
 	/** Set the canvas to use during drawing */
 	void SetCanvas(class UCanvas* InCanvas) { Canvas = InCanvas; }
-	void SetPlayerOwner(APlayerController* PlayerOwner) { this->PlayerOwner = PlayerOwner; }
+	void SetPlayerOwner(APlayerController* InPlayerOwner) { PlayerOwner = InPlayerOwner; }
 	APlayerController* GetPlayerOwner() { return PlayerOwner; }
 
 protected:
diff --git a/Engine/Source/Developer/GameplayDebugger/Private/GameplayDebuggingHUDComponent.cpp b/Engine/Source/Developer/GameplayDebugger/Private/GameplayDebuggingHUDComponent.cpp
--- a/Engine/Source/Developer/GameplayDebugger/Private/GameplayDebuggingHUDComponent.cpp
+++ b/Engine/Source/Developer/GameplayDebugger/Private/GameplayDebuggingHUDComponent.cpp
@@ -268,25 +268,25 @@ void AGameplayDebuggingHUDComponent::DrawPath(APlayerController* MyPC, class UGa
 	static const FColor PathColor(192,192,192);
 
 	const int32 NumPathVerts = DebugComponent->PathPoints.Num();
-	UWorld* World = GetWorld();
+	UWorld* DrawWorld = GetWorld();
 
-		for (int32 VertIdx=0; VertIdx < NumPathVerts-1; ++VertIdx)
-		{
-			FVector const VertLoc = DebugComponent->PathPoints[VertIdx] + NavigationDebugDrawing::PathOffset;
-			DrawDebugSolidBox(World, VertLoc, NavigationDebugDrawing::PathNodeBoxExtent, VertIdx < int32(DebugComponent->NextPathPointIndex) ? Grey : PathColor, false);
-
-			// draw line to next loc
-			FVector const NextVertLoc = DebugComponent->PathPoints[VertIdx+1] + NavigationDebugDrawing::PathOffset;
-			DrawDebugLine(World, VertLoc, NextVertLoc, VertIdx < int32(DebugComponent->NextPathPointIndex) ? Grey : PathColor, false
-				, -1.f, 0
-				, NavigationDebugDrawing::PathLineThickness);
-		}
+	for (int32 VertIdx=0; VertIdx < NumPathVerts-1; ++VertIdx)
+	{
+		FVector const VertLoc = DebugComponent->PathPoints[VertIdx] + NavigationDebugDrawing::PathOffset;
+		DrawDebugSolidBox(DrawWorld, VertLoc, NavigationDebugDrawing::PathNodeBoxExtent, VertIdx < int32(DebugComponent->NextPathPointIndex) ? Grey : PathColor, false);
 
-		// draw last vert
-		if (NumPathVerts > 0)
-		{
-			DrawDebugBox(World, DebugComponent->PathPoints[NumPathVerts-1] + NavigationDebugDrawing::PathOffset, FVector(15.f), Grey, false);
-		}
+		// draw line to next loc
+		FVector const NextVertLoc = DebugComponent->PathPoints[VertIdx+1] + NavigationDebugDrawing::PathOffset;
+		DrawDebugLine(DrawWorld, VertLoc, NextVertLoc, VertIdx < int32(DebugComponent->NextPathPointIndex) ? Grey : PathColor, false
+			, -1.f, 0
+			, NavigationDebugDrawing::PathLineThickness);
+	}
+
+	// draw last vert
+	if (NumPathVerts > 0)
+	{
+		DrawDebugBox(DrawWorld, DebugComponent->PathPoints[NumPathVerts-1] + NavigationDebugDrawing::PathOffset, FVector(15.f), Grey, false);
+	}
 
 #endif //!(UE_BUILD_SHIPPING || UE_BUILD_TEST)
 }
@@ -421,21 +421,23 @@ void AGameplayDebuggingHUDComponent::DrawEQSData(APlayerController* PC, class UG
 		return;
 	}
 
-	int32 Index = 0;
-	PrintString(DefaultContext, TEXT("{white}Queries: "));
-	for (auto CurrentQuery : DebugComponent->EQSLocalData)
 	{
-		if (EQSIndex == Index)
+		int32 Index = 0;
+		PrintString(DefaultContext, TEXT("{white}Queries: "));
+		for (auto CurrentQuery : DebugComponent->EQSLocalData)
 		{
-			PrintString(DefaultContext, FString::Printf(TEXT("{green}%s, "), *CurrentQuery.Name));
-		}
-		else
-		{
-			PrintString(DefaultContext, FString::Printf(TEXT("{yellow}%s, "), *CurrentQuery.Name));
+			if (EQSIndex == Index)
+			{
+				PrintString(DefaultContext, FString::Printf(TEXT("{green}%s, "), *CurrentQuery.Name));
+			}
+			else
+			{
+				PrintString(DefaultContext, FString::Printf(TEXT("{yellow}%s, "), *CurrentQuery.Name));
+			}
+			Index++;
 		}
-		Index++;
+		PrintString(DefaultContext, TEXT("\n"));
 	}
-	PrintString(DefaultContext, TEXT("\n"));
 
 	auto& CurrentLocalData = DebugComponent->EQSLocalData[EQSIndex];
 
diff --git a/Engine/Source/Developer/GameplayDebugger/Private/GameplayDebuggingReplicator.cpp b/Engine/Source/Developer/GameplayDebugger/Private/GameplayDebuggingReplicator.cpp
--- a/Engine/Source/Developer/GameplayDebugger/Private/GameplayDebuggingReplicator.cpp
+++ b/Engine/Source/Developer/GameplayDebugger/Private/GameplayDebuggingReplicator.cpp
@@ -508,9 +508,9 @@ void AGameplayDebuggingReplicator::ServerSetActorToDebug_Implementation(AActor*
 		}
 	}
 
-	if (UGameplayDebuggingComponent* DebugComponent = GetDebugComponent())
+	if (UGameplayDebuggingComponent* DebuggingComponent = GetDebugComponent())
 	{
-		DebugComponent->SetActorToDebug(InActor);
+		DebuggingComponent->SetActorToDebug(InActor);
 	}
 #endif
 }
@@ -548,16 +548,16 @@ void AGameplayDebuggingReplicator::OnDebugAIDelegate(class UCanvas* Canvas, clas
 
 	EnableDraw(true);
 	UWorld* World = GetWorld();
-	UGameplayDebuggingComponent* DebugComponent = GetDebugComponent();
-	if (World && DebugComponent && DebugComponent->GetOwnerRole() == ROLE_Authority)
+	UGameplayDebuggingComponent* DebuggingComponent = GetDebugComponent();
+	if (World && DebuggingComponent && DebuggingComponent->GetOwnerRole() == ROLE_Authority)
 	{
 		UGameplayDebuggingControllerComponent*  GDC = FindComponentByClass<UGameplayDebuggingControllerComponent>();
 		TArray<int32> OryginalReplicateViewDataCounters;
 
-		OryginalReplicateViewDataCounters = DebugComponent->ReplicateViewDataCounters;
+		OryginalReplicateViewDataCounters = DebuggingComponent->ReplicateViewDataCounters;
 		for (uint32 Index = 0; Index < EAIDebugDrawDataView::MAX; ++Index)
 		{
-			DebugComponent->ReplicateViewDataCounters[Index] = GameplayDebuggerSettings(this).CheckFlag((EAIDebugDrawDataView::Type)Index) ? 1 : 0;
+			DebuggingComponent->ReplicateViewDataCounters[Index] = GameplayDebuggerSettings(this).CheckFlag((EAIDebugDrawDataView::Type)Index) ? 1 : 0;
 		}
 
 		// looks like Simulate in UE4 Editor - let's find selected Pawn to debug
@@ -573,8 +573,8 @@ void AGameplayDebuggingReplicator::OnDebugAIDelegate(class UCanvas* Canvas, clas
 			}
 
 			//We needs to collect data manually in Simulate
-			DebugComponent->SetActorToDebug(NewTarget);
-			DebugComponent->CollectDataToReplicate(NewTarget->IsSelected());
+			DebuggingComponent->SetActorToDebug(NewTarget);
+			DebuggingComponent->CollectDataToReplicate(NewTarget->IsSelected());
 			DrawDebugData(Canvas, PC);
 		}
 
@@ -582,17 +582,17 @@ void AGameplayDebuggingReplicator::OnDebugAIDelegate(class UCanvas* Canvas, clas
 		ServerSetActorToDebug(FullSelectedTarget);
 		if (FullSelectedTarget)
 		{
-			DebugComponent->CollectDataToReplicate(true);
-			DebugComponent->SetEQSIndex(ActiveEQSIndex);
+			DebuggingComponent->CollectDataToReplicate(true);
+			DebuggingComponent->SetEQSIndex(ActiveEQSIndex);
 			DrawDebugData(Canvas, PC);
 		}
 
 		if (GetSelectedActorToDebug() != OldActor)
 		{
-			DebugComponent->MarkRenderStateDirty();
+			DebuggingComponent->MarkRenderStateDirty();
 		}
 
-		DebugComponent->ReplicateViewDataCounters = OryginalReplicateViewDataCounters;
+		DebuggingComponent->ReplicateViewDataCounters = OryginalReplicateViewDataCounters;
 
 	}
 #endif
diff --git a/Engine/Source/Editor/DetailCustomizations/Private/HardwareTargetingSettingsDetails.cpp b/Engine/Source/Editor/DetailCustomizations/Private/HardwareTargetingSettingsDetails.cpp
--- a/Engine/Source/Editor/DetailCustomizations/Private/HardwareTargetingSettingsDetails.cpp
+++ b/Engine/Source/Editor/DetailCustomizations/Private/HardwareTargetingSettingsDetails.cpp
@@ -258,12 +258,12 @@ void FHardwareTargetingSettingsDetails::CustomizeDetails(IDetailLayoutBuilder& D
 		DetailBuilder.HideProperty(PropertyName);
 
 		TSharedRef<IPropertyHandle> Property = DetailBuilder.GetProperty(PropertyName);
-		auto SetPropertyValue = [](EHardwareClass::Type NewValue, TSharedRef<IPropertyHandle> Property){
-			Property->SetValue(uint8(NewValue));
+		auto SetPropertyValue = [](EHardwareClass::Type NewValue, TSharedRef<IPropertyHandle> InProperty){
+			InProperty->SetValue(uint8(NewValue));
 		};
-		auto GetPropertyValue = [](TSharedRef<IPropertyHandle> Property){
+		auto GetPropertyValue = [](TSharedRef<IPropertyHandle> InProperty){
 			uint8 Value = 0;
-			Property->GetValue(Value);
+			InProperty->GetValue(Value);
 			return EHardwareClass::Type(Value);
 		};
 		
@@ -279,12 +279,12 @@ void FHardwareTargetingSettingsDetails::CustomizeDetails(IDetailLayoutBuilder& D
 		DetailBuilder.HideProperty(PropertyName);
 
 		TSharedRef<IPropertyHandle> Property = DetailBuilder.GetProperty(PropertyName);
-		auto SetPropertyValue = [](EGraphicsPreset::Type NewValue, TSharedRef<IPropertyHandle> Property){
-			Property->SetValue(uint8(NewValue));
+		auto SetPropertyValue = [](EGraphicsPreset::Type NewValue, TSharedRef<IPropertyHandle> InProperty){
+			InProperty->SetValue(uint8(NewValue));
 		};
-		auto GetPropertyValue = [](TSharedRef<IPropertyHandle> Property){
+		auto GetPropertyValue = [](TSharedRef<IPropertyHandle> InProperty){
 			uint8 Value = 0;
-			Property->GetValue(Value);
+			InProperty->GetValue(Value);
 			return EGraphicsPreset::Type(Value);
 		};
 		GraphicsPresetCombo = HardwareTargeting.MakeGraphicsPresetTargetCombo(
diff --git a/Engine/Source/Editor/DetailCustomizations/Private/MarginCustomization.cpp b/Engine/Source/Editor/DetailCustomizations/Private/MarginCustomization.cpp
--- a/Engine/Source/Editor/DetailCustomizations/Private/MarginCustomization.cpp
+++ b/Engine/Source/Editor/DetailCustomizations/Private/MarginCustomization.cpp
@@ -10,9 +10,9 @@ TSharedRef<IPropertyTypeCustomization> FMarginStructCustomization::MakeInstance(
 	return MakeShareable( new FMarginStructCustomization() );
 }
 
-void FMarginStructCustomization::CustomizeHeader( TSharedRef<class IPropertyHandle> StructPropertyHandle, class FDetailWidgetRow& HeaderRow, IPropertyTypeCustomizationUtils& StructCustomizationUtils )
+void FMarginStructCustomization::CustomizeHeader( TSharedRef<class IPropertyHandle> InStructPropertyHandle, class FDetailWidgetRow& HeaderRow, IPropertyTypeCustomizationUtils& StructCustomizationUtils )
 {
-	this->StructPropertyHandle = StructPropertyHandle;
+	StructPropertyHandle = InStructPropertyHandle;
 
 	const FString UVSpaceString( StructPropertyHandle->GetProperty()->GetMetaData( TEXT( "UVSpace" ) ) );
 	bIsMarginUsingUVSpace = UVSpaceString.Len() > 0 && UVSpaceString == TEXT( "true" );
@@ -44,7 +44,7 @@ void FMarginStructCustomization::CustomizeHeader( TSharedRef<class IPropertyHand
 	];
 }
 
-void FMarginStructCustomization::CustomizeChildren( TSharedRef<class IPropertyHandle> StructPropertyHandle, class IDetailChildrenBuilder& StructBuilder, IPropertyTypeCustomizationUtils& StructCustomizationUtils )
+void FMarginStructCustomization::CustomizeChildren( TSharedRef<class IPropertyHandle> InStructPropertyHandle, class IDetailChildrenBuilder& StructBuilder, IPropertyTypeCustomizationUtils& StructCustomizationUtils )
 {
 	const uint32 NumChildren = ChildPropertyHandles.Num();
 
diff --git a/Engine/Source/Editor/DistCurveEditor/Private/CurveEditorViewportClient.cpp b/Engine/Source/Editor/DistCurveEditor/Private/CurveEditorViewportClient.cpp
--- a/Engine/Source/Editor/DistCurveEditor/Private/CurveEditorViewportClient.cpp
+++ b/Engine/Source/Editor/DistCurveEditor/Private/CurveEditorViewportClient.cpp
@@ -1515,11 +1515,11 @@ void FCurveEditorViewportClient::MoveSelectedKeys(float DeltaIn, float DeltaOut)
 			EdInterface->SetKeyOut(SelKey.SubIndex, SelKey.KeyIndex, NewOut);
 		}
 
-		FCurveEditorModKey TestKey(SelKey.CurveIndex, SelKey.KeyIndex);
+		FCurveEditorModKey KeyToTest(SelKey.CurveIndex, SelKey.KeyIndex);
 
 		// If there is a change in the Input, apply it.
 		// This is slightly complicated because it may change the index of the selected key, so we have to update the selection as we do it.
-		if(FMath::Abs<float>(DeltaIn) > 0.f && !MovedInKeys.Contains(TestKey))
+		if(FMath::Abs<float>(DeltaIn) > 0.f && !MovedInKeys.Contains(KeyToTest))
 		{
 			SelKey.UnsnappedIn += DeltaIn;
 			float NewIn = SnapIn(SelKey.UnsnappedIn);
@@ -1578,8 +1578,8 @@ void FCurveEditorViewportClient::MoveSelectedKeys(float DeltaIn, float DeltaOut)
 			}
 
 			// Remember we have adjusted the In of this key.
-			TestKey.KeyIndex = NewKeyIndex;
-			MovedInKeys.Add(TestKey);
+			KeyToTest.KeyIndex = NewKeyIndex;
+			MovedInKeys.Add(KeyToTest);
 		}
 	} // FOR each selected key
 
diff --git a/Engine/Source/Editor/Documentation/Private/UDNParser.cpp b/Engine/Source/Editor/Documentation/Private/UDNParser.cpp
--- a/Engine/Source/Editor/Documentation/Private/UDNParser.cpp
+++ b/Engine/Source/Editor/Documentation/Private/UDNParser.cpp
@@ -731,11 +731,11 @@ void FUDNParser::AddContentToExcerpt(TSharedPtr<SVerticalBox> Box, const FString
 	}
 }
 
-TSharedRef< SWidget > FUDNParser::GenerateExcerptContent( const FString& Link, FExcerpt& Excerpt, const TArray<FString>& ContentLines, int32 StartingLineIndex )
+TSharedRef< SWidget > FUDNParser::GenerateExcerptContent( const FString& InLink, FExcerpt& Excerpt, const TArray<FString>& ContentLines, int32 StartingLineIndex )
 {
 	FMessageLog UDNParserLog(UDNParseErrorLog);
 
-	const FString SourcePath = FDocumentationLink::ToSourcePath( Link );
+	const FString SourcePath = FDocumentationLink::ToSourcePath( InLink );
 	const FString FullPath = FPaths::GetPath( SourcePath );
 
 	FSlateFontInfo Header1Font = FSlateFontInfo( FPaths::EngineContentDir() / TEXT("Slate/Fonts/Roboto-Regular.ttf"), 18 );
@@ -763,7 +763,7 @@ TSharedRef< SWidget > FUDNParser::GenerateExcerptContent( const FString& Link, F
 		{
 			if (ExcerptStack.Num() == 0 || Line.AdditionalContent[0] != ExcerptStack.Top())
 			{
-				UDNParserLog.NewPage( FText::FromString( Link + TEXT(" [") + Excerpt.Name + TEXT("]") ) );
+				UDNParserLog.NewPage( FText::FromString( InLink + TEXT(" [") + Excerpt.Name + TEXT("]") ) );
 				UDNParserLog.Error(FText::Format(LOCTEXT("ExcerptCloseError", "Line {0}: Excerpt {1} improperly closed."), FText::AsNumber(CurrentLineNumber), FText::FromString(Line.AdditionalContent[0])));
 				bCriticalError = true;
 				break;
@@ -781,7 +781,7 @@ TSharedRef< SWidget > FUDNParser::GenerateExcerptContent( const FString& Link, F
 		{
 			if ( !VariableName.IsEmpty() )
 			{
-				UDNParserLog.NewPage( FText::FromString( Link + TEXT(" [") + Excerpt.Name + TEXT("]") ) );
+				UDNParserLog.NewPage( FText::FromString( InLink + TEXT(" [") + Excerpt.Name + TEXT("]") ) );
 				UDNParserLog.Error(FText::Format(LOCTEXT("VariableOpenError", "Line {0}: Excerpt {1} improperly attempting to define a variable within a variable."), FText::AsNumber(CurrentLineNumber), FText::FromString(Line.AdditionalContent[0])));
 				bCriticalError = true;
 				break;
@@ -791,7 +791,7 @@ TSharedRef< SWidget > FUDNParser::GenerateExcerptContent( const FString& Link, F
 
 			if ( VariableName.IsEmpty() )
 			{
-				UDNParserLog.NewPage( FText::FromString( Link + TEXT(" [") + Excerpt.Name + TEXT("]") ) );
+				UDNParserLog.NewPage( FText::FromString( InLink + TEXT(" [") + Excerpt.Name + TEXT("]") ) );
 				UDNParserLog.Error(FText::Format(LOCTEXT("VariableWithOutName", "Line {0}: Excerpt {1} improperly attempted to define a variable with no name."), FText::AsNumber(CurrentLineNumber), FText::FromString(Line.AdditionalContent[0])));
 				bCriticalError = true;
 				break;
@@ -801,7 +801,7 @@ TSharedRef< SWidget > FUDNParser::GenerateExcerptContent( const FString& Link, F
 		{
 			if ( VariableName.IsEmpty() )
 			{
-				UDNParserLog.NewPage( FText::FromString( Link + TEXT(" [") + Excerpt.Name + TEXT("]") ) );
+				UDNParserLog.NewPage( FText::FromString( InLink + TEXT(" [") + Excerpt.Name + TEXT("]") ) );
 				UDNParserLog.Error(FText::Format(LOCTEXT("VariableCloseError", "Line {0}: Excerpt {1} improperly attempting to close a variable tag it never opened."), FText::AsNumber(CurrentLineNumber), FText::FromString(Line.AdditionalContent[0])));
 				bCriticalError = true;
 				break;
@@ -813,7 +813,7 @@ TSharedRef< SWidget > FUDNParser::GenerateExcerptContent( const FString& Link, F
 		{
 			if ( Line.AdditionalContent.Num() != 2 )
 			{
-				UDNParserLog.NewPage( FText::FromString( Link + TEXT(" [") + Excerpt.Name + TEXT("]") ) );
+				UDNParserLog.NewPage( FText::FromString( InLink + TEXT(" [") + Excerpt.Name + TEXT("]") ) );
 				UDNParserLog.Error(FText::Format(LOCTEXT("Variable", "Line {0}: Excerpt {1} improperly attempted to define a variable with no name."), FText::AsNumber(CurrentLineNumber), FText::FromString(Line.AdditionalContent[0])));
 				bCriticalError = true;
 				break;
@@ -823,7 +823,7 @@ TSharedRef< SWidget > FUDNParser::GenerateExcerptContent( const FString& Link, F
 
 			if ( VariableName.IsEmpty() )
 			{
-				UDNParserLog.NewPage( FText::FromString( Link + TEXT(" [") + Excerpt.Name + TEXT("]") ) );
+				UDNParserLog.NewPage( FText::FromString( InLink + TEXT(" [") + Excerpt.Name + TEXT("]") ) );
 				UDNParserLog.Error(FText::Format(LOCTEXT("VariableWithOutName", "Line {0}: Excerpt {1} improperly attempted to define a variable with no name."), FText::AsNumber(CurrentLineNumber), FText::FromString(Line.AdditionalContent[0])));
 				bCriticalError = true;
 				break;
@@ -953,7 +953,7 @@ TSharedRef< SWidget > FUDNParser::GenerateExcerptContent( const FString& Link, F
 				}
 				else
 				{
-					Excerpt.RichText += FString::Printf(TEXT("<a id=\"browser\" href=\"%s\" style=\"%s\">%s</>"), *Link, *Style.HyperlinkStyleName.ToString(), *Line.AdditionalContent[0]);
+					Excerpt.RichText += FString::Printf(TEXT("<a id=\"browser\" href=\"%s\" style=\"%s\">%s</>"), *InLink, *Style.HyperlinkStyleName.ToString(), *Line.AdditionalContent[0]);
 				}
 				break;
 			case FUDNLine::Image:
@@ -1011,7 +1011,7 @@ TSharedRef< SWidget > FUDNParser::GenerateExcerptContent( const FString& Link, F
 	{
 		if ( !bCriticalError )
 		{
-			UDNParserLog.NewPage( FText::FromString( Link + TEXT(" [") + Excerpt.Name + TEXT("]") ) );
+			UDNParserLog.NewPage( FText::FromString( InLink + TEXT(" [") + Excerpt.Name + TEXT("]") ) );
 		}
 
 		for (int32 i = 0; i < ExcerptStack.Num(); ++i)
diff --git a/Engine/Source/Editor/GraphEditor/Private/GraphActionNode.cpp b/Engine/Source/Editor/GraphEditor/Private/GraphActionNode.cpp
--- a/Engine/Source/Editor/GraphEditor/Private/GraphActionNode.cpp
+++ b/Engine/Source/Editor/GraphEditor/Private/GraphActionNode.cpp
@@ -164,18 +164,18 @@ TSharedPtr<FGraphActionNode> FGraphActionNode::NewRootNode()
 }
 
 //------------------------------------------------------------------------------
-FGraphActionNode::FGraphActionNode(int32 Grouping, int32 SectionID)
-	: SectionID(SectionID)
-	, Grouping(Grouping)
+FGraphActionNode::FGraphActionNode(int32 InGrouping, int32 InSectionID)
+	: SectionID(InSectionID)
+	, Grouping(InGrouping)
 	, bPendingRenameRequest(false)
 	, InsertOrder(0)
 {
 }
 
 //------------------------------------------------------------------------------
-FGraphActionNode::FGraphActionNode(TArray< TSharedPtr<FEdGraphSchemaAction> > const& ActionList, int32 Grouping, int32 SectionID)
-	: SectionID(SectionID)
-	, Grouping(Grouping)
+FGraphActionNode::FGraphActionNode(TArray< TSharedPtr<FEdGraphSchemaAction> > const& ActionList, int32 InGrouping, int32 InSectionID)
+	: SectionID(InSectionID)
+	, Grouping(InGrouping)
 	, Actions(ActionList)
 	, bPendingRenameRequest(false)
 	, InsertOrder(0)
@@ -195,13 +195,13 @@ TSharedPtr<FGraphActionNode> FGraphActionNode::AddChild(FGraphActionListBuilderB
 }
 
 //------------------------------------------------------------------------------
-TSharedPtr<FGraphActionNode> FGraphActionNode::AddSection(int32 Grouping, int32 InSectionID)
+TSharedPtr<FGraphActionNode> FGraphActionNode::AddSection(int32 InGrouping, int32 InSectionID)
 {
 	if ( !ChildSections.Contains(InSectionID) )
 	{
 		ChildSections.Add(InSectionID);
 
-		TSharedPtr<FGraphActionNode> Section = NewSectionHeadingNode(SharedThis(this), Grouping, InSectionID);
+		TSharedPtr<FGraphActionNode> Section = NewSectionHeadingNode(SharedThis(this), InGrouping, InSectionID);
 		InsertChild(Section);
 
 		return Section;
@@ -341,11 +341,11 @@ FString FGraphActionNode::GetCategoryPath() const
 	TWeakPtr<FGraphActionNode> AncestorNode = ParentNode;
 	while (AncestorNode.IsValid())
 	{
-		const FText& DisplayText = AncestorNode.Pin()->DisplayText;
+		const FText& AncestorDisplayText = AncestorNode.Pin()->DisplayText;
 
-		if( !DisplayText.IsEmpty() )
+		if( !AncestorDisplayText.IsEmpty() )
 		{
-			CategoryPath = DisplayText.ToString() + TEXT("|") + CategoryPath;
+			CategoryPath = AncestorDisplayText.ToString() + TEXT("|") + CategoryPath;
 		}
 		AncestorNode = AncestorNode.Pin()->GetParentNode();
 	}
diff --git a/Engine/Source/Editor/GraphEditor/Private/SGraphActionMenu.cpp b/Engine/Source/Editor/GraphEditor/Private/SGraphActionMenu.cpp
--- a/Engine/Source/Editor/GraphEditor/Private/SGraphActionMenu.cpp
+++ b/Engine/Source/Editor/GraphEditor/Private/SGraphActionMenu.cpp
@@ -926,16 +926,10 @@ void SGraphActionMenu::OnItemScrolledIntoView( TSharedPtr<FGraphActionNode> InAc
 
 TSharedRef<ITableRow> SGraphActionMenu::MakeWidget( TSharedPtr<FGraphActionNode> InItem, const TSharedRef<STableViewBase>& OwnerTable, bool bIsReadOnly )
 {
-	FText SectionTitle;
 	TSharedPtr<IToolTip> SectionToolTip;
 
 	if ( InItem->IsSectionHeadingNode() )
 	{
-		if ( OnGetSectionTitle.IsBound() )
-		{
-			SectionTitle = OnGetSectionTitle.Execute(InItem->SectionID);
-		}
-
 		if ( OnGetSectionToolTip.IsBound() )
 		{
 			SectionToolTip = OnGetSectionToolTip.Execute(InItem->SectionID);
diff --git a/Engine/Source/Editor/Kismet/Private/BlueprintActionMenuItem.cpp b/Engine/Source/Editor/Kismet/Private/BlueprintActionMenuItem.cpp
--- a/Engine/Source/Editor/Kismet/Private/BlueprintActionMenuItem.cpp
+++ b/Engine/Source/Editor/Kismet/Private/BlueprintActionMenuItem.cpp
@@ -91,11 +91,11 @@ static UEdGraphNode* FBlueprintMenuActionItemImpl::InvokeAction(const UBlueprint
  ******************************************************************************/
 
 //------------------------------------------------------------------------------
-FBlueprintActionMenuItem::FBlueprintActionMenuItem(UBlueprintNodeSpawner const* NodeSpawner, FBlueprintActionUiSpec const& UiSpec, IBlueprintNodeBinder::FBindingSet const& Bindings)
+FBlueprintActionMenuItem::FBlueprintActionMenuItem(UBlueprintNodeSpawner const* NodeSpawner, FBlueprintActionUiSpec const& UiSpec, IBlueprintNodeBinder::FBindingSet const& InBindings)
 	: Action(NodeSpawner)
 	, IconTint(UiSpec.IconTint)
 	, IconBrush(FEditorStyle::GetBrush(UiSpec.IconName))
-	, Bindings(Bindings)
+	, Bindings(InBindings)
 {
 	check(Action != nullptr);
 
diff --git a/Engine/Source/Editor/Kismet/Private/BlueprintActionMenuUtils.cpp b/Engine/Source/Editor/Kismet/Private/BlueprintActionMenuUtils.cpp
--- a/Engine/Source/Editor/Kismet/Private/BlueprintActionMenuUtils.cpp
+++ b/Engine/Source/Editor/Kismet/Private/BlueprintActionMenuUtils.cpp
@@ -514,7 +514,6 @@ void FBlueprintActionMenuUtils::MakeContextMenu(FBlueprintActionContext const& C
 	if(bCanHaveActorComponents)
 	{
 		// Don't allow actor components in static function graphs
-		UEdGraphSchema_K2 const* K2Schema = GetDefault<UEdGraphSchema_K2>();
 		for (UEdGraph* Graph : Context.Graphs)
 		{
 			bCanHaveActorComponents &= !K2Schema->IsStaticFunctionGraph(Graph);
diff --git a/Engine/Source/Editor/Kismet/Private/BlueprintDetailsCustomization.cpp b/Engine/Source/Editor/Kismet/Private/BlueprintDetailsCustomization.cpp
--- a/Engine/Source/Editor/Kismet/Private/BlueprintDetailsCustomization.cpp
+++ b/Engine/Source/Editor/Kismet/Private/BlueprintDetailsCustomization.cpp
@@ -46,15 +46,15 @@
 
 void FBlueprintDetails::AddEventsCategory(IDetailLayoutBuilder& DetailBuilder, UProperty* VariableProperty)
 {
-	UBlueprint* Blueprint = GetBlueprintObj();
-	check(Blueprint);
+	UBlueprint* BlueprintObj = GetBlueprintObj();
+	check(BlueprintObj);
 
 	if ( UObjectProperty* ComponentProperty = Cast<UObjectProperty>(VariableProperty) )
 	{
 		UClass* PropertyClass = ComponentProperty->PropertyClass;
 
 		// Check for Ed Graph vars that can generate events
-		if ( PropertyClass && Blueprint->AllowsDynamicBinding() )
+		if ( PropertyClass && BlueprintObj->AllowsDynamicBinding() )
 		{
 			if ( FBlueprintEditorUtils::CanClassGenerateEvents(PropertyClass) )
 			{
@@ -131,20 +131,20 @@ void FBlueprintDetails::AddEventsCategory(IDetailLayoutBuilder& DetailBuilder, U
 
 FReply FBlueprintDetails::HandleAddOrViewEventForVariable(const FName EventName, FName PropertyName, TWeakObjectPtr<UClass> PropertyClass)
 {
-	UBlueprint* Blueprint = GetBlueprintObj();
+	UBlueprint* BlueprintObj = GetBlueprintObj();
 
 	// Find the corresponding variable property in the Blueprint
-	UObjectProperty* VariableProperty = FindField<UObjectProperty>(Blueprint->SkeletonGeneratedClass, PropertyName);
+	UObjectProperty* VariableProperty = FindField<UObjectProperty>(BlueprintObj->SkeletonGeneratedClass, PropertyName);
 
 	if ( VariableProperty )
 	{
-		if ( !FKismetEditorUtilities::FindBoundEventForComponent(Blueprint, EventName, VariableProperty->GetFName()) )
+		if ( !FKismetEditorUtilities::FindBoundEventForComponent(BlueprintObj, EventName, VariableProperty->GetFName()) )
 		{
-			FKismetEditorUtilities::CreateNewBoundEventForClass(PropertyClass.Get(), EventName, Blueprint, VariableProperty);
+			FKismetEditorUtilities::CreateNewBoundEventForClass(PropertyClass.Get(), EventName, BlueprintObj, VariableProperty);
 		}
 		else
 		{
-			const UK2Node_ComponentBoundEvent* ExistingNode = FKismetEditorUtilities::FindBoundEventForComponent(Blueprint, EventName, VariableProperty->GetFName());
+			const UK2Node_ComponentBoundEvent* ExistingNode = FKismetEditorUtilities::FindBoundEventForComponent(BlueprintObj, EventName, VariableProperty->GetFName());
 			if ( ExistingNode )
 			{
 				FKismetEditorUtilities::BringKismetToFocusAttentionOnObject(ExistingNode);
@@ -157,9 +157,9 @@ FReply FBlueprintDetails::HandleAddOrViewEventForVariable(const FName EventName,
 
 int32 FBlueprintDetails::HandleAddOrViewIndexForButton(const FName EventName, FName PropertyName) const
 {
-	UBlueprint* Blueprint = GetBlueprintObj();
+	UBlueprint* BlueprintObj = GetBlueprintObj();
 
-	if ( FKismetEditorUtilities::FindBoundEventForComponent(Blueprint, EventName, PropertyName) )
+	if ( FKismetEditorUtilities::FindBoundEventForComponent(BlueprintObj, EventName, PropertyName) )
 	{
 		return 0; // View
 	}
diff --git a/Engine/Source/Editor/Kismet/Private/BlueprintEditor.cpp b/Engine/Source/Editor/Kismet/Private/BlueprintEditor.cpp
--- a/Engine/Source/Editor/Kismet/Private/BlueprintEditor.cpp
+++ b/Engine/Source/Editor/Kismet/Private/BlueprintEditor.cpp
@@ -6957,7 +6957,7 @@ void FBlueprintEditor::RestoreEditedObjectState()
 
 						for (UObject* OuterObject = InGraph->GetOuter(); OuterObject; OuterObject = OuterObject->GetOuter())
 						{
-							if (UBlueprint* Blueprint = Cast<UBlueprint>(OuterObject))
+							if (OuterObject->IsA<UBlueprint>())
 							{
 								// reached up to the blueprint for the graph, we are done climbing the tree
 								OpenCause = FDocumentTracker::OpenNewDocument;
diff --git a/Engine/Source/Editor/Kismet/Private/DiffUtils.cpp b/Engine/Source/Editor/Kismet/Private/DiffUtils.cpp
--- a/Engine/Source/Editor/Kismet/Private/DiffUtils.cpp
+++ b/Engine/Source/Editor/Kismet/Private/DiffUtils.cpp
@@ -293,9 +293,9 @@ TSharedPtr<FBlueprintDifferenceTreeEntry> FBlueprintDifferenceTreeEntry::CreateD
 
 TSharedPtr<FBlueprintDifferenceTreeEntry> FBlueprintDifferenceTreeEntry::CreateDefaultsCategoryEntryForMerge(FOnDiffEntryFocused FocusCallback, const TArray< TSharedPtr<FBlueprintDifferenceTreeEntry> >& Children, bool bHasRemoteDifferences, bool bHasLocalDifferences, bool bHasConflicts)
 {
-	const auto CreateDefaultsRootEntry = [](bool bHasRemoteDifferences, bool bHasLocalDifferences, bool bHasConflicts) -> TSharedRef<SWidget>
+	const auto CreateDefaultsRootEntry = [](bool bInHasRemoteDifferences, bool bInHasLocalDifferences, bool bInHasConflicts) -> TSharedRef<SWidget>
 	{
-		const FLinearColor BaseColor = DiffViewUtils::LookupColor(bHasRemoteDifferences || bHasLocalDifferences, bHasConflicts);
+		const FLinearColor BaseColor = DiffViewUtils::LookupColor(bInHasRemoteDifferences || bInHasLocalDifferences, bInHasConflicts);
 		return SNew(SHorizontalBox)
 			+ SHorizontalBox::Slot()
 			[
@@ -304,9 +304,9 @@ TSharedPtr<FBlueprintDifferenceTreeEntry> FBlueprintDifferenceTreeEntry::CreateD
 				.ColorAndOpacity(BaseColor)
 				.Text(NSLOCTEXT("FBlueprintDifferenceTreeEntry", "DefaultsLabel", "Defaults"))
 			]
-			+ DiffViewUtils::Box(true, DiffViewUtils::LookupColor(bHasRemoteDifferences, bHasConflicts))
+			+ DiffViewUtils::Box(true, DiffViewUtils::LookupColor(bInHasRemoteDifferences, bInHasConflicts))
 			+ DiffViewUtils::Box(true, BaseColor)
-			+ DiffViewUtils::Box(true, DiffViewUtils::LookupColor(bHasLocalDifferences, bHasConflicts));
+			+ DiffViewUtils::Box(true, DiffViewUtils::LookupColor(bInHasLocalDifferences, bInHasConflicts));
 	};
 
 	return TSharedPtr<FBlueprintDifferenceTreeEntry>(new FBlueprintDifferenceTreeEntry(
@@ -335,9 +335,9 @@ TSharedPtr<FBlueprintDifferenceTreeEntry> FBlueprintDifferenceTreeEntry::CreateC
 
 TSharedPtr<FBlueprintDifferenceTreeEntry> FBlueprintDifferenceTreeEntry::CreateComponentsCategoryEntryForMerge(FOnDiffEntryFocused FocusCallback, const TArray< TSharedPtr<FBlueprintDifferenceTreeEntry> >& Children, bool bHasRemoteDifferences, bool bHasLocalDifferences, bool bHasConflicts)
 {
-	const auto CreateComponentsRootEntry = [](bool bHasRemoteDifferences, bool bHasLocalDifferences, bool bHasConflicts) -> TSharedRef<SWidget>
+	const auto CreateComponentsRootEntry = [](bool bInHasRemoteDifferences, bool bInHasLocalDifferences, bool bInHasConflicts) -> TSharedRef<SWidget>
 	{
-		const FLinearColor BaseColor = DiffViewUtils::LookupColor(bHasRemoteDifferences || bHasLocalDifferences, bHasConflicts);
+		const FLinearColor BaseColor = DiffViewUtils::LookupColor(bInHasRemoteDifferences || bInHasLocalDifferences, bInHasConflicts);
 		return  SNew(SHorizontalBox)
 			+ SHorizontalBox::Slot()
 			[
@@ -346,9 +346,9 @@ TSharedPtr<FBlueprintDifferenceTreeEntry> FBlueprintDifferenceTreeEntry::CreateC
 				.ColorAndOpacity(BaseColor)
 				.Text(NSLOCTEXT("FBlueprintDifferenceTreeEntry", "SCSLabel", "Components"))
 			]
-			+ DiffViewUtils::Box(true, DiffViewUtils::LookupColor(bHasRemoteDifferences, bHasConflicts))
+			+ DiffViewUtils::Box(true, DiffViewUtils::LookupColor(bInHasRemoteDifferences, bInHasConflicts))
 			+ DiffViewUtils::Box(true, BaseColor)
-			+ DiffViewUtils::Box(true, DiffViewUtils::LookupColor(bHasLocalDifferences, bHasConflicts));
+			+ DiffViewUtils::Box(true, DiffViewUtils::LookupColor(bInHasLocalDifferences, bInHasConflicts));
 	};
 
 	return TSharedPtr<FBlueprintDifferenceTreeEntry>(new FBlueprintDifferenceTreeEntry(
diff --git a/Engine/Source/Editor/Kismet/Private/FindInBlueprints.cpp b/Engine/Source/Editor/Kismet/Private/FindInBlueprints.cpp
--- a/Engine/Source/Editor/Kismet/Private/FindInBlueprints.cpp
+++ b/Engine/Source/Editor/Kismet/Private/FindInBlueprints.cpp
@@ -249,14 +249,14 @@ bool FFindInBlueprintsResult::ParseSearchInfo(const TArray<FString> &InTokens, F
 
 void FFindInBlueprintsResult::AddExtraSearchInfo(FString InKey, FText InValue, TSharedPtr< FFindInBlueprintsResult > InParent)
 {
-	FText DisplayText;
+	FText ExtraSearchInfoText;
 
 	FFormatNamedArguments Args;
 	Args.Add(TEXT("Key"), FText::FromString(InKey));
 	Args.Add(TEXT("Value"), InValue);
-	DisplayText = FText::Format(LOCTEXT("ExtraSearchInfo", "{Key}: {Value}"), Args);
+	ExtraSearchInfoText = FText::Format(LOCTEXT("ExtraSearchInfo", "{Key}: {Value}"), Args);
 
-	TSharedPtr< FFindInBlueprintsResult > SearchResult(new FFindInBlueprintsResult(DisplayText, InParent) );
+	TSharedPtr< FFindInBlueprintsResult > SearchResult(new FFindInBlueprintsResult(ExtraSearchInfoText, InParent) );
 	InParent->Children.Add(SearchResult);
 }
 
diff --git a/Engine/Source/Editor/Kismet/Private/Nodes/K2Node_AsyncAction.cpp b/Engine/Source/Editor/Kismet/Private/Nodes/K2Node_AsyncAction.cpp
--- a/Engine/Source/Editor/Kismet/Private/Nodes/K2Node_AsyncAction.cpp
+++ b/Engine/Source/Editor/Kismet/Private/Nodes/K2Node_AsyncAction.cpp
@@ -99,12 +99,12 @@ void UK2Node_AsyncAction::GetMenuActions(FBlueprintActionDatabaseRegistrar& Acti
 					UK2Node_AsyncAction* AsyncTaskNode = CastChecked<UK2Node_AsyncAction>(NewNode);
 					if (FunctionPtr.IsValid())
 					{
-						UFunction* Function = FunctionPtr.Get();
-						UObjectProperty* ReturnProperty = CastChecked<UObjectProperty>(Function->GetReturnProperty());
+						UFunction* Func = FunctionPtr.Get();
+						UObjectProperty* ReturnProp = CastChecked<UObjectProperty>(Func->GetReturnProperty());
 						
-						AsyncTaskNode->ProxyFactoryFunctionName = Function->GetFName();
-						AsyncTaskNode->ProxyFactoryClass        = Function->GetOuterUClass();
-						AsyncTaskNode->ProxyClass               = ReturnProperty->PropertyClass;
+						AsyncTaskNode->ProxyFactoryFunctionName = Func->GetFName();
+						AsyncTaskNode->ProxyFactoryClass        = Func->GetOuterUClass();
+						AsyncTaskNode->ProxyClass               = ReturnProp->PropertyClass;
 					}
 				};
 				
diff --git a/Engine/Source/Editor/Kismet/Private/SBlueprintDiff.cpp b/Engine/Source/Editor/Kismet/Private/SBlueprintDiff.cpp
--- a/Engine/Source/Editor/Kismet/Private/SBlueprintDiff.cpp
+++ b/Engine/Source/Editor/Kismet/Private/SBlueprintDiff.cpp
@@ -82,9 +82,9 @@ FSCSDiffControl::FSCSDiffControl(
 	TArray< FSCSResolvedIdentifier > NewHierarchy = NewSCS.GetDisplayedHierarchy();
 	DiffUtils::CompareUnrelatedSCS(InOldBlueprint, OldHierarchy, InNewBlueprint, NewHierarchy, DifferingProperties);
 
-	const auto FocusSCSDifferenceEntry = [](FSCSDiffEntry Entry, FOnSCSDiffControlChanged SelectionCallback, FSCSDiffControl* Owner)
+	const auto FocusSCSDifferenceEntry = [](FSCSDiffEntry Entry, FOnSCSDiffControlChanged InSelectionCallback, FSCSDiffControl* Owner)
 	{
-		SelectionCallback.ExecuteIfBound();
+		InSelectionCallback.ExecuteIfBound();
 		if (Entry.TreeIdentifier.Name != NAME_None)
 		{
 			Owner->OldSCS.HighlightProperty(Entry.TreeIdentifier.Name, FPropertyPath());
@@ -120,10 +120,10 @@ FSCSDiffControl::FSCSDiffControl(
 		Children.Push(FBlueprintDifferenceTreeEntry::NoDifferencesEntry());
 	}
 
-	const auto ForwardSelection = [](FOnSCSDiffControlChanged SelectionCallback)
+	const auto ForwardSelection = [](FOnSCSDiffControlChanged InSelectionCallback)
 	{
 		// This allows the owning control to focus the correct tab (or do whatever else it likes):
-		SelectionCallback.ExecuteIfBound();
+		InSelectionCallback.ExecuteIfBound();
 	};
 
 	OutTreeEntries.Push(FBlueprintDifferenceTreeEntry::CreateComponentsCategoryEntry(
@@ -177,9 +177,9 @@ FCDODiffControl::FCDODiffControl(
 	TArray<FPropertySoftPath> OldProperties = OldDetails.GetDisplayedProperties();
 	TArray<FPropertySoftPath> NewProperties = NewDetails.GetDisplayedProperties();
 
-	const auto FindDiffering = [](TArray< FSingleObjectDiffEntry > const& InDifferences, const FPropertySoftPath& PropertyIdentifer) -> const FSingleObjectDiffEntry*
+	const auto FindDiffering = [](TArray< FSingleObjectDiffEntry > const& InDifferenceEntries, const FPropertySoftPath& PropertyIdentifer) -> const FSingleObjectDiffEntry*
 	{
-		for (const auto& Difference : InDifferences)
+		for (const auto& Difference : InDifferenceEntries)
 		{
 			if (Difference.Identifier == PropertyIdentifer)
 			{
@@ -231,10 +231,10 @@ FCDODiffControl::FCDODiffControl(
 			.ColorAndOpacity(DiffViewUtils::Differs());
 	};
 
-	const auto FocusDetailsDifferenceEntry = [](FPropertySoftPath Identifier, FCDODiffControl* Control, FOnCDODiffControlChanged SelectionCallback)
+	const auto FocusDetailsDifferenceEntry = [](FPropertySoftPath Identifier, FCDODiffControl* Control, FOnCDODiffControlChanged InSelectionCallback)
 	{
 		// This allows the owning control to focus the correct tab (or do whatever else it likes):
-		SelectionCallback.ExecuteIfBound();
+		InSelectionCallback.ExecuteIfBound();
 		Control->HighlightDifference(Identifier);
 	};
 
@@ -261,10 +261,10 @@ FCDODiffControl::FCDODiffControl(
 		Children.Push(FBlueprintDifferenceTreeEntry::NoDifferencesEntry() );
 	}
 
-	const auto ForwardSelection = [](FOnCDODiffControlChanged SelectionCallback)
+	const auto ForwardSelection = [](FOnCDODiffControlChanged InSelectionCallback)
 	{
 		// This allows the owning control to focus the correct tab (or do whatever else it likes):
-		SelectionCallback.ExecuteIfBound();
+		InSelectionCallback.ExecuteIfBound();
 	};
 
 	OutTreeEntries.Push(FBlueprintDifferenceTreeEntry::CreateDefaultsCategoryEntry(
diff --git a/Engine/Source/Editor/Kismet/Private/SBlueprintPalette.cpp b/Engine/Source/Editor/Kismet/Private/SBlueprintPalette.cpp
--- a/Engine/Source/Editor/Kismet/Private/SBlueprintPalette.cpp
+++ b/Engine/Source/Editor/Kismet/Private/SBlueprintPalette.cpp
@@ -1108,12 +1108,12 @@ bool SBlueprintPaletteItem::OnNameTextVerifyChanged(const FText& InNewText, FTex
 		}
 	}
 
-	UBlueprint* Blueprint = BlueprintEditorPtr.Pin()->GetBlueprintObj();
-	check(Blueprint != NULL);
+	UBlueprint* BlueprintObj = BlueprintEditorPtr.Pin()->GetBlueprintObj();
+	check(BlueprintObj != NULL);
 
-	if(Blueprint->SimpleConstructionScript != NULL)
+	if(BlueprintObj->SimpleConstructionScript != NULL)
 	{
-		TArray<USCS_Node*> Nodes = Blueprint->SimpleConstructionScript->GetAllNodes();
+		TArray<USCS_Node*> Nodes = BlueprintObj->SimpleConstructionScript->GetAllNodes();
 		for (TArray<USCS_Node*>::TConstIterator NodeIt(Nodes); NodeIt; ++NodeIt)
 		{
 			USCS_Node* Node = *NodeIt;
@@ -1125,7 +1125,7 @@ bool SBlueprintPaletteItem::OnNameTextVerifyChanged(const FText& InNewText, FTex
 		}
 	}
 
-	TSharedPtr<INameValidatorInterface> NameValidator = MakeShareable(new FKismetNameValidator(Blueprint, OriginalName, ValidationScope));
+	TSharedPtr<INameValidatorInterface> NameValidator = MakeShareable(new FKismetNameValidator(BlueprintObj, OriginalName, ValidationScope));
 
 	EValidatorResult ValidatorResult = NameValidator->IsValid(TextAsString);
 	if(ValidatorResult == EValidatorResult::AlreadyInUse)
@@ -1227,11 +1227,11 @@ void SBlueprintPaletteItem::OnNameTextCommitted(const FText& NewText, ETextCommi
 				const FScopedTransaction Transaction( LOCTEXT( "Rename Delegate", "Rename Event Dispatcher" ) );
 				const FName OldName =  Graph->GetFName();
 
-				UBlueprint* Blueprint = BlueprintEditorPtr.Pin()->GetBlueprintObj();
-				FBlueprintEditorUtils::RenameMemberVariable(Blueprint, OldName, NewName);
+				UBlueprint* BlueprintObj = BlueprintEditorPtr.Pin()->GetBlueprintObj();
+				FBlueprintEditorUtils::RenameMemberVariable(BlueprintObj, OldName, NewName);
 
 				TArray<UK2Node_BaseMCDelegate*> NodeUsingDelegate;
-				FBlueprintEditorUtils::GetAllNodesOfClass<UK2Node_BaseMCDelegate>(Blueprint, NodeUsingDelegate);
+				FBlueprintEditorUtils::GetAllNodesOfClass<UK2Node_BaseMCDelegate>(BlueprintObj, NodeUsingDelegate);
 				for (auto FuncIt = NodeUsingDelegate.CreateIterator(); FuncIt; ++FuncIt)
 				{
 					UK2Node_BaseMCDelegate* FunctionNode = *FuncIt;
@@ -1344,8 +1344,8 @@ FText SBlueprintPaletteItem::GetToolTipText() const
 			{			
 				if(UClass* ParentClass = CallFuncNode->FunctionReference.GetMemberParentClass(CallFuncNode))
 				{
-					UBlueprint* Blueprint = UBlueprint::GetBlueprintFromClass(ParentClass);
-					ClassDisplayName = (Blueprint != NULL) ? Blueprint->GetName() : ParentClass->GetName();
+					UBlueprint* BlueprintObj = UBlueprint::GetBlueprintFromClass(ParentClass);
+					ClassDisplayName = (BlueprintObj != NULL) ? BlueprintObj->GetName() : ParentClass->GetName();
 				}
 			}
 		}
@@ -1366,8 +1366,8 @@ FText SBlueprintPaletteItem::GetToolTipText() const
 			UClass* VarClass = VarAction->GetVariableClass();
 			if (bShowClassInTooltip && (VarClass != NULL))
 			{
-				UBlueprint* Blueprint = UBlueprint::GetBlueprintFromClass(VarClass);
-				ClassDisplayName = (Blueprint != NULL) ? Blueprint->GetName() : VarClass->GetName();
+				UBlueprint* BlueprintObj = UBlueprint::GetBlueprintFromClass(VarClass);
+				ClassDisplayName = (BlueprintObj != NULL) ? BlueprintObj->GetName() : VarClass->GetName();
 			}
 			else
 			{
@@ -1385,8 +1385,8 @@ FText SBlueprintPaletteItem::GetToolTipText() const
 				UClass* VarClass = CastChecked<UClass>(LocalVarAction->GetVariableScope()->GetOuter());
 				if (bShowClassInTooltip && (VarClass != NULL))
 				{
-					UBlueprint* Blueprint = UBlueprint::GetBlueprintFromClass(VarClass);
-					ClassDisplayName = (Blueprint != NULL) ? Blueprint->GetName() : VarClass->GetName();
+					UBlueprint* BlueprintObj = UBlueprint::GetBlueprintFromClass(VarClass);
+					ClassDisplayName = (BlueprintObj != NULL) ? BlueprintObj->GetName() : VarClass->GetName();
 				}
 				else
 				{
diff --git a/Engine/Source/Editor/Kismet/Private/SMyBlueprint.cpp b/Engine/Source/Editor/Kismet/Private/SMyBlueprint.cpp
--- a/Engine/Source/Editor/Kismet/Private/SMyBlueprint.cpp
+++ b/Engine/Source/Editor/Kismet/Private/SMyBlueprint.cpp
@@ -815,18 +815,18 @@ TSharedRef<SWidget> SMyBlueprint::OnCreateWidgetForAction(FCreateWidgetForAction
 	return BlueprintEditorPtr.IsValid() ? SNew(SBlueprintPaletteItem, InCreateData, BlueprintEditorPtr.Pin()) : SNew(SBlueprintPaletteItem, InCreateData, GetBlueprintObj());
 }
 
-void SMyBlueprint::GetChildGraphs(UEdGraph* EdGraph, FGraphActionListBuilderBase& OutAllActions, FString ParentCategory)
+void SMyBlueprint::GetChildGraphs(UEdGraph* InEdGraph, FGraphActionListBuilderBase& OutAllActions, FString ParentCategory)
 {
 	// Grab subgraphs
-	const UEdGraphSchema* Schema = EdGraph->GetSchema();
+	const UEdGraphSchema* Schema = InEdGraph->GetSchema();
 
 	// Grab display info
 	FGraphDisplayInfo EdGraphDisplayInfo;
-	Schema->GetGraphDisplayInformation(*EdGraph, EdGraphDisplayInfo);
+	Schema->GetGraphDisplayInformation(*InEdGraph, EdGraphDisplayInfo);
 	const FText EdGraphDisplayName = EdGraphDisplayInfo.DisplayName;
 
 	// Grab children graphs
-	for (TArray<UEdGraph*>::TIterator It(EdGraph->SubGraphs); It; ++It)
+	for (TArray<UEdGraph*>::TIterator It(InEdGraph->SubGraphs); It; ++It)
 	{
 		UEdGraph* Graph = *It;
 		check(Graph);
@@ -874,39 +874,39 @@ struct FCreateEdGraphSchemaActionHelper
 	}
 };
 
-void SMyBlueprint::GetChildEvents(UEdGraph const* EdGraph, int32 const SectionId, FGraphActionListBuilderBase& OutAllActions) const
+void SMyBlueprint::GetChildEvents(UEdGraph const* InEdGraph, int32 const SectionId, FGraphActionListBuilderBase& OutAllActions) const
 {
-	if (!ensure(EdGraph != NULL))
+	if (!ensure(InEdGraph != NULL))
 	{
 		return;
 	}
 
 	// grab the parent graph's name
-	UEdGraphSchema const* Schema = EdGraph->GetSchema();
+	UEdGraphSchema const* Schema = InEdGraph->GetSchema();
 	FGraphDisplayInfo EdGraphDisplayInfo;
-	Schema->GetGraphDisplayInformation(*EdGraph, EdGraphDisplayInfo);
+	Schema->GetGraphDisplayInformation(*InEdGraph, EdGraphDisplayInfo);
 	FText const EdGraphDisplayName = EdGraphDisplayInfo.DisplayName;
 	FString const ActionCategory = EdGraphDisplayName.ToString();
 
-	FCreateEdGraphSchemaActionHelper::CreateAll<FEdGraphSchemaAction_K2Event, UK2Node_Event>(EdGraph, SectionId, OutAllActions, ActionCategory);
-	FCreateEdGraphSchemaActionHelper::CreateAll<FEdGraphSchemaAction_K2InputAction, UK2Node_InputKey>(EdGraph, SectionId, OutAllActions, ActionCategory);
-	FCreateEdGraphSchemaActionHelper::CreateAll<FEdGraphSchemaAction_K2InputAction, UK2Node_InputAction>(EdGraph, SectionId, OutAllActions, ActionCategory);
+	FCreateEdGraphSchemaActionHelper::CreateAll<FEdGraphSchemaAction_K2Event, UK2Node_Event>(InEdGraph, SectionId, OutAllActions, ActionCategory);
+	FCreateEdGraphSchemaActionHelper::CreateAll<FEdGraphSchemaAction_K2InputAction, UK2Node_InputKey>(InEdGraph, SectionId, OutAllActions, ActionCategory);
+	FCreateEdGraphSchemaActionHelper::CreateAll<FEdGraphSchemaAction_K2InputAction, UK2Node_InputAction>(InEdGraph, SectionId, OutAllActions, ActionCategory);
 }
 
 void SMyBlueprint::GetLocalVariables(FGraphActionListBuilderBase& OutAllActions) const
 {
 	// We want to pull local variables from the top level function graphs
-	UEdGraph* EdGraph = FBlueprintEditorUtils::GetTopLevelGraph(GetFocusedGraph());
-	if( EdGraph )
+	UEdGraph* TopLevelGraph = FBlueprintEditorUtils::GetTopLevelGraph(GetFocusedGraph());
+	if( TopLevelGraph )
 	{
 		// grab the parent graph's name
-		UEdGraphSchema const* Schema = EdGraph->GetSchema();
+		UEdGraphSchema const* Schema = TopLevelGraph->GetSchema();
 		FGraphDisplayInfo EdGraphDisplayInfo;
-		Schema->GetGraphDisplayInformation(*EdGraph, EdGraphDisplayInfo);
+		Schema->GetGraphDisplayInformation(*TopLevelGraph, EdGraphDisplayInfo);
 		FText const EdGraphDisplayName = EdGraphDisplayInfo.DisplayName;
 
 		TArray<UK2Node_FunctionEntry*> FunctionEntryNodes;
-		EdGraph->GetNodesOfClass<UK2Node_FunctionEntry>(FunctionEntryNodes);
+		TopLevelGraph->GetNodesOfClass<UK2Node_FunctionEntry>(FunctionEntryNodes);
 
 		// Search in all FunctionEntry nodes for their local variables
 		FString ActionCategory;
@@ -921,7 +921,7 @@ void SMyBlueprint::GetLocalVariables(FGraphActionListBuilderBase& OutAllActions)
 					Category = FString();
 				}
 
-				UFunction* Func = FindField<UFunction>(GetBlueprintObj()->SkeletonGeneratedClass, EdGraph->GetFName());
+				UFunction* Func = FindField<UFunction>(GetBlueprintObj()->SkeletonGeneratedClass, TopLevelGraph->GetFName());
 				if (Func)
 				{
 					TSharedPtr<FEdGraphSchemaAction_K2LocalVar> NewVarAction = MakeShareable(new FEdGraphSchemaAction_K2LocalVar(Category, FText::FromName(Variable.VarName), TEXT(""), 0));
@@ -952,8 +952,8 @@ void SMyBlueprint::CollectAllActions(FGraphActionListBuilderBase& OutAllActions)
 {
 	const UEdGraphSchema_K2* K2Schema = GetDefault<UEdGraphSchema_K2>();
 
-	UBlueprint* Blueprint = GetBlueprintObj();
-	check(Blueprint);
+	UBlueprint* BlueprintObj = GetBlueprintObj();
+	check(BlueprintObj);
 
 	EFieldIteratorFlags::SuperClassFlags FieldIteratorSuperFlag = EFieldIteratorFlags::IncludeSuper;
 	if ( ShowUserVarsOnly() )
@@ -1045,12 +1045,12 @@ void SMyBlueprint::CollectAllActions(FGraphActionListBuilderBase& OutAllActions)
 	}; 
 
 	// Initialise action sorting instance
-	FGraphActionSort SortList( Blueprint->CategorySorting );
+	FGraphActionSort SortList( BlueprintObj->CategorySorting );
 	// List of names of functions we implement
 	ImplementedFunctionCache.Empty();
 
 	// Grab Variables
-	for (TFieldIterator<UProperty> PropertyIt(Blueprint->SkeletonGeneratedClass, FieldIteratorSuperFlag); PropertyIt; ++PropertyIt)
+	for (TFieldIterator<UProperty> PropertyIt(BlueprintObj->SkeletonGeneratedClass, FieldIteratorSuperFlag); PropertyIt; ++PropertyIt)
 	{
 		UProperty* Property = *PropertyIt;
 		FName PropName = Property->GetFName();
@@ -1075,7 +1075,7 @@ void SMyBlueprint::CollectAllActions(FGraphActionListBuilderBase& OutAllActions)
 		FString PropertyCategory = FObjectEditorUtils::GetCategory(Property);
 		const FString UserCategoryName = FEditorCategoryUtils::GetCategoryDisplayString( PropertyCategory );
 
-		if ((CategoryName == Blueprint->GetFName()) || (CategoryName == K2Schema->VR_DefaultCategory))
+		if ((CategoryName == BlueprintObj->GetFName()) || (CategoryName == K2Schema->VR_DefaultCategory))
 		{
 			CategoryName = NAME_None;		// default, so place in 'non' category
 			PropertyCategory = FString();
@@ -1094,7 +1094,7 @@ void SMyBlueprint::CollectAllActions(FGraphActionListBuilderBase& OutAllActions)
 			TSharedPtr<FEdGraphSchemaAction_K2Var> NewVarAction = MakeShareable(new FEdGraphSchemaAction_K2Var(PropertyCategory, PropertyDesc, PropertyTooltip.ToString(), 0));
 			const UArrayProperty* ArrayProperty = Cast<const UArrayProperty>(Property);
 			const UProperty* TestProperty = ArrayProperty ? ArrayProperty->Inner : Property;
-			NewVarAction->SetVariableInfo(PropertyName, Blueprint->SkeletonGeneratedClass, Cast<UBoolProperty>(TestProperty) != nullptr);
+			NewVarAction->SetVariableInfo(PropertyName, BlueprintObj->SkeletonGeneratedClass, Cast<UBoolProperty>(TestProperty) != nullptr);
 			NewVarAction->SectionID = NodeSectionID::VARIABLE;
 			SortList.AddAction( UserCategoryName, NewVarAction );
 		}
@@ -1105,15 +1105,15 @@ void SMyBlueprint::CollectAllActions(FGraphActionListBuilderBase& OutAllActions)
 			if (Property->HasAllPropertyFlags(CPF_Edit) || !PropertyCategory.IsEmpty())
 			{
 				NewFuncAction = MakeShareable(new FEdGraphSchemaAction_K2Delegate(PropertyCategory, PropertyDesc, PropertyTooltip.ToString(), 0));
-				NewFuncAction->SetDelegateInfo(PropertyName, Blueprint->SkeletonGeneratedClass);
+				NewFuncAction->SetDelegateInfo(PropertyName, BlueprintObj->SkeletonGeneratedClass);
 				NewFuncAction->EdGraph = NULL;
 				NewFuncAction->SectionID = NodeSectionID::DELEGATE;
 				SortList.AddAction( UserCategoryName, NewFuncAction );
 			}
 
 			UClass* OwnerClass = CastChecked<UClass>(Property->GetOuter());
-			UEdGraph* Graph = FBlueprintEditorUtils::GetDelegateSignatureGraphByName(Blueprint, PropertyName);
-			if (Graph && OwnerClass && (Blueprint == OwnerClass->ClassGeneratedBy))
+			UEdGraph* Graph = FBlueprintEditorUtils::GetDelegateSignatureGraphByName(BlueprintObj, PropertyName);
+			if (Graph && OwnerClass && (BlueprintObj == OwnerClass->ClassGeneratedBy))
 			{
 				if ( NewFuncAction.IsValid() )
 				{
@@ -1129,9 +1129,9 @@ void SMyBlueprint::CollectAllActions(FGraphActionListBuilderBase& OutAllActions)
 	SortList.GetAllActions( OutAllActions );
 
 	// Grab functions implemented by the blueprint
-	for (int32 i = 0; i < Blueprint->FunctionGraphs.Num(); i++)
+	for (int32 i = 0; i < BlueprintObj->FunctionGraphs.Num(); i++)
 	{
-		UEdGraph* Graph = Blueprint->FunctionGraphs[i];
+		UEdGraph* Graph = BlueprintObj->FunctionGraphs[i];
 		check(Graph);
 
 		bool bIsConstructionScript = Graph->GetFName() == K2Schema->FN_UserConstructionScript;
@@ -1140,7 +1140,7 @@ void SMyBlueprint::CollectAllActions(FGraphActionListBuilderBase& OutAllActions)
 		Graph->GetSchema()->GetGraphDisplayInformation(*Graph, DisplayInfo);
 
 		FString FunctionCategory;
-		if(UFunction* Function = Blueprint->SkeletonGeneratedClass->FindFunctionByName(Graph->GetFName()))
+		if(UFunction* Function = BlueprintObj->SkeletonGeneratedClass->FindFunctionByName(Graph->GetFName()))
 		{
 			FunctionCategory = Function->GetMetaData(FBlueprintMetadata::MD_FunctionCategory);
 		}
@@ -1164,9 +1164,9 @@ void SMyBlueprint::CollectAllActions(FGraphActionListBuilderBase& OutAllActions)
 		ImplementedFunctionCache.Add(Graph->GetFName());
 	}
 	// Grab macros implemented by the blueprint
-	for (int32 i = 0; i < Blueprint->MacroGraphs.Num(); i++)
+	for (int32 i = 0; i < BlueprintObj->MacroGraphs.Num(); i++)
 	{
-		UEdGraph* Graph = Blueprint->MacroGraphs[i];
+		UEdGraph* Graph = BlueprintObj->MacroGraphs[i];
 		check(Graph);
 		
 		const FName FunctionName = Graph->GetFName();
@@ -1194,12 +1194,12 @@ void SMyBlueprint::CollectAllActions(FGraphActionListBuilderBase& OutAllActions)
 	TArray<FName> OverridableFunctionNames;
 
 	// Cache potentially overridable functions
-	for ( TFieldIterator<UFunction> FunctionIt(Blueprint->ParentClass, EFieldIteratorFlags::IncludeSuper); FunctionIt; ++FunctionIt )
+	for ( TFieldIterator<UFunction> FunctionIt(BlueprintObj->ParentClass, EFieldIteratorFlags::IncludeSuper); FunctionIt; ++FunctionIt )
 	{
 		const UFunction* Function = *FunctionIt;
 		const FName FunctionName = Function->GetFName();
 
-		if ( UEdGraphSchema_K2::CanKismetOverrideFunction(Function) && !OverridableFunctionNames.Contains(FunctionName) && !ImplementedFunctionCache.Contains(FunctionName) && !FBlueprintEditorUtils::FindOverrideForFunction(Blueprint, CastChecked<UClass>(Function->GetOuter()), Function->GetFName()) )
+		if ( UEdGraphSchema_K2::CanKismetOverrideFunction(Function) && !OverridableFunctionNames.Contains(FunctionName) && !ImplementedFunctionCache.Contains(FunctionName) && !FBlueprintEditorUtils::FindOverrideForFunction(BlueprintObj, CastChecked<UClass>(Function->GetOuter()), Function->GetFName()) )
 		{
 			FString FunctionTooltip = UK2Node_CallFunction::GetDefaultTooltipForFunction(Function);
 			FText FunctionDesc = FText::FromString(Function->GetMetaData(TEXT("FriendlyName")));
@@ -1220,9 +1220,9 @@ void SMyBlueprint::CollectAllActions(FGraphActionListBuilderBase& OutAllActions)
 	}
 
 	// Also function implemented for interfaces
-	for (int32 i=0; i < Blueprint->ImplementedInterfaces.Num(); i++)
+	for (int32 i=0; i < BlueprintObj->ImplementedInterfaces.Num(); i++)
 	{
-		FBPInterfaceDescription& InterfaceDesc = Blueprint->ImplementedInterfaces[i];
+		FBPInterfaceDescription& InterfaceDesc = BlueprintObj->ImplementedInterfaces[i];
 		for (int32 FuncIdx = 0; FuncIdx < InterfaceDesc.Graphs.Num(); FuncIdx++)
 		{
 			UEdGraph* Graph = InterfaceDesc.Graphs[FuncIdx];
@@ -1233,7 +1233,7 @@ void SMyBlueprint::CollectAllActions(FGraphActionListBuilderBase& OutAllActions)
 			FString FunctionDesc = FunctionName.ToString();
 			
 			FString FunctionCategory;
-			if ( UFunction* Function = Blueprint->SkeletonGeneratedClass->FindFunctionByName(Graph->GetFName()) )
+			if ( UFunction* Function = BlueprintObj->SkeletonGeneratedClass->FindFunctionByName(Graph->GetFName()) )
 			{
 				FunctionCategory = Function->GetMetaData(FBlueprintMetadata::MD_FunctionCategory);
 			}
@@ -1250,7 +1250,7 @@ void SMyBlueprint::CollectAllActions(FGraphActionListBuilderBase& OutAllActions)
 	}
 
 	// also walk up the class chain to look for overridable functions in natively implemented interfaces
-	for ( UClass* TempClass=Blueprint->ParentClass; TempClass; TempClass=TempClass->GetSuperClass() )
+	for ( UClass* TempClass=BlueprintObj->ParentClass; TempClass; TempClass=TempClass->GetSuperClass() )
 	{
 		for (int32 Idx=0; Idx<TempClass->Interfaces.Num(); ++Idx)
 		{
@@ -1281,9 +1281,9 @@ void SMyBlueprint::CollectAllActions(FGraphActionListBuilderBase& OutAllActions)
 	}
 
 	// Grab ubergraph pages
-	for (int32 i = 0; i < Blueprint->UbergraphPages.Num(); i++)
+	for (int32 i = 0; i < BlueprintObj->UbergraphPages.Num(); i++)
 	{
-		UEdGraph* Graph = Blueprint->UbergraphPages[i];
+		UEdGraph* Graph = BlueprintObj->UbergraphPages[i];
 		check(Graph);
 		
 		FGraphDisplayInfo DisplayInfo;
@@ -1300,9 +1300,9 @@ void SMyBlueprint::CollectAllActions(FGraphActionListBuilderBase& OutAllActions)
 	}
 
 	// Grab intermediate pages
-	for (int32 i = 0; i < Blueprint->IntermediateGeneratedGraphs.Num(); i++)
+	for (int32 i = 0; i < BlueprintObj->IntermediateGeneratedGraphs.Num(); i++)
 	{
-		UEdGraph* Graph = Blueprint->IntermediateGeneratedGraphs[i];
+		UEdGraph* Graph = BlueprintObj->IntermediateGeneratedGraphs[i];
 		check(Graph);
 		
 		const FName FunctionName(*(FString(TEXT("$INTERMEDIATE$_")) + Graph->GetName()));
@@ -1611,7 +1611,7 @@ void SMyBlueprint::OnActionDoubleClicked(const TArray< TSharedPtr<FEdGraphSchema
 
 void SMyBlueprint::ExecuteAction(TSharedPtr<FEdGraphSchemaAction> InAction)
 {
-	UBlueprint* Blueprint = BlueprintEditorPtr.Pin()->GetBlueprintObj();
+	UBlueprint* BlueprintObj = BlueprintEditorPtr.Pin()->GetBlueprintObj();
 	if(InAction.IsValid())
 	{
 		if(InAction->GetTypeId() == FEdGraphSchemaAction_K2Graph::StaticGetTypeId())
@@ -1642,12 +1642,12 @@ void SMyBlueprint::ExecuteAction(TSharedPtr<FEdGraphSchemaAction> InAction)
 				ObjectProperty->PropertyClass &&
 				ObjectProperty->PropertyClass->IsChildOf(UTimelineComponent::StaticClass()))
 			{
-				for (int32 i=0; i<Blueprint->Timelines.Num(); i++)
+				for (int32 i=0; i<BlueprintObj->Timelines.Num(); i++)
 				{
 					// Convert the Timeline's name to a variable name before comparing it to the variable
-					if (FName(*UTimelineTemplate::TimelineTemplateNameToVariableName(Blueprint->Timelines[i]->GetFName())) == VarAction->GetVariableName())
+					if (FName(*UTimelineTemplate::TimelineTemplateNameToVariableName(BlueprintObj->Timelines[i]->GetFName())) == VarAction->GetVariableName())
 					{
-						BlueprintEditorPtr.Pin()->OpenDocument(Blueprint->Timelines[i], FDocumentTracker::OpenNewDocument);
+						BlueprintEditorPtr.Pin()->OpenDocument(BlueprintObj->Timelines[i], FDocumentTracker::OpenNewDocument);
 					}
 				}
 			}
@@ -1893,28 +1893,28 @@ bool SMyBlueprint::CanOpenGraph() const
 
 void SMyBlueprint::OpenGraph(FDocumentTracker::EOpenDocumentCause InCause)
 {
-	UEdGraph* EdGraph = NULL;
+	UEdGraph* GraphToOpen = NULL;
 
 	if (FEdGraphSchemaAction_K2Graph* GraphAction = SelectionAsGraph())
 	{
-		EdGraph = GraphAction->EdGraph;
+		GraphToOpen = GraphAction->EdGraph;
 	}
 	else if (FEdGraphSchemaAction_K2Delegate* DelegateAction = SelectionAsDelegate())
 	{
-		EdGraph = DelegateAction->EdGraph;
+		GraphToOpen = DelegateAction->EdGraph;
 	}
 	else if (FEdGraphSchemaAction_K2Event* EventAction = SelectionAsEvent())
 	{
-		EdGraph = EventAction->NodeTemplate->GetGraph();
+		GraphToOpen = EventAction->NodeTemplate->GetGraph();
 	}
 	else if (FEdGraphSchemaAction_K2InputAction* InputAction = SelectionAsInputAction())
 	{
-		EdGraph = InputAction->NodeTemplate->GetGraph();
+		GraphToOpen = InputAction->NodeTemplate->GetGraph();
 	}
 	
-	if (EdGraph)
+	if (GraphToOpen)
 	{
-		BlueprintEditorPtr.Pin()->OpenDocument(EdGraph, InCause);
+		BlueprintEditorPtr.Pin()->OpenDocument(GraphToOpen, InCause);
 	}
 }
 
@@ -2171,18 +2171,18 @@ bool SMyBlueprint::IsEditingMode() const
 
 void SMyBlueprint::OnDeleteDelegate(FEdGraphSchemaAction_K2Delegate* InDelegateAction)
 {
-	UEdGraph* EdGraph = InDelegateAction->EdGraph;
-	UBlueprint* Blueprint = GetBlueprintObj();
-	if (EdGraph && Blueprint)
+	UEdGraph* GraphToActOn = InDelegateAction->EdGraph;
+	UBlueprint* BlueprintObj = GetBlueprintObj();
+	if (GraphToActOn && BlueprintObj)
 	{
 		const FScopedTransaction Transaction( LOCTEXT("RemoveDelegate", "Remove Event Dispatcher") );
-		Blueprint->Modify();
+		BlueprintObj->Modify();
 
-		BlueprintEditorPtr.Pin()->CloseDocumentTab(EdGraph);
-		EdGraph->Modify();
+		BlueprintEditorPtr.Pin()->CloseDocumentTab(GraphToActOn);
+		GraphToActOn->Modify();
 
-		FBlueprintEditorUtils::RemoveMemberVariable(Blueprint, EdGraph->GetFName());
-		FBlueprintEditorUtils::RemoveGraph(Blueprint, EdGraph, EGraphRemoveFlags::Recompile);
+		FBlueprintEditorUtils::RemoveMemberVariable(BlueprintObj, GraphToActOn->GetFName());
+		FBlueprintEditorUtils::RemoveGraph(BlueprintObj, GraphToActOn, EGraphRemoveFlags::Recompile);
 
 		for (TObjectIterator<UK2Node_CreateDelegate> It(RF_ClassDefaultObject | RF_PendingKill); It; ++It)
 		{
@@ -2622,9 +2622,9 @@ void SMyBlueprint::ExpandCategory(const FString& CategoryName)
 bool SMyBlueprint::MoveCategoryBeforeCategory( const FString& InCategoryToMove, const FString& InTargetCategory )
 {
 	bool bResult = false;
-	UBlueprint* Blueprint = BlueprintEditorPtr.Pin()->GetBlueprintObj();
+	UBlueprint* BlueprintObj = BlueprintEditorPtr.Pin()->GetBlueprintObj();
 
-	if( Blueprint )
+	if( BlueprintObj )
 	{
 		// Find root categories
 		int32 RootCategoryDelim = InCategoryToMove.Find( TEXT( "|" ));
@@ -2632,7 +2632,7 @@ bool SMyBlueprint::MoveCategoryBeforeCategory( const FString& InCategoryToMove,
 		RootCategoryDelim = InTargetCategory.Find( TEXT( "|" ));
 		FName TargetCategory = RootCategoryDelim == INDEX_NONE ? *InTargetCategory : *InTargetCategory.Left( RootCategoryDelim );
 
-		TArray<FName>& CategorySort = Blueprint->CategorySorting;
+		TArray<FName>& CategorySort = BlueprintObj->CategorySorting;
 		const int32 RemovalIndex = CategorySort.Find( CategoryToMove );
 		// Remove existing sort index
 		if( RemovalIndex != INDEX_NONE )
diff --git a/Engine/Source/Editor/Kismet/Private/SSCSEditor.cpp b/Engine/Source/Editor/Kismet/Private/SSCSEditor.cpp
--- a/Engine/Source/Editor/Kismet/Private/SSCSEditor.cpp
+++ b/Engine/Source/Editor/Kismet/Private/SSCSEditor.cpp
@@ -3510,24 +3510,24 @@ void SSCSEditor::GetSelectedItemsForContextMenu(TArray<FComponentEventConstructi
 
 TSharedPtr< SWidget > SSCSEditor::CreateContextMenu()
 {
-	TArray<FSCSEditorTreeNodePtrType> SelectedNodes = SCSTreeWidget->GetSelectedItems();
+	TArray<FSCSEditorTreeNodePtrType> SelectedItems = SCSTreeWidget->GetSelectedItems();
 
-	if (SelectedNodes.Num() > 0 || CanPasteNodes())
+	if (SelectedItems.Num() > 0 || CanPasteNodes())
 	{
 		const bool CloseAfterSelection = true;
 		FMenuBuilder MenuBuilder( CloseAfterSelection, CommandList );
 
 		bool bOnlyShowPasteOption = false;
 
-		if (SelectedNodes.Num() > 0)
+		if (SelectedItems.Num() > 0)
 		{
-			if (SelectedNodes.Num() == 1 && SelectedNodes[0]->GetNodeType() == FSCSEditorTreeNode::RootActorNode)
+			if (SelectedItems.Num() == 1 && SelectedItems[0]->GetNodeType() == FSCSEditorTreeNode::RootActorNode)
 			{
 				bOnlyShowPasteOption = true;
 			}
 			else
 			{
-				for (auto SelectedNode : SelectedNodes)
+				for (auto SelectedNode : SelectedItems)
 				{
 					if (SelectedNode->GetNodeType() != FSCSEditorTreeNode::ComponentNode)
 					{
diff --git a/Engine/Source/Editor/Persona/Private/SSkeletonSmartNameManager.cpp b/Engine/Source/Editor/Persona/Private/SSkeletonSmartNameManager.cpp
--- a/Engine/Source/Editor/Persona/Private/SSkeletonSmartNameManager.cpp
+++ b/Engine/Source/Editor/Persona/Private/SSkeletonSmartNameManager.cpp
@@ -377,8 +377,8 @@ void SCurveNameManager::OnDeleteNameClicked()
 		FAssetData& Data = AnimationAssets[Idx];
 		bool bRemove = true;
 
-		const FString* SkeletonData = Data.TagsAndValues.Find("Skeleton");
-		if(SkeletonData && *SkeletonData == CurrentSkeletonName)
+		const FString* SkeletonDataTag = Data.TagsAndValues.Find("Skeleton");
+		if(SkeletonDataTag && *SkeletonDataTag == CurrentSkeletonName)
 		{
 			const FString* CurveData = Data.TagsAndValues.Find(USkeleton::CurveTag);
 			FString CurveDataCopy;
diff --git a/Engine/Source/Runtime/AIModule/Private/Navigation/PathFollowingComponent.cpp b/Engine/Source/Runtime/AIModule/Private/Navigation/PathFollowingComponent.cpp
--- a/Engine/Source/Runtime/AIModule/Private/Navigation/PathFollowingComponent.cpp
+++ b/Engine/Source/Runtime/AIModule/Private/Navigation/PathFollowingComponent.cpp
@@ -51,17 +51,17 @@ UPathFollowingComponent::UPathFollowingComponent(const FObjectInitializer& Objec
 	Status = EPathFollowingStatus::Idle;
 }
 
-void UPathFollowingComponent::LogPathHelper(const AActor* LogOwner, FNavigationPath* LogPath, const AActor* LogGoalActor)
+void UPathFollowingComponent::LogPathHelper(const AActor* LogOwner, FNavigationPath* InLogPath, const AActor* LogGoalActor)
 {
 #if ENABLE_VISUAL_LOG
 	FVisualLogger& Vlog = FVisualLogger::Get();
 	if (Vlog.IsRecording() &&
-		LogPath && LogPath->IsValid() && LogPath->GetPathPoints().Num())
+		InLogPath && InLogPath->IsValid() && InLogPath->GetPathPoints().Num())
 	{
 		FVisualLogEntry* Entry = Vlog.GetEntryToWrite(LogOwner, LogOwner->GetWorld()->TimeSeconds);
-		LogPath->DescribeSelfToVisLog(Entry);
+		InLogPath->DescribeSelfToVisLog(Entry);
 
-		const FVector PathEnd = LogPath->GetPathPoints().Last().Location;
+		const FVector PathEnd = InLogPath->GetPathPoints().Last().Location;
 		if (LogGoalActor)
 		{
 			const FVector GoalLoc = LogGoalActor->GetActorLocation();
@@ -77,12 +77,12 @@ void UPathFollowingComponent::LogPathHelper(const AActor* LogOwner, FNavigationP
 #endif // ENABLE_VISUAL_LOG
 }
 
-void UPathFollowingComponent::LogPathHelper(const AActor* LogOwner, FNavPathSharedPtr LogPath, const AActor* LogGoalActor)
+void UPathFollowingComponent::LogPathHelper(const AActor* LogOwner, FNavPathSharedPtr InLogPath, const AActor* LogGoalActor)
 {
 #if ENABLE_VISUAL_LOG
-	if (LogPath.IsValid())
+	if (InLogPath.IsValid())
 	{
-		LogPathHelper(LogOwner, LogPath.Get(), LogGoalActor);
+		LogPathHelper(LogOwner, InLogPath.Get(), LogGoalActor);
 	}
 #endif // ENABLE_VISUAL_LOG
 }
diff --git a/Engine/Source/Runtime/Engine/Private/Materials/MaterialInstance.cpp b/Engine/Source/Runtime/Engine/Private/Materials/MaterialInstance.cpp
--- a/Engine/Source/Runtime/Engine/Private/Materials/MaterialInstance.cpp
+++ b/Engine/Source/Runtime/Engine/Private/Materials/MaterialInstance.cpp
@@ -286,22 +286,22 @@ void UMaterialInstance::PropagateDataToMaterialProxy()
 	}
 }
 
-void FMaterialInstanceResource::GameThread_SetParent(UMaterialInterface* InParent)
+void FMaterialInstanceResource::GameThread_SetParent(UMaterialInterface* ParentMaterialInterface)
 {
 	check(IsInGameThread());
 
-	if( GameThreadParent != InParent )
+	if( GameThreadParent != ParentMaterialInterface )
 	{
 		// Set the game thread accessible parent.
 		UMaterialInterface* OldParent = GameThreadParent;
-		GameThreadParent = InParent;
+		GameThreadParent = ParentMaterialInterface;
 
 		// Set the rendering thread's parent and instance pointers.
-		check(InParent != NULL);
+		check(ParentMaterialInterface != NULL);
 		ENQUEUE_UNIQUE_RENDER_COMMAND_TWOPARAMETER(
 			InitMaterialInstanceResource,
 			FMaterialInstanceResource*,Resource,this,
-			UMaterialInterface*,Parent,InParent,
+			UMaterialInterface*,Parent,ParentMaterialInterface,
 		{
 			Resource->Parent = Parent;
 			Resource->InvalidateUniformExpressionCache();
@@ -1538,9 +1538,9 @@ void UMaterialInstance::BeginCacheForCookedPlatformData( const ITargetPlatform *
 		// Cache shaders for each shader format, storing the results in CachedMaterialResourcesForCooking so they will be available during saving
 		for (int32 FormatIndex = 0; FormatIndex < DesiredShaderFormats.Num(); FormatIndex++)
 		{
-			const EShaderPlatform TargetPlatform = ShaderFormatToLegacyShaderPlatform(DesiredShaderFormats[FormatIndex]);
+			const EShaderPlatform TargetShaderPlatform = ShaderFormatToLegacyShaderPlatform(DesiredShaderFormats[FormatIndex]);
 
-			CacheResourceShadersForCooking(TargetPlatform, *CachedMaterialResourcesForPlatform );
+			CacheResourceShadersForCooking(TargetShaderPlatform, *CachedMaterialResourcesForPlatform );
 		}
 	}
 }
@@ -1719,7 +1719,6 @@ void UMaterialInstance::PostLoad()
 		ITargetPlatformManagerModule* TPM = GetTargetPlatformManager();
 		if (TPM && (TPM->RestrictFormatsToRuntimeOnly() == false)) 
 		{
-			ITargetPlatformManagerModule* TPM = GetTargetPlatformManager();
 			TArray<ITargetPlatform*> Platforms = TPM->GetActiveTargetPlatforms();
 			// Cache for all the shader formats that the cooking target requires
 			for (int32 FormatIndex = 0; FormatIndex < Platforms.Num(); FormatIndex++)
@@ -2392,11 +2391,11 @@ void UMaterialInstance::AllMaterialsCacheResourceShadersForRendering()
 }
 
 
-bool UMaterialInstance::IsChildOf(const UMaterialInterface* Parent) const
+bool UMaterialInstance::IsChildOf(const UMaterialInterface* ParentMaterialInterface) const
 {
 	const UMaterialInterface* Material = this;
 
-	while (Material != Parent && Material != nullptr)
+	while (Material != ParentMaterialInterface && Material != nullptr)
 	{
 		const UMaterialInstance* MaterialInstance = Cast<const UMaterialInstance>(Material);
 		Material = (MaterialInstance != nullptr) ? MaterialInstance->Parent : nullptr;
diff --git a/Engine/Source/Runtime/GameplayAbilities/Private/AbilitySystemComponent_Abilities.cpp b/Engine/Source/Runtime/GameplayAbilities/Private/AbilitySystemComponent_Abilities.cpp
--- a/Engine/Source/Runtime/GameplayAbilities/Private/AbilitySystemComponent_Abilities.cpp
+++ b/Engine/Source/Runtime/GameplayAbilities/Private/AbilitySystemComponent_Abilities.cpp
@@ -571,10 +571,10 @@ UGameplayAbility* UAbilitySystemComponent::CreateNewInstanceOfAbility(FGameplayA
 	check(Ability);
 	check(Ability->HasAllFlags(RF_ClassDefaultObject));
 
-	AActor* OwnerActor = GetOwner();
-	check(OwnerActor);
+	AActor* Owner = GetOwner();
+	check(Owner);
 
-	UGameplayAbility * AbilityInstance = NewObject<UGameplayAbility>(OwnerActor, Ability->GetClass());
+	UGameplayAbility * AbilityInstance = NewObject<UGameplayAbility>(Owner, Ability->GetClass());
 	check(AbilityInstance);
 
 	// Add it to one of our instance lists so that it doesn't GC.
diff --git a/Engine/Source/Runtime/GameplayAbilities/Private/GameplayAbilityTypes.cpp b/Engine/Source/Runtime/GameplayAbilities/Private/GameplayAbilityTypes.cpp
--- a/Engine/Source/Runtime/GameplayAbilities/Private/GameplayAbilityTypes.cpp
+++ b/Engine/Source/Runtime/GameplayAbilities/Private/GameplayAbilityTypes.cpp
@@ -55,9 +55,9 @@ void FGameplayAbilityActorInfo::InitFromActor(AActor *InOwnerActor, AActor *InAv
 	}
 }
 
-void FGameplayAbilityActorInfo::SetAvatarActor(AActor *AvatarActor)
+void FGameplayAbilityActorInfo::SetAvatarActor(AActor *InAvatarActor)
 {
-	InitFromActor(OwnerActor.Get(), AvatarActor, AbilitySystemComponent.Get());
+	InitFromActor(OwnerActor.Get(), InAvatarActor, AbilitySystemComponent.Get());
 }
 
 void FGameplayAbilityActorInfo::ClearActorInfo()
@@ -155,9 +155,9 @@ void FGameplayAbilitySpec::PostReplicatedAdd(const struct FGameplayAbilitySpecCo
 	}
 }
 
-void FGameplayAbilitySpecContainer::RegisterWithOwner(UAbilitySystemComponent* Owner)
+void FGameplayAbilitySpecContainer::RegisterWithOwner(UAbilitySystemComponent* InOwner)
 {
-	this->Owner = Owner;
+	Owner = InOwner;
 }
 
 // ----------------------------------------------------
diff --git a/Engine/Source/Runtime/GameplayAbilities/Private/GameplayEffect.cpp b/Engine/Source/Runtime/GameplayAbilities/Private/GameplayEffect.cpp
--- a/Engine/Source/Runtime/GameplayAbilities/Private/GameplayEffect.cpp
+++ b/Engine/Source/Runtime/GameplayAbilities/Private/GameplayEffect.cpp
@@ -421,7 +421,7 @@ FGameplayEffectSpec::FGameplayEffectSpec()
 
 }
 
-FGameplayEffectSpec::FGameplayEffectSpec(const UGameplayEffect* InDef, const FGameplayEffectContextHandle& InEffectContext, float Level)
+FGameplayEffectSpec::FGameplayEffectSpec(const UGameplayEffect* InDef, const FGameplayEffectContextHandle& InEffectContext, float InLevel)
 	: Def(InDef)
 	, StackCount(1)
 	, bCompletedSourceAttributeCapture(false)
@@ -429,7 +429,7 @@ FGameplayEffectSpec::FGameplayEffectSpec(const UGameplayEffect* InDef, const FGa
 	, bDurationLocked(false)
 {
 	check(Def);	
-	SetLevel(Level);
+	SetLevel(InLevel);
 	SetContext(InEffectContext);
 
 	// Init our ModifierSpecs
@@ -447,7 +447,7 @@ FGameplayEffectSpec::FGameplayEffectSpec(const UGameplayEffect* InDef, const FGa
 
 	for (const UGameplayEffect* TargetDef : TargetEffectDefs)
 	{
-		TargetEffectSpecs.Add(FGameplayEffectSpecHandle(new FGameplayEffectSpec(TargetDef, EffectContext, Level)));
+		TargetEffectSpecs.Add(FGameplayEffectSpecHandle(new FGameplayEffectSpec(TargetDef, EffectContext, InLevel)));
 	}
 
 	// Make Granted AbilitySpecs (caller may modify these specs after creating spec, which is why we dont just reference them from the def)
diff --git a/Engine/Source/Runtime/GameplayAbilities/Private/GameplayEffectAggregator.cpp b/Engine/Source/Runtime/GameplayAbilities/Private/GameplayEffectAggregator.cpp
--- a/Engine/Source/Runtime/GameplayAbilities/Private/GameplayEffectAggregator.cpp
+++ b/Engine/Source/Runtime/GameplayAbilities/Private/GameplayEffectAggregator.cpp
@@ -151,11 +151,11 @@ void FAggregator::ExecModOnBaseValue(TEnumAsByte<EGameplayModOp::Type> ModifierO
 	BroadcastOnDirty();
 }
 
-float FAggregator::SumMods(const TArray<FAggregatorMod> &Mods, float Bias, const FAggregatorEvaluateParameters& Parameters) const
+float FAggregator::SumMods(const TArray<FAggregatorMod> &InMods, float Bias, const FAggregatorEvaluateParameters& Parameters) const
 {
 	float Sum = Bias;
 
-	for (const FAggregatorMod& Mod : Mods)
+	for (const FAggregatorMod& Mod : InMods)
 	{
 		if (Mod.Qualifies(Parameters))
 		{
@@ -223,20 +223,20 @@ bool FAggregator::HasPredictedMods() const
 	return NumPredictiveMods > 0;
 }
 
-void FAggregator::RemoveModsWithActiveHandle(TArray<FAggregatorMod>& Mods, FActiveGameplayEffectHandle ActiveHandle)
+void FAggregator::RemoveModsWithActiveHandle(TArray<FAggregatorMod>& InMods, FActiveGameplayEffectHandle ActiveHandle)
 {
 	check(ActiveHandle.IsValid());
 
-	for(int32 idx=Mods.Num()-1; idx >= 0; --idx)
+	for(int32 idx=InMods.Num()-1; idx >= 0; --idx)
 	{
-		if (Mods[idx].ActiveHandle == ActiveHandle)
+		if (InMods[idx].ActiveHandle == ActiveHandle)
 		{
-			if (Mods[idx].IsPredicted)
+			if (InMods[idx].IsPredicted)
 			{
 				NumPredictiveMods--;
 			}
 
-			Mods.RemoveAtSwap(idx, 1, false);
+			InMods.RemoveAtSwap(idx, 1, false);
 		}
 	}
 }
