[UE-11594] Fix a regression that was blocking component attachment in the Blueprint editor.

- was introduced in CL# 2464931 (as part of CL# 2464930)

diff --git a/Engine/Source/Editor/Kismet/Private/SSCSEditor.cpp b/Engine/Source/Editor/Kismet/Private/SSCSEditor.cpp
--- a/Engine/Source/Editor/Kismet/Private/SSCSEditor.cpp
+++ b/Engine/Source/Editor/Kismet/Private/SSCSEditor.cpp
@@ -1673,11 +1673,11 @@ FText SSCS_RowWidget::GetIntroducedInToolTipText() const
 						IntroducedInTooltip = LOCTEXT("IntroducedInNativeError", "Unknown native source (via C++ code)");
 					}
 				}
-				else if (!TreeNode->IsNative() && ComponentTemplate->CreationMethod == EComponentCreationMethod::Native)
+				else if (TreeNode->IsInstanced() && ComponentTemplate->CreationMethod == EComponentCreationMethod::Native && !ComponentTemplate->HasAnyFlags(RF_DefaultSubObject))
 				{
 					IntroducedInTooltip = FText::Format(LOCTEXT("IntroducedInCPPErrorFmt", "{0} (via C++ code)"), FBlueprintEditorUtils::GetFriendlyClassDisplayName(BestClass));
 				}
-				else if (ComponentTemplate->CreationMethod == EComponentCreationMethod::UserConstructionScript)
+				else if (TreeNode->IsInstanced() && ComponentTemplate->CreationMethod == EComponentCreationMethod::UserConstructionScript)
 				{
 					IntroducedInTooltip = FText::Format(LOCTEXT("IntroducedInUCSErrorFmt", "{0} (via an Add Component call)"), FBlueprintEditorUtils::GetFriendlyClassDisplayName(BestClass));
 				}
@@ -2081,12 +2081,12 @@ void SSCS_RowWidget::OnDragEnter( const FGeometry& MyGeometry, const FDragDropEv
 					// Can't attach Static components to mobile ones
 					Message = LOCTEXT("DropActionToolTip_Error_CannotAttachStationary", "Cannot attach Stationary components to movable ones.");
 				}
-				else if ((!NodePtr->IsNative() && HoveredTemplate->CreationMethod == EComponentCreationMethod::Native))
+				else if ((NodePtr->IsInstanced() && HoveredTemplate->CreationMethod == EComponentCreationMethod::Native && !HoveredTemplate->HasAnyFlags(RF_DefaultSubObject)))
 				{
 					// Can't attach to post-construction C++-added components as they exist outside of the CDO and are not known at SCS execution time
 					Message = LOCTEXT("DropActionToolTip_Error_CannotAttachCPPAdded", "Cannot attach to components added in post-construction C++ code.");
 				}
-				else if (HoveredTemplate->CreationMethod == EComponentCreationMethod::UserConstructionScript)
+				else if (NodePtr->IsInstanced() && HoveredTemplate->CreationMethod == EComponentCreationMethod::UserConstructionScript)
 				{
 					// Can't attach to UCS-added components as they exist outside of the CDO and are not known at SCS execution time
 					Message = LOCTEXT("DropActionToolTip_Error_CannotAttachUCSAdded", "Cannot attach to components added in the Construction Script.");
