[UE-8552] The Blueprint context menu has a new pop-out sub-menu that gives the user more control over the menu's target context (which now includes the current blueprint by default).

diff --git a/Engine/Source/Editor/BlueprintGraph/Private/BlueprintEditorSettings.cpp b/Engine/Source/Editor/BlueprintGraph/Private/BlueprintEditorSettings.cpp
--- a/Engine/Source/Editor/BlueprintGraph/Private/BlueprintEditorSettings.cpp
+++ b/Engine/Source/Editor/BlueprintGraph/Private/BlueprintEditorSettings.cpp
@@ -11,7 +11,7 @@ UBlueprintEditorSettings::UBlueprintEditorSettings(const FObjectInitializer& Obj
 	, bDrawMidpointArrowsInBlueprints(false)
 	, bShowGraphInstructionText(true)
 	// Workflow Settings
-	, bUseTargetContextForNodeMenu(true)
+	, bSplitContextTargetSettings(true)
 	, bExposeAllMemberComponentFunctions(true)
 	, bShowContextualFavorites(false)
 	, bFlattenFavoritesMenus(true)
diff --git a/Engine/Source/Editor/BlueprintGraph/Public/BlueprintEditorSettings.h b/Engine/Source/Editor/BlueprintGraph/Public/BlueprintEditorSettings.h
--- a/Engine/Source/Editor/BlueprintGraph/Public/BlueprintEditorSettings.h
+++ b/Engine/Source/Editor/BlueprintGraph/Public/BlueprintEditorSettings.h
@@ -31,9 +31,9 @@ public:
 
 // Workflow Settings
 public:
-	/** If enabled, will use the blueprint's (or output pin's) class to narrow down context menu results. */
-	UPROPERTY(EditAnywhere, config, Category=Workflow, meta=(DisplayName="Context Menu: Use Target Context"))
-	bool bUseTargetContextForNodeMenu;
+	/** If enabled, we'll save off your chosen target setting based off of the context (allowing you to have different preferences based off what you're operating on). */
+	UPROPERTY(EditAnywhere, config, Category=Workflow, meta=(DisplayName="Context Menu: Divide Context Target Preferences"))
+	bool bSplitContextTargetSettings;
 
 	/** If enabled, then ALL component functions are exposed to the context menu (when the contextual target is a component owner). Ignores "ExposeFunctionCategories" metadata for components. */
 	UPROPERTY(EditAnywhere, config, Category=Workflow, meta=(DisplayName="Context Menu: Expose All Sub-Component Functions"))
diff --git a/Engine/Source/Editor/EditorStyle/Private/SlateEditorStyle.cpp b/Engine/Source/Editor/EditorStyle/Private/SlateEditorStyle.cpp
--- a/Engine/Source/Editor/EditorStyle/Private/SlateEditorStyle.cpp
+++ b/Engine/Source/Editor/EditorStyle/Private/SlateEditorStyle.cpp
@@ -4839,6 +4839,20 @@ void FSlateEditorStyle::FStyle::SetupPersonaStyle()
 		Set( "BlueprintEditor.EnableSimulation.Small", new IMAGE_BRUSH( "Icons/icon_Enable_Simulation_40px", Icon20x20));
 		Set( "SCS.NativeComponent", new IMAGE_BRUSH( "Icons/NativeSCSComponent", Icon20x20 ));
 		Set( "SCS.Component", new IMAGE_BRUSH( "Icons/SCSComponent", Icon20x20 ));
+
+		const FButtonStyle BlueprintContextTargetsButtonStyle = FButtonStyle()
+			.SetNormal(IMAGE_BRUSH("Common/TreeArrow_Collapsed_Hovered", Icon10x10, FLinearColor(0.2f, 0.2f, 0.2f, 1.f)))
+			.SetHovered(IMAGE_BRUSH("Common/TreeArrow_Collapsed_Hovered", Icon10x10, FLinearColor(0.4f, 0.4f, 0.4f, 1.f)))
+			.SetPressed(IMAGE_BRUSH("Common/TreeArrow_Collapsed_Hovered", Icon10x10, FLinearColor(0.2f, 0.2f, 0.2f, 1.f)));
+// 		const FCheckBoxStyle BlueprintContextTargetsButtonStyle = FCheckBoxStyle()
+// 			.SetCheckBoxType(ESlateCheckBoxType::CheckBox)
+// 			.SetUncheckedImage(IMAGE_BRUSH("Common/TreeArrow_Collapsed", Icon10x10, FLinearColor(0.2f, 0.2f, 0.2f, 1.f)))
+// 			.SetUncheckedHoveredImage(IMAGE_BRUSH("Common/TreeArrow_Collapsed", Icon10x10, FLinearColor(0.4f, 0.4f, 0.4f, 1.f)))
+// 			.SetUncheckedPressedImage(IMAGE_BRUSH("Common/TreeArrow_Collapsed", Icon10x10, FLinearColor(0.2f, 0.2f, 0.2f, 1.f)))
+// 			.SetCheckedImage(IMAGE_BRUSH("Common/TreeArrow_Collapsed_Hovered", Icon10x10, FLinearColor(0.2f, 0.2f, 0.2f, 1.f)))
+// 			.SetCheckedHoveredImage(IMAGE_BRUSH("Common/TreeArrow_Collapsed_Hovered", Icon10x10, FLinearColor(0.4f, 0.4f, 0.4f, 1.f)))
+// 			.SetCheckedPressedImage(IMAGE_BRUSH("Common/TreeArrow_Collapsed_Hovered", Icon10x10, FLinearColor(0.2f, 0.2f, 0.2f, 1.f)));
+		Set("BlueprintEditor.ContextMenu.TargetsButton", BlueprintContextTargetsButtonStyle);
 	}
 
 	// Kismet linear expression display
diff --git a/Engine/Source/Editor/Kismet/Private/BlueprintActionMenuUtils.cpp b/Engine/Source/Editor/Kismet/Private/BlueprintActionMenuUtils.cpp
--- a/Engine/Source/Editor/Kismet/Private/BlueprintActionMenuUtils.cpp
+++ b/Engine/Source/Editor/Kismet/Private/BlueprintActionMenuUtils.cpp
@@ -102,10 +102,11 @@ namespace BlueprintActionMenuUtilsImpl
 	/**
 	 * 
 	 * 
-	 * @param  MainMenuFilter	
+	 * @param  MainMenuFilter    
+	 * @param  ContextTargetMask    
 	 * @return 
 	 */
-	static FBlueprintActionFilter MakeCallOnMemberFilter(FBlueprintActionFilter const& MainMenuFilter);
+	static FBlueprintActionFilter MakeCallOnMemberFilter(FBlueprintActionFilter const& MainMenuFilter, uint32 ContextTargetMask);
 
 	/**
 	 * 
@@ -253,11 +254,22 @@ static UClass* BlueprintActionMenuUtilsImpl::GetPinClassType(UEdGraphPin const*
 		}
 	}
 
+	if (PinObjClass != nullptr)
+	{
+		if (UBlueprint* ClassBlueprint = Cast<UBlueprint>(PinObjClass->ClassGeneratedBy))
+		{
+			if (ClassBlueprint->SkeletonGeneratedClass != nullptr)
+			{
+				PinObjClass = ClassBlueprint->SkeletonGeneratedClass;
+			}
+		}
+	}
+
 	return PinObjClass;
 }
 
 //------------------------------------------------------------------------------
-static FBlueprintActionFilter BlueprintActionMenuUtilsImpl::MakeCallOnMemberFilter(FBlueprintActionFilter const& MainMenuFilter)
+static FBlueprintActionFilter BlueprintActionMenuUtilsImpl::MakeCallOnMemberFilter(FBlueprintActionFilter const& MainMenuFilter, uint32 ContextTargetMask)
 {
 	FBlueprintActionFilter CallOnMemberFilter;
 	CallOnMemberFilter.Context = MainMenuFilter.Context;
@@ -277,18 +289,36 @@ static FBlueprintActionFilter BlueprintActionMenuUtilsImpl::MakeCallOnMemberFilt
 		CallOnMemberFilter.AddRejectionTest(FBlueprintActionFilter::FRejectionTestDelegate::CreateStatic(IsUnexposedMemberAction));
 	}
 
-	for (UClass const* TargetClass : MainMenuFilter.TargetClasses)
+	bool bForceAddComponents = ((ContextTargetMask & EContextTargetFlags::TARGET_SubComponents) != 0);
+
+	TArray<UClass*> TargetClasses = MainMenuFilter.TargetClasses;
+	if (bForceAddComponents && (CallOnMemberFilter.TargetClasses.Num() == 0))
+	{
+		for (UBlueprint const* TargetBlueprint : MainMenuFilter.Context.Blueprints)
+		{
+			UClass* BpClass = TargetBlueprint->SkeletonGeneratedClass;
+			if (BpClass != nullptr)
+			{
+				TargetClasses.Add(BpClass);
+			}
+		}
+	}
+
+	for (UClass const* TargetClass : TargetClasses)
 	{
 		for (TFieldIterator<UObjectProperty> PropertyIt(TargetClass, EFieldIteratorFlags::IncludeSuper); PropertyIt; ++PropertyIt)
 		{
 			UObjectProperty* ObjectProperty = *PropertyIt;
-
-			if (!ObjectProperty->HasAnyPropertyFlags(CPF_BlueprintVisible) ||
-				!(ObjectProperty->PropertyClass->IsChildOf<UActorComponent>() || ObjectProperty->HasMetaData(FBlueprintMetadata::MD_ExposeFunctionCategories)))
+			if (!ObjectProperty->HasAnyPropertyFlags(CPF_BlueprintVisible))
 			{
 				continue;
 			}
-			CallOnMemberFilter.Context.SelectedObjects.Add(ObjectProperty);
+
+			if ( ObjectProperty->HasMetaData(FBlueprintMetadata::MD_ExposeFunctionCategories) || 
+				(bForceAddComponents && ObjectProperty->PropertyClass->IsChildOf<UActorComponent>()) )
+			{
+				CallOnMemberFilter.Context.SelectedObjects.Add(ObjectProperty);
+			}			
 		}
 	}
 
@@ -401,7 +431,7 @@ void FBlueprintActionMenuUtils::MakePaletteMenu(FBlueprintActionContext const& C
 }
 
 //------------------------------------------------------------------------------
-void FBlueprintActionMenuUtils::MakeContextMenu(FBlueprintActionContext const& Context, bool bIsContextSensitive, FBlueprintActionMenuBuilder& MenuOut)
+void FBlueprintActionMenuUtils::MakeContextMenu(FBlueprintActionContext const& Context, bool bIsContextSensitive, uint32 ClassTargetMask, FBlueprintActionMenuBuilder& MenuOut)
 {
 	using namespace BlueprintActionMenuUtilsImpl;
 
@@ -430,7 +460,6 @@ void FBlueprintActionMenuUtils::MakeContextMenu(FBlueprintActionContext const& C
 	LevelActorsFilter.AddRejectionTest(FBlueprintActionFilter::FRejectionTestDelegate::CreateStatic(IsUnBoundSpawner));
 
 	const UBlueprintEditorSettings* BlueprintSettings = GetDefault<UBlueprintEditorSettings>();
-	bool const bAddTargetContext  = bIsContextSensitive && BlueprintSettings->bUseTargetContextForNodeMenu;
 	bool bCanOperateOnLevelActors = bIsContextSensitive && (Context.Pins.Num() == 0);
 	bool bCanHaveActorComponents  = bIsContextSensitive;
 	// determine if we can operate on certain object selections (level actors, 
@@ -441,7 +470,7 @@ void FBlueprintActionMenuUtils::MakeContextMenu(FBlueprintActionContext const& C
 		if (BlueprintClass != nullptr)
 		{
 			bCanOperateOnLevelActors &= BlueprintClass->IsChildOf<ALevelScriptActor>();
-			if (bAddTargetContext)
+			if (bIsContextSensitive && (ClassTargetMask & EContextTargetFlags::TARGET_Blueprint))
 			{
 				MainMenuFilter.TargetClasses.Add(BlueprintClass);
 			}
@@ -520,34 +549,41 @@ void FBlueprintActionMenuUtils::MakeContextMenu(FBlueprintActionContext const& C
 		}
 	}
 
-	if (bAddTargetContext)
+	if (bIsContextSensitive)
 	{
-		bool bContextPinIsObj = false;
-
 		// if we're dragging from a pin, we further extend the context to cover
 		// that pin and any other pins it sits beside 
 		for (UEdGraphPin* ContextPin : Context.Pins)
 		{
-			// we only want the pin to be the target class when it is an output
-			// (doesn't make sense to get members to plug into their parent)
-			if (ContextPin->Direction == EGPD_Input)
+			if (UClass* PinObjClass = GetPinClassType(ContextPin))
 			{
-				continue;
+				if (ClassTargetMask & EContextTargetFlags::TARGET_PinObject)
+				{
+					MainMenuFilter.TargetClasses.Add(PinObjClass);
+				}
 			}
 
-			if (UClass* PinObjClass = GetPinClassType(ContextPin))
+			UEdGraphNode* OwningNode = ContextPin->GetOwningNodeUnchecked();
+			if ((OwningNode != nullptr) && (ClassTargetMask & EContextTargetFlags::TARGET_NodeTarget))
 			{
-				if (!bContextPinIsObj)
+				// @TODO: should we search instead by name/DefaultToSelf
+				if (UEdGraphPin* TargetPin = K2Schema->FindSelfPin(*OwningNode, EGPD_Input))
 				{
-					MainMenuFilter.TargetClasses.Empty();
+					if (UClass* TargetClass = GetPinClassType(TargetPin))
+					{
+						MainMenuFilter.TargetClasses.Add(TargetClass);
+					}
 				}
-				MainMenuFilter.TargetClasses.Add(PinObjClass);
-				bContextPinIsObj = true;
+			}
+
+			if ((ClassTargetMask & EContextTargetFlags::TARGET_SiblingPinObjects) == 0)
+			{
+				continue;
 			}
 
 			for (UEdGraphPin* NodePin : ContextPin->GetOwningNode()->Pins)
 			{
-				if ((NodePin->Direction == ContextPin->Direction) && !bContextPinIsObj)
+				if ((NodePin->Direction == EGPD_Output))
 				{
 					if (UClass* PinClass = GetPinClassType(NodePin))
 					{
@@ -559,7 +595,7 @@ void FBlueprintActionMenuUtils::MakeContextMenu(FBlueprintActionContext const& C
 	}
 
 	// should be called AFTER the MainMenuFilter if fully constructed
-	FBlueprintActionFilter CallOnMemberFilter = MakeCallOnMemberFilter(MainMenuFilter);
+	FBlueprintActionFilter CallOnMemberFilter = MakeCallOnMemberFilter(MainMenuFilter, ClassTargetMask);
 
 	FBlueprintActionFilter AddComponentFilter;
 	AddComponentFilter.Context = MainMenuFilter.Context;
diff --git a/Engine/Source/Editor/Kismet/Private/SBlueprintActionMenu.cpp b/Engine/Source/Editor/Kismet/Private/SBlueprintActionMenu.cpp
--- a/Engine/Source/Editor/Kismet/Private/SBlueprintActionMenu.cpp
+++ b/Engine/Source/Editor/Kismet/Private/SBlueprintActionMenu.cpp
@@ -16,6 +16,7 @@
 #include "BlueprintPaletteFavorites.h"
 #include "IDocumentation.h"
 #include "SSCSEditor.h"
+#include "SBlueprintContextTargetMenu.h"
 
 #define LOCTEXT_NAMESPACE "SBlueprintGraphContextMenu"
 
@@ -230,6 +231,21 @@ void SBlueprintActionMenu::Construct( const FArguments& InArgs, TSharedPtr<FBlue
 		}
 	}
 
+	FBlueprintActionContext MenuContext;
+	ConstructActionContext(MenuContext);
+
+	TSharedPtr<SComboButton> TargetContextSubMenuButton;
+	// @TODO: would be nice if we could use a checkbox style for this, and have a different state for open/closed
+	SAssignNew(TargetContextSubMenuButton, SComboButton)
+		.MenuPlacement(MenuPlacement_MenuRight)
+		.HasDownArrow(false)
+		.ButtonStyle(FEditorStyle::Get(), "BlueprintEditor.ContextMenu.TargetsButton")
+		.MenuContent()
+		[
+			SAssignNew(ContextTargetSubMenu, SBlueprintContextTargetMenu, MenuContext)
+				.OnTargetMaskChanged(this, &SBlueprintActionMenu::OnContextTargetsChanged)
+		];
+
 	// Build the widget layout
 	SBorder::Construct( SBorder::FArguments()
 		.BorderImage( FEditorStyle::GetBrush("Menu.Background") )
@@ -295,16 +311,25 @@ void SBlueprintActionMenu::Construct( const FArguments& InArgs, TSharedPtr<FBlue
 							.Text(LOCTEXT("BlueprintActionMenuContextToggle", "Context Sensitive"))
 						]
 					]
+
+					+SHorizontalBox::Slot()
+					.HAlign(HAlign_Right)
+					.VAlign(VAlign_Center)
+					.AutoWidth()
+					.Padding(3.f, 0.f, 0.f, 0.f)
+					[
+						TargetContextSubMenuButton.ToSharedRef()
+					]
 				]
 
 				// ACTION LIST 
 				+SVerticalBox::Slot()
 				[
 					SAssignNew(GraphActionMenu, SGraphActionMenu)
-					.OnActionSelected(this, &SBlueprintActionMenu::OnActionSelected)
-					.OnCreateWidgetForAction( SGraphActionMenu::FOnCreateWidgetForAction::CreateSP(this, &SBlueprintActionMenu::OnCreateWidgetForAction) )
-					.OnCollectAllActions(this, &SBlueprintActionMenu::CollectAllActions)
-					.OnCreateCustomRowExpander_Static(&CreateCustomBlueprintActionExpander)
+						.OnActionSelected(this, &SBlueprintActionMenu::OnActionSelected)
+						.OnCreateWidgetForAction(SGraphActionMenu::FOnCreateWidgetForAction::CreateSP(this, &SBlueprintActionMenu::OnCreateWidgetForAction))
+						.OnCollectAllActions(this, &SBlueprintActionMenu::CollectAllActions)
+						.OnCreateCustomRowExpander_Static(&CreateCustomBlueprintActionExpander)
 				]
 			]
 		]
@@ -381,60 +406,79 @@ void SBlueprintActionMenu::OnContextToggleChanged(ECheckBoxState CheckState)
 	GraphActionMenu->RefreshAllActions(true, false);
 }
 
+void SBlueprintActionMenu::OnContextTargetsChanged(uint32 /*ContextTargetMask*/)
+{
+	GraphActionMenu->RefreshAllActions(/*bPreserveExpansion =*/true, /*bHandleOnSelectionEvent =*/false);
+}
+
 ECheckBoxState SBlueprintActionMenu::ContextToggleIsChecked() const
 {
 	return EditorPtr.Pin()->GetIsContextSensitive() ? ECheckBoxState::Checked : ECheckBoxState::Unchecked;
 }
 
 void SBlueprintActionMenu::CollectAllActions(FGraphActionListBuilderBase& OutAllActions)
-{
-	UBlueprint* Blueprint = EditorPtr.Pin()->GetBlueprintObj();
-	
+{	
 	check(EditorPtr.IsValid());
 	TSharedPtr<FBlueprintEditor> BlueprintEditor = EditorPtr.Pin();
 	bool const bIsContextSensitive = BlueprintEditor->GetIsContextSensitive();
-	
+
+	uint32 ContextTargetMask = 0;
+	if (bIsContextSensitive && ContextTargetSubMenu.IsValid())
+	{
+		ContextTargetMask = ContextTargetSubMenu->GetContextTargetMask();
+	}
+
 	FBlueprintActionContext FilterContext;
-	FilterContext.Blueprints.Add(Blueprint);
+	ConstructActionContext(FilterContext);
+	
+	FBlueprintActionMenuBuilder MenuBuilder(EditorPtr);
+	// NOTE: cannot call GetGraphContextActions() during serialization and GC due to its use of FindObject()
+	if(!GIsSavingPackage && !GIsGarbageCollecting)
+	{
+		FBlueprintActionMenuUtils::MakeContextMenu(FilterContext, bIsContextSensitive, ContextTargetMask, MenuBuilder);
+	}
+	// copy the added options back to the main list
+	OutAllActions.Append(MenuBuilder); // @TODO: Avoid this copy
+	// also try adding promote to variable if we can do so.
+	TryInsertPromoteToVariable(FilterContext, OutAllActions);
+}
+
+void SBlueprintActionMenu::ConstructActionContext(FBlueprintActionContext& ContextDescOut)
+{
+	check(EditorPtr.IsValid());
+	TSharedPtr<FBlueprintEditor> BlueprintEditor = EditorPtr.Pin();
+	bool const bIsContextSensitive = BlueprintEditor->GetIsContextSensitive();
+
+	UBlueprint* Blueprint = BlueprintEditor->GetBlueprintObj();
+	ContextDescOut.Blueprints.Add(Blueprint);
 	// we still want context from the graph (even if the user has unchecked
 	// "Context Sensitive"), otherwise the user would be presented with nodes
 	// that can't be placed in the graph... if the user isn't being presented
 	// with a valid node, then fix it up in filtering
-	FilterContext.Graphs.Add(GraphObj);
-	
+	ContextDescOut.Graphs.Add(GraphObj);
+
 	if (bIsContextSensitive)
 	{
-		FilterContext.Pins = DraggedFromPins;
-		
+		ContextDescOut.Pins = DraggedFromPins;
+
 		// Get selection from the "My Blueprint" view.
 		FEdGraphSchemaAction_K2Var* SelectedVar = BlueprintEditor->GetMyBlueprintWidget()->SelectionAsVar();
 		if ((SelectedVar != nullptr) && (SelectedVar->GetProperty() != nullptr))
 		{
-			FilterContext.SelectedObjects.Add(SelectedVar->GetProperty());
+			ContextDescOut.SelectedObjects.Add(SelectedVar->GetProperty());
 		}
 		// If the selection come from the SCS editor, add it to the filter context.
-		else if ( Blueprint->SkeletonGeneratedClass && BlueprintEditor->GetSCSEditor().IsValid() )
+		else if (Blueprint->SkeletonGeneratedClass && BlueprintEditor->GetSCSEditor().IsValid())
 		{
 			TArray<FSCSEditorTreeNodePtrType> Nodes = BlueprintEditor->GetSCSEditor()->GetSelectedNodes();
-			if ( Nodes.Num() == 1 && Nodes[0]->GetNodeType() == FSCSEditorTreeNode::ComponentNode )
+			if (Nodes.Num() == 1 && Nodes[0]->GetNodeType() == FSCSEditorTreeNode::ComponentNode)
 			{
 				FName PropertyName = Nodes[0]->GetVariableName();
 				UObjectProperty* VariableProperty = FindField<UObjectProperty>(Blueprint->SkeletonGeneratedClass, PropertyName);
-				FilterContext.SelectedObjects.Add(VariableProperty);
+				ContextDescOut.SelectedObjects.Add(VariableProperty);
 			}
 		}
 	}
-	
-	FBlueprintActionMenuBuilder MenuBuilder(EditorPtr);
-	// NOTE: cannot call GetGraphContextActions() during serialization and GC due to its use of FindObject()
-	if(!GIsSavingPackage && !GIsGarbageCollecting)
-	{
-		FBlueprintActionMenuUtils::MakeContextMenu(FilterContext, bIsContextSensitive, MenuBuilder);
-	}
-	// copy the added options back to the main list
-	OutAllActions.Append(MenuBuilder); // @TODO: Avoid this copy
-	// also try adding promote to variable if we can do so.
-	TryInsertPromoteToVariable(FilterContext, OutAllActions);
 }
 
 TSharedRef<SEditableTextBox> SBlueprintActionMenu::GetFilterTextBox()
diff --git a/Engine/Source/Editor/Kismet/Private/SBlueprintActionMenu.h b/Engine/Source/Editor/Kismet/Private/SBlueprintActionMenu.h
--- a/Engine/Source/Editor/Kismet/Private/SBlueprintActionMenu.h
+++ b/Engine/Source/Editor/Kismet/Private/SBlueprintActionMenu.h
@@ -5,6 +5,7 @@
 
 class  SGraphActionMenu;
 struct FBlueprintActionContext;
+class  SBlueprintContextTargetMenu;
 
 /*******************************************************************************
 * SBlueprintActionMenu
@@ -42,6 +43,7 @@ protected:
 	FText GetSearchContextDesc() const;
 	void OnContextToggleChanged(ECheckBoxState CheckState);
 	ECheckBoxState ContextToggleIsChecked() const;
+	void OnContextTargetsChanged(uint32 ContextTargetMask);
 
 	void OnActionSelected( const TArray< TSharedPtr<FEdGraphSchemaAction> >& SelectedAction );
 
@@ -50,6 +52,9 @@ protected:
 	/** Callback used to populate all actions list in SGraphActionMenu */
 	void CollectAllActions(FGraphActionListBuilderBase& OutAllActions);
 
+	/**  */
+	void ConstructActionContext(FBlueprintActionContext& ContextDescOut);
+
 	/** Functioin to try to insert a promote to variable entry if it is possible to do so. */
 	void TryInsertPromoteToVariable(FBlueprintActionContext const& Context, FGraphActionListBuilderBase& OutAllActions);
 
@@ -64,6 +69,7 @@ private:
 
 	TSharedPtr<SGraphActionMenu> GraphActionMenu;
 	TWeakPtr<FBlueprintEditor> EditorPtr;
+	TSharedPtr<SBlueprintContextTargetMenu> ContextTargetSubMenu;
 
 	bool bActionExecuted;
 };
diff --git a/Engine/Source/Editor/Kismet/Private/SBlueprintContextTargetMenu.cpp b/Engine/Source/Editor/Kismet/Private/SBlueprintContextTargetMenu.cpp
new file mode 100644
--- /dev/null
+++ b/Engine/Source/Editor/Kismet/Private/SBlueprintContextTargetMenu.cpp
@@ -0,0 +1,517 @@
+// Copyright 1998-2015 Epic Games, Inc. All Rights Reserved.
+
+#include "BlueprintEditorPrivatePCH.h"
+#include "SBlueprintContextTargetMenu.h"
+#include "BlueprintActionMenuUtils.h"	// for EContextTargetFlags
+#include "BlueprintActionFilter.h"		// for FBlueprintActionContext
+#include "BlueprintEditorUtils.h"		// for DoesSupportComponents
+#include "BlueprintEditorSettings.h"
+
+#define LOCTEXT_NAMESPACE "SBlueprintContextTargetMenu"
+
+/*******************************************************************************
+ * FContextMenuTargetProfile
+ ******************************************************************************/
+ 
+namespace ContextMenuTargetProfileImpl
+{
+	const FString ConfigSection("BlueprintEditor.ContextTargets");
+	const FString SharedProfileName("SharedMenuProfile");
+
+	/**
+	 * Query methods to check and see if the user is dragging off a specific pin
+	 * type (for the context menu).
+	 * 
+	 * @param  MenuContext    Defines the context from which the menu was requested.
+	 * @return True if the context contains a pin matching the type in question, otherwise false.
+	 */
+	static bool HasExecPinContext(const FBlueprintActionContext& MenuContext);
+	static bool HasObjectPinContext(const FBlueprintActionContext& MenuContext);
+
+	static bool IsObjectPin(UEdGraphPin* Pin);
+	static UClass* GetPinClass(UEdGraphPin const* Pin);
+	static bool HasAnyExposedComponents(UClass* TargetClass);
+
+	/**
+	 * Returns a user facing description that provides a short succinct title,
+	 * discerning this profile from others.
+	 */
+	static FText GetProfileDescription(const FBlueprintActionContext& MenuContext);
+
+	/**
+	 * Determines which target options are unused, given the current context.
+	 */
+	static uint32 GetIncompatibleTargetFlags(const FBlueprintActionContext& MenuContext, uint32& HasComponentsMaskOut);
+
+	/**
+	 * Controls how user set targets are saved. Determines how to split context 
+	 * scenarios by providing a separate save name for each.
+	 */
+	static FString GetProfileSaveName(const FBlueprintActionContext& MenuContext);
+}
+
+//------------------------------------------------------------------------------
+static bool ContextMenuTargetProfileImpl::HasObjectPinContext(const FBlueprintActionContext& MenuContext)
+{
+	for (UEdGraphPin* ContextPin : MenuContext.Pins)
+	{
+		if (IsObjectPin(ContextPin))
+		{
+			return true;
+		}
+	}
+	return false;
+}
+
+//------------------------------------------------------------------------------
+static bool ContextMenuTargetProfileImpl::HasExecPinContext(const FBlueprintActionContext& MenuContext)
+{
+	const UEdGraphSchema_K2* K2Schema = GetDefault<UEdGraphSchema_K2>();
+	for (UEdGraphPin* ContextPin : MenuContext.Pins)
+	{
+		if (K2Schema->IsExecPin(*ContextPin))
+		{
+			return true;
+		}
+	}
+	return false;
+}
+
+//------------------------------------------------------------------------------
+static bool ContextMenuTargetProfileImpl::IsObjectPin(UEdGraphPin* Pin)
+{
+	return GetPinClass(Pin) != nullptr;
+}
+
+//------------------------------------------------------------------------------
+static UClass* ContextMenuTargetProfileImpl::GetPinClass(UEdGraphPin const* Pin)
+{
+	UClass* PinObjClass = nullptr;
+
+	FEdGraphPinType const& PinType = Pin->PinType;
+	if ((PinType.PinCategory == UEdGraphSchema_K2::PC_Object) ||
+		(PinType.PinCategory == UEdGraphSchema_K2::PC_Interface))
+	{
+		bool const bIsSelfPin = !PinType.PinSubCategoryObject.IsValid();
+		if (bIsSelfPin)
+		{
+			PinObjClass = CastChecked<UK2Node>(Pin->GetOwningNode())->GetBlueprint()->SkeletonGeneratedClass;
+		}
+		else
+		{
+			PinObjClass = Cast<UClass>(PinType.PinSubCategoryObject.Get());
+			if (PinObjClass != nullptr)
+			{
+				if (UBlueprint* ClassBlueprint = Cast<UBlueprint>(PinObjClass->ClassGeneratedBy))
+				{
+					if (ClassBlueprint->SkeletonGeneratedClass != nullptr)
+					{
+						PinObjClass = ClassBlueprint->SkeletonGeneratedClass;
+					}
+				}
+			}
+		}
+	}
+
+	return PinObjClass;
+}
+
+//------------------------------------------------------------------------------
+static bool ContextMenuTargetProfileImpl::HasAnyExposedComponents(UClass* TargetClass)
+{
+	for (TFieldIterator<UObjectProperty> PropertyIt(TargetClass, EFieldIteratorFlags::IncludeSuper); PropertyIt; ++PropertyIt)
+	{
+		UObjectProperty* ObjectProperty = *PropertyIt;
+		if (!ObjectProperty->HasAnyPropertyFlags(CPF_BlueprintVisible) || !ObjectProperty->PropertyClass->IsChildOf<UActorComponent>())
+		{
+			continue;
+		}
+
+		return true;
+	}
+	return false;
+}
+
+//------------------------------------------------------------------------------
+FText ContextMenuTargetProfileImpl::GetProfileDescription(const FBlueprintActionContext& MenuContext)
+{
+	if (MenuContext.Pins.Num() == 0)
+	{
+		return LOCTEXT("BlueprintContextTarget", "Blueprint Context Target(s)");
+	}
+	else if (ContextMenuTargetProfileImpl::HasObjectPinContext(MenuContext))
+	{
+		return LOCTEXT("ObjectContextTarget", "Object Context Target(s)");
+	}
+	else if (ContextMenuTargetProfileImpl::HasExecPinContext(MenuContext))
+	{
+		return LOCTEXT("ExecContextTarget", "Exec Context Target(s)");
+	}
+	else
+	{
+// 		const FEdGraphPinType& PinType = MenuContext.Pins[0]->PinType;
+// 		FText PinTypename = UEdGraphSchema_K2::TypeToText(PinType);
+// 		return FText::Format(LOCTEXT("TypedPinContextTarget", "{0} Context Target(s)"), PinTypename);
+
+		return LOCTEXT("PODPinContextTarget", "Pin Context Target(s)");
+	}
+}
+
+//------------------------------------------------------------------------------
+uint32 ContextMenuTargetProfileImpl::GetIncompatibleTargetFlags(const FBlueprintActionContext& MenuContext, uint32& HasComponentsMaskOut)
+{
+	uint32 IncompatibleFlags = (EContextTargetFlags::TARGET_NodeTarget | EContextTargetFlags::TARGET_PinObject | 
+		EContextTargetFlags::TARGET_SiblingPinObjects | EContextTargetFlags::TARGET_SubComponents);
+
+	TArray<UClass*> BlueprintClasses;
+
+	HasComponentsMaskOut = 0;
+	for (UBlueprint* Blueprint : MenuContext.Blueprints)
+	{
+		UClass* BpClass = Blueprint->SkeletonGeneratedClass;
+		BlueprintClasses.Add(BpClass);
+
+		if (HasAnyExposedComponents(BpClass))
+		{
+			HasComponentsMaskOut |= EContextTargetFlags::TARGET_Blueprint;
+		}
+	}
+
+	const UEdGraphSchema_K2* K2Schema = GetDefault<UEdGraphSchema_K2>();
+	for (UEdGraphPin* ContextPin : MenuContext.Pins)
+	{
+		UEdGraphNode* OwnerNode = ContextPin->GetOwningNodeUnchecked();
+		if (OwnerNode != nullptr)
+		{
+			UEdGraphPin* SelfPin = K2Schema->FindSelfPin(*OwnerNode, EGPD_Input);
+			UClass* SelfClass = (SelfPin != nullptr) ? GetPinClass(SelfPin) : nullptr;
+			// if this node has a hidden self pin, then that probably means its 
+			// target is a static lib, or the same as the current blueprint
+			bool const bHasUniqueTarget = (SelfPin != nullptr) && !SelfPin->bHidden;// && !BlueprintClasses.Contains(SelfClass);
+			// @TODO: check that "self" is different from this blueprint
+			if (bHasUniqueTarget)
+			{
+				IncompatibleFlags &= ~EContextTargetFlags::TARGET_NodeTarget;
+				if (HasAnyExposedComponents(SelfClass))
+				{
+					HasComponentsMaskOut |= EContextTargetFlags::TARGET_NodeTarget;
+				}
+			}
+
+			for (UEdGraphPin* SiblingPin : OwnerNode->Pins)
+			{
+				UClass* PinClass = GetPinClass(SiblingPin);
+				if ((SiblingPin != ContextPin) && (SiblingPin->Direction == EGPD_Output) && (PinClass != nullptr))
+				{
+					IncompatibleFlags &= ~EContextTargetFlags::TARGET_SiblingPinObjects;
+					if (HasAnyExposedComponents(PinClass))
+					{
+						HasComponentsMaskOut |= EContextTargetFlags::TARGET_SiblingPinObjects;
+					}
+				}
+			}
+		}
+
+		if (ContextPin->Direction == EGPD_Input)
+		{
+			continue;
+		}
+
+		if (UClass* PinClass = GetPinClass(ContextPin))
+		{
+			IncompatibleFlags &= ~EContextTargetFlags::TARGET_PinObject;
+			if (HasAnyExposedComponents(PinClass))
+			{
+				HasComponentsMaskOut |= EContextTargetFlags::TARGET_PinObject;
+			}
+		}
+	}
+
+	if (HasComponentsMaskOut != 0)
+	{
+		IncompatibleFlags &= ~EContextTargetFlags::TARGET_SubComponents;
+	}
+
+	return IncompatibleFlags;
+}
+
+
+//------------------------------------------------------------------------------
+static FString ContextMenuTargetProfileImpl::GetProfileSaveName(const FBlueprintActionContext& MenuContext)
+{
+	const UBlueprintEditorSettings* BpSettings = GetDefault<UBlueprintEditorSettings>();
+	if (!BpSettings->bSplitContextTargetSettings)
+	{
+		return ContextMenuTargetProfileImpl::SharedProfileName;
+	}
+
+	if (MenuContext.Pins.Num() == 0)
+	{
+		return TEXT("NoPinProfile");
+	}
+	else if (HasExecPinContext(MenuContext))
+	{
+		return TEXT("ExecPinProfile");
+	}
+
+	uint32 HasComponentMask = 0;
+	uint32 IncompatibleTargetFlags = GetIncompatibleTargetFlags(MenuContext, HasComponentMask);
+	return FString::Printf(TEXT("MenuProfile_%d"), IncompatibleTargetFlags);
+
+// 	if (MenuContext.Pins.Num() == 0)
+// 	{
+// 		return TEXT("NoPin");
+// 	}
+// 	else if (ContextMenuTargetProfileImpl::HasObjectPinContext(MenuContext))
+// 	{
+// 		return TEXT("ObjectPin");
+// 	}
+// 	else
+// 	{
+// 		return TEXT("PODPin");
+// 	}	
+}
+
+//------------------------------------------------------------------------------
+FContextMenuTargetProfile::FContextMenuTargetProfile()
+	: ProfileSaveName(TEXT("Default"))
+	, HasComponentsMask(0)
+	, IncompatibleTargetFlags(0)
+	, SavedTargetFlags(~IncompatibleTargetFlags)
+{
+}
+
+//------------------------------------------------------------------------------
+FContextMenuTargetProfile::FContextMenuTargetProfile(const FBlueprintActionContext& MenuContext)
+	: ProfileSaveName(ContextMenuTargetProfileImpl::GetProfileSaveName(MenuContext))
+	, HasComponentsMask(0)
+	, IncompatibleTargetFlags(ContextMenuTargetProfileImpl::GetIncompatibleTargetFlags(MenuContext, HasComponentsMask))
+	, SavedTargetFlags(~IncompatibleTargetFlags)
+{
+	const UBlueprintEditorSettings* BpSettings = GetDefault<UBlueprintEditorSettings>();
+	FString const BpConfigKey = BpSettings->GetClass()->GetPathName();
+
+	bool bOldUseTargetContextForNodeMenu = true;
+	if (GConfig->GetBool(*BpConfigKey, TEXT("bUseTargetContextForNodeMenu"), bOldUseTargetContextForNodeMenu, GEditorUserSettingsIni) && !bOldUseTargetContextForNodeMenu)
+	{
+		SavedTargetFlags = 0;
+	}
+
+	if (!LoadProfile())
+	{
+		// maybe they were originally using the shared context profile? so let's default to that
+		FString CachedProfileSaveName = ProfileSaveName;
+		ProfileSaveName = ContextMenuTargetProfileImpl::SharedProfileName;
+		LoadProfile();
+		ProfileSaveName = CachedProfileSaveName;
+	}
+}
+
+//------------------------------------------------------------------------------
+uint32 FContextMenuTargetProfile::GetContextTargetMask() const
+{
+	return SavedTargetFlags & ~IncompatibleTargetFlags;
+}
+
+//------------------------------------------------------------------------------
+void FContextMenuTargetProfile::SetContextTarget(EContextTargetFlags::Type Flag, bool bClear)
+{
+	if (bClear)
+	{
+		SavedTargetFlags &= ~Flag;
+	}
+	else
+	{
+		SavedTargetFlags |= Flag;
+	}
+	SaveProfile();
+}
+
+//------------------------------------------------------------------------------
+uint32 FContextMenuTargetProfile::GetIncompatibleTargetsMask() const
+{
+	return IncompatibleTargetFlags;
+}
+
+//------------------------------------------------------------------------------
+bool FContextMenuTargetProfile::IsTargetEnabled(EContextTargetFlags::Type Flag) const 
+{
+	if (IncompatibleTargetFlags & Flag)
+	{
+		return false;
+	}
+	else if (Flag == EContextTargetFlags::TARGET_SubComponents)
+	{
+		return ((SavedTargetFlags|EContextTargetFlags::TARGET_Blueprint) & HasComponentsMask) != 0;
+	}
+	return true;
+}
+
+//------------------------------------------------------------------------------
+void FContextMenuTargetProfile::SaveProfile() const
+{
+	GConfig->SetInt(*ContextMenuTargetProfileImpl::ConfigSection, *ProfileSaveName, SavedTargetFlags, GEditorIni);
+}
+
+//------------------------------------------------------------------------------
+bool FContextMenuTargetProfile::LoadProfile() 
+{
+	int32 SavedFlags = SavedTargetFlags;
+	if (GConfig->GetInt(*ContextMenuTargetProfileImpl::ConfigSection, *ProfileSaveName, SavedFlags, GEditorIni))
+	{
+		SavedTargetFlags = SavedFlags;
+		return true;
+	}
+	return false;
+}
+
+/*******************************************************************************
+ * SBlueprintContextTargetMenu
+ ******************************************************************************/
+
+/**  */
+namespace BlueprintContextTargetMenuImpl
+{
+	static FText GetContextTargetDisplayName(UEnum* Enum, EContextTargetFlags::Type ContextTarget);
+}
+
+//------------------------------------------------------------------------------
+static FText BlueprintContextTargetMenuImpl::GetContextTargetDisplayName(UEnum* Enum, EContextTargetFlags::Type ContextTarget)
+{
+	if (Enum != nullptr)
+	{
+		return Enum->GetDisplayNameText(ContextTarget);
+	}
+	return LOCTEXT("UnrecognizedTarget", "Error: <UNRECOGNIZED>");
+}
+
+//------------------------------------------------------------------------------
+void SBlueprintContextTargetMenu::Construct(const FArguments& InArgs, const FBlueprintActionContext& MenuContext)
+{
+	TargetProfile = FContextMenuTargetProfile(MenuContext);
+	OnTargetMaskChanged = InArgs._OnTargetMaskChanged;
+
+	FSlateFontInfo HeaderFontStyle = FEditorStyle::GetFontStyle("BlueprintEditor.ActionMenu.ContextDescriptionFont");
+	HeaderFontStyle.Size -= 2.f;
+	FText const HeaderText = ContextMenuTargetProfileImpl::GetProfileDescription(MenuContext);
+
+	FText const MenuToolTip = LOCTEXT("MenuToolTip", "Select whose functions/variables you want to see.\nNOTE: Unchecking everything is akin to 'SHOW EVERYTHING' (you're choosing to have NO target context and to not limit the scope)");
+
+	TSharedPtr<SHorizontalBox> MenuBody;
+	SBorder::Construct(SBorder::FArguments()
+		.BorderImage(FEditorStyle::GetBrush("Menu.Background"))
+		.Padding(5.f)
+		.ToolTipText(MenuToolTip)
+		[
+			SNew(SBox)
+				.MinDesiredWidth(200)
+				.ToolTipText(MenuToolTip)
+				.Padding(FMargin(0.f, 0.f, 0.f, 18.f))
+			[
+				SNew(SVerticalBox)
+				+SVerticalBox::Slot()
+				.AutoHeight()
+				[
+					SNew(SBorder)
+						.BorderImage(FEditorStyle::GetBrush("ToolPanel.GroupBorder"))
+						.ForegroundColor(FEditorStyle::GetSlateColor("DefaultForeground"))
+					[
+						SNew(STextBlock)
+							.Text(HeaderText)
+							.Font(HeaderFontStyle)
+					]
+				]
+
+				+SVerticalBox::Slot()
+					.VAlign(VAlign_Fill)
+					.HAlign(HAlign_Fill)
+				[
+					SAssignNew(MenuBody, SHorizontalBox)
+						.ToolTipText(MenuToolTip)
+				]
+			]
+		]
+	);
+
+	const uint32 ColumnCount = 2;
+	TSharedPtr<SVerticalBox> Columns[ColumnCount];
+
+	for (int32 Col = 0; Col < ColumnCount; ++Col)
+	{
+		MenuBody->AddSlot()
+			.AutoWidth()
+		[
+			SAssignNew(Columns[Col], SVerticalBox)
+				.ToolTipText(MenuToolTip)
+		];
+	}
+
+	UEnum* const TargetEnum = FindObject<UEnum>(ANY_PACKAGE, TEXT("EContextTargetFlags"));
+	uint32 const GreatestFlag = (EContextTargetFlags::ContextTargetFlagsEnd & ~1);
+
+	int32  ColIndex = 0;
+	uint32 TargetFlag = (1<<0);
+	while (TargetFlag <= GreatestFlag)
+	{
+		EContextTargetFlags::Type ContextTarget = (EContextTargetFlags::Type)TargetFlag;
+		TargetFlag <<= 1;
+
+		if (TargetEnum && TargetEnum->HasMetaData(TEXT("Hidden"), ContextTarget))
+		{
+			continue;
+		}
+		
+		FText const MenuName = BlueprintContextTargetMenuImpl::GetContextTargetDisplayName(TargetEnum, ContextTarget);
+		const FContextMenuTargetProfile& ProfileRef = TargetProfile;
+
+		Columns[ColIndex]->AddSlot()
+			.AutoHeight()
+			.VAlign(VAlign_Top)
+			.Padding(3.f, 2.5f)
+		[
+			SNew(SCheckBox)
+				.IsEnabled_Raw(&TargetProfile, &FContextMenuTargetProfile::IsTargetEnabled, ContextTarget)
+				.IsChecked(this, &SBlueprintContextTargetMenu::GetTargetCheckedState, ContextTarget)
+				.OnCheckStateChanged(this, &SBlueprintContextTargetMenu::OnTargetCheckStateChanged, ContextTarget)
+				.ToolTipText_Lambda([TargetEnum, ContextTarget, &ProfileRef]()->FText
+				{
+					if (!ProfileRef.IsTargetEnabled(ContextTarget))
+					{
+						return LOCTEXT("DisabledTargetTooltip", "This target is invalid or redundent for this context.");
+					}
+					else if (TargetEnum != nullptr)
+					{
+						return TargetEnum->GetToolTipText(ContextTarget);
+					}
+					return LOCTEXT("GenericTargetTooltip", "Include variables/functions that belong to this target.");
+				})
+			[
+				SNew(STextBlock).Text(MenuName)
+			]
+		];
+		
+		
+		ColIndex = (ColIndex + 1) % ColumnCount;
+	}
+}
+
+//------------------------------------------------------------------------------
+uint32 SBlueprintContextTargetMenu::GetContextTargetMask() const
+{
+	return TargetProfile.GetContextTargetMask();
+}
+
+//------------------------------------------------------------------------------
+void SBlueprintContextTargetMenu::OnTargetCheckStateChanged(const ECheckBoxState NewCheckedState, EContextTargetFlags::Type ContextTarget)
+{
+	TargetProfile.SetContextTarget(ContextTarget, (NewCheckedState != ECheckBoxState::Checked));
+	OnTargetMaskChanged.ExecuteIfBound(TargetProfile.GetContextTargetMask());
+}
+
+//------------------------------------------------------------------------------
+ECheckBoxState SBlueprintContextTargetMenu::GetTargetCheckedState(EContextTargetFlags::Type ContextTarget) const
+{
+	return ((TargetProfile.GetContextTargetMask() & ContextTarget) != 0) ? ECheckBoxState::Checked : ECheckBoxState::Unchecked;
+}
+
+#undef LOCTEXT_NAMESPACE
diff --git a/Engine/Source/Editor/Kismet/Private/SBlueprintContextTargetMenu.h b/Engine/Source/Editor/Kismet/Private/SBlueprintContextTargetMenu.h
new file mode 100644
--- /dev/null
+++ b/Engine/Source/Editor/Kismet/Private/SBlueprintContextTargetMenu.h
@@ -0,0 +1,84 @@
+// Copyright 1998-2015 Epic Games, Inc. All Rights Reserved.
+
+#pragma once
+
+#include "SMenuAnchor.h"
+#include "BlueprintActionMenuUtils.h" // for EContextTargetFlags
+
+struct FBlueprintActionContext;
+
+/*******************************************************************************
+ * FContextMenuTargetProfile
+ ******************************************************************************/
+
+/** Used internally by SBlueprintContextTargetMenu, to track and save the user's context target settings */
+struct FContextMenuTargetProfile
+{
+public:
+	FContextMenuTargetProfile();
+	FContextMenuTargetProfile(const FBlueprintActionContext& MenuContext);
+
+	uint32 GetContextTargetMask() const;
+	void   SetContextTarget(EContextTargetFlags::Type Flag, bool bClear = false);
+	uint32 GetIncompatibleTargetsMask() const;
+
+	bool   IsTargetEnabled(EContextTargetFlags::Type Flag) const;
+
+private:
+	void SaveProfile() const;
+	bool LoadProfile();
+
+	FString ProfileSaveName;
+	uint32  HasComponentsMask;
+	uint32  IncompatibleTargetFlags;
+	uint32  SavedTargetFlags;
+};
+
+/*******************************************************************************
+ * SBlueprintContextTargetMenu
+ ******************************************************************************/
+
+/**  */
+class SBlueprintContextTargetMenu : public SBorder
+{
+	DECLARE_DELEGATE_OneParam(FOnTargetMaskChanged, uint32);
+
+public:
+	SLATE_BEGIN_ARGS(SBlueprintContextTargetMenu){}
+		SLATE_EVENT(FOnTargetMaskChanged, OnTargetMaskChanged)
+	SLATE_END_ARGS()
+
+	// SWidget interface
+	void Construct(const FArguments& InArgs, const FBlueprintActionContext& MenuContext);
+	// End SWidget interface
+
+	/**
+	 * Returns a EContextTargetFlags::Type bitmask, defining which targets the 
+	 * user currently has enabled.
+	 */
+	uint32 GetContextTargetMask() const;
+
+private:
+	/**
+	 * Internal UI callback that handles when one of the context targets is 
+	 * checked/unchecked.
+	 * 
+	 * @param  NewCheckedState	The checkbox's new state.
+	 * @param  ContextTarget    A flag identifying which checkbox was altered.
+	 */
+	void OnTargetCheckStateChanged(const ECheckBoxState NewCheckedState, EContextTargetFlags::Type ContextTarget);
+
+	/**
+	 * Internal UI callback that determines the checkbox state for one of the 
+	 * context target options.
+	 * 
+	 * @param  ContextTarget    A flag identifying which checkbox we want the state of.
+	 * @return ECheckBoxState::Checked if the user has the flag set, otherwise ECheckBoxState::UnChecked
+	 */
+	ECheckBoxState GetTargetCheckedState(EContextTargetFlags::Type ContextTarget) const;
+
+	/** Determines what flags are incompatible with the current context, and saves/loads the user's choice settings */
+	FContextMenuTargetProfile TargetProfile;
+	/** Delegate for external users to hook into (so they can act when the menu's settings are changed). */
+	FOnTargetMaskChanged OnTargetMaskChanged;
+};
\ No newline at end of file
diff --git a/Engine/Source/Editor/Kismet/Public/BlueprintActionMenuUtils.h b/Engine/Source/Editor/Kismet/Public/BlueprintActionMenuUtils.h
--- a/Engine/Source/Editor/Kismet/Public/BlueprintActionMenuUtils.h
+++ b/Engine/Source/Editor/Kismet/Public/BlueprintActionMenuUtils.h
@@ -6,6 +6,21 @@
 struct FBlueprintActionMenuBuilder;
 struct FBlueprintActionContext;
 
+UENUM()
+namespace EContextTargetFlags
+{
+	enum Type
+	{
+		TARGET_Blueprint			= 0x00000001 UMETA(DisplayName="This Blueprint", ToolTip="Include functions and variables that belong to this Blueprint."),
+		TARGET_SubComponents		= 0x00000002 UMETA(DisplayName="Components", ToolTip="Include functions that belong to components of this Blueprint and/or the other target classes."),
+		TARGET_NodeTarget			= 0x00000004 UMETA(DisplayName="Node Target", ToolTip="Include functions and variables that belong to the same class that the pin's node does."),
+		TARGET_PinObject			= 0x00000008 UMETA(DisplayName="Pin Type Class", ToolTip="Include functions and variables that belong to this pin type."),
+		TARGET_SiblingPinObjects	= 0x00000010 UMETA(DisplayName="Other Object Outputs", ToolTip="Include functions and variables that belong to any of this node's output types."),
+
+		ContextTargetFlagsEnd UMETA(Hidden), // +1 to the last flag (so we can easily iterate these flags)
+	};
+}
+
 struct FBlueprintActionMenuUtils
 {
 	/**
@@ -25,10 +40,11 @@ struct FBlueprintActionMenuUtils
 	 * filters that're used to construct the menu.
 	 *
 	 * @param  Context				Contains the blueprint/graph/pin that the menu is for.
-	 * @param  SelectedProperties	A set of selected properties to offer contextual menu options for.
+	 * @param  bIsContextSensitive	
+	 * @param  ClassTargetMask		
 	 * @param  MenuOut				The structure that will be populated with context menu items.
 	 */
-	KISMET_API static void MakeContextMenu(FBlueprintActionContext const& Context, bool bIsContextSensitive, FBlueprintActionMenuBuilder& MenuOut);
+	KISMET_API static void MakeContextMenu(FBlueprintActionContext const& Context, bool bIsContextSensitive, uint32 ClassTargetMask, FBlueprintActionMenuBuilder& MenuOut);
 
 	/**
 	 * A centralized utility function for constructing the blueprint favorites
diff --git a/Engine/Source/Editor/UnrealEd/Private/Commandlets/DumpBlueprintsInfoCommandlet.cpp b/Engine/Source/Editor/UnrealEd/Private/Commandlets/DumpBlueprintsInfoCommandlet.cpp
--- a/Engine/Source/Editor/UnrealEd/Private/Commandlets/DumpBlueprintsInfoCommandlet.cpp
+++ b/Engine/Source/Editor/UnrealEd/Private/Commandlets/DumpBlueprintsInfoCommandlet.cpp
@@ -2112,7 +2112,8 @@ static double DumpBlueprintInfoUtils::GetContextMenuActions(FBlueprintGraphActio
 		FBlueprintActionMenuBuilder MenuBuilder(nullptr);
 		{
 			FScopedDurationTimer DurationTimer(MenuBuildDuration);
-			FBlueprintActionMenuUtils::MakeContextMenu(FilterContext, /*bIsContextSensitive =*/true, MenuBuilder);
+			FBlueprintActionMenuUtils::MakeContextMenu(FilterContext, /*bIsContextSensitive =*/true, 
+				(EContextTargetFlags::TARGET_Blueprint | EContextTargetFlags::TARGET_NodeTarget | EContextTargetFlags::TARGET_PinObject | EContextTargetFlags::TARGET_SiblingPinObjects), MenuBuilder);
 		}
 		ActionBuilder.Append(MenuBuilder);
 	}
