Unify Slate's FInputGesture and Engine's FInputChord as FInputChord defined in Slate

diff --git a/Engine/Source/Editor/GraphEditor/Private/SNodePanel.cpp b/Engine/Source/Editor/GraphEditor/Private/SNodePanel.cpp
--- a/Engine/Source/Editor/GraphEditor/Private/SNodePanel.cpp
+++ b/Engine/Source/Editor/GraphEditor/Private/SNodePanel.cpp
@@ -864,10 +864,10 @@ FReply SNodePanel::OnMouseButtonUp( const FGeometry& MyGeometry, const FPointerE
 
 				if(OnSpawnNodeByShortcut.IsBound())
 				{
-					OnSpawnNodeByShortcut.Execute(LastKeyGestureDetected, PanelCoordToGraphCoord(  MyGeometry.AbsoluteToLocal( MouseEvent.GetScreenSpacePosition() ) ));
+					OnSpawnNodeByShortcut.Execute(LastKeyChordDetected, PanelCoordToGraphCoord(  MyGeometry.AbsoluteToLocal( MouseEvent.GetScreenSpacePosition() ) ));
 				}
 
-				LastKeyGestureDetected = FInputGesture();
+				LastKeyChordDetected = FInputChord();
 			}
 		}
 		else if ( Marquee.IsValid() )
@@ -932,7 +932,7 @@ FReply SNodePanel::OnKeyDown( const FGeometry& MyGeometry, const FKeyEvent& InKe
 {
 	if( IsEditable.Get() )
 	{
-		LastKeyGestureDetected = FInputGesture(InKeyEvent.GetKey(), EModifierKey::FromBools(InKeyEvent.IsControlDown(), InKeyEvent.IsAltDown(), InKeyEvent.IsShiftDown(), InKeyEvent.IsCommandDown()));
+		LastKeyChordDetected = FInputChord(InKeyEvent.GetKey(), EModifierKey::FromBools(InKeyEvent.IsControlDown(), InKeyEvent.IsAltDown(), InKeyEvent.IsShiftDown(), InKeyEvent.IsCommandDown()));
 	}
 
 	return FReply::Unhandled();
@@ -940,9 +940,9 @@ FReply SNodePanel::OnKeyDown( const FGeometry& MyGeometry, const FKeyEvent& InKe
 
 FReply SNodePanel::OnKeyUp( const FGeometry& MyGeometry, const FKeyEvent& InKeyEvent )
 {
-	if(LastKeyGestureDetected.Key == InKeyEvent.GetKey())
+	if(LastKeyChordDetected.Key == InKeyEvent.GetKey())
 	{
-		LastKeyGestureDetected = FInputGesture();
+		LastKeyChordDetected = FInputChord();
 	}
 
 	return FReply::Unhandled();
@@ -950,7 +950,7 @@ FReply SNodePanel::OnKeyUp( const FGeometry& MyGeometry, const FKeyEvent& InKeyE
 
 void SNodePanel::OnFocusLost( const FFocusEvent& InFocusEvent )
 {
-	LastKeyGestureDetected = FInputGesture();
+	LastKeyChordDetected = FInputChord();
 }
 
 FReply SNodePanel::OnTouchGesture( const FGeometry& MyGeometry, const FPointerEvent& GestureEvent )
diff --git a/Engine/Source/Editor/GraphEditor/Public/SNodePanel.h b/Engine/Source/Editor/GraphEditor/Public/SNodePanel.h
--- a/Engine/Source/Editor/GraphEditor/Public/SNodePanel.h
+++ b/Engine/Source/Editor/GraphEditor/Public/SNodePanel.h
@@ -875,8 +875,8 @@ protected:
 	/** Invoked when the user may be attempting to spawn a node using a shortcut */
 	SGraphEditor::FOnSpawnNodeByShortcut OnSpawnNodeByShortcut;
 
-	/** The last key gesture detected in this graph panel */
-	FInputGesture LastKeyGestureDetected;
+	/** The last key chord detected in this graph panel */
+	FInputChord LastKeyChordDetected;
 
 	/** The current transaction for undo/redo */
 	TSharedPtr<FScopedTransaction> ScopedTransactionPtr;
diff --git a/Engine/Source/Editor/InputBindingEditor/Private/InputBindingEditorPrivatePCH.h b/Engine/Source/Editor/InputBindingEditor/Private/InputBindingEditorPrivatePCH.h
--- a/Engine/Source/Editor/InputBindingEditor/Private/InputBindingEditorPrivatePCH.h
+++ b/Engine/Source/Editor/InputBindingEditor/Private/InputBindingEditorPrivatePCH.h
@@ -14,7 +14,7 @@
 /* Private includes
  *****************************************************************************/
 
-#include "SGestureEditBox.h"
-#include "SGestureTreeItem.h"
-#include "SGestureEditor.h"
+#include "SChordEditBox.h"
+#include "SChordTreeItem.h"
+#include "SChordEditor.h"
 #include "SInputBindingEditorPanel.h"
diff --git a/Engine/Source/Editor/InputBindingEditor/Private/Widgets/SChordEditBox.cpp b/Engine/Source/Editor/InputBindingEditor/Private/Widgets/SChordEditBox.cpp
new file mode 100644
--- /dev/null
+++ b/Engine/Source/Editor/InputBindingEditor/Private/Widgets/SChordEditBox.cpp
@@ -0,0 +1,276 @@
+// Copyright 1998-2015 Epic Games, Inc. All Rights Reserved.
+
+#include "InputBindingEditorPrivatePCH.h"
+
+
+#define LOCTEXT_NAMESPACE "SChordEditBox"
+
+
+/* SChordEditBox interface
+ *****************************************************************************/
+
+BEGIN_SLATE_FUNCTION_BUILD_OPTIMIZATION
+void SChordEditBox::Construct( const FArguments& InArgs, TSharedPtr<FChordTreeItem> InputCommand )
+{
+	BorderImageNormal = FEditorStyle::GetBrush( "EditableTextBox.Background.Normal" );
+	BorderImageHovered = FEditorStyle::GetBrush( "EditableTextBox.Background.Hovered" );
+	BorderImageFocused = FEditorStyle::GetBrush( "EditableTextBox.Background.Focused" );
+
+	static const FName InvertedForegroundName("InvertedForeground");
+
+	ChildSlot
+	[
+		SAssignNew( ConflictPopup, SMenuAnchor )
+		.Placement(MenuPlacement_ComboBox)
+		.OnGetMenuContent( this, &SChordEditBox::OnGetContentForConflictPopup )
+		.OnMenuOpenChanged(this, &SChordEditBox::OnConflictPopupOpenChanged)
+		[
+			SNew( SBox )
+			.WidthOverride( 200.0f )
+			[
+				SNew( SBorder )
+				.VAlign(VAlign_Center)
+				.Padding( FMargin( 4.0f, 2.0f ) )
+				.BorderImage( this, &SChordEditBox::GetBorderImage )
+				.ForegroundColor( FEditorStyle::GetSlateColor(InvertedForegroundName) )
+				[
+					SNew( SHorizontalBox )
+					+ SHorizontalBox::Slot()
+					.FillWidth(1)
+					.VAlign(VAlign_Center)
+					[
+						SAssignNew( ChordEditor, SChordEditor, InputCommand )
+						.OnEditBoxLostFocus( this, &SChordEditBox::OnChordEditorLostFocus )
+						.OnChordChanged( this, &SChordEditBox::OnChordChanged )
+						.OnEditingStarted( this, &SChordEditBox::OnChordEditingStarted )
+						.OnEditingStopped( this, &SChordEditBox::OnChordEditingStopped )
+					]
+					+ SHorizontalBox::Slot()
+					.AutoWidth()
+					.HAlign(HAlign_Right)
+					.VAlign(VAlign_Center)
+					[
+						// Remove binding button
+						SNew(SButton)
+						.Visibility( this, &SChordEditBox::GetChordRemoveButtonVisibility )
+						.ButtonStyle( FEditorStyle::Get(), "NoBorder" )
+						.ContentPadding(0)
+						.OnClicked( this, &SChordEditBox::OnChordRemoveButtonClicked )
+						.ForegroundColor( FSlateColor::UseForeground() )
+						.IsFocusable(false)
+						.ToolTipText(LOCTEXT("ChordEditButtonRemove_ToolTip", "Remove this binding") )
+						[
+							SNew( SBox )
+							[
+								SNew( SImage )
+								.Image( FEditorStyle::GetBrush( "Symbols.X" ) )
+								.ColorAndOpacity( FLinearColor(.7f,0,0,.75f) )
+							]
+						]
+					]
+				]
+			]
+		]
+	];
+}
+END_SLATE_FUNCTION_BUILD_OPTIMIZATION
+
+
+const FSlateBrush* SChordEditBox::GetBorderImage() const
+{
+	if ( ChordEditor->HasKeyboardFocus() )
+	{
+		return BorderImageFocused;
+	}
+	else
+	{
+		if ( ChordEditor->IsHovered() )
+		{
+			return BorderImageHovered;
+		}
+		else
+		{
+			return BorderImageNormal;
+		}
+	}
+}
+
+
+FText SChordEditBox::GetNotificationMessage() const
+{
+	return ChordEditor->GetNotificationText();
+}
+
+
+void SChordEditBox::OnChordEditorLostFocus()
+{
+	if( (!ChordAcceptButton.IsValid() || ChordAcceptButton->HasMouseCapture() == false) && !ChordEditor->IsTyping() )
+	{
+		if( ChordEditor->IsEditing() && ChordEditor->IsEditedChordValid() && !ChordEditor->HasConflict() )
+		{
+			ChordEditor->CommitNewChord();
+		}
+
+		ChordEditor->StopEditing();
+	}
+}
+
+
+void SChordEditBox::OnChordEditingStarted()
+{
+	ConflictPopup->SetIsOpen(false);
+}
+
+
+void SChordEditBox::OnChordEditingStopped()
+{
+	if( ChordEditor->IsEditedChordValid() && !ChordEditor->HasConflict() )
+	{
+		ChordEditor->CommitNewChord();
+	}
+}
+
+
+void SChordEditBox::OnChordChanged()
+{
+	if( ChordEditor->HasConflict() )
+	{
+		ConflictPopup->SetIsOpen(true, true);
+	}
+	else
+	{
+		ConflictPopup->SetIsOpen(false);
+
+		if (ChordEditor->IsEditedChordValid())
+		{
+			ChordEditor->CommitNewChord();
+			ChordEditor->StopEditing();
+			FSlateApplication::Get().ClearKeyboardFocus(EFocusCause::SetDirectly);
+		}
+	}
+}
+
+EVisibility SChordEditBox::GetChordRemoveButtonVisibility() const
+{
+	if( !ChordEditor->IsEditing() && ChordEditor->IsActiveChordValid() )
+	{
+		// Chord should display a button.  What the button does is defined by the two cases above
+		return EVisibility::Visible;
+	}
+	else
+	{
+		// Nothing to remove, but still take up the space
+		return EVisibility::Hidden;
+	}
+}
+
+
+FReply SChordEditBox::OnChordRemoveButtonClicked()
+{
+	if( !ChordEditor->IsEditing() )
+	{
+		ChordEditor->RemoveActiveChord();
+	}
+
+	return FReply::Handled();
+}
+
+
+FReply SChordEditBox::OnAcceptNewChordButtonClicked()
+{
+	if( ChordEditor->IsEditing() )
+	{
+		ChordEditor->CommitNewChord();
+		ChordEditor->StopEditing();
+		
+	}
+
+	ConflictPopup->SetIsOpen(false);
+
+	return FReply::Handled();
+}
+
+
+TSharedRef<SWidget> SChordEditBox::OnGetContentForConflictPopup()
+{
+	return SNew(SBorder)
+		.BorderImage( FEditorStyle::GetBrush("NotificationList.ItemBackground")  )
+		[
+			SNew( SVerticalBox )
+
+			+ SVerticalBox::Slot()
+				.VAlign(VAlign_Center)
+				.Padding( 2.0f, 0.0f )
+				.AutoHeight()
+				[
+					SNew( STextBlock )
+						.WrapTextAt(200.0f)
+						.ColorAndOpacity( FLinearColor( .75f, 0.0f, 0.0f, 1.0f ) )
+						.Text( this, &SChordEditBox::GetNotificationMessage )
+						.Visibility( this, &SChordEditBox::GetNotificationVisibility )
+				]
+
+			+ SVerticalBox::Slot()
+				.Padding( 2.0f )
+				.VAlign(VAlign_Center)
+				.HAlign(HAlign_Right)
+				.AutoHeight()
+				[
+					SAssignNew( ChordAcceptButton, SButton )
+						.ContentPadding(1)
+						.ToolTipText( LOCTEXT("ChordAcceptButton_ToolTip", "Accept this new binding") )
+						.OnClicked( this, &SChordEditBox::OnAcceptNewChordButtonClicked )
+						[
+							SNew( SHorizontalBox )
+
+							+ SHorizontalBox::Slot()
+								.AutoWidth()
+								[
+									SNew( SImage )
+										.Image( FEditorStyle::GetBrush( "Symbols.Check" ) )
+										.ColorAndOpacity( FLinearColor(0,.7f,0,.75f) )
+								]
+
+							+ SHorizontalBox::Slot()
+								.AutoWidth()
+								.Padding(2.0f,0.0f)
+								[
+									SNew( STextBlock )
+										.Text( LOCTEXT("ChordAcceptButtonText_Override", "Override") )
+								]
+						]
+				]
+		];
+}
+
+void SChordEditBox::OnConflictPopupOpenChanged(bool bIsOpen)
+{
+	if(!bIsOpen)
+	{
+		ChordEditor->StopEditing();
+	}
+}
+
+EVisibility SChordEditBox::GetNotificationVisibility() const
+{
+	return !ChordEditor->GetNotificationText().IsEmpty() ? EVisibility::Visible : EVisibility::Collapsed;
+}
+
+
+FReply SChordEditBox::OnMouseButtonDown( const FGeometry& MyGeometry, const FPointerEvent& InMouseEvent )
+{
+	// This is a passthrough if the chord edit box gets a mouse click in the button area and the button isn't visible. 
+	// we should focus the lower level editing widget in this case
+	if ( InMouseEvent.GetEffectingButton() == EKeys::LeftMouseButton )
+	{
+		if( !ChordEditor->IsEditing() )
+		{
+			ChordEditor->StartEditing();
+		}
+		return FReply::Handled().SetUserFocus(ChordEditor.ToSharedRef(), EFocusCause::Mouse);
+	}
+
+	return FReply::Handled();
+}
+
+#undef LOCTEXT_NAMESPACE
diff --git a/Engine/Source/Editor/InputBindingEditor/Private/Widgets/SChordTreeItem.cpp b/Engine/Source/Editor/InputBindingEditor/Private/Widgets/SChordTreeItem.cpp
new file mode 100644
--- /dev/null
+++ b/Engine/Source/Editor/InputBindingEditor/Private/Widgets/SChordTreeItem.cpp
@@ -0,0 +1,89 @@
+// Copyright 1998-2015 Epic Games, Inc. All Rights Reserved.
+
+#include "InputBindingEditorPrivatePCH.h"
+#include "SChordTreeItem.h"
+#include "UICommandDragDropOp.h"
+
+
+
+
+void SChordTreeItem::Construct(const FArguments& InArgs, const TSharedRef<STableViewBase>& InOwnerTable, TSharedPtr<FChordTreeItem> InItem)
+{
+	TreeItem = InItem;
+
+	SMultiColumnTableRow< TSharedPtr<FChordTreeItem> >::Construct(FSuperRowType::FArguments(), InOwnerTable);
+}
+
+TSharedRef<SWidget> SChordTreeItem::GenerateWidgetForColumn(const FName& ColumnName)
+{
+	if (ColumnName == "Name")
+	{
+		return SNew(SHorizontalBox)
+
+			+ SHorizontalBox::Slot()
+			.AutoWidth()
+			.Padding(2.0f)
+			.VAlign(VAlign_Center)
+			[
+				SNew(SExpanderArrow, SharedThis(this))
+			]
+
+		+ SHorizontalBox::Slot()
+			.AutoWidth()
+			.VAlign(VAlign_Center)
+			[
+				SNew(SVerticalBox)
+
+				+ SVerticalBox::Slot()
+				.AutoHeight()
+				[
+					SNew(STextBlock)
+					.Text(TreeItem->CommandInfo->GetLabel())
+					.ToolTipText(TreeItem->CommandInfo->GetDescription())
+				]
+
+				+ SVerticalBox::Slot()
+					.Padding(0.0f, 0.0f, 0.0f, 5.0f)
+					.AutoHeight()
+					[
+						SNew(STextBlock)
+						.Font(FEditorStyle::GetFontStyle("InputBindingEditor.SmallFont"))
+						.ColorAndOpacity(FLinearColor::Gray)
+						.Text(TreeItem->CommandInfo->GetDescription())
+					]
+			];
+	}
+	else if (ColumnName == "Binding")
+	{
+		return SNew(SBox)
+			.Padding(2.0f)
+			.VAlign(VAlign_Center)
+			[
+				SNew(SChordEditBox, TreeItem)
+			];
+	}
+	else
+	{
+		return SNew(STextBlock).Text(NSLOCTEXT("ChordTreeItem", "UnknownColumn", "Unknown Column"));
+	}
+}
+
+FReply SChordTreeItem::OnMouseButtonDown(const FGeometry& MyGeometry, const FPointerEvent& MouseEvent)
+{
+	if (MouseEvent.GetEffectingButton() == EKeys::LeftMouseButton && FMultiBoxSettings::IsInToolbarEditMode())
+	{
+		return FReply::Handled().DetectDrag(SharedThis(this), MouseEvent.GetEffectingButton());
+	}
+
+	return FReply::Unhandled();
+}
+
+FReply SChordTreeItem::OnDragDetected(const FGeometry& MyGeometry, const FPointerEvent& MouseEvent)
+{
+	FToolBarBuilder TempBuilder(MakeShareable(new FUICommandList), FMultiBoxCustomization::None);
+	TempBuilder.AddToolBarButton(TreeItem->CommandInfo.ToSharedRef());
+
+	TSharedRef<SWidget> CursorDecorator = TempBuilder.MakeWidget();
+
+	return FReply::Handled().BeginDragDrop(FUICommandDragDropOp::New(TreeItem->CommandInfo.ToSharedRef(), NAME_None, TempBuilder.MakeWidget(), FVector2D::ZeroVector));
+}
diff --git a/Engine/Source/Editor/InputBindingEditor/Private/Widgets/SChordTreeItem.h b/Engine/Source/Editor/InputBindingEditor/Private/Widgets/SChordTreeItem.h
new file mode 100644
--- /dev/null
+++ b/Engine/Source/Editor/InputBindingEditor/Private/Widgets/SChordTreeItem.h
@@ -0,0 +1,66 @@
+// Copyright 1998-2015 Epic Games, Inc. All Rights Reserved.
+
+#pragma once
+
+
+/**
+ * An item for the chord tree view
+ */
+struct FChordTreeItem
+{
+	// Note these are mutually exclusive
+	TWeakPtr<FBindingContext> BindingContext;
+	TSharedPtr<FUICommandInfo> CommandInfo;
+
+	TSharedPtr<FBindingContext> GetBindingContext() { return BindingContext.Pin(); }
+
+	bool IsContext() const { return BindingContext.IsValid(); }
+	bool IsCommand() const { return CommandInfo.IsValid(); }
+};
+
+
+typedef STreeView< TSharedPtr<FChordTreeItem> > SChordTree;
+
+
+/**
+ * A widget which visualizes a command info      .              
+ */
+class SChordTreeItem
+	: public SMultiColumnTableRow< TSharedPtr< FChordTreeItem > >
+ {
+public:
+
+	SLATE_BEGIN_ARGS( SChordTreeItem ){}
+	SLATE_END_ARGS()
+
+ public:
+
+	/**
+	 * Constructs the widget.
+	 *
+	 * @param InArgs The Slate argument list.
+	 * @param InOwnerTable The table that owns this tree item.
+	 * @param InItem The actual item to be displayed.
+	 */
+	void Construct( const FArguments& InArgs, const TSharedRef<STableViewBase>& InOwnerTable, TSharedPtr<FChordTreeItem> InItem );
+
+	/**
+	 * Called to generate a widget for each column.
+	 *
+	 * @param ColumnName The name of the column that needs a widget.
+	 */
+	TSharedRef<SWidget> GenerateWidgetForColumn( const FName& ColumnName );
+
+public:
+
+	// SWidget interface
+
+	virtual FReply OnMouseButtonDown( const FGeometry& MyGeometry, const FPointerEvent& MouseEvent ) override;
+
+	virtual FReply OnDragDetected( const FGeometry& MyGeometry, const FPointerEvent& MouseEvent ) override;
+
+private:
+
+	/** Holds the tree item being visualized. */
+	TSharedPtr<FChordTreeItem> TreeItem;
+};
diff --git a/Engine/Source/Editor/InputBindingEditor/Private/Widgets/SInputBindingEditorPanel.cpp b/Engine/Source/Editor/InputBindingEditor/Private/Widgets/SInputBindingEditorPanel.cpp
--- a/Engine/Source/Editor/InputBindingEditor/Private/Widgets/SInputBindingEditorPanel.cpp
+++ b/Engine/Source/Editor/InputBindingEditor/Private/Widgets/SInputBindingEditorPanel.cpp
@@ -6,7 +6,7 @@
 #define LOCTEXT_NAMESPACE "SInputBindingEditorPanel"
 
 
-/* SGestureEditor interface
+/* SChordEditor interface
  *****************************************************************************/
 
 void SInputBindingEditorPanel::Construct(const FArguments& InArgs)
@@ -32,7 +32,7 @@ void SInputBindingEditorPanel::Construct(const FArguments& InArgs)
 			+ SVerticalBox::Slot()
 			.FillHeight(1)
 			[
-				SAssignNew( GestureTree, SGestureTree )
+				SAssignNew( ChordTree, SChordTree )
 				.TreeItemsSource( &ContextVisibleList )
 				.OnGetChildren( this, &SInputBindingEditorPanel::OnGetChildrenForTreeItem )
 				.OnGenerateRow( this, &SInputBindingEditorPanel::OnGenerateWidgetForTreeItem )
@@ -84,12 +84,12 @@ void SInputBindingEditorPanel::Construct(const FArguments& InArgs)
 		for( auto CurContextIt( ContextVisibleList.CreateConstIterator() ); CurContextIt; ++CurContextIt )
 		{
 			const bool bShouldExpand = true;
-			GestureTree->SetItemExpansion( *CurContextIt, bShouldExpand );
+			ChordTree->SetItemExpansion( *CurContextIt, bShouldExpand );
 		}
 	}
 
 
-	GestureTree->RequestTreeRefresh();
+	ChordTree->RequestTreeRefresh();
 }
 
 
@@ -136,8 +136,8 @@ FReply SInputBindingEditorPanel::OnBindingColumnClicked()
 
 	const bool bSortName = false;
 
-	GestureSortMode = FGestureSort( bSortName, bSortUp );
-	GestureTree->RequestTreeRefresh();
+	ChordSortMode = FChordSort( bSortName, bSortUp );
+	ChordTree->RequestTreeRefresh();
 	return FReply::Handled();
 }
 
@@ -149,8 +149,8 @@ FReply SInputBindingEditorPanel::OnNameColumnClicked()
 
 	const bool bSortName = true;
 
-	GestureSortMode = FGestureSort( bSortName, bSortUp );
-	GestureTree->RequestTreeRefresh();
+	ChordSortMode = FChordSort( bSortName, bSortUp );
+	ChordTree->RequestTreeRefresh();
 
 	return FReply::Handled();
 }
@@ -179,7 +179,7 @@ void SInputBindingEditorPanel::UpdateContextMasterList()
 	{
 		const auto& Context = Contexts[ListIndex];
 
-		TSharedRef<FGestureTreeItem> TreeItem( new FGestureTreeItem );
+		TSharedRef<FChordTreeItem> TreeItem( new FChordTreeItem );
 		TreeItem->BindingContext = Context;
 		ContextMasterList.Add( TreeItem );
 	}
@@ -227,12 +227,12 @@ void SInputBindingEditorPanel::FilterVisibleContextList()
 				ContextVisibleList.Add( ContextMasterList[Context] );
 
 				// Expand all filtered contexts
-				GestureTree->SetItemExpansion( ContextMasterList[Context], true );
+				ChordTree->SetItemExpansion( ContextMasterList[Context], true );
 			}
 		}
 	}
 
-	GestureTree->RequestTreeRefresh();
+	ChordTree->RequestTreeRefresh();
 }
 
 
@@ -249,14 +249,14 @@ void SInputBindingEditorPanel::OnCommandsChanged()
 }
 
 
-TSharedRef< ITableRow > SInputBindingEditorPanel::OnGenerateWidgetForTreeItem( TSharedPtr<FGestureTreeItem> InTreeItem, const TSharedRef<STableViewBase>& OwnerTable )
+TSharedRef< ITableRow > SInputBindingEditorPanel::OnGenerateWidgetForTreeItem( TSharedPtr<FChordTreeItem> InTreeItem, const TSharedRef<STableViewBase>& OwnerTable )
 {
 	if( InTreeItem->IsContext() )
 	{
 		static const FName InvertedForegroundName("InvertedForeground");
 		// contexts do not need columns
 		return
-			SNew( STableRow< TSharedPtr<FGestureTreeItem> >, OwnerTable )
+			SNew( STableRow< TSharedPtr<FChordTreeItem> >, OwnerTable )
 			.Padding(3)
 			[
 				SNew( SBorder )
@@ -286,12 +286,12 @@ TSharedRef< ITableRow > SInputBindingEditorPanel::OnGenerateWidgetForTreeItem( T
 	}
 	else
 	{
-		return SNew( SGestureTreeItem, OwnerTable, InTreeItem );
+		return SNew( SChordTreeItem, OwnerTable, InTreeItem );
 	}
 }
 
 
-void SInputBindingEditorPanel::OnGetChildrenForTreeItem( TSharedPtr<FGestureTreeItem> InTreeItem, TArray< TSharedPtr< FGestureTreeItem > >& OutChildren )
+void SInputBindingEditorPanel::OnGetChildrenForTreeItem( TSharedPtr<FChordTreeItem> InTreeItem, TArray< TSharedPtr< FChordTreeItem > >& OutChildren )
 {
 	if( InTreeItem->IsContext() )
 	{
@@ -304,13 +304,13 @@ void SInputBindingEditorPanel::OnGetChildrenForTreeItem( TSharedPtr<FGestureTree
 		{
 			if( FilterStrings.Num() == 0 || bContextPassesFilter || StringPassesFilter( CommandInfos[CommandIndex]->GetLabel(), FilterStrings ) || StringPassesFilter( CommandInfos[CommandIndex]->GetInputText(), FilterStrings ) )
 			{
-				TSharedPtr<FGestureTreeItem> NewItem( new FGestureTreeItem );
+				TSharedPtr<FChordTreeItem> NewItem( new FChordTreeItem );
 				NewItem->CommandInfo = CommandInfos[ CommandIndex ];
 				OutChildren.Add( NewItem );
 			}
 		}
 
-		OutChildren.Sort( GestureSortMode );
+		OutChildren.Sort( ChordSortMode );
 	}
 }
 
diff --git a/Engine/Source/Editor/InputBindingEditor/Private/Widgets/SInputBindingEditorPanel.h b/Engine/Source/Editor/InputBindingEditor/Private/Widgets/SInputBindingEditorPanel.h
--- a/Engine/Source/Editor/InputBindingEditor/Private/Widgets/SInputBindingEditorPanel.h
+++ b/Engine/Source/Editor/InputBindingEditor/Private/Widgets/SInputBindingEditorPanel.h
@@ -6,14 +6,14 @@
 /**
  * A gesture sort functor.  Sorts by name or gesture and ascending or descending
  */
-struct FGestureSort
+struct FChordSort
 {
-	FGestureSort( bool bInSortName, bool bInSortUp )
+	FChordSort( bool bInSortName, bool bInSortUp )
 		: bSortName( bInSortName )
 		, bSortUp( bInSortUp )
 	{ }
 
-	bool operator()( const TSharedPtr<FGestureTreeItem>& A, const TSharedPtr<FGestureTreeItem>& B ) const
+	bool operator()( const TSharedPtr<FChordTreeItem>& A, const TSharedPtr<FChordTreeItem>& B ) const
 	{
 		if( bSortName )
 		{
@@ -52,7 +52,7 @@ public:
 
 	/** Default constructor. */
 	SInputBindingEditorPanel()
-		: GestureSortMode( true, false )
+		: ChordSortMode( true, false )
 	{ }
 
 	/** Destructor. */
@@ -90,12 +90,12 @@ private:
 	/**
 	 * Generates widget for an item in the gesture tree.
 	 */
-	TSharedRef< ITableRow > OnGenerateWidgetForTreeItem( TSharedPtr<FGestureTreeItem> InTreeItem, const TSharedRef<STableViewBase>& OwnerTable );
+	TSharedRef< ITableRow > OnGenerateWidgetForTreeItem( TSharedPtr<FChordTreeItem> InTreeItem, const TSharedRef<STableViewBase>& OwnerTable );
 
 	/**
-	 * Gets children FGestureTreeItems from the passed in tree item.  Note: Only contexts have children and those children are the actual gestures.
+	 * Gets children FChordTreeItems from the passed in tree item.  Note: Only contexts have children and those children are the actual gestures.
 	 */
-	void OnGetChildrenForTreeItem( TSharedPtr<FGestureTreeItem> InTreeItem, TArray< TSharedPtr< FGestureTreeItem > >& OutChildren );
+	void OnGetChildrenForTreeItem( TSharedPtr<FChordTreeItem> InTreeItem, TArray< TSharedPtr< FChordTreeItem > >& OutChildren );
 
 	/**
 	 * Called when the binding column is clicked.  We sort by binding in this case .
@@ -119,19 +119,19 @@ private:
 private:
 
 	/** List of all known contexts. */
-	TArray< TSharedPtr<FGestureTreeItem> > ContextMasterList;
+	TArray< TSharedPtr<FChordTreeItem> > ContextMasterList;
 	
 	/** List of contexts visible in the tree. */
-	TArray< TSharedPtr<FGestureTreeItem> > ContextVisibleList;
+	TArray< TSharedPtr<FChordTreeItem> > ContextVisibleList;
 	
 	/** Search box */
 	TSharedPtr<SWidget> SearchBox;
 
-	/** Gesture tree widget. */
-	TSharedPtr< SGestureTree > GestureTree;
+	/** Chord tree widget. */
+	TSharedPtr< SChordTree > ChordTree;
 	
 	/** The current gesture sort to use. */
-	FGestureSort GestureSortMode;
+	FChordSort ChordSortMode;
 	
 	/** The current list of filter strings to filter gestures by. */
 	TArray<FString> FilterStrings;
diff --git a/Engine/Source/Editor/Kismet/Private/BlueprintEditor.cpp b/Engine/Source/Editor/Kismet/Private/BlueprintEditor.cpp
--- a/Engine/Source/Editor/Kismet/Private/BlueprintEditor.cpp
+++ b/Engine/Source/Editor/Kismet/Private/BlueprintEditor.cpp
@@ -7334,7 +7334,7 @@ void FBlueprintEditor::DestroyPreview()
 }
 
 
-FReply FBlueprintEditor::OnSpawnGraphNodeByShortcut(FInputGesture InGesture, const FVector2D& InPosition, UEdGraph* InGraph)
+FReply FBlueprintEditor::OnSpawnGraphNodeByShortcut(FInputChord InChord, const FVector2D& InPosition, UEdGraph* InGraph)
 {
 	UEdGraph* Graph = InGraph;
 	if (Graph == nullptr)
@@ -7343,7 +7343,7 @@ FReply FBlueprintEditor::OnSpawnGraphNodeByShortcut(FInputGesture InGesture, con
 	}
 
 	FBlueprintPaletteListBuilder PaletteBuilder(GetBlueprintObj());
-	FBlueprintSpawnNodeCommands::Get().GetGraphActionByGesture(InGesture, PaletteBuilder, InGraph);
+	FBlueprintSpawnNodeCommands::Get().GetGraphActionByChord(InChord, PaletteBuilder, InGraph);
 
 	TSet<const UEdGraphNode*> NodesToSelect;
 	FVector2D NodeSpawnPos = InPosition;
diff --git a/Engine/Source/Editor/Kismet/Private/BlueprintEditorCommands.cpp b/Engine/Source/Editor/Kismet/Private/BlueprintEditorCommands.cpp
--- a/Engine/Source/Editor/Kismet/Private/BlueprintEditorCommands.cpp
+++ b/Engine/Source/Editor/Kismet/Private/BlueprintEditorCommands.cpp
@@ -18,58 +18,58 @@ PRAGMA_DISABLE_OPTIMIZATION
 void FBlueprintEditorCommands::RegisterCommands()
 {
 	// Edit commands
-	UI_COMMAND( FindInBlueprint, "Search", "Finds references to functions, events, variables, and pins in the current Blueprint", EUserInterfaceActionType::Button, FInputGesture(EModifierKey::Control, EKeys::F) );
-	UI_COMMAND( FindInBlueprints, "Find in Blueprints", "Find references to functions, events and variables in ALL Blueprints", EUserInterfaceActionType::Button, FInputGesture(EModifierKey::Control | EModifierKey::Shift, EKeys::F) );
-	UI_COMMAND( ReparentBlueprint, "Reparent Blueprint", "Change the parent of this Blueprint", EUserInterfaceActionType::Button, FInputGesture() );
+	UI_COMMAND( FindInBlueprint, "Search", "Finds references to functions, events, variables, and pins in the current Blueprint", EUserInterfaceActionType::Button, FInputChord(EModifierKey::Control, EKeys::F) );
+	UI_COMMAND( FindInBlueprints, "Find in Blueprints", "Find references to functions, events and variables in ALL Blueprints", EUserInterfaceActionType::Button, FInputChord(EModifierKey::Control | EModifierKey::Shift, EKeys::F) );
+	UI_COMMAND( ReparentBlueprint, "Reparent Blueprint", "Change the parent of this Blueprint", EUserInterfaceActionType::Button, FInputChord() );
 
-	UI_COMMAND( CompileBlueprint, "Compile", "Compile the blueprint", EUserInterfaceActionType::Button, FInputGesture(EKeys::F7) );
-	UI_COMMAND( RefreshAllNodes, "Refresh All nodes", "Refreshes all nodes in the graph to account for external changes", EUserInterfaceActionType::Button, FInputGesture() );
-	UI_COMMAND( DeleteUnusedVariables, "Delete Unused Variables", "Deletes any variables that are never used", EUserInterfaceActionType::Button, FInputGesture() );
+	UI_COMMAND( CompileBlueprint, "Compile", "Compile the blueprint", EUserInterfaceActionType::Button, FInputChord(EKeys::F7) );
+	UI_COMMAND( RefreshAllNodes, "Refresh All nodes", "Refreshes all nodes in the graph to account for external changes", EUserInterfaceActionType::Button, FInputChord() );
+	UI_COMMAND( DeleteUnusedVariables, "Delete Unused Variables", "Deletes any variables that are never used", EUserInterfaceActionType::Button, FInputChord() );
 
-	UI_COMMAND( FindReferencesFromClass, "List references (from class)", "Find all objects that the class references", EUserInterfaceActionType::Button, FInputGesture() );
-	UI_COMMAND( FindReferencesFromBlueprint, "List referenced (from blueprint)", "Find all objects that the blueprint references", EUserInterfaceActionType::Button, FInputGesture() );
-	UI_COMMAND( RepairCorruptedBlueprint, "Repair corrupted blueprint", "Attempts to repair a corrupted blueprint that cannot be saved", EUserInterfaceActionType::Button, FInputGesture() );
+	UI_COMMAND( FindReferencesFromClass, "List references (from class)", "Find all objects that the class references", EUserInterfaceActionType::Button, FInputChord() );
+	UI_COMMAND( FindReferencesFromBlueprint, "List referenced (from blueprint)", "Find all objects that the blueprint references", EUserInterfaceActionType::Button, FInputChord() );
+	UI_COMMAND( RepairCorruptedBlueprint, "Repair corrupted blueprint", "Attempts to repair a corrupted blueprint that cannot be saved", EUserInterfaceActionType::Button, FInputChord() );
 
 	// View commands 
-	UI_COMMAND( ZoomToWindow, "Zoom to Graph Extents", "Fit the current view to the entire graph", EUserInterfaceActionType::Button, FInputGesture() );
-	UI_COMMAND( ZoomToSelection, "Zoom to Selection", "Fit the current view to the selection", EUserInterfaceActionType::Button, FInputGesture(EKeys::Home) );
-	UI_COMMAND( NavigateToParent, "Go to parent graph", "Open the parent graph", EUserInterfaceActionType::Button, FInputGesture(EKeys::PageUp) );
-	UI_COMMAND( NavigateToParentBackspace, "Go to parent graph", "Open the parent graph", EUserInterfaceActionType::Button, FInputGesture(EKeys::BackSpace) );
-	UI_COMMAND( NavigateToChild, "Go to child graph", "Open the child graph", EUserInterfaceActionType::Button, FInputGesture(EKeys::PageDown) );
+	UI_COMMAND( ZoomToWindow, "Zoom to Graph Extents", "Fit the current view to the entire graph", EUserInterfaceActionType::Button, FInputChord() );
+	UI_COMMAND( ZoomToSelection, "Zoom to Selection", "Fit the current view to the selection", EUserInterfaceActionType::Button, FInputChord(EKeys::Home) );
+	UI_COMMAND( NavigateToParent, "Go to parent graph", "Open the parent graph", EUserInterfaceActionType::Button, FInputChord(EKeys::PageUp) );
+	UI_COMMAND( NavigateToParentBackspace, "Go to parent graph", "Open the parent graph", EUserInterfaceActionType::Button, FInputChord(EKeys::BackSpace) );
+	UI_COMMAND( NavigateToChild, "Go to child graph", "Open the child graph", EUserInterfaceActionType::Button, FInputChord(EKeys::PageDown) );
 
 	// Preview commands
-	UI_COMMAND( ResetCamera, "Reset Camera", "Resets the camera to focus on the mesh", EUserInterfaceActionType::Button, FInputGesture() );
-	UI_COMMAND( EnableSimulation, "Simulation", "Enables the simulation of the blueprint and ticking", EUserInterfaceActionType::ToggleButton, FInputGesture() );
-	UI_COMMAND( ShowFloor, "Show Floor", "Toggles a ground mesh for collision", EUserInterfaceActionType::ToggleButton, FInputGesture() );
-	UI_COMMAND( ShowGrid, "Show Grid", "Toggles viewport grid", EUserInterfaceActionType::ToggleButton, FInputGesture() );
+	UI_COMMAND( ResetCamera, "Reset Camera", "Resets the camera to focus on the mesh", EUserInterfaceActionType::Button, FInputChord() );
+	UI_COMMAND( EnableSimulation, "Simulation", "Enables the simulation of the blueprint and ticking", EUserInterfaceActionType::ToggleButton, FInputChord() );
+	UI_COMMAND( ShowFloor, "Show Floor", "Toggles a ground mesh for collision", EUserInterfaceActionType::ToggleButton, FInputChord() );
+	UI_COMMAND( ShowGrid, "Show Grid", "Toggles viewport grid", EUserInterfaceActionType::ToggleButton, FInputChord() );
 
 	// Debugging commands
-	UI_COMMAND( EnableAllBreakpoints,"Enable All Breakpoints", "Enable all breakpoints", EUserInterfaceActionType::Button, FInputGesture() );
-	UI_COMMAND( DisableAllBreakpoints, "Disable All Breakpoints", "Disable all breakpoints", EUserInterfaceActionType::Button, FInputGesture() );
-	UI_COMMAND( ClearAllBreakpoints, "Delete All Breakpoints", "Delete all breakpoints", EUserInterfaceActionType::Button, FInputGesture(EModifierKey::Control | EModifierKey::Shift, EKeys::F9) );
-	UI_COMMAND(	ClearAllWatches, "Delete All Watches", "Delete all watches", EUserInterfaceActionType::Button, FInputGesture() );
+	UI_COMMAND( EnableAllBreakpoints,"Enable All Breakpoints", "Enable all breakpoints", EUserInterfaceActionType::Button, FInputChord() );
+	UI_COMMAND( DisableAllBreakpoints, "Disable All Breakpoints", "Disable all breakpoints", EUserInterfaceActionType::Button, FInputChord() );
+	UI_COMMAND( ClearAllBreakpoints, "Delete All Breakpoints", "Delete all breakpoints", EUserInterfaceActionType::Button, FInputChord(EModifierKey::Control | EModifierKey::Shift, EKeys::F9) );
+	UI_COMMAND(	ClearAllWatches, "Delete All Watches", "Delete all watches", EUserInterfaceActionType::Button, FInputChord() );
 
 	// New documents
-	UI_COMMAND( AddNewVariable, "Variable", "Adds a new variable to this blueprint.", EUserInterfaceActionType::Button, FInputGesture() );
-	UI_COMMAND( AddNewLocalVariable, "Local Variable", "Adds a new local variable to this graph.", EUserInterfaceActionType::Button, FInputGesture() );
-	UI_COMMAND( AddNewFunction, "Function", "Add a new function graph", EUserInterfaceActionType::Button, FInputGesture() );
-	UI_COMMAND(	AddNewMacroDeclaration, "Macro", "Add a new macro declaration graph", EUserInterfaceActionType::Button, FInputGesture() );
-	UI_COMMAND( AddNewAnimationGraph, "Anim Graph", "Add a new animation graph", EUserInterfaceActionType::Button, FInputGesture() );
-	UI_COMMAND( AddNewEventGraph, "Graph", "Add a new event graph", EUserInterfaceActionType::Button, FInputGesture() );
-	UI_COMMAND( AddNewDelegate, "Event Dispatcher", "Add a new event dispatcher", EUserInterfaceActionType::Button, FInputGesture() );
+	UI_COMMAND( AddNewVariable, "Variable", "Adds a new variable to this blueprint.", EUserInterfaceActionType::Button, FInputChord() );
+	UI_COMMAND( AddNewLocalVariable, "Local Variable", "Adds a new local variable to this graph.", EUserInterfaceActionType::Button, FInputChord() );
+	UI_COMMAND( AddNewFunction, "Function", "Add a new function graph", EUserInterfaceActionType::Button, FInputChord() );
+	UI_COMMAND(	AddNewMacroDeclaration, "Macro", "Add a new macro declaration graph", EUserInterfaceActionType::Button, FInputChord() );
+	UI_COMMAND( AddNewAnimationGraph, "Anim Graph", "Add a new animation graph", EUserInterfaceActionType::Button, FInputChord() );
+	UI_COMMAND( AddNewEventGraph, "Graph", "Add a new event graph", EUserInterfaceActionType::Button, FInputChord() );
+	UI_COMMAND( AddNewDelegate, "Event Dispatcher", "Add a new event dispatcher", EUserInterfaceActionType::Button, FInputChord() );
 
 	// Development commands
-	UI_COMMAND( SaveIntermediateBuildProducts, "Save Intermediate Build Products", "Should the compiler save intermediate build products for debugging.", EUserInterfaceActionType::ToggleButton, FInputGesture() );
+	UI_COMMAND( SaveIntermediateBuildProducts, "Save Intermediate Build Products", "Should the compiler save intermediate build products for debugging.", EUserInterfaceActionType::ToggleButton, FInputChord() );
 
-	UI_COMMAND( RecompileGraphEditor, "Recompile Graph Editor", "Recompiles and reloads C++ code for the graph editor", EUserInterfaceActionType::Button, FInputGesture() );
-	UI_COMMAND( RecompileKismetCompiler, "Recompile Blueprint Compiler", "Recompiles and reloads C++ code for the blueprint compiler", EUserInterfaceActionType::Button, FInputGesture() );
-	UI_COMMAND( RecompileBlueprintEditor, "Recompile Blueprint Editor", "Recompiles and reloads C++ code for the blueprint editor", EUserInterfaceActionType::Button, FInputGesture() );
-	UI_COMMAND( RecompilePersona, "Recompile Persona", "Recompiles and reloads C++ code for Persona", EUserInterfaceActionType::Button, FInputGesture() );
+	UI_COMMAND( RecompileGraphEditor, "Recompile Graph Editor", "Recompiles and reloads C++ code for the graph editor", EUserInterfaceActionType::Button, FInputChord() );
+	UI_COMMAND( RecompileKismetCompiler, "Recompile Blueprint Compiler", "Recompiles and reloads C++ code for the blueprint compiler", EUserInterfaceActionType::Button, FInputChord() );
+	UI_COMMAND( RecompileBlueprintEditor, "Recompile Blueprint Editor", "Recompiles and reloads C++ code for the blueprint editor", EUserInterfaceActionType::Button, FInputChord() );
+	UI_COMMAND( RecompilePersona, "Recompile Persona", "Recompiles and reloads C++ code for Persona", EUserInterfaceActionType::Button, FInputChord() );
 
-	UI_COMMAND(GenerateNativeCode, "Generate Native Code", "Generate C++ code from the blueprint", EUserInterfaceActionType::Button, FInputGesture());
+	UI_COMMAND(GenerateNativeCode, "Generate Native Code", "Generate C++ code from the blueprint", EUserInterfaceActionType::Button, FInputChord());
 
 	// SCC commands
-	UI_COMMAND( BeginBlueprintMerge, "Merge", "Shows the Blueprint merge panel and toolbar, allowing the user to resolve conflicted blueprints", EUserInterfaceActionType::Button, FInputGesture() );
+	UI_COMMAND( BeginBlueprintMerge, "Merge", "Shows the Blueprint merge panel and toolbar, allowing the user to resolve conflicted blueprints", EUserInterfaceActionType::Button, FInputChord() );
 }
 
 PRAGMA_ENABLE_OPTIMIZATION
@@ -151,7 +151,7 @@ class FNodeSpawnInfo
 public:
 	virtual ~FNodeSpawnInfo () {}
 
-	/** Holds the UI Command to verify gestures for this action are held */
+	/** Holds the UI Command to verify chords for this action are held */
 	TSharedPtr< FUICommandInfo > CommandInfo;
 
 	/**
@@ -376,12 +376,12 @@ void FBlueprintSpawnNodeCommands::RegisterCommands()
 				FParse::Bool(*NodeSpawns[x], TEXT("Cmd="), bCmd);
 			}
 
-			FInputGesture Gesture(Key, EModifierKey::FromBools(bCtrl, bAlt, bShift, bCmd));
+			FInputChord Chord(Key, EModifierKey::FromBools(bCtrl, bAlt, bShift, bCmd));
 
 			FText CommandLabelText = FText::FromString( CommandLabel );
 			FText Description = FText::Format( NSLOCTEXT("BlueprintEditor", "NodeSpawnDescription", "Hold down the bound keys and left click in the graph panel to spawn a {0} node."), CommandLabelText );
 
-			FUICommandInfo::MakeCommandInfo( this->AsShared(), CommandInfo, FName(*NodeSpawns[x]), CommandLabelText, Description, FSlateIcon(FEditorStyle::GetStyleSetName(), *FString::Printf(TEXT("%s.%s"), *this->GetContextName().ToString(), *NodeSpawns[x])), EUserInterfaceActionType::Button, Gesture );
+			FUICommandInfo::MakeCommandInfo( this->AsShared(), CommandInfo, FName(*NodeSpawns[x]), CommandLabelText, Description, FSlateIcon(FEditorStyle::GetStyleSetName(), *FString::Printf(TEXT("%s.%s"), *this->GetContextName().ToString(), *NodeSpawns[x])), EUserInterfaceActionType::Button, Chord );
 
 			InfoPtr->CommandInfo = CommandInfo;
 
@@ -390,18 +390,18 @@ void FBlueprintSpawnNodeCommands::RegisterCommands()
 	}
 
 	TSharedPtr<FNodeSpawnInfo> AddActorRefAction = MakeShareable(new FActorRefSpawnInfo);
-	UI_COMMAND(AddActorRefAction->CommandInfo, "Add Selected Actor Reference(s)", "Spawns node(s) which reference the currently selected actor(s) in the level.", EUserInterfaceActionType::Button, FInputGesture(EKeys::R));
+	UI_COMMAND(AddActorRefAction->CommandInfo, "Add Selected Actor Reference(s)", "Spawns node(s) which reference the currently selected actor(s) in the level.", EUserInterfaceActionType::Button, FInputChord(EKeys::R));
 	NodeCommands.Add(AddActorRefAction);
 }
 
 
-void FBlueprintSpawnNodeCommands::GetGraphActionByGesture(FInputGesture& InGesture, FBlueprintPaletteListBuilder& InPaletteBuilder, UEdGraph* InDestGraph) const
+void FBlueprintSpawnNodeCommands::GetGraphActionByChord(FInputChord& InChord, FBlueprintPaletteListBuilder& InPaletteBuilder, UEdGraph* InDestGraph) const
 {
-	if(InGesture.IsValidGesture())
+	if(InChord.IsValidChord())
 	{
 		for(int32 x = 0; x < NodeCommands.Num(); ++x)
 		{
-			if(NodeCommands[x]->CommandInfo->GetActiveGesture().Get() == InGesture)
+			if(NodeCommands[x]->CommandInfo->GetActiveChord().Get() == InChord)
 			{
 				NodeCommands[x]->GetActions(InPaletteBuilder, InDestGraph);
 			}
diff --git a/Engine/Source/Editor/Kismet/Private/BlueprintEditorCommands.h b/Engine/Source/Editor/Kismet/Private/BlueprintEditorCommands.h
--- a/Engine/Source/Editor/Kismet/Private/BlueprintEditorCommands.h
+++ b/Engine/Source/Editor/Kismet/Private/BlueprintEditorCommands.h
@@ -81,20 +81,20 @@ class FBlueprintSpawnNodeCommands : public TCommands<FBlueprintSpawnNodeCommands
 public:
 
 	FBlueprintSpawnNodeCommands()
-		: TCommands<FBlueprintSpawnNodeCommands>(TEXT("BlueprintEditorSpawnNodes"), NSLOCTEXT("Contexts", "BlueprintEditor_SpawnNodes", "Blueprint Editor - Spawn Nodes by gesture"), NAME_None, FEditorStyle::GetStyleSetName())
+		: TCommands<FBlueprintSpawnNodeCommands>(TEXT("BlueprintEditorSpawnNodes"), NSLOCTEXT("Contexts", "BlueprintEditor_SpawnNodes", "Blueprint Editor - Spawn Nodes by chord"), NAME_None, FEditorStyle::GetStyleSetName())
 	{
 	}	
 
 	virtual void RegisterCommands() override;
 
 	/**
-	 * Returns a graph action assigned to the passed in gesture
+	 * Returns a graph action assigned to the passed in chord
 	 *
-	 * @param InGesture			The gesture to use for lookup
+	 * @param InChord			The chord to use for lookup
 	 * @param InPaletteBuilder	The Blueprint palette to create the graph action for, used for validation purposes and to link any important node data to the blueprint
 	 * @param InDestGraph		The graph to create the graph action for, used for validation purposes and to link any important node data to the graph
 	 */
-	void GetGraphActionByGesture(FInputGesture& InGesture, struct FBlueprintPaletteListBuilder& InPaletteBuilder, UEdGraph* InDestGraph) const;
+	void GetGraphActionByChord(FInputChord& InChord, struct FBlueprintPaletteListBuilder& InPaletteBuilder, UEdGraph* InDestGraph) const;
 
 private:
 	/** An array of all the possible commands for spawning nodes */
diff --git a/Engine/Source/Editor/Kismet/Private/SBlueprintEditorToolbar.cpp b/Engine/Source/Editor/Kismet/Private/SBlueprintEditorToolbar.cpp
--- a/Engine/Source/Editor/Kismet/Private/SBlueprintEditorToolbar.cpp
+++ b/Engine/Source/Editor/Kismet/Private/SBlueprintEditorToolbar.cpp
@@ -334,20 +334,20 @@ TSharedRef<SWidget> FKismet2Menu::MakeDiffMenu(FBlueprintEditor& Kismet)
 
 void FFullBlueprintEditorCommands::RegisterCommands() 
 {
-	UI_COMMAND(Compile, "Compile", "Compile the blueprint", EUserInterfaceActionType::Button, FInputGesture());
+	UI_COMMAND(Compile, "Compile", "Compile the blueprint", EUserInterfaceActionType::Button, FInputChord());
 	
-	UI_COMMAND(SaveOnCompile_Never, "Never", "Sets the save-on-compile option to 'Never', meaning that your Blueprints will not be saved when they are compiled", EUserInterfaceActionType::RadioButton, FInputGesture());
-	UI_COMMAND(SaveOnCompile_SuccessOnly, "On Success Only", "Sets the save-on-compile option to 'Success Only', meaning that your Blueprints will be saved whenever they are successfully compiled", EUserInterfaceActionType::RadioButton, FInputGesture());
-	UI_COMMAND(SaveOnCompile_Always, "Always", "Sets the save-on-compile option to 'Always', meaning that your Blueprints will be saved whenever they ar compiled (even if there were errors)", EUserInterfaceActionType::RadioButton, FInputGesture());
+	UI_COMMAND(SaveOnCompile_Never, "Never", "Sets the save-on-compile option to 'Never', meaning that your Blueprints will not be saved when they are compiled", EUserInterfaceActionType::RadioButton, FInputChord());
+	UI_COMMAND(SaveOnCompile_SuccessOnly, "On Success Only", "Sets the save-on-compile option to 'Success Only', meaning that your Blueprints will be saved whenever they are successfully compiled", EUserInterfaceActionType::RadioButton, FInputChord());
+	UI_COMMAND(SaveOnCompile_Always, "Always", "Sets the save-on-compile option to 'Always', meaning that your Blueprints will be saved whenever they ar compiled (even if there were errors)", EUserInterfaceActionType::RadioButton, FInputChord());
 
-	UI_COMMAND(SwitchToScriptingMode, "Graph", "Switches to Graph Editing Mode", EUserInterfaceActionType::ToggleButton, FInputGesture());
-	UI_COMMAND(SwitchToBlueprintDefaultsMode, "Defaults", "Switches to Class Defaults Mode", EUserInterfaceActionType::ToggleButton, FInputGesture());
-	UI_COMMAND(SwitchToComponentsMode, "Components", "Switches to Components Mode", EUserInterfaceActionType::ToggleButton, FInputGesture());
+	UI_COMMAND(SwitchToScriptingMode, "Graph", "Switches to Graph Editing Mode", EUserInterfaceActionType::ToggleButton, FInputChord());
+	UI_COMMAND(SwitchToBlueprintDefaultsMode, "Defaults", "Switches to Class Defaults Mode", EUserInterfaceActionType::ToggleButton, FInputChord());
+	UI_COMMAND(SwitchToComponentsMode, "Components", "Switches to Components Mode", EUserInterfaceActionType::ToggleButton, FInputChord());
 
-	UI_COMMAND(EditGlobalOptions, "Class Settings", "Edit Class Settings (Previously known as Blueprint Props)", EUserInterfaceActionType::ToggleButton, FInputGesture());
-	UI_COMMAND(EditClassDefaults, "Class Defaults", "Edit the initial values of your class.", EUserInterfaceActionType::ToggleButton, FInputGesture());
+	UI_COMMAND(EditGlobalOptions, "Class Settings", "Edit Class Settings (Previously known as Blueprint Props)", EUserInterfaceActionType::ToggleButton, FInputChord());
+	UI_COMMAND(EditClassDefaults, "Class Defaults", "Edit the initial values of your class.", EUserInterfaceActionType::ToggleButton, FInputChord());
 
-	UI_COMMAND(JumpToErrorNode, "Jump to Error Node", "When enabled, then the Blueprint will snap focus to nodes producing an error during compilation", EUserInterfaceActionType::ToggleButton, FInputGesture());
+	UI_COMMAND(JumpToErrorNode, "Jump to Error Node", "When enabled, then the Blueprint will snap focus to nodes producing an error during compilation", EUserInterfaceActionType::ToggleButton, FInputChord());
 }
 
 //////////////////////////////////////////////////////////////////////////
diff --git a/Engine/Source/Editor/Kismet/Private/SBlueprintFavoritesPalette.cpp b/Engine/Source/Editor/Kismet/Private/SBlueprintFavoritesPalette.cpp
--- a/Engine/Source/Editor/Kismet/Private/SBlueprintFavoritesPalette.cpp
+++ b/Engine/Source/Editor/Kismet/Private/SBlueprintFavoritesPalette.cpp
@@ -256,9 +256,9 @@ public:
 	/** Registers context menu commands for the blueprint favorites palette. */
 	virtual void RegisterCommands() override
 	{
-		UI_COMMAND(RemoveSingleFavorite, "Remove from Favorites",          "Removes this item from your favorites list.",                 EUserInterfaceActionType::Button, FInputGesture());
-		UI_COMMAND(RemoveSubFavorites,   "Remove Category from Favorites", "Removes all the nodes in this category from your favorites.", EUserInterfaceActionType::Button, FInputGesture());
-		UI_COMMAND(ClearFavorites,       "Clear All Favorites",			   "Clears out all of your favorited nodes.",                     EUserInterfaceActionType::Button, FInputGesture());
+		UI_COMMAND(RemoveSingleFavorite, "Remove from Favorites",          "Removes this item from your favorites list.",                 EUserInterfaceActionType::Button, FInputChord());
+		UI_COMMAND(RemoveSubFavorites,   "Remove Category from Favorites", "Removes all the nodes in this category from your favorites.", EUserInterfaceActionType::Button, FInputChord());
+		UI_COMMAND(ClearFavorites,       "Clear All Favorites",			   "Clears out all of your favorited nodes.",                     EUserInterfaceActionType::Button, FInputChord());
 	}
 };
 
diff --git a/Engine/Source/Editor/Kismet/Private/SBlueprintLibraryPalette.cpp b/Engine/Source/Editor/Kismet/Private/SBlueprintLibraryPalette.cpp
--- a/Engine/Source/Editor/Kismet/Private/SBlueprintLibraryPalette.cpp
+++ b/Engine/Source/Editor/Kismet/Private/SBlueprintLibraryPalette.cpp
@@ -151,10 +151,10 @@ public:
 	/** Registers context menu commands for the blueprint library palette. */
 	virtual void RegisterCommands() override
 	{
-		UI_COMMAND(AddSingleFavorite,    "Add to Favorites",               "Adds this item to your favorites list.",                      EUserInterfaceActionType::Button, FInputGesture());
-		UI_COMMAND(AddSubFavorites,      "Add Category to Favorites",      "Adds all the nodes in this category to your favorites.",      EUserInterfaceActionType::Button, FInputGesture());
-		UI_COMMAND(RemoveSingleFavorite, "Remove from Favorites",          "Removes this item from your favorites list.",                 EUserInterfaceActionType::Button, FInputGesture());
-		UI_COMMAND(RemoveSubFavorites,   "Remove Category from Favorites", "Removes all the nodes in this category from your favorites.", EUserInterfaceActionType::Button, FInputGesture());
+		UI_COMMAND(AddSingleFavorite,    "Add to Favorites",               "Adds this item to your favorites list.",                      EUserInterfaceActionType::Button, FInputChord());
+		UI_COMMAND(AddSubFavorites,      "Add Category to Favorites",      "Adds all the nodes in this category to your favorites.",      EUserInterfaceActionType::Button, FInputChord());
+		UI_COMMAND(RemoveSingleFavorite, "Remove from Favorites",          "Removes this item from your favorites list.",                 EUserInterfaceActionType::Button, FInputChord());
+		UI_COMMAND(RemoveSubFavorites,   "Remove Category from Favorites", "Removes all the nodes in this category from your favorites.", EUserInterfaceActionType::Button, FInputChord());
 	}
 };
 /*******************************************************************************
diff --git a/Engine/Source/Editor/Kismet/Private/SBlueprintSubPalette.cpp b/Engine/Source/Editor/Kismet/Private/SBlueprintSubPalette.cpp
--- a/Engine/Source/Editor/Kismet/Private/SBlueprintSubPalette.cpp
+++ b/Engine/Source/Editor/Kismet/Private/SBlueprintSubPalette.cpp
@@ -176,7 +176,7 @@ public:
 	 */
 	virtual void RegisterCommands() override
 	{
-		UI_COMMAND(RefreshPalette, "Refresh List", "Refreshes the list of nodes.", EUserInterfaceActionType::Button, FInputGesture());
+		UI_COMMAND(RefreshPalette, "Refresh List", "Refreshes the list of nodes.", EUserInterfaceActionType::Button, FInputChord());
 	}
 };
 
diff --git a/Engine/Source/Editor/Kismet/Private/SMyBlueprint.cpp b/Engine/Source/Editor/Kismet/Private/SMyBlueprint.cpp
--- a/Engine/Source/Editor/Kismet/Private/SMyBlueprint.cpp
+++ b/Engine/Source/Editor/Kismet/Private/SMyBlueprint.cpp
@@ -49,14 +49,14 @@
 
 void FMyBlueprintCommands::RegisterCommands() 
 {
-	UI_COMMAND( OpenGraph, "Open Graph", "Opens up this function, macro, or event graph's graph panel up.", EUserInterfaceActionType::Button, FInputGesture() );
-	UI_COMMAND( OpenGraphInNewTab, "Open in New Tab", "Opens up this function, macro, or event graph's graph panel up in a new tab. Hold down Ctrl and double click for shortcut.", EUserInterfaceActionType::Button, FInputGesture() );
-	UI_COMMAND( FocusNode, "Focus", "Focuses on the associated node", EUserInterfaceActionType::Button, FInputGesture() );
-	UI_COMMAND( FocusNodeInNewTab, "Focus in New Tab", "Focuses on the associated node in a new tab", EUserInterfaceActionType::Button, FInputGesture() );
-	UI_COMMAND( ImplementFunction, "Implement Function", "Implements this overridable function as a new function.", EUserInterfaceActionType::Button, FInputGesture() );
-	UI_COMMAND( FindEntry, "Find References", "Searches for all references of this function or variable.", EUserInterfaceActionType::Button, FInputGesture() );
-	UI_COMMAND(DeleteEntry, "Delete", "Deletes this function or variable from this blueprint.", EUserInterfaceActionType::Button, FInputGesture(EKeys::Platform_Delete));
-	UI_COMMAND( GotoNativeVarDefinition, "Goto Code Definition", "Goto the native code definition of this variable", EUserInterfaceActionType::Button, FInputGesture() );
+	UI_COMMAND( OpenGraph, "Open Graph", "Opens up this function, macro, or event graph's graph panel up.", EUserInterfaceActionType::Button, FInputChord() );
+	UI_COMMAND( OpenGraphInNewTab, "Open in New Tab", "Opens up this function, macro, or event graph's graph panel up in a new tab. Hold down Ctrl and double click for shortcut.", EUserInterfaceActionType::Button, FInputChord() );
+	UI_COMMAND( FocusNode, "Focus", "Focuses on the associated node", EUserInterfaceActionType::Button, FInputChord() );
+	UI_COMMAND( FocusNodeInNewTab, "Focus in New Tab", "Focuses on the associated node in a new tab", EUserInterfaceActionType::Button, FInputChord() );
+	UI_COMMAND( ImplementFunction, "Implement Function", "Implements this overridable function as a new function.", EUserInterfaceActionType::Button, FInputChord() );
+	UI_COMMAND( FindEntry, "Find References", "Searches for all references of this function or variable.", EUserInterfaceActionType::Button, FInputChord() );
+	UI_COMMAND(DeleteEntry, "Delete", "Deletes this function or variable from this blueprint.", EUserInterfaceActionType::Button, FInputChord(EKeys::Platform_Delete));
+	UI_COMMAND( GotoNativeVarDefinition, "Goto Code Definition", "Goto the native code definition of this variable", EUserInterfaceActionType::Button, FInputChord() );
 }
 
 //////////////////////////////////////////////////////////////////////////
diff --git a/Engine/Source/Editor/Kismet/Public/BlueprintEditor.h b/Engine/Source/Editor/Kismet/Public/BlueprintEditor.h
--- a/Engine/Source/Editor/Kismet/Public/BlueprintEditor.h
+++ b/Engine/Source/Editor/Kismet/Public/BlueprintEditor.h
@@ -459,8 +459,8 @@ public:
 	/** Find all instances of the selected custom event. */
 	void OnFindInstancesCustomEvent();
 
-	/** Handles spawning a graph node in the current graph using the passed in gesture */
-	FReply OnSpawnGraphNodeByShortcut(FInputGesture InGesture, const FVector2D& InPosition, UEdGraph* InGraph);
+	/** Handles spawning a graph node in the current graph using the passed in chord */
+	FReply OnSpawnGraphNodeByShortcut(FInputChord InChord, const FVector2D& InPosition, UEdGraph* InGraph);
 
 	/** 
 	 * Perform the actual promote to variable action on the given pin in the given blueprint.
diff --git a/Engine/Source/Editor/LandscapeEditor/Private/LandscapeEditorCommands.cpp b/Engine/Source/Editor/LandscapeEditor/Private/LandscapeEditorCommands.cpp
--- a/Engine/Source/Editor/LandscapeEditor/Private/LandscapeEditorCommands.cpp
+++ b/Engine/Source/Editor/LandscapeEditor/Private/LandscapeEditorCommands.cpp
@@ -8,83 +8,83 @@
 
 void FLandscapeEditorCommands::RegisterCommands()
 {
-	UI_COMMAND(ManageMode, "Mode - Manage", "", EUserInterfaceActionType::RadioButton, FInputGesture());
+	UI_COMMAND(ManageMode, "Mode - Manage", "", EUserInterfaceActionType::RadioButton, FInputChord());
 	NameToCommandMap.Add("ToolMode_Manage", ManageMode);
-	UI_COMMAND(SculptMode, "Mode - Sculpt", "", EUserInterfaceActionType::RadioButton, FInputGesture());
+	UI_COMMAND(SculptMode, "Mode - Sculpt", "", EUserInterfaceActionType::RadioButton, FInputChord());
 	NameToCommandMap.Add("ToolMode_Sculpt", SculptMode);
-	UI_COMMAND(PaintMode, "Mode - Paint", "", EUserInterfaceActionType::RadioButton, FInputGesture());
+	UI_COMMAND(PaintMode, "Mode - Paint", "", EUserInterfaceActionType::RadioButton, FInputChord());
 	NameToCommandMap.Add("ToolMode_Paint", PaintMode);
 
-	UI_COMMAND(NewLandscape, "Tool - New Landscape", "", EUserInterfaceActionType::RadioButton, FInputGesture());
+	UI_COMMAND(NewLandscape, "Tool - New Landscape", "", EUserInterfaceActionType::RadioButton, FInputChord());
 	NameToCommandMap.Add("Tool_NewLandscape", NewLandscape);
 
-	UI_COMMAND(ResizeLandscape, "Tool - Change Component Size", "", EUserInterfaceActionType::RadioButton, FInputGesture());
+	UI_COMMAND(ResizeLandscape, "Tool - Change Component Size", "", EUserInterfaceActionType::RadioButton, FInputChord());
 	NameToCommandMap.Add("Tool_ResizeLandscape", ResizeLandscape);
 
-	UI_COMMAND(SculptTool, "Tool - Sculpt", "", EUserInterfaceActionType::RadioButton, FInputGesture());
+	UI_COMMAND(SculptTool, "Tool - Sculpt", "", EUserInterfaceActionType::RadioButton, FInputChord());
 	NameToCommandMap.Add("Tool_Sculpt", SculptTool);
-	UI_COMMAND(PaintTool, "Tool - Paint", "", EUserInterfaceActionType::RadioButton, FInputGesture());
+	UI_COMMAND(PaintTool, "Tool - Paint", "", EUserInterfaceActionType::RadioButton, FInputChord());
 	NameToCommandMap.Add("Tool_Paint", PaintTool);
-	UI_COMMAND(SmoothTool, "Tool - Smooth", "", EUserInterfaceActionType::RadioButton, FInputGesture());
+	UI_COMMAND(SmoothTool, "Tool - Smooth", "", EUserInterfaceActionType::RadioButton, FInputChord());
 	NameToCommandMap.Add("Tool_Smooth", SmoothTool);
-	UI_COMMAND(FlattenTool, "Tool - Flatten", "", EUserInterfaceActionType::RadioButton, FInputGesture());
+	UI_COMMAND(FlattenTool, "Tool - Flatten", "", EUserInterfaceActionType::RadioButton, FInputChord());
 	NameToCommandMap.Add("Tool_Flatten", FlattenTool);
-	UI_COMMAND(RampTool, "Tool - Ramp", "", EUserInterfaceActionType::RadioButton, FInputGesture());
+	UI_COMMAND(RampTool, "Tool - Ramp", "", EUserInterfaceActionType::RadioButton, FInputChord());
 	NameToCommandMap.Add("Tool_Ramp", RampTool);
-	UI_COMMAND(ErosionTool, "Tool - Erosion", "", EUserInterfaceActionType::RadioButton, FInputGesture());
+	UI_COMMAND(ErosionTool, "Tool - Erosion", "", EUserInterfaceActionType::RadioButton, FInputChord());
 	NameToCommandMap.Add("Tool_Erosion", ErosionTool);
-	UI_COMMAND(HydroErosionTool, "Tool - Hydraulic Erosion", "", EUserInterfaceActionType::RadioButton, FInputGesture());
+	UI_COMMAND(HydroErosionTool, "Tool - Hydraulic Erosion", "", EUserInterfaceActionType::RadioButton, FInputChord());
 	NameToCommandMap.Add("Tool_HydraErosion", HydroErosionTool);
-	UI_COMMAND(NoiseTool, "Tool - Noise", "", EUserInterfaceActionType::RadioButton, FInputGesture());
+	UI_COMMAND(NoiseTool, "Tool - Noise", "", EUserInterfaceActionType::RadioButton, FInputChord());
 	NameToCommandMap.Add("Tool_Noise", NoiseTool);
-	UI_COMMAND(RetopologizeTool, "Tool - Retopologize", "", EUserInterfaceActionType::RadioButton, FInputGesture());
+	UI_COMMAND(RetopologizeTool, "Tool - Retopologize", "", EUserInterfaceActionType::RadioButton, FInputChord());
 	NameToCommandMap.Add("Tool_Retopologize", RetopologizeTool);
-	UI_COMMAND(VisibilityTool, "Tool - Visibility", "", EUserInterfaceActionType::RadioButton, FInputGesture());
+	UI_COMMAND(VisibilityTool, "Tool - Visibility", "", EUserInterfaceActionType::RadioButton, FInputChord());
 	NameToCommandMap.Add("Tool_Visibility", VisibilityTool);
 
-	UI_COMMAND(SelectComponentTool, "Tool - Component Selection", "", EUserInterfaceActionType::RadioButton, FInputGesture());
+	UI_COMMAND(SelectComponentTool, "Tool - Component Selection", "", EUserInterfaceActionType::RadioButton, FInputChord());
 	NameToCommandMap.Add("Tool_Select", SelectComponentTool);
-	UI_COMMAND(AddComponentTool, "Tool - Add Components", "", EUserInterfaceActionType::RadioButton, FInputGesture());
+	UI_COMMAND(AddComponentTool, "Tool - Add Components", "", EUserInterfaceActionType::RadioButton, FInputChord());
 	NameToCommandMap.Add("Tool_AddComponent", AddComponentTool);
-	UI_COMMAND(DeleteComponentTool, "Tool - Delete Components", "", EUserInterfaceActionType::RadioButton, FInputGesture());
+	UI_COMMAND(DeleteComponentTool, "Tool - Delete Components", "", EUserInterfaceActionType::RadioButton, FInputChord());
 	NameToCommandMap.Add("Tool_DeleteComponent", DeleteComponentTool);
-	UI_COMMAND(MoveToLevelTool, "Tool - Move to Level", "Moves the selected landscape components to the current streaming level", EUserInterfaceActionType::RadioButton, FInputGesture());
+	UI_COMMAND(MoveToLevelTool, "Tool - Move to Level", "Moves the selected landscape components to the current streaming level", EUserInterfaceActionType::RadioButton, FInputChord());
 	NameToCommandMap.Add("Tool_MoveToLevel", MoveToLevelTool);
 
-	UI_COMMAND(RegionSelectTool, "Tool - Region Selection", "", EUserInterfaceActionType::RadioButton, FInputGesture());
+	UI_COMMAND(RegionSelectTool, "Tool - Region Selection", "", EUserInterfaceActionType::RadioButton, FInputChord());
 	NameToCommandMap.Add("Tool_Mask", RegionSelectTool);
-	UI_COMMAND(RegionCopyPasteTool, "Tool - Copy/Paste", "", EUserInterfaceActionType::RadioButton, FInputGesture());
+	UI_COMMAND(RegionCopyPasteTool, "Tool - Copy/Paste", "", EUserInterfaceActionType::RadioButton, FInputChord());
 	NameToCommandMap.Add("Tool_CopyPaste", RegionCopyPasteTool);
 
-	UI_COMMAND(SplineTool, "Tool - Edit Splines", "", EUserInterfaceActionType::RadioButton, FInputGesture());
+	UI_COMMAND(SplineTool, "Tool - Edit Splines", "", EUserInterfaceActionType::RadioButton, FInputChord());
 	NameToCommandMap.Add("Tool_Splines", SplineTool);
 
-	UI_COMMAND(CircleBrush, "Brush - Circle", "", EUserInterfaceActionType::RadioButton, FInputGesture());
+	UI_COMMAND(CircleBrush, "Brush - Circle", "", EUserInterfaceActionType::RadioButton, FInputChord());
 	NameToCommandMap.Add("BrushSet_Circle", CircleBrush);
-	UI_COMMAND(AlphaBrush, "Brush - Alpha", "", EUserInterfaceActionType::RadioButton, FInputGesture());
+	UI_COMMAND(AlphaBrush, "Brush - Alpha", "", EUserInterfaceActionType::RadioButton, FInputChord());
 	NameToCommandMap.Add("BrushSet_Alpha", AlphaBrush);
-	UI_COMMAND(AlphaBrush_Pattern, "Brush - Pattern", "", EUserInterfaceActionType::RadioButton, FInputGesture());
+	UI_COMMAND(AlphaBrush_Pattern, "Brush - Pattern", "", EUserInterfaceActionType::RadioButton, FInputChord());
 	NameToCommandMap.Add("BrushSet_Pattern", AlphaBrush_Pattern);
-	UI_COMMAND(ComponentBrush, "Brush - Component", "", EUserInterfaceActionType::RadioButton, FInputGesture());
+	UI_COMMAND(ComponentBrush, "Brush - Component", "", EUserInterfaceActionType::RadioButton, FInputChord());
 	NameToCommandMap.Add("BrushSet_Component", ComponentBrush);
-	UI_COMMAND(GizmoBrush, "Brush - Gizmo", "", EUserInterfaceActionType::RadioButton, FInputGesture());
+	UI_COMMAND(GizmoBrush, "Brush - Gizmo", "", EUserInterfaceActionType::RadioButton, FInputChord());
 	NameToCommandMap.Add("BrushSet_Gizmo", GizmoBrush);
 	NameToCommandMap.Add("BrushSet_Splines", SplineTool);
 
-	UI_COMMAND(CircleBrush_Smooth, "Circle Brush - Smooth Falloff", "", EUserInterfaceActionType::RadioButton, FInputGesture());
+	UI_COMMAND(CircleBrush_Smooth, "Circle Brush - Smooth Falloff", "", EUserInterfaceActionType::RadioButton, FInputChord());
 	NameToCommandMap.Add("Circle_Smooth", CircleBrush_Smooth);
-	UI_COMMAND(CircleBrush_Linear, "Circle Brush - Linear Falloff", "", EUserInterfaceActionType::RadioButton, FInputGesture());
+	UI_COMMAND(CircleBrush_Linear, "Circle Brush - Linear Falloff", "", EUserInterfaceActionType::RadioButton, FInputChord());
 	NameToCommandMap.Add("Circle_Linear", CircleBrush_Linear);
-	UI_COMMAND(CircleBrush_Spherical, "Circle Brush - Spherical Falloff", "", EUserInterfaceActionType::RadioButton, FInputGesture());
+	UI_COMMAND(CircleBrush_Spherical, "Circle Brush - Spherical Falloff", "", EUserInterfaceActionType::RadioButton, FInputChord());
 	NameToCommandMap.Add("Circle_Spherical", CircleBrush_Spherical);
-	UI_COMMAND(CircleBrush_Tip, "Circle Brush - Tip Falloff", "", EUserInterfaceActionType::RadioButton, FInputGesture());
+	UI_COMMAND(CircleBrush_Tip, "Circle Brush - Tip Falloff", "", EUserInterfaceActionType::RadioButton, FInputChord());
 	NameToCommandMap.Add("Circle_Tip", CircleBrush_Tip);
 
-	UI_COMMAND(ViewModeNormal, "Normal", "", EUserInterfaceActionType::RadioButton, FInputGesture());
-	UI_COMMAND(ViewModeLOD, "LOD", "", EUserInterfaceActionType::RadioButton, FInputGesture());
-	UI_COMMAND(ViewModeLayerDensity, "Layer Density", "", EUserInterfaceActionType::RadioButton, FInputGesture());
-	UI_COMMAND(ViewModeLayerDebug, "Layer Debug", "", EUserInterfaceActionType::RadioButton, FInputGesture());
-	UI_COMMAND(ViewModeWireframeOnTop, "Wireframe on Top", "", EUserInterfaceActionType::RadioButton, FInputGesture());
+	UI_COMMAND(ViewModeNormal, "Normal", "", EUserInterfaceActionType::RadioButton, FInputChord());
+	UI_COMMAND(ViewModeLOD, "LOD", "", EUserInterfaceActionType::RadioButton, FInputChord());
+	UI_COMMAND(ViewModeLayerDensity, "Layer Density", "", EUserInterfaceActionType::RadioButton, FInputChord());
+	UI_COMMAND(ViewModeLayerDebug, "Layer Debug", "", EUserInterfaceActionType::RadioButton, FInputChord());
+	UI_COMMAND(ViewModeWireframeOnTop, "Wireframe on Top", "", EUserInterfaceActionType::RadioButton, FInputChord());
 }
 
 #undef LOCTEXT_NAMESPACE
\ No newline at end of file
diff --git a/Engine/Source/Editor/Layers/Private/LayerCollectionViewCommands.h b/Engine/Source/Editor/Layers/Private/LayerCollectionViewCommands.h
--- a/Engine/Source/Editor/Layers/Private/LayerCollectionViewCommands.h
+++ b/Engine/Source/Editor/Layers/Private/LayerCollectionViewCommands.h
@@ -28,19 +28,19 @@ public:
 	 */
 	virtual void RegisterCommands() override
 	{
-		UI_COMMAND( CreateEmptyLayer, "Create Empty Layer", "Creates a new empty Layer", EUserInterfaceActionType::Button, FInputGesture() );
-		UI_COMMAND( AddSelectedActorsToNewLayer, "Add Selected Actors to New Layer", "Adds the actors currently selected in the active viewport to a new layer", EUserInterfaceActionType::Button, FInputGesture() );
-		UI_COMMAND( AddSelectedActorsToSelectedLayer, "Add Selected Actors to Selected Layers", "Adds the actors currently selected in the viewport to the selected layers", EUserInterfaceActionType::Button, FInputGesture() );
-		UI_COMMAND( RemoveSelectedActorsFromSelectedLayer, "Remove Selected Actors from Layers", "Removes the actors currently selected in the viewport from the selected layers", EUserInterfaceActionType::Button, FInputGesture() );
+		UI_COMMAND( CreateEmptyLayer, "Create Empty Layer", "Creates a new empty Layer", EUserInterfaceActionType::Button, FInputChord() );
+		UI_COMMAND( AddSelectedActorsToNewLayer, "Add Selected Actors to New Layer", "Adds the actors currently selected in the active viewport to a new layer", EUserInterfaceActionType::Button, FInputChord() );
+		UI_COMMAND( AddSelectedActorsToSelectedLayer, "Add Selected Actors to Selected Layers", "Adds the actors currently selected in the viewport to the selected layers", EUserInterfaceActionType::Button, FInputChord() );
+		UI_COMMAND( RemoveSelectedActorsFromSelectedLayer, "Remove Selected Actors from Layers", "Removes the actors currently selected in the viewport from the selected layers", EUserInterfaceActionType::Button, FInputChord() );
 
-		UI_COMMAND( SelectActors, "Select Actors", "Actors in the selected Layers becomes the viewport's selection", EUserInterfaceActionType::Button, FInputGesture() );
-		UI_COMMAND( AppendActorsToSelection, "Append Actors to Selection", "Adds the Actors in the selected Layers to the viewport's existing selection", EUserInterfaceActionType::Button, FInputGesture() );
-		UI_COMMAND( DeselectActors, "Deselect Actors", "Removes the Actors in the selected Layers from the viewport's existing selection", EUserInterfaceActionType::Button, FInputGesture() );
+		UI_COMMAND( SelectActors, "Select Actors", "Actors in the selected Layers becomes the viewport's selection", EUserInterfaceActionType::Button, FInputChord() );
+		UI_COMMAND( AppendActorsToSelection, "Append Actors to Selection", "Adds the Actors in the selected Layers to the viewport's existing selection", EUserInterfaceActionType::Button, FInputChord() );
+		UI_COMMAND( DeselectActors, "Deselect Actors", "Removes the Actors in the selected Layers from the viewport's existing selection", EUserInterfaceActionType::Button, FInputChord() );
 
-		UI_COMMAND( ToggleSelectedLayersVisibility, "Toggle Visibility of Selected Layers", "Toggles the visibility of the selected layers in the viewports", EUserInterfaceActionType::Button, FInputGesture() );
-		UI_COMMAND( MakeAllLayersVisible, "Make All Layers Visible", "Toggles all layers to a visible state in the viewports", EUserInterfaceActionType::Button, FInputGesture() );
+		UI_COMMAND( ToggleSelectedLayersVisibility, "Toggle Visibility of Selected Layers", "Toggles the visibility of the selected layers in the viewports", EUserInterfaceActionType::Button, FInputChord() );
+		UI_COMMAND( MakeAllLayersVisible, "Make All Layers Visible", "Toggles all layers to a visible state in the viewports", EUserInterfaceActionType::Button, FInputChord() );
 	
-		UI_COMMAND( RequestRenameLayer, "Rename", "Rename the selected layer.", EUserInterfaceActionType::Button, FInputGesture( EKeys::F2 ) );	
+		UI_COMMAND( RequestRenameLayer, "Rename", "Rename the selected layer.", EUserInterfaceActionType::Button, FInputChord( EKeys::F2 ) );	
 	}
 
 public:
