[UE-8610] - IWCE: The first time you add a component it clobbers the DefaultSceneRoot
First SceneComponent added no longer clobbers DefaultSceneRoot.
DefaultSceneRoot can be clobbered by dragging another SceneComponent onto it or by deleting it.
Changed tooltips to explain these actions to users.
Also, changed the way children are reparented when a node is deleted. If possible, children are copied to removed node's parent instead of promoting one child to replace the removed node.

diff --git a/Engine/Source/Editor/Kismet/Private/SSCSEditor.cpp b/Engine/Source/Editor/Kismet/Private/SSCSEditor.cpp
--- a/Engine/Source/Editor/Kismet/Private/SSCSEditor.cpp
+++ b/Engine/Source/Editor/Kismet/Private/SSCSEditor.cpp
@@ -1798,11 +1798,21 @@ TSharedPtr<SWidget> SSCS_RowWidget::BuildSceneRootDropActionMenu(FSCSEditorTreeN
 			FUIAction(
 				FExecuteAction::CreateSP(this, &SSCS_RowWidget::OnAttachToDropAction, DroppedNodePtr),
 				FCanExecuteAction()));
+
+		FSCSEditorTreeNodePtrType NodePtr = GetNode();
+		const bool bIsDefaultSceneRoot = NodePtr->IsDefaultSceneRoot();
+
+		FText NewRootNodeText = bIsDefaultSceneRoot
+			? FText::Format(LOCTEXT("DropActionToolTip_MakeNewRootNodeAndDelete", "Make {0} the new root. The default root will be deleted."), DroppedVariableNameText)
+			: FText::Format(LOCTEXT("DropActionToolTip_MakeNewRootNode", "Make {0} the new root."), DroppedVariableNameText);
+
+		FText NewRootNodeFromCopyText = bIsDefaultSceneRoot
+			? FText::Format(LOCTEXT("DropActionToolTip_MakeNewRootNodeFromCopyAndDelete", "Copy {0} to a new variable and make it the new root. The default root will be deleted."), DroppedVariableNameText)
+			: FText::Format(LOCTEXT("DropActionToolTip_MakeNewRootNodeFromCopy", "Copy {0} to a new variable and make it the new root."), DroppedVariableNameText);
+
 		MenuBuilder.AddMenuEntry(
 			LOCTEXT("DropActionLabel_MakeNewRootNode", "Make New Root"),
-			bDroppedInSameBlueprint
-			? FText::Format( LOCTEXT("DropActionToolTip_MakeNewRootNode", "Make {0} the new root."), DroppedVariableNameText )
-			: FText::Format( LOCTEXT("DropActionToolTip_MakeNewRootNodeFromCopy", "Copy {0} to a new variable and make it the new root."), DroppedVariableNameText ),
+			bDroppedInSameBlueprint ? NewRootNodeText : NewRootNodeFromCopyText,
 			FSlateIcon(),
 			FUIAction(
 				FExecuteAction::CreateSP(this, &SSCS_RowWidget::OnMakeNewRootDropAction, DroppedNodePtr),
@@ -1967,8 +1977,7 @@ void SSCS_RowWidget::OnDragEnter( const FGeometry& MyGeometry, const FDragDropEv
 				else if (NodePtr == SceneRootNodePtr)
 				{
 					bool bCanMakeNewRoot = false;
-					bool bCanAttachToRoot = !NodePtr->IsDefaultSceneRoot()
-						&& !DraggedNodePtr->IsDirectlyAttachedTo(NodePtr)
+					bool bCanAttachToRoot = !DraggedNodePtr->IsDirectlyAttachedTo(NodePtr)
 						&& HoveredTemplate->CanAttachAsChild(DraggedTemplate, NAME_None)
 						&& DraggedTemplate->Mobility >= HoveredTemplate->Mobility
 						&& (!HoveredTemplate->IsEditorOnly() || DraggedTemplate->IsEditorOnly());
@@ -2007,8 +2016,17 @@ void SSCS_RowWidget::OnDragEnter( const FGeometry& MyGeometry, const FDragDropEv
 					{
 						if (bCanMakeNewRoot)
 						{
-							// Only available action is to copy the dragged node to the other Blueprint and make it the new root
-							Message = FText::Format(LOCTEXT("DropActionToolTip_DropMakeNewRootNodeFromCopy", "Drop here to copy {0} to a new variable and make it the new root."), DraggedNodePtr->GetDisplayName());
+							if (NodePtr->IsDefaultSceneRoot())
+							{
+								// Only available action is to copy the dragged node to the other Blueprint and make it the new root
+								// Default root will be deleted
+								Message = FText::Format(LOCTEXT("DropActionToolTip_DropMakeNewRootNodeFromCopyAndDelete", "Drop here to copy {0} to a new variable and make it the new root. The default root will be deleted."), DraggedNodePtr->GetDisplayName());
+							}
+							else
+							{
+								// Only available action is to copy the dragged node to the other Blueprint and make it the new root
+								Message = FText::Format(LOCTEXT("DropActionToolTip_DropMakeNewRootNodeFromCopy", "Drop here to copy {0} to a new variable and make it the new root."), DraggedNodePtr->GetDisplayName());
+							}
 							DragRowOp->PendingDropAction = FSCSRowDragDropOp::DropAction_MakeNewRoot;
 						}
 						else if (bCanAttachToRoot)
@@ -2028,8 +2046,17 @@ void SSCS_RowWidget::OnDragEnter( const FGeometry& MyGeometry, const FDragDropEv
 					}
 					else if (bCanMakeNewRoot)
 					{
-						// Only available action is to make the dragged node the new root
-						Message = FText::Format(LOCTEXT("DropActionToolTip_DropMakeNewRootNode", "Drop here to make {0} the new root."), DraggedNodePtr->GetDisplayName());
+						if (NodePtr->IsDefaultSceneRoot())
+						{
+							// Only available action is to make the dragged node the new root
+							// Default root will be deleted
+							Message = FText::Format(LOCTEXT("DropActionToolTip_DropMakeNewRootNodeAndDelete", "Drop here to make {0} the new root. The default root will be deleted."), DraggedNodePtr->GetDisplayName());
+						}
+						else
+						{
+							// Only available action is to make the dragged node the new root
+							Message = FText::Format(LOCTEXT("DropActionToolTip_DropMakeNewRootNode", "Drop here to make {0} the new root."), DraggedNodePtr->GetDisplayName());
+						}
 						DragRowOp->PendingDropAction = FSCSRowDragDropOp::DropAction_MakeNewRoot;
 					}
 					else if (bCanAttachToRoot)
@@ -2489,15 +2516,17 @@ void SSCS_RowWidget::OnMakeNewRootDropAction(FSCSEditorTreeNodePtrType DroppedNo
 	// Create a transaction record
 	const FScopedTransaction TransactionContext(LOCTEXT("MakeNewSceneRoot", "Make New Scene Root"));
 
+	FSCSEditorTreeNodePtrType OldSceneRootNodePtr;
+
+	// Remember whether or not we're replacing the default scene root
+	bool bWasDefaultSceneRoot = SceneRootNodePtr.IsValid() && SceneRootNodePtr->IsDefaultSceneRoot();
+
 	if (SCSEditorPtr->GetEditorMode() == EComponentEditorMode::BlueprintSCS)
 	{
 		// Get the current Blueprint context
 		UBlueprint* Blueprint = GetBlueprint();
 		check(Blueprint != NULL && Blueprint->SimpleConstructionScript != nullptr);
 
-		// Remember whether or not we're replacing the default scene root
-		bool bWasDefaultSceneRoot = SceneRootNodePtr.IsValid() && SceneRootNodePtr->IsDefaultSceneRoot();
-
 		// Clone the component if it's being dropped into a different SCS
 		if(DroppedNodePtr->GetBlueprint() != Blueprint)
 		{
@@ -2546,20 +2575,13 @@ void SSCS_RowWidget::OnMakeNewRootDropAction(FSCSEditorTreeNodePtrType DroppedNo
 			}
 		}
 
-		if(!bWasDefaultSceneRoot)
-		{
-			check(SceneRootNodePtr->CanReparent());
+		check(bWasDefaultSceneRoot || SceneRootNodePtr->CanReparent());
 
-			// Remove the current scene root node from the SCS context
-			Blueprint->SimpleConstructionScript->RemoveNode(SceneRootNodePtr->GetSCSNode());
-		}
+		// Remove the current scene root node from the SCS context
+		Blueprint->SimpleConstructionScript->RemoveNode(SceneRootNodePtr->GetSCSNode());
 
 		// Save old root node
-		FSCSEditorTreeNodePtrType OldSceneRootNodePtr;
-		if(!bWasDefaultSceneRoot)
-		{
-			OldSceneRootNodePtr = SceneRootNodePtr;
-		}
+		OldSceneRootNodePtr = SceneRootNodePtr;
 
 		// Set node we are dropping as new root
 		SceneRootNodePtr = DroppedNodePtr;
@@ -2567,13 +2589,21 @@ void SSCS_RowWidget::OnMakeNewRootDropAction(FSCSEditorTreeNodePtrType DroppedNo
 		// Add dropped node to the SCS context
 		Blueprint->SimpleConstructionScript->AddNode(SceneRootNodePtr->GetSCSNode());
 
-		// Set old root as child of new root
-		if(OldSceneRootNodePtr.IsValid())
+		// Remove or re-parent the old root
+		if (OldSceneRootNodePtr.IsValid())
 		{
+			check(SceneRootNodePtr->CanReparent());
+
+			// Set old root as child of new root
 			SceneRootNodePtr->AddChild(OldSceneRootNodePtr);
 
 			// Expand the new scene root as we've just added a child to it
 			SCSEditorPtr->SetNodeExpansionState(SceneRootNodePtr, true);
+
+			if (bWasDefaultSceneRoot)
+			{
+				SCSEditorPtr->RemoveComponentNode(OldSceneRootNodePtr);
+			}
 		}
 	}
 	else    // EComponentEditorMode::ActorInstance
@@ -2584,20 +2614,30 @@ void SSCS_RowWidget::OnMakeNewRootDropAction(FSCSEditorTreeNodePtrType DroppedNo
 			DroppedNodePtr->GetParent()->RemoveChild(DroppedNodePtr);
 		}
 
-		check(SceneRootNodePtr->CanReparent());
-
 		// Save old root node
-		FSCSEditorTreeNodePtrType OldSceneRootNodePtr = SceneRootNodePtr;
+		OldSceneRootNodePtr = SceneRootNodePtr;
 
 		// Set node we are dropping as new root
 		SceneRootNodePtr = DroppedNodePtr;
 
-		// Set old root as child of new root
-		check(OldSceneRootNodePtr.IsValid());
-		SceneRootNodePtr->AddChild(OldSceneRootNodePtr);
+		// Remove or re-parent the old root
+		if (OldSceneRootNodePtr.IsValid())
+		{
+			if (bWasDefaultSceneRoot)
+			{
+				SCSEditorPtr->RemoveComponentNode(OldSceneRootNodePtr);
+			}
+			else
+			{
+				check(SceneRootNodePtr->CanReparent());
+
+				// Set old root as child of new root
+				SceneRootNodePtr->AddChild(OldSceneRootNodePtr);
 
-		// Expand the new scene root as we've just added a child to it
-		SCSEditorPtr->SetNodeExpansionState(SceneRootNodePtr, true);
+				// Expand the new scene root as we've just added a child to it
+				SCSEditorPtr->SetNodeExpansionState(SceneRootNodePtr, true);
+			}
+		}
 	}
 
 	PostDragDropAction(true);
@@ -4515,10 +4555,10 @@ UActorComponent* SSCSEditor::AddNewNode(USCS_Node* NewNode, UObject* Asset, bool
 	check(Blueprint != nullptr && Blueprint->SimpleConstructionScript != nullptr);
 
 	// Reset the scene root node if it's set to the default one that's managed by the SCS
-	if(SceneRootNodePtr.IsValid() && SceneRootNodePtr->GetSCSNode() == Blueprint->SimpleConstructionScript->GetDefaultSceneRootNode())
-	{
-		SceneRootNodePtr.Reset();
-	}
+	//if(SceneRootNodePtr.IsValid() && SceneRootNodePtr->GetSCSNode() == Blueprint->SimpleConstructionScript->GetDefaultSceneRootNode())
+	//{
+	//	SceneRootNodePtr.Reset();
+	//}
 
 	// Add the new node to the editor tree
 	NewNodePtr = AddTreeNode(NewNode, SceneRootNodePtr, false);
@@ -4562,12 +4602,12 @@ UActorComponent* SSCSEditor::AddNewNodeForInstancedComponent(UActorComponent* Ne
 		NewNodePtr = AddTreeNodeFromComponent(NewSceneComponent);
 
 		// Remove the old scene root node if it's set to the default one
-		if(SceneRootNodePtr.IsValid() && SceneRootNodePtr->IsDefaultSceneRoot())
-		{
-			RemoveComponentNode(SceneRootNodePtr);
-			RootNodes.Remove( SceneRootNodePtr );
-			SceneRootNodePtr.Reset();
-		}
+		//if(SceneRootNodePtr.IsValid() && SceneRootNodePtr->IsDefaultSceneRoot())
+		//{
+		//	RemoveComponentNode(SceneRootNodePtr);
+		//	RootNodes.Remove( SceneRootNodePtr );
+		//	SceneRootNodePtr.Reset();
+		//}
 	}
 	else
 	{
diff --git a/Engine/Source/Runtime/Engine/Classes/Engine/SimpleConstructionScript.h b/Engine/Source/Runtime/Engine/Classes/Engine/SimpleConstructionScript.h
--- a/Engine/Source/Runtime/Engine/Classes/Engine/SimpleConstructionScript.h
+++ b/Engine/Source/Runtime/Engine/Classes/Engine/SimpleConstructionScript.h
@@ -57,8 +57,8 @@ class USimpleConstructionScript : public UObject
 	/** Remove this node from the script (will take all its children with it) */
 	ENGINE_API void RemoveNode(USCS_Node* Node);
 
-	/** Remove this node from the script and promote its first child to replace it */
-	ENGINE_API USCS_Node* RemoveNodeAndPromoteChildren(USCS_Node* Node);
+	/** Remove this node from the script and if it's the root, promote its first child to replace it */
+	ENGINE_API void RemoveNodeAndPromoteChildren(USCS_Node* Node);
 
 	/** Find the parent node of this one. Returns NULL if node is not in tree or if is root */
 	ENGINE_API USCS_Node* FindParentNode(USCS_Node* InNode) const;
diff --git a/Engine/Source/Runtime/Engine/Private/SimpleConstructionScript.cpp b/Engine/Source/Runtime/Engine/Private/SimpleConstructionScript.cpp
--- a/Engine/Source/Runtime/Engine/Private/SimpleConstructionScript.cpp
+++ b/Engine/Source/Runtime/Engine/Private/SimpleConstructionScript.cpp
@@ -495,39 +495,39 @@ void USimpleConstructionScript::RemoveNode(USCS_Node* Node)
 	}
 }
 
-USCS_Node* USimpleConstructionScript::RemoveNodeAndPromoteChildren(USCS_Node* Node)
+void USimpleConstructionScript::RemoveNodeAndPromoteChildren(USCS_Node* Node)
 {
 	Node->Modify();
-	   
-	USCS_Node* ChildToPromote = NULL;
 
-	// Pick the first child to promote as a new 'root'
-	if (Node->ChildNodes.Num() > 0)
+	if (RootNodes.Contains(Node))
 	{
-		int32 PromoteIndex = 0;
-		ChildToPromote = Node->ChildNodes[PromoteIndex];
+		USCS_Node* ChildToPromote = NULL;
 
-		// if this is an editor-only component, then it can't have any game-component children (better make sure that's the case)
-		if (ChildToPromote->ComponentTemplate != NULL && ChildToPromote->ComponentTemplate->IsEditorOnly())
+		// Pick the first child to promote as a new 'root'
+		if (Node->ChildNodes.Num() > 0)
 		{
-			for (int32 ChildIndex = 1; ChildIndex < Node->ChildNodes.Num(); ++ChildIndex)
+			int32 PromoteIndex = 0;
+			ChildToPromote = Node->ChildNodes[PromoteIndex];
+
+			// if this is an editor-only component, then it can't have any game-component children (better make sure that's the case)
+			if (ChildToPromote->ComponentTemplate != NULL && ChildToPromote->ComponentTemplate->IsEditorOnly())
 			{
-				USCS_Node* Child = Node->ChildNodes[ChildIndex];
-				// we found a game-component sibling, better make it the child to promote
-				if (Child->ComponentTemplate != NULL && !Child->ComponentTemplate->IsEditorOnly())
+				for (int32 ChildIndex = 1; ChildIndex < Node->ChildNodes.Num(); ++ChildIndex)
 				{
-					ChildToPromote = Child;
-					PromoteIndex = ChildIndex;
-					break;
+					USCS_Node* Child = Node->ChildNodes[ChildIndex];
+					// we found a game-component sibling, better make it the child to promote
+					if (Child->ComponentTemplate != NULL && !Child->ComponentTemplate->IsEditorOnly())
+					{
+						ChildToPromote = Child;
+						PromoteIndex = ChildIndex;
+						break;
+					}
 				}
 			}
-		}
 
-		Node->ChildNodes.RemoveAt(PromoteIndex);
-	}
+			Node->ChildNodes.RemoveAt(PromoteIndex);
+		}
 
-	if (RootNodes.Contains(Node))
-	{
 		Modify();
 
 		if(ChildToPromote != NULL)
@@ -558,23 +558,14 @@ USCS_Node* USimpleConstructionScript::RemoveNodeAndPromoteChildren(USCS_Node* No
 
 		ParentNode->Modify();
 
-		if ( ChildToPromote != NULL )
-		{
-			ChildToPromote->Modify();
-
-			// Insert promoted node next to node being removed.
-			int32 Location = ParentNode->ChildNodes.Find(Node);
-			ParentNode->ChildNodes.Insert(ChildToPromote,Location);
-			ChildToPromote->ChildNodes.Append(Node->ChildNodes);	
-		}
-		// remove node
+		// remove node and move children onto parent
+		int32 Location = ParentNode->ChildNodes.Find(Node);
 		ParentNode->ChildNodes.Remove(Node);
+		ParentNode->ChildNodes.Insert(Node->ChildNodes, Location);
 	}
 
-	// Clear out references to promoted children
+	// Clear out references to previous children
 	Node->ChildNodes.Empty();
-
-	return ChildToPromote;
 }
 
 
