- Drag-drop indicators take into account the tint on the indicator brush
- Added basic drag-drop outline indicators to the SCS tree
- Removed redundant definition of SlateEditorStyle::NormalTableRowStyle & added the basic drag-drop indicators to the normal style

diff --git a/Engine/Source/Editor/EditorStyle/Private/SlateEditorStyle.cpp b/Engine/Source/Editor/EditorStyle/Private/SlateEditorStyle.cpp
--- a/Engine/Source/Editor/EditorStyle/Private/SlateEditorStyle.cpp
+++ b/Engine/Source/Editor/EditorStyle/Private/SlateEditorStyle.cpp
@@ -178,7 +178,10 @@ void FSlateEditorStyle::FStyle::SetupGeneralStyles()
 		.SetInactiveBrush( IMAGE_BRUSH( "Common/Selection", Icon8x8, SelectionColor_Inactive ) )
 		.SetInactiveHoveredBrush( IMAGE_BRUSH( "Common/Selection", Icon8x8, SelectionColor_Inactive ) )
 		.SetTextColor( DefaultForeground )
-		.SetSelectedTextColor( InvertedForeground );
+		.SetSelectedTextColor(InvertedForeground)
+		.SetDropIndicator_Above(BOX_BRUSH("Common/DropZoneIndicator_Above", FMargin(10.0f / 16.0f, 10.0f / 16.0f, 0, 0), SelectionColor))
+		.SetDropIndicator_Onto(BOX_BRUSH("Common/DropZoneIndicator_Onto", FMargin(4.0f / 16.0f), SelectionColor))
+		.SetDropIndicator_Below(BOX_BRUSH("Common/DropZoneIndicator_Below", FMargin(10.0f / 16.0f, 0, 0, 10.0f / 16.0f), SelectionColor));
 
 	// Normal Text
 	{
@@ -822,20 +825,6 @@ void FSlateEditorStyle::FStyle::SetupGeneralStyles()
 			);
 	}//
 
-
-	NormalTableRowStyle = FTableRowStyle()
-		.SetEvenRowBackgroundBrush(FSlateNoResource())
-		.SetEvenRowBackgroundHoveredBrush(IMAGE_BRUSH("Common/Selection", Icon8x8, FLinearColor(1.0f, 1.0f, 1.0f, 0.1f)))
-		.SetOddRowBackgroundBrush(FSlateNoResource())
-		.SetOddRowBackgroundHoveredBrush(IMAGE_BRUSH("Common/Selection", Icon8x8, FLinearColor(1.0f, 1.0f, 1.0f, 0.1f)))
-		.SetSelectorFocusedBrush(BORDER_BRUSH("Common/Selector", FMargin(4.f / 16.f), SelectorColor))
-		.SetActiveBrush(IMAGE_BRUSH("Common/Selection", Icon8x8, SelectionColor))
-		.SetActiveHoveredBrush(IMAGE_BRUSH("Common/Selection", Icon8x8, SelectionColor))
-		.SetInactiveBrush(IMAGE_BRUSH("Common/Selection", Icon8x8, SelectionColor_Inactive))
-		.SetInactiveHoveredBrush(IMAGE_BRUSH("Common/Selection", Icon8x8, SelectionColor_Inactive))
-		.SetTextColor(DefaultForeground)
-		.SetSelectedTextColor(InvertedForeground);
-
 	// Lists, Trees
 	{
 		Set( "TableView.Row", FTableRowStyle( NormalTableRowStyle) );
diff --git a/Engine/Source/Editor/Kismet/Private/SSCSEditor.cpp b/Engine/Source/Editor/Kismet/Private/SSCSEditor.cpp
--- a/Engine/Source/Editor/Kismet/Private/SSCSEditor.cpp
+++ b/Engine/Source/Editor/Kismet/Private/SSCSEditor.cpp
@@ -1274,7 +1274,12 @@ void SSCS_RowWidget::Construct( const FArguments& InArgs, TSharedPtr<SSCSEditor>
 				&FEditorStyle::Get().GetWidgetStyle<FTableRowStyle>("TableView.NoHoverTableRow") :
 				&FEditorStyle::Get().GetWidgetStyle<FTableRowStyle>("SceneOutliner.TableViewRow")) //@todo create editor style for the SCS tree
 		.Padding(FMargin(0.f, 0.f, 0.f, 4.f))
-		.ShowSelection(!bIsSeparator);
+		.ShowSelection(!bIsSeparator)
+		.OnDragDetected(this, &SSCS_RowWidget::HandleOnDragDetected)
+		.OnDragEnter(this, &SSCS_RowWidget::HandleOnDragEnter)
+		.OnDragLeave(this, &SSCS_RowWidget::HandleOnDragLeave)
+		.OnCanAcceptDrop(this, &SSCS_RowWidget::HandleOnCanAcceptDrop)
+		.OnAcceptDrop(this, &SSCS_RowWidget::HandleOnAcceptDrop);
 
 	SMultiColumnTableRow<FSCSEditorTreeNodePtrType>::Construct( Args, InOwnerTableView.ToSharedRef() );
 }
@@ -1836,7 +1841,7 @@ FReply SSCS_RowWidget::OnMouseButtonDown(const FGeometry& MyGeometry, const FPoi
 	}
 }
 
-FReply SSCS_RowWidget::OnDragDetected( const FGeometry& MyGeometry, const FPointerEvent& MouseEvent )
+FReply SSCS_RowWidget::HandleOnDragDetected( const FGeometry& MyGeometry, const FPointerEvent& MouseEvent )
 {
 	auto SCSEditorPtr = SCSEditor.Pin();
 	if (MouseEvent.IsMouseButtonDown(EKeys::LeftMouseButton)
@@ -1868,7 +1873,7 @@ FReply SSCS_RowWidget::OnDragDetected( const FGeometry& MyGeometry, const FPoint
 	return FReply::Unhandled();
 }
 
-void SSCS_RowWidget::OnDragEnter( const FGeometry& MyGeometry, const FDragDropEvent& DragDropEvent )
+void SSCS_RowWidget::HandleOnDragEnter( const FDragDropEvent& DragDropEvent )
 {
 	TSharedPtr<FDragDropOperation> Operation = DragDropEvent.GetOperation();
 	if (!Operation.IsValid())
@@ -2165,12 +2170,13 @@ void SSCS_RowWidget::OnDragEnter( const FGeometry& MyGeometry, const FDragDropEv
 		TSharedPtr<SSCSEditor> PinnedEditor = SCSEditor.Pin();
 		if ( PinnedEditor.IsValid() && PinnedEditor->SCSTreeWidget.IsValid() )
 		{
-			PinnedEditor->SCSTreeWidget->OnDragEnter( MyGeometry, DragDropEvent );
+			// The widget geometry is irrelevant to the tree widget's OnDragEnter
+			PinnedEditor->SCSTreeWidget->OnDragEnter( FGeometry(), DragDropEvent );
 		}
 	}
 }
 
-void SSCS_RowWidget::OnDragLeave( const FDragDropEvent& DragDropEvent )
+void SSCS_RowWidget::HandleOnDragLeave(const FDragDropEvent& DragDropEvent)
 {
 	TSharedPtr<FSCSRowDragDropOp> DragRowOp = DragDropEvent.GetOperationAs<FSCSRowDragDropOp>();
 	if (DragRowOp.IsValid())
@@ -2194,7 +2200,33 @@ void SSCS_RowWidget::OnDragLeave( const FDragDropEvent& DragDropEvent )
 	}
 }
 
-FReply SSCS_RowWidget::OnDrop( const FGeometry& MyGeometry, const FDragDropEvent& DragDropEvent )
+TOptional<EItemDropZone> SSCS_RowWidget::HandleOnCanAcceptDrop(const FDragDropEvent& DragDropEvent, EItemDropZone DropZone, FSCSEditorTreeNodePtrType TargetItem)
+{
+	TOptional<EItemDropZone> ReturnDropZone;
+
+	TSharedPtr<FDragDropOperation> Operation = DragDropEvent.GetOperation();
+	if (Operation.IsValid())
+	{
+		if (Operation->IsOfType<FSCSRowDragDropOp>() && ( Cast<USceneComponent>(GetNode()->GetComponentTemplate()) != nullptr ))
+		{
+			TSharedPtr<FSCSRowDragDropOp> DragRowOp = StaticCastSharedPtr<FSCSRowDragDropOp>(Operation);
+			check(DragRowOp.IsValid());
+
+			if (DragRowOp->PendingDropAction != FSCSRowDragDropOp::DropAction_None)
+			{
+				ReturnDropZone = EItemDropZone::OntoItem;
+			}
+		}
+		else if (Operation->IsOfType<FExternalDragOperation>() || Operation->IsOfType<FAssetDragDropOp>())
+		{
+			ReturnDropZone = EItemDropZone::OntoItem;
+		}
+	}
+
+	return ReturnDropZone;
+}
+
+FReply SSCS_RowWidget::HandleOnAcceptDrop( const FDragDropEvent& DragDropEvent, EItemDropZone DropZone, FSCSEditorTreeNodePtrType TargetItem )
 {
 	TSharedPtr<FDragDropOperation> Operation = DragDropEvent.GetOperation();
 	if (!Operation.IsValid())
@@ -2243,7 +2275,8 @@ FReply SSCS_RowWidget::OnDrop( const FGeometry& MyGeometry, const FDragDropEvent
 		TSharedPtr<SSCSEditor> PinnedEditor = SCSEditor.Pin();
 		if ( PinnedEditor.IsValid() && PinnedEditor->SCSTreeWidget.IsValid() )
 		{
-			PinnedEditor->SCSTreeWidget->OnDrop( MyGeometry, DragDropEvent );
+			// The widget geometry is irrelevant to the tree widget's OnDrop
+			PinnedEditor->SCSTreeWidget->OnDrop( FGeometry(), DragDropEvent );
 		}
 	}
 
diff --git a/Engine/Source/Editor/Kismet/Public/SSCSEditor.h b/Engine/Source/Editor/Kismet/Public/SSCSEditor.h
--- a/Engine/Source/Editor/Kismet/Public/SSCSEditor.h
+++ b/Engine/Source/Editor/Kismet/Public/SSCSEditor.h
@@ -482,10 +482,6 @@ public:
 
 	// SWidget interface
 	virtual FReply OnMouseButtonDown(const FGeometry& MyGeometry, const FPointerEvent& MouseEvent) override;
-	virtual FReply OnDragDetected( const FGeometry& MyGeometry, const FPointerEvent& MouseEvent ) override;
-	virtual void OnDragEnter( const FGeometry& MyGeometry, const FDragDropEvent& DragDropEvent ) override;
-	virtual void OnDragLeave( const FDragDropEvent& DragDropEvent ) override;
-	virtual FReply OnDrop( const FGeometry& MyGeometry, const FDragDropEvent& DragDropEvent ) override;
 	// End of SWidget interface
 
 	/** Get the blueprint we are editing */
@@ -523,6 +519,13 @@ private:
 	/** Creates a tooltip for this row */
 	TSharedRef<SToolTip> CreateToolTipWidget() const;
 
+	/** Drag-drop handlers */
+	void HandleOnDragEnter(const FDragDropEvent& DragDropEvent);
+	void HandleOnDragLeave(const FDragDropEvent& DragDropEvent);
+	FReply HandleOnDragDetected(const FGeometry& MyGeometry, const FPointerEvent& MouseEvent);
+	TOptional<EItemDropZone> HandleOnCanAcceptDrop(const FDragDropEvent& DragDropEvent, EItemDropZone DropZone, FSCSEditorTreeNodePtrType TargetItem);
+	FReply HandleOnAcceptDrop(const FDragDropEvent& DragDropEvent, EItemDropZone DropZone, FSCSEditorTreeNodePtrType TargetItem);
+
 	/** Handler for attaching a single node to this node */
 	void OnAttachToDropAction(FSCSEditorTreeNodePtrType DroppedNodePtr)
 	{
@@ -668,11 +671,10 @@ public:
 	/** Object construction - mostly defers to the base STreeView */
 	void Construct( const FArguments& InArgs );
 
-	/** Handler for drag over operations */
-	FReply OnDragOver( const FGeometry& MyGeometry, const FDragDropEvent& DragDropEvent );
-
-	/** Handler for drop operations */
-	FReply OnDrop( const FGeometry& MyGeometry, const FDragDropEvent& DragDropEvent );
+	// SWidget interface
+	virtual FReply OnDragOver( const FGeometry& MyGeometry, const FDragDropEvent& DragDropEvent ) override;
+	virtual FReply OnDrop( const FGeometry& MyGeometry, const FDragDropEvent& DragDropEvent ) override;
+	// End SWidget interface
 
 private:
 	/** Pointer to the SSCSEditor that owns this widget */
diff --git a/Engine/Source/Runtime/Slate/Public/Widgets/Views/STableRow.h b/Engine/Source/Runtime/Slate/Public/Widgets/Views/STableRow.h
--- a/Engine/Source/Runtime/Slate/Public/Widgets/Views/STableRow.h
+++ b/Engine/Source/Runtime/Slate/Public/Widgets/Views/STableRow.h
@@ -223,7 +223,7 @@ public:
 				DropIndicatorBrush,
 				MyClippingRect,
 				ESlateDrawEffect::None,
-				InWidgetStyle.GetColorAndOpacityTint()
+				DropIndicatorBrush->GetTint(InWidgetStyle) * InWidgetStyle.GetColorAndOpacityTint()
 			);
 		}
 
diff --git a/Engine/Source/Runtime/SlateCore/Public/Styling/SlateTypes.h b/Engine/Source/Runtime/SlateCore/Public/Styling/SlateTypes.h
--- a/Engine/Source/Runtime/SlateCore/Public/Styling/SlateTypes.h
+++ b/Engine/Source/Runtime/SlateCore/Public/Styling/SlateTypes.h
@@ -1080,11 +1080,6 @@ struct SLATECORE_API FTableRowStyle : public FSlateWidgetStyle
 	FSlateBrush OddRowBackgroundBrush;
 	FTableRowStyle& SetOddRowBackgroundBrush( const FSlateBrush& InOddRowBackgroundBrush ){ OddRowBackgroundBrush = InOddRowBackgroundBrush; return *this; }
 
-	/** Brush used when drawing the drag-drop indicator lines above and below rows */
-	UPROPERTY(EditAnywhere, Category = Appearance)
-	FSlateBrush DragDropLineIndicatorBrush;
-	FTableRowStyle& SetDragDropLineIndicatorBrush(const FSlateBrush& InDragDropLineIndicatorBrush){ DragDropLineIndicatorBrush = InDragDropLineIndicatorBrush; return *this; }
-
 	/** Text color used for all rows */
 	UPROPERTY(EditAnywhere, BlueprintReadWrite, Category=Appearance)
 	FSlateColor TextColor;
