Expanding a function node with function nodes inside that are scoped to the current Blueprint will no longer be missing the call function nodes.

Moves the duplicated nodes into the new graph (and the Blueprint) before checking if they are valid, this allows functions and variables to resolve correctly if self context.

diff --git a/Engine/Source/Editor/Kismet/Private/BlueprintEditor.cpp b/Engine/Source/Editor/Kismet/Private/BlueprintEditor.cpp
--- a/Engine/Source/Editor/Kismet/Private/BlueprintEditor.cpp
+++ b/Engine/Source/Editor/Kismet/Private/BlueprintEditor.cpp
@@ -6148,6 +6148,9 @@ void FBlueprintEditor::ExpandNode(UEdGraphNode* InNodeToExpand, UEdGraph* InSour
 	{
 		UEdGraphNode* Node = SourceGraph->Nodes.Pop();
 
+		Node->Modify();
+		Node->Rename(/*NewName=*/ NULL, /*NewOuter=*/ DestinationGraph);
+
 		// We do not check CanPasteHere when determining CanCollapseNodes, unlike CanCollapseSelectionToFunction/Macro,
 		// so when expanding a collapsed graph we don't want to check the CanPasteHere function:
 		if (!bIsCollapsedGraph && !Node->CanPasteHere(DestinationGraph))
@@ -6156,9 +6159,7 @@ void FBlueprintEditor::ExpandNode(UEdGraphNode* InNodeToExpand, UEdGraph* InSour
 			continue;
 		}
 
-		Node->Modify();
-
-		Node->Rename(/*NewName=*/ NULL, /*NewOuter=*/ DestinationGraph);
+		// Successfully added the node to the graph, we may need to remove flags
 		if (Node->HasAllFlags(RF_Transient) && !DestinationGraph->HasAllFlags(RF_Transient))
 		{
 			Node->ClearFlags(RF_Transient);
