Editor: Added OpenAssetPickerMenu to programmatically trigger the asset picker on a SContentPicker

diff --git a/Engine/Source/Editor/Kismet/Private/WorkflowOrientedApp/SContentReference.cpp b/Engine/Source/Editor/Kismet/Private/WorkflowOrientedApp/SContentReference.cpp
--- a/Engine/Source/Editor/Kismet/Private/WorkflowOrientedApp/SContentReference.cpp
+++ b/Engine/Source/Editor/Kismet/Private/WorkflowOrientedApp/SContentReference.cpp
@@ -32,6 +32,19 @@ void SContentReference::Construct(const FArguments& InArgs)
 
 	static const FName InvertedForegroundName("InvertedForeground");
 
+	PickerComboButton = SNew(SComboButton)
+		.ButtonStyle(FEditorStyle::Get(), "NoBorder")
+		.ContentPadding(1.f)
+		.Visibility(this, &SContentReference::GetPickButtonVisibility)
+		.OnGetMenuContent(this, &SContentReference::MakeAssetPickerMenu)
+		.HasDownArrow(false)
+		.ToolTipText(LOCTEXT("PickAsset", "Pick an asset from a popup menu"))
+		.ButtonContent()
+		[
+			SNew(SImage)
+			.Image(FEditorStyle::GetBrush(InArgs._Style, ".PickAsset"))
+		];
+
 	// Create the widgets
 	ChildSlot
 	[
@@ -58,28 +71,6 @@ void SContentReference::Construct(const FArguments& InArgs)
 			]
 		]
 
-		// Use current selection in content browser button
-		//@TODO: Remove this code
-		/*
-		+SHorizontalBox::Slot()
-		.AutoWidth()
-		.HAlign(HAlign_Center)
-		.VAlign(VAlign_Center)
-		.Padding(1,0)
-		[
-			SNew(SButton)
-			.Style("NoBorder")
-			.OnClicked(this, &SContentReference::OnClickUseButton)
-			.ContentPadding(1.f)
-			.Visibility(this, &SContentReference::GetUseButtonVisibility)
-			.ToolTipText(LOCTEXT("UseSelected", "Try to use selected asset from content browser").ToString())
-			[
-				SNew(SImage)
-				.Image( FEditorStyle::GetBrush(InArgs._Style, ".UseSelectionFromContentBrowser") )
-			]
-		]
-		*/
-
 		// Pick an asset
 		+SHorizontalBox::Slot()
 		.AutoWidth()
@@ -87,18 +78,7 @@ void SContentReference::Construct(const FArguments& InArgs)
 		.VAlign(VAlign_Center)
 		.Padding(1,0)
 		[
-			SAssignNew(PickerComboButton, SComboButton)
-			.ButtonStyle( FEditorStyle::Get(), "NoBorder" )
-			.ContentPadding(1.f)
-			.Visibility(this, &SContentReference::GetPickButtonVisibility)
-			.OnGetMenuContent( this, &SContentReference::MakeAssetPickerMenu )
-			.HasDownArrow(false)
-			.ToolTipText(LOCTEXT("PickAsset", "Pick an asset from a popup menu"))
-			.ButtonContent()
-			[
-				SNew(SImage)
-				.Image( FEditorStyle::GetBrush(InArgs._Style, ".PickAsset") )
-			]
+			PickerComboButton.ToSharedRef()
 		]
 			
 		// Find in content browser button
@@ -160,6 +140,11 @@ void SContentReference::Construct(const FArguments& InArgs)
 	];
 }
 
+void SContentReference::OpenAssetPickerMenu()
+{
+	PickerComboButton->SetIsOpen(true);
+}
+
 EVisibility SContentReference::GetUseButtonVisibility() const
 {
 	return AllowSelectingNewAsset.Get() ? EVisibility::Visible : EVisibility::Collapsed;
diff --git a/Engine/Source/Editor/Kismet/Public/WorkflowOrientedApp/SContentReference.h b/Engine/Source/Editor/Kismet/Public/WorkflowOrientedApp/SContentReference.h
--- a/Engine/Source/Editor/Kismet/Public/WorkflowOrientedApp/SContentReference.h
+++ b/Engine/Source/Editor/Kismet/Public/WorkflowOrientedApp/SContentReference.h
@@ -71,6 +71,9 @@ public:
 	// Construct a SContentReference
 	void Construct(const FArguments& InArgs);
 
+	// Causes the asset picker to open, allowing the user to select a new asset
+	void OpenAssetPickerMenu();
+
 protected:
 	EVisibility GetUseButtonVisibility() const;
 	EVisibility GetPickButtonVisibility() const;
