Find-All-In-Blueprints will again inform the user when Blueprints are not-indexed and need to re-save and will help the user re-save the needed Blueprints.

Revived code to display a banner which informs the user of un-indexed BP count and will report any failures to re-save.

Cleans up the Find-in-Blueprint manager of some code that serializes and encodes the FiB AR data.

Levels will always have a tag for FiB so the system can recognize when they are successfully indexed.

diff --git a/Engine/Source/Editor/Kismet/Private/FindInBlueprintManager.cpp b/Engine/Source/Editor/Kismet/Private/FindInBlueprintManager.cpp
--- a/Engine/Source/Editor/Kismet/Private/FindInBlueprintManager.cpp
+++ b/Engine/Source/Editor/Kismet/Private/FindInBlueprintManager.cpp
@@ -47,6 +47,82 @@ const FText FFindInBlueprintSearchTags::FiB_IsReference = LOCTEXT("IsReference",
 const FText FFindInBlueprintSearchTags::FiB_Glyph = LOCTEXT("Glyph", "Glyph");
 const FText FFindInBlueprintSearchTags::FiB_GlyphColor = LOCTEXT("GlyphColor", "GlyphColor");
 
+/** Helper functions for serialization of types to and from an FString */
+namespace FiBSerializationHelpers
+{
+	/**
+	* Helper function to handle properly encoding and serialization of a type into an FString
+	*
+	* @param InValue				Value to serialize
+	* @param bInIncludeSize		If true, include the size of the type. This will place an int32
+									before the value in the FString. This is needed for non-basic types
+									because everything is stored in an FString and is impossible to distinguish
+	*/
+	template<class Type>
+	const FString Serialize(Type& InValue, bool bInIncludeSize)
+	{
+		TArray<uint8> SerializedData;
+		FMemoryWriter Ar(SerializedData, /*bIsPersistent=*/ true);
+
+		Ar << InValue;
+		Ar.Close();
+		FString Result = BytesToString(SerializedData.GetData(), SerializedData.Num());
+
+		// If the size is included, prepend it onto the Result string.
+		if(bInIncludeSize)
+		{
+			SerializedData.Empty();
+			FMemoryWriter ArWithLength(SerializedData, /*bIsPersistent=*/ true);
+			int32 Length = Result.Len();
+			ArWithLength << Length;
+
+			Result = BytesToString(SerializedData.GetData(), SerializedData.Num()) + Result;
+		}
+		return Result;
+	}
+
+	/** Helper function to handle properly decoding of uint8 arrays so they can be deserialized as their respective types */
+	void DecodeFromStream(FBufferReader& InStream, int32 InBytes, TArray<uint8>& OutDerivedData)
+	{
+		// Read, as a byte string, the number of characters composing the Lookup Table for the Json.
+		FString SizeOfDataAsHex;
+		SizeOfDataAsHex.GetCharArray().AddUninitialized(InBytes + 1);
+		SizeOfDataAsHex[InBytes] = TEXT('\0');
+		InStream.Serialize((char*)SizeOfDataAsHex.GetCharArray().GetData(), sizeof(TCHAR) * InBytes);
+
+		// Convert the number (which is stored in 1 serialized byte per TChar) into an int32
+		OutDerivedData.Empty();
+		OutDerivedData.AddUninitialized(InBytes);
+		StringToBytes(SizeOfDataAsHex, OutDerivedData.GetData(), InBytes);
+	}
+
+	/** Helper function to deserialize from a Stream the sizeof the templated type */
+	template<class Type>
+	Type Deserialize(FBufferReader& InStream)
+	{
+		TArray<uint8> DerivedData;
+		DecodeFromStream(InStream, sizeof(Type), DerivedData);
+		FMemoryReader SizeOfDataAr(DerivedData, /*bIsPersistent=*/ true);
+
+		Type ReturnValue;
+		SizeOfDataAr << ReturnValue;
+		return ReturnValue;
+	}
+
+	/** Helper function to deserialize from a Stream a certain number of bytes */
+	template<class Type>
+	Type Deserialize(FBufferReader& InStream, int32 InBytes)
+	{
+		TArray<uint8> DerivedData;
+		DecodeFromStream(InStream, InBytes, DerivedData);
+		FMemoryReader SizeOfDataAr(DerivedData, /*bIsPersistent=*/ true);
+
+		Type ReturnValue;
+		SizeOfDataAr << ReturnValue;
+		return ReturnValue;
+	}
+}
+
 namespace BlueprintSearchMetaDataHelpers
 {
 	/** Json Writer used for serializing FText's in the correct format for Find-in-Blueprints */
@@ -149,21 +225,7 @@ namespace BlueprintSearchMetaDataHelpers
 		/** Converts the lookup table of ints (which are stored as identifiers and string values in the Json) and the FText's they represent to an FString. */
 		FString GetSerializedLookupTable()
 		{
-			TArray<uint8> SerializedData;
-			FMemoryWriter Ar(SerializedData, /*bIsPersistent=*/ true);
-
-			Ar << LookupTable;
-			Ar.Close();
-
-			FString Result = BytesToString(SerializedData.GetData(), SerializedData.Num());
-
-			SerializedData.Empty();
-
-			// Attach, as the first value and in an FString, the length of the Lookup table, this is to help sort through the data during deserialization
-			FMemoryWriter ArWithLength(SerializedData, /*bIsPersistent=*/ true);
-			int32 Length = Result.Len();
-			ArWithLength << Length;
-			return BytesToString(SerializedData.GetData(), SerializedData.Num()) + Result;
+			return FiBSerializationHelpers::Serialize< TMap< int32, FText > >(LookupTable, true);
 		}
 
 		struct FLookupTableItem
@@ -554,6 +616,217 @@ namespace BlueprintSearchMetaDataHelpers
 	}
 }
 
+class FCacheAllBlueprintsTickableObject
+{
+
+public:
+
+	FCacheAllBlueprintsTickableObject(TArray<FString>& InUncachedBlueprints)
+		: TickCacheIndex(0)
+		, UncachedBlueprints(InUncachedBlueprints)
+		, bIsStarted(false)
+		, bIsCancelled(false)
+		, bRecursionGuard(false)
+	{
+		// Start the Blueprint indexing 'progress' notification
+		FNotificationInfo Info( LOCTEXT("BlueprintIndexMessage", "Indexing Blueprints...") );
+		Info.bFireAndForget = false;
+		Info.ButtonDetails.Add(FNotificationButtonInfo(
+			LOCTEXT("BlueprintIndexCancel","Cancel"),
+			LOCTEXT("BlueprintIndexCancelToolTip","Cancels indexing Blueprints."), FSimpleDelegate::CreateRaw(this, &FCacheAllBlueprintsTickableObject::OnCancelCaching, false)));
+
+		ProgressNotification = FSlateNotificationManager::Get().AddNotification(Info);
+		if (ProgressNotification.IsValid())
+		{
+			ProgressNotification.Pin()->SetCompletionState(SNotificationItem::CS_Pending);
+		}
+	}
+
+	~FCacheAllBlueprintsTickableObject()
+	{
+
+	}
+
+	/** Returns the current cache index of the object */
+	int32 GetCurrentCacheIndex() const
+	{
+		return TickCacheIndex + 1;
+	}
+
+	/** Returns the name of the current Blueprint being cached */
+	FString GetCurrentCacheBlueprintName() const
+	{
+		if(UncachedBlueprints.Num() && TickCacheIndex >= 0)
+		{
+			return UncachedBlueprints[TickCacheIndex];
+		}
+		return FString();
+	}
+
+	/** Returns the progress as a percent */
+	float GetCacheProgress() const
+	{
+		return (float)TickCacheIndex / (float)UncachedBlueprints.Num();
+	}
+
+	/** Cancels caching and destroys this object */
+	void OnCancelCaching(bool bIsImmediate)
+	{
+		ProgressNotification.Pin()->SetText(LOCTEXT("BlueprintIndexCancelled", "Cancelled Indexing Blueprints!"));
+
+		ProgressNotification.Pin()->SetCompletionState(SNotificationItem::CS_Fail);
+		ProgressNotification.Pin()->ExpireAndFadeout();
+
+		// Sometimes we can't wait another tick to shutdown, so make the callback immediately.
+		if (bIsImmediate)
+		{
+			FFindInBlueprintSearchManager::Get().FinishedCachingBlueprints(TickCacheIndex, FailedToCacheList);
+		}
+		else
+		{
+			bIsCancelled = true;
+		}
+	}
+
+	/** Enables the caching process */
+	void Start() { bIsStarted = true; }
+
+	/** FTickableEditorObject interface */
+	EActiveTimerReturnType Tick(double InCurrentTime, float InDeltaTime)
+	{
+		// Protect against Slate recursion if a modal dialog appears from loading/resaving an asset
+		if (bRecursionGuard == true)
+		{
+			return EActiveTimerReturnType::Continue;
+		}
+		TGuardValue<bool> Guard(bRecursionGuard, true);
+
+		if(!bIsStarted)
+		{
+			return EActiveTimerReturnType::Continue;
+		}
+
+		if (bIsCancelled || GWarn->ReceivedUserCancel())
+		{
+			FFindInBlueprintSearchManager::Get().FinishedCachingBlueprints(TickCacheIndex, FailedToCacheList);
+			return EActiveTimerReturnType::Stop;
+		}
+		else
+		{
+			FAssetRegistryModule* AssetRegistryModule = &FModuleManager::LoadModuleChecked<FAssetRegistryModule>(TEXT("AssetRegistry"));
+			FAssetData AssetData = AssetRegistryModule->Get().GetAssetByObjectPath(FName(*UncachedBlueprints[TickCacheIndex]));
+
+			// Cache whether the Blueprint is already loaded
+			bool bIsLoaded = AssetData.IsAssetLoaded();
+
+			bool bIsWorldAsset = AssetData.GetClass() == UWorld::StaticClass();
+
+			// Construct a full package filename with path so we can query the read only status and save to disk
+			FString FinalPackageFilename = FPackageName::LongPackageNameToFilename(AssetData.PackageName.ToString());
+			if( FinalPackageFilename.Len() > 0 && FPaths::GetExtension(FinalPackageFilename).Len() == 0 )
+			{
+				FinalPackageFilename += bIsWorldAsset ? FPackageName::GetMapPackageExtension() : FPackageName::GetAssetPackageExtension();
+			}
+			FText ErrorMessage;
+			bool bValidFilename = FEditorFileUtils::IsFilenameValidForSaving( FinalPackageFilename, ErrorMessage );
+			if ( bValidFilename )
+			{
+				bValidFilename = bIsWorldAsset ? FEditorFileUtils::IsValidMapFilename( FinalPackageFilename, ErrorMessage ) : FPackageName::IsValidLongPackageName( FinalPackageFilename, false, &ErrorMessage );
+			}
+
+			bool bIsAssetReadOnlyOnDisk = IFileManager::Get().IsReadOnly( *FinalPackageFilename );
+			bool bFailedToCache = true;
+
+			if (!bIsAssetReadOnlyOnDisk)
+			{
+				UObject* Asset = AssetData.GetAsset();
+				if(Asset)
+				{
+					// Assume the package was correctly checked out from SCC
+					bool bOutPackageLocallyWritable = true;
+
+					UPackage* Package = AssetData.GetPackage();
+
+					ISourceControlProvider& SourceControlProvider = ISourceControlModule::Get().GetProvider();
+					// Trusting the SCC status in the package file cache to minimize network activity during save.
+					const FSourceControlStatePtr SourceControlState = SourceControlProvider.GetState(Package, EStateCacheUsage::Use);
+					// If the package is in the depot, and not recognized as editable by source control, and not read-only, then we know the user has made the package locally writable!
+					const bool bSCCCanEdit = !SourceControlState.IsValid() || SourceControlState->CanCheckIn() || SourceControlState->IsIgnored() || SourceControlState->IsUnknown();
+					const bool bSCCIsCheckedOut = SourceControlState.IsValid() && SourceControlState->IsCheckedOut();
+					const bool bInDepot = SourceControlState.IsValid() && SourceControlState->IsSourceControlled();
+					if ( !bSCCCanEdit && bInDepot && !bIsAssetReadOnlyOnDisk && SourceControlProvider.UsesLocalReadOnlyState() && !bSCCIsCheckedOut )
+					{
+						bOutPackageLocallyWritable = false;
+					}
+
+					// Save the package if the file is writable
+					if(bOutPackageLocallyWritable)
+					{
+						UWorld* WorldAsset = Cast<UWorld>(Asset);
+
+						// Save the package
+						EObjectFlags ObjectFlags = (WorldAsset == nullptr)? RF_Standalone : RF_NoFlags;
+						
+						if (GEditor->SavePackage(Package, WorldAsset, ObjectFlags, *FinalPackageFilename))
+						{
+							bFailedToCache = false;
+						}
+					}
+				}
+			}
+
+			if (bFailedToCache)
+			{
+				FailedToCacheList.Add(UncachedBlueprints[TickCacheIndex]);
+			}
+
+			++TickCacheIndex;
+
+			// Check if done caching Blueprints
+			if(TickCacheIndex == UncachedBlueprints.Num())
+			{
+				ProgressNotification.Pin()->SetCompletionState(SNotificationItem::CS_Success);
+				ProgressNotification.Pin()->ExpireAndFadeout();
+
+				ProgressNotification.Pin()->SetText(LOCTEXT("BlueprintIndexComplete", "Finished indexing Blueprints!"));
+
+				FFindInBlueprintSearchManager::Get().FinishedCachingBlueprints(TickCacheIndex, FailedToCacheList);
+				return EActiveTimerReturnType::Stop;
+			}
+			else
+			{
+				FFormatNamedArguments Args;
+				Args.Add(TEXT("Percent"), FText::AsPercent(GetCacheProgress()));
+				ProgressNotification.Pin()->SetText(FText::Format(LOCTEXT("BlueprintIndexProgress", "Indexing Blueprints... ({Percent})"), Args));
+			}
+		}
+		return EActiveTimerReturnType::Continue;
+	}
+
+private:
+
+	/** The current index, increases at a rate of once per tick */
+	int32 TickCacheIndex;
+
+	/** The list of uncached Blueprints that are in the process of being cached */
+	TArray<FString> UncachedBlueprints;
+
+	/** Notification that appears and details progress */
+	TWeakPtr<SNotificationItem> ProgressNotification;
+
+	/** List of Blueprints that failed to be saved */
+	TArray<FString> FailedToCacheList;
+
+	/** TRUE if the caching process is started */
+	bool bIsStarted;
+
+	/** TRUE if the user has requested to cancel the caching process */
+	bool bIsCancelled;
+
+	/** Guard to prevent TickRecursion */
+	bool bRecursionGuard = false;
+};
+
 FFindInBlueprintSearchManager& FFindInBlueprintSearchManager::Get()
 {
 	if (Instance == NULL)
@@ -567,6 +840,7 @@ FFindInBlueprintSearchManager& FFindInBlueprintSearchManager::Get()
 
 FFindInBlueprintSearchManager::FFindInBlueprintSearchManager()
 	: bIsPausing(false)
+	, CachingObject(nullptr)
 {
 }
 
@@ -607,63 +881,77 @@ void FFindInBlueprintSearchManager::OnAssetAdded(const FAssetData& InAssetData)
 		int32* IndexPtr = SearchMap.Find(BlueprintPackagePath);
 		if(!IndexPtr)
 		{
-			FSearchData NewSearchData;
-
-			NewSearchData.BlueprintPath = InAssetData.ObjectPath.ToString();
-			NewSearchData.bMarkedForDeletion = false;
-
 			if(InAssetData.IsAssetLoaded())
 			{
-				UBlueprint* BlueprintAsset = nullptr;
-
 				if(InAssetData.GetClass()->IsChildOf(UBlueprint::StaticClass()))
 				{
-					BlueprintAsset = Cast<UBlueprint>(InAssetData.GetAsset());
+					if(UBlueprint* BlueprintAsset = Cast<UBlueprint>(InAssetData.GetAsset()))
+					{
+						// Cache the searchable data
+						AddOrUpdateBlueprintSearchMetadata(BlueprintAsset);
+					}
 				}
 				else if(InAssetData.GetClass()->IsChildOf(UWorld::StaticClass()))
 				{
 					UWorld* WorldAsset = Cast<UWorld>(InAssetData.GetAsset());
 					if(WorldAsset->PersistentLevel)
 					{
-						BlueprintAsset = Cast<UBlueprint>((UObject*)WorldAsset->PersistentLevel->GetLevelScriptBlueprint(true));
-					}
-				}
+						TArray<UBlueprint*> LevelBlueprints;
+						LevelBlueprints = WorldAsset->PersistentLevel->GetLevelBlueprints();
 
-				// Levels do not always have Blueprints
-				if(BlueprintAsset)
-				{
-					// Cache the searchable data
-					AddOrUpdateBlueprintSearchMetadata(BlueprintAsset);
+						for(UBlueprint* BlueprintAsset : LevelBlueprints)
+						{
+							// Cache the searchable data
+							AddOrUpdateBlueprintSearchMetadata(BlueprintAsset);
+						}
 
-					NewSearchData.Blueprint = Cast<UBlueprint>(BlueprintAsset);
+					}
 				}
-
-				// No more work, we have added the Blueprint if there is one, and if it's a level without one, it need not be added.
-				return;
 			}
 			else if(const FString* FiBSearchData = InAssetData.TagsAndValues.Find("FiB"))
 			{
+				FSearchData NewSearchData;
+
+				NewSearchData.BlueprintPath = InAssetData.ObjectPath.ToString();
+				NewSearchData.bMarkedForDeletion = false;
+
 				// Since the asset was not loaded, pull out the searchable data stored in the asset
 				NewSearchData.Value = *FiBSearchData;
+				AddSearchDataToDatabase(NewSearchData);
+			}
+			else
+			{
+				// The asset is uncached, we will want to inform the user that this is the case
+				UncachedBlueprints.Add(InAssetData.ObjectPath.ToString());
 			}
-
-			int32 ArrayIndex = SearchArray.Add(NewSearchData);
-
-			// Add the asset file path to the map along with the index into the array
-			SearchMap.Add(BlueprintPackagePath, ArrayIndex);
 		}
 	}
 }
 
+int32 FFindInBlueprintSearchManager::AddSearchDataToDatabase(FSearchData& InSearchData)
+{
+	int32 ArrayIndex = SearchArray.Add(InSearchData);
+
+	// Add the asset file path to the map along with the index into the array
+	SearchMap.Add(InSearchData.BlueprintPath, ArrayIndex);
+
+	return ArrayIndex;
+}
+
+void FFindInBlueprintSearchManager::RemoveBlueprintByPath(FString InPath)
+{
+	int32* SearchIdx = SearchMap.Find(InPath);
+
+	if(SearchIdx)
+	{
+		SearchArray[*SearchIdx].bMarkedForDeletion = true;
+	}
+}
 void FFindInBlueprintSearchManager::OnAssetRemoved(const class FAssetData& InAssetData)
 {
 	if(InAssetData.IsAssetLoaded())
 	{
-		int32* SearchIdx = SearchMap.Find(InAssetData.GetAsset()->GetOutermost()->GetPathName());
-		if(SearchIdx)
-		{
-			SearchArray[*SearchIdx].bMarkedForDeletion = true;
-		}
+		RemoveBlueprintByPath(InAssetData.ObjectPath.ToString());
 	}
 }
 
@@ -672,19 +960,14 @@ void FFindInBlueprintSearchManager::OnAssetRenamed(const class FAssetData& InAss
 	// Renaming removes the item from the manager, it will be re-added in the OnAssetAdded event under the new name.
 	if(InAssetData.IsAssetLoaded())
 	{
-		FString Test = FPaths::GetPath(InOldName) / FPaths::GetBaseFilename(InOldName);
-		int32* SearchIdx = SearchMap.Find(FPaths::GetPath(InOldName) / FPaths::GetBaseFilename(InOldName));
-		if(SearchIdx)
-		{
-			SearchArray[*SearchIdx].bMarkedForDeletion = true;
-		}
+		RemoveBlueprintByPath(InOldName);
 	}
 }
 
 void FFindInBlueprintSearchManager::OnAssetLoaded(UObject* InAsset)
 {
 	UBlueprint* BlueprintAsset = nullptr;
-	FString BlueprintPackagePath;
+	FString BlueprintPackagePath = InAsset->GetPathName();
 
 	if(UWorld* WorldAsset = Cast<UWorld>(InAsset))
 	{
@@ -693,15 +976,15 @@ void FFindInBlueprintSearchManager::OnAssetLoaded(UObject* InAsset)
 			BlueprintAsset = Cast<UBlueprint>((UObject*)WorldAsset->PersistentLevel->GetLevelScriptBlueprint(true));
 			if(BlueprintAsset)
 			{
-				BlueprintPackagePath = BlueprintAsset->GetOutermost()->GetPathName();
+				BlueprintPackagePath = BlueprintAsset->GetPathName();
 			}
 		}
 	}
 	else
 	{
 		BlueprintAsset = Cast<UBlueprint>(InAsset);
-		BlueprintPackagePath = BlueprintAsset->GetPathName();
 	}
+
 	if(BlueprintAsset)
 	{
 		// Find and update the item in the search array. Searches may currently be active, this will do no harm to them
@@ -714,7 +997,6 @@ void FFindInBlueprintSearchManager::OnAssetLoaded(UObject* InAsset)
 		{
 			// That index should never have a Blueprint already, but if it does, it should be the same Blueprint!
 			ensure(!SearchArray[*IndexPtr].Blueprint.IsValid() || SearchArray[*IndexPtr].Blueprint == BlueprintAsset);
-			//check(!SearchArray[*IndexPtr].Blueprint);
 			SearchArray[*IndexPtr].Blueprint = BlueprintAsset;
 		}
 	}
@@ -798,8 +1080,6 @@ void FFindInBlueprintSearchManager::AddOrUpdateBlueprintSearchMetadata(UBlueprin
 	// Allow only one thread modify the search data at a time
 	FScopeLock ScopeLock(&SafeModifyCacheCriticalSection);
 
-	FString BlueprintPackagePath = InBlueprint->GetOutermost()->GetPathName();
-
 	FString BlueprintPath;
 	if(FBlueprintEditorUtils::IsLevelScriptBlueprint(InBlueprint))
 	{
@@ -813,16 +1093,15 @@ void FFindInBlueprintSearchManager::AddOrUpdateBlueprintSearchMetadata(UBlueprin
 		BlueprintPath = InBlueprint->GetPathName();
 	}
 
-	int32* IndexPtr = SearchMap.Find(BlueprintPackagePath);
+	int32* IndexPtr = SearchMap.Find(BlueprintPath);
 	int32 Index = 0;
 	if(!IndexPtr)
 	{
 		FSearchData SearchData;
 		SearchData.Blueprint = InBlueprint;
 		SearchData.BlueprintPath = BlueprintPath;
-		Index = SearchArray.Add(SearchData);
 
-		SearchMap.Add(BlueprintPackagePath, Index);
+		Index = AddSearchDataToDatabase(SearchData);
 	}
 	else
 	{
@@ -932,8 +1211,14 @@ FString FFindInBlueprintSearchManager::QuerySingleBlueprint(UBlueprint* InBluepr
 		AddOrUpdateBlueprintSearchMetadata(InBlueprint, true);
 	}
 
-	FString Key = InBlueprint->GetOutermost()->GetPathName();
-
+	FString Key = InBlueprint->GetPathName();
+	if (ULevel* LevelOuter = Cast<ULevel>(InBlueprint->GetOuter()))
+	{
+		if (UWorld* WorldOuter = Cast<UWorld>(LevelOuter->GetOuter()))
+		{
+			Key = WorldOuter->GetPathName();
+		}
+	}
 	int32* ArrayIdx = SearchMap.Find(Key);
 	// This should always be true since we make sure to refresh the search data for this Blueprint when doing the search, unless we do not rebuild the searchable data
 	check((bInRebuildSearchData && ArrayIdx && *ArrayIdx < SearchArray.Num()) || !bInRebuildSearchData);
@@ -1101,6 +1386,116 @@ FString FFindInBlueprintSearchManager::ConvertFTextToHexString(FText InValue)
 	return BytesToHex(SerializedData.GetData(), SerializedData.Num());
 }
 
+void FFindInBlueprintSearchManager::OnCacheAllUncachedBlueprints(bool bInSourceControlActive)
+{
+	// Multiple threads can be adding to this at the same time
+	FScopeLock ScopeLock(&SafeModifyCacheCriticalSection);
+
+	if(ISourceControlModule::Get().IsEnabled())
+	{
+		FEditorFileUtils::CheckoutPackages(UncachedBlueprints);
+	}
+
+	// Start the cache process.
+	CachingObject->Start();
+}
+
+void FFindInBlueprintSearchManager::CacheAllUncachedBlueprints(TWeakPtr< SFindInBlueprints > InSourceWidget, FWidgetActiveTimerDelegate& OutActiveTimerDelegate)
+{
+	// Do not start another caching process if one is in progress
+	if(!IsCacheInProgress())
+	{
+		// Add all failed to cache Blueprints to the UncachedBlueprints list and resubmit for caching
+		UncachedBlueprints.Append(FailedToCachePaths);
+		FailedToCachePaths.Empty();
+
+		CachingObject = new FCacheAllBlueprintsTickableObject(UncachedBlueprints);
+		OutActiveTimerDelegate.BindRaw(CachingObject, &FCacheAllBlueprintsTickableObject::Tick);
+
+		if(!ISourceControlModule::Get().IsEnabled())
+		{
+			// Offer to start up Source Control
+			ISourceControlModule::Get().ShowLoginDialog(FSourceControlLoginClosed::CreateRaw(this, &FFindInBlueprintSearchManager::OnCacheAllUncachedBlueprints), ELoginWindowMode::Modeless, EOnLoginWindowStartup::PreserveProvider);
+		}
+		else
+		{
+			OnCacheAllUncachedBlueprints(true);
+		}
+
+		SourceCachingWidget = InSourceWidget;
+	}
+}
+
+void FFindInBlueprintSearchManager::CancelCacheAll(SFindInBlueprints* InFindInBlueprintWidget)
+{
+	if(IsCacheInProgress() && ((SourceCachingWidget.IsValid() && SourceCachingWidget.Pin().Get() == InFindInBlueprintWidget) || !SourceCachingWidget.IsValid()))
+	{
+		CachingObject->OnCancelCaching(!SourceCachingWidget.IsValid());
+		SourceCachingWidget.Reset();
+	}
+}
+
+int32 FFindInBlueprintSearchManager::GetCurrentCacheIndex() const
+{
+	int32 CachingIndex = 0;
+	if(CachingObject)
+	{
+		CachingIndex = CachingObject->GetCurrentCacheIndex();
+	}
+
+	return CachingIndex;
+}
+
+FString FFindInBlueprintSearchManager::GetCurrentCacheBlueprintName() const
+{
+	FString CachingBPName;
+	if(CachingObject)
+	{
+		CachingBPName = CachingObject->GetCurrentCacheBlueprintName();
+	}
+
+	return CachingBPName;
+}
+
+float  FFindInBlueprintSearchManager::GetCacheProgress() const
+{
+	float ReturnCacheValue = 1.0f;
+
+	if(CachingObject)
+	{
+		ReturnCacheValue = CachingObject->GetCacheProgress();
+	}
+
+	return ReturnCacheValue;
+}
+
+void FFindInBlueprintSearchManager::FinishedCachingBlueprints(int32 InNumberCached, TArray<FString>& InFailedToCacheList)
+{
+	// Multiple threads could be adding to this at the same time
+	FScopeLock ScopeLock(&SafeModifyCacheCriticalSection);
+
+	FailedToCachePaths = InFailedToCacheList;
+
+	// Clean out the list by the count that were successfully cached, this allows us to continue the caching later.
+	UncachedBlueprints.RemoveAt(0, InNumberCached);
+
+	// Signal the callback, so the source FindInBlueprint can resubmit their search queries
+	if(SourceCachingWidget.IsValid())
+	{
+		SourceCachingWidget.Pin()->OnCacheComplete();
+	}
+	SourceCachingWidget.Reset();
+
+	// Delete the object and NULL it out so we can do it again in the future if needed (if it was canceled)
+	delete CachingObject;
+	CachingObject = nullptr;
+}
+
+bool FFindInBlueprintSearchManager::IsCacheInProgress() const
+{
+	return CachingObject != nullptr;
+}
+
 TSharedPtr< FJsonObject > FFindInBlueprintSearchManager::ConvertJsonStringToObject(FString InJsonString)
 {
 	/** The searchable data is more complicated than a Json string, the Json being the main searchable body that is parsed. Below is a diagram of the full data:
@@ -1117,33 +1512,12 @@ TSharedPtr< FJsonObject > FFindInBlueprintSearchManager::ConvertJsonStringToObje
 	int32 SizeOfData;
 	FBufferReader ReaderStream((void*)*InJsonString, InJsonString.Len() * sizeof(TCHAR), false);
 
-	// Read, as a byte string, the number of characters composing the Lookup Table for the Json.
-	FString SizeOfDataAsHex;
-	SizeOfDataAsHex.GetCharArray().AddUninitialized(5);
-	SizeOfDataAsHex[4] = TEXT('\0');
-	ReaderStream.Serialize((char*)SizeOfDataAsHex.GetCharArray().GetData(), sizeof(TCHAR) * 4);
-
-	// Convert the number (which is stored in 1 serialized byte per TChar) into an int32
-	DerivedData.Empty();
-	DerivedData.AddUninitialized(4);
-	StringToBytes(SizeOfDataAsHex, DerivedData.GetData(), 4);
-	FMemoryReader SizeOfDataAr(DerivedData, /*bIsPersistent=*/ true);
-	SizeOfDataAr << SizeOfData;
-
-	// With the size of the TMap in hand, let's serialize JUST that (as a byte string)
-	FString MapAsByteString;
-	MapAsByteString.GetCharArray().AddUninitialized(SizeOfData + 1);
-	MapAsByteString[SizeOfData] = TEXT('\0');
-	ReaderStream.Serialize((char*)MapAsByteString.GetCharArray().GetData(), sizeof(TCHAR) * SizeOfData);
-
-	// With the hex of the TMap, serialize that into a TMap
-	DerivedData.Empty();
-	DerivedData.AddUninitialized(SizeOfData);
-	StringToBytes(MapAsByteString, DerivedData.GetData(), SizeOfData);
+ 	// Read, as a byte string, the number of characters composing the Lookup Table for the Json.
+	SizeOfData = FiBSerializationHelpers::Deserialize<int32>(ReaderStream);
 
+ 	// With the size of the TMap in hand, let's serialize JUST that (as a byte string)
 	TMap<int32, FText> LookupTable;
-	FMemoryReader LookupTableAr(DerivedData, /*bIsPersistent=*/ true);
-	LookupTableAr << LookupTable;
+	LookupTable = FiBSerializationHelpers::Deserialize< TMap<int32, FText> >(ReaderStream, SizeOfData);
 
 	// The original BufferReader should be positioned at the Json
 	TSharedPtr< FJsonObject > JsonObject = NULL;
diff --git a/Engine/Source/Editor/Kismet/Private/FindInBlueprints.cpp b/Engine/Source/Editor/Kismet/Private/FindInBlueprints.cpp
--- a/Engine/Source/Editor/Kismet/Private/FindInBlueprints.cpp
+++ b/Engine/Source/Editor/Kismet/Private/FindInBlueprints.cpp
@@ -358,7 +358,7 @@ UBlueprint* FFindInBlueprintsResult::GetParentBlueprint() const
 	else
 	{
 		GIsEditorLoadingPackage = true;
-		UObject* Object = LoadObject<UObject>(NULL, *DisplayText.ToString(), NULL, 0, NULL); //StaticLoadObject(UObject::StaticClass(), NULL, *DisplayText.ToString());
+		UObject* Object = LoadObject<UObject>(NULL, *DisplayText.ToString(), NULL, 0, NULL);
 		GIsEditorLoadingPackage = false;
 
 		if(UBlueprint* BlueprintObj = Cast<UBlueprint>(Object))
@@ -1042,6 +1042,177 @@ void SFindInBlueprints::Construct( const FArguments& InArgs, TSharedPtr<FBluepri
 		];
 }
 
+void SFindInBlueprints::ConditionallyAddCacheBar()
+{
+	FFindInBlueprintSearchManager& FindInBlueprintManager = FFindInBlueprintSearchManager::Get();
+
+	// Do not add a second cache bar and do not add it when there are no uncached Blueprints
+	if(FindInBlueprintManager.GetNumberUncachedBlueprints() > 0 || FindInBlueprintManager.GetFailedToCacheCount() > 0)
+	{
+		if(MainVerticalBox.IsValid() && !CacheBarSlot.IsValid())
+		{
+			// Create a single string of all the Blueprint paths that failed to cache, on separate lines
+			FString PackageList;
+			TArray<FString> FailedToCacheList = FFindInBlueprintSearchManager::Get().GetFailedToCachePathList();
+			for (FString Package : FailedToCacheList)
+			{
+				PackageList += Package + TEXT("\n");
+			}
+
+			// Lambda to put together the popup menu detailing the failed to cache paths
+			auto OnDisplayCacheFailLambda = [](TWeakPtr<SWidget> InParentWidget, FString InPackageList)->FReply
+			{
+				if (InParentWidget.IsValid())
+				{
+					TSharedRef<SWidget> DisplayWidget = 
+						SNew(SBox)
+						.MaxDesiredHeight(512)
+						.MaxDesiredWidth(512)
+						.Content()
+						[
+							SNew(SBorder)
+							.BorderImage(FEditorStyle::GetBrush("ToolPanel.GroupBorder"))
+							[
+								SNew(SScrollBox)
+								+SScrollBox::Slot()
+								[
+									SNew(SMultiLineEditableText)
+									.AutoWrapText(true)
+									.IsReadOnly(true)
+									.Text(FText::FromString(InPackageList))
+								]
+							]
+						];
+
+					FSlateApplication::Get().PushMenu(
+						InParentWidget.Pin().ToSharedRef(),
+						DisplayWidget,
+						FSlateApplication::Get().GetCursorPos(),
+						FPopupTransitionEffect(FPopupTransitionEffect::TypeInPopup)
+						);	
+				}
+				return FReply::Handled();
+			};
+
+			MainVerticalBox.Pin()->AddSlot()
+				.AutoHeight()
+				[
+					SAssignNew(CacheBarSlot, SBorder)
+					.Visibility( this, &SFindInBlueprints::GetCachingBarVisibility )
+					.BorderBackgroundColor( this, &SFindInBlueprints::GetCachingBarColor )
+					.BorderImage( FCoreStyle::Get().GetBrush("ErrorReporting.Box") )
+					.Padding( FMargin(3,1) )
+					[
+						SNew(SVerticalBox)
+
+						+SVerticalBox::Slot()
+						.AutoHeight()
+						[
+							SNew(SHorizontalBox)
+							+SHorizontalBox::Slot()
+							.VAlign(EVerticalAlignment::VAlign_Center)
+							.AutoWidth()
+							[
+								SNew(STextBlock)
+								.Text(this, &SFindInBlueprints::GetUncachedBlueprintWarningText)
+								.ColorAndOpacity( FCoreStyle::Get().GetColor("ErrorReporting.ForegroundColor") )
+							]
+
+							// Cache All button
+							+SHorizontalBox::Slot()
+								.AutoWidth()
+								.VAlign(EVerticalAlignment::VAlign_Center)
+								.Padding(6.0f, 2.0f, 4.0f, 2.0f)
+								[
+									SNew(SButton)
+									.Text(LOCTEXT("IndexAllBlueprints", "Index All"))
+									.OnClicked( this, &SFindInBlueprints::OnCacheAllBlueprints )
+									.Visibility( this, &SFindInBlueprints::GetCacheAllButtonVisibility )
+									.ToolTip(IDocumentation::Get()->CreateToolTip(
+									LOCTEXT("IndexAlLBlueprints_Tooltip", "Loads all non-indexed Blueprints and saves them with their search data. This can be a very slow process, feel free to cancel it whenever and benefit from the Blueprints it indexes."),
+									NULL,
+									TEXT("Shared/Editors/BlueprintEditor"),
+									TEXT("FindInBlueprint_IndexAll")))
+								]
+
+
+							// View of failed Blueprint paths
+							+SHorizontalBox::Slot()
+								.AutoWidth()
+								.Padding(4.0f, 2.0f, 0.0f, 2.0f)
+								[
+									SNew(SButton)
+									.Text(LOCTEXT("ShowFailedPackages", "Show Failed Packages"))
+									.OnClicked(FOnClicked::CreateLambda(OnDisplayCacheFailLambda, TWeakPtr<SWidget>(SharedThis(this)), PackageList))
+									.Visibility( this, &SFindInBlueprints::GetFailedToCacheListVisibility )
+									.ToolTip(IDocumentation::Get()->CreateToolTip(
+									LOCTEXT("FailedCache_Tooltip", "Displays a list of packages that failed to save."),
+									NULL,
+									TEXT("Shared/Editors/BlueprintEditor"),
+									TEXT("FindInBlueprint_FailedCache")))
+								]
+
+							// Cache progress bar
+							+SHorizontalBox::Slot()
+								.FillWidth(1.0f)
+								.Padding(4.0f, 2.0f, 4.0f, 2.0f)
+								[
+									SNew(SProgressBar)
+									.Percent( this, &SFindInBlueprints::GetPercentCompleteCache )
+									.Visibility( this, &SFindInBlueprints::GetCachingProgressBarVisiblity )
+								]
+
+							// Cancel button
+							+SHorizontalBox::Slot()
+								.AutoWidth()
+								.Padding(4.0f, 2.0f, 0.0f, 2.0f)
+								[
+									SNew(SButton)
+									.Text(LOCTEXT("CancelCacheAll", "Cancel"))
+									.OnClicked( this, &SFindInBlueprints::OnCancelCacheAll )
+									.Visibility( this, &SFindInBlueprints::GetCachingProgressBarVisiblity )
+									.ToolTipText( LOCTEXT("CancelCacheAll_Tooltip", "Stops the caching process from where ever it is, can be started back up where it left off when needed.") )
+								]
+
+							// "X" to remove the bar
+							+SHorizontalBox::Slot()
+								.HAlign(HAlign_Right)
+								[
+									SNew(SButton)
+									.ButtonStyle( FCoreStyle::Get(), "NoBorder" )
+									.ContentPadding(0)
+									.HAlign(HAlign_Center)
+									.VAlign(VAlign_Center)
+									.OnClicked( this, &SFindInBlueprints::OnRemoveCacheBar )
+									.ForegroundColor( FSlateColor::UseForeground() )
+									[
+										SNew(SImage)
+										.Image( FCoreStyle::Get().GetBrush("EditableComboBox.Delete") )
+										.ColorAndOpacity( FSlateColor::UseForeground() )
+									]
+								]
+						]
+
+						+SVerticalBox::Slot()
+							.AutoHeight()
+							.Padding(8.0f, 0.0f, 0.0f, 2.0f)
+							[
+								SNew(STextBlock)
+								.Text(this, &SFindInBlueprints::GetCurrentCacheBlueprintName)
+								.Visibility( this, &SFindInBlueprints::GetCachingBlueprintNameVisiblity )
+								.ColorAndOpacity( FCoreStyle::Get().GetColor("ErrorReporting.ForegroundColor") )
+							]
+					]
+				];
+		}
+	}
+	else
+	{
+		// Because there are no uncached Blueprints, remove the bar
+		OnRemoveCacheBar();
+	}
+}
+
 FReply SFindInBlueprints::OnRemoveCacheBar()
 {
 	if(MainVerticalBox.IsValid() && CacheBarSlot.IsValid())
@@ -1059,6 +1230,8 @@ SFindInBlueprints::~SFindInBlueprints()
 		StreamSearch->Stop();
 		StreamSearch->EnsureCompletion();
 	}
+
+	FFindInBlueprintSearchManager::Get().CancelCacheAll(this);
 }
 
 EActiveTimerReturnType SFindInBlueprints::UpdateSearchResults( double InCurrentTime, float InDeltaTime )
@@ -1091,6 +1264,9 @@ EActiveTimerReturnType SFindInBlueprints::UpdateSearchResults( double InCurrentT
 				TreeView->RequestTreeRefresh();
 			}
 
+			// Add the cache bar if needed.
+			ConditionallyAddCacheBar();
+
 			StreamSearch->EnsureCompletion();
 			StreamSearch.Reset();
 		}
@@ -1380,6 +1556,126 @@ EVisibility SFindInBlueprints::GetSearchbarVisiblity() const
 	return StreamSearch.IsValid()? EVisibility::Visible : EVisibility::Collapsed;
 }
 
+FReply SFindInBlueprints::OnCacheAllBlueprints()
+{
+	if(!FFindInBlueprintSearchManager::Get().IsCacheInProgress())
+	{
+		// Request from the SearchManager a delegate to use for ticking the cache system.
+		FWidgetActiveTimerDelegate WidgetActiveTimer;
+		FFindInBlueprintSearchManager::Get().CacheAllUncachedBlueprints(SharedThis(this), WidgetActiveTimer);
+		RegisterActiveTimer(0.f, WidgetActiveTimer);
+	}
+
+	return FReply::Handled();
+}
+
+FReply SFindInBlueprints::OnCancelCacheAll()
+{
+	FFindInBlueprintSearchManager::Get().CancelCacheAll(this);
+
+	// Resubmit the last search
+	OnSearchTextCommitted(SearchTextField->GetText(), ETextCommit::OnEnter);
+
+	return FReply::Handled();
+}
+
+int32 SFindInBlueprints::GetCurrentCacheIndex() const
+{
+	return FFindInBlueprintSearchManager::Get().GetCurrentCacheIndex();
+}
+
+TOptional<float> SFindInBlueprints::GetPercentCompleteCache() const
+{
+	return FFindInBlueprintSearchManager::Get().GetCacheProgress();
+}
+
+EVisibility SFindInBlueprints::GetCachingProgressBarVisiblity() const
+{
+	return IsCacheInProgress()? EVisibility::Visible : EVisibility::Hidden;
+}
+
+EVisibility SFindInBlueprints::GetCacheAllButtonVisibility() const
+{
+	return IsCacheInProgress()? EVisibility::Collapsed : EVisibility::Visible;
+}
+
+EVisibility SFindInBlueprints::GetCachingBarVisibility() const
+{
+	FFindInBlueprintSearchManager& FindInBlueprintManager = FFindInBlueprintSearchManager::Get();
+	return (FindInBlueprintManager.GetNumberUncachedBlueprints() > 0 || FindInBlueprintManager.GetFailedToCacheCount())? EVisibility::Visible : EVisibility::Collapsed;
+}
+
+EVisibility SFindInBlueprints::GetCachingBlueprintNameVisiblity() const
+{
+	return IsCacheInProgress()? EVisibility::Visible : EVisibility::Collapsed;
+}
+
+EVisibility SFindInBlueprints::GetFailedToCacheListVisibility() const
+{
+	return FFindInBlueprintSearchManager::Get().GetFailedToCacheCount()? EVisibility::Visible : EVisibility::Collapsed;
+}
+
+bool SFindInBlueprints::IsCacheInProgress() const
+{
+	return FFindInBlueprintSearchManager::Get().IsCacheInProgress();
+}
+
+FSlateColor SFindInBlueprints::GetCachingBarColor() const
+{
+	// The caching bar's default color is a darkish red
+	FSlateColor ReturnColor = FSlateColor(FLinearColor(0.4f, 0.0f, 0.0f));
+	if(IsCacheInProgress())
+	{
+		// It turns yellow when in progress
+		ReturnColor = FSlateColor(FLinearColor(0.4f, 0.4f, 0.0f));
+	}
+	return ReturnColor;
+}
+
+FText SFindInBlueprints::GetUncachedBlueprintWarningText() const
+{
+	FFindInBlueprintSearchManager& FindInBlueprintManager = FFindInBlueprintSearchManager::Get();
+
+	int32 FailedToCacheCount = FindInBlueprintManager.GetFailedToCacheCount();
+
+	// The number of unindexed Blueprints is the total of those that failed to cache and those that haven't been attempted yet.
+	FFormatNamedArguments Args;
+	Args.Add(TEXT("Count"), FindInBlueprintManager.GetNumberUncachedBlueprints() + FailedToCacheCount);
+
+	FText ReturnDisplayText;
+	if(IsCacheInProgress())
+	{
+		Args.Add(TEXT("CurrentIndex"), FindInBlueprintManager.GetCurrentCacheIndex());
+
+		ReturnDisplayText = FText::Format(LOCTEXT("CachingBlueprints", "Indexing Blueprints... {CurrentIndex}/{Count}"), Args);
+	}
+	else
+	{
+		ReturnDisplayText = FText::Format(LOCTEXT("UncachedBlueprints", "Search incomplete. {Count} Blueprints are not indexed!"), Args);
+
+		if (FailedToCacheCount > 0)
+		{
+			FFormatNamedArguments ArgsWithCacheFails;
+			Args.Add(TEXT("BaseMessage"), ReturnDisplayText);
+			Args.Add(TEXT("CacheFails"), FailedToCacheCount);
+			ReturnDisplayText = FText::Format(LOCTEXT("UncachedBlueprintsWithCacheFails", "{BaseMessage} {CacheFails} Blueprints failed to cache."), Args);
+		}
+	}
+
+	return ReturnDisplayText;
+}
+
+FText SFindInBlueprints::GetCurrentCacheBlueprintName() const
+{
+	return FText::FromString(FFindInBlueprintSearchManager::Get().GetCurrentCacheBlueprintName());
+}
+
+void SFindInBlueprints::OnCacheComplete()
+{
+	// Resubmit the last search, which will also remove the bar if needed
+	OnSearchTextCommitted(SearchTextField->GetText(), ETextCommit::OnEnter);
+}
+
 TSharedPtr<SWidget> SFindInBlueprints::OnContextMenuOpening()
 {
 	const bool bShouldCloseWindowAfterMenuSelection = true;
diff --git a/Engine/Source/Editor/Kismet/Public/FindInBlueprintManager.h b/Engine/Source/Editor/Kismet/Public/FindInBlueprintManager.h
--- a/Engine/Source/Editor/Kismet/Public/FindInBlueprintManager.h
+++ b/Engine/Source/Editor/Kismet/Public/FindInBlueprintManager.h
@@ -166,9 +166,51 @@ public:
 	/** Serializes an FText to memory and converts the memory into a string of hex characters */
 	static FString ConvertFTextToHexString(FText InValue);
 
+	/** Returns the number of uncached Blueprints */
+	int32 GetNumberUncachedBlueprints() const { return UncachedBlueprints.Num(); };
+
+	/**
+	 * Starts caching all uncached Blueprints at a rate of 1 per tick
+	 *
+	 * @param InSourceWidget		The source FindInBlueprints widget, this widget will be informed when caching is complete
+	 */
+	void CacheAllUncachedBlueprints(TWeakPtr< class SFindInBlueprints > InSourceWidgetm, FWidgetActiveTimerDelegate& OutActiveTimerDelegate);
+	void OnCacheAllUncachedBlueprints(bool bInSourceControlActive);
+
+	/** Stops the caching process where it currently is at, the rest can be continued later */
+	void CancelCacheAll(SFindInBlueprints* InFindInBlueprintWidget);
+
+	/** Returns the current index in the caching */
+	int32 GetCurrentCacheIndex() const;
+
+	/** Returns the name of the current Blueprint being cached */
+	FString GetCurrentCacheBlueprintName() const;
+
+	/** Returns the progress complete on the caching */
+	float GetCacheProgress() const;
+
+	/** Returns the list of Blueprint paths that failed to cache */
+	TArray<FString> GetFailedToCachePathList() const { return FailedToCachePaths; }
+
+	/** Returns the number of Blueprints that failed to cache */
+	int32 GetFailedToCacheCount() const { return FailedToCachePaths.Num(); }
+
+	/** Returns TRUE if caching failed */
+	bool HasCachingFailed() const { return FailedToCachePaths.Num() > 0; };
+	/**
+	 * Callback to note that Blueprint caching is complete
+	 *
+	 * @param InNumberCached		The number of Blueprints cached, to be chopped off the existing array so the rest (if any) can be finished later
+	 */
+	void FinishedCachingBlueprints(int32 InNumberCached, TArray<FString>& InFailedToCacheList);
+
+	/** Returns TRUE if Blueprints are being cached. */
+	bool IsCacheInProgress() const;
+
 	/** Serializes an FString to memory and converts the memory into a string of hex characters */
 	static FString ConvertFStringToHexString(FString InValue);
 
+	/** Given a fully constructed Find-in-Blueprint FString of searchable data, will parse and construct a JsonObject */
 	static TSharedPtr< class FJsonObject > ConvertJsonStringToObject(FString InJsonString);
 
 private:
@@ -202,6 +244,19 @@ private:
 	/** Builds the cache from all available Blueprint assets that the asset registry has discovered at the time of this function. Occurs on startup */
 	void BuildCache();
 
+	/**
+	 * Helper to properly add a Blueprint's SearchData to the database
+	 *
+	 * @param InSearchData		Data to add to the database
+	 * @return					Index into the SearchArray for looking up the added item
+	 */
+	int32 AddSearchDataToDatabase(FSearchData& InSearchData);
+
+	/** Removes a Blueprint from being managed by the FiB system by passing in the UBlueprint's path */
+	void RemoveBlueprintByPath(FString InPath);
+
+	/** Removes a World Blueprint with compound FiB searchable data from being managed by the FiB system by passing in the World's path */
+	void RemoveWorldByPath(FString InPath);
 protected:
 	/** Maps the Blueprint paths to their index in the SearchArray */
 	TMap<FString, int> SearchMap;
@@ -230,6 +285,15 @@ protected:
 	/** Because we are unable to query for the module on another thread, cache it for use later */
 	class FAssetRegistryModule* AssetRegistryModule;
 
-	/** A list of all the FindInBlueprints widgets that need to be informed that caching is complete */
-	TArray< TWeakPtr<class SFindInBlueprints> > OnCachingCompleteCallbackWidgets;
+	/** FindInBlueprints widget that started the cache process */
+	TWeakPtr<SFindInBlueprints> SourceCachingWidget;
+
+	/** Blueprint paths that have not been cached for searching due to lack of data, this means that they are either older Blueprints, or the DDC cannot find the data */
+	TArray<FString> UncachedBlueprints;
+
+	/** List of paths for Blueprints that failed to cache */
+	TArray<FString> FailedToCachePaths;
+
+	/** Tickable object that does the caching of uncached Blueprints at a rate of once per tick */
+	class FCacheAllBlueprintsTickableObject* CachingObject;
 };
diff --git a/Engine/Source/Editor/Kismet/Public/FindInBlueprints.h b/Engine/Source/Editor/Kismet/Public/FindInBlueprints.h
--- a/Engine/Source/Editor/Kismet/Public/FindInBlueprints.h
+++ b/Engine/Source/Editor/Kismet/Public/FindInBlueprints.h
@@ -229,7 +229,8 @@ public:
 	/** Focuses this widget's search box, and changes the mode as well, and optionally the search terms */
 	void FocusForUse(bool bSetFindWithinBlueprint, FString NewSearchTerms = FString(), bool bSelectFirstResult = false);
 
-
+	/** Called when caching Blueprints is complete, if this widget initiated the indexing */
+	void OnCacheComplete();
 private:
 	/** Processes results of the ongoing async stream search */
 	EActiveTimerReturnType UpdateSearchResults( double InCurrentTime, float InDeltaTime );
@@ -273,9 +274,51 @@ private:
 	/** Returns the progress bar visiblity */
 	EVisibility GetSearchbarVisiblity() const;
 
+	/** Adds the "cache" bar at the bottom of the Find-in-Blueprints widget, to notify the user that the search is incomplete */
+	void ConditionallyAddCacheBar();
+
 	/** Callback to remove the "cache" bar when a button is pressed */
 	FReply OnRemoveCacheBar();
 
+	/** Callback to return the cache bar's display text, informing the user of the situation */
+	FText GetUncachedBlueprintWarningText() const;
+
+	/** Callback to return the cache bar's current indexing Blueprint name */
+	FText GetCurrentCacheBlueprintName() const;
+
+	/** Callback to cache all uncached Blueprints */
+	FReply OnCacheAllBlueprints();
+
+	/** Callback to cancel the caching process */
+	FReply OnCancelCacheAll();
+
+	/** Retrieves the current index of the Blueprint caching process */
+	int32 GetCurrentCacheIndex() const;
+
+	/** Gets the percent complete of the caching process */
+	TOptional<float> GetPercentCompleteCache() const;
+
+	/** Returns the visibility of the caching progress bar, visible when in progress, hidden when not */
+	EVisibility GetCachingProgressBarVisiblity() const;
+
+	/** Returns the visibility of the "Cache All" button, visible when not caching, collapsed when caching is in progress */
+	EVisibility GetCacheAllButtonVisibility() const;
+
+	/** Returns the caching bar's visibility, it goes invisible when there is nothing to be cached. The next search will remove this bar or make it visible again */
+	EVisibility GetCachingBarVisibility() const;
+
+	/** Returns the visibility of the caching Blueprint name, visible when in progress, collapsed when not */
+	EVisibility GetCachingBlueprintNameVisiblity() const;
+
+	/** Returns the visibility of the popup button that displays the list of Blueprints that failed to cache */
+	EVisibility GetFailedToCacheListVisibility() const;
+
+	/** Returns TRUE if Blueprint caching is in progress */
+	bool IsCacheInProgress() const;
+
+	/** Returns the color of the caching bar */
+	FSlateColor GetCachingBarColor() const;
+
 	/** Callback to build the context menu when right clicking in the tree */
 	TSharedPtr<SWidget> OnContextMenuOpening();
 
diff --git a/Engine/Source/Editor/UnrealEd/Public/Kismet2/KismetEditorUtilities.h b/Engine/Source/Editor/UnrealEd/Public/Kismet2/KismetEditorUtilities.h
--- a/Engine/Source/Editor/UnrealEd/Public/Kismet2/KismetEditorUtilities.h
+++ b/Engine/Source/Editor/UnrealEd/Public/Kismet2/KismetEditorUtilities.h
@@ -10,6 +10,37 @@
 
 class UBlueprintGeneratedClass;
 
+//////////////////////////////////////////////////////////////////////////
+// FBlueprintUnloader
+
+/** Utility struct, used to aid in unloading and replacing a specific blueprint. */
+struct UNREALED_API FBlueprintUnloader
+{
+public:
+	FBlueprintUnloader(UBlueprint* OldBlueprint);
+
+	/** 
+	 * Unloads the specified Blueprint (marking it pending-kill, and removing it 
+	 * from its outer package). Optionally, will unload the package as well.
+	 *
+	 * @param  bResetPackage	Whether or not this should unload the entire package.
+	 */
+	void UnloadBlueprint(const bool bResetPackage);
+	
+	/** 
+	 * Replaces all old references to the original blueprints (its class/CDO/etc.)
+	 * @param  NewBlueprint	The blueprint to replace old references with
+	 */
+	void ReplaceStaleRefs(UBlueprint* NewBlueprint);
+
+private:
+	TWeakObjectPtr<UBlueprint> OldBlueprint;
+	UClass*  OldGeneratedClass;
+	UObject* OldCDO;
+	UClass*  OldSkeletonClass;
+	UObject* OldSkelCDO;
+};
+
 //////////////////////////////////////////////////////////////////////////
 // FKismetEditorUtilities
 
diff --git a/Engine/Source/Runtime/Engine/Private/World.cpp b/Engine/Source/Runtime/Engine/Private/World.cpp
--- a/Engine/Source/Runtime/Engine/Private/World.cpp
+++ b/Engine/Source/Runtime/Engine/Private/World.cpp
@@ -49,6 +49,7 @@
 	#include "Editor/UnrealEd/Public/Kismet2/BlueprintEditorUtils.h"
 	#include "Editor/UnrealEd/Classes/ThumbnailRendering/WorldThumbnailInfo.h"
 	#include "SlateBasics.h"
+	#include "Editor/Kismet/Public/FindInBlueprintManager.h"
 #endif
 
 #include "MallocProfiler.h"
@@ -5519,10 +5520,17 @@ void UWorld::GetAssetRegistryTags(TArray<FAssetRegistryTag>& OutTags) const
 {
 	if(PersistentLevel && PersistentLevel->OwningWorld)
 	{
-		for (UBlueprint* Blueprint : PersistentLevel->GetLevelBlueprints())
+		TArray<UBlueprint*> LevelBlueprints = PersistentLevel->GetLevelBlueprints();
+		for (UBlueprint* Blueprint : LevelBlueprints)
 		{
 			Blueprint->GetAssetRegistryTags(OutTags);
 		}
+
+		// If there are no Blueprints, add empty FiB data so the manager knows that the Blueprint is indexed.
+		if (LevelBlueprints.Num() == 0)
+		{
+			OutTags.Add(FAssetRegistryTag("FiB", FString(), FAssetRegistryTag::TT_Hidden));
+		}
 	}
 }
 #endif
