When entering PIE, there will no longer be errors from function call nodes in the level Blueprint due to missing scope.

Find-in-Blueprints will no longer attempt to gather Blueprints that do not have a skeleton class.

diff --git a/Engine/Source/Editor/Kismet/Private/FindInBlueprintManager.cpp b/Engine/Source/Editor/Kismet/Private/FindInBlueprintManager.cpp
--- a/Engine/Source/Editor/Kismet/Private/FindInBlueprintManager.cpp
+++ b/Engine/Source/Editor/Kismet/Private/FindInBlueprintManager.cpp
@@ -1110,7 +1110,12 @@ void FFindInBlueprintSearchManager::AddOrUpdateBlueprintSearchMetadata(UBlueprin
 
 	// Build the search data
 	SearchArray[Index].BlueprintPath = BlueprintPath;
-	SearchArray[Index].Value = GatherBlueprintSearchMetadata(InBlueprint);
+
+	// Cannot successfully gather most searchable data if there is no SkeletonGeneratedClass, so don't try, leave it as whatever it was last set to
+	if(InBlueprint->SkeletonGeneratedClass != nullptr)
+	{
+		SearchArray[Index].Value = GatherBlueprintSearchMetadata(InBlueprint);
+	}
 	SearchArray[Index].bMarkedForDeletion = false;
 }
 
