#ue4 code to track issues: UE-11595 and UE-11597

diff --git a/Engine/Source/Editor/Kismet/Private/SBlueprintActionMenu.cpp b/Engine/Source/Editor/Kismet/Private/SBlueprintActionMenu.cpp
--- a/Engine/Source/Editor/Kismet/Private/SBlueprintActionMenu.cpp
+++ b/Engine/Source/Editor/Kismet/Private/SBlueprintActionMenu.cpp
@@ -449,14 +449,21 @@ void SBlueprintActionMenu::ConstructActionContext(FBlueprintActionContext& Conte
 	TSharedPtr<FBlueprintEditor> BlueprintEditor = EditorPtr.Pin();
 	bool const bIsContextSensitive = BlueprintEditor->GetIsContextSensitive();
 
-	UBlueprint* Blueprint = BlueprintEditor->GetBlueprintObj();
-	ContextDescOut.Blueprints.Add(Blueprint);
 	// we still want context from the graph (even if the user has unchecked
 	// "Context Sensitive"), otherwise the user would be presented with nodes
 	// that can't be placed in the graph... if the user isn't being presented
 	// with a valid node, then fix it up in filtering
 	ContextDescOut.Graphs.Add(GraphObj);
 
+	UBlueprint* Blueprint = BlueprintEditor->GetBlueprintObj();
+	const bool bBlueprintIsValid = IsValid(Blueprint) && Blueprint->GeneratedClass && (Blueprint->GeneratedClass->ClassGeneratedBy == Blueprint);
+	if (!ensure(bBlueprintIsValid))  // to track UE-11597 and UE-11595
+	{
+		return;
+	}
+
+	ContextDescOut.Blueprints.Add(Blueprint);
+
 	if (bIsContextSensitive)
 	{
 		ContextDescOut.Pins = DraggedFromPins;
diff --git a/Engine/Source/Runtime/Engine/Private/BlueprintGeneratedClass.cpp b/Engine/Source/Runtime/Engine/Private/BlueprintGeneratedClass.cpp
--- a/Engine/Source/Runtime/Engine/Private/BlueprintGeneratedClass.cpp
+++ b/Engine/Source/Runtime/Engine/Private/BlueprintGeneratedClass.cpp
@@ -123,6 +123,11 @@ void UBlueprintGeneratedClass::GetRequiredPreloadDependencies(TArray<UObject*>&
 
 UClass* UBlueprintGeneratedClass::GetAuthoritativeClass()
 {
+	if (nullptr == ClassGeneratedBy) // to track UE-11597 and UE-11595
+	{
+		UE_LOG(LogBlueprint, Fatal, TEXT("UBlueprintGeneratedClass::GetAuthoritativeClass: ClassGeneratedBy is null. class '%s'"), *GetPathName());
+	}
+
 	UBlueprint* GeneratingBP = CastChecked<UBlueprint>(ClassGeneratedBy);
 
 	check(GeneratingBP);
