[UE-12443] Guarding against invalid action pointers when attempting to construct a tooltip for the BP menu's favorite toggle widgets.

diff --git a/Engine/Source/Editor/Kismet/Private/BlueprintPaletteFavorites.cpp b/Engine/Source/Editor/Kismet/Private/BlueprintPaletteFavorites.cpp
--- a/Engine/Source/Editor/Kismet/Private/BlueprintPaletteFavorites.cpp
+++ b/Engine/Source/Editor/Kismet/Private/BlueprintPaletteFavorites.cpp
@@ -255,8 +255,11 @@ bool UBlueprintPaletteFavorites::CanBeFavorited(TSharedPtr<FEdGraphSchemaAction>
 bool UBlueprintPaletteFavorites::IsFavorited(TSharedPtr<FEdGraphSchemaAction> PaletteAction) const
 {
 	bool bIsFavorited = false;
-
-	if (PaletteAction->GetTypeId() == FBlueprintDragDropMenuItem::StaticGetTypeId())
+	if (!PaletteAction.IsValid())
+	{
+		bIsFavorited = false;
+	}
+	else if (PaletteAction->GetTypeId() == FBlueprintDragDropMenuItem::StaticGetTypeId())
 	{
 		FBlueprintDragDropMenuItem* CollectionMenuItem = (FBlueprintDragDropMenuItem*)PaletteAction.Get();
 
diff --git a/Engine/Source/Editor/Kismet/Private/SBlueprintActionMenu.cpp b/Engine/Source/Editor/Kismet/Private/SBlueprintActionMenu.cpp
--- a/Engine/Source/Editor/Kismet/Private/SBlueprintActionMenu.cpp
+++ b/Engine/Source/Editor/Kismet/Private/SBlueprintActionMenu.cpp
@@ -155,9 +155,13 @@ private:
 	ECheckBoxState GetFavoritedState() const
 	{
 		ECheckBoxState FavoriteState = ECheckBoxState::Unchecked;
-		if (GEditor->EditorUserSettings->BlueprintFavorites->IsFavorited(ActionPtr.Pin()))
+		if (ActionPtr.IsValid() && (GEditor != nullptr))
 		{
-			FavoriteState = ECheckBoxState::Checked;
+			const UEditorUserSettings& EditorSettings = GEditor->GetEditorUserSettings();
+			if (UBlueprintPaletteFavorites* BlueprintFavorites = EditorSettings.BlueprintFavorites)
+			{
+				FavoriteState = BlueprintFavorites->IsFavorited(ActionPtr.Pin()) ? ECheckBoxState::Checked : ECheckBoxState::Unchecked;
+			}
 		}
 		return FavoriteState;
 	}
