Using the up and down arrow keys for selection in the MyBlueprint window will correctly select them.

Refactored logic for what selection type is allowed in a SGraphActionMenu to the widgets that react to the selection.

diff --git a/Engine/Source/Editor/AIGraph/Private/SGraphEditorActionMenuAI.cpp b/Engine/Source/Editor/AIGraph/Private/SGraphEditorActionMenuAI.cpp
--- a/Engine/Source/Editor/AIGraph/Private/SGraphEditorActionMenuAI.cpp
+++ b/Engine/Source/Editor/AIGraph/Private/SGraphEditorActionMenuAI.cpp
@@ -68,26 +68,29 @@ TSharedRef<SEditableTextBox> SGraphEditorActionMenuAI::GetFilterTextBox()
 	return GraphActionMenu->GetFilterTextBox();
 }
 
-void SGraphEditorActionMenuAI::OnActionSelected( const TArray< TSharedPtr<FEdGraphSchemaAction> >& SelectedAction )
+void SGraphEditorActionMenuAI::OnActionSelected( const TArray< TSharedPtr<FEdGraphSchemaAction> >& SelectedAction, ESelectInfo::Type InSelectionType )
 {
-	bool bDoDismissMenus = false;
-
-	if (GraphObj)
+	if (InSelectionType == ESelectInfo::OnMouseClick  || InSelectionType == ESelectInfo::OnKeyPress || SelectedAction.Num() == 0)
 	{
-		for ( int32 ActionIndex = 0; ActionIndex < SelectedAction.Num(); ActionIndex++ )
-		{
-			TSharedPtr<FEdGraphSchemaAction> CurrentAction = SelectedAction[ActionIndex];
+		bool bDoDismissMenus = false;
 
-			if ( CurrentAction.IsValid() )
+		if (GraphObj)
+		{
+			for ( int32 ActionIndex = 0; ActionIndex < SelectedAction.Num(); ActionIndex++ )
 			{
-				CurrentAction->PerformAction(GraphObj, DraggedFromPins, NewNodePosition);
-				bDoDismissMenus = true;
+				TSharedPtr<FEdGraphSchemaAction> CurrentAction = SelectedAction[ActionIndex];
+
+				if ( CurrentAction.IsValid() )
+				{
+					CurrentAction->PerformAction(GraphObj, DraggedFromPins, NewNodePosition);
+					bDoDismissMenus = true;
+				}
 			}
 		}
-	}
 
-	if (bDoDismissMenus)
-	{
-		FSlateApplication::Get().DismissAllMenus();
+		if (bDoDismissMenus)
+		{
+			FSlateApplication::Get().DismissAllMenus();
+		}
 	}
 }
diff --git a/Engine/Source/Editor/AIGraph/Public/SGraphEditorActionMenuAI.h b/Engine/Source/Editor/AIGraph/Public/SGraphEditorActionMenuAI.h
--- a/Engine/Source/Editor/AIGraph/Public/SGraphEditorActionMenuAI.h
+++ b/Engine/Source/Editor/AIGraph/Public/SGraphEditorActionMenuAI.h
@@ -42,7 +42,7 @@ protected:
 	SGraphEditor::FActionMenuClosed OnClosedCallback;
 	TSharedPtr<SGraphActionMenu> GraphActionMenu;
 
-	void OnActionSelected( const TArray< TSharedPtr<FEdGraphSchemaAction> >& SelectedAction );
+	void OnActionSelected( const TArray< TSharedPtr<FEdGraphSchemaAction> >& SelectedAction, ESelectInfo::Type InSelectionType );
 
 	/** Callback used to populate all actions list in SGraphActionMenu */
 	void CollectAllActions(FGraphActionListBuilderBase& OutAllActions);
diff --git a/Engine/Source/Editor/BehaviorTreeEditor/Private/SBehaviorTreeBlackboardView.cpp b/Engine/Source/Editor/BehaviorTreeEditor/Private/SBehaviorTreeBlackboardView.cpp
--- a/Engine/Source/Editor/BehaviorTreeEditor/Private/SBehaviorTreeBlackboardView.cpp
+++ b/Engine/Source/Editor/BehaviorTreeEditor/Private/SBehaviorTreeBlackboardView.cpp
@@ -513,13 +513,16 @@ FText SBehaviorTreeBlackboardView::HandleGetSectionTitle(int32 SectionID) const
 	return FText();
 }
 
-void SBehaviorTreeBlackboardView::HandleActionSelected(const TArray< TSharedPtr<FEdGraphSchemaAction> >& SelectedActions) const
+void SBehaviorTreeBlackboardView::HandleActionSelected(const TArray< TSharedPtr<FEdGraphSchemaAction> >& SelectedActions, ESelectInfo::Type InSelectionType) const
 {
-	if(SelectedActions.Num() > 0)
+	if (InSelectionType == ESelectInfo::OnMouseClick  || InSelectionType == ESelectInfo::OnKeyPress || SelectedActions.Num() == 0)
 	{
-		check(SelectedActions[0]->GetTypeId() == FEdGraphSchemaAction_BlackboardEntry::StaticGetTypeId());
-		TSharedPtr<FEdGraphSchemaAction_BlackboardEntry> BlackboardEntry = StaticCastSharedPtr<FEdGraphSchemaAction_BlackboardEntry>(SelectedActions[0]);
-		OnEntrySelected.ExecuteIfBound(&BlackboardEntry->Key, BlackboardEntry->bIsInherited);
+		if(SelectedActions.Num() > 0)
+		{
+			check(SelectedActions[0]->GetTypeId() == FEdGraphSchemaAction_BlackboardEntry::StaticGetTypeId());
+			TSharedPtr<FEdGraphSchemaAction_BlackboardEntry> BlackboardEntry = StaticCastSharedPtr<FEdGraphSchemaAction_BlackboardEntry>(SelectedActions[0]);
+			OnEntrySelected.ExecuteIfBound(&BlackboardEntry->Key, BlackboardEntry->bIsInherited);
+		}
 	}
 }
 
diff --git a/Engine/Source/Editor/BehaviorTreeEditor/Private/SBehaviorTreeBlackboardView.h b/Engine/Source/Editor/BehaviorTreeEditor/Private/SBehaviorTreeBlackboardView.h
--- a/Engine/Source/Editor/BehaviorTreeEditor/Private/SBehaviorTreeBlackboardView.h
+++ b/Engine/Source/Editor/BehaviorTreeEditor/Private/SBehaviorTreeBlackboardView.h
@@ -106,7 +106,7 @@ protected:
 	FText HandleGetSectionTitle(int32 SectionID) const;
 
 	/** Delegate handler used when an action is selected */
-	void HandleActionSelected(const TArray< TSharedPtr<FEdGraphSchemaAction> >& SelectedActions) const;
+	void HandleActionSelected(const TArray< TSharedPtr<FEdGraphSchemaAction> >& SelectedActions, ESelectInfo::Type InSelectionType) const;
 
 	/** Delegate handler used to generate an action's context menu */
 	TSharedPtr<SWidget> HandleContextMenuOpening(TSharedRef<FUICommandList> ToolkitCommands) const;
diff --git a/Engine/Source/Editor/GraphEditor/Private/SGraphActionMenu.cpp b/Engine/Source/Editor/GraphEditor/Private/SGraphActionMenu.cpp
--- a/Engine/Source/Editor/GraphEditor/Private/SGraphActionMenu.cpp
+++ b/Engine/Source/Editor/GraphEditor/Private/SGraphActionMenu.cpp
@@ -1126,11 +1126,7 @@ void SGraphActionMenu::OnItemSelected( TSharedPtr< FGraphActionNode > InSelected
 {
 	if (!bIgnoreUIUpdate)
 	{
-		// Filter out selection changes that should not trigger execution
-		if( ( SelectInfo == ESelectInfo::OnMouseClick )  || ( SelectInfo == ESelectInfo::OnKeyPress ) || !InSelectedItem.IsValid())		
-		{
-			HandleSelection(InSelectedItem);
-		}
+		HandleSelection(InSelectedItem, SelectInfo);
 	}
 }
 
@@ -1201,7 +1197,7 @@ bool SGraphActionMenu::OnMouseButtonDownEvent( TWeakPtr<FEdGraphSchemaAction> In
 		{
 			if( SelectedNode->GetPrimaryAction().Get() == InAction.Pin().Get() )
 			{				
-				bResult = HandleSelection( SelectedNode );
+				bResult = HandleSelection( SelectedNode, ESelectInfo::OnMouseClick );
 			}
 		}
 	}
@@ -1288,19 +1284,19 @@ void SGraphActionMenu::AddReferencedObjects( FReferenceCollector& Collector )
 	}
 }
 
-bool SGraphActionMenu::HandleSelection( TSharedPtr< FGraphActionNode > &InSelectedItem )
+bool SGraphActionMenu::HandleSelection( TSharedPtr< FGraphActionNode > &InSelectedItem, ESelectInfo::Type InSelectionType )
 {
 	bool bResult = false;
 	if( OnActionSelected.IsBound() )
 	{
 		if ( InSelectedItem.IsValid() && InSelectedItem->IsActionNode() )
 		{
-			OnActionSelected.Execute(InSelectedItem->Actions);
+			OnActionSelected.Execute(InSelectedItem->Actions, InSelectionType);
 			bResult = true;
 		}
 		else
 		{
-			OnActionSelected.Execute(TArray< TSharedPtr<FEdGraphSchemaAction> >());
+			OnActionSelected.Execute(TArray< TSharedPtr<FEdGraphSchemaAction> >(), InSelectionType);
 			bResult = true;
 		}
 	}
diff --git a/Engine/Source/Editor/GraphEditor/Private/SGraphEditorActionMenu.cpp b/Engine/Source/Editor/GraphEditor/Private/SGraphEditorActionMenu.cpp
--- a/Engine/Source/Editor/GraphEditor/Private/SGraphEditorActionMenu.cpp
+++ b/Engine/Source/Editor/GraphEditor/Private/SGraphEditorActionMenu.cpp
@@ -61,25 +61,28 @@ TSharedRef<SEditableTextBox> SGraphEditorActionMenu::GetFilterTextBox()
 }
 
 
-void SGraphEditorActionMenu::OnActionSelected( const TArray< TSharedPtr<FEdGraphSchemaAction> >& SelectedAction )
+void SGraphEditorActionMenu::OnActionSelected( const TArray< TSharedPtr<FEdGraphSchemaAction> >& SelectedAction, ESelectInfo::Type InSelectionType )
 {
-	bool bDoDismissMenus = true;
-
-	if ( GraphObj != NULL )
+	if (InSelectionType == ESelectInfo::OnMouseClick  || InSelectionType == ESelectInfo::OnKeyPress || SelectedAction.Num() == 0)
 	{
-		for ( int32 ActionIndex = 0; ActionIndex < SelectedAction.Num(); ActionIndex++ )
-		{
-			TSharedPtr<FEdGraphSchemaAction> CurrentAction = SelectedAction[ActionIndex];
+		bool bDoDismissMenus = true;
 
-			if ( CurrentAction.IsValid() )
+		if ( GraphObj != NULL )
+		{
+			for ( int32 ActionIndex = 0; ActionIndex < SelectedAction.Num(); ActionIndex++ )
 			{
-				if ( bDoDismissMenus )
+				TSharedPtr<FEdGraphSchemaAction> CurrentAction = SelectedAction[ActionIndex];
+
+				if ( CurrentAction.IsValid() )
 				{
-					FSlateApplication::Get().DismissAllMenus();
-					bDoDismissMenus = false;
-				}
+					if ( bDoDismissMenus )
+					{
+						FSlateApplication::Get().DismissAllMenus();
+						bDoDismissMenus = false;
+					}
 
-				CurrentAction->PerformAction(GraphObj, DraggedFromPins, NewNodePosition);
+					CurrentAction->PerformAction(GraphObj, DraggedFromPins, NewNodePosition);
+				}
 			}
 		}
 	}
diff --git a/Engine/Source/Editor/GraphEditor/Private/SGraphEditorActionMenu.h b/Engine/Source/Editor/GraphEditor/Private/SGraphEditorActionMenu.h
--- a/Engine/Source/Editor/GraphEditor/Private/SGraphEditorActionMenu.h
+++ b/Engine/Source/Editor/GraphEditor/Private/SGraphEditorActionMenu.h
@@ -41,7 +41,7 @@ protected:
 
 	TSharedPtr<SGraphActionMenu> GraphActionMenu;
 
-	void OnActionSelected( const TArray< TSharedPtr<FEdGraphSchemaAction> >& SelectedAction );
+	void OnActionSelected( const TArray< TSharedPtr<FEdGraphSchemaAction> >& SelectedAction, ESelectInfo::Type InSelectionType );
 
 	/** Callback used to populate all actions list in SGraphActionMenu */
 	void CollectAllActions(FGraphActionListBuilderBase& OutAllActions);
diff --git a/Engine/Source/Editor/GraphEditor/Public/SGraphActionMenu.h b/Engine/Source/Editor/GraphEditor/Public/SGraphActionMenu.h
--- a/Engine/Source/Editor/GraphEditor/Public/SGraphActionMenu.h
+++ b/Engine/Source/Editor/GraphEditor/Public/SGraphActionMenu.h
@@ -72,7 +72,7 @@ public:
 	/** Delegate that can be used to create a custom "expander" widget for a particular row */
 	DECLARE_DELEGATE_RetVal_OneParam( TSharedRef<SExpanderArrow>, FOnCreateCustomRowExpander, FCustomExpanderData const& );
 	/** Delegate executed when an action is selected */
-	DECLARE_DELEGATE_OneParam( FOnActionSelected, const TArray< TSharedPtr<FEdGraphSchemaAction> >& );
+	DECLARE_DELEGATE_TwoParams( FOnActionSelected, const TArray< TSharedPtr<FEdGraphSchemaAction> >&, ESelectInfo::Type );
 	/** Delegate executed when an action is double clicked */
 	DECLARE_DELEGATE_OneParam( FOnActionDoubleClicked, const TArray< TSharedPtr<FEdGraphSchemaAction> >& );
 	/** Delegate executed when an action is dragged */
@@ -295,7 +295,7 @@ protected:
 	 *
 	 * @return true if item selection delegate was executed
 	 */
-	bool HandleSelection( TSharedPtr< FGraphActionNode > &InSelectedItem );
+	bool HandleSelection( TSharedPtr< FGraphActionNode > &InSelectedItem, ESelectInfo::Type InSelectionType );
 
 	/** Called when tree item is double clicked */
 	void OnItemDoubleClicked( TSharedPtr< FGraphActionNode > InClickedItem );
diff --git a/Engine/Source/Editor/Kismet/Private/SBlueprintActionMenu.cpp b/Engine/Source/Editor/Kismet/Private/SBlueprintActionMenu.cpp
--- a/Engine/Source/Editor/Kismet/Private/SBlueprintActionMenu.cpp
+++ b/Engine/Source/Editor/Kismet/Private/SBlueprintActionMenu.cpp
@@ -504,24 +504,27 @@ TSharedRef<SWidget> SBlueprintActionMenu::OnCreateWidgetForAction(FCreateWidgetF
 	return SNew(SBlueprintPaletteItem, InCreateData, EditorPtr.Pin()); 
 }
 
-void SBlueprintActionMenu::OnActionSelected( const TArray< TSharedPtr<FEdGraphSchemaAction> >& SelectedAction )
+void SBlueprintActionMenu::OnActionSelected( const TArray< TSharedPtr<FEdGraphSchemaAction> >& SelectedAction, ESelectInfo::Type InSelectionType )
 {
-	for ( int32 ActionIndex = 0; ActionIndex < SelectedAction.Num(); ActionIndex++ )
+	if (InSelectionType == ESelectInfo::OnMouseClick  || InSelectionType == ESelectInfo::OnKeyPress || SelectedAction.Num() == 0)
 	{
-		if ( SelectedAction[ActionIndex].IsValid() && GraphObj != NULL )
+		for ( int32 ActionIndex = 0; ActionIndex < SelectedAction.Num(); ActionIndex++ )
 		{
-			// Don't dismiss when clicking on dummy action
-			if ( !bActionExecuted && (SelectedAction[ActionIndex]->GetTypeId() != FEdGraphSchemaAction_Dummy::StaticGetTypeId()))
+			if ( SelectedAction[ActionIndex].IsValid() && GraphObj != NULL )
 			{
-				FSlateApplication::Get().DismissAllMenus();
-				bActionExecuted = true;
-			}
-
-			UEdGraphNode* ResultNode = SelectedAction[ActionIndex]->PerformAction(GraphObj, DraggedFromPins, NewNodePosition);
-
-			if ( ResultNode != NULL )
-			{
-				NewNodePosition.Y += UEdGraphSchema_K2::EstimateNodeHeight( ResultNode );
+				// Don't dismiss when clicking on dummy action
+				if ( !bActionExecuted && (SelectedAction[ActionIndex]->GetTypeId() != FEdGraphSchemaAction_Dummy::StaticGetTypeId()))
+				{
+					FSlateApplication::Get().DismissAllMenus();
+					bActionExecuted = true;
+				}
+
+				UEdGraphNode* ResultNode = SelectedAction[ActionIndex]->PerformAction(GraphObj, DraggedFromPins, NewNodePosition);
+
+				if ( ResultNode != NULL )
+				{
+					NewNodePosition.Y += UEdGraphSchema_K2::EstimateNodeHeight( ResultNode );
+				}
 			}
 		}
 	}
diff --git a/Engine/Source/Editor/Kismet/Private/SBlueprintActionMenu.h b/Engine/Source/Editor/Kismet/Private/SBlueprintActionMenu.h
--- a/Engine/Source/Editor/Kismet/Private/SBlueprintActionMenu.h
+++ b/Engine/Source/Editor/Kismet/Private/SBlueprintActionMenu.h
@@ -45,7 +45,7 @@ protected:
 	ECheckBoxState ContextToggleIsChecked() const;
 	void OnContextTargetsChanged(uint32 ContextTargetMask);
 
-	void OnActionSelected( const TArray< TSharedPtr<FEdGraphSchemaAction> >& SelectedAction );
+	void OnActionSelected( const TArray< TSharedPtr<FEdGraphSchemaAction> >& SelectedAction, ESelectInfo::Type InSelectionType );
 
 	TSharedRef<SWidget> OnCreateWidgetForAction(struct FCreateWidgetForActionData* const InCreateData);
 
diff --git a/Engine/Source/Editor/Kismet/Private/SMyBlueprint.cpp b/Engine/Source/Editor/Kismet/Private/SMyBlueprint.cpp
--- a/Engine/Source/Editor/Kismet/Private/SMyBlueprint.cpp
+++ b/Engine/Source/Editor/Kismet/Private/SMyBlueprint.cpp
@@ -1497,9 +1497,12 @@ FReply SMyBlueprint::OnCategoryDragged(const FString& InCategory, const FPointer
 	return FReply::Handled().BeginDragDrop(DragOperation);
 }
 
-void SMyBlueprint::OnGlobalActionSelected(const TArray< TSharedPtr<FEdGraphSchemaAction> >& InActions)
+void SMyBlueprint::OnGlobalActionSelected(const TArray< TSharedPtr<FEdGraphSchemaAction> >& InActions, ESelectInfo::Type InSelectionType)
 {
-	OnActionSelected(InActions);
+	if (InSelectionType == ESelectInfo::OnMouseClick  || InSelectionType == ESelectInfo::OnKeyPress || InSelectionType == ESelectInfo::OnNavigation || InActions.Num() == 0)
+	{
+		OnActionSelected(InActions);
+	}
 }
 
 void SMyBlueprint::OnActionSelected( const TArray< TSharedPtr<FEdGraphSchemaAction> >& InActions )
diff --git a/Engine/Source/Editor/Kismet/Private/SMyBlueprint.h b/Engine/Source/Editor/Kismet/Private/SMyBlueprint.h
--- a/Engine/Source/Editor/Kismet/Private/SMyBlueprint.h
+++ b/Engine/Source/Editor/Kismet/Private/SMyBlueprint.h
@@ -153,7 +153,7 @@ private:
 	FReply OnCategoryDragged(const FString& InCategory, const FPointerEvent& MouseEvent);
 	void OnActionSelected(const TArray< TSharedPtr<FEdGraphSchemaAction> >& InActions);
 	static void OnActionSelectedHelper(TSharedPtr<FEdGraphSchemaAction> InAction, UBlueprint* Blueprint, TSharedRef<SKismetInspector> Inspector);
-	void OnGlobalActionSelected(const TArray< TSharedPtr<FEdGraphSchemaAction> >& InActions);
+	void OnGlobalActionSelected(const TArray< TSharedPtr<FEdGraphSchemaAction> >& InActions, ESelectInfo::Type InSelectionType);
 	void OnActionDoubleClicked(const TArray< TSharedPtr<FEdGraphSchemaAction> >& InActions);
 	void ExecuteAction(TSharedPtr<FEdGraphSchemaAction> InAction);
 	TSharedPtr<SWidget> OnContextMenuOpening();
diff --git a/Engine/Source/Editor/SoundClassEditor/Private/SSoundClassActionMenu.cpp b/Engine/Source/Editor/SoundClassEditor/Private/SSoundClassActionMenu.cpp
--- a/Engine/Source/Editor/SoundClassEditor/Private/SSoundClassActionMenu.cpp
+++ b/Engine/Source/Editor/SoundClassEditor/Private/SSoundClassActionMenu.cpp
@@ -108,7 +108,7 @@ void SSoundClassActionMenuItem::OnNewSoundClassNameEntered( const FText& NewText
 	TArray< TSharedPtr<FEdGraphSchemaAction> > ActionList;
 	ActionList.Add( InAction );
 
-	Owner.Pin()->OnActionSelected(ActionList);
+	Owner.Pin()->OnActionSelected(ActionList, ESelectInfo::OnKeyPress);
 }
 
 ///////////////////////////////////////////////
@@ -168,19 +168,22 @@ TSharedRef<SWidget> SSoundClassActionMenu::OnCreateWidgetForAction(struct FCreat
 }
 
 
-void SSoundClassActionMenu::OnActionSelected( const TArray< TSharedPtr<FEdGraphSchemaAction> >& SelectedActions )
+void SSoundClassActionMenu::OnActionSelected( const TArray< TSharedPtr<FEdGraphSchemaAction> >& SelectedActions, ESelectInfo::Type InSelectionType )
 {
-	if ( GraphObj != NULL )
+	if (InSelectionType == ESelectInfo::OnMouseClick  || InSelectionType == ESelectInfo::OnKeyPress || SelectedActions.Num() == 0)
 	{
-		for ( int32 ActionIndex = 0; ActionIndex < SelectedActions.Num(); ActionIndex++ )
+		if ( GraphObj != NULL )
 		{
-			TSharedPtr<FEdGraphSchemaAction> CurrentAction = SelectedActions[ActionIndex];
-
-			if ( CurrentAction.IsValid() )
+			for ( int32 ActionIndex = 0; ActionIndex < SelectedActions.Num(); ActionIndex++ )
 			{
-				FSlateApplication::Get().DismissAllMenus();
+				TSharedPtr<FEdGraphSchemaAction> CurrentAction = SelectedActions[ActionIndex];
+
+				if ( CurrentAction.IsValid() )
+				{
+					FSlateApplication::Get().DismissAllMenus();
 
-				CurrentAction->PerformAction(GraphObj, DraggedFromPins, NewNodePosition);
+					CurrentAction->PerformAction(GraphObj, DraggedFromPins, NewNodePosition);
+				}
 			}
 		}
 	}
diff --git a/Engine/Source/Editor/SoundClassEditor/Private/SSoundClassActionMenu.h b/Engine/Source/Editor/SoundClassEditor/Private/SSoundClassActionMenu.h
--- a/Engine/Source/Editor/SoundClassEditor/Private/SSoundClassActionMenu.h
+++ b/Engine/Source/Editor/SoundClassEditor/Private/SSoundClassActionMenu.h
@@ -64,7 +64,7 @@ protected:
 
 	TSharedPtr<SGraphActionMenu> GraphActionMenu;
 
-	void OnActionSelected( const TArray< TSharedPtr<FEdGraphSchemaAction> >& SelectedActions );
+	void OnActionSelected( const TArray< TSharedPtr<FEdGraphSchemaAction> >& SelectedActions, ESelectInfo::Type InSelectionType );
 
 	TSharedRef<SWidget> OnCreateWidgetForAction(struct FCreateWidgetForActionData* const InCreateData);
 
