Modified some ensures with Find-in-Blueprints to provide more details.

diff --git a/Engine/Source/Editor/Kismet/Private/FindInBlueprintManager.cpp b/Engine/Source/Editor/Kismet/Private/FindInBlueprintManager.cpp
--- a/Engine/Source/Editor/Kismet/Private/FindInBlueprintManager.cpp
+++ b/Engine/Source/Editor/Kismet/Private/FindInBlueprintManager.cpp
@@ -967,7 +967,7 @@ void FFindInBlueprintSearchManager::OnAssetRenamed(const class FAssetData& InAss
 void FFindInBlueprintSearchManager::OnAssetLoaded(UObject* InAsset)
 {
 	UBlueprint* BlueprintAsset = nullptr;
-	FString BlueprintPackagePath = InAsset->GetPathName();
+	FString BlueprintPath = InAsset->GetPathName();
 
 	if(UWorld* WorldAsset = Cast<UWorld>(InAsset))
 	{
@@ -976,7 +976,7 @@ void FFindInBlueprintSearchManager::OnAssetLoaded(UObject* InAsset)
 			BlueprintAsset = Cast<UBlueprint>((UObject*)WorldAsset->PersistentLevel->GetLevelScriptBlueprint(true));
 			if(BlueprintAsset)
 			{
-				BlueprintPackagePath = BlueprintAsset->GetPathName();
+				BlueprintPath = BlueprintAsset->GetPathName();
 			}
 		}
 	}
@@ -990,13 +990,14 @@ void FFindInBlueprintSearchManager::OnAssetLoaded(UObject* InAsset)
 		// Find and update the item in the search array. Searches may currently be active, this will do no harm to them
 
 		// Confirm that the Blueprint has not been added already, this can occur during duplication of Blueprints.
-		int32* IndexPtr = SearchMap.Find(BlueprintPackagePath);
+		int32* IndexPtr = SearchMap.Find(BlueprintPath);
 
 		// The asset registry might not have informed us of this asset yet.
 		if(IndexPtr)
 		{
 			// That index should never have a Blueprint already, but if it does, it should be the same Blueprint!
-			ensure(!SearchArray[*IndexPtr].Blueprint.IsValid() || SearchArray[*IndexPtr].Blueprint == BlueprintAsset);
+			ensureMsgf(!SearchArray[*IndexPtr].Blueprint.IsValid() || SearchArray[*IndexPtr].Blueprint == BlueprintAsset, TEXT("Blueprint in database has path %s and is being stomped by %s"), SearchArray[*IndexPtr].BlueprintPath, BlueprintPath);
+			ensureMsgf(!SearchArray[*IndexPtr].Blueprint.IsValid() || SearchArray[*IndexPtr].BlueprintPath == BlueprintPath, TEXT("Blueprint in database has path %s and is being stomped by %s"), SearchArray[*IndexPtr].BlueprintPath, BlueprintPath);
 			SearchArray[*IndexPtr].Blueprint = BlueprintAsset;
 		}
 	}
