Mac build fix and removal of tautological assignment found by PVS-Studio

diff --git a/Engine/Source/Editor/Kismet/Private/FindInBlueprintManager.cpp b/Engine/Source/Editor/Kismet/Private/FindInBlueprintManager.cpp
--- a/Engine/Source/Editor/Kismet/Private/FindInBlueprintManager.cpp
+++ b/Engine/Source/Editor/Kismet/Private/FindInBlueprintManager.cpp
@@ -996,8 +996,8 @@ void FFindInBlueprintSearchManager::OnAssetLoaded(UObject* InAsset)
 		if(IndexPtr)
 		{
 			// That index should never have a Blueprint already, but if it does, it should be the same Blueprint!
-			ensureMsgf(!SearchArray[*IndexPtr].Blueprint.IsValid() || SearchArray[*IndexPtr].Blueprint == BlueprintAsset, TEXT("Blueprint in database has path %s and is being stomped by %s"), SearchArray[*IndexPtr].BlueprintPath, BlueprintPath);
-			ensureMsgf(!SearchArray[*IndexPtr].Blueprint.IsValid() || SearchArray[*IndexPtr].BlueprintPath == BlueprintPath, TEXT("Blueprint in database has path %s and is being stomped by %s"), SearchArray[*IndexPtr].BlueprintPath, BlueprintPath);
+			ensureMsgf(!SearchArray[*IndexPtr].Blueprint.IsValid() || SearchArray[*IndexPtr].Blueprint == BlueprintAsset, TEXT("Blueprint in database has path %s and is being stomped by %s"), *(SearchArray[*IndexPtr].BlueprintPath), *BlueprintPath);
+			ensureMsgf(!SearchArray[*IndexPtr].Blueprint.IsValid() || SearchArray[*IndexPtr].BlueprintPath == BlueprintPath, TEXT("Blueprint in database has path %s and is being stomped by %s"), *(SearchArray[*IndexPtr].BlueprintPath), *BlueprintPath);
 			SearchArray[*IndexPtr].Blueprint = BlueprintAsset;
 		}
 	}
@@ -1314,10 +1314,6 @@ void FFindInBlueprintSearchManager::CleanCache()
 					// Build the new map/array
 					NewSearchMap.Add(SearchValuePair.Key, NewSearchArray.Add(SearchArray[SearchValuePair.Value]) );
 				}
-				else
-				{
-					SearchArray[SearchValuePair.Value].Blueprint = SearchArray[SearchValuePair.Value].Blueprint;
-				}
 			}
 		}
 	}
