Fix crash applying default changes when an actor needed to be reinstanced

diff --git a/Engine/Source/Editor/Kismet/Private/SSCSEditor.cpp b/Engine/Source/Editor/Kismet/Private/SSCSEditor.cpp
--- a/Engine/Source/Editor/Kismet/Private/SSCSEditor.cpp
+++ b/Engine/Source/Editor/Kismet/Private/SSCSEditor.cpp
@@ -5505,6 +5505,9 @@ void SSCSEditor::OnApplyChangesToBlueprint() const
 	{
 		const FScopedTransaction Transaction(LOCTEXT("PushToBlueprintDefaults_Transaction", "Apply Changes to Blueprint"));
 
+		// Cache the actor label as by the time we need it, it may be invalid
+		const FString ActorLabel = Actor->GetActorLabel();
+
 		// Perform the actual copy
 		{
 			AActor* BlueprintCDO = Actor->GetClass()->GetDefaultObject<AActor>();
@@ -5522,6 +5525,7 @@ void SSCSEditor::OnApplyChangesToBlueprint() const
 				{
 					FBlueprintEditorUtils::MarkBlueprintAsModified(Blueprint);
 					FKismetEditorUtilities::CompileBlueprint(Blueprint);
+					Actor = nullptr; // It is unsafe to use Actor after this point as it may have been reinstanced, so set it to null to make this obvious
 				}
 			}
 		}
@@ -5539,14 +5543,14 @@ void SSCSEditor::OnApplyChangesToBlueprint() const
 				FFormatNamedArguments Args;
 				Args.Add(TEXT("BlueprintName"), FText::FromName(Blueprint->GetFName()));
 				Args.Add(TEXT("NumChangedProperties"), NumChangedProperties);
-				Args.Add(TEXT("ActorName"), FText::FromString(Actor->GetActorLabel()));
+				Args.Add(TEXT("ActorName"), FText::FromString(ActorLabel));
 				NotificationInfo.Text = FText::Format(LOCTEXT("PushToBlueprintDefaults_ApplySuccess", "Updated Blueprint {BlueprintName} ({NumChangedProperties} property changes applied from actor {ActorName})."), Args);
 			}
 			else
 			{
 				FFormatNamedArguments Args;
 				Args.Add(TEXT("BlueprintName"), FText::FromName(Blueprint->GetFName()));
-				Args.Add(TEXT("ActorName"), FText::FromString(Actor->GetActorLabel()));
+				Args.Add(TEXT("ActorName"), FText::FromString(ActorLabel));
 				NotificationInfo.Text = FText::Format(LOCTEXT("PushOneToBlueprintDefaults_ApplySuccess", "Updated Blueprint {BlueprintName} (1 property change applied from actor {ActorName})."), Args);
 			}
 			CompletionState = SNotificationItem::CS_Success;
