BP: Natice Code Generation supports Function Libraries

diff --git a/Engine/Source/Editor/Kismet/Private/NativeCodeGenerationTool.cpp b/Engine/Source/Editor/Kismet/Private/NativeCodeGenerationTool.cpp
--- a/Engine/Source/Editor/Kismet/Private/NativeCodeGenerationTool.cpp
+++ b/Engine/Source/Editor/Kismet/Private/NativeCodeGenerationTool.cpp
@@ -465,7 +465,7 @@ void FNativeCodeGenerationTool::Open(UBlueprint& Blueprint, TSharedRef< class FB
 bool FNativeCodeGenerationTool::CanGenerate(const UBlueprint& Blueprint)
 {
 	return (Blueprint.Status == EBlueprintStatus::BS_UpToDate)
-		&& (Blueprint.BlueprintType == EBlueprintType::BPTYPE_Normal)
+		&& (Blueprint.BlueprintType == EBlueprintType::BPTYPE_Normal || Blueprint.BlueprintType == EBlueprintType::BPTYPE_FunctionLibrary)
 		&& (Blueprint.GeneratedClass->GetClass() == UBlueprintGeneratedClass::StaticClass());
 }
 
