Blueprint - The "Add New" button on the MyBlueprint tab is now "happy green".

diff --git a/Engine/Source/Editor/Kismet/Private/SMyBlueprint.cpp b/Engine/Source/Editor/Kismet/Private/SMyBlueprint.cpp
--- a/Engine/Source/Editor/Kismet/Private/SMyBlueprint.cpp
+++ b/Engine/Source/Editor/Kismet/Private/SMyBlueprint.cpp
@@ -223,6 +223,7 @@ void SMyBlueprint::Construct(const FArguments& InArgs, TWeakPtr<FBlueprintEditor
 	{
 		AddNewMenu = SNew(SComboButton)
 			.ComboButtonStyle(FEditorStyle::Get(), "ToolbarComboButton")
+			.ButtonStyle(FEditorStyle::Get(), "FlatButton.Success")
 			.ForegroundColor(FLinearColor::White)
 			.ToolTipText(LOCTEXT("AddNewToolTip", "Add a new Variable, Graph, Function, Macro, or Event Dispatcher."))
 			.OnGetMenuContent(this, &SMyBlueprint::CreateAddNewMenuWidget)
