#ue4 Overridden component can be moved using 'gizmo'.

It looks like the change was missed while merging cl#2418656 into main branch (cl#2418846)

diff --git a/Engine/Source/Editor/Kismet/Private/SCSEditorViewportClient.cpp b/Engine/Source/Editor/Kismet/Private/SCSEditorViewportClient.cpp
--- a/Engine/Source/Editor/Kismet/Private/SCSEditorViewportClient.cpp
+++ b/Engine/Source/Editor/Kismet/Private/SCSEditorViewportClient.cpp
@@ -392,8 +392,7 @@ bool FSCSEditorViewportClient::InputWidgetDelta( FViewport* Viewport, EAxisList:
 				{
 					FSCSEditorTreeNodePtrType SelectedNodePtr = *It;
 					// Don't allow editing of a root node, inherited SCS node or child node that also has a movable (non-root) parent node selected
-					const bool bCanEdit = !SelectedNodePtr->IsRootComponent() && !SelectedNodePtr->IsInherited()
-						&& !IsMovableParentNodeSelected(SelectedNodePtr, SelectedNodes);
+					const bool bCanEdit = !SelectedNodePtr->IsRootComponent() && !IsMovableParentNodeSelected(SelectedNodePtr, SelectedNodes);
 
 					if(bCanEdit)
 					{
