Can drag and drop pins onto a variety of K2Node_EditablePinBase nodes to add pins to them.

Shifted a lot of pin adding functionality from Blueprint Details customizion to the nodes themselves or BlueprintEditorUtils.

diff --git a/Engine/Source/Editor/BlueprintGraph/Classes/EdGraphSchema_K2.h b/Engine/Source/Editor/BlueprintGraph/Classes/EdGraphSchema_K2.h
--- a/Engine/Source/Editor/BlueprintGraph/Classes/EdGraphSchema_K2.h
+++ b/Engine/Source/Editor/BlueprintGraph/Classes/EdGraphSchema_K2.h
@@ -425,6 +425,8 @@ public:
 	virtual void SplitPin(UEdGraphPin* Pin) const override;
 	virtual void RecombinePin(UEdGraphPin* Pin) const override;
 	virtual void OnPinConnectionDoubleCicked(UEdGraphPin* PinA, UEdGraphPin* PinB, const FVector2D& GraphPosition) const override;
+	virtual UEdGraphPin* DropPinOnNode(UEdGraphNode* InTargetNode, const FString& InSourcePinName, const FEdGraphPinType& InSourcePinType, EEdGraphPinDirection InSourcePinDirection) const override;
+	virtual bool SupportsDropPinOnNode(UEdGraphNode* InTargetNode, const FEdGraphPinType& InSourcePinType, EEdGraphPinDirection InSourcePinDirection, FText& OutErrorMessage) const override;
 	// End EdGraphSchema Interface
 
 	/**
diff --git a/Engine/Source/Editor/BlueprintGraph/Classes/K2Node_CustomEvent.h b/Engine/Source/Editor/BlueprintGraph/Classes/K2Node_CustomEvent.h
--- a/Engine/Source/Editor/BlueprintGraph/Classes/K2Node_CustomEvent.h
+++ b/Engine/Source/Editor/BlueprintGraph/Classes/K2Node_CustomEvent.h
@@ -42,6 +42,7 @@ class UK2Node_CustomEvent : public UK2Node_Event
 
 	// Begin UK2Node_EditablePinBase interface
 	virtual UEdGraphPin* CreatePinFromUserDefinition(const TSharedPtr<FUserPinInfo> NewPinInfo) override;
+	virtual bool CanCreateUserDefinedPin(const FEdGraphPinType& InPinType, EEdGraphPinDirection InDesiredDirection, FText& OutErrorMessage) override;
 	// Begin UK2Node_EditablePinBase interface
 
 	virtual bool IsUsedByAuthorityOnlyDelegate() const override;
diff --git a/Engine/Source/Editor/BlueprintGraph/Classes/K2Node_EditablePinBase.h b/Engine/Source/Editor/BlueprintGraph/Classes/K2Node_EditablePinBase.h
--- a/Engine/Source/Editor/BlueprintGraph/Classes/K2Node_EditablePinBase.h
+++ b/Engine/Source/Editor/BlueprintGraph/Classes/K2Node_EditablePinBase.h
@@ -17,29 +17,29 @@ struct FUserPinInfo
 	UPROPERTY()
 	struct FEdGraphPinType PinType;
 
+	/** Desired direction for the pin. The direction will be forced to work with the node if necessary */
+	UPROPERTY()
+	TEnumAsByte<EEdGraphPinDirection> DesiredPinDirection;
+
 	/** The default value of the pin */
 	UPROPERTY()
 	FString PinDefaultValue;
 
-	/** Direction of the pin */
-	///var EEdGraphPinDirection	PinDirection;
-
-
-		friend FArchive& operator <<(FArchive& Ar, FUserPinInfo& Info)
-		{
-			Ar << Info.PinName;
+	friend FArchive& operator <<(FArchive& Ar, FUserPinInfo& Info)
+	{
+		Ar << Info.PinName;
 
-			Ar << Info.PinType.bIsArray;
-			Ar << Info.PinType.bIsReference;
+		Ar << Info.PinType.bIsArray;
+		Ar << Info.PinType.bIsReference;
 
-			Ar << Info.PinType.PinCategory;
-			Ar << Info.PinType.PinSubCategory;
+		Ar << Info.PinType.PinCategory;
+		Ar << Info.PinType.PinSubCategory;
 
-			Ar << Info.PinType.PinSubCategoryObject;
-			Ar << Info.PinDefaultValue;
+		Ar << Info.PinType.PinSubCategoryObject;
+		Ar << Info.PinDefaultValue;
 
-			return Ar;
-		}
+		return Ar;
+	}
 	
 };
 
@@ -103,13 +103,23 @@ class UK2Node_EditablePinBase : public UK2Node
 	BLUEPRINTGRAPH_API virtual bool ShouldShowNodeProperties() const override { return bIsEditable; }
 	// End of UK2Node interface
 
+	/**
+	 * Queries if a user defined pin of the passed type can be constructed on this node. Node will return false by default and must opt into this functionality
+	 *
+	 * @param			InPinType			The type info for the pin to create
+	 * @param			OutErrorMessage		Only filled with an error if there is pin add support but there is an error with the pin type
+	 * @return			TRUE if a user defined pin can be constructed
+	 */
+	BLUEPRINTGRAPH_API virtual bool CanCreateUserDefinedPin(const FEdGraphPinType& InPinType, EEdGraphPinDirection InDesiredDirection, FText& OutErrorMessage) { return false; }
+
 	/**
 	 * Creates a UserPinInfo from the specified information, and also adds a pin based on that description to the node
 	 *
-	 * @param	InPinName	Name of the pin to create
-	 * @param	InPinType	The type info for the pin to create
+	 * @param	InPinName				Name of the pin to create
+	 * @param	InPinType				The type info for the pin to create
+	 * @param	InDesiredDirection		Desired direction of the pin, will auto-correct if the direction is not allowed on the pin.
 	 */
-	BLUEPRINTGRAPH_API UEdGraphPin* CreateUserDefinedPin(const FString& InPinName, const FEdGraphPinType& InPinType);
+	BLUEPRINTGRAPH_API UEdGraphPin* CreateUserDefinedPin(const FString& InPinName, const FEdGraphPinType& InPinType, EEdGraphPinDirection InDesiredDirection);
 
 	/**
 	 * Removes a pin from the user-defined array, and removes the pin with the same name from the Pins array
diff --git a/Engine/Source/Editor/BlueprintGraph/Classes/K2Node_FunctionEntry.h b/Engine/Source/Editor/BlueprintGraph/Classes/K2Node_FunctionEntry.h
--- a/Engine/Source/Editor/BlueprintGraph/Classes/K2Node_FunctionEntry.h
+++ b/Engine/Source/Editor/BlueprintGraph/Classes/K2Node_FunctionEntry.h
@@ -53,6 +53,7 @@ class UK2Node_FunctionEntry : public UK2Node_FunctionTerminator
 	// End UK2Node_EditablePinBase interface
 
 	// Begin K2Node_FunctionTerminator interface
+	virtual bool CanCreateUserDefinedPin(const FEdGraphPinType& InPinType, EEdGraphPinDirection InDesiredDirection, FText& OutErrorMessage) override;
 	virtual UEdGraphPin* CreatePinFromUserDefinition(const TSharedPtr<FUserPinInfo> NewPinInfo) override;
 	// End K2Node_FunctionTerminator interface
 
diff --git a/Engine/Source/Editor/BlueprintGraph/Classes/K2Node_FunctionResult.h b/Engine/Source/Editor/BlueprintGraph/Classes/K2Node_FunctionResult.h
--- a/Engine/Source/Editor/BlueprintGraph/Classes/K2Node_FunctionResult.h
+++ b/Engine/Source/Editor/BlueprintGraph/Classes/K2Node_FunctionResult.h
@@ -23,6 +23,7 @@ class UK2Node_FunctionResult : public UK2Node_FunctionTerminator
 	// End UK2Node interface
 
 	// Begin K2Node_FunctionTerminator interface
+	virtual bool CanCreateUserDefinedPin(const FEdGraphPinType& InPinType, EEdGraphPinDirection InDesiredDirection, FText& OutErrorMessage) override;
 	virtual UEdGraphPin* CreatePinFromUserDefinition(const TSharedPtr<FUserPinInfo> NewPinInfo) override;
 	// End K2Node_FunctionTerminator interface
 };
diff --git a/Engine/Source/Editor/BlueprintGraph/Classes/K2Node_FunctionTerminator.h b/Engine/Source/Editor/BlueprintGraph/Classes/K2Node_FunctionTerminator.h
--- a/Engine/Source/Editor/BlueprintGraph/Classes/K2Node_FunctionTerminator.h
+++ b/Engine/Source/Editor/BlueprintGraph/Classes/K2Node_FunctionTerminator.h
@@ -25,10 +25,15 @@ class UK2Node_FunctionTerminator : public UK2Node_EditablePinBase
 	// Begin UEdGraphNode interface
 	virtual bool CanDuplicateNode() const override { return false; }
 	virtual FLinearColor GetNodeTitleColor() const override;
+	virtual FString CreateUniquePinName(FString SourcePinName) const override;
 	// End UEdGraphNode interface
 
 	// Begin UK2Node interface
 	virtual bool NodeCausesStructuralBlueprintChange() const override { return true; }
 	// End UK2Node interface
+
+	// Begin UK2Node_EditablePinBase interface
+	virtual bool CanCreateUserDefinedPin(const FEdGraphPinType& InPinType, EEdGraphPinDirection InDesiredDirection, FText& OutErrorMessage) override;
+	// End UK2Node_EditablePinBase interface
 };
 
diff --git a/Engine/Source/Editor/BlueprintGraph/Classes/K2Node_MacroInstance.h b/Engine/Source/Editor/BlueprintGraph/Classes/K2Node_MacroInstance.h
--- a/Engine/Source/Editor/BlueprintGraph/Classes/K2Node_MacroInstance.h
+++ b/Engine/Source/Editor/BlueprintGraph/Classes/K2Node_MacroInstance.h
@@ -66,6 +66,10 @@ public:
 	virtual bool ShouldDrawCompact() const override;
 	// End UK2Node interface
 
+	// Begin UK2Node_EditablePinBase interface
+	virtual bool CanCreateUserDefinedPin(const FEdGraphPinType& InPinType, EEdGraphPinDirection InDesiredDirection, FText& OutErrorMessage) override { return false; }
+	// End UK2Node_EditablePinBase interface
+
 	void SetMacroGraph(UEdGraph* Graph) { MacroGraphReference.SetGraph(Graph); }
 	UEdGraph* GetMacroGraph() const { return MacroGraphReference.GetGraph(); }
 	UBlueprint* GetSourceBlueprint() const { return MacroGraphReference.GetBlueprint(); }
diff --git a/Engine/Source/Editor/BlueprintGraph/Classes/K2Node_MathExpression.h b/Engine/Source/Editor/BlueprintGraph/Classes/K2Node_MathExpression.h
--- a/Engine/Source/Editor/BlueprintGraph/Classes/K2Node_MathExpression.h
+++ b/Engine/Source/Editor/BlueprintGraph/Classes/K2Node_MathExpression.h
@@ -43,6 +43,10 @@ public:
 	virtual void GetMenuActions(FBlueprintActionDatabaseRegistrar& ActionRegistrar) const override;
 	// End of UK2Node interface
 
+	// Begin UK2Node_EditablePinBase interface
+	virtual bool CanCreateUserDefinedPin(const FEdGraphPinType& InPinType, EEdGraphPinDirection InDesiredDirection, FText& OutErrorMessage) override { return false; }
+	// End UK2Node_EditablePinBase interface
+
 private:
 	/**
 	* Clears this node's sub-graph, and then takes the supplied string and
diff --git a/Engine/Source/Editor/BlueprintGraph/Classes/K2Node_Tunnel.h b/Engine/Source/Editor/BlueprintGraph/Classes/K2Node_Tunnel.h
--- a/Engine/Source/Editor/BlueprintGraph/Classes/K2Node_Tunnel.h
+++ b/Engine/Source/Editor/BlueprintGraph/Classes/K2Node_Tunnel.h
@@ -55,6 +55,7 @@ class UK2Node_Tunnel : public UK2Node_EditablePinBase
 	BLUEPRINTGRAPH_API virtual bool CanModifyExecutionWires() override;
 	BLUEPRINTGRAPH_API virtual ERenamePinResult RenameUserDefinedPin(const FString& OldName, const FString& NewName, bool bTest = false) override;
 	virtual bool CanUseRefParams() const override { return true; }
+	virtual bool CanCreateUserDefinedPin(const FEdGraphPinType& InPinType, EEdGraphPinDirection InDesiredDirection, FText& OutErrorMessage) override;
 	// End UK2Node_EditablePinBase interface
 public:
 	// The input pins of this tunnel go to the output pins of InputSinkNode (can be NULL).
diff --git a/Engine/Source/Editor/BlueprintGraph/Private/EdGraphSchema_K2.cpp b/Engine/Source/Editor/BlueprintGraph/Private/EdGraphSchema_K2.cpp
--- a/Engine/Source/Editor/BlueprintGraph/Private/EdGraphSchema_K2.cpp
+++ b/Engine/Source/Editor/BlueprintGraph/Private/EdGraphSchema_K2.cpp
@@ -5701,6 +5701,108 @@ UK2Node_VariableSet* UEdGraphSchema_K2::SpawnVariableSetNode(const FVector2D Gra
 	return FEdGraphSchemaAction_K2NewNode::SpawnNodeFromTemplate<UK2Node_VariableSet>(ParentGraph, NodeTemplate, GraphPosition);
 }
 
+UEdGraphPin* UEdGraphSchema_K2::DropPinOnNode(UEdGraphNode* InTargetNode, const FString& InSourcePinName, const FEdGraphPinType& InSourcePinType, EEdGraphPinDirection InSourcePinDirection) const
+{
+	UEdGraphPin* ResultPin = nullptr;
+	if (UK2Node_EditablePinBase* EditablePinNode = Cast<UK2Node_EditablePinBase>(InTargetNode))
+	{
+		EditablePinNode->Modify();
+
+		if (InSourcePinDirection == EGPD_Output && Cast<UK2Node_FunctionEntry>(InTargetNode))
+		{
+			if (UK2Node_EditablePinBase* ResultNode = FBlueprintEditorUtils::FindOrCreateFunctionResultNode(EditablePinNode))
+			{
+				EditablePinNode = ResultNode;
+			}
+			else
+			{
+				// If we did not successfully find or create a result node, just fail out
+				return nullptr;
+			}
+		}
+		else if (InSourcePinDirection == EGPD_Input && Cast<UK2Node_FunctionResult>(InTargetNode))
+		{
+			TArray<UK2Node_FunctionEntry*> FunctionEntryNode;
+			InTargetNode->GetGraph()->GetNodesOfClass(FunctionEntryNode);
+
+			if (FunctionEntryNode.Num() == 1)
+			{
+				EditablePinNode = FunctionEntryNode[0];
+			}
+			else
+			{
+				// If we did not successfully find the entry node, just fail out
+				return nullptr;
+			}
+		}
+
+		FString NewPinName = InSourcePinName;
+		ResultPin = EditablePinNode->CreateUserDefinedPin(NewPinName, InSourcePinType, (InSourcePinDirection == EGPD_Input)? EGPD_Output : EGPD_Input);
+
+		FParamsChangedHelper ParamsChangedHelper;
+		ParamsChangedHelper.ModifiedBlueprints.Add(FBlueprintEditorUtils::FindBlueprintForNode(InTargetNode));
+		FBlueprintEditorUtils::MarkBlueprintAsStructurallyModified(FBlueprintEditorUtils::FindBlueprintForNode(InTargetNode));
+
+		ParamsChangedHelper.Broadcast(FBlueprintEditorUtils::FindBlueprintForNode(InTargetNode), EditablePinNode, InTargetNode->GetGraph());
+
+		for (auto GraphIt = ParamsChangedHelper.ModifiedGraphs.CreateIterator(); GraphIt; ++GraphIt)
+		{
+			if(UEdGraph* ModifiedGraph = *GraphIt)
+			{
+				ModifiedGraph->NotifyGraphChanged();
+			}
+		}
+
+		// Now update all the blueprints that got modified
+		for (auto BlueprintIt = ParamsChangedHelper.ModifiedBlueprints.CreateIterator(); BlueprintIt; ++BlueprintIt)
+		{
+			if(UBlueprint* Blueprint = *BlueprintIt)
+			{
+				FBlueprintEditorUtils::MarkBlueprintAsModified(Blueprint);
+				Blueprint->BroadcastChanged();
+			}
+		}
+	}
+	return ResultPin;
+}
+
+bool UEdGraphSchema_K2::SupportsDropPinOnNode(UEdGraphNode* InTargetNode, const FEdGraphPinType& InSourcePinType, EEdGraphPinDirection InSourcePinDirection, FText& OutErrorMessage) const
+{
+	bool bIsSupported = false;
+	if (UK2Node_EditablePinBase* EditablePinNode = Cast<UK2Node_EditablePinBase>(InTargetNode))
+	{
+		if (InSourcePinDirection == EGPD_Output && Cast<UK2Node_FunctionEntry>(InTargetNode))
+		{
+			// Just check with the Function Entry and see if it's legal, we'll create/use a result node if the user drops
+			bIsSupported = EditablePinNode->CanCreateUserDefinedPin(InSourcePinType, InSourcePinDirection, OutErrorMessage);
+
+			if (bIsSupported)
+			{
+				OutErrorMessage = LOCTEXT("AddConnectResultNode", "Add Pin to Result Node");
+			}
+		}
+		else if (InSourcePinDirection == EGPD_Input && Cast<UK2Node_FunctionResult>(InTargetNode))
+		{
+			// Just check with the Function Result and see if it's legal, we'll create/use a result node if the user drops
+			bIsSupported = EditablePinNode->CanCreateUserDefinedPin(InSourcePinType, InSourcePinDirection, OutErrorMessage);
+
+			if (bIsSupported)
+			{
+				OutErrorMessage = LOCTEXT("AddPinEntryNode", "Add Pin to Entry Node");
+			}
+		}
+		else
+		{
+			bIsSupported = EditablePinNode->CanCreateUserDefinedPin(InSourcePinType, (InSourcePinDirection == EGPD_Input)? EGPD_Output : EGPD_Input, OutErrorMessage);
+			if (bIsSupported)
+			{
+				OutErrorMessage = LOCTEXT("AddPinToNode", "Add Pin to Node");
+			}
+		}
+	}
+	return bIsSupported;
+}
+
 /////////////////////////////////////////////////////
 
 #undef LOCTEXT_NAMESPACE
diff --git a/Engine/Source/Editor/BlueprintGraph/Private/K2Node_BaseAsyncTask.cpp b/Engine/Source/Editor/BlueprintGraph/Private/K2Node_BaseAsyncTask.cpp
--- a/Engine/Source/Editor/BlueprintGraph/Private/K2Node_BaseAsyncTask.cpp
+++ b/Engine/Source/Editor/BlueprintGraph/Private/K2Node_BaseAsyncTask.cpp
@@ -199,7 +199,7 @@ bool UK2Node_BaseAsyncTask::FBaseAsyncTaskHelper::CopyEventSignature(UK2Node_Cus
 		{
 			FEdGraphPinType PinType;
 			bResult &= Schema->ConvertPropertyToPinType(Param, /*out*/ PinType);
-			bResult &= (NULL != CENode->CreateUserDefinedPin(Param->GetName(), PinType));
+			bResult &= (NULL != CENode->CreateUserDefinedPin(Param->GetName(), PinType, EGPD_Output));
 		}
 	}
 	return bResult;
diff --git a/Engine/Source/Editor/BlueprintGraph/Private/K2Node_CustomEvent.cpp b/Engine/Source/Editor/BlueprintGraph/Private/K2Node_CustomEvent.cpp
--- a/Engine/Source/Editor/BlueprintGraph/Private/K2Node_CustomEvent.cpp
+++ b/Engine/Source/Editor/BlueprintGraph/Private/K2Node_CustomEvent.cpp
@@ -142,6 +142,29 @@ void UK2Node_CustomEvent::RefreshNodeTitle()
 	CachedNodeTitle.MarkDirty();
 }
 
+bool UK2Node_CustomEvent::CanCreateUserDefinedPin(const FEdGraphPinType& InPinType, EEdGraphPinDirection InDesiredDirection, FText& OutErrorMessage)
+{
+	if (!IsEditable())
+	{
+		return false;
+	}
+
+	// Make sure that if this is an exec node we are allowed one.
+	const UEdGraphSchema_K2* Schema = GetDefault<UEdGraphSchema_K2>();
+	if(InDesiredDirection == EGPD_Input)
+	{
+		OutErrorMessage = NSLOCTEXT("K2Node", "AddInputPinError", "Cannot add input pins to function result node!");
+		return false;
+	}
+	else if (InPinType.PinCategory == Schema->PC_Exec && !CanModifyExecutionWires())
+	{
+		OutErrorMessage = LOCTEXT("MultipleExecPinError", "Cannot support more exec pins!");
+		return false;
+	}
+
+	return true;
+}
+
 UEdGraphPin* UK2Node_CustomEvent::CreatePinFromUserDefinition(const TSharedPtr<FUserPinInfo> NewPinInfo)
 {
 	UEdGraphPin* NewPin = CreatePin(
@@ -345,7 +368,7 @@ UK2Node_CustomEvent* UK2Node_CustomEvent::CreateFromFunction(FVector2D GraphPosi
 			{
 				FEdGraphPinType PinType;
 				K2Schema->ConvertPropertyToPinType(Param, /*out*/ PinType);
-				CustomEventNode->CreateUserDefinedPin(Param->GetName(), PinType);
+				CustomEventNode->CreateUserDefinedPin(Param->GetName(), PinType, EGPD_Output);
 			}
 		}
 
diff --git a/Engine/Source/Editor/BlueprintGraph/Private/K2Node_EditablePinBase.cpp b/Engine/Source/Editor/BlueprintGraph/Private/K2Node_EditablePinBase.cpp
--- a/Engine/Source/Editor/BlueprintGraph/Private/K2Node_EditablePinBase.cpp
+++ b/Engine/Source/Editor/BlueprintGraph/Private/K2Node_EditablePinBase.cpp
@@ -18,7 +18,7 @@ void UK2Node_EditablePinBase::AllocateDefaultPins()
 	}
 }
 
-UEdGraphPin* UK2Node_EditablePinBase::CreateUserDefinedPin(const FString& InPinName, const FEdGraphPinType& InPinType)
+UEdGraphPin* UK2Node_EditablePinBase::CreateUserDefinedPin(const FString& InPinName, const FEdGraphPinType& InPinType, EEdGraphPinDirection InDesiredDirection)
 {
 	// Sanitize the name, if needed
 	const FString NewPinName = CreateUniquePinName(InPinName);
@@ -27,12 +27,11 @@ UEdGraphPin* UK2Node_EditablePinBase::CreateUserDefinedPin(const FString& InPinN
 	TSharedPtr<FUserPinInfo> NewPinInfo = MakeShareable( new FUserPinInfo() );
 	NewPinInfo->PinName = NewPinName;
 	NewPinInfo->PinType = InPinType;
+	NewPinInfo->DesiredPinDirection = InDesiredDirection;
 	UserDefinedPins.Add(NewPinInfo);
 
 	// Then, add the pin to the actual Pins array
 	UEdGraphPin* NewPin = CreatePinFromUserDefinition(NewPinInfo);
-	
-	check(NewPin);
 
 	return NewPin;
 }
@@ -218,7 +217,7 @@ bool UK2Node_EditablePinBase::CreateUserDefinedPinsForFunctionEntryExit(const UF
 			FEdGraphPinType PinType;
 			K2Schema->ConvertPropertyToPinType(Param, /*out*/ PinType);
 
-			const bool bPinGood = CreateUserDefinedPin(Param->GetName(), PinType) != NULL;
+			const bool bPinGood = CreateUserDefinedPin(Param->GetName(), PinType, Direction) != NULL;
 
 			bAllPinsGood = bAllPinsGood && bPinGood;
 		}
diff --git a/Engine/Source/Editor/BlueprintGraph/Private/K2Node_FunctionEntry.cpp b/Engine/Source/Editor/BlueprintGraph/Private/K2Node_FunctionEntry.cpp
--- a/Engine/Source/Editor/BlueprintGraph/Private/K2Node_FunctionEntry.cpp
+++ b/Engine/Source/Editor/BlueprintGraph/Private/K2Node_FunctionEntry.cpp
@@ -198,8 +198,29 @@ void UK2Node_FunctionEntry::RemoveOutputPin(UEdGraphPin* PinToRemove)
 	}
 }
 
+bool UK2Node_FunctionEntry::CanCreateUserDefinedPin(const FEdGraphPinType& InPinType, EEdGraphPinDirection InDesiredDirection, FText& OutErrorMessage)
+{
+	bool bResult = Super::CanCreateUserDefinedPin(InPinType, InDesiredDirection, OutErrorMessage);
+	if (bResult)
+	{
+		if(InDesiredDirection == EGPD_Input)
+		{
+			OutErrorMessage = LOCTEXT("AddInputPinError", "Cannot add input pins to function entry node!");
+			bResult = false;
+		}
+	}
+	return bResult;
+}
+
 UEdGraphPin* UK2Node_FunctionEntry::CreatePinFromUserDefinition(const TSharedPtr<FUserPinInfo> NewPinInfo)
 {
+	// Make sure that if this is an exec node we are allowed one.
+	const UEdGraphSchema_K2* Schema = GetDefault<UEdGraphSchema_K2>();
+	if (NewPinInfo->PinType.PinCategory == Schema->PC_Exec && !CanModifyExecutionWires())
+	{
+		return nullptr;
+	}
+
 	UEdGraphPin* NewPin = CreatePin(
 		EGPD_Output, 
 		NewPinInfo->PinType.PinCategory, 
diff --git a/Engine/Source/Editor/BlueprintGraph/Private/K2Node_FunctionResult.cpp b/Engine/Source/Editor/BlueprintGraph/Private/K2Node_FunctionResult.cpp
--- a/Engine/Source/Editor/BlueprintGraph/Private/K2Node_FunctionResult.cpp
+++ b/Engine/Source/Editor/BlueprintGraph/Private/K2Node_FunctionResult.cpp
@@ -93,6 +93,20 @@ void UK2Node_FunctionResult::AllocateDefaultPins()
 	FFillDefaultPinValueHelper::FillAll(this);
 }
 
+bool UK2Node_FunctionResult::CanCreateUserDefinedPin(const FEdGraphPinType& InPinType, EEdGraphPinDirection InDesiredDirection, FText& OutErrorMessage)
+{
+	bool bResult = Super::CanCreateUserDefinedPin(InPinType, InDesiredDirection, OutErrorMessage);
+	if (bResult)
+	{
+		if(InDesiredDirection == EGPD_Output)
+		{
+			OutErrorMessage = NSLOCTEXT("K2Node", "AddOutputPinError", "Cannot add output pins to function result node!");
+			bResult = false;
+		}
+	}
+	return bResult;
+}
+
 UEdGraphPin* UK2Node_FunctionResult::CreatePinFromUserDefinition(const TSharedPtr<FUserPinInfo> NewPinInfo)
 {
 	UEdGraphPin* Pin = CreatePin(
diff --git a/Engine/Source/Editor/BlueprintGraph/Private/K2Node_FunctionTerminator.cpp b/Engine/Source/Editor/BlueprintGraph/Private/K2Node_FunctionTerminator.cpp
--- a/Engine/Source/Editor/BlueprintGraph/Private/K2Node_FunctionTerminator.cpp
+++ b/Engine/Source/Editor/BlueprintGraph/Private/K2Node_FunctionTerminator.cpp
@@ -5,6 +5,8 @@
 #include "K2Node_FunctionTerminator.h"
 #include "GraphEditorSettings.h"
 
+#define LOCTEXT_NAMESPACE "K2Node"
+
 UK2Node_FunctionTerminator::UK2Node_FunctionTerminator(const FObjectInitializer& ObjectInitializer)
 	: Super(ObjectInitializer)
 {
@@ -14,3 +16,34 @@ FLinearColor UK2Node_FunctionTerminator::GetNodeTitleColor() const
 {
 	return GetDefault<UGraphEditorSettings>()->FunctionTerminatorNodeTitleColor;
 }
+
+FString UK2Node_FunctionTerminator::CreateUniquePinName(FString InSourcePinName) const
+{
+	const UFunction* FoundFunction = FFunctionFromNodeHelper::FunctionFromNode(this);
+
+	FString ResultName = InSourcePinName;
+	int UniqueNum = 0;
+	// Prevent the unique name from being the same as another of the UFunction's properties
+	while(FindPin(ResultName) || FindField<const UProperty>(FoundFunction, *ResultName) != NULL)
+	{
+		ResultName = FString::Printf(TEXT("%s%d"), *InSourcePinName, ++UniqueNum);
+	}
+	return ResultName;
+}
+
+bool UK2Node_FunctionTerminator::CanCreateUserDefinedPin(const FEdGraphPinType& InPinType, EEdGraphPinDirection InDesiredDirection, FText& OutErrorMessage)
+{
+	const bool bIsEditable = IsEditable();
+
+	// Make sure that if this is an exec node we are allowed one.
+	const UEdGraphSchema_K2* Schema = GetDefault<UEdGraphSchema_K2>();
+	if (bIsEditable && InPinType.PinCategory == Schema->PC_Exec && !CanModifyExecutionWires())
+	{
+		OutErrorMessage = LOCTEXT("MultipleExecPinError", "Cannot support more exec pins!");
+		return false;
+	}
+
+	return bIsEditable;
+}
+
+#undef LOCTEXT_NAMESPACE
diff --git a/Engine/Source/Editor/BlueprintGraph/Private/K2Node_MathExpression.cpp b/Engine/Source/Editor/BlueprintGraph/Private/K2Node_MathExpression.cpp
--- a/Engine/Source/Editor/BlueprintGraph/Private/K2Node_MathExpression.cpp
+++ b/Engine/Source/Editor/BlueprintGraph/Private/K2Node_MathExpression.cpp
@@ -1202,7 +1202,7 @@ public:
 		if (RootFragment.IsValid())
 		{
 			// connect the final node of the expression with the math-node's output
-			UEdGraphPin* ReturnPin = ExitNode->CreateUserDefinedPin(TEXT("ReturnValue"), RootFragment->GetOutputType());
+			UEdGraphPin* ReturnPin = ExitNode->CreateUserDefinedPin(TEXT("ReturnValue"), RootFragment->GetOutputType(), EGPD_Input);
 			if (!RootFragment->ConnectToInput(ReturnPin, MessageLog))
 			{
 				MessageLog.Error(*LOCTEXT("ResultConnectError", "Failed to connect the generated nodes with expression's result pin: '@@'").ToString(),
@@ -1477,7 +1477,7 @@ private:
 			// currently, generated expressions ALWAYS take a float (it is the most versatile type)
 			DefaultType.PinCategory = K2Schema->PC_Float;
 			
-			UEdGraphPin* NewInputPin = EntryNode->CreateUserDefinedPin(VariableIdentifier, DefaultType);
+			UEdGraphPin* NewInputPin = EntryNode->CreateUserDefinedPin(VariableIdentifier, DefaultType, EGPD_Output);
 			InputPinFragment = MakeShareable(new FCodeGenFragment_InputPin(NewInputPin));
 		}
 		
diff --git a/Engine/Source/Editor/BlueprintGraph/Private/K2Node_Tunnel.cpp b/Engine/Source/Editor/BlueprintGraph/Private/K2Node_Tunnel.cpp
--- a/Engine/Source/Editor/BlueprintGraph/Private/K2Node_Tunnel.cpp
+++ b/Engine/Source/Editor/BlueprintGraph/Private/K2Node_Tunnel.cpp
@@ -7,6 +7,8 @@
 #include "K2Node_Tunnel.h"
 #include "K2Node_Composite.h"
 
+#define LOCTEXT_NAMESPACE "K2Node"
+
 UK2Node_Tunnel::UK2Node_Tunnel(const FObjectInitializer& ObjectInitializer)
 	: Super(ObjectInitializer)
 {
@@ -126,10 +128,40 @@ UK2Node_Tunnel* UK2Node_Tunnel::GetOutputSource() const
 	return OutputSourceNode;
 }
 
+bool UK2Node_Tunnel::CanCreateUserDefinedPin(const FEdGraphPinType& InPinType, EEdGraphPinDirection InDesiredDirection, FText& OutErrorMessage)
+{
+	bool bResult = true;
+	// Make sure that if this is an exec node we are allowed one.
+	const UEdGraphSchema_K2* Schema = GetDefault<UEdGraphSchema_K2>();
+	if (InPinType.PinCategory == Schema->PC_Exec && !CanModifyExecutionWires())
+	{
+		OutErrorMessage = LOCTEXT("MultipleExecPinError", "Cannot support more exec pins!");
+		bResult = false;
+	}
+	else if(InDesiredDirection == EGPD_Input && !bCanHaveInputs)
+	{
+		OutErrorMessage = LOCTEXT("AddTunnelInputPinError", "Cannot add input pins to entry node!");
+		bResult = false;
+	}
+	else if(InDesiredDirection == EGPD_Output && !bCanHaveOutputs)
+	{
+		OutErrorMessage = LOCTEXT("AddTunnelOutputPinError", "Cannot add output pins to entry node!");
+		bResult = false;
+	}
+	return bResult;
+}
+
 UEdGraphPin* UK2Node_Tunnel::CreatePinFromUserDefinition(const TSharedPtr<FUserPinInfo> NewPinInfo)
 {
 	// Create the new pin
 	EEdGraphPinDirection Direction = bCanHaveInputs ? EGPD_Input : EGPD_Output;
+
+	// Let the user pick the pin direction if legal
+	if ( (bCanHaveInputs && NewPinInfo->DesiredPinDirection == EGPD_Input) || (bCanHaveOutputs && NewPinInfo->DesiredPinDirection == EGPD_Output) )
+	{
+		Direction = NewPinInfo->DesiredPinDirection;
+	}
+
 	UEdGraphPin* Result = CreatePin(
 		Direction, 
 		NewPinInfo->PinType.PinCategory, 
@@ -159,7 +191,7 @@ UEdGraphPin* UK2Node_Tunnel::CreatePinFromUserDefinition(const TSharedPtr<FUserP
 	}
 	else if (UK2Node_Tunnel* TunnelNode = Cast<UK2Node_Tunnel>(TargetNode))
 	{
-		TunnelNode->CreateUserDefinedPin(NewPinInfo->PinName, NewPinInfo->PinType);
+		TunnelNode->CreateUserDefinedPin(NewPinInfo->PinName, NewPinInfo->PinType, (Direction == EGPD_Input)? EGPD_Output : EGPD_Input);
 	}
 
 	//@TODO: Automatically update loaded macro instances when this node is changed too
@@ -206,3 +238,4 @@ UObject* UK2Node_Tunnel::GetJumpTargetForDoubleClick() const
 
 	return TargetNode;
 }
+#undef LOCTEXT_NAMESPACE
\ No newline at end of file
diff --git a/Engine/Source/Editor/GraphEditor/Private/DragConnection.cpp b/Engine/Source/Editor/GraphEditor/Private/DragConnection.cpp
--- a/Engine/Source/Editor/GraphEditor/Private/DragConnection.cpp
+++ b/Engine/Source/Editor/GraphEditor/Private/DragConnection.cpp
@@ -34,9 +34,7 @@ void FDragConnection::HoverTargetChanged()
 {
 	TArray<FPinConnectionResponse> UniqueMessages;
 
-	UEdGraphPin* TargetPinObj = GetHoveredPin();
-
-	if (TargetPinObj != NULL)
+	if (UEdGraphPin* TargetPinObj = GetHoveredPin())
 	{
 		TArray<UEdGraphPin*> ValidSourcePins;
 		ValidateGraphPinList(/*out*/ ValidSourcePins);
@@ -65,6 +63,33 @@ void FDragConnection::HoverTargetChanged()
 			}
 		}
 	}
+	else if(UEdGraphNode* TargetNodeObj = GetHoveredNode())
+	{
+		TArray<UEdGraphPin*> ValidSourcePins;
+		ValidateGraphPinList(/*out*/ ValidSourcePins);
+
+		// Check the schema for connection responses
+		for (UEdGraphPin* StartingPinObj : ValidSourcePins)
+		{
+			FPinConnectionResponse Response;
+			FText ResponseText;
+			if (StartingPinObj->GetOwningNode() != TargetNodeObj && StartingPinObj->GetSchema()->SupportsDropPinOnNode(TargetNodeObj, StartingPinObj->PinType, StartingPinObj->Direction, ResponseText))
+			{
+				Response.Response = ECanCreateConnectionResponse::CONNECT_RESPONSE_MAKE;
+			}
+			else
+			{
+				Response.Response = ECanCreateConnectionResponse::CONNECT_RESPONSE_DISALLOW;
+			}
+			
+			// Do not display an error if there is no message
+			if (!ResponseText.IsEmpty())
+			{
+				Response.Message = ResponseText;
+				UniqueMessages.AddUnique(Response);
+			}
+		}
+	}
 
 	// Let the user know the status of dropping now
 	if (UniqueMessages.Num() == 0)
@@ -194,6 +219,54 @@ FReply FDragConnection::DroppedOnPin(FVector2D ScreenPosition, FVector2D GraphPo
 	return FReply::Handled();
 }
 
+FReply FDragConnection::DroppedOnNode(FVector2D ScreenPosition, FVector2D GraphPosition)
+{
+	bool bHandledPinDropOnNode = false;
+	UEdGraphNode* HoveredNode = GetHoveredNode();
+
+	if (HoveredNode)
+	{
+		// Gather any source drag pins
+		TArray<UEdGraphPin*> ValidSourcePins;
+		ValidateGraphPinList(/*out*/ ValidSourcePins);
+
+
+		if (ValidSourcePins.Num())
+		{
+			for (UEdGraphPin* SourcePin : ValidSourcePins)
+			{
+				// Check for pin drop support
+				FText ResponseText;
+				if (SourcePin->GetOwningNode() != HoveredNode && SourcePin->GetSchema()->SupportsDropPinOnNode(HoveredNode, SourcePin->PinType, SourcePin->Direction, ResponseText))
+				{
+					bHandledPinDropOnNode = true;
+
+					// Find which pin name to use and drop the pin on the node
+					FString PinName = SourcePin->PinFriendlyName.IsEmpty()? SourcePin->PinName : SourcePin->PinFriendlyName.ToString();
+
+					const FScopedTransaction Transaction( NSLOCTEXT("UnrealEd", "AddInParam", "Add In Parameter" ) );
+
+					UEdGraphPin* EdGraphPin = HoveredNode->GetSchema()->DropPinOnNode(GetHoveredNode(), PinName, SourcePin->PinType, SourcePin->Direction);
+
+					if(EdGraphPin)
+					{
+						SourcePin->Modify();
+						EdGraphPin->Modify();
+						SourcePin->GetSchema()->TryCreateConnection(SourcePin, EdGraphPin);
+					}
+				}
+
+				// If we have not handled the pin drop on node and there is an error message, do not let other actions occur.
+				if(!bHandledPinDropOnNode && !ResponseText.IsEmpty())
+				{
+					bHandledPinDropOnNode = true;
+				}
+			}
+		}
+	}
+	return bHandledPinDropOnNode? FReply::Handled() : FReply::Unhandled();
+}
+
 FReply FDragConnection::DroppedOnPanel( const TSharedRef< SWidget >& Panel, FVector2D ScreenPosition, FVector2D GraphPosition, UEdGraph& Graph)
 {
 	// Gather any source drag pins
diff --git a/Engine/Source/Editor/GraphEditor/Private/DragConnection.h b/Engine/Source/Editor/GraphEditor/Private/DragConnection.h
--- a/Engine/Source/Editor/GraphEditor/Private/DragConnection.h
+++ b/Engine/Source/Editor/GraphEditor/Private/DragConnection.h
@@ -18,6 +18,7 @@ public:
 	// FGraphEditorDragDropAction interface
 	virtual void HoverTargetChanged() override;
 	virtual FReply DroppedOnPin(FVector2D ScreenPosition, FVector2D GraphPosition) override;
+	virtual FReply DroppedOnNode(FVector2D ScreenPosition, FVector2D GraphPosition) override;
 	virtual FReply DroppedOnPanel(const TSharedRef< SWidget >& Panel, FVector2D ScreenPosition, FVector2D GraphPosition, UEdGraph& Graph) override;
 	virtual void OnDragBegin(const TSharedRef<class SGraphPin>& InPin) override;
 	virtual void OnDragged(const class FDragDropEvent& DragDropEvent) override;
diff --git a/Engine/Source/Editor/Kismet/Private/BlueprintDetailsCustomization.cpp b/Engine/Source/Editor/Kismet/Private/BlueprintDetailsCustomization.cpp
--- a/Engine/Source/Editor/Kismet/Private/BlueprintDetailsCustomization.cpp
+++ b/Engine/Source/Editor/Kismet/Private/BlueprintDetailsCustomization.cpp
@@ -2735,40 +2735,11 @@ bool FBaseBlueprintGraphActionDetails::ConditionallyCleanUpResultNode()
 
 bool FBaseBlueprintGraphActionDetails::AttemptToCreateResultNode()
 {
-	UEdGraph* Graph = GetGraph();
-	
-	if( !FunctionResultNodePtr.IsValid() && Graph && GetBlueprintObj() )
+	if (!FunctionResultNodePtr.IsValid())
 	{
-		FGraphNodeCreator<UK2Node_FunctionResult> ResultNodeCreator(*Graph);
-		UK2Node_FunctionResult* FunctionResult = ResultNodeCreator.CreateNode();
-
-		UK2Node_EditablePinBase * FunctionEntryNode = FunctionEntryNodePtr.Get();
-		check(FunctionEntryNode);
-
-		const UEdGraphSchema_K2* Schema = Cast<const UEdGraphSchema_K2>(FunctionResult->GetSchema());
-		FunctionResult->NodePosX = FunctionEntryNode->NodePosX + FunctionEntryNode->NodeWidth + 256;
-		FunctionResult->NodePosY = FunctionEntryNode->NodePosY;
-		UEdGraphSchema_K2::SetNodeMetaData(FunctionResult, FNodeMetadata::DefaultGraphNode);
-		ResultNodeCreator.Finalize();
-		
-		// Connect the function entry to the result node, if applicable
-		UEdGraphPin* ThenPin = Schema->FindExecutionPin(*FunctionEntryNode, EGPD_Output);
-		UEdGraphPin* ReturnPin = Schema->FindExecutionPin(*FunctionResult, EGPD_Input);
-
-		if(ThenPin->LinkedTo.Num() == 0) 
-		{
-			ThenPin->MakeLinkTo(ReturnPin);
-		}
-		else
-		{
-			// Bump the result node up a bit, so it's less likely to fall behind the node the entry is already connected to
-			FunctionResult->NodePosY -= 100;
-		}
-
-		FunctionResultNodePtr = FunctionResult;
+		FunctionResultNodePtr = FBlueprintEditorUtils::FindOrCreateFunctionResultNode(FunctionEntryNodePtr.Get());
 	}
-
-	return (FunctionResultNodePtr.IsValid());
+	return FunctionResultNodePtr.IsValid();
 }
 
 void FBaseBlueprintGraphActionDetails::SetRefreshDelegate(FSimpleDelegate RefreshDelegate, bool bForInputs)
@@ -3111,7 +3082,8 @@ void FBlueprintDelegateActionDetails::OnFunctionSelected(TSharedPtr<FString> Fun
 				UProperty* FuncParam = *PropIt;
 				FEdGraphPinType TypeOut;
 				Schema->ConvertPropertyToPinType(FuncParam, TypeOut);
-				FunctionEntryNode->CreateUserDefinedPin(FuncParam->GetName(), TypeOut);
+				UEdGraphPin* EdGraphPin = FunctionEntryNode->CreateUserDefinedPin(FuncParam->GetName(), TypeOut, EGPD_Output);
+				ensure(EdGraphPin);
 			}
 
 			OnParamsChanged(FunctionEntryNode);
@@ -3119,176 +3091,6 @@ void FBlueprintDelegateActionDetails::OnFunctionSelected(TSharedPtr<FString> Fun
 	}
 }
 
-struct FBasePinChangeHelper
-{
-	static bool NodeIsNotTransient(const UK2Node* Node)
-	{
-		return (NULL != Node)
-			&& !Node->HasAnyFlags(RF_Transient) 
-			&& (NULL != Cast<UEdGraph>(Node->GetOuter()));
-	}
-
-	virtual void EditCompositeTunnelNode(UK2Node_Tunnel* TunnelNode) {}
-
-	virtual void EditMacroInstance(UK2Node_MacroInstance* MacroInstance, UBlueprint* Blueprint) {}
-
-	virtual void EditCallSite(UK2Node_CallFunction* CallSite, UBlueprint* Blueprint) {}
-
-	virtual void EditDelegates(UK2Node_BaseMCDelegate* CallSite, UBlueprint* Blueprint) {}
-
-	virtual void EditCreateDelegates(UK2Node_CreateDelegate* CallSite) {}
-
-	void Broadcast(UK2Node_EditablePinBase* InTargetNode, FBaseBlueprintGraphActionDetails& BlueprintGraphActionDetails, UEdGraph* Graph)
-	{
-		if (UK2Node_Tunnel* TunnelNode = Cast<UK2Node_Tunnel>(InTargetNode))
-		{
-			UBlueprint* Blueprint = FBlueprintEditorUtils::FindBlueprintForGraphChecked(Graph);
-
-			const bool bIsTopLevelFunctionGraph = Blueprint->MacroGraphs.Contains(Graph);
-
-			if (bIsTopLevelFunctionGraph)
-			{
-				// Editing a macro, hit all loaded instances (in open blueprints)
-				for (TObjectIterator<UK2Node_MacroInstance> It(RF_Transient); It; ++It)
-				{
-					UK2Node_MacroInstance* MacroInstance = *It;
-					if (NodeIsNotTransient(MacroInstance) && (MacroInstance->GetMacroGraph() == Graph))
-					{
-						EditMacroInstance(MacroInstance, FBlueprintEditorUtils::FindBlueprintForNode(MacroInstance));
-					}
-				}
-			}
-			else if(NodeIsNotTransient(TunnelNode))
-			{
-				// Editing a composite node, hit the single instance in the parent graph		
-				EditCompositeTunnelNode(TunnelNode);
-			}
-		}
-		else if (UK2Node_FunctionTerminator* FunctionDefNode = Cast<UK2Node_FunctionTerminator>(BlueprintGraphActionDetails.GetFunctionEntryNode().Get()))
-		{
-			// Reconstruct all function call sites that call this function (in open blueprints)
-			for (TObjectIterator<UK2Node_CallFunction> It(RF_Transient); It; ++It)
-			{
-				UK2Node_CallFunction* CallSite = *It;
-				if (NodeIsNotTransient(CallSite))
-				{
-					UBlueprint* CallSiteBlueprint = FBlueprintEditorUtils::FindBlueprintForNode(CallSite);
-
-					const bool bNameMatches = (CallSite->FunctionReference.GetMemberName() == FunctionDefNode->SignatureName);
-					const UClass* MemberParentClass = CallSite->FunctionReference.GetMemberParentClass(CallSite->GetBlueprintClassFromNode());
-					const bool bClassMatchesEasy = (MemberParentClass != NULL) && (MemberParentClass->IsChildOf(FunctionDefNode->SignatureClass));
-					const bool bClassMatchesHard = (CallSiteBlueprint != NULL) && (CallSite->FunctionReference.IsSelfContext()) && (FunctionDefNode->SignatureClass == NULL) && (CallSiteBlueprint == BlueprintGraphActionDetails.GetBlueprintObj());
-					const bool bValidSchema = CallSite->GetSchema() != NULL;
-
-					if (bNameMatches && bValidSchema && (bClassMatchesEasy || bClassMatchesHard))
-					{
-						EditCallSite(CallSite, CallSiteBlueprint);
-					}
-				}
-			}
-
-			if(FBlueprintEditorUtils::IsDelegateSignatureGraph(Graph))
-			{
-				FName GraphName = Graph->GetFName();
-				for (TObjectIterator<UK2Node_BaseMCDelegate> It(RF_Transient); It; ++It)
-				{
-					if(NodeIsNotTransient(*It) && (GraphName == It->GetPropertyName()))
-					{
-						UBlueprint* CallSiteBlueprint = FBlueprintEditorUtils::FindBlueprintForNode(*It);
-						EditDelegates(*It, CallSiteBlueprint);
-					}
-				}
-			}
-
-			for (TObjectIterator<UK2Node_CreateDelegate> It(RF_Transient); It; ++It)
-			{
-				if(NodeIsNotTransient(*It))
-				{
-					EditCreateDelegates(*It);
-				}
-			}
-		}
-	}
-};
-
-struct FParamsChangedHelper : public FBasePinChangeHelper
-{
-	TSet<UBlueprint*> ModifiedBlueprints;
-	TSet<UEdGraph*> ModifiedGraphs;
-
-	virtual void EditCompositeTunnelNode(UK2Node_Tunnel* TunnelNode) override
-	{
-		if (TunnelNode->InputSinkNode != NULL)
-		{
-			TunnelNode->InputSinkNode->ReconstructNode();
-		}
-
-		if (TunnelNode->OutputSourceNode != NULL)
-		{
-			TunnelNode->OutputSourceNode->ReconstructNode();
-		}
-	}
-
-	virtual void EditMacroInstance(UK2Node_MacroInstance* MacroInstance, UBlueprint* Blueprint) override
-	{
-		MacroInstance->ReconstructNode();
-		if (Blueprint)
-		{
-			ModifiedBlueprints.Add(Blueprint);
-		}
-	}
-
-	virtual void EditCallSite(UK2Node_CallFunction* CallSite, UBlueprint* Blueprint) override
-	{
-		CallSite->Modify();
-		CallSite->ReconstructNode();
-		if (Blueprint != NULL)
-		{
-			ModifiedBlueprints.Add(Blueprint);
-		}
-	}
-
-	virtual void EditDelegates(UK2Node_BaseMCDelegate* CallSite, UBlueprint* Blueprint) override
-	{
-		CallSite->Modify();
-		CallSite->ReconstructNode();
-		if (auto AssignNode = Cast<UK2Node_AddDelegate>(CallSite))
-		{
-			if (auto DelegateInPin = AssignNode->GetDelegatePin())
-			{
-				for(auto DelegateOutPinIt = DelegateInPin->LinkedTo.CreateIterator(); DelegateOutPinIt; ++DelegateOutPinIt)
-				{
-					UEdGraphPin* DelegateOutPin = *DelegateOutPinIt;
-					if (auto CustomEventNode = (DelegateOutPin ? Cast<UK2Node_CustomEvent>(DelegateOutPin->GetOwningNode()) : NULL))
-					{
-						CustomEventNode->ReconstructNode();
-					}
-				}
-			}
-		}
-		if (Blueprint != NULL)
-		{
-			ModifiedBlueprints.Add(Blueprint);
-		}
-	}
-
-	virtual void EditCreateDelegates(UK2Node_CreateDelegate* CallSite) override
-	{
-		UBlueprint* Blueprint = NULL;
-		UEdGraph* Graph = NULL;
-		CallSite->HandleAnyChange(Graph, Blueprint);
-		if(Blueprint)
-		{
-			ModifiedBlueprints.Add(Blueprint);
-		}
-		if(Graph)
-		{
-			ModifiedGraphs.Add(Graph);
-		}
-	}
-
-};
-
 void FBaseBlueprintGraphActionDetails::OnParamsChanged(UK2Node_EditablePinBase* TargetNode, bool bForceRefresh)
 {
 	UEdGraph* Graph = GetGraph();
@@ -3305,7 +3107,7 @@ void FBaseBlueprintGraphActionDetails::OnParamsChanged(UK2Node_EditablePinBase*
 		ParamsChangedHelper.ModifiedBlueprints.Add(GetBlueprintObj());
 		FBlueprintEditorUtils::MarkBlueprintAsStructurallyModified(GetBlueprintObj());
 
-		ParamsChangedHelper.Broadcast(TargetNode, *this, Graph);
+		ParamsChangedHelper.Broadcast(GetBlueprintObj(), TargetNode, Graph);
 
 		for (auto GraphIt = ParamsChangedHelper.ModifiedGraphs.CreateIterator(); GraphIt; ++GraphIt)
 		{
@@ -3408,7 +3210,7 @@ bool FBaseBlueprintGraphActionDetails::OnPinRenamed(UK2Node_EditablePinBase* Tar
 		PinRenamedHelper.ModifiedBlueprints.Add(GetBlueprintObj());
 
 		// GATHER 
-		PinRenamedHelper.Broadcast(TargetNode, *this, Graph);
+		PinRenamedHelper.Broadcast(GetBlueprintObj(), TargetNode, Graph);
 
 		// TEST
 		for(auto NodeIter = PinRenamedHelper.NodesToRename.CreateIterator(); NodeIter; ++NodeIter)
@@ -3918,58 +3720,13 @@ ECheckBoxState FBlueprintGraphActionDetails::GetIsPureFunction() const
 	return (EntryNode->ExtraFlags & FUNC_BlueprintPure) ? ECheckBoxState::Checked :  ECheckBoxState::Unchecked;
 }
 
-bool FBaseBlueprintGraphActionDetails::IsPinNameUnique(const FString& TestName) const
-{
-	UK2Node_EditablePinBase * FunctionEntryNode = FunctionEntryNodePtr.Get();
-	UK2Node_EditablePinBase * FunctionResultNode = FunctionResultNodePtr.Get();
-
-	// Check the built in pins
-	for( TArray<UEdGraphPin*>::TIterator it(FunctionEntryNode->Pins); it; ++it )
-	{
-		if( (*it)->PinName == TestName )
-		{
-			return false;
-		}
-	}
-
-	// If there is a result node, that isn't the same as the entry node, check it as well
-	if( FunctionResultNode && (FunctionResultNode != FunctionEntryNode) )
-	{
-		for( TArray<UEdGraphPin*>::TIterator it(FunctionResultNode->Pins); it; ++it )
-		{
-			if( (*it)->PinName == TestName )
-			{
-				return false;
-			}
-		}
-	}
-
-	return true;
-}
-
-void FBaseBlueprintGraphActionDetails::GenerateUniqueParameterName( const FString &BaseName, FString &Result ) const
-{
-	UK2Node_EditablePinBase * FunctionEntryNode = FunctionEntryNodePtr.Get();
-	check(FunctionEntryNode);
-
-	const auto FoundFunction = FFunctionFromNodeHelper::FunctionFromNode(FunctionEntryNode);
-
-	Result = BaseName;
-	int UniqueNum = 0;
-	// Prevent the unique name from being the same as another of the UFunction's properties
-	while(!IsPinNameUnique(Result) || FindField<const UProperty>(FoundFunction, *Result) != NULL)
-	{
-		Result = FString::Printf(TEXT("%s%d"), *BaseName, ++UniqueNum);
-	}
-}
-
 FReply FBaseBlueprintGraphActionDetails::OnAddNewInputClicked()
 {
 	UK2Node_EditablePinBase * FunctionEntryNode = FunctionEntryNodePtr.Get();
 
 	if( FunctionEntryNode )
 	{
-		const FScopedTransaction Transaction( LOCTEXT( "AddInParam", "Add In Parameter" ) );
+		FScopedTransaction Transaction( LOCTEXT( "AddInParam", "Add In Parameter" ) );
 		FunctionEntryNode->Modify();
 
 		FEdGraphPinType PinType = MyBlueprint.Pin()->GetLastFunctionPinTypeUsed();
@@ -3981,11 +3738,15 @@ FReply FBaseBlueprintGraphActionDetails::OnAddNewInputClicked()
 			MyBlueprint.Pin()->ResetLastPinType();
 			PinType = MyBlueprint.Pin()->GetLastFunctionPinTypeUsed();
 		}
-		FString NewPinName;
-		GenerateUniqueParameterName(TEXT("NewParam"), NewPinName);
-		FunctionEntryNode->CreateUserDefinedPin(NewPinName, PinType);
-
-		OnParamsChanged(FunctionEntryNode, true);
+		FString NewPinName = TEXT("NewParam");
+		if (FunctionEntryNode->CreateUserDefinedPin(NewPinName, PinType, EGPD_Output))
+		{
+			OnParamsChanged(FunctionEntryNode, true);
+		}
+		else
+		{
+			Transaction.Cancel();
+		}
 	}
 
 	return FReply::Handled();
@@ -4039,15 +3800,19 @@ FReply FBlueprintGraphActionDetails::OnAddNewOutputClicked()
 			MyBlueprint.Pin()->ResetLastPinType();
 			PinType = MyBlueprint.Pin()->GetLastFunctionPinTypeUsed();
 		}
-		FString NewPinName;
-		GenerateUniqueParameterName(TEXT("NewParam"), NewPinName);
-		FunctionResultNode->CreateUserDefinedPin(NewPinName, PinType);
-		
-		OnParamsChanged(FunctionResultNode, true);
+		FString NewPinName = TEXT("NewParam");
+		if (FunctionResultNode->CreateUserDefinedPin(NewPinName, PinType, EGPD_Input))
+		{
+			OnParamsChanged(FunctionResultNode, true);
 
-		if ( !PreviousResultNode )
+			if ( !PreviousResultNode )
+			{
+				DetailsLayoutPtr->ForceRefreshDetails();
+			}
+		}
+		else
 		{
-			DetailsLayoutPtr->ForceRefreshDetails();
+			Transaction.Cancel();
 		}
 	}
 	else
diff --git a/Engine/Source/Editor/Kismet/Private/BlueprintDetailsCustomization.h b/Engine/Source/Editor/Kismet/Private/BlueprintDetailsCustomization.h
--- a/Engine/Source/Editor/Kismet/Private/BlueprintDetailsCustomization.h
+++ b/Engine/Source/Editor/Kismet/Private/BlueprintDetailsCustomization.h
@@ -258,10 +258,6 @@ public:
 	
 	FReply OnAddNewInputClicked();
 
-	/** Utility functions for pin names */
-	bool IsPinNameUnique(const FString& TestName) const;
-	void GenerateUniqueParameterName( const FString &BaseName, FString &Result ) const;
-
 protected:
 	/** Tries to create the result node (if there are output args) */
 	bool AttemptToCreateResultNode();
diff --git a/Engine/Source/Editor/Kismet/Private/BlueprintEditor.cpp b/Engine/Source/Editor/Kismet/Private/BlueprintEditor.cpp
--- a/Engine/Source/Editor/Kismet/Private/BlueprintEditor.cpp
+++ b/Engine/Source/Editor/Kismet/Private/BlueprintEditor.cpp
@@ -5742,7 +5742,7 @@ void FBlueprintEditor::ExtractEventTemplateForFunction(class UK2Node_CustomEvent
 							Pin->PinType.bIsArray,
 							Pin->PinType.bIsReference,
 							UniquePortName);
-						InResultNode->CreateUserDefinedPin(UniquePortName, Pin->PinType);
+						InResultNode->CreateUserDefinedPin(UniquePortName, Pin->PinType, EGPD_Input);
 					}
 					PinLink->BreakAllPinLinks();
 					PinLink->MakeLinkTo(RemotePortPin);
@@ -5752,7 +5752,7 @@ void FBlueprintEditor::ExtractEventTemplateForFunction(class UK2Node_CustomEvent
 					InEntryNode->Modify();
 
 					FString UniquePortName = InGatewayNode->CreateUniquePinName(Pin->PinName);
-					InEntryNode->CreateUserDefinedPin(UniquePortName, Pin->PinType);
+					InEntryNode->CreateUserDefinedPin(UniquePortName, Pin->PinType, EGPD_Output);
 				}
 			}
 		}
@@ -5899,7 +5899,7 @@ void FBlueprintEditor::CollapseNodesIntoGraph(UEdGraphNode* InGatewayNode, UK2No
 								LocalPin->PinType.bIsReference,
 								UniquePortName);
 
-							LocalPortPin = LocalPort->CreateUserDefinedPin(UniquePortName, LocalPin->PinType);
+							LocalPortPin = LocalPort->CreateUserDefinedPin(UniquePortName, LocalPin->PinType, (LocalPin->Direction == EGPD_Input)? EGPD_Output : EGPD_Input);
 						}
 					}
 				}
diff --git a/Engine/Source/Editor/UMGEditor/Private/WidgetBlueprintCompiler.cpp b/Engine/Source/Editor/UMGEditor/Private/WidgetBlueprintCompiler.cpp
--- a/Engine/Source/Editor/UMGEditor/Private/WidgetBlueprintCompiler.cpp
+++ b/Engine/Source/Editor/UMGEditor/Private/WidgetBlueprintCompiler.cpp
@@ -68,7 +68,7 @@ void FWidgetBlueprintCompiler::CreateFunctionList()
 				FEdGraphPinType PinType;
 				K2Schema->ConvertPropertyToPinType(Property, /*out*/ PinType);
 
-				UEdGraphPin* ReturnPin = ReturnNode->CreateUserDefinedPin(TEXT("ReturnValue"), PinType);
+				UEdGraphPin* ReturnPin = ReturnNode->CreateUserDefinedPin(TEXT("ReturnValue"), PinType, EGPD_Input);
 
 				// Auto-connect the pins for entry and exit, so that by default the signature is properly generated
 				UEdGraphPin* EntryNodeExec = K2Schema->FindExecutionPin(*EntryNode, EGPD_Output);
diff --git a/Engine/Source/Editor/UnrealEd/Private/Kismet2/BlueprintEditorUtils.cpp b/Engine/Source/Editor/UnrealEd/Private/Kismet2/BlueprintEditorUtils.cpp
--- a/Engine/Source/Editor/UnrealEd/Private/Kismet2/BlueprintEditorUtils.cpp
+++ b/Engine/Source/Editor/UnrealEd/Private/Kismet2/BlueprintEditorUtils.cpp
@@ -274,6 +274,158 @@ static void RenameVariableReferences(UBlueprint* Blueprint, UClass* VariableClas
 	}
 }
 
+//////////////////////////////////////
+// FBasePinChangeHelper
+
+void FBasePinChangeHelper::Broadcast(UBlueprint* InBlueprint, UK2Node_EditablePinBase* InTargetNode, UEdGraph* Graph)
+{
+	if (UK2Node_Tunnel* TunnelNode = Cast<UK2Node_Tunnel>(InTargetNode))
+	{
+		UBlueprint* Blueprint = FBlueprintEditorUtils::FindBlueprintForGraphChecked(Graph);
+
+		const bool bIsTopLevelFunctionGraph = Blueprint->MacroGraphs.Contains(Graph);
+
+		if (bIsTopLevelFunctionGraph)
+		{
+			// Editing a macro, hit all loaded instances (in open blueprints)
+			for (TObjectIterator<UK2Node_MacroInstance> It(RF_Transient); It; ++It)
+			{
+				UK2Node_MacroInstance* MacroInstance = *It;
+				if (NodeIsNotTransient(MacroInstance) && (MacroInstance->GetMacroGraph() == Graph))
+				{
+					EditMacroInstance(MacroInstance, FBlueprintEditorUtils::FindBlueprintForNode(MacroInstance));
+				}
+			}
+		}
+		else if(NodeIsNotTransient(TunnelNode))
+		{
+			// Editing a composite node, hit the single instance in the parent graph		
+			EditCompositeTunnelNode(TunnelNode);
+		}
+	}
+	else if (UK2Node_FunctionTerminator* FunctionDefNode = Cast<UK2Node_FunctionTerminator>(InTargetNode))
+	{
+		// Reconstruct all function call sites that call this function (in open blueprints)
+		for (TObjectIterator<UK2Node_CallFunction> It(RF_Transient); It; ++It)
+		{
+			UK2Node_CallFunction* CallSite = *It;
+			if (NodeIsNotTransient(CallSite))
+			{
+				UBlueprint* CallSiteBlueprint = FBlueprintEditorUtils::FindBlueprintForNode(CallSite);
+
+				const bool bNameMatches = (CallSite->FunctionReference.GetMemberName() == FunctionDefNode->SignatureName);
+				const UClass* MemberParentClass = CallSite->FunctionReference.GetMemberParentClass(CallSite->GetBlueprintClassFromNode());
+				const bool bClassMatchesEasy = (MemberParentClass != NULL) && (MemberParentClass->IsChildOf(FunctionDefNode->SignatureClass));
+				const bool bClassMatchesHard = (CallSiteBlueprint != NULL) && (CallSite->FunctionReference.IsSelfContext()) && (FunctionDefNode->SignatureClass == NULL) && (CallSiteBlueprint == InBlueprint);
+				const bool bValidSchema = CallSite->GetSchema() != NULL;
+
+				if (bNameMatches && bValidSchema && (bClassMatchesEasy || bClassMatchesHard))
+				{
+					EditCallSite(CallSite, CallSiteBlueprint);
+				}
+			}
+		}
+
+		if(FBlueprintEditorUtils::IsDelegateSignatureGraph(Graph))
+		{
+			FName GraphName = Graph->GetFName();
+			for (TObjectIterator<UK2Node_BaseMCDelegate> It(RF_Transient); It; ++It)
+			{
+				if(NodeIsNotTransient(*It) && (GraphName == It->GetPropertyName()))
+				{
+					UBlueprint* CallSiteBlueprint = FBlueprintEditorUtils::FindBlueprintForNode(*It);
+					EditDelegates(*It, CallSiteBlueprint);
+				}
+			}
+		}
+
+		for (TObjectIterator<UK2Node_CreateDelegate> It(RF_Transient); It; ++It)
+		{
+			if(NodeIsNotTransient(*It))
+			{
+				EditCreateDelegates(*It);
+			}
+		}
+	}
+}
+
+//////////////////////////////////////
+// FParamsChangedHelper
+
+void FParamsChangedHelper::EditCompositeTunnelNode(UK2Node_Tunnel* TunnelNode)
+{
+	if (TunnelNode->InputSinkNode != NULL)
+	{
+		TunnelNode->InputSinkNode->ReconstructNode();
+	}
+
+	if (TunnelNode->OutputSourceNode != NULL)
+	{
+		TunnelNode->OutputSourceNode->ReconstructNode();
+	}
+}
+
+void FParamsChangedHelper::EditMacroInstance(UK2Node_MacroInstance* MacroInstance, UBlueprint* Blueprint)
+{
+	MacroInstance->ReconstructNode();
+	if (Blueprint)
+	{
+		ModifiedBlueprints.Add(Blueprint);
+	}
+}
+
+void FParamsChangedHelper::EditCallSite(UK2Node_CallFunction* CallSite, UBlueprint* Blueprint)
+{
+	CallSite->Modify();
+	CallSite->ReconstructNode();
+	if (Blueprint != NULL)
+	{
+		ModifiedBlueprints.Add(Blueprint);
+	}
+}
+
+void FParamsChangedHelper::EditDelegates(UK2Node_BaseMCDelegate* CallSite, UBlueprint* Blueprint)
+{
+	CallSite->Modify();
+	CallSite->ReconstructNode();
+	if (auto AssignNode = Cast<UK2Node_AddDelegate>(CallSite))
+	{
+		if (auto DelegateInPin = AssignNode->GetDelegatePin())
+		{
+			for(auto DelegateOutPinIt = DelegateInPin->LinkedTo.CreateIterator(); DelegateOutPinIt; ++DelegateOutPinIt)
+			{
+				UEdGraphPin* DelegateOutPin = *DelegateOutPinIt;
+				if (auto CustomEventNode = (DelegateOutPin ? Cast<UK2Node_CustomEvent>(DelegateOutPin->GetOwningNode()) : NULL))
+				{
+					CustomEventNode->ReconstructNode();
+				}
+			}
+		}
+	}
+	if (Blueprint != NULL)
+	{
+		ModifiedBlueprints.Add(Blueprint);
+	}
+}
+
+void FParamsChangedHelper::EditCreateDelegates(UK2Node_CreateDelegate* CallSite)
+{
+	UBlueprint* Blueprint = NULL;
+	UEdGraph* Graph = NULL;
+	CallSite->HandleAnyChange(Graph, Blueprint);
+	if(Blueprint)
+	{
+		ModifiedBlueprints.Add(Blueprint);
+	}
+	if(Graph)
+	{
+		ModifiedGraphs.Add(Graph);
+	}
+}
+
+//////////////////////////////////////
+// FBlueprintEditorUtils
+
 void FBlueprintEditorUtils::RefreshAllNodes(UBlueprint* Blueprint)
 {
 	if (!Blueprint || !Blueprint->HasAllFlags(RF_LoadCompleted))
@@ -4142,28 +4294,37 @@ void FBlueprintEditorUtils::RemoveLocalVariable(UBlueprint* InBlueprint, const U
 	}
 }
 
-FFunctionFromNodeHelper::FFunctionFromNodeHelper(UObject* Obj) : Function(FunctionFromNode(Cast<UK2Node>(Obj))), Node(Cast<UK2Node>(Obj))
+FFunctionFromNodeHelper::FFunctionFromNodeHelper(const UObject* Obj) : Function(FunctionFromNode(Cast<UK2Node>(Obj))), Node(Cast<UK2Node>(Obj))
 {
 
 }
 
-UFunction* FFunctionFromNodeHelper::FunctionFromNode(UK2Node* Node)
+UFunction* FFunctionFromNodeHelper::FunctionFromNode(const UK2Node* Node)
 {
 	UFunction* Function = NULL;
 	UBlueprint* Blueprint = Node ? Node->GetBlueprint() : NULL;
 	const UClass* SearchScope = Blueprint ? Blueprint->SkeletonGeneratedClass : NULL;
 	if (SearchScope)
 	{
-		if (UK2Node_Event* EventNode = Cast<UK2Node_Event>(Node))
+		if (const UK2Node_FunctionResult* ResultNode = Cast<UK2Node_FunctionResult>(Node))
 		{
-			// We need to search up the class hierachy by name or functions like CanAddParentNode will fail:
-			Function = SearchScope->FindFunctionByName(EventNode->EventReference.GetMemberName());
+			// Function result nodes cannot resolve the UFunction, so find the entry node and use that for finding the UFunction
+			TArray<UK2Node_FunctionEntry*> EntryNodes;
+			ResultNode->GetGraph()->GetNodesOfClass(EntryNodes);
+
+			check(EntryNodes.Num() == 1);
+			Node = EntryNodes[0];
 		}
-		else if (UK2Node_FunctionEntry* FunctionNode = Cast<UK2Node_FunctionEntry>(Node))
+		else if (const UK2Node_FunctionEntry* FunctionNode = Cast<UK2Node_FunctionEntry>(Node))
 		{
 			const FName FunctionName = (FunctionNode->CustomGeneratedFunctionName != NAME_None) ? FunctionNode->CustomGeneratedFunctionName : FunctionNode->GetGraph()->GetFName();
 			Function = SearchScope->FindFunctionByName(FunctionName);
 		}
+		else if (const UK2Node_Event* EventNode = Cast<UK2Node_Event>(Node))
+		{
+			// We need to search up the class hierachy by name or functions like CanAddParentNode will fail:
+			Function = SearchScope->FindFunctionByName(EventNode->EventReference.GetMemberName());
+		}
 	}
 
 	return Function;
@@ -7345,5 +7506,51 @@ FString FBlueprintEditorUtils::GetClassNameWithoutSuffix(const UClass* Class)
 	}
 }
 
+UK2Node_FunctionResult* FBlueprintEditorUtils::FindOrCreateFunctionResultNode(UK2Node_EditablePinBase* InFunctionEntryNode)
+{
+	UK2Node_FunctionResult* FunctionResult = nullptr;
+
+	if (InFunctionEntryNode)
+	{
+		UEdGraph* Graph = InFunctionEntryNode->GetGraph();
+
+		TArray<UK2Node_FunctionResult*> ResultNode;
+		Graph->GetNodesOfClass(ResultNode);
+
+		if (Graph && ResultNode.Num() == 0)
+		{
+			FGraphNodeCreator<UK2Node_FunctionResult> ResultNodeCreator(*Graph);
+			FunctionResult = ResultNodeCreator.CreateNode();
+
+			const UEdGraphSchema_K2* Schema = Cast<const UEdGraphSchema_K2>(FunctionResult->GetSchema());
+			FunctionResult->NodePosX = InFunctionEntryNode->NodePosX + InFunctionEntryNode->NodeWidth + 256;
+			FunctionResult->NodePosY = InFunctionEntryNode->NodePosY;
+			FunctionResult->bIsEditable = true;
+			UEdGraphSchema_K2::SetNodeMetaData(FunctionResult, FNodeMetadata::DefaultGraphNode);
+			ResultNodeCreator.Finalize();
+
+			// Connect the function entry to the result node, if applicable
+			UEdGraphPin* ThenPin = Schema->FindExecutionPin(*InFunctionEntryNode, EGPD_Output);
+			UEdGraphPin* ReturnPin = Schema->FindExecutionPin(*FunctionResult, EGPD_Input);
+
+			if(ThenPin->LinkedTo.Num() == 0) 
+			{
+				ThenPin->MakeLinkTo(ReturnPin);
+			}
+			else
+			{
+				// Bump the result node up a bit, so it's less likely to fall behind the node the entry is already connected to
+				FunctionResult->NodePosY -= 100;
+			}
+		}
+		else
+		{
+			FunctionResult = ResultNode[0];
+		}
+	}
+
+	return FunctionResult;
+}
+
 #undef LOCTEXT_NAMESPACE
 
diff --git a/Engine/Source/Editor/UnrealEd/Public/Kismet2/BlueprintEditorUtils.h b/Engine/Source/Editor/UnrealEd/Public/Kismet2/BlueprintEditorUtils.h
--- a/Engine/Source/Editor/UnrealEd/Public/Kismet2/BlueprintEditorUtils.h
+++ b/Engine/Source/Editor/UnrealEd/Public/Kismet2/BlueprintEditorUtils.h
@@ -33,11 +33,52 @@ namespace EGraphRemoveFlags
 struct UNREALED_API FFunctionFromNodeHelper
 {
 	UFunction* const Function;
-	UK2Node* const Node;
+	const UK2Node* const Node;
 
-	static UFunction* FunctionFromNode(UK2Node* Node);
+	static UFunction* FunctionFromNode(const UK2Node* Node);
+
+	FFunctionFromNodeHelper(const UObject* Obj);
+};
+
+class UNREALED_API FBasePinChangeHelper
+{
+public:
+	static bool NodeIsNotTransient(const UK2Node* Node)
+	{
+		return (NULL != Node)
+			&& !Node->HasAnyFlags(RF_Transient) 
+			&& (NULL != Cast<UEdGraph>(Node->GetOuter()));
+	}
+
+	virtual void EditCompositeTunnelNode(class UK2Node_Tunnel* TunnelNode) {}
+
+	virtual void EditMacroInstance(class UK2Node_MacroInstance* MacroInstance, UBlueprint* Blueprint) {}
+
+	virtual void EditCallSite(class UK2Node_CallFunction* CallSite, UBlueprint* Blueprint) {}
+
+	virtual void EditDelegates(class UK2Node_BaseMCDelegate* CallSite, UBlueprint* Blueprint) {}
+
+	virtual void EditCreateDelegates(class UK2Node_CreateDelegate* CallSite) {}
+
+	void Broadcast(UBlueprint* InBlueprint, class UK2Node_EditablePinBase* InTargetNode, UEdGraph* Graph);
+};
+
+class UNREALED_API FParamsChangedHelper : public FBasePinChangeHelper
+{
+public:
+	TSet<UBlueprint*> ModifiedBlueprints;
+	TSet<UEdGraph*> ModifiedGraphs;
+
+	virtual void EditCompositeTunnelNode(class UK2Node_Tunnel* TunnelNode) override;
+
+	virtual void EditMacroInstance(class UK2Node_MacroInstance* MacroInstance, UBlueprint* Blueprint) override;
+
+	virtual void EditCallSite(class UK2Node_CallFunction* CallSite, UBlueprint* Blueprint) override;
+
+	virtual void EditDelegates(class UK2Node_BaseMCDelegate* CallSite, UBlueprint* Blueprint) override;
+
+	virtual void EditCreateDelegates(class UK2Node_CreateDelegate* CallSite) override;
 
-	FFunctionFromNodeHelper(UObject* Obj);
 };
 
 DECLARE_CYCLE_STAT_EXTERN(TEXT("Notify Blueprint Changed"), EKismetCompilerStats_NotifyBlueprintChanged, STATGROUP_KismetCompiler, );
@@ -1095,6 +1136,14 @@ public:
 	/** Checks if a graph (or any sub-graphs or referenced graphs) have latent function nodes */
 	static bool CheckIfGraphHasLatentFunctions(UEdGraph* InGraph);
 
+	/**
+	 * Creates a function result node or returns the current one if one exists
+	 *
+	 * @param InFunctionEntryNode		The function entry node to spawn the result node for
+	 * @return							Spawned result node
+	 */
+	static class UK2Node_FunctionResult* FindOrCreateFunctionResultNode(class UK2Node_EditablePinBase* InFunctionEntryNode);
+
 protected:
 	// Removes all NULL graph references from the SubGraphs array and recurses thru the non-NULL ones
 	static void CleanNullGraphReferencesRecursive(UEdGraph* Graph);
diff --git a/Engine/Source/Runtime/Engine/Classes/EdGraph/EdGraphNode.h b/Engine/Source/Runtime/Engine/Classes/EdGraph/EdGraphNode.h
--- a/Engine/Source/Runtime/Engine/Classes/EdGraph/EdGraphNode.h
+++ b/Engine/Source/Runtime/Engine/Classes/EdGraph/EdGraphNode.h
@@ -223,7 +223,7 @@ public:
 	}
 
 	/** Generate a unique pin name, trying to stick close to a passed in name */
-	FString CreateUniquePinName(FString SourcePinName) const
+	virtual FString CreateUniquePinName(FString SourcePinName) const
 	{
 		FString PinName(SourcePinName);
 		
diff --git a/Engine/Source/Runtime/Engine/Classes/EdGraph/EdGraphSchema.h b/Engine/Source/Runtime/Engine/Classes/EdGraph/EdGraphSchema.h
--- a/Engine/Source/Runtime/Engine/Classes/EdGraph/EdGraphSchema.h
+++ b/Engine/Source/Runtime/Engine/Classes/EdGraph/EdGraphSchema.h
@@ -739,4 +739,26 @@ class ENGINE_API UEdGraphSchema : public UObject
 
 	/* When a node is removed, this method determines whether we should remove it immediately or use the old (slower) code path that results in all node being recreated: */
 	virtual bool ShouldAlwaysPurgeOnModification() const { return true; }
+
+	/*
+	 * Some schemas have nodes that support the user dynamically adding pins when dropping a connection on the node
+	 *
+	 * @param InTargetNode					Node to check for pin adding support
+	 * @param InSourcePinName				Name of the pin being dropped, a new pin of similar name will be constructed
+	 * @param InSourcePinType				Type of pin to drop onto the node
+	 * @param InSourcePinDirection			Direction of the source pin
+	 * @return								Returns the new pin if created
+	 */
+	virtual UEdGraphPin* DropPinOnNode(UEdGraphNode* InTargetNode, const FString& InSourcePinName, const FEdGraphPinType& InSourcePinType, EEdGraphPinDirection InSourcePinDirection) const { return nullptr; }
+
+	/**
+	 * Checks if the node supports dropping a pin on it
+	 *
+	 * @param InTargetNode					Node to check for pin adding support
+	 * @param InSourcePinType				Type of pin to drop onto the node
+	 * @param InSourcePinDirection			Direction of the source pin
+	 * @param OutErrorMessage				Only filled with an error if there is pin add support but there is an error with the pin type
+	 * @return								Returns TRUE if there is support for dropping the pin on the node
+	 */
+	virtual bool SupportsDropPinOnNode(UEdGraphNode* InTargetNode, const FEdGraphPinType& InSourcePinType, EEdGraphPinDirection InSourcePinDirection, FText& OutErrorMessage) const { return false; }
 };
