Renaming macro pins will not rename any macros within the graph that have pins with the same name.

Clicking on a macro instance within a macro graph will no longer show details for any macro.

diff --git a/Engine/Source/Editor/Kismet/Private/BlueprintDetailsCustomization.cpp b/Engine/Source/Editor/Kismet/Private/BlueprintDetailsCustomization.cpp
--- a/Engine/Source/Editor/Kismet/Private/BlueprintDetailsCustomization.cpp
+++ b/Engine/Source/Editor/Kismet/Private/BlueprintDetailsCustomization.cpp
@@ -3190,21 +3190,14 @@ bool FBaseBlueprintGraphActionDetails::OnPinRenamed(UK2Node_EditablePinBase* Tar
 	{
 		FPinRenamedHelper PinRenamedHelper;
 
-		if( Graph )
+		if (FunctionEntryNodePtr.IsValid())
 		{
-			for( TArray<UEdGraphNode*>::TConstIterator NodeIt(Graph->Nodes); NodeIt; ++NodeIt )
-			{
-				UK2Node_EditablePinBase* PinNode = Cast<UK2Node_EditablePinBase>(*NodeIt);
-				if( PinNode )
-				{
-					PinRenamedHelper.NodesToRename.Add(PinNode);
-				}
-			}
-			check( PinRenamedHelper.NodesToRename.Find(TargetNode) );
+			PinRenamedHelper.NodesToRename.Add(FunctionEntryNodePtr.Get());
 		}
-		else
+
+		if (FunctionResultNodePtr.IsValid())
 		{
-			PinRenamedHelper.NodesToRename.Add(TargetNode);
+			PinRenamedHelper.NodesToRename.Add(FunctionEntryNodePtr.Get());
 		}
 
 		PinRenamedHelper.ModifiedBlueprints.Add(GetBlueprintObj());
@@ -3268,7 +3261,7 @@ UEdGraph* FBaseBlueprintGraphActionDetails::GetGraph() const
 		{
 			return Cast<UK2Node_Composite>(Object)->BoundGraph;
 		}
-		else if (Object->IsA<UK2Node_Tunnel>() || Object->IsA<UK2Node_FunctionTerminator>())
+		else if (!Object->IsA<UK2Node_MacroInstance>() && (Object->IsA<UK2Node_Tunnel>() || Object->IsA<UK2Node_FunctionTerminator>()))
 		{
 			return Cast<UK2Node>(Object)->GetGraph();
 		}
