Github integration request 991, Modifies Ctrl/Alt detection when dropping a variable into a blueprint so modifiers are detected at any time during the operation.
#UE-12948 - GitHub 991 : Detect Ctrl/Alt when dropping a Variable into a blueprint
#991
#Branch UE4
#Proj Kismet

diff --git a/Engine/Source/Editor/Kismet/Private/BPVariableDragDropAction.cpp b/Engine/Source/Editor/Kismet/Private/BPVariableDragDropAction.cpp
--- a/Engine/Source/Editor/Kismet/Private/BPVariableDragDropAction.cpp
+++ b/Engine/Source/Editor/Kismet/Private/BPVariableDragDropAction.cpp
@@ -468,15 +468,23 @@ FReply FKismetVariableDragDropAction::DroppedOnPanel( const TSharedRef< SWidget
 			// call analytics
 			AnalyticCallback.ExecuteIfBound();
 
-			// Asking for a getter
-			if (bControlDrag || !CanExecuteMakeSetter(NewNodeParams, VariableProperty))
+			// Take into account current state of modifier keys in case the user changed his mind
+			auto ModifierKeys = FSlateApplication::Get().GetModifierKeys();
+			const bool bModifiedKeysActive = ModifierKeys.IsControlDown() || ModifierKeys.IsAltDown();
+			const bool bAutoCreateGetter = bModifiedKeysActive ? ModifierKeys.IsControlDown() : bControlDrag;
+			const bool bAutoCreateSetter = bModifiedKeysActive ? ModifierKeys.IsAltDown() : bAltDrag;
+			// Handle Getter/Setters
+			if (bAutoCreateGetter || bAutoCreateSetter)
 			{
-				MakeGetter(NewNodeParams);
-			}
-			// Asking for a setter
-			else if (bAltDrag && CanExecuteMakeSetter(NewNodeParams, VariableProperty))
-			{
-				MakeSetter(NewNodeParams);
+				if (bAutoCreateGetter || !CanExecuteMakeSetter(NewNodeParams, VariableProperty))
+				{
+					MakeGetter(NewNodeParams);
+					NewNodeParams.GraphPosition.Y += 50.f;
+				}
+				if (bAutoCreateSetter && CanExecuteMakeSetter( NewNodeParams, VariableProperty))
+				{
+					MakeSetter(NewNodeParams);
+				}
 			}
 			// Show selection menu
 			else
