Adding pins to macros and functions will correctly refresh the details panel.

diff --git a/Engine/Source/Editor/Kismet/Private/BlueprintDetailsCustomization.cpp b/Engine/Source/Editor/Kismet/Private/BlueprintDetailsCustomization.cpp
--- a/Engine/Source/Editor/Kismet/Private/BlueprintDetailsCustomization.cpp
+++ b/Engine/Source/Editor/Kismet/Private/BlueprintDetailsCustomization.cpp
@@ -1909,16 +1909,27 @@ void FBlueprintGraphArgumentGroupLayout::GenerateChildContent( IDetailChildrenBu
 	{
 		TArray<TSharedPtr<FUserPinInfo>> Pins = TargetNode->UserDefinedPins;
 
-		bool bIsInputNode = TargetNode == GraphActionDetailsPtr.Pin()->GetFunctionEntryNode().Get();
-		for (int32 i = 0; i < Pins.Num(); ++i)
-		{
-			TSharedRef<class FBlueprintGraphArgumentLayout> BlueprintArgumentLayout = MakeShareable(new FBlueprintGraphArgumentLayout(
-				TWeakPtr<FUserPinInfo>(Pins[i]),
-				TargetNode.Get(),
-				GraphActionDetailsPtr,
-				FName(*FString::Printf(bIsInputNode ? TEXT("InputArgument%i") : TEXT("OutputArgument%i"), i)),
-				bIsInputNode));
-			ChildrenBuilder.AddChildCustomBuilder(BlueprintArgumentLayout);
+		if(Pins.Num() > 0)
+		{
+			bool bIsInputNode = TargetNode == GraphActionDetailsPtr.Pin()->GetFunctionEntryNode().Get();
+			for (int32 i = 0; i < Pins.Num(); ++i)
+			{
+				TSharedRef<class FBlueprintGraphArgumentLayout> BlueprintArgumentLayout = MakeShareable(new FBlueprintGraphArgumentLayout(
+					TWeakPtr<FUserPinInfo>(Pins[i]),
+					TargetNode.Get(),
+					GraphActionDetailsPtr,
+					FName(*FString::Printf(bIsInputNode ? TEXT("InputArgument%i") : TEXT("OutputArgument%i"), i)),
+					bIsInputNode));
+				ChildrenBuilder.AddChildCustomBuilder(BlueprintArgumentLayout);
+			}
+		}
+		else
+		{
+			// Add a null widget for this section, keeps it around for callbacks to refresh
+			ChildrenBuilder.AddChildContent(FText::GetEmpty()).ValueContent()
+				[
+					SNullWidget::NullWidget
+				];
 		}
 	}
 }
