[UE-12201] Ensuring that pasted components are marked as public

diff --git a/Engine/Source/Editor/Kismet/Private/SSCSEditor.cpp b/Engine/Source/Editor/Kismet/Private/SSCSEditor.cpp
--- a/Engine/Source/Editor/Kismet/Private/SSCSEditor.cpp
+++ b/Engine/Source/Editor/Kismet/Private/SSCSEditor.cpp
@@ -4782,6 +4782,8 @@ void SSCSEditor::PasteNodes()
 			// Get the component object instance
 			UActorComponent* NewActorComponent = NewObjectIt->Value;
 			check(NewActorComponent);
+			// make sure that our flags match those set in USimpleConstructionScript::CreateNode:
+			NewActorComponent->SetFlags(RF_ArchetypeObject | RF_Transactional | RF_Public);
 
 			// Create a new SCS node to contain the new component and add it to the tree
 			NewActorComponent = AddNewNode(Blueprint->SimpleConstructionScript->CreateNodeAndRenameComponent(NewActorComponent), NULL, false, false);
