[UE-13328] Blocking users from modifying Blueprint properties while PIE'ing.

diff --git a/Engine/Source/Editor/Kismet/Private/SKismetInspector.cpp b/Engine/Source/Editor/Kismet/Private/SKismetInspector.cpp
--- a/Engine/Source/Editor/Kismet/Private/SKismetInspector.cpp
+++ b/Engine/Source/Editor/Kismet/Private/SKismetInspector.cpp
@@ -712,6 +712,12 @@ EVisibility SKismetInspector::GetPropertyViewVisibility() const
 bool SKismetInspector::IsPropertyEditingEnabled() const
 {
 	bool bIsEditable = true;
+
+	if (BlueprintEditorPtr.IsValid())
+	{
+		bIsEditable = BlueprintEditorPtr.Pin()->InEditingMode();
+	}
+
 	for (const TWeakObjectPtr<UObject>& SelectedObject : SelectedObjects)
 	{
 		UActorComponent* Component = Cast<UActorComponent>(SelectedObject.Get());
