[UE-12166] Cleaning up duplicate call-on-member node options and improving the call-on-member workflow.

diff --git a/Engine/Source/Editor/BlueprintGraph/Private/BlueprintEditorSettings.cpp b/Engine/Source/Editor/BlueprintGraph/Private/BlueprintEditorSettings.cpp
--- a/Engine/Source/Editor/BlueprintGraph/Private/BlueprintEditorSettings.cpp
+++ b/Engine/Source/Editor/BlueprintGraph/Private/BlueprintEditorSettings.cpp
@@ -14,6 +14,7 @@ UBlueprintEditorSettings::UBlueprintEditorSettings(const FObjectInitializer& Obj
 	, bSplitContextTargetSettings(true)
 	, bExposeAllMemberComponentFunctions(true)
 	, bShowContextualFavorites(false)
+	, bCompactCallOnMemberNodes(false)
 	, bFlattenFavoritesMenus(true)
 	, bFavorPureCastNodes(false)
 	, bAutoCastObjectConnections(false)
diff --git a/Engine/Source/Editor/BlueprintGraph/Private/BlueprintFunctionNodeSpawner.cpp b/Engine/Source/Editor/BlueprintGraph/Private/BlueprintFunctionNodeSpawner.cpp
--- a/Engine/Source/Editor/BlueprintGraph/Private/BlueprintFunctionNodeSpawner.cpp
+++ b/Engine/Source/Editor/BlueprintGraph/Private/BlueprintFunctionNodeSpawner.cpp
@@ -17,6 +17,7 @@
 #include "K2Node_CallDataTableFunction.h"
 #include "K2Node_CallArrayFunction.h"
 #include "K2Node_VariableGet.h"
+#include "BlueprintEditorSettings.h"
 
 #define LOCTEXT_NAMESPACE "BlueprintFunctionNodeSpawner"
 
@@ -391,10 +392,16 @@ UEdGraphNode* UBlueprintFunctionNodeSpawner::Invoke(UEdGraph* ParentGraph, FBind
 	FCustomizeNodeDelegate PostSpawnSetupDelegate = FCustomizeNodeDelegate::CreateStatic(PostSpawnSetupLambda, GetFunction(), SetNodeFieldDelegate, CustomizeNodeDelegate);
 
 	UClass* SpawnClass = NodeClass;
-	if ((Bindings.Num() == 1) && (*Bindings.CreateConstIterator())->IsA<UObjectProperty>())
+
+	const UBlueprintEditorSettings* BPSettings = GetDefault<UBlueprintEditorSettings>();
+	bool const bIsTemplateNode = FBlueprintNodeTemplateCache::IsTemplateOuter(ParentGraph);
+
+	bool const bSpawnCallOnMember = (Bindings.Num() == 1) && (*Bindings.CreateConstIterator())->IsA<UObjectProperty>();
+	if (bSpawnCallOnMember && (bIsTemplateNode || BPSettings->bCompactCallOnMemberNodes))
 	{
 		SpawnClass = UK2Node_CallFunctionOnMember::StaticClass();
 	}
+
 	// if this spawner was set up to spawn a bound node, reset this so the 
 	// bound nodes get positioned properly
 	BlueprintFunctionNodeSpawnerImpl::BindingOffset = FVector2D::ZeroVector;
diff --git a/Engine/Source/Editor/BlueprintGraph/Private/K2Node_CallFunctionOnMember.cpp b/Engine/Source/Editor/BlueprintGraph/Private/K2Node_CallFunctionOnMember.cpp
--- a/Engine/Source/Editor/BlueprintGraph/Private/K2Node_CallFunctionOnMember.cpp
+++ b/Engine/Source/Editor/BlueprintGraph/Private/K2Node_CallFunctionOnMember.cpp
@@ -35,7 +35,7 @@ UEdGraphPin* UK2Node_CallFunctionOnMember::CreateSelfPin(const UFunction* Functi
 FString UK2Node_CallFunctionOnMember::GetFunctionContextString() const
 {
 	UClass* MemberVarClass = MemberVariableToCallOn.GetMemberParentClass(GetBlueprintClassFromNode());
-	FString CallFunctionClassName = (MemberVarClass != NULL) ? MemberVarClass->GetName() : TEXT("NULL");
+	FString CallFunctionClassName = (MemberVarClass != NULL) ? MemberVarClass->GetDisplayNameText().ToString() : TEXT("NULL");
 	return FString(TEXT("\n")) + FString::Printf(TEXT("%s %s (%s)"), *LOCTEXT("CallFunctionOnMemberDifferentContext", "Target is").ToString(), *CallFunctionClassName, *MemberVariableToCallOn.GetMemberName().ToString());
 }
 
diff --git a/Engine/Source/Editor/BlueprintGraph/Public/BlueprintEditorSettings.h b/Engine/Source/Editor/BlueprintGraph/Public/BlueprintEditorSettings.h
--- a/Engine/Source/Editor/BlueprintGraph/Public/BlueprintEditorSettings.h
+++ b/Engine/Source/Editor/BlueprintGraph/Public/BlueprintEditorSettings.h
@@ -43,6 +43,10 @@ public:
 	UPROPERTY(EditAnywhere, config, Category=Workflow, meta=(DisplayName="Context Menu: Show Favorites Section"))
 	bool bShowContextualFavorites;
 
+	/** If enabled, then call-on-member actions will be spawned as a single node (instead of a GetMember + FunctionCall node). */
+	UPROPERTY(EditAnywhere, config, Category=Workflow)
+	bool bCompactCallOnMemberNodes;
+
 	/** If enabled, then your Blueprint favorites will be uncategorized, leaving you with less nested categories to sort through. */
 	UPROPERTY(EditAnywhere, config, Category=Workflow)
 	bool bFlattenFavoritesMenus;
diff --git a/Engine/Source/Editor/Kismet/Private/BlueprintActionMenuItem.cpp b/Engine/Source/Editor/Kismet/Private/BlueprintActionMenuItem.cpp
--- a/Engine/Source/Editor/Kismet/Private/BlueprintActionMenuItem.cpp
+++ b/Engine/Source/Editor/Kismet/Private/BlueprintActionMenuItem.cpp
@@ -34,6 +34,25 @@ namespace FBlueprintMenuActionItemImpl
 	 * @return The spawned node (could be an existing one if the event was already placed).
 	 */
 	static UEdGraphNode* InvokeAction(const UBlueprintNodeSpawner* Action, UEdGraph* ParentGraph, FVector2D const Location, IBlueprintNodeBinder::FBindingSet const& Bindings, bool& bOutNewNode);
+
+	/**
+	 * 
+	 * 
+	 * @param  LeadingPin    
+	 * @param  Node    
+	 * @return 
+	 */
+	static bool IsNodeLinked(UEdGraphPin* LeadingPin, UEdGraphNode& Node);
+
+	/**
+	 * 
+	 * 
+	 * @param  ParentGraph    
+	 * @param  FromPin    
+	 * @param  SpawnedNodesBeginIndex    
+	 * @return 
+	 */
+	static UEdGraphNode* AutowireSpawnedNodes(UEdGraphPin* FromPin, const TArray<UEdGraphNode*>& GraphNodes, int32 const NodesBeginIndex);
 }
 
 //------------------------------------------------------------------------------
@@ -86,6 +105,94 @@ static UEdGraphNode* FBlueprintMenuActionItemImpl::InvokeAction(const UBlueprint
 	return SpawnedNode;
 }
 
+//------------------------------------------------------------------------------
+static bool FBlueprintMenuActionItemImpl::IsNodeLinked(UEdGraphPin* LeadingPin, UEdGraphNode& Node)
+{
+	const EEdGraphPinDirection PinDirection = LeadingPin->Direction;
+	
+	for (UEdGraphPin* Link : LeadingPin->LinkedTo)
+	{
+		UEdGraphNode* LinkNode = Link->GetOwningNode();
+		if (LinkNode == &Node)
+		{
+			return true;
+		}
+
+		for (UEdGraphPin* NodePin : LinkNode->Pins)
+		{
+			if (NodePin->Direction != PinDirection)
+			{
+				continue;
+			}
+
+			if (IsNodeLinked(NodePin, Node))
+			{
+				return true;
+			}
+		}
+	}
+	return false;
+}
+
+//------------------------------------------------------------------------------
+static UEdGraphNode* FBlueprintMenuActionItemImpl::AutowireSpawnedNodes(UEdGraphPin* FromPin, const TArray<UEdGraphNode*>& GraphNodes, int32 const NodesBeginIndex)
+{
+	int32 const CurrentGraphNodeCount = GraphNodes.Num();
+	int32 const NewNodeCount = CurrentGraphNodeCount - NodesBeginIndex;
+
+	TArray<UEdGraphNode*> OrderedNewNodes;
+	OrderedNewNodes.Reserve(NewNodeCount);
+	// @TODO: there's gotta be a better way to blit these in
+	for (int32 NodexIndex = NodesBeginIndex; NodexIndex < CurrentGraphNodeCount; ++NodexIndex)
+	{
+		OrderedNewNodes.Add(GraphNodes[NodexIndex]);
+	}
+
+	const EEdGraphPinDirection PinDirection = FromPin->Direction;
+	// should lhs come before rhs?
+	OrderedNewNodes.Sort([PinDirection](UEdGraphNode& Lhs, UEdGraphNode& Rhs)->bool
+		{
+			for (UEdGraphPin* NodePin : Rhs.Pins)
+			{
+				if (NodePin->Direction != PinDirection)
+				{
+					continue;
+				}
+
+				if (IsNodeLinked(NodePin, Lhs))
+				{
+					return false;
+				}
+			}
+			return true;
+		});
+		
+
+	int32 const PreAutowireConnectionCount = FromPin->LinkedTo.Num();
+
+	UEdGraphPin* OldPinLink = nullptr;
+	if (PreAutowireConnectionCount > 0)
+	{
+		OldPinLink = FromPin->LinkedTo[0];
+	}
+
+	for (UEdGraphNode* NewNode : OrderedNewNodes)
+	{
+		NewNode->AutowireNewNode(FromPin);
+
+		int32 const NewConnectionCount = FromPin->LinkedTo.Num();
+		if (NewConnectionCount == 0)
+		{
+			continue;
+		}
+		else if ((NewConnectionCount != PreAutowireConnectionCount) || (FromPin->LinkedTo[0] != OldPinLink))
+		{
+			return NewNode;
+		}
+	}
+	return nullptr;
+}
+
 /*******************************************************************************
  * FBlueprintMenuActionItem
  ******************************************************************************/
@@ -149,6 +256,8 @@ UEdGraphNode* FBlueprintActionMenuItem::PerformAction(UEdGraph* ParentGraph, UEd
 			}
 		}
 
+		int32 const PreInvokeNodeCount = ParentGraph->Nodes.Num();
+
 		bool bNewNode = false;
 		LastSpawnedNode = InvokeAction(Action, ParentGraph, ModifiedLocation, BindingsSubset, /*out*/ bNewNode);
 		// could already be an existent node, so we have to add here (can't 
@@ -161,7 +270,7 @@ UEdGraphNode* FBlueprintActionMenuItem::PerformAction(UEdGraph* ParentGraph, UEd
 		{
 			// make sure to auto-wire after we position the new node (in case
 			// the auto-wire creates a conversion node to put between them)
-			LastSpawnedNode->AutowireNewNode(FromPin);
+			FBlueprintMenuActionItemImpl::AutowireSpawnedNodes(FromPin, ParentGraph->Nodes, PreInvokeNodeCount);
 		}
 
 		if (bNewNode)
diff --git a/Engine/Source/Editor/Kismet/Private/BlueprintActionMenuUtils.cpp b/Engine/Source/Editor/Kismet/Private/BlueprintActionMenuUtils.cpp
--- a/Engine/Source/Editor/Kismet/Private/BlueprintActionMenuUtils.cpp
+++ b/Engine/Source/Editor/Kismet/Private/BlueprintActionMenuUtils.cpp
@@ -292,14 +292,14 @@ static FBlueprintActionFilter BlueprintActionMenuUtilsImpl::MakeCallOnMemberFilt
 	bool bForceAddComponents = ((ContextTargetMask & EContextTargetFlags::TARGET_SubComponents) != 0);
 
 	TArray<UClass*> TargetClasses = MainMenuFilter.TargetClasses;
-	if (bForceAddComponents && (CallOnMemberFilter.TargetClasses.Num() == 0))
+	if (bForceAddComponents && (TargetClasses.Num() == 0))
 	{
 		for (UBlueprint const* TargetBlueprint : MainMenuFilter.Context.Blueprints)
 		{
 			UClass* BpClass = TargetBlueprint->SkeletonGeneratedClass;
 			if (BpClass != nullptr)
 			{
-				TargetClasses.Add(BpClass);
+				TargetClasses.AddUnique(BpClass);
 			}
 		}
 	}
@@ -308,14 +308,14 @@ static FBlueprintActionFilter BlueprintActionMenuUtilsImpl::MakeCallOnMemberFilt
 	{
 		for (TFieldIterator<UObjectProperty> PropertyIt(TargetClass, EFieldIteratorFlags::IncludeSuper); PropertyIt; ++PropertyIt)
 		{
-			UObjectProperty* ObjectProperty = *PropertyIt;
+ 			UObjectProperty* ObjectProperty = *PropertyIt;
 			if (!ObjectProperty->HasAnyPropertyFlags(CPF_BlueprintVisible))
 			{
 				continue;
 			}
 
 			if ( ObjectProperty->HasMetaData(FBlueprintMetadata::MD_ExposeFunctionCategories) || 
-				(bForceAddComponents && ObjectProperty->PropertyClass->IsChildOf<UActorComponent>()) )
+				(bForceAddComponents && FBlueprintEditorUtils::IsSCSComponentProperty(ObjectProperty)) )
 			{
 				CallOnMemberFilter.Context.SelectedObjects.Add(ObjectProperty);
 			}			
@@ -472,7 +472,7 @@ void FBlueprintActionMenuUtils::MakeContextMenu(FBlueprintActionContext const& C
 			bCanOperateOnLevelActors &= BlueprintClass->IsChildOf<ALevelScriptActor>();
 			if (bIsContextSensitive && (ClassTargetMask & EContextTargetFlags::TARGET_Blueprint))
 			{
-				MainMenuFilter.TargetClasses.Add(BlueprintClass);
+				MainMenuFilter.TargetClasses.AddUnique(BlueprintClass);
 			}
 		}
 		bCanHaveActorComponents &= FBlueprintEditorUtils::DoesSupportComponents(Blueprint);
@@ -559,7 +559,7 @@ void FBlueprintActionMenuUtils::MakeContextMenu(FBlueprintActionContext const& C
 			{
 				if (ClassTargetMask & EContextTargetFlags::TARGET_PinObject)
 				{
-					MainMenuFilter.TargetClasses.Add(PinObjClass);
+					MainMenuFilter.TargetClasses.AddUnique(PinObjClass);
 				}
 			}
 
@@ -571,7 +571,7 @@ void FBlueprintActionMenuUtils::MakeContextMenu(FBlueprintActionContext const& C
 				{
 					if (UClass* TargetClass = GetPinClassType(TargetPin))
 					{
-						MainMenuFilter.TargetClasses.Add(TargetClass);
+						MainMenuFilter.TargetClasses.AddUnique(TargetClass);
 					}
 				}
 			}
@@ -587,7 +587,7 @@ void FBlueprintActionMenuUtils::MakeContextMenu(FBlueprintActionContext const& C
 				{
 					if (UClass* PinClass = GetPinClassType(NodePin))
 					{
-						MainMenuFilter.TargetClasses.Add(PinClass);
+						MainMenuFilter.TargetClasses.AddUnique(PinClass);
 					}
 				}
 			}
diff --git a/Engine/Source/Editor/Kismet/Private/SBlueprintContextTargetMenu.cpp b/Engine/Source/Editor/Kismet/Private/SBlueprintContextTargetMenu.cpp
--- a/Engine/Source/Editor/Kismet/Private/SBlueprintContextTargetMenu.cpp
+++ b/Engine/Source/Editor/Kismet/Private/SBlueprintContextTargetMenu.cpp
@@ -6,6 +6,7 @@
 #include "BlueprintActionFilter.h"		// for FBlueprintActionContext
 #include "BlueprintEditorUtils.h"		// for DoesSupportComponents
 #include "BlueprintEditorSettings.h"
+#include "Components/ActorComponent.h"
 
 #define LOCTEXT_NAMESPACE "SBlueprintContextTargetMenu"
 
@@ -26,7 +27,10 @@ namespace ContextMenuTargetProfileImpl
 	 * @return True if the context contains a pin matching the type in question, otherwise false.
 	 */
 	static bool HasExecPinContext(const FBlueprintActionContext& MenuContext);
+	static bool HasComponentPinContext(const FBlueprintActionContext& MenuContext);
 	static bool HasObjectPinContext(const FBlueprintActionContext& MenuContext);
+	static bool HasClassPinContext(const FBlueprintActionContext& MenuContext);
+	static bool HasInterfacePinContext(const FBlueprintActionContext& MenuContext);
 
 	static bool IsObjectPin(UEdGraphPin* Pin);
 	static UClass* GetPinClass(UEdGraphPin const* Pin);
@@ -50,6 +54,35 @@ namespace ContextMenuTargetProfileImpl
 	static FString GetProfileSaveName(const FBlueprintActionContext& MenuContext);
 }
 
+//------------------------------------------------------------------------------
+static bool ContextMenuTargetProfileImpl::HasExecPinContext(const FBlueprintActionContext& MenuContext)
+{
+	const UEdGraphSchema_K2* K2Schema = GetDefault<UEdGraphSchema_K2>();
+	for (UEdGraphPin* ContextPin : MenuContext.Pins)
+	{
+		if (K2Schema->IsExecPin(*ContextPin))
+		{
+			return true;
+		}
+	}
+	return false;
+}
+
+//------------------------------------------------------------------------------
+static bool ContextMenuTargetProfileImpl::HasComponentPinContext(const FBlueprintActionContext& MenuContext)
+{
+	const UEdGraphSchema_K2* K2Schema = GetDefault<UEdGraphSchema_K2>();
+	for (UEdGraphPin* ContextPin : MenuContext.Pins)
+	{
+		UClass* PinClass = GetPinClass(ContextPin);
+		if ((PinClass != nullptr) && PinClass->IsChildOf<UActorComponent>())
+		{
+			return true;
+		}
+	}
+	return false;
+}
+
 //------------------------------------------------------------------------------
 static bool ContextMenuTargetProfileImpl::HasObjectPinContext(const FBlueprintActionContext& MenuContext)
 {
@@ -64,12 +97,24 @@ static bool ContextMenuTargetProfileImpl::HasObjectPinContext(const FBlueprintAc
 }
 
 //------------------------------------------------------------------------------
-static bool ContextMenuTargetProfileImpl::HasExecPinContext(const FBlueprintActionContext& MenuContext)
+static bool ContextMenuTargetProfileImpl::HasClassPinContext(const FBlueprintActionContext& MenuContext)
 {
-	const UEdGraphSchema_K2* K2Schema = GetDefault<UEdGraphSchema_K2>();
 	for (UEdGraphPin* ContextPin : MenuContext.Pins)
 	{
-		if (K2Schema->IsExecPin(*ContextPin))
+		if (ContextPin->PinType.PinCategory == UEdGraphSchema_K2::PC_Class)
+		{
+			return true;
+		}
+	}
+	return false;
+}
+
+//------------------------------------------------------------------------------
+static bool ContextMenuTargetProfileImpl::HasInterfacePinContext(const FBlueprintActionContext& MenuContext)
+{
+	for (UEdGraphPin* ContextPin : MenuContext.Pins)
+	{
+		if (ContextPin->PinType.PinCategory == UEdGraphSchema_K2::PC_Interface)
 		{
 			return true;
 		}
@@ -139,11 +184,23 @@ FText ContextMenuTargetProfileImpl::GetProfileDescription(const FBlueprintAction
 	{
 		return LOCTEXT("BlueprintContextTarget", "Blueprint Context Target(s)");
 	}
-	else if (ContextMenuTargetProfileImpl::HasObjectPinContext(MenuContext))
+	else if (HasComponentPinContext(MenuContext))
+	{
+		return LOCTEXT("ComponentContextTarget", "Component Context Target(s)");
+	}
+// 	else if (HasClassPinContext(MenuContext))
+// 	{
+// 		return LOCTEXT("ClassContextTarget", "Class Context Target(s)");
+// 	}
+//	else if (HasInterfacePinContext(MenuContext))
+// 	{
+// 		return LOCTEXT("InterfaceContextTarget", "Interface Context Target(s)");
+// 	}
+	else if (HasObjectPinContext(MenuContext))
 	{
 		return LOCTEXT("ObjectContextTarget", "Object Context Target(s)");
 	}
-	else if (ContextMenuTargetProfileImpl::HasExecPinContext(MenuContext))
+	else if (HasExecPinContext(MenuContext))
 	{
 		return LOCTEXT("ExecContextTarget", "Exec Context Target(s)");
 	}
@@ -227,7 +284,7 @@ uint32 ContextMenuTargetProfileImpl::GetIncompatibleTargetFlags(const FBlueprint
 		}
 	}
 
-	if (HasComponentsMaskOut != 0)
+	if ((HasComponentsMaskOut != 0) && !HasComponentPinContext(MenuContext))
 	{
 		IncompatibleFlags &= ~EContextTargetFlags::TARGET_SubComponents;
 	}
@@ -253,6 +310,18 @@ static FString ContextMenuTargetProfileImpl::GetProfileSaveName(const FBlueprint
 	{
 		return TEXT("ExecPinProfile");
 	}
+	else if (HasComponentPinContext(MenuContext))
+	{
+		return TEXT("ComponentPinProfile");
+	}
+// 	else if (HasClassPinContext(MenuContext))
+// 	{
+// 		return TEXT("ClassPinProfile");
+// 	}
+// 	else if (HasInterfacePinContext(MenuContext))
+// 	{
+// 		return TEXT("InterfacePinProfile");
+// 	}
 
 	uint32 HasComponentMask = 0;
 	uint32 IncompatibleTargetFlags = GetIncompatibleTargetFlags(MenuContext, HasComponentMask);
diff --git a/Engine/Source/Editor/UnrealEd/Private/Kismet2/BlueprintEditorUtils.cpp b/Engine/Source/Editor/UnrealEd/Private/Kismet2/BlueprintEditorUtils.cpp
--- a/Engine/Source/Editor/UnrealEd/Private/Kismet2/BlueprintEditorUtils.cpp
+++ b/Engine/Source/Editor/UnrealEd/Private/Kismet2/BlueprintEditorUtils.cpp
@@ -5748,6 +5748,76 @@ void FBlueprintEditorUtils::UpdateRootComponentReference(UBlueprint* Blueprint)
 	}
 }
 
+bool FBlueprintEditorUtils::IsSCSComponentProperty(UObjectProperty* MemberProperty)
+{
+	if (!MemberProperty->PropertyClass->IsChildOf<UActorComponent>())
+	{
+		return false;
+	}
+
+
+	UClass* OwnerClass = MemberProperty->GetOwnerClass();
+	UBlueprintGeneratedClass* BpClassOwner = Cast<UBlueprintGeneratedClass>(OwnerClass);
+
+	if (BpClassOwner == nullptr)
+	{
+		// if this isn't directly a blueprint property, then we check if it is a 
+		// associated with a natively added component (which would still be  
+		// accessible through the SCS tree)
+
+		if (OwnerClass == nullptr)
+		{
+			return false;
+		}
+		else if (const AActor* ActorCDO = GetDefault<AActor>(OwnerClass))
+		{
+			TInlineComponentArray<UActorComponent*> CDOComponents;
+			ActorCDO->GetComponents(CDOComponents);
+
+			const void* PropertyAddress = MemberProperty->ContainerPtrToValuePtr<void>(ActorCDO);
+			UObject* PropertyValue = MemberProperty->GetObjectPropertyValue(PropertyAddress);
+
+			for (UActorComponent* Component : CDOComponents)
+			{
+				if (!Component->GetClass()->IsChildOf(MemberProperty->PropertyClass))
+				{
+					continue;
+				}
+
+				if (PropertyValue == Component)
+				{
+					return true;
+				}
+			}
+		}
+		return false;
+	}
+
+	FMemberReference MemberRef;
+	MemberRef.SetFromField<UProperty>(MemberProperty, /*bIsConsideredSelfContext =*/false);
+	bool const bIsGuidValid = MemberRef.GetMemberGuid().IsValid();
+
+	if (BpClassOwner->SimpleConstructionScript != nullptr)
+	{
+		TArray<USCS_Node*> SCSNodes = BpClassOwner->SimpleConstructionScript->GetAllNodes();
+		for (USCS_Node* ScsNode : SCSNodes)
+		{
+			if (bIsGuidValid && ScsNode->VariableGuid.IsValid())
+			{
+				if (ScsNode->VariableGuid == MemberRef.GetMemberGuid())
+				{
+					return true;
+				}
+			}
+			else if (ScsNode->VariableName == MemberRef.GetMemberName())
+			{
+				return true;
+			}
+		}
+	}
+	return false;
+}
+
 /** Temporary fix for cut-n-paste error that failed to carry transactional flags */
 void FBlueprintEditorUtils::UpdateTransactionalFlags(UBlueprint* Blueprint)
 {
diff --git a/Engine/Source/Editor/UnrealEd/Public/Kismet2/BlueprintEditorUtils.h b/Engine/Source/Editor/UnrealEd/Public/Kismet2/BlueprintEditorUtils.h
--- a/Engine/Source/Editor/UnrealEd/Public/Kismet2/BlueprintEditorUtils.h
+++ b/Engine/Source/Editor/UnrealEd/Public/Kismet2/BlueprintEditorUtils.h
@@ -949,6 +949,9 @@ public:
 	/** Ensures that the CDO root component reference is valid for Actor-based Blueprints */
 	static void UpdateRootComponentReference(UBlueprint* Blueprint);
 
+	/** Determines if this property is associated with a component that would be displayed in the SCS editor */
+	static bool IsSCSComponentProperty(UObjectProperty* MemberProperty);
+
 	//////////////////////////////////////////////////////////////////////////
 	// Interface
 
