[UE-9602] Giving the SCSEditorViewport its own delete command (so on mac it you can still delete with the "backspace" key).

diff --git a/Engine/Source/Editor/Kismet/Private/BlueprintEditorCommands.cpp b/Engine/Source/Editor/Kismet/Private/BlueprintEditorCommands.cpp
--- a/Engine/Source/Editor/Kismet/Private/BlueprintEditorCommands.cpp
+++ b/Engine/Source/Editor/Kismet/Private/BlueprintEditorCommands.cpp
@@ -409,4 +409,12 @@ void FBlueprintSpawnNodeCommands::GetGraphActionByChord(FInputChord& InChord, FB
 	}
 }
 
+//////////////////////////////////////////////////////////////////////////
+// FSCSEditorViewportCommands
+
+void FSCSEditorViewportCommands::RegisterCommands()
+{
+	UI_COMMAND(DeleteComponent, "Delete", "Delete current selection", EUserInterfaceActionType::Button, FInputChord(EKeys::Platform_Delete));
+}
+
 #undef LOCTEXT_NAMESPACE
\ No newline at end of file
diff --git a/Engine/Source/Editor/Kismet/Private/BlueprintEditorCommands.h b/Engine/Source/Editor/Kismet/Private/BlueprintEditorCommands.h
--- a/Engine/Source/Editor/Kismet/Private/BlueprintEditorCommands.h
+++ b/Engine/Source/Editor/Kismet/Private/BlueprintEditorCommands.h
@@ -100,3 +100,19 @@ private:
 	/** An array of all the possible commands for spawning nodes */
 	TArray< TSharedPtr< class FNodeSpawnInfo > > NodeCommands;
 };
+
+//////////////////////////////////////////////////////////////////////////
+// FSCSEditorViewportCommands
+
+class FSCSEditorViewportCommands : public TCommands<FSCSEditorViewportCommands>
+{
+public:
+
+	FSCSEditorViewportCommands()
+		: TCommands<FSCSEditorViewportCommands>(TEXT("SCSEditorViewport"), NSLOCTEXT("Contexts", "SCSEditorViewport", "SCS Editor Viewport"), NAME_None, FEditorStyle::GetStyleSetName())
+	{}
+
+	virtual void RegisterCommands() override;
+
+	TSharedPtr< FUICommandInfo > DeleteComponent;
+};
diff --git a/Engine/Source/Editor/Kismet/Private/SSCSEditorViewport.cpp b/Engine/Source/Editor/Kismet/Private/SSCSEditorViewport.cpp
--- a/Engine/Source/Editor/Kismet/Private/SSCSEditorViewport.cpp
+++ b/Engine/Source/Editor/Kismet/Private/SSCSEditorViewport.cpp
@@ -353,31 +353,49 @@ TSharedPtr<SWidget> SSCSEditorViewport::MakeViewportToolbar()
 
 void SSCSEditorViewport::BindCommands()
 {
+	FSCSEditorViewportCommands::Register(); // make sure the viewport specific commands have been registered
+
+	TSharedPtr<FBlueprintEditor> BlueprintEditor = BlueprintEditorPtr.Pin();
+	TSharedPtr<SSCSEditor> SCSEditorWidgetPtr = BlueprintEditor->GetSCSEditor();
+	SSCSEditor* SCSEditorWidget = SCSEditorWidgetPtr.Get();
+	// for mac, we have to bind a command that would override the BP-Editor's 
+	// "NavigateToParentBackspace" command, because the delete key is the 
+	// backspace key for that platform (and "NavigateToParentBackspace" does not 
+	// make sense in the viewport window... it blocks the generic delete command)
+	// 
+	// NOTE: this needs to come before we map any other actions (so it is 
+	// prioritized first)
+	CommandList->MapAction(
+		FSCSEditorViewportCommands::Get().DeleteComponent,
+		FExecuteAction::CreateSP(SCSEditorWidget, &SSCSEditor::OnDeleteNodes),
+		FCanExecuteAction::CreateSP(SCSEditorWidget, &SSCSEditor::CanDeleteNodes)
+	);
+
 	const FBlueprintEditorCommands& Commands = FBlueprintEditorCommands::Get();
 
-	GetCommandList()->Append( BlueprintEditorPtr.Pin()->GetToolkitCommands() );
-	GetCommandList()->Append( BlueprintEditorPtr.Pin()->GetSCSEditor()->CommandList.ToSharedRef() );
+	CommandList->Append(BlueprintEditor->GetToolkitCommands());
+	CommandList->Append(BlueprintEditor->GetSCSEditor()->CommandList.ToSharedRef());
 	SEditorViewport::BindCommands();
 
 	BlueprintEditorPtr.Pin()->GetToolkitCommands()->MapAction(
-		FBlueprintEditorCommands::Get().EnableSimulation,
+		Commands.EnableSimulation,
 		FExecuteAction::CreateSP(this, &SSCSEditorViewport::ToggleIsSimulateEnabled),
 		FCanExecuteAction(),
 		FIsActionChecked::CreateSP(ViewportClient.Get(), &FSCSEditorViewportClient::GetIsSimulateEnabled));
 
 	// Toggle camera lock on/off
 	CommandList->MapAction(
-		FBlueprintEditorCommands::Get().ResetCamera,
+		Commands.ResetCamera,
 		FExecuteAction::CreateSP(ViewportClient.Get(), &FSCSEditorViewportClient::ResetCamera) );
 
 	CommandList->MapAction(
-		FBlueprintEditorCommands::Get().ShowFloor,
+		Commands.ShowFloor,
 		FExecuteAction::CreateSP(ViewportClient.Get(), &FSCSEditorViewportClient::ToggleShowFloor),
 		FCanExecuteAction(),
 		FIsActionChecked::CreateSP(ViewportClient.Get(), &FSCSEditorViewportClient::GetShowFloor));
 
 	CommandList->MapAction(
-		FBlueprintEditorCommands::Get().ShowGrid,
+		Commands.ShowGrid,
 		FExecuteAction::CreateSP(ViewportClient.Get(), &FSCSEditorViewportClient::ToggleShowGrid),
 		FCanExecuteAction(),
 		FIsActionChecked::CreateSP(ViewportClient.Get(), &FSCSEditorViewportClient::GetShowGrid));
