[UE-13263] Accounting for GEditor->GetSelectedObjects() no longer returning selected content browser objects. Now properly using (and loading) selected assets to offer bound AddComponent nodes.

diff --git a/Engine/Source/Editor/Kismet/Private/BlueprintActionMenuUtils.cpp b/Engine/Source/Editor/Kismet/Private/BlueprintActionMenuUtils.cpp
--- a/Engine/Source/Editor/Kismet/Private/BlueprintActionMenuUtils.cpp
+++ b/Engine/Source/Editor/Kismet/Private/BlueprintActionMenuUtils.cpp
@@ -21,6 +21,8 @@
 #include "Engine/Selection.h"
 #include "Engine/LevelScriptActor.h"
 #include "Engine/LevelScriptBlueprint.h"
+#include "Editor/ContentBrowser/Public/ContentBrowserModule.h" // for GetSelectedAssets()
+#include "ComponentAssetBroker.h"	// for GetPrimaryComponentForAsset()
 
 #define LOCTEXT_NAMESPACE "BlueprintActionMenuUtils"
 
@@ -602,13 +604,26 @@ void FBlueprintActionMenuUtils::MakeContextMenu(FBlueprintActionContext const& C
 	AddComponentFilter.PermittedNodeTypes.Add(UK2Node_AddComponent::StaticClass());
 	AddComponentFilter.AddRejectionTest(FBlueprintActionFilter::FRejectionTestDelegate::CreateStatic(IsUnBoundSpawner));
 
-	for (FSelectionIterator SelectionIt(*GEditor->GetSelectedObjects()); SelectionIt; ++SelectionIt)
+
+	FContentBrowserModule& ContentBrowserModule = FModuleManager::LoadModuleChecked<FContentBrowserModule>("ContentBrowser");
+	TArray<FAssetData> SelectedAssets;
+	ContentBrowserModule.Get().GetSelectedAssets(SelectedAssets);
+
+	for (FAssetData& Asset : SelectedAssets)
 	{
-		UObject* PerspectiveAsset = *SelectionIt;
-		if (PerspectiveAsset->IsAsset())
+		UClass* AssetClass = Asset.GetClass();
+		// filter here (rather than in FBlueprintActionFilter) so we only load
+		// assets that we can use
+		if ((AssetClass == nullptr) || (FComponentAssetBrokerage::GetPrimaryComponentForAsset(AssetClass) == nullptr))
 		{
-			AddComponentFilter.Context.SelectedObjects.Add(PerspectiveAsset);
+			continue;
 		}
+
+		// @TODO: loading assets here may be slow (but we need a UObject to 
+		//        properly bind to), consider adding a editor option that will 
+		//        only offer then if the asset is already loaded
+		UObject* AssetObj = Asset.GetAsset();
+		AddComponentFilter.Context.SelectedObjects.Add(AssetObj);
 	}
 
 	//--------------------------------------
