UE-11802 Crash while compiling blueprints

The actor on preview level is saved into transition buffer, so the level should be modified. The actor should be destroyed inside a transition scope.

diff --git a/Engine/Source/Editor/Kismet/Private/BlueprintEditor.cpp b/Engine/Source/Editor/Kismet/Private/BlueprintEditor.cpp
--- a/Engine/Source/Editor/Kismet/Private/BlueprintEditor.cpp
+++ b/Engine/Source/Editor/Kismet/Private/BlueprintEditor.cpp
@@ -6635,8 +6635,8 @@ void FBlueprintEditor::NotifyPostChange(const FPropertyChangedEvent& PropertyCha
 		FBlueprintEditorUtils::PostEditChangeBlueprintActors(Blueprint);
 	}
 
-	// Force updates to occur immediately during interactive mode (otherwise the preview won't refresh because it won't be ticking)
-	UpdateSCSPreview(PropertyChangedEvent.ChangeType == EPropertyChangeType::Interactive);
+	// Force updates to occur immediately, so the change is inside a transaction scope. See: UE-11802
+	UpdateSCSPreview(true);
 }
 
 void FBlueprintEditor::OnFinishedChangingProperties(const FPropertyChangedEvent& PropertyChangedEvent)
@@ -7310,7 +7310,7 @@ void FBlueprintEditor::DestroyPreview()
 	if ( PreviewActor != nullptr )
 	{
 		check(PreviewScene.GetWorld());
-		PreviewScene.GetWorld()->EditorDestroyActor(PreviewActor, false);
+		PreviewScene.GetWorld()->EditorDestroyActor(PreviewActor, true);
 	}
 
 	UBlueprint* PreviewBlueprint = GetBlueprintObj();
