Added a menu option when right clicking on Blueprint nodes to view node's documentation.

Moved the "GoToDocumentation" command from MaterialEditorCommands up to GraphEditorCommands.

diff --git a/Engine/Source/Editor/BlueprintGraph/Private/EdGraphSchema_K2.cpp b/Engine/Source/Editor/BlueprintGraph/Private/EdGraphSchema_K2.cpp
--- a/Engine/Source/Editor/BlueprintGraph/Private/EdGraphSchema_K2.cpp
+++ b/Engine/Source/Editor/BlueprintGraph/Private/EdGraphSchema_K2.cpp
@@ -1402,6 +1402,12 @@ void UEdGraphSchema_K2::GetContextMenuActions(const UEdGraph* CurrentGraph, cons
 					MenuBuilder->EndSection();
 				}
 			}
+
+			MenuBuilder->BeginSection("EdGraphSchemaDocumentation", LOCTEXT("DocumentationHeader", "Documentation"));
+			{
+				MenuBuilder->AddMenuEntry( FGraphEditorCommands::Get().GoToDocumentation );
+			}
+			MenuBuilder->EndSection();
 		}
 	}
 
diff --git a/Engine/Source/Editor/EditorStyle/Private/SlateEditorStyle.cpp b/Engine/Source/Editor/EditorStyle/Private/SlateEditorStyle.cpp
--- a/Engine/Source/Editor/EditorStyle/Private/SlateEditorStyle.cpp
+++ b/Engine/Source/Editor/EditorStyle/Private/SlateEditorStyle.cpp
@@ -3948,6 +3948,8 @@ void FSlateEditorStyle::FStyle::SetupGraphEditorStyles()
 		Set( "GraphEditor.HideUnusedPins", new IMAGE_BRUSH( "Icons/hide_unusedpins", Icon40x40 ) );
 		Set( "GraphEditor.HideUnusedPins.Small", new IMAGE_BRUSH( "Icons/hide_unusedpins", Icon20x20 ) );
 
+		Set( "GraphEditor.GoToDocumentation", new IMAGE_BRUSH( "Common/icon_Help_Hover_16x", Icon16x16 ) );
+
 		// Graph editor widgets
 		{
 			// EditableTextBox
@@ -5698,7 +5700,6 @@ void FSlateEditorStyle::FStyle::SetupToolkitStyles()
 		Set( "MaterialEditor.CameraHome.Small", new IMAGE_BRUSH( "Icons/icon_MatEd_Home_40x", Icon20x20 ) );
 		Set( "MaterialEditor.FindInMaterial", new IMAGE_BRUSH( "Icons/icon_Blueprint_Find_40px", Icon40x40 ) );
 		Set( "MaterialEditor.FindInMaterial.Small", new IMAGE_BRUSH( "Icons/icon_Blueprint_Find_40px", Icon20x20 ) );
-		Set( "MaterialEditor.GoToDocumentation", new IMAGE_BRUSH( "Common/icon_Help_Hover_16x", Icon16x16 ) );
 	}
 
 	// Material Instance Editor
diff --git a/Engine/Source/Editor/GraphEditor/Private/GraphEditorActions.cpp b/Engine/Source/Editor/GraphEditor/Private/GraphEditorActions.cpp
--- a/Engine/Source/Editor/GraphEditor/Private/GraphEditorActions.cpp
+++ b/Engine/Source/Editor/GraphEditor/Private/GraphEditorActions.cpp
@@ -76,6 +76,8 @@ void FGraphEditorCommandsImpl::RegisterCommands()
 
 	UI_COMMAND( ZoomIn, "Zoom In", "Zoom in on the graph editor", EUserInterfaceActionType::Button, FInputChord(EKeys::Add))
 	UI_COMMAND( ZoomOut, "Zoom Out", "Zoom out from the graph editor", EUserInterfaceActionType::Button, FInputChord(EKeys::Subtract))
+
+	UI_COMMAND( GoToDocumentation, "View Documentation", "View documentation for this node.", EUserInterfaceActionType::Button, FInputChord());
 }
 
 
diff --git a/Engine/Source/Editor/GraphEditor/Public/GraphEditorActions.h b/Engine/Source/Editor/GraphEditor/Public/GraphEditorActions.h
--- a/Engine/Source/Editor/GraphEditor/Public/GraphEditorActions.h
+++ b/Engine/Source/Editor/GraphEditor/Public/GraphEditorActions.h
@@ -108,6 +108,9 @@ public:
 	// Zoom in and out on the graph editor
 	TSharedPtr< FUICommandInfo > ZoomIn;
 	TSharedPtr< FUICommandInfo > ZoomOut;
+
+	// Go to node documentation
+	TSharedPtr< FUICommandInfo > GoToDocumentation;
 };
 
 class GRAPHEDITOR_API FGraphEditorCommands
diff --git a/Engine/Source/Editor/Kismet/Private/BlueprintEditor.cpp b/Engine/Source/Editor/Kismet/Private/BlueprintEditor.cpp
--- a/Engine/Source/Editor/Kismet/Private/BlueprintEditor.cpp
+++ b/Engine/Source/Editor/Kismet/Private/BlueprintEditor.cpp
@@ -27,6 +27,7 @@
 #include "GraphEditorActions.h"
 #include "SNodePanel.h"
 #include "SDockTab.h"
+#include "EditorClassUtils.h"
 
 #include "SBlueprintEditorToolbar.h"
 #include "FindInBlueprints.h"
@@ -1133,6 +1134,11 @@ TSharedRef<SGraphEditor> FBlueprintEditor::CreateGraphEditorWidget(TSharedRef<FT
 				FIsActionChecked(),
 				FIsActionButtonVisible::CreateSP( this, &FBlueprintEditor::IsNativeCodeBrowsingAvailable )
 				);
+
+			GraphEditorCommands->MapAction(FGraphEditorCommands::Get().GoToDocumentation,
+				FExecuteAction::CreateSP(this, &FBlueprintEditor::OnGoToDocumentation),
+				FCanExecuteAction::CreateSP(this, &FBlueprintEditor::CanGoToDocumentation)
+				);
 		}
 	}
 
@@ -5586,6 +5592,41 @@ void FBlueprintEditor::OnGoToDefinition()
 	OnNodeDoubleClicked(Cast<UEdGraphNode>(*SelectedNodes.CreateConstIterator()));
 }
 
+FString FBlueprintEditor::GetDocLinkForSelectedNode()
+{
+	FString DocumentationLink;
+
+	const FGraphPanelSelectionSet SelectedNodes = GetSelectedNodes();
+	if (SelectedNodes.Num() == 1)
+	{
+		UEdGraphNode* SelectedGraphNode = Cast<UEdGraphNode>(*SelectedNodes.CreateConstIterator());
+		if (SelectedGraphNode != NULL)
+		{
+			FString DocLink = SelectedGraphNode->GetDocumentationLink();
+			FString DocExcerpt = SelectedGraphNode->GetDocumentationExcerptName();
+
+			DocumentationLink = FEditorClassUtils::GetDocumentationLinkFromExcerpt(DocLink, DocExcerpt);
+		}
+	}
+
+	return DocumentationLink;
+}
+
+void FBlueprintEditor::OnGoToDocumentation()
+{
+	FString DocumentationLink = GetDocLinkForSelectedNode();
+	if (!DocumentationLink.IsEmpty())
+	{
+		IDocumentation::Get()->Open(DocumentationLink, FDocumentationSourceInfo(TEXT("rightclick_bpnode")));
+	}
+}
+
+bool FBlueprintEditor::CanGoToDocumentation()
+{
+	FString DocumentationLink = GetDocLinkForSelectedNode();
+	return !DocumentationLink.IsEmpty();
+}
+
 void FBlueprintEditor::ToggleSaveIntermediateBuildProducts()
 {
 	bSaveIntermediateBuildProducts = !bSaveIntermediateBuildProducts;
diff --git a/Engine/Source/Editor/Kismet/Public/BlueprintEditor.h b/Engine/Source/Editor/Kismet/Public/BlueprintEditor.h
--- a/Engine/Source/Editor/Kismet/Public/BlueprintEditor.h
+++ b/Engine/Source/Editor/Kismet/Public/BlueprintEditor.h
@@ -945,6 +945,15 @@ private:
 	/** Checks to see if it is possible to jump to the selected node's graph definition. */
 	bool CanGoToDefinition() const;
 
+	/** Open documentation for the selected node class */
+	void OnGoToDocumentation();
+
+	/** Can we open documentation for the selected node */
+	bool CanGoToDocumentation();
+
+	/** Util to try and get doc link for the currently selected node */
+	FString GetDocLinkForSelectedNode();
+
 public://@TODO
 	TSharedPtr<FDocumentTracker> DocumentManager;
 
diff --git a/Engine/Source/Editor/MaterialEditor/Private/MaterialEditor.cpp b/Engine/Source/Editor/MaterialEditor/Private/MaterialEditor.cpp
--- a/Engine/Source/Editor/MaterialEditor/Private/MaterialEditor.cpp
+++ b/Engine/Source/Editor/MaterialEditor/Private/MaterialEditor.cpp
@@ -3814,7 +3814,7 @@ TSharedRef<SGraphEditor> FMaterialEditor::CreateGraphEditorWidget()
 			FExecuteAction::CreateSP(this, &FMaterialEditor::OnCreateComponentMaskNode)
 			);
 
-		GraphEditorCommands->MapAction(FMaterialEditorCommands::Get().GoToDocumentation,
+		GraphEditorCommands->MapAction(FGraphEditorCommands::Get().GoToDocumentation,
 			FExecuteAction::CreateSP(this, &FMaterialEditor::OnGoToDocumentation),
 			FCanExecuteAction::CreateSP(this, &FMaterialEditor::CanGoToDocumentation)
 			);
diff --git a/Engine/Source/Editor/MaterialEditor/Private/MaterialEditorActions.cpp b/Engine/Source/Editor/MaterialEditor/Private/MaterialEditorActions.cpp
--- a/Engine/Source/Editor/MaterialEditor/Private/MaterialEditorActions.cpp
+++ b/Engine/Source/Editor/MaterialEditor/Private/MaterialEditorActions.cpp
@@ -52,7 +52,6 @@ void FMaterialEditorCommands::RegisterCommands()
 	UI_COMMAND( BreakLink, "Break Link", "Deletes this link.", EUserInterfaceActionType::Button, FInputChord() );
 	UI_COMMAND( ForceRefreshPreviews, "Force Refresh Previews", "Forces a refresh of all previews", EUserInterfaceActionType::Button, FInputChord(EKeys::SpaceBar) );
 	UI_COMMAND( CreateComponentMaskNode, "Create ComponentMask Node", "Creates a ComponentMask node at the current cursor position.", EUserInterfaceActionType::Button, FInputChord(EModifierKey::Shift, EKeys::C));
-	UI_COMMAND( GoToDocumentation, "View Documentation", "View documentation for this node.", EUserInterfaceActionType::Button, FInputChord());
 	UI_COMMAND( FindInMaterial, "Search", "Finds expressions and comments in the current Material", EUserInterfaceActionType::Button, FInputChord(EModifierKey::Control, EKeys::F));
 }
 
diff --git a/Engine/Source/Editor/MaterialEditor/Public/MaterialEditorActions.h b/Engine/Source/Editor/MaterialEditor/Public/MaterialEditorActions.h
--- a/Engine/Source/Editor/MaterialEditor/Public/MaterialEditorActions.h
+++ b/Engine/Source/Editor/MaterialEditor/Public/MaterialEditorActions.h
@@ -155,9 +155,6 @@ public:
 	/** Create component mask node */
 	TSharedPtr< FUICommandInfo > CreateComponentMaskNode;
 
-	/** Go to node documentation */
-	TSharedPtr< FUICommandInfo > GoToDocumentation;
-
 
 	/**
 	 * Initialize commands
diff --git a/Engine/Source/Editor/UnrealEd/Private/MaterialGraphNode.cpp b/Engine/Source/Editor/UnrealEd/Private/MaterialGraphNode.cpp
--- a/Engine/Source/Editor/UnrealEd/Private/MaterialGraphNode.cpp
+++ b/Engine/Source/Editor/UnrealEd/Private/MaterialGraphNode.cpp
@@ -392,7 +392,7 @@ void UMaterialGraphNode::GetContextMenuActions(const FGraphNodeContextMenuBuilde
 
 		Context.MenuBuilder->BeginSection("MaterialEditorMenuDocumentation");
 		{
-			Context.MenuBuilder->AddMenuEntry(FMaterialEditorCommands::Get().GoToDocumentation);
+			Context.MenuBuilder->AddMenuEntry(FGraphEditorCommands::Get().GoToDocumentation);
 		}
 		Context.MenuBuilder->EndSection();
 
