Fix for NonUnityTestCompile_OnMac

diff --git a/Engine/Source/Editor/Kismet/Private/BlueprintEditor.cpp b/Engine/Source/Editor/Kismet/Private/BlueprintEditor.cpp
--- a/Engine/Source/Editor/Kismet/Private/BlueprintEditor.cpp
+++ b/Engine/Source/Editor/Kismet/Private/BlueprintEditor.cpp
@@ -28,6 +28,7 @@
 #include "SNodePanel.h"
 #include "SDockTab.h"
 #include "EditorClassUtils.h"
+#include "IDocumentation.h"
 
 #include "SBlueprintEditorToolbar.h"
 #include "FindInBlueprints.h"
