Make Override categories list respect hidecategories/functions

diff --git a/Engine/Source/Editor/Kismet/Private/SMyBlueprint.cpp b/Engine/Source/Editor/Kismet/Private/SMyBlueprint.cpp
--- a/Engine/Source/Editor/Kismet/Private/SMyBlueprint.cpp
+++ b/Engine/Source/Editor/Kismet/Private/SMyBlueprint.cpp
@@ -1200,7 +1200,11 @@ void SMyBlueprint::CollectAllActions(FGraphActionListBuilderBase& OutAllActions)
 		const UFunction* Function = *FunctionIt;
 		const FName FunctionName = Function->GetFName();
 
-		if ( UEdGraphSchema_K2::CanKismetOverrideFunction(Function) && !OverridableFunctionNames.Contains(FunctionName) && !ImplementedFunctionCache.Contains(FunctionName) && !FBlueprintEditorUtils::FindOverrideForFunction(BlueprintObj, CastChecked<UClass>(Function->GetOuter()), Function->GetFName()) )
+		if (    UEdGraphSchema_K2::CanKismetOverrideFunction(Function) 
+			 && !OverridableFunctionNames.Contains(FunctionName) 
+			 && !ImplementedFunctionCache.Contains(FunctionName) 
+			 && !FObjectEditorUtils::IsFunctionHiddenFromClass(Function, BlueprintObj->ParentClass)
+			 && !FBlueprintEditorUtils::FindOverrideForFunction(BlueprintObj, CastChecked<UClass>(Function->GetOuter()), Function->GetFName()) )
 		{
 			FString FunctionTooltip = UK2Node_CallFunction::GetDefaultTooltipForFunction(Function);
 			FText FunctionDesc = K2Schema->GetFriendlySignatureName(Function);
