[UE-12917] Inherited components marked EditAnywhere can now be re-edited after being set to 'None' in the Blueprint editor.

change summary:
- revised SKismetInspector::UpdateFromObjects() to consider potentially differing property values when constructing the property array for object filtering when the CDO is the object that's being edited (i.e. when native component templates are selected for editing).

diff --git a/Engine/Source/Editor/Kismet/Private/SKismetInspector.cpp b/Engine/Source/Editor/Kismet/Private/SKismetInspector.cpp
--- a/Engine/Source/Editor/Kismet/Private/SKismetInspector.cpp
+++ b/Engine/Source/Editor/Kismet/Private/SKismetInspector.cpp
@@ -613,15 +613,36 @@ void SKismetInspector::UpdateFromObjects(const TArray<UObject*>& PropertyObjects
 
 				if(Object != EditableComponentTemplate)
 				{
-					for(TFieldIterator<UObjectProperty> ObjPropIt(Object->GetClass()); ObjPropIt; ++ObjPropIt)
+					UObjectProperty* ObjectProperty = FindField<UObjectProperty>(Object->GetClass(), EditableComponentTemplate->GetFName());
+					if(ObjectProperty != nullptr)
 					{
-						UObjectProperty* ObjectProperty = *ObjPropIt;
-						check(ObjectProperty != NULL);
-
-						// If the property value matches the current component template, add it as a selected property for filtering
-						if(EditableComponentTemplate == ObjectProperty->GetObjectPropertyValue_InContainer(Object))
+						SelectedObjectProperties.Add(ObjectProperty);
+					}
+					else if(UActorComponent* Archetype = Cast<UActorComponent>(EditableComponentTemplate->GetArchetype()))
+					{
+						if(AActor* Owner = Archetype->GetOwner())
 						{
-							SelectedObjectProperties.Add(ObjectProperty);
+							if(UClass* OwnerClass = Owner->GetClass())
+							{
+								AActor* OwnerCDO = CastChecked<AActor>(OwnerClass->GetDefaultObject());
+								for(TFieldIterator<UObjectProperty> ObjPropIt(OwnerClass, EFieldIteratorFlags::IncludeSuper); ObjPropIt; ++ObjPropIt)
+								{
+									ObjectProperty = *ObjPropIt;
+									check(ObjectProperty != nullptr);
+
+									// If the property value matches the current archetype, add it as a selected property for filtering
+									if(Archetype->GetClass()->IsChildOf(ObjectProperty->PropertyClass)
+										&& Archetype == ObjectProperty->GetObjectPropertyValue_InContainer(OwnerCDO))
+									{
+										ObjectProperty = FindField<UObjectProperty>(Object->GetClass(), ObjectProperty->GetFName());
+										if(ObjectProperty != nullptr)
+										{
+											SelectedObjectProperties.Add(ObjectProperty);
+											break;
+										}
+									}
+								}
+							}
 						}
 					}
 				}
