Fix spelling error

diff --git a/Engine/Source/Editor/Kismet/Private/SSCSEditor.cpp b/Engine/Source/Editor/Kismet/Private/SSCSEditor.cpp
--- a/Engine/Source/Editor/Kismet/Private/SSCSEditor.cpp
+++ b/Engine/Source/Editor/Kismet/Private/SSCSEditor.cpp
@@ -1217,12 +1217,12 @@ UActorComponent* FSCSEditorTreeNodeComponent::INTERNAL_GetOverridenComponentTemp
 
 	FComponentKey Key(GetSCSNode());
 
-	const bool BlueprintCanOverrideComponentFormKey = Key.IsValid()
+	const bool BlueprintCanOverrideComponentFromKey = Key.IsValid()
 		&& Blueprint
 		&& Blueprint->ParentClass
 		&& Blueprint->ParentClass->IsChildOf(Key.OwnerClass);
 
-	if (BlueprintCanOverrideComponentFormKey)
+	if (BlueprintCanOverrideComponentFromKey)
 	{
 		UInheritableComponentHandler* InheritableComponentHandler = Blueprint->GetInheritableComponentHandler(bCreateIfNecessary);
 		if (InheritableComponentHandler)
