Fixes fallout from CL# 2506859 causing right click on K2 variable nodes to not work.

diff --git a/Engine/Source/Editor/Kismet/Private/BlueprintEditor.cpp b/Engine/Source/Editor/Kismet/Private/BlueprintEditor.cpp
--- a/Engine/Source/Editor/Kismet/Private/BlueprintEditor.cpp
+++ b/Engine/Source/Editor/Kismet/Private/BlueprintEditor.cpp
@@ -5620,7 +5620,10 @@ FString FBlueprintEditor::GetDocLinkForSelectedNode()
 			FString DocLink = SelectedGraphNode->GetDocumentationLink();
 			FString DocExcerpt = SelectedGraphNode->GetDocumentationExcerptName();
 
-			DocumentationLink = FEditorClassUtils::GetDocumentationLinkFromExcerpt(DocLink, DocExcerpt);
+			if (!DocLink.IsEmpty() && !DocExcerpt.IsEmpty())
+			{
+				DocumentationLink = FEditorClassUtils::GetDocumentationLinkFromExcerpt(DocLink, DocExcerpt);
+			}
 		}
 	}
 
