This change moves the config variable option in the details panel to the advanced section and updates the tooltip to explain use and config location.
#UE-13102 - Blueprints: "Config Variable" should be an advanced property, and have a more descriptive tool-tip
#Branch UE4
#Proj Kismet

diff --git a/Engine/Source/Editor/Kismet/Private/BlueprintDetailsCustomization.cpp b/Engine/Source/Editor/Kismet/Private/BlueprintDetailsCustomization.cpp
--- a/Engine/Source/Editor/Kismet/Private/BlueprintDetailsCustomization.cpp
+++ b/Engine/Source/Editor/Kismet/Private/BlueprintDetailsCustomization.cpp
@@ -371,9 +371,19 @@ void FBlueprintVarActionDetails::CustomizeDetails( IDetailLayoutBuilder& DetailL
 		.ToolTip(ExposeToMatineeTooltip)
 	];
 
-	TSharedPtr<SToolTip> ExposeToConfigTooltip = IDocumentation::Get()->CreateToolTip(LOCTEXT("VariableExposeToConfig_Tooltip", "Allow this variable be set from the config?"), NULL, DocLink, TEXT("ExposeToConfig"));
+	// Build the property specific config variable tool tip
+	FFormatNamedArguments ConfigTooltipArgs;
+	if (UClass* OwnerClass = VariableProperty->GetOwnerClass())
+	{
+		OwnerClass = OwnerClass->GetAuthoritativeClass();
+		ConfigTooltipArgs.Add(TEXT("ConfigPath"), FText::FromString(OwnerClass->GetDefaultConfigFilename()));
+		ConfigTooltipArgs.Add(TEXT("ConfigSection"), FText::FromString(OwnerClass->GetPathName()));
+	}
+	const FText LocalisedTooltip = FText::Format(LOCTEXT("VariableExposeToConfig_Tooltip", "Should this variable read it's default value from a config file if it is present?\r\n\r\nThis is used for customising variable default values and behavior between different projects and configurations.\r\n\r\nConfig file [{ConfigPath}]\r\nConfig section [{ConfigSection}]"), ConfigTooltipArgs); 
+
+	TSharedPtr<SToolTip> ExposeToConfigTooltip = IDocumentation::Get()->CreateToolTip(LocalisedTooltip, NULL, DocLink, TEXT("ExposeToConfig"));
 
-	Category.AddCustomRow( LOCTEXT("VariableExposeToConfig", "Config Variable") )
+	Category.AddCustomRow( LOCTEXT("VariableExposeToConfig", "Config Variable"), true )
 	.Visibility(TAttribute<EVisibility>(this, &FBlueprintVarActionDetails::ExposeConfigVisibility))
 	.NameContent()
 	[
