[UE-13774] Fixing up scenarios where an SCS has multiple RootNodes and one of them gets picked as the tree's root (removing and re-nesting the others)

diff --git a/Engine/Source/Editor/Kismet/Private/SSCSEditor.cpp b/Engine/Source/Editor/Kismet/Private/SSCSEditor.cpp
--- a/Engine/Source/Editor/Kismet/Private/SSCSEditor.cpp
+++ b/Engine/Source/Editor/Kismet/Private/SSCSEditor.cpp
@@ -4193,12 +4193,32 @@ void SSCSEditor::UpdateTree(bool bRegenerateTreeNodes)
 			{
 				if(ParentBPStack[StackIndex]->SimpleConstructionScript != NULL)
 				{
-					const TArray<USCS_Node*>& SCS_RootNodes = ParentBPStack[StackIndex]->SimpleConstructionScript->GetRootNodes();
-					for(int32 NodeIndex = 0; NodeIndex < SCS_RootNodes.Num(); ++NodeIndex)
+					USimpleConstructionScript* SCS = ParentBPStack[StackIndex]->SimpleConstructionScript;
+					const TArray<USCS_Node*>& SCS_RootNodes = SCS->GetRootNodes();
+
+					// sometimes (like when reparenting a blueprint removes  
+					// inherited root nodes) the SCS will end up with RootNodes
+					// where one of them is the tree's sole root, its sibling
+					// RootNodes need to be nested
+					bool bNestRemainingScsNodes = false;
+
+					for(int32 NodeIndex = 0; NodeIndex < SCS_RootNodes.Num(); )
 					{
 						USCS_Node* SCS_Node = SCS_RootNodes[NodeIndex];
 						check(SCS_Node != NULL);
 
+						if (bNestRemainingScsNodes)
+						{
+							// don't worry about adding it back in as a child, 
+							// AddTreeNode() below will tap into the SCS and add  
+							// it back in
+							SCS->RemoveNode(SCS_Node);
+						}
+						else
+						{
+							++NodeIndex;
+						}
+
 						if(SCS_Node->ParentComponentOrVariableName != NAME_None)
 						{
 							USceneComponent* ParentComponent = SCS_Node->GetParentComponentTemplate(ParentBPStack[0]);
@@ -4215,6 +4235,13 @@ void SSCSEditor::UpdateTree(bool bRegenerateTreeNodes)
 						{
 							AddTreeNode(SCS_Node, SceneRootNodePtr, StackIndex > 0);
 						}
+
+						// if AddTreeNode() just made this node the SceneRootNodePtr
+						// then we know it shouldn't have sibling RootNodes
+						if (SceneRootNodePtr.IsValid() && (SceneRootNodePtr->GetSCSNode() == SCS_Node))
+						{
+							bNestRemainingScsNodes = true;
+						}
 					}
 				}
 			}
