[UE-9395] Perf - Remove redundant preview Actor construction pass on BP editor load.

change summary:
- modified FBlueprintEditor::UpdatePreviewActor() to remove the call to RerunConstructionScripts() that occurred after the initial preview Actor spawn.

note: the RerunConstructionScripts() call was added a while back as part of the fix for TTP# 301026 (CL# 1727280). it appears to no longer be necessary as the details and component in the viewport appear to be correct w/o the extra pass.

diff --git a/Engine/Source/Editor/Kismet/Private/BlueprintEditor.cpp b/Engine/Source/Editor/Kismet/Private/BlueprintEditor.cpp
--- a/Engine/Source/Editor/Kismet/Private/BlueprintEditor.cpp
+++ b/Engine/Source/Editor/Kismet/Private/BlueprintEditor.cpp
@@ -7334,9 +7334,6 @@ void FBlueprintEditor::UpdatePreviewActor(UBlueprint* InBlueprint, bool bInForce
 			{
 				PreviewBlueprint->SimpleConstructionScript->SetComponentEditorActorInstance(PreviewActor);
 			}
-
-			// Run the construction scripts again, otherwise the actor will appear as though it's had a script pass first, rather than the default properties as shown in the details panel
-			PreviewActor->RerunConstructionScripts();
 		}
 	}
 	else if ( PreviewActor )
