[UE-9412] Fixing up Blueprint Palette tooltips so the more match the nodes they spawn (specifically focused on AddComponent nodes).

diff --git a/Engine/Source/Editor/BlueprintGraph/Private/BlueprintComponentNodeSpawner.cpp b/Engine/Source/Editor/BlueprintGraph/Private/BlueprintComponentNodeSpawner.cpp
--- a/Engine/Source/Editor/BlueprintGraph/Private/BlueprintComponentNodeSpawner.cpp
+++ b/Engine/Source/Editor/BlueprintGraph/Private/BlueprintComponentNodeSpawner.cpp
@@ -48,7 +48,7 @@ static FText BlueprintComponentNodeSpawnerImpl::GetDefaultMenuCategory(TSubclass
 UBlueprintComponentNodeSpawner* UBlueprintComponentNodeSpawner::Create(const FComponentTypeEntry& Entry)
 {
 	UClass* ComponentClass = Entry.ComponentClass;
-	if (ComponentClass == nullptr )
+	if (ComponentClass == nullptr)
 	{
 		// unloaded class, must be blueprint created. Create an entry. We'll load the class when we spawn the node:
 
@@ -70,6 +70,8 @@ UBlueprintComponentNodeSpawner* UBlueprintComponentNodeSpawner::Create(const FCo
 			MenuSignature.Keywords = FText::FromString(TEXT(" "));
 		}
 		MenuSignature.IconName = FClassIconFinder::FindIconNameForClass(nullptr);
+		MenuSignature.DocLink  = TEXT("Shared/GraphNodes/Blueprint/UK2Node_AddComponent");
+		MenuSignature.DocExcerptTag = TEXT("AddComponent");
 
 		return NodeSpawner;
 	}
@@ -99,6 +101,8 @@ UBlueprintComponentNodeSpawner* UBlueprintComponentNodeSpawner::Create(const FCo
 		MenuSignature.Keywords = FText::FromString(TEXT(" "));
 	}
 	MenuSignature.IconName = FClassIconFinder::FindIconNameForClass(AuthoritativeClass);
+	MenuSignature.DocLink  = TEXT("Shared/GraphNodes/Blueprint/UK2Node_AddComponent");
+	MenuSignature.DocExcerptTag = AuthoritativeClass->GetName();
 
 	return NodeSpawner;
 }
@@ -128,7 +132,7 @@ UEdGraphNode* UBlueprintComponentNodeSpawner::Invoke(UEdGraph* ParentGraph, FBin
 		UBlueprint* Blueprint = AddCompNode->GetBlueprint();
 		
 		UFunction* AddComponentFunc = FindFieldChecked<UFunction>(AActor::StaticClass(), UK2Node_AddComponent::GetAddComponentFunctionName());
-		AddCompNode->FunctionReference.SetFromField<UFunction>(AddComponentFunc, FBlueprintEditorUtils::IsActorBased(Blueprint));
+		AddCompNode->FunctionReference.SetFromField<UFunction>(AddComponentFunc, !bIsTemplateNode && FBlueprintEditorUtils::IsActorBased(Blueprint));
 
 		UserDelegate.ExecuteIfBound(NewNode, bIsTemplateNode);
 	};
diff --git a/Engine/Source/Editor/BlueprintGraph/Private/BlueprintNodeSpawner.cpp b/Engine/Source/Editor/BlueprintGraph/Private/BlueprintNodeSpawner.cpp
--- a/Engine/Source/Editor/BlueprintGraph/Private/BlueprintNodeSpawner.cpp
+++ b/Engine/Source/Editor/BlueprintGraph/Private/BlueprintNodeSpawner.cpp
@@ -200,6 +200,29 @@ FBlueprintActionUiSpec const& UBlueprintNodeSpawner::PrimeDefaultUiSpec(UEdGraph
 		bTemplateNodeFetched = true;
 	}
 
+	//--------------------------------------
+	// Verify Documentation Link
+	//--------------------------------------
+
+	if (MenuSignature.DocExcerptTag.IsEmpty())
+	{
+		NodeTemplate = bTemplateNodeFetched ? NodeTemplate : GetTemplateNode(TargetGraph);
+		if (NodeTemplate != nullptr)
+		{
+			MenuSignature.DocLink = NodeTemplate->GetDocumentationLink();
+			MenuSignature.DocExcerptTag = NodeTemplate->GetDocumentationExcerptName();
+		}
+
+		// if a target graph was provided, then we've done all we can to spawn a
+		// template node... we have to default to something
+		if (MenuSignature.DocExcerptTag.IsEmpty() && (TargetGraph != nullptr))
+		{
+			// want to set it to something so we won't end up back in this condition
+			MenuSignature.DocExcerptTag = TEXT(" ");
+		}
+		bTemplateNodeFetched = true;
+	}
+
 	return MenuSignature;
 }
 
diff --git a/Engine/Source/Editor/BlueprintGraph/Public/BlueprintNodeSpawner.h b/Engine/Source/Editor/BlueprintGraph/Public/BlueprintNodeSpawner.h
--- a/Engine/Source/Editor/BlueprintGraph/Public/BlueprintNodeSpawner.h
+++ b/Engine/Source/Editor/BlueprintGraph/Public/BlueprintNodeSpawner.h
@@ -27,6 +27,8 @@ struct FBlueprintActionUiSpec
 	FText	Keywords;
 	FName   IconName;
 	FLinearColor IconTint;
+	FString DocLink;
+	FString DocExcerptTag;
 };
 
 /*******************************************************************************
diff --git a/Engine/Source/Editor/Kismet/Private/BlueprintActionMenuItem.cpp b/Engine/Source/Editor/Kismet/Private/BlueprintActionMenuItem.cpp
--- a/Engine/Source/Editor/Kismet/Private/BlueprintActionMenuItem.cpp
+++ b/Engine/Source/Editor/Kismet/Private/BlueprintActionMenuItem.cpp
@@ -7,6 +7,7 @@
 #include "BlueprintEditorUtils.h"	// for AnalyticsTrackNewNode(), MarkBlueprintAsModified(), etc.
 #include "ScopedTransaction.h"
 #include "SNodePanel.h"				// for GetSnapGridSize()
+#include "IDocumentation.h"			// for GetPage()
 
 #define LOCTEXT_NAMESPACE "BlueprintActionMenuItem"
 
@@ -210,6 +211,12 @@ FBlueprintActionMenuItem::FBlueprintActionMenuItem(UBlueprintNodeSpawner const*
 	Category = UiSpec.Category.ToString();
 	TooltipDescription = UiSpec.Tooltip.ToString();
 	Keywords = UiSpec.Keywords;
+	DocExcerptRef.DocLink = UiSpec.DocLink;
+	DocExcerptRef.DocExcerptName = UiSpec.DocExcerptTag;
+	// we may fill out the UiSpec's DocLink with whitespace (so we can tell the
+	// difference between an empty DocLink and one that still needs to be filled 
+	// out), but this could cause troubles later with FDocumentation::GetPage()
+	DocExcerptRef.DocLink.Trim();
 }
 
 //------------------------------------------------------------------------------
@@ -341,6 +348,12 @@ void FBlueprintActionMenuItem::AppendBindings(const FBlueprintActionContext& Con
 	Keywords  = UiSpec.Keywords;	
 	IconBrush = FEditorStyle::GetBrush(UiSpec.IconName);
 	IconTint  = UiSpec.IconTint;
+	DocExcerptRef.DocLink = UiSpec.DocLink;
+	DocExcerptRef.DocExcerptName = UiSpec.DocExcerptTag;
+	// we may fill out the UiSpec's DocLink with whitespace (so we can tell the
+	// difference between an empty DocLink and one that still needs to be filled 
+	// out), but this could cause troubles later with FDocumentation::GetPage()
+	DocExcerptRef.DocLink.Trim();
 }
 
 //------------------------------------------------------------------------------
@@ -350,6 +363,24 @@ FSlateBrush const* FBlueprintActionMenuItem::GetMenuIcon(FSlateColor& ColorOut)
 	return IconBrush;
 }
 
+//------------------------------------------------------------------------------
+const FBlueprintActionMenuItem::FDocExcerptRef& FBlueprintActionMenuItem::GetDocumentationExcerpt() const
+{
+	return DocExcerptRef;
+}
+
+//------------------------------------------------------------------------------
+bool FBlueprintActionMenuItem::FDocExcerptRef::IsValid() const
+{
+	if (DocLink.IsEmpty())
+	{
+		return false;
+	}
+
+	TSharedRef<IDocumentationPage> DocumentationPage = IDocumentation::Get()->GetPage(DocLink, /*Config =*/nullptr);
+	return DocumentationPage->HasExcerpt(DocExcerptName);
+}
+
 //------------------------------------------------------------------------------
 UBlueprintNodeSpawner const* FBlueprintActionMenuItem::GetRawAction() const
 {
diff --git a/Engine/Source/Editor/Kismet/Private/BlueprintActionMenuUtils.cpp b/Engine/Source/Editor/Kismet/Private/BlueprintActionMenuUtils.cpp
--- a/Engine/Source/Editor/Kismet/Private/BlueprintActionMenuUtils.cpp
+++ b/Engine/Source/Editor/Kismet/Private/BlueprintActionMenuUtils.cpp
@@ -4,6 +4,7 @@
 #include "BlueprintActionMenuUtils.h"
 #include "BlueprintActionMenuBuilder.h"
 #include "BlueprintActionMenuItem.h"
+#include "BlueprintDragDropMenuItem.h"
 #include "BlueprintActionFilter.h"
 #include "BlueprintNodeSpawner.h"
 #include "BlueprintEditorUtils.h"	// for DoesSupportComponents()
@@ -727,6 +728,11 @@ const UK2Node* FBlueprintActionMenuUtils::ExtractNodeTemplateFromAction(TSharedP
 			FBlueprintActionMenuItem* NewNodeActionMenuItem = (FBlueprintActionMenuItem*)PaletteAction.Get();
 			TemplateNode = Cast<UK2Node>(NewNodeActionMenuItem->GetRawAction()->GetTemplateNode());
 		}
+		else if (ActionId == FBlueprintDragDropMenuItem::StaticGetTypeId())
+		{
+// 			FBlueprintDragDropMenuItem* DragDropActionMenuItem = (FBlueprintDragDropMenuItem*)PaletteAction.Get();
+// 			TemplateNode = Cast<UK2Node>(DragDropActionMenuItem->GetSampleAction()->GetTemplateNode());
+		}
 		// if this action inherits from FEdGraphSchemaAction_K2NewNode
 		else if (ActionId == FEdGraphSchemaAction_K2NewNode::StaticGetTypeId() ||
 			ActionId == FEdGraphSchemaAction_K2AssignDelegate::StaticGetTypeId() ||
diff --git a/Engine/Source/Editor/Kismet/Private/SBlueprintPalette.cpp b/Engine/Source/Editor/Kismet/Private/SBlueprintPalette.cpp
--- a/Engine/Source/Editor/Kismet/Private/SBlueprintPalette.cpp
+++ b/Engine/Source/Editor/Kismet/Private/SBlueprintPalette.cpp
@@ -282,25 +282,17 @@ static void GetPaletteItemIcon(TSharedPtr<FEdGraphSchemaAction> ActionIn, UBluep
 	// Default to tooltip based on action supplied
 	ToolTipOut = (ActionIn->TooltipDescription.Len() > 0) ? ActionIn->TooltipDescription : ActionIn->MenuDescription.ToString();
 
-	if (ActionIn->GetTypeId() == FEdGraphSchemaAction_K2TargetNode::StaticGetTypeId())
+	if (ActionIn->GetTypeId() == FBlueprintActionMenuItem::StaticGetTypeId())
 	{
-		FEdGraphSchemaAction_K2TargetNode* TargetNodeAction = (FEdGraphSchemaAction_K2TargetNode*)ActionIn.Get();
-		if (UK2Node_Event* EventNode = Cast<UK2Node_Event>(TargetNodeAction->NodeTemplate))
-		{
-			BrushOut = FEditorStyle::GetBrush(TEXT("GraphEditor.Event_16x"));
-		}
+		FBlueprintActionMenuItem* NodeSpawnerAction = (FBlueprintActionMenuItem*)ActionIn.Get();
+		BrushOut = NodeSpawnerAction->GetMenuIcon(ColorOut);
 	}
-	else if (ActionIn->GetTypeId() == FEdGraphSchemaAction_K2AddComponent::StaticGetTypeId())
+	else if (ActionIn->GetTypeId() == FBlueprintDragDropMenuItem::StaticGetTypeId())
 	{
-		FEdGraphSchemaAction_K2AddComponent* AddCompAction = (FEdGraphSchemaAction_K2AddComponent*)ActionIn.Get();
-		// Get icon from class
-		UClass* ComponentClass = *(AddCompAction->ComponentClass);
-		if (ComponentClass)
-		{
-			BrushOut = FClassIconFinder::FindIconForClass(ComponentClass);
-			ToolTipOut = ComponentClass->GetName();
-		}
+		FBlueprintDragDropMenuItem* DragDropAction = (FBlueprintDragDropMenuItem*)ActionIn.Get();
+		BrushOut = DragDropAction->GetMenuIcon(ColorOut);
 	}
+	// for backwards compatibility:
 	else if (UK2Node const* const NodeTemplate = FBlueprintActionMenuUtils::ExtractNodeTemplateFromAction(ActionIn))
 	{
 		// If the node wants to create tooltip text, use that instead, because its probably more detailed
@@ -315,19 +307,20 @@ static void GetPaletteItemIcon(TSharedPtr<FEdGraphSchemaAction> ActionIn, UBluep
 		BrushOut = FEditorStyle::GetBrush(NodeTemplate->GetPaletteIcon(IconLinearColor));
 		ColorOut = IconLinearColor;
 	}
-	else if(ActionIn->GetTypeId() == FEdGraphSchemaAction_K2Graph::StaticGetTypeId())
+	// for MyBlueprint tab specific actions:
+	else if (ActionIn->GetTypeId() == FEdGraphSchemaAction_K2Graph::StaticGetTypeId())
 	{
 		FEdGraphSchemaAction_K2Graph const* GraphAction = (FEdGraphSchemaAction_K2Graph const*)ActionIn.Get();
 		GetSubGraphIcon(GraphAction, BlueprintIn, BrushOut, ToolTipOut);
 	}
-	else if(ActionIn->GetTypeId() == FEdGraphSchemaAction_K2Delegate::StaticGetTypeId())
+	else if (ActionIn->GetTypeId() == FEdGraphSchemaAction_K2Delegate::StaticGetTypeId())
 	{
 		FEdGraphSchemaAction_K2Delegate* DelegateAction = (FEdGraphSchemaAction_K2Delegate*)ActionIn.Get();
 
 		BrushOut = FEditorStyle::GetBrush(TEXT("GraphEditor.Delegate_16x"));
-		ToolTipOut = FString::Printf( *LOCTEXT( "Delegate_Tooltip", "Event Dispatcher '%s'" ).ToString(), *DelegateAction->GetDelegateName().ToString());
+		ToolTipOut = FString::Printf(*LOCTEXT("Delegate_Tooltip", "Event Dispatcher '%s'").ToString(), *DelegateAction->GetDelegateName().ToString());
 	}
-	else if(ActionIn->GetTypeId() == FEdGraphSchemaAction_K2Var::StaticGetTypeId())
+	else if (ActionIn->GetTypeId() == FEdGraphSchemaAction_K2Var::StaticGetTypeId())
 	{
 		FEdGraphSchemaAction_K2Var* VarAction = (FEdGraphSchemaAction_K2Var*)ActionIn.Get();
 
@@ -338,7 +331,7 @@ static void GetPaletteItemIcon(TSharedPtr<FEdGraphSchemaAction> ActionIn, UBluep
 		DocLinkOut = TEXT("Shared/Editor/Blueprint/VariableTypes");
 		DocExcerptOut = GetVarType(VarClass, VarAction->GetVariableName(), false, false);
 	}
-	else if(ActionIn->GetTypeId() == FEdGraphSchemaAction_K2LocalVar::StaticGetTypeId())
+	else if (ActionIn->GetTypeId() == FEdGraphSchemaAction_K2LocalVar::StaticGetTypeId())
 	{
 		FEdGraphSchemaAction_K2LocalVar* LocalVarAction = (FEdGraphSchemaAction_K2LocalVar*)ActionIn.Get();
 
@@ -349,36 +342,19 @@ static void GetPaletteItemIcon(TSharedPtr<FEdGraphSchemaAction> ActionIn, UBluep
 		DocLinkOut = TEXT("Shared/Editor/Blueprint/VariableTypes");
 		DocExcerptOut = GetVarType(VarScope, LocalVarAction->GetVariableName(), false);
 	}
-	else if(ActionIn->GetTypeId() == FEdGraphSchemaAction_K2Enum::StaticGetTypeId())
+	else if (ActionIn->GetTypeId() == FEdGraphSchemaAction_K2Enum::StaticGetTypeId())
 	{
 		BrushOut = FEditorStyle::GetBrush(TEXT("GraphEditor.EnumGlyph"));
-		ToolTipOut = FString::Printf( *LOCTEXT("Enum_Tooltip", "Enum Asset").ToString() );
+		ToolTipOut = FString::Printf(*LOCTEXT("Enum_Tooltip", "Enum Asset").ToString());
 	}
-	else if(ActionIn->GetTypeId() == FEdGraphSchemaAction_K2Struct::StaticGetTypeId())
+	else if (ActionIn->GetTypeId() == FEdGraphSchemaAction_K2Struct::StaticGetTypeId())
 	{
 		BrushOut = FEditorStyle::GetBrush(TEXT("GraphEditor.StructGlyph"));
-		ToolTipOut =  FString::Printf( *LOCTEXT("Struct_Tooltip", "Struct Asset").ToString() );
+		ToolTipOut = FString::Printf(*LOCTEXT("Struct_Tooltip", "Struct Asset").ToString());
 	}
-	else if (ActionIn->GetTypeId() == FEdGraphSchemaAction_K2AddComment::StaticGetTypeId() ||
-		ActionIn->GetTypeId() == FEdGraphSchemaAction_NewStateComment::StaticGetTypeId())
+	else if (ActionIn->GetTypeId() != FEdGraphSchemaAction_Dummy::StaticGetTypeId())
 	{
-		ToolTipOut = LOCTEXT("Comment_Tooltip", "Create a resizable comment box.").ToString();
-		BrushOut = FEditorStyle::GetBrush(TEXT("GraphEditor.Comment_16x"));
-	}
-	else if (ActionIn->GetTypeId() == FBlueprintActionMenuItem::StaticGetTypeId())
-	{
-		FBlueprintActionMenuItem* NodeSpawnerAction = (FBlueprintActionMenuItem*)ActionIn.Get();
-		BrushOut = NodeSpawnerAction->GetMenuIcon(ColorOut);
-	}
-	else if (ActionIn->GetTypeId() == FBlueprintDragDropMenuItem::StaticGetTypeId())
-	{
-		FBlueprintDragDropMenuItem* DragDropAction = (FBlueprintDragDropMenuItem*)ActionIn.Get();
-		BrushOut = DragDropAction->GetMenuIcon(ColorOut);
-	}
-	else if (ActionIn->GetTypeId() == FEdGraphSchemaAction_K2AddDocumentation::StaticGetTypeId())
-	{
-		ToolTipOut = LOCTEXT("Documentation_Tooltip", "Create a documentation node.").ToString();
-		BrushOut = FEditorStyle::GetBrush(TEXT("GraphEditor.Documentation_16x"));
+		ensureMsgf(false, TEXT("Unhandled GraphSchemaAction, wanting a menu icon: %s"), *ActionIn->GetTypeId().ToString());
 	}
 }
 
@@ -1313,13 +1289,13 @@ FText SBlueprintPaletteItem::GetToolTipText() const
 {
 	TSharedPtr<FEdGraphSchemaAction> PaletteAction = ActionPtr.Pin();
 
-	FString ToolTipText;
-	FString ClassDisplayName;
+	FText ToolTipText;
+	FText ClassDisplayName;
 
 	if (PaletteAction.IsValid())
 	{
 		// Default tooltip is taken from the action
-		ToolTipText = (PaletteAction->TooltipDescription.Len() > 0) ? PaletteAction->TooltipDescription : PaletteAction->MenuDescription.ToString();
+		ToolTipText = (PaletteAction->TooltipDescription.Len() > 0) ? FText::FromString(PaletteAction->TooltipDescription) : PaletteAction->MenuDescription;
 
 		if(PaletteAction->GetTypeId() == FEdGraphSchemaAction_K2AddComponent::StaticGetTypeId())
 		{
@@ -1328,20 +1304,14 @@ FText SBlueprintPaletteItem::GetToolTipText() const
 			UClass* ComponentClass = *(AddCompAction->ComponentClass);
 			if (ComponentClass)
 			{
-				ToolTipText = ComponentClass->GetToolTipText().ToString();
+				ToolTipText = ComponentClass->GetToolTipText();
 			}
 		}
 		else if (UK2Node const* const NodeTemplate = FBlueprintActionMenuUtils::ExtractNodeTemplateFromAction(PaletteAction))
 		{
-			// Display the native title of the node when alt is held
-			if(FSlateApplication::Get().GetModifierKeys().IsAltDown())
-			{
-				return FText::FromString(NodeTemplate->GetNodeTitle(ENodeTitleType::ListView).BuildSourceString());
-			}
-
 			// If the node wants to create tooltip text, use that instead, because its probably more detailed
-			FString NodeToolTipText = NodeTemplate->GetTooltipText().ToString();
-			if (NodeToolTipText.Len() > 0)
+			FText NodeToolTipText = NodeTemplate->GetTooltipText();
+			if (!NodeToolTipText.IsEmpty())
 			{
 				ToolTipText = NodeToolTipText;
 			}
@@ -1351,7 +1321,14 @@ FText SBlueprintPaletteItem::GetToolTipText() const
 				if(UClass* ParentClass = CallFuncNode->FunctionReference.GetMemberParentClass(CallFuncNode->GetBlueprintClassFromNode()))
 				{
 					UBlueprint* BlueprintObj = UBlueprint::GetBlueprintFromClass(ParentClass);
-					ClassDisplayName = (BlueprintObj != NULL) ? BlueprintObj->GetName() : ParentClass->GetName();
+					if (BlueprintObj == nullptr)
+					{
+						ClassDisplayName = ParentClass->GetDisplayNameText();
+					}
+					else if (!BlueprintObj->HasAnyFlags(RF_Transient))
+					{
+						ClassDisplayName = FText::FromName(BlueprintObj->GetFName());
+					}					
 				}
 			}
 		}
@@ -1363,7 +1340,7 @@ FText SBlueprintPaletteItem::GetToolTipText() const
 				FGraphDisplayInfo DisplayInfo;
 				GraphAction->EdGraph->GetSchema()->GetGraphDisplayInformation(*(GraphAction->EdGraph), DisplayInfo);
 
-				ToolTipText = DisplayInfo.Tooltip;
+				ToolTipText = FText::FromString(DisplayInfo.Tooltip);
 			}
 		}
 		else if (PaletteAction->GetTypeId() == FEdGraphSchemaAction_K2Var::StaticGetTypeId())
@@ -1373,13 +1350,13 @@ FText SBlueprintPaletteItem::GetToolTipText() const
 			if (bShowClassInTooltip && (VarClass != NULL))
 			{
 				UBlueprint* BlueprintObj = UBlueprint::GetBlueprintFromClass(VarClass);
-				ClassDisplayName = (BlueprintObj != NULL) ? BlueprintObj->GetName() : VarClass->GetName();
+				ClassDisplayName = (BlueprintObj != NULL) ? FText::FromName(BlueprintObj->GetFName()) : VarClass->GetDisplayNameText();
 			}
 			else
 			{
 				FString Result = GetVarTooltip(Blueprint, VarClass, VarAction->GetVariableName());
 				// Only use the variable tooltip if it has been filled out.
-				ToolTipText = !Result.IsEmpty() ? Result : GetVarType(VarClass, VarAction->GetVariableName(), true, true);
+				ToolTipText = FText::FromString( !Result.IsEmpty() ? Result : GetVarType(VarClass, VarAction->GetVariableName(), true, true) );
 			}
 		}
 		else if (PaletteAction->GetTypeId() == FEdGraphSchemaAction_K2LocalVar::StaticGetTypeId())
@@ -1392,14 +1369,14 @@ FText SBlueprintPaletteItem::GetToolTipText() const
 				if (bShowClassInTooltip && (VarClass != NULL))
 				{
 					UBlueprint* BlueprintObj = UBlueprint::GetBlueprintFromClass(VarClass);
-					ClassDisplayName = (BlueprintObj != NULL) ? BlueprintObj->GetName() : VarClass->GetName();
+					ClassDisplayName = (BlueprintObj != NULL) ? FText::FromName(BlueprintObj->GetFName()) : VarClass->GetDisplayNameText();
 				}
 				else
 				{
 					FString Result;
 					FBlueprintEditorUtils::GetBlueprintVariableMetaData(Blueprint, LocalVarAction->GetVariableName(), LocalVarAction->GetVariableScope(), TEXT("tooltip"), Result);
 					// Only use the variable tooltip if it has been filled out.
-					ToolTipText = !Result.IsEmpty() ? Result : GetVarType(LocalVarAction->GetVariableScope(), LocalVarAction->GetVariableName(), true, true);
+					ToolTipText = FText::FromString( !Result.IsEmpty() ? Result : GetVarType(LocalVarAction->GetVariableScope(), LocalVarAction->GetVariableName(), true, true) );
 				}
 			}
 		}
@@ -1408,14 +1385,14 @@ FText SBlueprintPaletteItem::GetToolTipText() const
 			FEdGraphSchemaAction_K2Delegate* DelegateAction = (FEdGraphSchemaAction_K2Delegate*)PaletteAction.Get();
 			
 			FString Result = GetVarTooltip(Blueprint, DelegateAction->GetDelegateClass(), DelegateAction->GetDelegateName());
-			ToolTipText = !Result.IsEmpty() ? Result : DelegateAction->GetDelegateName().ToString();
+			ToolTipText = !Result.IsEmpty() ? FText::FromString(Result) : FText::FromName(DelegateAction->GetDelegateName());
 		}
 		else if (PaletteAction->GetTypeId() == FEdGraphSchemaAction_K2Enum::StaticGetTypeId())
 		{
 			FEdGraphSchemaAction_K2Enum* EnumAction = (FEdGraphSchemaAction_K2Enum*)PaletteAction.Get();
 			if (EnumAction->Enum)
 			{
-				ToolTipText = EnumAction->Enum->GetName();
+				ToolTipText = FText::FromName(EnumAction->Enum->GetFName());
 			}
 		}
 		else if (PaletteAction->GetTypeId() == FEdGraphSchemaAction_K2TargetNode::StaticGetTypeId())
@@ -1423,33 +1400,45 @@ FText SBlueprintPaletteItem::GetToolTipText() const
 			FEdGraphSchemaAction_K2TargetNode* TargetNodeAction = (FEdGraphSchemaAction_K2TargetNode*)PaletteAction.Get();
 			if (TargetNodeAction->NodeTemplate)
 			{
-				ToolTipText = TargetNodeAction->NodeTemplate->GetTooltipText().ToString();
+				ToolTipText = TargetNodeAction->NodeTemplate->GetTooltipText();
 			}
 		}
 	}
 
 	if (bShowClassInTooltip && !ClassDisplayName.IsEmpty())
 	{
-		FString Format = LOCTEXT("BlueprintItemClassTooltip", "%s\nClass: %s").ToString();
-		ToolTipText = FString::Printf(*Format, *ToolTipText, *ClassDisplayName);
+		ToolTipText = FText::Format(LOCTEXT("BlueprintItemClassTooltip", "{0}\nClass: {1}"), ToolTipText, ClassDisplayName);
 	}
 
-	return FText::FromString(ToolTipText);
+	return ToolTipText;
 }
 
 TSharedPtr<SToolTip> SBlueprintPaletteItem::ConstructToolTipWidget() const
 {
 	TSharedPtr<FEdGraphSchemaAction> PaletteAction = ActionPtr.Pin();
+	UEdGraphNode const* const NodeTemplate = FBlueprintActionMenuUtils::ExtractNodeTemplateFromAction(PaletteAction);
 
-	FString DocLink, DocExcerptName;
+	FBlueprintActionMenuItem::FDocExcerptRef DocExcerptRef;
 
 	if (PaletteAction.IsValid())
 	{
-		if (UEdGraphNode const* const NodeTemplate = FBlueprintActionMenuUtils::ExtractNodeTemplateFromAction(PaletteAction))
+		if (NodeTemplate != nullptr)
 		{
 			// Take rich tooltip from node
-			DocLink = NodeTemplate->GetDocumentationLink();
-			DocExcerptName = NodeTemplate->GetDocumentationExcerptName();
+			DocExcerptRef.DocLink = NodeTemplate->GetDocumentationLink();
+			DocExcerptRef.DocExcerptName = NodeTemplate->GetDocumentationExcerptName();
+
+			// sometimes, with FBlueprintActionMenuItem's, the NodeTemplate 
+			// doesn't always reflect the node that will be spawned (some things 
+			// we don't want to be executed until spawn time, like adding of 
+			// component templates)... in that case, the 
+			// FBlueprintActionMenuItem's may have a more specific documentation 
+			// link of its own (most of the time, it will reflect the NodeTemplate's)
+			if ( !DocExcerptRef.IsValid() && (PaletteAction->GetTypeId() == FBlueprintActionMenuItem::StaticGetTypeId()) )
+			{
+				FBlueprintActionMenuItem* NodeSpawnerAction = (FBlueprintActionMenuItem*)PaletteAction.Get();
+				DocExcerptRef = NodeSpawnerAction->GetDocumentationExcerpt();
+			}
 		}
 		else if (PaletteAction->GetTypeId() == FEdGraphSchemaAction_K2Graph::StaticGetTypeId())
 		{
@@ -1459,8 +1448,8 @@ TSharedPtr<SToolTip> SBlueprintPaletteItem::ConstructToolTipWidget() const
 				FGraphDisplayInfo DisplayInfo;
 				GraphAction->EdGraph->GetSchema()->GetGraphDisplayInformation(*(GraphAction->EdGraph), DisplayInfo);
 
-				DocLink = DisplayInfo.DocLink;
-				DocExcerptName = DisplayInfo.DocExcerptName;
+				DocExcerptRef.DocLink = DisplayInfo.DocLink;
+				DocExcerptRef.DocExcerptName = DisplayInfo.DocExcerptName;
 			}
 		}
 		else if (PaletteAction->GetTypeId() == FEdGraphSchemaAction_K2Var::StaticGetTypeId())
@@ -1470,28 +1459,28 @@ TSharedPtr<SToolTip> SBlueprintPaletteItem::ConstructToolTipWidget() const
 			if (!bShowClassInTooltip || VarClass == NULL)
 			{
 				// Don't show big tooltip if we are showing class as well (means we are not in MyBlueprint)
-				DocLink = TEXT("Shared/Editors/BlueprintEditor/GraphTypes");
-				DocExcerptName = TEXT("Variable");
+				DocExcerptRef.DocLink = TEXT("Shared/Editors/BlueprintEditor/GraphTypes");
+				DocExcerptRef.DocExcerptName = TEXT("Variable");
 			}
 		}
 		else if (PaletteAction->GetTypeId() == FEdGraphSchemaAction_K2Event::StaticGetTypeId())
 		{
-			DocLink = TEXT("Shared/Editors/BlueprintEditor/GraphTypes");
-			DocExcerptName = TEXT("Event");
+			DocExcerptRef.DocLink = TEXT("Shared/Editors/BlueprintEditor/GraphTypes");
+			DocExcerptRef.DocExcerptName = TEXT("Event");
 		}
 		else if (PaletteAction->GetTypeId() == FEdGraphSchemaAction_K2AddComment::StaticGetTypeId() ||
 			PaletteAction->GetTypeId() == FEdGraphSchemaAction_NewStateComment::StaticGetTypeId())
 		{
 			// Taking tooltip from action is fine
 			const UEdGraphNode_Comment* DefaultComment = GetDefault<UEdGraphNode_Comment>();
-			DocLink = DefaultComment->GetDocumentationLink();
-			DocExcerptName = DefaultComment->GetDocumentationExcerptName();
+			DocExcerptRef.DocLink = DefaultComment->GetDocumentationLink();
+			DocExcerptRef.DocExcerptName = DefaultComment->GetDocumentationExcerptName();
 		}
 		else if (PaletteAction->GetTypeId() == FEdGraphSchemaAction_K2LocalVar::StaticGetTypeId())
 		{
 			// Don't show big tooltip if we are showing class as well (means we are not in MyBlueprint)
-			DocLink = TEXT("Shared/Editors/BlueprintEditor/GraphTypes");
-			DocExcerptName = TEXT("LocalVariable");
+			DocExcerptRef.DocLink = TEXT("Shared/Editors/BlueprintEditor/GraphTypes");
+			DocExcerptRef.DocExcerptName = TEXT("LocalVariable");
 		}
 	}
 
@@ -1499,46 +1488,108 @@ TSharedPtr<SToolTip> SBlueprintPaletteItem::ConstructToolTipWidget() const
 	TAttribute<FText> TextAttribute;
 	TextAttribute.Bind(this, &SBlueprintPaletteItem::GetToolTipText);
 
-	TSharedRef< SToolTip > TooltipWidget = IDocumentation::Get()->CreateToolTip(TextAttribute, NULL, DocLink, DocExcerptName);
+	TSharedRef< SToolTip > TooltipWidget = IDocumentation::Get()->CreateToolTip(TextAttribute, NULL, DocExcerptRef.DocLink, DocExcerptRef.DocExcerptName);
 
 	// English speakers have no real need to know this exists.
-	if(FInternationalization::Get().GetCurrentCulture()->GetTwoLetterISOLanguageName() != TEXT("en"))
+	if ( (NodeTemplate != nullptr) && (FInternationalization::Get().GetCurrentCulture()->GetTwoLetterISOLanguageName() != TEXT("en")) )
 	{
-		if (UK2Node const* const NodeTemplate = FBlueprintActionMenuUtils::ExtractNodeTemplateFromAction(PaletteAction))
+		FText NativeNodeName = FText::FromString(NodeTemplate->GetNodeTitle(ENodeTitleType::ListView).BuildSourceString());
+		const FTextBlockStyle& SubduedTextStyle = FEditorStyle::GetWidgetStyle<FTextBlockStyle>("Documentation.SDocumentationTooltipSubdued");
+
+		TSharedPtr<SToolTip> InternationalTooltip;
+		TSharedPtr<SVerticalBox> TooltipBody;
+
+		SAssignNew(InternationalTooltip, SToolTip)
+			// Emulate text-only tool-tip styling that SToolTip uses 
+			// when no custom content is supplied.  We want node tool-
+			// tips to be styled just like text-only tool-tips
+			.BorderImage( FCoreStyle::Get().GetBrush("ToolTip.BrightBackground") )
+			.TextMargin(FMargin(11.0f))
+		[
+			SAssignNew(TooltipBody, SVerticalBox)
+		];
+
+		if (!DocExcerptRef.IsValid())
 		{
-			struct Local
+			auto GetNativeNamePromptVisibility = []()->EVisibility
 			{
-				static EVisibility GetNativeNodeNameVisibility()
-				{
-					return FSlateApplication::Get().GetModifierKeys().IsAltDown()? EVisibility::Collapsed : EVisibility::Visible;
-				}
+				FModifierKeysState KeyState = FSlateApplication::Get().GetModifierKeys();
+				return KeyState.IsAltDown() ? EVisibility::Collapsed : EVisibility::Visible;
 			};
 
-			return 
-				SNew(SToolTip)
+			TooltipBody->AddSlot()
+			[
+				SNew(STextBlock)
+					.TextStyle(FEditorStyle::Get(), "Documentation.SDocumentationTooltip")
+					.Text(NativeNodeName)
+					.Visibility_Lambda([GetNativeNamePromptVisibility]()->EVisibility
+					{
+						return (GetNativeNamePromptVisibility() == EVisibility::Visible) ? EVisibility::Collapsed : EVisibility::Visible;
+					})
+			];
 
-				// Emulate text-only tool-tip styling that SToolTip uses when no custom content is supplied.  We want node tool-tips to 
-				// be styled just like text-only tool-tips
-				.BorderImage( FCoreStyle::Get().GetBrush("ToolTip.BrightBackground") )
-				.TextMargin(FMargin(11.0f))
+			TooltipBody->AddSlot()
+			[
+				SNew(SHorizontalBox)
+					.Visibility_Lambda(GetNativeNamePromptVisibility)
+				+SHorizontalBox::Slot()
 				[
-					SNew(SVerticalBox)
-					+SVerticalBox::Slot()
-					[
-						TooltipWidget->GetContentWidget()
-					]
-					+SVerticalBox::Slot()
-						.AutoHeight()
-						.HAlign( HAlign_Right )
-					[
+					TooltipWidget->GetContentWidget()
+				]
+			];
 
-						SNew( STextBlock )
-						.Text( LOCTEXT( "NativeNodeName", "hold (Alt) for native node name" ) )
-						.TextStyle( FEditorStyle::Get(), "Documentation.SDocumentationTooltipSubdued")
-						.Visibility_Static(&Local::GetNativeNodeNameVisibility)
-					]
-				];
+			TooltipBody->AddSlot()
+				.AutoHeight()
+				.HAlign(HAlign_Center)
+				.Padding(0.f, 8.f, 0.f, 0.f)
+			[
+
+				SNew(STextBlock)
+					.Text( LOCTEXT("NativeNodeName", "hold (Alt) for native node name") )
+					.TextStyle(&SubduedTextStyle)
+					.Visibility_Lambda(GetNativeNamePromptVisibility)
+			];
 		}
+		else
+		{
+			auto GetNativeNodeNameVisibility = []()->EVisibility
+			{
+				FModifierKeysState KeyState = FSlateApplication::Get().GetModifierKeys();
+				return KeyState.IsAltDown() && KeyState.IsControlDown() ? EVisibility::Visible : EVisibility::Collapsed;
+			};
+
+			// give the "advanced" tooltip a header
+			TooltipBody->AddSlot()
+				.AutoHeight()
+				.HAlign(HAlign_Right)
+				.Padding(0.f, 0.f, 0.f, 8.f)
+			[
+				SNew(SHorizontalBox)
+				+SHorizontalBox::Slot()
+					.AutoWidth()
+				[
+					SNew(STextBlock)
+					.TextStyle(&SubduedTextStyle)
+						.Text(LOCTEXT("NativeNodeNameLabel", "Native Node Name: "))
+						.Visibility_Lambda(GetNativeNodeNameVisibility)
+				]
+				+SHorizontalBox::Slot()
+					.AutoWidth()
+				[
+					SNew(STextBlock)
+						.TextStyle(&SubduedTextStyle)
+						.Text(NativeNodeName)
+						.Visibility_Lambda(GetNativeNodeNameVisibility)
+				]
+			];
+
+			TooltipBody->AddSlot()
+			[
+				TooltipWidget->GetContentWidget()
+			];
+		}
+
+		return InternationalTooltip;
 	}
 	return TooltipWidget;
 }
diff --git a/Engine/Source/Editor/Kismet/Private/SBlueprintSubPalette.cpp b/Engine/Source/Editor/Kismet/Private/SBlueprintSubPalette.cpp
--- a/Engine/Source/Editor/Kismet/Private/SBlueprintSubPalette.cpp
+++ b/Engine/Source/Editor/Kismet/Private/SBlueprintSubPalette.cpp
@@ -301,8 +301,7 @@ void SBlueprintSubPalette::RefreshActionsList(bool bPreserveExpansion)
 //------------------------------------------------------------------------------
 TSharedRef<SWidget> SBlueprintSubPalette::OnCreateWidgetForAction(FCreateWidgetForActionData* const InCreateData)
 {
-	return SNew(SBlueprintPaletteItem, InCreateData, BlueprintEditorPtr.Pin())
-		.ShowClassInTooltip(true);
+	return SNew(SBlueprintPaletteItem, InCreateData, BlueprintEditorPtr.Pin());
 }
 
 //------------------------------------------------------------------------------
diff --git a/Engine/Source/Editor/Kismet/Public/BlueprintActionMenuItem.h b/Engine/Source/Editor/Kismet/Public/BlueprintActionMenuItem.h
--- a/Engine/Source/Editor/Kismet/Public/BlueprintActionMenuItem.h
+++ b/Engine/Source/Editor/Kismet/Public/BlueprintActionMenuItem.h
@@ -51,7 +51,16 @@ public:
 	 */
 	FSlateBrush const* GetMenuIcon(FSlateColor& ColorOut);
 
-	
+	/** Utility struct for pairing documentation page names with excerpt names */
+	struct FDocExcerptRef
+	{
+		bool IsValid() const;
+
+		FString DocLink;
+		FString DocExcerptName;		
+	};
+	/**  */
+	const FDocExcerptRef& GetDocumentationExcerpt() const;
 
 private:
 	/** Specialized node-spawner, that comprises the action portion of this menu entry. */
@@ -62,4 +71,6 @@ private:
 	FSlateBrush const* IconBrush;
 	/** */
 	IBlueprintNodeBinder::FBindingSet Bindings;
+	/** References the documentation page/excerpt pertaining to the node this will spawn */
+	FDocExcerptRef DocExcerptRef;
 };
diff --git a/Engine/Source/Runtime/Engine/Classes/GameFramework/Actor.h b/Engine/Source/Runtime/Engine/Classes/GameFramework/Actor.h
--- a/Engine/Source/Runtime/Engine/Classes/GameFramework/Actor.h
+++ b/Engine/Source/Runtime/Engine/Classes/GameFramework/Actor.h
@@ -862,10 +862,15 @@ public:
 	bool HasAuthority() const;
 
 	/** 
-	 * BLUEPRINT INTERNAL USE ONLY for Add Component nodes
-	 * Creates a new component owned by this Actor based a component template found in the owning Blueprint.
-	 * Automatic attachment causes the first component created to become the root component, and all subsequent components
-	 * will be attached to the root component.  In manual mode, it is up to the user to attach or set as root
+	 * Creates a new component and assigns ownership to the Actor this is 
+	 * called for. Automatic attachment causes the first component created to 
+	 * become the root, and all subsequent components to be attached under that 
+	 * root. When bManualAttachment is set, automatic attachment is 
+	 * skipped and it is up to the user to attach the resulting component (or 
+	 * set it up as the root) themselves.
+	 *
+	 * @see UK2Node_AddComponent	DO NOT CALL MANUALLY - BLUEPRINT INTERNAL USE ONLY (for Add Component nodes)
+	 *
 	 * @param TemplateName					The name of the Component Template to use.
 	 * @param bManualAttachment				Whether manual or automatic attachment is to be used
 	 * @param RelativeTransform				The relative transform between the new component and its attach parent (automatic only)
