Sub-items in the My Blueprint panel will no longer be mis-categorized into the Function's section.

Category grouping is no longer relevant for matching to sub-items. In the same section, never want two categories of the same name.

diff --git a/Engine/Source/Editor/GraphEditor/Private/GraphActionNode.cpp b/Engine/Source/Editor/GraphEditor/Private/GraphActionNode.cpp
--- a/Engine/Source/Editor/GraphEditor/Private/GraphActionNode.cpp
+++ b/Engine/Source/Editor/GraphEditor/Private/GraphActionNode.cpp
@@ -506,7 +506,6 @@ TSharedPtr<FGraphActionNode> FGraphActionNode::FindMatchingParent(FString const&
 		if (ChildNode->IsCategoryNode())
 		{
 			if ((NodeToAdd->SectionID == ChildNode->SectionID) &&
-				(NodeToAdd->Grouping  == ChildNode->Grouping)  &&
 				(ParentName == ChildNode->DisplayText.ToString()))
 			{
 				FoundCategoryNode = ChildNode;
diff --git a/Engine/Source/Editor/Kismet/Private/SMyBlueprint.cpp b/Engine/Source/Editor/Kismet/Private/SMyBlueprint.cpp
--- a/Engine/Source/Editor/Kismet/Private/SMyBlueprint.cpp
+++ b/Engine/Source/Editor/Kismet/Private/SMyBlueprint.cpp
@@ -463,7 +463,7 @@ void SMyBlueprint::OnCategoryNameCommitted(const FText& InNewText, ETextCommit::
 					FEdGraphSchemaAction_K2Graph* GraphAction = (FEdGraphSchemaAction_K2Graph*)Actions[i].Get();
 
 					// Don't allow changing the category of a graph who's parent is not the current Blueprint
-					if(!FBlueprintEditorUtils::IsPaletteActionReadOnly(Actions[i], BlueprintEditorPtr.Pin()) && FBlueprintEditorUtils::FindBlueprintForGraph(GraphAction->EdGraph) == GetBlueprintObj())
+					if(GraphAction && !FBlueprintEditorUtils::IsPaletteActionReadOnly(Actions[i], BlueprintEditorPtr.Pin()) && FBlueprintEditorUtils::FindBlueprintForGraph(GraphAction->EdGraph) == GetBlueprintObj())
 					{
 						TWeakObjectPtr<UK2Node_EditablePinBase> EntryNode;
 						TWeakObjectPtr<UK2Node_EditablePinBase> ResultNode;
@@ -816,7 +816,7 @@ TSharedRef<SWidget> SMyBlueprint::OnCreateWidgetForAction(FCreateWidgetForAction
 	return BlueprintEditorPtr.IsValid() ? SNew(SBlueprintPaletteItem, InCreateData, BlueprintEditorPtr.Pin()) : SNew(SBlueprintPaletteItem, InCreateData, GetBlueprintObj());
 }
 
-void SMyBlueprint::GetChildGraphs(UEdGraph* InEdGraph, FGraphActionListBuilderBase& OutAllActions, FString ParentCategory)
+void SMyBlueprint::GetChildGraphs(UEdGraph* InEdGraph, int32 const SectionId, FGraphActionListBuilderBase& OutAllActions, FString ParentCategory)
 {
 	// Grab subgraphs
 	const UEdGraphSchema* Schema = InEdGraph->GetSchema();
@@ -839,17 +839,18 @@ void SMyBlueprint::GetChildGraphs(UEdGraph* InEdGraph, FGraphActionListBuilderBa
 		FText DisplayText = ChildGraphDisplayInfo.DisplayName;
 
 		FString Category = ((ParentCategory.IsEmpty()) ? "" : ParentCategory + "|") + EdGraphDisplayName.ToString();
-		FString FunctionTooltip = DisplayText.ToString();
-		FText FunctionDesc = DisplayText;
+		FString ChildTooltip = DisplayText.ToString();
+		FText ChildDesc = DisplayText;
 		const FName DisplayName =  FName(*DisplayText.ToString());
 
-		TSharedPtr<FEdGraphSchemaAction_K2Graph> NewFuncAction = MakeShareable(new FEdGraphSchemaAction_K2Graph(EEdGraphSchemaAction_K2Graph::Subgraph, Category, FunctionDesc, FunctionTooltip, 1));
-		NewFuncAction->FuncName = DisplayName;
-		NewFuncAction->EdGraph = Graph;
-		NewFuncAction->SectionID = NodeSectionID::FUNCTION;
-		OutAllActions.AddAction(NewFuncAction);
+		TSharedPtr<FEdGraphSchemaAction_K2Graph> NewChildAction = MakeShareable(new FEdGraphSchemaAction_K2Graph(EEdGraphSchemaAction_K2Graph::Subgraph, Category, ChildDesc, ChildTooltip, 1));
+		NewChildAction->FuncName = DisplayName;
+		NewChildAction->EdGraph = Graph;
+		NewChildAction->SectionID = SectionId;
+		OutAllActions.AddAction(NewChildAction);
 		
-		GetChildGraphs(Graph, OutAllActions, Category);
+		GetChildGraphs(Graph, SectionId, OutAllActions, Category);
+		GetChildEvents(Graph, SectionId, OutAllActions, Category);
 	}
 }
 
@@ -875,7 +876,7 @@ struct FCreateEdGraphSchemaActionHelper
 	}
 };
 
-void SMyBlueprint::GetChildEvents(UEdGraph const* InEdGraph, int32 const SectionId, FGraphActionListBuilderBase& OutAllActions) const
+void SMyBlueprint::GetChildEvents(UEdGraph const* InEdGraph, int32 const SectionId, FGraphActionListBuilderBase& OutAllActions, FString ParentCategory) const
 {
 	if (!ensure(InEdGraph != NULL))
 	{
@@ -887,7 +888,7 @@ void SMyBlueprint::GetChildEvents(UEdGraph const* InEdGraph, int32 const Section
 	FGraphDisplayInfo EdGraphDisplayInfo;
 	Schema->GetGraphDisplayInformation(*InEdGraph, EdGraphDisplayInfo);
 	FText const EdGraphDisplayName = EdGraphDisplayInfo.DisplayName;
-	FString const ActionCategory = EdGraphDisplayName.ToString();
+	FString ActionCategory = ((ParentCategory.IsEmpty()) ? "" : ParentCategory + "|") + EdGraphDisplayName.ToString();
 
 	FCreateEdGraphSchemaActionHelper::CreateAll<FEdGraphSchemaAction_K2Event, UK2Node_Event>(InEdGraph, SectionId, OutAllActions, ActionCategory);
 	FCreateEdGraphSchemaActionHelper::CreateAll<FEdGraphSchemaAction_K2InputAction, UK2Node_InputKey>(InEdGraph, SectionId, OutAllActions, ActionCategory);
@@ -1159,8 +1160,8 @@ void SMyBlueprint::CollectAllActions(FGraphActionListBuilderBase& OutAllActions)
 
 		OutAllActions.AddAction(NewFuncAction);
 
-		GetChildGraphs(Graph, OutAllActions);
-		GetChildEvents(Graph, NewFuncAction->SectionID, OutAllActions);
+		GetChildGraphs(Graph, NewFuncAction->SectionID, OutAllActions, FunctionCategory);
+		GetChildEvents(Graph, NewFuncAction->SectionID, OutAllActions, FunctionCategory);
 
 		ImplementedFunctionCache.Add(Graph->GetFName());
 	}
@@ -1170,23 +1171,23 @@ void SMyBlueprint::CollectAllActions(FGraphActionListBuilderBase& OutAllActions)
 		UEdGraph* Graph = BlueprintObj->MacroGraphs[i];
 		check(Graph);
 		
-		const FName FunctionName = Graph->GetFName();
+		const FName MacroName = Graph->GetFName();
 
 		FGraphDisplayInfo DisplayInfo;
 		Graph->GetSchema()->GetGraphDisplayInformation(*Graph, DisplayInfo);
 
 		FString MacroCategory = GetGraphCategory(Graph);
 
-		TSharedPtr<FEdGraphSchemaAction_K2Graph> NewFuncAction = MakeShareable(new FEdGraphSchemaAction_K2Graph(EEdGraphSchemaAction_K2Graph::Macro, MacroCategory, DisplayInfo.PlainName, DisplayInfo.Tooltip, 1));
-		NewFuncAction->SectionID = NodeSectionID::MACRO;
-		NewFuncAction->FuncName = FunctionName;
-		NewFuncAction->EdGraph = Graph;
-		OutAllActions.AddAction(NewFuncAction);
+		TSharedPtr<FEdGraphSchemaAction_K2Graph> NewMacroAction = MakeShareable(new FEdGraphSchemaAction_K2Graph(EEdGraphSchemaAction_K2Graph::Macro, MacroCategory, DisplayInfo.PlainName, DisplayInfo.Tooltip, 1));
+		NewMacroAction->SectionID = NodeSectionID::MACRO;
+		NewMacroAction->FuncName = MacroName;
+		NewMacroAction->EdGraph = Graph;
+		OutAllActions.AddAction(NewMacroAction);
 
-		GetChildGraphs(Graph, OutAllActions);
-		GetChildEvents(Graph, NewFuncAction->SectionID, OutAllActions);
+		GetChildGraphs(Graph, NewMacroAction->SectionID, OutAllActions, MacroCategory);
+		GetChildEvents(Graph, NewMacroAction->SectionID, OutAllActions, MacroCategory);
 
-		ImplementedFunctionCache.Add(FunctionName);
+		ImplementedFunctionCache.Add(MacroName);
 	}
 
 	OverridableFunctionActions.Reset();
@@ -1249,8 +1250,8 @@ void SMyBlueprint::CollectAllActions(FGraphActionListBuilderBase& OutAllActions)
 			NewFuncAction->SectionID = NodeSectionID::INTERFACE;
 			OutAllActions.AddAction(NewFuncAction);
 
-			GetChildGraphs(Graph, OutAllActions);
-			GetChildEvents(Graph, NewFuncAction->SectionID, OutAllActions);
+			GetChildGraphs(Graph, NewFuncAction->SectionID, OutAllActions, FunctionCategory);
+			GetChildEvents(Graph, NewFuncAction->SectionID, OutAllActions, FunctionCategory);
 		}
 	}
 
@@ -1293,14 +1294,14 @@ void SMyBlueprint::CollectAllActions(FGraphActionListBuilderBase& OutAllActions)
 		FGraphDisplayInfo DisplayInfo;
 		Graph->GetSchema()->GetGraphDisplayInformation(*Graph, DisplayInfo);
 
-		TSharedPtr<FEdGraphSchemaAction_K2Graph> NewFuncAction = MakeShareable(new FEdGraphSchemaAction_K2Graph(EEdGraphSchemaAction_K2Graph::Graph, FString(), DisplayInfo.PlainName, DisplayInfo.Tooltip, 2));
-		NewFuncAction->FuncName = Graph->GetFName();
-		NewFuncAction->EdGraph = Graph;
-		NewFuncAction->SectionID = NodeSectionID::GRAPH;
-		OutAllActions.AddAction(NewFuncAction);
+		TSharedPtr<FEdGraphSchemaAction_K2Graph> NeUbergraphAction = MakeShareable(new FEdGraphSchemaAction_K2Graph(EEdGraphSchemaAction_K2Graph::Graph, FString(), DisplayInfo.PlainName, DisplayInfo.Tooltip, 2));
+		NeUbergraphAction->FuncName = Graph->GetFName();
+		NeUbergraphAction->EdGraph = Graph;
+		NeUbergraphAction->SectionID = NodeSectionID::GRAPH;
+		OutAllActions.AddAction(NeUbergraphAction);
 
-		GetChildGraphs(Graph, OutAllActions);
-		GetChildEvents(Graph, NewFuncAction->SectionID, OutAllActions);
+		GetChildGraphs(Graph, NeUbergraphAction->SectionID, OutAllActions);
+		GetChildEvents(Graph, NeUbergraphAction->SectionID, OutAllActions);
 	}
 
 	// Grab intermediate pages
@@ -1309,16 +1310,16 @@ void SMyBlueprint::CollectAllActions(FGraphActionListBuilderBase& OutAllActions)
 		UEdGraph* Graph = BlueprintObj->IntermediateGeneratedGraphs[i];
 		check(Graph);
 		
-		const FName FunctionName(*(FString(TEXT("$INTERMEDIATE$_")) + Graph->GetName()));
-		FString FunctionTooltip = FunctionName.ToString();
-		FString FunctionDesc = FunctionName.ToString();
-		TSharedPtr<FEdGraphSchemaAction_K2Graph> NewFuncAction = MakeShareable(new FEdGraphSchemaAction_K2Graph(EEdGraphSchemaAction_K2Graph::Graph, FString(), FText::FromString(FunctionDesc), FunctionTooltip, 1));
-		NewFuncAction->FuncName = FunctionName;
-		NewFuncAction->EdGraph = Graph;
-		OutAllActions.AddAction(NewFuncAction);
-
-		GetChildGraphs(Graph, OutAllActions);
-		GetChildEvents(Graph, NewFuncAction->SectionID, OutAllActions);
+		const FName IntermediateName(*(FString(TEXT("$INTERMEDIATE$_")) + Graph->GetName()));
+		FString IntermediateTooltip = IntermediateName.ToString();
+		FString IntermediateDesc = IntermediateName.ToString();
+		TSharedPtr<FEdGraphSchemaAction_K2Graph> NewIntermediateAction = MakeShareable(new FEdGraphSchemaAction_K2Graph(EEdGraphSchemaAction_K2Graph::Graph, FString(), FText::FromString(IntermediateDesc), IntermediateTooltip, 1));
+		NewIntermediateAction->FuncName = IntermediateName;
+		NewIntermediateAction->EdGraph = Graph;
+		OutAllActions.AddAction(NewIntermediateAction);
+
+		GetChildGraphs(Graph, NewIntermediateAction->SectionID, OutAllActions);
+		GetChildEvents(Graph, NewIntermediateAction->SectionID, OutAllActions);
 	}
 
 	if(GetLocalActionsListVisibility().IsVisible())
diff --git a/Engine/Source/Editor/Kismet/Private/SMyBlueprint.h b/Engine/Source/Editor/Kismet/Private/SMyBlueprint.h
--- a/Engine/Source/Editor/Kismet/Private/SMyBlueprint.h
+++ b/Engine/Source/Editor/Kismet/Private/SMyBlueprint.h
@@ -141,8 +141,8 @@ private:
 	/** Callback used to populate all actions list in SGraphActionMenu */
 	void CollectAllActions(FGraphActionListBuilderBase& OutAllActions);
 	void CollectStaticSections(TArray<int32>& StaticSectionIDs);
-	void GetChildGraphs(UEdGraph* EdGraph, FGraphActionListBuilderBase& OutAllActions, FString ParentCategory = FString());
-	void GetChildEvents(UEdGraph const* EdGraph, int32 const SectionId, FGraphActionListBuilderBase& OutAllActions) const;
+	void GetChildGraphs(UEdGraph* EdGraph, int32 const SectionId, FGraphActionListBuilderBase& OutAllActions, FString ParentCategory = FString());
+	void GetChildEvents(UEdGraph const* EdGraph, int32 const SectionId, FGraphActionListBuilderBase& OutAllActions, FString ParentCategory = FString()) const;
 	void GetLocalVariables(FGraphActionListBuilderBase& OutAllActions) const;
 	
 	/** Handles the visibility of the local action list */
