Getting rid of an invalid ensure in BP menuing code.

diff --git a/Engine/Source/Editor/Kismet/Private/SBlueprintPalette.cpp b/Engine/Source/Editor/Kismet/Private/SBlueprintPalette.cpp
--- a/Engine/Source/Editor/Kismet/Private/SBlueprintPalette.cpp
+++ b/Engine/Source/Editor/Kismet/Private/SBlueprintPalette.cpp
@@ -352,10 +352,6 @@ static void GetPaletteItemIcon(TSharedPtr<FEdGraphSchemaAction> ActionIn, UBluep
 		BrushOut = FEditorStyle::GetBrush(TEXT("GraphEditor.StructGlyph"));
 		ToolTipOut = FString::Printf(*LOCTEXT("Struct_Tooltip", "Struct Asset").ToString());
 	}
-	else if (ActionIn->GetTypeId() != FEdGraphSchemaAction_Dummy::StaticGetTypeId())
-	{
-		ensureMsgf(false, TEXT("Unhandled GraphSchemaAction, wanting a menu icon: %s"), *ActionIn->GetTypeId().ToString());
-	}
 }
 
 /**
