[UE-10578] Component tree in Blueprint editor will display UCS-added nodes when the option is enabled.

diff --git a/Engine/Source/Editor/Kismet/Private/BlueprintEditor.cpp b/Engine/Source/Editor/Kismet/Private/BlueprintEditor.cpp
--- a/Engine/Source/Editor/Kismet/Private/BlueprintEditor.cpp
+++ b/Engine/Source/Editor/Kismet/Private/BlueprintEditor.cpp
@@ -7346,6 +7346,12 @@ void FBlueprintEditor::UpdatePreviewActor(UBlueprint* InBlueprint, bool bInForce
 	{
 		InBlueprint->SimpleConstructionScript->EndEditorComponentConstruction();
 	}
+
+	// Update the tree to include any UCS-constructed components
+	if ( SCSEditor.IsValid() )
+	{
+		SCSEditor->UpdateTree();
+	}
 }
 
 void FBlueprintEditor::DestroyPreview()
diff --git a/Engine/Source/Editor/Kismet/Private/SSCSEditor.cpp b/Engine/Source/Editor/Kismet/Private/SSCSEditor.cpp
--- a/Engine/Source/Editor/Kismet/Private/SSCSEditor.cpp
+++ b/Engine/Source/Editor/Kismet/Private/SSCSEditor.cpp
@@ -3980,13 +3980,10 @@ FSCSEditorTreeNodePtrType SSCSEditor::GetNodeFromActorComponent(const UActorComp
 		}
 
 		// If we have a valid component archetype instance, attempt to find a tree node that corresponds to it
-		if ((EditorMode == EComponentEditorMode::BlueprintSCS && ActorComponent->IsTemplate()) || EditorMode == EComponentEditorMode::ActorInstance)
+		const TArray<FSCSEditorTreeNodePtrType>& Nodes = GetRootNodes();
+		for (int32 i = 0; i < Nodes.Num() && !NodePtr.IsValid(); i++)
 		{
-			const TArray<FSCSEditorTreeNodePtrType>& Nodes = GetRootNodes();
-			for (int32 i = 0; i < Nodes.Num() && !NodePtr.IsValid(); i++)
-			{
-				NodePtr = FindTreeNode(ActorComponent, Nodes[i]);
-			}
+			NodePtr = FindTreeNode(ActorComponent, Nodes[i]);
 		}
 
 		// If we didn't find it in the tree, step up the chain to the parent of the given component and recursively see if that is in the tree (unless the flag is false)
@@ -4218,6 +4215,29 @@ void SSCSEditor::UpdateTree(bool bRegenerateTreeNodes)
 					}
 				}
 			}
+
+			AActor* PreviewActorInstance = PreviewActor.Get();
+			if(PreviewActorInstance != nullptr && !GetDefault<UBlueprintEditorSettings>()->bHideConstructionScriptComponentsInDetailsView)
+			{
+				TInlineComponentArray<UActorComponent*> Components;
+				PreviewActorInstance->GetComponents(Components);
+
+				for(auto Component : Components)
+				{
+					if(Component->CreationMethod == EComponentCreationMethod::UserConstructionScript)
+					{
+						USceneComponent* SceneComponent = Cast<USceneComponent>(Component);
+						if(SceneComponent != nullptr)
+						{
+							AddTreeNodeFromComponent(SceneComponent);
+						}
+						else
+						{
+							AddRootComponentTreeNode(Component);
+						}
+					}
+				}
+			}
 		}
 		else    // EComponentEditorMode::ActorInstance
 		{
@@ -4671,7 +4691,7 @@ bool SSCSEditor::IsComponentSelected(const UPrimitiveComponent* PrimComponent) c
 {
 	check(PrimComponent != NULL);
 
-	FSCSEditorTreeNodePtrType NodePtr = GetNodeFromActorComponent(PrimComponent);
+	FSCSEditorTreeNodePtrType NodePtr = GetNodeFromActorComponent(PrimComponent, false);
 	if(NodePtr.IsValid() && SCSTreeWidget.IsValid())
 	{
 		return SCSTreeWidget->IsItemSelected(NodePtr);
@@ -5186,8 +5206,15 @@ FSCSEditorTreeNodePtrType SSCSEditor::AddTreeNodeFromComponent(USceneComponent*
 		FSCSEditorTreeNodePtrType ParentNodePtr = FindTreeNode(InSceneComponent->AttachParent);
 		if(!ParentNodePtr.IsValid())
 		{
-			// Recursively add the parent node to the tree if it does not exist yet
-			ParentNodePtr = AddTreeNodeFromComponent(InSceneComponent->AttachParent);
+			// If the actual attach parent wasn't found, attempt to find its archetype.
+			// This handles the BP editor case where we might add UCS component nodes taken
+			// from the preview actor instance, which are not themselves template objects.
+			ParentNodePtr = FindTreeNode(Cast<USceneComponent>(InSceneComponent->AttachParent->GetArchetype()));
+			if(!ParentNodePtr.IsValid())
+			{
+				// Recursively add the parent node to the tree if it does not exist yet
+				ParentNodePtr = AddTreeNodeFromComponent(InSceneComponent->AttachParent);
+			}
 		}
 
 		// Add a new tree node for the given scene component
