Fix for a crash in BP editor when an invalid variable is dropped

diff --git a/Engine/Source/Editor/Kismet/Private/BPVariableDragDropAction.cpp b/Engine/Source/Editor/Kismet/Private/BPVariableDragDropAction.cpp
--- a/Engine/Source/Editor/Kismet/Private/BPVariableDragDropAction.cpp
+++ b/Engine/Source/Editor/Kismet/Private/BPVariableDragDropAction.cpp
@@ -594,18 +594,22 @@ FReply FKismetVariableDragDropAction::DroppedOnCategory(FString Category)
 
 bool FKismetVariableDragDropAction::CanVariableBeDropped(const UProperty* InVariableProperty, const UEdGraph& InGraph) const
 {
-	UObject* Outer = InVariableProperty->GetOuter();
+	bool bCanVariableBeDropped = false;
+	if (!InVariableProperty)
+	{
+		UObject* Outer = InVariableProperty->GetOuter();
 
-	// Only allow variables to be placed within the same blueprint (otherwise the self context on the dropped node will be invalid)
-	bool bCanVariableBeDropped = IsFromBlueprint(FBlueprintEditorUtils::FindBlueprintForGraph(&InGraph));
+		// Only allow variables to be placed within the same blueprint (otherwise the self context on the dropped node will be invalid)
+		bCanVariableBeDropped = IsFromBlueprint(FBlueprintEditorUtils::FindBlueprintForGraph(&InGraph));
 
-	// Local variables have some special conditions for being allowed to be placed
-	if(bCanVariableBeDropped && Outer->IsA(UFunction::StaticClass()))
-	{
-		// Check if the top level graph has the same name as the function, if they do not then the variable cannot be placed in the graph
-		if(FBlueprintEditorUtils::GetTopLevelGraph(&InGraph)->GetFName() != Outer->GetFName())
+		// Local variables have some special conditions for being allowed to be placed
+		if (bCanVariableBeDropped && Outer->IsA(UFunction::StaticClass()))
 		{
-			bCanVariableBeDropped = false;
+			// Check if the top level graph has the same name as the function, if they do not then the variable cannot be placed in the graph
+			if (FBlueprintEditorUtils::GetTopLevelGraph(&InGraph)->GetFName() != Outer->GetFName())
+			{
+				bCanVariableBeDropped = false;
+			}
 		}
 	}
 	return bCanVariableBeDropped;
