Fix in FKismetVariableDragDropAction::CanVariableBeDropped

diff --git a/Engine/Source/Editor/Kismet/Private/BPVariableDragDropAction.cpp b/Engine/Source/Editor/Kismet/Private/BPVariableDragDropAction.cpp
--- a/Engine/Source/Editor/Kismet/Private/BPVariableDragDropAction.cpp
+++ b/Engine/Source/Editor/Kismet/Private/BPVariableDragDropAction.cpp
@@ -595,7 +595,7 @@ FReply FKismetVariableDragDropAction::DroppedOnCategory(FString Category)
 bool FKismetVariableDragDropAction::CanVariableBeDropped(const UProperty* InVariableProperty, const UEdGraph& InGraph) const
 {
 	bool bCanVariableBeDropped = false;
-	if (!InVariableProperty)
+	if (InVariableProperty)
 	{
 		UObject* Outer = InVariableProperty->GetOuter();
 
