Added a confirmation message box for indexing all Blueprint for Find-in-Blueprints. This box explains the situation and gives the user the ability to opt out of checking out and saving all Blueprints (and instead just load Blueprints).

diff --git a/Engine/Source/Editor/Kismet/Private/FindInBlueprintManager.cpp b/Engine/Source/Editor/Kismet/Private/FindInBlueprintManager.cpp
--- a/Engine/Source/Editor/Kismet/Private/FindInBlueprintManager.cpp
+++ b/Engine/Source/Editor/Kismet/Private/FindInBlueprintManager.cpp
@@ -621,12 +621,13 @@ class FCacheAllBlueprintsTickableObject
 
 public:
 
-	FCacheAllBlueprintsTickableObject(TArray<FString>& InUncachedBlueprints)
+	FCacheAllBlueprintsTickableObject(TArray<FString>& InUncachedBlueprints, bool bInCheckOutAndSave)
 		: TickCacheIndex(0)
 		, UncachedBlueprints(InUncachedBlueprints)
 		, bIsStarted(false)
 		, bIsCancelled(false)
 		, bRecursionGuard(false)
+		, bCheckOutAndSave(bInCheckOutAndSave)
 	{
 		// Start the Blueprint indexing 'progress' notification
 		FNotificationInfo Info( LOCTEXT("BlueprintIndexMessage", "Indexing Blueprints...") );
@@ -735,12 +736,12 @@ public:
 			}
 
 			bool bIsAssetReadOnlyOnDisk = IFileManager::Get().IsReadOnly( *FinalPackageFilename );
-			bool bFailedToCache = true;
+			bool bFailedToCache = bCheckOutAndSave;
 
-			if (!bIsAssetReadOnlyOnDisk)
+			if (!bIsAssetReadOnlyOnDisk || !bCheckOutAndSave)
 			{
 				UObject* Asset = AssetData.GetAsset();
-				if(Asset)
+				if(Asset && bCheckOutAndSave)
 				{
 					// Assume the package was correctly checked out from SCC
 					bool bOutPackageLocallyWritable = true;
@@ -825,6 +826,9 @@ private:
 
 	/** Guard to prevent TickRecursion */
 	bool bRecursionGuard = false;
+
+	/** If TRUE, Blueprints will be checked out and resaved after being loaded */
+	bool bCheckOutAndSave;
 };
 
 FFindInBlueprintSearchManager& FFindInBlueprintSearchManager::Get()
@@ -1407,24 +1411,33 @@ void FFindInBlueprintSearchManager::CacheAllUncachedBlueprints(TWeakPtr< SFindIn
 	// Do not start another caching process if one is in progress
 	if(!IsCacheInProgress())
 	{
-		// Add all failed to cache Blueprints to the UncachedBlueprints list and resubmit for caching
-		UncachedBlueprints.Append(FailedToCachePaths);
-		FailedToCachePaths.Empty();
-
-		CachingObject = new FCacheAllBlueprintsTickableObject(UncachedBlueprints);
-		OutActiveTimerDelegate.BindRaw(CachingObject, &FCacheAllBlueprintsTickableObject::Tick);
+		FText DialogTitle = LOCTEXT("ConfirmIndexAll_Title", "Indexing All");
+		FFormatNamedArguments Args;
+		Args.Add(TEXT("PackageCount"), UncachedBlueprints.Num());
+		const EAppReturnType::Type ReturnValue = FMessageDialog::Open(EAppMsgType::YesNoCancel, FText::Format(LOCTEXT("CacheAllConfirmationMessage", "This process can take a long time; there are {PackageCount} Blueprints to load.\n\nWould you like to checkout, load, and save all Blueprints to make this indexing permanant? Otherwise, all Blueprints will still be loaded but you will be required to re-index the next time you start the editor!"), Args), &DialogTitle);
 
-		if(!ISourceControlModule::Get().IsEnabled())
+		// If Yes is chosen, checkout and save all Blueprints, if No is chosen, only load all Blueprints
+		if (ReturnValue != EAppReturnType::Cancel)
 		{
-			// Offer to start up Source Control
-			ISourceControlModule::Get().ShowLoginDialog(FSourceControlLoginClosed::CreateRaw(this, &FFindInBlueprintSearchManager::OnCacheAllUncachedBlueprints), ELoginWindowMode::Modeless, EOnLoginWindowStartup::PreserveProvider);
-		}
-		else
-		{
-			OnCacheAllUncachedBlueprints(true);
-		}
+			// Add all failed to cache Blueprints to the UncachedBlueprints list and resubmit for caching
+			UncachedBlueprints.Append(FailedToCachePaths);
+			FailedToCachePaths.Empty();
+
+			CachingObject = new FCacheAllBlueprintsTickableObject(UncachedBlueprints, ReturnValue == EAppReturnType::Yes);
+			OutActiveTimerDelegate.BindRaw(CachingObject, &FCacheAllBlueprintsTickableObject::Tick);
 
-		SourceCachingWidget = InSourceWidget;
+			if(!ISourceControlModule::Get().IsEnabled() && ReturnValue == EAppReturnType::Yes)
+			{
+				// Offer to start up Source Control
+				ISourceControlModule::Get().ShowLoginDialog(FSourceControlLoginClosed::CreateRaw(this, &FFindInBlueprintSearchManager::OnCacheAllUncachedBlueprints), ELoginWindowMode::Modeless, EOnLoginWindowStartup::PreserveProvider);
+			}
+			else
+			{
+				OnCacheAllUncachedBlueprints(true);
+			}
+
+			SourceCachingWidget = InSourceWidget;
+		}
 	}
 }
 
