[UE-14068] Expose the array variable from the SCS tree when an actor's component has been stored in an array (when drag and dropping into a BP graph).

diff --git a/Engine/Source/Editor/Kismet/Private/SSCSEditor.cpp b/Engine/Source/Editor/Kismet/Private/SSCSEditor.cpp
--- a/Engine/Source/Editor/Kismet/Private/SSCSEditor.cpp
+++ b/Engine/Source/Editor/Kismet/Private/SSCSEditor.cpp
@@ -261,6 +261,9 @@ void FSCSRowDragDropOp::HoverTargetChanged()
 {
 	bool bHoverHandled = false;
 
+	FSlateColor IconTint = FLinearColor::White;
+	const FSlateBrush* ErrorSymbol = FEditorStyle::GetBrush(TEXT("Graph.ConnectorFeedback.Error"));
+
 	if(SourceNodes.Num() > 1)
 	{
 		// Display an error message if attempting to drop multiple source items onto a node
@@ -268,15 +271,22 @@ void FSCSRowDragDropOp::HoverTargetChanged()
 		if (VarNodeUnderCursor != NULL)
 		{
 			// Icon/text to draw on tooltip
-			FSlateColor IconColor = FLinearColor::White;
-			const FSlateBrush* StatusSymbol = FEditorStyle::GetBrush(TEXT("Graph.ConnectorFeedback.Error"));
 			FText Message = LOCTEXT("InvalidMultiDropTarget", "Cannot replace node with multiple nodes");
-			SetSimpleFeedbackMessage(StatusSymbol, IconColor, Message);
+			SetSimpleFeedbackMessage(ErrorSymbol, IconTint, Message);
 
 			bHoverHandled = true;
 		}
 	}
 
+	UProperty* VariableProperty = GetVariableProperty();
+	if (!bHoverHandled && (VariableProperty == nullptr))
+	{
+		FText Message = LOCTEXT("CannotFindProperty", "Cannot find corresponding variable (make sure component has been assigned to one)");
+		SetSimpleFeedbackMessage(ErrorSymbol, IconTint, Message);
+
+		bHoverHandled = true;
+	}
+
 	if(!bHoverHandled)
 	{
 		FKismetVariableDragDropAction::HoverTargetChanged();
@@ -769,8 +779,15 @@ FString FSCSEditorTreeNodeComponentBase::GetDisplayString() const
 	FName VariableName = GetVariableName();
 	UActorComponent* ComponentTemplate = GetComponentTemplate();
 
+	UBlueprint* Blueprint = GetBlueprint();
+	UClass* VariableOwner = (Blueprint != nullptr) ? Blueprint->SkeletonGeneratedClass : nullptr;
+
+	bool const bHasValidVarName = (VariableName != NAME_None);
+	bool const bIsArrayVariable = bHasValidVarName && (VariableOwner != nullptr) && 
+		FindField<UArrayProperty>(VariableOwner, VariableName);
+
 	// Only display SCS node variable names in the tree if they have not been autogenerated
-	if (VariableName != NAME_None)
+	if ((VariableName != NAME_None) && !bIsArrayVariable)
 	{
 		return VariableName.ToString();
 	}
diff --git a/Engine/Source/Editor/UnrealEd/Private/Kismet2/ComponentEditorUtils.cpp b/Engine/Source/Editor/UnrealEd/Private/Kismet2/ComponentEditorUtils.cpp
--- a/Engine/Source/Editor/UnrealEd/Private/Kismet2/ComponentEditorUtils.cpp
+++ b/Engine/Source/Editor/UnrealEd/Private/Kismet2/ComponentEditorUtils.cpp
@@ -770,6 +770,28 @@ FName FComponentEditorUtils::FindVariableNameGivenComponentInstance(UActorCompon
 					}
 				}
 			}
+
+			for (TFieldIterator<UArrayProperty> PropIt(OwnerClass, EFieldIteratorFlags::IncludeSuper); PropIt; ++PropIt)
+			{
+				UArrayProperty* TestProperty = *PropIt;
+				void* ArrayPropInstAddress = TestProperty->ContainerPtrToValuePtr<void>(OwnerCDO);
+
+				UObjectProperty* ArrayEntryProp = Cast<UObjectProperty>(TestProperty->Inner);
+				if ((ArrayEntryProp == nullptr) || !ArrayEntryProp->PropertyClass->IsChildOf<UActorComponent>())
+				{
+					continue;
+				}
+
+				FScriptArrayHelper ArrayHelper(TestProperty, ArrayPropInstAddress);
+				for (int32 ComponentIndex = 0; ComponentIndex < ArrayHelper.Num(); ++ComponentIndex)
+				{
+					UObject* ArrayElement = ArrayEntryProp->GetObjectPropertyValue(ArrayHelper.GetRawPtr(ComponentIndex));
+					if (ArrayElement == Archetype)
+					{
+						return TestProperty->GetFName();
+					}
+				}
+			}
 		}
 	}
 
