[UE-14411] Fix for crash introduced by 2454136 which made the overridable function menu complete - but did not handle blueprints that had no event graph (gameplay ability blueprints are one example, but there may be others)

diff --git a/Engine/Source/Editor/Kismet/Private/SMyBlueprint.cpp b/Engine/Source/Editor/Kismet/Private/SMyBlueprint.cpp
--- a/Engine/Source/Editor/Kismet/Private/SMyBlueprint.cpp
+++ b/Engine/Source/Editor/Kismet/Private/SMyBlueprint.cpp
@@ -2025,7 +2025,10 @@ void SMyBlueprint::ImplementFunction(FEdGraphSchemaAction_K2Graph* GraphAction)
 	check(OverrideFunc);
 	UClass* const OverrideFuncClass = CastChecked<UClass>(OverrideFunc->GetOuter())->GetAuthoritativeClass();
 
-	if (UEdGraphSchema_K2::FunctionCanBePlacedAsEvent(OverrideFunc))
+	// Some types of blueprints don't have an event graph (IE gameplay ability blueprints), in that case just make a new graph, even
+	// for events:
+	UEdGraph* EventGraph = FBlueprintEditorUtils::FindEventGraph(GetBlueprintObj());
+	if (UEdGraphSchema_K2::FunctionCanBePlacedAsEvent(OverrideFunc) && EventGraph)
 	{
 		// Add to event graph
 		FName EventName = OverrideFunc->GetFName();
@@ -2037,7 +2040,6 @@ void SMyBlueprint::ImplementFunction(FEdGraphSchemaAction_K2Graph* GraphAction)
 		}
 		else
 		{
-			UEdGraph* EventGraph = FBlueprintEditorUtils::FindEventGraph(GetBlueprintObj());
 			UK2Node_Event* NewEventNodeTemplate = NewObject<UK2Node_Event>();
 			NewEventNodeTemplate->EventReference.SetExternalMember(EventName, OverrideFuncClass);
 			NewEventNodeTemplate->bOverrideFunction = true;
