[UE-13202] Relative transform updates resulting from component attach/detach operations in the Blueprint editor are now propagated correctly.

change summary:
- modified SSCS_RowWidget::OnAttachToDropAction() to propagate changes to relative transform properties applied to the scene component archetype
- modified SSCS_RowWidget::OnDetachFromDropAction() to propagate changes to relative transform properties applied to the scene component archetype
- switched to use GetComponentToWorld() instead of GetSocketTransform() when AttachSocketName is NAME_None when calculating relative transform on attach/detach; this ensures that we're returning the component's actual world transform and not the component-modified transform (e.g. SpringArmComponent)
- modified FBlueprintEditorUtils::PostEditChangeBlueprintActors() to accept a 'bComponentEditChange' flag (default == false)
- modified FBlueprintEditorUtils::PostEditChangeBlueprintActors() to broadcast the Level Editor 'ComponentsEdited' event when 'bComponentEditChange' is true and one or more of the Actor instances are in the Level Editor's current selection set; this ensures that the components tree in the Levle Editor's Details view is always refreshed whenever attachment/detachment occurs in the Blueprint Class scene component hierarchy
- modified SSCS_RowWidget::PostDragDropAction() to pass 'true' to FBlueprintEditorUtils::PostEditChangeBlueprintActors() to indicate a component edit change

diff --git a/Engine/Source/Editor/Kismet/Private/SSCSEditor.cpp b/Engine/Source/Editor/Kismet/Private/SSCSEditor.cpp
--- a/Engine/Source/Editor/Kismet/Private/SSCSEditor.cpp
+++ b/Engine/Source/Editor/Kismet/Private/SSCSEditor.cpp
@@ -2356,6 +2356,16 @@ void SSCS_RowWidget::OnAttachToDropAction(const TArray<FSCSEditorTreeNodePtrType
 				// Get the associated component template if it is a scene component, so we can adjust the transform
 				USceneComponent* SceneComponentTemplate = Cast<USceneComponent>(DroppedNodePtr->GetComponentTemplate());
 
+				// Cache current default values for propagation
+				FVector OldRelativeLocation, OldRelativeScale3D;
+				FRotator OldRelativeRotation;
+				if(SceneComponentTemplate)
+				{
+					OldRelativeLocation = SceneComponentTemplate->RelativeLocation;
+					OldRelativeRotation = SceneComponentTemplate->RelativeRotation;
+					OldRelativeScale3D = SceneComponentTemplate->RelativeScale3D;
+				}
+
 				// Check for a valid parent node
 				FSCSEditorTreeNodePtrType ParentNodePtr = DroppedNodePtr->GetParent();
 				if(ParentNodePtr.IsValid())
@@ -2400,7 +2410,7 @@ void SSCS_RowWidget::OnAttachToDropAction(const TArray<FSCSEditorTreeNodePtrType
 				{
 					// If we find a match, calculate its new position relative to the scene root component instance in its current scene
 					FTransform ComponentToWorld(SceneComponentTemplate->RelativeRotation, SceneComponentTemplate->RelativeLocation, SceneComponentTemplate->RelativeScale3D);
-					FTransform ParentToWorld = ParentSceneComponent->GetSocketTransform(SceneComponentTemplate->AttachSocketName);
+					FTransform ParentToWorld = SceneComponentTemplate->AttachSocketName != NAME_None ? ParentSceneComponent->GetSocketTransform(SceneComponentTemplate->AttachSocketName, RTS_World) : ParentSceneComponent->GetComponentToWorld();
 					FTransform RelativeTM = ComponentToWorld.GetRelativeTransform(ParentToWorld);
 
 					// Store new relative location value (if not set to absolute)
@@ -2421,6 +2431,23 @@ void SSCS_RowWidget::OnAttachToDropAction(const TArray<FSCSEditorTreeNodePtrType
 						SceneComponentTemplate->RelativeScale3D = RelativeTM.GetScale3D();
 					}
 				}
+
+				// Propagate any default value changes out to all instances of the template. If we didn't do this, then instances could incorrectly override the new default value with the old default value when construction scripts are re-run.
+				if(SceneComponentTemplate)
+				{
+					TArray<UObject*> InstancedSceneComponents;
+					SceneComponentTemplate->GetArchetypeInstances(InstancedSceneComponents);
+					for(int32 InstanceIndex = 0; InstanceIndex < InstancedSceneComponents.Num(); ++InstanceIndex)
+					{
+						USceneComponent* InstancedSceneComponent = Cast<USceneComponent>(InstancedSceneComponents[InstanceIndex]);
+						if(InstancedSceneComponent != nullptr)
+						{
+							FComponentEditorUtils::ApplyDefaultValueChange(InstancedSceneComponent, InstancedSceneComponent->RelativeLocation, OldRelativeLocation, SceneComponentTemplate->RelativeLocation);
+							FComponentEditorUtils::ApplyDefaultValueChange(InstancedSceneComponent, InstancedSceneComponent->RelativeRotation, OldRelativeRotation, SceneComponentTemplate->RelativeRotation);
+							FComponentEditorUtils::ApplyDefaultValueChange(InstancedSceneComponent, InstancedSceneComponent->RelativeScale3D,  OldRelativeScale3D,  SceneComponentTemplate->RelativeScale3D);
+						}
+					}
+				}
 			}
 		}
 	}
@@ -2467,6 +2494,9 @@ void SSCS_RowWidget::OnDetachFromDropAction(const TArray<FSCSEditorTreeNodePtrTy
 
 		for(const auto& DroppedNodePtr : DroppedNodePtrs)
 		{
+			FVector OldRelativeLocation, OldRelativeScale3D;
+			FRotator OldRelativeRotation;
+
 			check(DroppedNodePtr.IsValid());
 
 			// Detach the node from its parent
@@ -2476,6 +2506,11 @@ void SSCS_RowWidget::OnDetachFromDropAction(const TArray<FSCSEditorTreeNodePtrTy
 			USceneComponent* SceneComponentTemplate = Cast<USceneComponent>(DroppedNodePtr->GetComponentTemplate());
 			if(SceneComponentTemplate)
 			{
+				// Cache current default values for propagation
+				OldRelativeLocation = SceneComponentTemplate->RelativeLocation;
+				OldRelativeRotation = SceneComponentTemplate->RelativeRotation;
+				OldRelativeScale3D = SceneComponentTemplate->RelativeScale3D;
+
 				// Save current state
 				SceneComponentTemplate->Modify();
 
@@ -2510,7 +2545,7 @@ void SSCS_RowWidget::OnDetachFromDropAction(const TArray<FSCSEditorTreeNodePtrTy
 			{
 				// If we find a match, calculate its new position relative to the scene root component instance in the preview scene
 				FTransform ComponentToWorld(SceneComponentTemplate->RelativeRotation, SceneComponentTemplate->RelativeLocation, SceneComponentTemplate->RelativeScale3D);
-				FTransform ParentToWorld = InstancedSceneRootComponent->GetSocketTransform(SceneComponentTemplate->AttachSocketName);
+				FTransform ParentToWorld = SceneComponentTemplate->AttachSocketName != NAME_None ? InstancedSceneRootComponent->GetSocketTransform(SceneComponentTemplate->AttachSocketName, RTS_World) : InstancedSceneRootComponent->GetComponentToWorld();
 				FTransform RelativeTM = ComponentToWorld.GetRelativeTransform(ParentToWorld);
 
 				// Store new relative location value (if not set to absolute)
@@ -2531,6 +2566,23 @@ void SSCS_RowWidget::OnDetachFromDropAction(const TArray<FSCSEditorTreeNodePtrTy
 					SceneComponentTemplate->RelativeScale3D = RelativeTM.GetScale3D();
 				}
 			}
+
+			// Propagate any default value changes out to all instances of the template. If we didn't do this, then instances could incorrectly override the new default value with the old default value when construction scripts are re-run.
+			if(SceneComponentTemplate)
+			{
+				TArray<UObject*> InstancedSceneComponents;
+				SceneComponentTemplate->GetArchetypeInstances(InstancedSceneComponents);
+				for(int32 InstanceIndex = 0; InstanceIndex < InstancedSceneComponents.Num(); ++InstanceIndex)
+				{
+					USceneComponent* InstancedSceneComponent = Cast<USceneComponent>(InstancedSceneComponents[InstanceIndex]);
+					if(InstancedSceneComponent != nullptr)
+					{
+						FComponentEditorUtils::ApplyDefaultValueChange(InstancedSceneComponent, InstancedSceneComponent->RelativeLocation, OldRelativeLocation, SceneComponentTemplate->RelativeLocation);
+						FComponentEditorUtils::ApplyDefaultValueChange(InstancedSceneComponent, InstancedSceneComponent->RelativeRotation, OldRelativeRotation, SceneComponentTemplate->RelativeRotation);
+						FComponentEditorUtils::ApplyDefaultValueChange(InstancedSceneComponent, InstancedSceneComponent->RelativeScale3D,  OldRelativeScale3D,  SceneComponentTemplate->RelativeScale3D);
+					}
+				}
+			}
 		}
 	}
 	else    // EComponentEditorMode::ActorInstance
@@ -2712,7 +2764,7 @@ void SSCS_RowWidget::PostDragDropAction(bool bRegenerateTreeNodes)
 				UBlueprint* Blueprint = GetBlueprint();
 				if(Blueprint != nullptr)
 				{
-					FBlueprintEditorUtils::PostEditChangeBlueprintActors(Blueprint);
+					FBlueprintEditorUtils::PostEditChangeBlueprintActors(Blueprint, true);
 				}
 			}
 		}
diff --git a/Engine/Source/Editor/UnrealEd/Private/Kismet2/BlueprintEditorUtils.cpp b/Engine/Source/Editor/UnrealEd/Private/Kismet2/BlueprintEditorUtils.cpp
--- a/Engine/Source/Editor/UnrealEd/Private/Kismet2/BlueprintEditorUtils.cpp
+++ b/Engine/Source/Editor/UnrealEd/Private/Kismet2/BlueprintEditorUtils.cpp
@@ -49,6 +49,7 @@
 #include "Editor/Blutility/Public/IBlutilityModule.h"
 
 #include "Engine/InheritableComponentHandler.h"
+#include "Editor/LevelEditor/Public/LevelEditor.h"
 
 #include "EditorCategoryUtils.h"
 #include "EngineUtils.h"
@@ -6890,10 +6891,14 @@ void FBlueprintEditorUtils::UpdateOldPureFunctions( UBlueprint* Blueprint )
 }
 
 /** Call PostEditChange() on any Actors that are based on this Blueprint */
-void FBlueprintEditorUtils::PostEditChangeBlueprintActors(UBlueprint* Blueprint)
+void FBlueprintEditorUtils::PostEditChangeBlueprintActors(UBlueprint* Blueprint, bool bComponentEditChange)
 {
 	if (Blueprint->GeneratedClass && Blueprint->GeneratedClass->IsChildOf(AActor::StaticClass()))
 	{
+		// Get the selected Actor set in the level editor context
+		bool bEditorSelectionChanged = false;
+		const USelection* CurrentEditorActorSelection = GEditor ? GEditor->GetSelectedActors() : nullptr;
+
 		TArray<UObject*> MatchingBlueprintObjects;
 		GetObjectsOfClass(Blueprint->GeneratedClass, MatchingBlueprintObjects, true, (RF_ClassDefaultObject | RF_PendingKill));
 
@@ -6902,6 +6907,19 @@ void FBlueprintEditorUtils::PostEditChangeBlueprintActors(UBlueprint* Blueprint)
 			// We know the class was derived from AActor because we checked the Blueprint->GeneratedClass.
 			AActor* Actor = static_cast<AActor*>(ObjIt);
 			Actor->PostEditChange();
+
+			// Broadcast edit notification if necessary so that the level editor's detail panel is refreshed
+			bEditorSelectionChanged |= CurrentEditorActorSelection && CurrentEditorActorSelection->IsSelected(Actor);
+		}
+
+		// Broadcast edit notifications if necessary so that level editor details are refreshed (e.g. components tree)
+		if(bEditorSelectionChanged)
+		{
+			if(bComponentEditChange)
+			{
+				FLevelEditorModule& LevelEditor = FModuleManager::LoadModuleChecked<FLevelEditorModule>("LevelEditor");
+				LevelEditor.BroadcastComponentsEdited();
+			}
 		}
 	}	
 
diff --git a/Engine/Source/Editor/UnrealEd/Public/Kismet2/BlueprintEditorUtils.h b/Engine/Source/Editor/UnrealEd/Public/Kismet2/BlueprintEditorUtils.h
--- a/Engine/Source/Editor/UnrealEd/Public/Kismet2/BlueprintEditorUtils.h
+++ b/Engine/Source/Editor/UnrealEd/Public/Kismet2/BlueprintEditorUtils.h
@@ -961,7 +961,7 @@ public:
 	static bool PropertyValueToString(const UProperty* Property, const uint8* Container, FString& OutForm);
 
 	/** Call PostEditChange() on all Actors based on the given Blueprint */
-	static void PostEditChangeBlueprintActors(UBlueprint* Blueprint);
+	static void PostEditChangeBlueprintActors(UBlueprint* Blueprint, bool bComponentEditChange = false);
 
 	/** Checks if the property can be modified in given blueprint */
 	static bool IsPropertyReadOnlyInCurrentBlueprint(const UBlueprint* Blueprint, const UProperty* Property);
