Added messages to Find-in-Blueprints' Index All action to warn the editor may become unresponsive. These warnings appear both in the dialog that initially appears when indexing all, and also to appear while indexing.

diff --git a/Engine/Source/Editor/Kismet/Private/FindInBlueprintManager.cpp b/Engine/Source/Editor/Kismet/Private/FindInBlueprintManager.cpp
--- a/Engine/Source/Editor/Kismet/Private/FindInBlueprintManager.cpp
+++ b/Engine/Source/Editor/Kismet/Private/FindInBlueprintManager.cpp
@@ -1414,7 +1414,7 @@ void FFindInBlueprintSearchManager::CacheAllUncachedBlueprints(TWeakPtr< SFindIn
 		FText DialogTitle = LOCTEXT("ConfirmIndexAll_Title", "Indexing All");
 		FFormatNamedArguments Args;
 		Args.Add(TEXT("PackageCount"), UncachedBlueprints.Num());
-		const EAppReturnType::Type ReturnValue = FMessageDialog::Open(EAppMsgType::YesNoCancel, FText::Format(LOCTEXT("CacheAllConfirmationMessage", "This process can take a long time; there are {PackageCount} Blueprints to load.\n\nWould you like to checkout, load, and save all Blueprints to make this indexing permanant? Otherwise, all Blueprints will still be loaded but you will be required to re-index the next time you start the editor!"), Args), &DialogTitle);
+		const EAppReturnType::Type ReturnValue = FMessageDialog::Open(EAppMsgType::YesNoCancel, FText::Format(LOCTEXT("CacheAllConfirmationMessage", "This process can take a long time and the editor may become unresponsive; there are {PackageCount} Blueprints to load.\n\nWould you like to checkout, load, and save all Blueprints to make this indexing permanant? Otherwise, all Blueprints will still be loaded but you will be required to re-index the next time you start the editor!"), Args), &DialogTitle);
 
 		// If Yes is chosen, checkout and save all Blueprints, if No is chosen, only load all Blueprints
 		if (ReturnValue != EAppReturnType::Cancel)
diff --git a/Engine/Source/Editor/Kismet/Private/FindInBlueprints.cpp b/Engine/Source/Editor/Kismet/Private/FindInBlueprints.cpp
--- a/Engine/Source/Editor/Kismet/Private/FindInBlueprints.cpp
+++ b/Engine/Source/Editor/Kismet/Private/FindInBlueprints.cpp
@@ -1129,7 +1129,7 @@ void SFindInBlueprints::ConditionallyAddCacheBar()
 									.OnClicked( this, &SFindInBlueprints::OnCacheAllBlueprints )
 									.Visibility( this, &SFindInBlueprints::GetCacheAllButtonVisibility )
 									.ToolTip(IDocumentation::Get()->CreateToolTip(
-									LOCTEXT("IndexAlLBlueprints_Tooltip", "Loads all non-indexed Blueprints and saves them with their search data. This can be a very slow process, feel free to cancel it whenever and benefit from the Blueprints it indexes."),
+									LOCTEXT("IndexAlLBlueprints_Tooltip", "Loads all non-indexed Blueprints and saves them with their search data. This can be a very slow process and the editor may become unresponsive."),
 									NULL,
 									TEXT("Shared/Editors/BlueprintEditor"),
 									TEXT("FindInBlueprint_IndexAll")))
@@ -1197,10 +1197,23 @@ void SFindInBlueprints::ConditionallyAddCacheBar()
 							.AutoHeight()
 							.Padding(8.0f, 0.0f, 0.0f, 2.0f)
 							[
-								SNew(STextBlock)
-								.Text(this, &SFindInBlueprints::GetCurrentCacheBlueprintName)
-								.Visibility( this, &SFindInBlueprints::GetCachingBlueprintNameVisiblity )
-								.ColorAndOpacity( FCoreStyle::Get().GetColor("ErrorReporting.ForegroundColor") )
+								SNew(SVerticalBox)
+								+SVerticalBox::Slot()
+								.AutoHeight()
+								[
+									SNew(STextBlock)
+									.Text(this, &SFindInBlueprints::GetCurrentCacheBlueprintName)
+									.Visibility( this, &SFindInBlueprints::GetCachingBlueprintNameVisiblity )
+									.ColorAndOpacity( FCoreStyle::Get().GetColor("ErrorReporting.ForegroundColor") )
+								]
+
+								+SVerticalBox::Slot()
+								.AutoHeight()
+								[
+									SNew(STextBlock)
+									.Text(LOCTEXT("FiBUnresponsiveEditorWarning", "NOTE: the editor may become unresponsive for some time!"))
+									.TextStyle(&FCoreStyle::Get().GetWidgetStyle<FTextBlockStyle>( "SmallText" ))
+								]
 							]
 					]
 				];
