Fix inability to replace DefaultSceneRoot in IWCE

diff --git a/Engine/Source/Editor/Kismet/Private/SSCSEditor.cpp b/Engine/Source/Editor/Kismet/Private/SSCSEditor.cpp
--- a/Engine/Source/Editor/Kismet/Private/SSCSEditor.cpp
+++ b/Engine/Source/Editor/Kismet/Private/SSCSEditor.cpp
@@ -2729,6 +2729,7 @@ void SSCS_RowWidget::OnMakeNewRootDropAction(FSCSEditorTreeNodePtrType DroppedNo
 			if (bWasDefaultSceneRoot)
 			{
 				SCSEditorPtr->RemoveComponentNode(OldSceneRootNodePtr);
+				SCSEditorPtr->GetActorContext()->SetRootComponent(CastChecked<USceneComponent>(DroppedNodePtr->GetComponentTemplate()));
 			}
 			else
 			{
@@ -5106,9 +5107,11 @@ void SSCSEditor::RemoveComponentNode(FSCSEditorTreeNodePtrType InNodePtr)
 				SCSTreeWidget->ClearSelection();
 			}
 
+			const bool bWasDefaultSceneRoot = InNodePtr.IsValid() && InNodePtr->IsDefaultSceneRoot();
+
 			// Destroy the component instance
 			ComponentInstance->Modify();
-			ComponentInstance->DestroyComponent(true);
+			ComponentInstance->DestroyComponent(!bWasDefaultSceneRoot);
 		}
 	}
 }
diff --git a/Engine/Source/Runtime/Engine/Private/Components/SceneComponent.cpp b/Engine/Source/Runtime/Engine/Private/Components/SceneComponent.cpp
--- a/Engine/Source/Runtime/Engine/Private/Components/SceneComponent.cpp
+++ b/Engine/Source/Runtime/Engine/Private/Components/SceneComponent.cpp
@@ -647,6 +647,20 @@ void USceneComponent::OnComponentDestroyed()
 
 	ScopedMovementStack.Reset();
 
+	// Detach children before destroying
+	for (int32 Index = AttachChildren.Num()-1; Index >= 0; --Index)
+	{
+		USceneComponent* Child = AttachChildren[Index];
+		if (AttachParent)
+		{
+			Child->AttachTo(AttachParent);
+		}
+		else
+		{
+			Child->DetachFromParent();
+		}
+	}
+
 	// Ensure we are detached before destroying
 	DetachFromParent();
 }
