Undoing an "Apply Changes to Blueprint" action will no longer leave the editor in an invalid state that will lead to crashes with proceeding to redo and undo, or move the selected component.

diff --git a/Engine/Source/Editor/Kismet/Private/SSCSEditor.cpp b/Engine/Source/Editor/Kismet/Private/SSCSEditor.cpp
--- a/Engine/Source/Editor/Kismet/Private/SSCSEditor.cpp
+++ b/Engine/Source/Editor/Kismet/Private/SSCSEditor.cpp
@@ -5605,33 +5605,54 @@ void SSCSEditor::OnApplyChangesToBlueprint() const
 
 	if (Actor != NULL && Blueprint != NULL && Actor->GetClass()->ClassGeneratedBy == Blueprint)
 	{
-		const FScopedTransaction Transaction(LOCTEXT("PushToBlueprintDefaults_Transaction", "Apply Changes to Blueprint"));
-
 		// Cache the actor label as by the time we need it, it may be invalid
 		const FString ActorLabel = Actor->GetActorLabel();
-
-		// Perform the actual copy
 		{
-			AActor* BlueprintCDO = Actor->GetClass()->GetDefaultObject<AActor>();
-			if (BlueprintCDO != NULL)
+			const FScopedTransaction Transaction(LOCTEXT("PushToBlueprintDefaults_Transaction", "Apply Changes to Blueprint"));
+
+			// The component selection state should be maintained
+			GEditor->GetSelectedComponents()->Modify();
+
+			Actor->Modify();
+
+			// Mark components that are either native or from the SCS as modified so they will be restored
+			const TArray<UActorComponent*> Components = Actor->GetComponents();
+			for (UActorComponent* ActorComponent : Components)
 			{
-				const auto CopyOptions = (EditorUtilities::ECopyOptions::Type)(EditorUtilities::ECopyOptions::OnlyCopyEditOrInterpProperties | EditorUtilities::ECopyOptions::PropagateChangesToArchetypeInstances);
-				NumChangedProperties = EditorUtilities::CopyActorProperties(Actor, BlueprintCDO, CopyOptions);
-				if (Actor->GetInstanceComponents().Num() > 0)
+				if (ActorComponent->CreationMethod == EComponentCreationMethod::SimpleConstructionScript || ActorComponent->CreationMethod == EComponentCreationMethod::Native)
 				{
-					FKismetEditorUtilities::AddComponentsToBlueprint(Blueprint, Actor->GetInstanceComponents());
-					NumChangedProperties += Actor->GetInstanceComponents().Num();
-					Actor->ClearInstanceComponents(true);
+					ActorComponent->Modify();
 				}
-				if (NumChangedProperties > 0)
+			}
+
+			// Perform the actual copy
+			{
+				AActor* BlueprintCDO = Actor->GetClass()->GetDefaultObject<AActor>();
+				if (BlueprintCDO != NULL)
 				{
-					FBlueprintEditorUtils::MarkBlueprintAsModified(Blueprint);
-					FKismetEditorUtilities::CompileBlueprint(Blueprint);
-					Actor = nullptr; // It is unsafe to use Actor after this point as it may have been reinstanced, so set it to null to make this obvious
+					const auto CopyOptions = (EditorUtilities::ECopyOptions::Type)(EditorUtilities::ECopyOptions::OnlyCopyEditOrInterpProperties | EditorUtilities::ECopyOptions::PropagateChangesToArchetypeInstances);
+					NumChangedProperties = EditorUtilities::CopyActorProperties(Actor, BlueprintCDO, CopyOptions);
+					if (Actor->GetInstanceComponents().Num() > 0)
+					{
+						FKismetEditorUtilities::AddComponentsToBlueprint(Blueprint, Actor->GetInstanceComponents());
+						NumChangedProperties += Actor->GetInstanceComponents().Num();
+						Actor->ClearInstanceComponents(true);
+					}
+					if (NumChangedProperties > 0)
+					{
+						Actor = nullptr; // It is unsafe to use Actor after this point as it may have been reinstanced, so set it to null to make this obvious
+					}
 				}
 			}
 		}
 
+		// Compile the BP outside of the transaction
+ 		if (NumChangedProperties > 0)
+ 		{
+			FBlueprintEditorUtils::MarkBlueprintAsModified(Blueprint);
+			FKismetEditorUtilities::CompileBlueprint(Blueprint);
+ 		}
+
 		// Set up a notification record to indicate success/failure
 		FNotificationInfo NotificationInfo(FText::GetEmpty());
 		NotificationInfo.FadeInDuration = 1.0f;
