Fixed crash dragging the details panel in the BP editor before the viewport has ever been initialized

https://jira.ol.epicgames.net/browse/UE-15075

diff --git a/Engine/Source/Editor/Kismet/Private/BlueprintEditor.cpp b/Engine/Source/Editor/Kismet/Private/BlueprintEditor.cpp
--- a/Engine/Source/Editor/Kismet/Private/BlueprintEditor.cpp
+++ b/Engine/Source/Editor/Kismet/Private/BlueprintEditor.cpp
@@ -1176,7 +1176,6 @@ TSharedRef<SGraphEditor> FBlueprintEditor::CreateGraphEditorWidget(TSharedRef<FT
 	TSharedRef<SGraphEditor> Editor = SNew(SGraphEditor)
 		.AdditionalCommands(GraphEditorCommands)
 		.IsEditable(this, &FBlueprintEditor::IsEditable, InGraph)
-		.DisplayAsReadOnly(this, &FBlueprintEditor::IsGraphReadOnly, InGraph)
 		.TitleBar(TitleBarWidget)
 		.Appearance(this, &FBlueprintEditor::GetGraphAppearance, InGraph)
 		.GraphToEdit(InGraph)
@@ -7521,12 +7520,7 @@ UEdGraph* FBlueprintEditor::GetFocusedGraph() const
 
 bool FBlueprintEditor::IsEditable(UEdGraph* InGraph) const
 {
-	return InEditingMode() && !IsGraphReadOnly(InGraph);
-}
-
-bool FBlueprintEditor::IsGraphReadOnly(UEdGraph* InGraph) const
-{
-	return !(InGraph && InGraph->bEditable && IsGraphPanelEnabled(InGraph));
+	return InEditingMode() && InGraph && InGraph->bEditable && IsGraphPanelEnabled(InGraph);
 }
 
 bool FBlueprintEditor::IsGraphPanelEnabled(UEdGraph* InGraph) const
@@ -7608,19 +7602,22 @@ void FBlueprintEditor::TryInvokingDetailsTab(bool bFlash)
 	{
 		TSharedPtr<SDockTab> BlueprintTab = FGlobalTabmanager::Get()->GetMajorTabForTabManager(TabManager.ToSharedRef());
 
-		// We don't want to force this tab into existance when the blueprint editor isn't in the foreground and actively
+		// We don't want to force this tab into existence when the blueprint editor isn't in the foreground and actively
 		// being interacted with.  So we make sure the window it's in is focused and the tab is in the foreground.
-		if ( BlueprintTab.IsValid() && BlueprintTab->IsForeground() && BlueprintTab->GetParentWindow()->HasFocusedDescendants() )
+		if ( BlueprintTab.IsValid() && BlueprintTab->IsForeground() && BlueprintTab->GetParentWindow()->HasFocusedDescendants()  )
 		{
-			// Show the details panel if it doesn't exist.
-			TabManager->InvokeTab(FBlueprintEditorTabs::DetailsID);
-
-			if ( bFlash )
+			if( !Inspector.IsValid() || (Inspector->GetOwnerTab()->GetDockArea().IsValid()) ) 
 			{
-				TSharedPtr<SDockTab> OwnerTab = Inspector->GetOwnerTab();
-				if ( OwnerTab.IsValid() )
+				// Show the details panel if it doesn't exist.
+				TabManager->InvokeTab(FBlueprintEditorTabs::DetailsID);
+
+				if(bFlash)
 				{
-					OwnerTab->FlashTab();
+					TSharedPtr<SDockTab> OwnerTab = Inspector->GetOwnerTab();
+					if(OwnerTab.IsValid())
+					{
+						OwnerTab->FlashTab();
+					}
 				}
 			}
 		}
