Re-adding some changes removed in CL# 2541596. Verified UE-15075 functions.

diff --git a/Engine/Source/Editor/Kismet/Private/BlueprintEditor.cpp b/Engine/Source/Editor/Kismet/Private/BlueprintEditor.cpp
--- a/Engine/Source/Editor/Kismet/Private/BlueprintEditor.cpp
+++ b/Engine/Source/Editor/Kismet/Private/BlueprintEditor.cpp
@@ -1176,6 +1176,7 @@ TSharedRef<SGraphEditor> FBlueprintEditor::CreateGraphEditorWidget(TSharedRef<FT
 	TSharedRef<SGraphEditor> Editor = SNew(SGraphEditor)
 		.AdditionalCommands(GraphEditorCommands)
 		.IsEditable(this, &FBlueprintEditor::IsEditable, InGraph)
+		.DisplayAsReadOnly(this, &FBlueprintEditor::IsGraphReadOnly, InGraph)
 		.TitleBar(TitleBarWidget)
 		.Appearance(this, &FBlueprintEditor::GetGraphAppearance, InGraph)
 		.GraphToEdit(InGraph)
@@ -7520,7 +7521,12 @@ UEdGraph* FBlueprintEditor::GetFocusedGraph() const
 
 bool FBlueprintEditor::IsEditable(UEdGraph* InGraph) const
 {
-	return InEditingMode() && InGraph && InGraph->bEditable && IsGraphPanelEnabled(InGraph);
+	return InEditingMode() && !IsGraphReadOnly(InGraph);
+}
+
+bool FBlueprintEditor::IsGraphReadOnly(UEdGraph* InGraph) const
+{
+	return !(InGraph && InGraph->bEditable && IsGraphPanelEnabled(InGraph));
 }
 
 bool FBlueprintEditor::IsGraphPanelEnabled(UEdGraph* InGraph) const
