Editor: Added the ability to double-click on the asset name in a SContentReference widget to open that asset

diff --git a/Engine/Source/Editor/Kismet/Private/WorkflowOrientedApp/SContentReference.cpp b/Engine/Source/Editor/Kismet/Private/WorkflowOrientedApp/SContentReference.cpp
--- a/Engine/Source/Editor/Kismet/Private/WorkflowOrientedApp/SContentReference.cpp
+++ b/Engine/Source/Editor/Kismet/Private/WorkflowOrientedApp/SContentReference.cpp
@@ -61,6 +61,7 @@ void SContentReference::Construct(const FArguments& InArgs)
 			.BorderBackgroundColor( FLinearColor::White )
 			.ForegroundColor(FEditorStyle::GetSlateColor(InvertedForegroundName))
 			.ToolTipText(this, &SContentReference::GetAssetFullName)
+			.OnMouseDoubleClick(this, &SContentReference::OnDoubleClickedOnAssetName)
 			[
 				SNew(SBox)
 				.WidthOverride(InArgs._WidthOverride)
@@ -253,6 +254,20 @@ FText SContentReference::GetAssetFullName() const
 	}
 }
 
+FReply SContentReference::OnDoubleClickedOnAssetName(const FGeometry& InMyGeometry, const FPointerEvent& InMouseEvent)
+{
+	OpenAssetToEdit();
+	return FReply::Handled();
+}
+
+void SContentReference::OpenAssetToEdit()
+{
+	if (UObject* Asset = AssetReference.Get())
+	{
+		GEditor->EditObject(Asset);
+	}
+}
+
 const FSlateBrush* SContentReference::GetBorderImage() const
 {
 	return AssetReferenceNameBorderWidget->IsHovered() ? BorderImageHovered : BorderImageNormal;
diff --git a/Engine/Source/Editor/Kismet/Public/WorkflowOrientedApp/SContentReference.h b/Engine/Source/Editor/Kismet/Public/WorkflowOrientedApp/SContentReference.h
--- a/Engine/Source/Editor/Kismet/Public/WorkflowOrientedApp/SContentReference.h
+++ b/Engine/Source/Editor/Kismet/Public/WorkflowOrientedApp/SContentReference.h
@@ -91,6 +91,9 @@ protected:
 	FText GetAssetShortName() const;
 	FText GetAssetFullName() const;
 
+	FReply OnDoubleClickedOnAssetName(const FGeometry& InMyGeometry, const FPointerEvent& InMouseEvent);
+	void OpenAssetToEdit();
+
 	const FSlateBrush* GetBorderImage() const;
 
 	static void FindObjectInContentBrowser(UObject* Object);
