Can duplicate macro graphs in the My Blueprint window.

diff --git a/Engine/Source/Editor/BlueprintGraph/Private/EdGraphSchema_K2.cpp b/Engine/Source/Editor/BlueprintGraph/Private/EdGraphSchema_K2.cpp
--- a/Engine/Source/Editor/BlueprintGraph/Private/EdGraphSchema_K2.cpp
+++ b/Engine/Source/Editor/BlueprintGraph/Private/EdGraphSchema_K2.cpp
@@ -5104,7 +5104,9 @@ TSharedPtr<FEdGraphSchemaAction> UEdGraphSchema_K2::GetCreateDocumentNodeAction(
 
 bool UEdGraphSchema_K2::CanDuplicateGraph(UEdGraph* InSourceGraph) const
 {
-	if(GetGraphType(InSourceGraph) == GT_Function)
+	EGraphType GraphType = GetGraphType(InSourceGraph);
+
+	if(GraphType == GT_Function)
 	{
 		UBlueprint* SourceBP = FBlueprintEditorUtils::FindBlueprintForGraph(InSourceGraph);
 
@@ -5127,7 +5129,7 @@ bool UEdGraphSchema_K2::CanDuplicateGraph(UEdGraph* InSourceGraph) const
 		}
 	}
 
-	return true;
+	return GraphType == GT_Function || GraphType == GT_Macro;
 }
 
 UEdGraph* UEdGraphSchema_K2::DuplicateGraph(UEdGraph* GraphToDuplicate) const
diff --git a/Engine/Source/Editor/Kismet/Private/SMyBlueprint.cpp b/Engine/Source/Editor/Kismet/Private/SMyBlueprint.cpp
--- a/Engine/Source/Editor/Kismet/Private/SMyBlueprint.cpp
+++ b/Engine/Source/Editor/Kismet/Private/SMyBlueprint.cpp
@@ -843,6 +843,10 @@ void SMyBlueprint::GetChildGraphs(UEdGraph* InEdGraph, int32 const SectionId, FG
 		{
 			Category = FText::Format(FText::FromString(TEXT("{0}|{1}")), ParentCategory, EdGraphDisplayName);
 		}
+		else
+		{
+			Category = EdGraphDisplayName;
+		}
 
 		FString ChildTooltip = DisplayText.ToString();
 		FText ChildDesc = DisplayText;
@@ -898,6 +902,10 @@ void SMyBlueprint::GetChildEvents(UEdGraph const* InEdGraph, int32 const Section
 	{
 		ActionCategory = FText::Format(FText::FromString(TEXT("{0}|{1}")), ParentCategory, EdGraphDisplayName);
 	}
+	else
+	{
+		ActionCategory = EdGraphDisplayName;
+	}
 
 	FCreateEdGraphSchemaActionHelper::CreateAll<FEdGraphSchemaAction_K2Event, UK2Node_Event>(InEdGraph, SectionId, OutAllActions, ActionCategory);
 	FCreateEdGraphSchemaActionHelper::CreateAll<FEdGraphSchemaAction_K2InputAction, UK2Node_InputKey>(InEdGraph, SectionId, OutAllActions, ActionCategory);
@@ -2435,7 +2443,7 @@ bool SMyBlueprint::CanDuplicateAction() const
 	if (FEdGraphSchemaAction_K2Graph* GraphAction = SelectionAsGraph())
 	{
 		// Only support function graph duplication
-		if(GraphAction->EdGraph && GraphAction->EdGraph->GetSchema()->GetGraphType(GraphAction->EdGraph) == GT_Function)
+		if(GraphAction->EdGraph)
 		{
 			return GraphAction->EdGraph->GetSchema()->CanDuplicateGraph(GraphAction->EdGraph);
 		}
@@ -2476,9 +2484,16 @@ void SMyBlueprint::OnDuplicateAction()
 
 		// Only function duplication is supported
 		EGraphType GraphType = DuplicatedGraph->GetSchema()->GetGraphType(GraphAction->EdGraph);
-		check(GraphType == GT_Function);
+		check(GraphType == GT_Function || GraphType == GT_Macro);
 
-		GetBlueprintObj()->FunctionGraphs.Add(DuplicatedGraph);
+		if (GraphType == GT_Function)
+		{
+			GetBlueprintObj()->FunctionGraphs.Add(DuplicatedGraph);
+		}
+		else if (GraphType == GT_Macro)
+		{
+			GetBlueprintObj()->MacroGraphs.Add(DuplicatedGraph);
+		}
 		FBlueprintEditorUtils::MarkBlueprintAsStructurallyModified(GetBlueprintObj());
 
 		DuplicateActionName = DuplicatedGraph->GetFName();
