Can set Function and Macro graph metadata for CompactNodeTitle and Keywords.

Changed the StandardMacros Increment, Decrement, and Negate functions to use this metadata and removed the code forcing them to appear as compact nodes from UK2Node_MacroInstance.

Compact Node Title will force nodes to display as Compact Nodes if set. This information is also appended to the Keywords, no need to set in the Keywords as well.

Keywords makes nodes easier to find.

diff --git a/Engine/Source/Editor/BlueprintGraph/Classes/K2Node_EditablePinBase.h b/Engine/Source/Editor/BlueprintGraph/Classes/K2Node_EditablePinBase.h
--- a/Engine/Source/Editor/BlueprintGraph/Classes/K2Node_EditablePinBase.h
+++ b/Engine/Source/Editor/BlueprintGraph/Classes/K2Node_EditablePinBase.h
@@ -65,6 +65,12 @@ struct FKismetUserDeclaredFunctionMetadata
 	UPROPERTY()
 	FText Category;
 
+	UPROPERTY()
+	FText Keywords;
+
+	UPROPERTY()
+	FText CompactNodeTitle;
+
 	UPROPERTY()
 	FLinearColor InstanceTitleColor;
 
diff --git a/Engine/Source/Editor/BlueprintGraph/Private/K2Node_MacroInstance.cpp b/Engine/Source/Editor/BlueprintGraph/Private/K2Node_MacroInstance.cpp
--- a/Engine/Source/Editor/BlueprintGraph/Private/K2Node_MacroInstance.cpp
+++ b/Engine/Source/Editor/BlueprintGraph/Private/K2Node_MacroInstance.cpp
@@ -138,13 +138,14 @@ FText UK2Node_MacroInstance::GetTooltipText() const
 
 FText UK2Node_MacroInstance::GetKeywords() const
 {
-	// This might need to be refined, but seems better than providing no keywords:
-	UBlueprint* MacroBP = MacroGraphReference.GetBlueprint();
-	if (MacroBP)
+	FText Keywords = GetAssociatedGraphMetadata(GetMacroGraph())->Keywords;
+
+	// If the Macro has Compact Node Title data, append the compact node title as a Keyword so it can be searched.
+	if (ShouldDrawCompact())
 	{
-		return FText::FromString(MacroBP->BlueprintDescription);
+		Keywords = FText::Format(FText::FromString(TEXT("{0} {1}")), Keywords, GetCompactNodeTitle());
 	}
-	return UK2Node_Tunnel::GetKeywords();
+	return Keywords;
 }
 
 FText UK2Node_MacroInstance::GetNodeTitle(ENodeTitleType::Type TitleType) const
@@ -362,52 +363,12 @@ FName UK2Node_MacroInstance::GetPaletteIcon(FLinearColor& OutColor) const
 
 FText UK2Node_MacroInstance::GetCompactNodeTitle() const
 {
-	// Special case handling for standard macros
-	// @TODO Change this to use metadata by allowing macros to specify CompactNodeTitle metadata
-	UEdGraph* MacroGraph = MacroGraphReference.GetGraph();
-	if(MacroGraph != nullptr && MacroGraph->GetOuter()->GetName() == TEXT("StandardMacros"))
-	{
-		FName MacroName = FName(*MacroGraph->GetName());
-		if(	MacroName == TEXT("IncrementFloat" ) ||
-			MacroName == TEXT("IncrementInt"))
-		{
-			return LOCTEXT("IncrementCompactNodeTitle", "++");
-		}
-		else if( MacroName == TEXT("DecrementFloat") || 
-				 MacroName == TEXT("DecrementInt"))
-		{
-			return LOCTEXT("DecrementCompactNodeTitle", "--");
-		}
-		else if( MacroName == TEXT("NegateFloat") || 
-				 MacroName == TEXT("NegateInt") )
-		{
-			return LOCTEXT("NegateCompactNodeTitle", "-");
-		}
-	}	
-
-	return Super::GetCompactNodeTitle();
+	return GetAssociatedGraphMetadata(GetMacroGraph())->CompactNodeTitle;	
 }
 
 bool UK2Node_MacroInstance::ShouldDrawCompact() const
 {
-	// Special case handling for standard macros
-	// @TODO Change this to use metadata by allowing macros to specify CompactNodeTitle metadata
-	UEdGraph* MacroGraph = MacroGraphReference.GetGraph();
-	if(MacroGraph != nullptr && MacroGraph->GetOuter()->GetName() == TEXT("StandardMacros"))
-	{
-		FName MacroName = FName(*MacroGraph->GetName());
-		if(	MacroName == TEXT("IncrementFloat" ) ||
-			MacroName == TEXT("IncrementInt") || 
-			MacroName == TEXT("DecrementFloat") || 
-			MacroName == TEXT("DecrementInt") || 
-			MacroName == TEXT("NegateFloat") || 
-			MacroName == TEXT("NegateInt") )
-		{
-			return true;
-		}
-	}
-
-	return Super::ShouldDrawCompact();
+	return !GetCompactNodeTitle().IsEmpty();
 }
 
 bool UK2Node_MacroInstance::CanPasteHere(const UEdGraph* TargetGraph) const
diff --git a/Engine/Source/Editor/Kismet/Private/BlueprintDetailsCustomization.cpp b/Engine/Source/Editor/Kismet/Private/BlueprintDetailsCustomization.cpp
--- a/Engine/Source/Editor/Kismet/Private/BlueprintDetailsCustomization.cpp
+++ b/Engine/Source/Editor/Kismet/Private/BlueprintDetailsCustomization.cpp
@@ -2341,8 +2341,8 @@ void FBlueprintGraphActionDetails::CustomizeDetails( IDetailLayoutBuilder& Detai
 				TSharedPtr<SComboButton> NewComboButton;
 				TSharedPtr<SListView<TSharedPtr<FText>>> NewListView;
 
-				const FString DocLink = TEXT("Shared/Editors/BlueprintEditor/VariableDetails");
-				TSharedPtr<SToolTip> CategoryTooltip = IDocumentation::Get()->CreateToolTip(LOCTEXT("EditCategoryName_Tooltip", "The category of the variable; editing this will place the variable into another category or create a new one."), NULL, DocLink, TEXT("Category"));
+				const FString DocLink = TEXT("Shared/Editors/BlueprintEditor/GraphDetails");
+				TSharedPtr<SToolTip> CategoryTooltip = IDocumentation::Get()->CreateToolTip(LOCTEXT("EditCategoryName_Tooltip", "The category of the graph; editing this will place the graph into another category or create a new one."), NULL, DocLink, TEXT("Category"));
 
 				Category.AddCustomRow( LOCTEXT("CategoryLabel", "Category") )
 					.NameContent()
@@ -2389,6 +2389,44 @@ void FBlueprintGraphActionDetails::CustomizeDetails( IDetailLayoutBuilder& Detai
 
 				CategoryComboButton = NewComboButton;
 				CategoryListView = NewListView;
+
+				TSharedPtr<SToolTip> KeywordsTooltip = IDocumentation::Get()->CreateToolTip(LOCTEXT("EditKeywords_Tooltip", "Keywords for searching for the function or macro."), NULL, DocLink, TEXT("Keywords"));
+				Category.AddCustomRow( LOCTEXT("KeywordsLabel", "Keywords") )
+					.NameContent()
+					[
+						SNew(STextBlock)
+						.Text( LOCTEXT("KeywordsLabel", "Keywords") )
+						.ToolTip(CategoryTooltip)
+						.Font( IDetailLayoutBuilder::GetDetailFont() )
+					]
+				.ValueContent()
+					[
+						SNew(SEditableTextBox)
+						.Text(this, &FBlueprintGraphActionDetails::OnGetKeywordsText)
+						.OnTextCommitted(this, &FBlueprintGraphActionDetails::OnKeywordsTextCommitted )
+						.ToolTip(KeywordsTooltip)
+						.RevertTextOnEscape(true)
+						.Font( IDetailLayoutBuilder::GetDetailFont() )
+					];
+
+				TSharedPtr<SToolTip> CompactNodeTitleTooltip = IDocumentation::Get()->CreateToolTip(LOCTEXT("EditCompactNodeTitle_Tooltip", "Sets the compact node title for calls to this function or macro. Compact node titles convert a node to display as a compact node and are used as a keyword for searching."), NULL, DocLink, TEXT("Compact Node Title"));
+				Category.AddCustomRow( LOCTEXT("CompactNodeTitleLabel", "Compact Node Title") )
+					.NameContent()
+					[
+						SNew(STextBlock)
+						.Text( LOCTEXT("CompactNodeTitleLabel", "Compact Node Title") )
+						.ToolTip(CategoryTooltip)
+						.Font( IDetailLayoutBuilder::GetDetailFont() )
+					]
+				.ValueContent()
+					[
+						SNew(SEditableTextBox)
+						.Text(this, &FBlueprintGraphActionDetails::OnGetCompactNodeTitleText)
+						.OnTextCommitted(this, &FBlueprintGraphActionDetails::OnCompactNodeTitleTextCommitted )
+						.ToolTip(CompactNodeTitleTooltip)
+						.RevertTextOnEscape(true)
+						.Font( IDetailLayoutBuilder::GetDetailFont() )
+					];
 			}
 
 			if (IsAccessSpecifierVisible())
@@ -3470,6 +3508,85 @@ TSharedRef< ITableRow > FBlueprintGraphActionDetails::MakeCategoryViewWidget( TS
 		];
 }
 
+FText FBlueprintGraphActionDetails::OnGetKeywordsText() const
+{
+	FText ResultKeywords;
+	if (FKismetUserDeclaredFunctionMetadata* Metadata = GetMetadataBlock())
+	{
+		ResultKeywords = Metadata->Keywords;
+	}
+	return ResultKeywords;
+}
+
+void FBlueprintGraphActionDetails::OnKeywordsTextCommitted(const FText& NewText, ETextCommit::Type InTextCommit)
+{
+	if (InTextCommit == ETextCommit::OnEnter || InTextCommit == ETextCommit::OnUserMovedFocus)
+	{
+		if (FKismetUserDeclaredFunctionMetadata* Metadata = GetMetadataBlock())
+		{
+			// Remove excess whitespace and prevent keywords with just spaces
+			FText Keywords = FText::TrimPrecedingAndTrailing(NewText);
+
+			if (!Keywords.EqualTo(Metadata->Keywords))
+			{
+				Metadata->Keywords = Keywords;
+
+				if(auto Function = FindFunction())
+				{
+					Function->Modify();
+					Function->SetMetaData(FBlueprintMetadata::MD_FunctionKeywords, *Keywords.ToString());
+				}
+				OnParamsChanged(GetFunctionEntryNode().Get(), true);
+				FBlueprintEditorUtils::MarkBlueprintAsModified(GetBlueprintObj());
+			}
+		}
+	}
+}
+
+FText FBlueprintGraphActionDetails::OnGetCompactNodeTitleText() const
+{
+	FText ResultKeywords;
+	if (FKismetUserDeclaredFunctionMetadata* Metadata = GetMetadataBlock())
+	{
+		ResultKeywords = Metadata->CompactNodeTitle;
+	}
+	return ResultKeywords;
+}
+
+void FBlueprintGraphActionDetails::OnCompactNodeTitleTextCommitted(const FText& NewText, ETextCommit::Type InTextCommit)
+{
+	if (InTextCommit == ETextCommit::OnEnter || InTextCommit == ETextCommit::OnUserMovedFocus)
+	{
+		if (FKismetUserDeclaredFunctionMetadata* Metadata = GetMetadataBlock())
+		{
+			// Remove excess whitespace and prevent CompactNodeTitle with just spaces
+			FText CompactNodeTitle = FText::TrimPrecedingAndTrailing(NewText);
+
+			if (!CompactNodeTitle.EqualTo(Metadata->CompactNodeTitle))
+			{
+				Metadata->CompactNodeTitle = CompactNodeTitle;
+
+				if(auto Function = FindFunction())
+				{
+					Function->Modify();
+
+					if (CompactNodeTitle.IsEmpty())
+					{
+						// Remove the metadata from the function, empty listings will still display the node as Compact
+						Function->RemoveMetaData(FBlueprintMetadata::MD_FunctionKeywords);
+					}
+					else
+					{
+						Function->SetMetaData(FBlueprintMetadata::MD_CompactNodeTitle, *CompactNodeTitle.ToString());
+					}
+				}
+				OnParamsChanged(GetFunctionEntryNode().Get(), true);
+				FBlueprintEditorUtils::MarkBlueprintAsModified(GetBlueprintObj());
+			}
+		}
+	}
+}
+
 FText FBlueprintGraphActionDetails::AccessSpecifierProperName( uint32 AccessSpecifierFlag ) const
 {
 	switch(AccessSpecifierFlag)
diff --git a/Engine/Source/Editor/Kismet/Private/BlueprintDetailsCustomization.h b/Engine/Source/Editor/Kismet/Private/BlueprintDetailsCustomization.h
--- a/Engine/Source/Editor/Kismet/Private/BlueprintDetailsCustomization.h
+++ b/Engine/Source/Editor/Kismet/Private/BlueprintDetailsCustomization.h
@@ -480,6 +480,12 @@ private:
 	
 	FText OnGetCategoryText() const;
 	void OnCategoryTextCommitted(const FText& NewText, ETextCommit::Type InTextCommit);
+
+	FText OnGetKeywordsText() const;
+	void OnKeywordsTextCommitted(const FText& NewText, ETextCommit::Type InTextCommit);
+
+	FText OnGetCompactNodeTitleText() const;
+	void OnCompactNodeTitleTextCommitted(const FText& NewText, ETextCommit::Type InTextCommit);
 	
 	FText AccessSpecifierProperName( uint32 AccessSpecifierFlag ) const;
 	bool IsAccessSpecifierVisible() const;
diff --git a/Engine/Source/Editor/KismetCompiler/Private/KismetCompiler.cpp b/Engine/Source/Editor/KismetCompiler/Private/KismetCompiler.cpp
--- a/Engine/Source/Editor/KismetCompiler/Private/KismetCompiler.cpp
+++ b/Engine/Source/Editor/KismetCompiler/Private/KismetCompiler.cpp
@@ -1220,6 +1220,19 @@ void FKismetCompilerContext::PrecompileFunction(FKismetFunctionContext& Context)
 		{
 			Context.Function->SetMetaData(FBlueprintMetadata::MD_FunctionCategory, *FunctionMetaData.Category.ToString());
 		}
+
+		// Set up the function keywords
+		if( !FunctionMetaData.Keywords.IsEmpty() )
+		{
+			Context.Function->SetMetaData(FBlueprintMetadata::MD_FunctionKeywords, *FunctionMetaData.Keywords.ToString());
+		}
+
+		// Set up the function compact node title
+		if( !FunctionMetaData.CompactNodeTitle.IsEmpty() )
+		{
+			Context.Function->SetMetaData(FBlueprintMetadata::MD_CompactNodeTitle, *FunctionMetaData.CompactNodeTitle.ToString());
+		}
+
 		// Set as blutility function
 		if( FunctionMetaData.bCallInEditor )
 		{
