When indexing all Blueprints, hitting "No" will correctly skip checking out all un-indexed Blueprints.

diff --git a/Engine/Source/Editor/Kismet/Private/FindInBlueprintManager.cpp b/Engine/Source/Editor/Kismet/Private/FindInBlueprintManager.cpp
--- a/Engine/Source/Editor/Kismet/Private/FindInBlueprintManager.cpp
+++ b/Engine/Source/Editor/Kismet/Private/FindInBlueprintManager.cpp
@@ -1392,12 +1392,12 @@ FString FFindInBlueprintSearchManager::ConvertFTextToHexString(FText InValue)
 	return BytesToHex(SerializedData.GetData(), SerializedData.Num());
 }
 
-void FFindInBlueprintSearchManager::OnCacheAllUncachedBlueprints(bool bInSourceControlActive)
+void FFindInBlueprintSearchManager::OnCacheAllUncachedBlueprints(bool bInSourceControlActive, bool bCheckoutAndSave)
 {
 	// Multiple threads can be adding to this at the same time
 	FScopeLock ScopeLock(&SafeModifyCacheCriticalSection);
 
-	if(ISourceControlModule::Get().IsEnabled())
+	if(ISourceControlModule::Get().IsEnabled() && bCheckoutAndSave)
 	{
 		FEditorFileUtils::CheckoutPackages(UncachedBlueprints);
 	}
@@ -1423,17 +1423,18 @@ void FFindInBlueprintSearchManager::CacheAllUncachedBlueprints(TWeakPtr< SFindIn
 			UncachedBlueprints.Append(FailedToCachePaths);
 			FailedToCachePaths.Empty();
 
-			CachingObject = new FCacheAllBlueprintsTickableObject(UncachedBlueprints, ReturnValue == EAppReturnType::Yes);
+			const bool bCheckoutAndSave = ReturnValue == EAppReturnType::Yes;
+			CachingObject = new FCacheAllBlueprintsTickableObject(UncachedBlueprints, bCheckoutAndSave);
 			OutActiveTimerDelegate.BindRaw(CachingObject, &FCacheAllBlueprintsTickableObject::Tick);
 
-			if(!ISourceControlModule::Get().IsEnabled() && ReturnValue == EAppReturnType::Yes)
+			if(!ISourceControlModule::Get().IsEnabled() && bCheckoutAndSave)
 			{
 				// Offer to start up Source Control
-				ISourceControlModule::Get().ShowLoginDialog(FSourceControlLoginClosed::CreateRaw(this, &FFindInBlueprintSearchManager::OnCacheAllUncachedBlueprints), ELoginWindowMode::Modeless, EOnLoginWindowStartup::PreserveProvider);
+				ISourceControlModule::Get().ShowLoginDialog(FSourceControlLoginClosed::CreateRaw(this, &FFindInBlueprintSearchManager::OnCacheAllUncachedBlueprints, bCheckoutAndSave), ELoginWindowMode::Modeless, EOnLoginWindowStartup::PreserveProvider);
 			}
 			else
 			{
-				OnCacheAllUncachedBlueprints(true);
+				OnCacheAllUncachedBlueprints(true, bCheckoutAndSave);
 			}
 
 			SourceCachingWidget = InSourceWidget;
diff --git a/Engine/Source/Editor/Kismet/Public/FindInBlueprintManager.h b/Engine/Source/Editor/Kismet/Public/FindInBlueprintManager.h
--- a/Engine/Source/Editor/Kismet/Public/FindInBlueprintManager.h
+++ b/Engine/Source/Editor/Kismet/Public/FindInBlueprintManager.h
@@ -175,7 +175,7 @@ public:
 	 * @param InSourceWidget		The source FindInBlueprints widget, this widget will be informed when caching is complete
 	 */
 	void CacheAllUncachedBlueprints(TWeakPtr< class SFindInBlueprints > InSourceWidgetm, FWidgetActiveTimerDelegate& OutActiveTimerDelegate);
-	void OnCacheAllUncachedBlueprints(bool bInSourceControlActive);
+	void OnCacheAllUncachedBlueprints(bool bInSourceControlActive, bool bCheckoutAndSave);
 
 	/** Stops the caching process where it currently is at, the rest can be continued later */
 	void CancelCacheAll(SFindInBlueprints* InFindInBlueprintWidget);
