[UE-15138] Fix interactive editing (and related issues) in the Blueprint Class editor details tab.

change summary:
- FBlueprintEditor::OnBlueprintChangeImpl() no longer calls DestroyPreview(); this can occur within the context of a transaction, and thus we should not delete or respawn the preview actor instance here.
- modified FBlueprintEditor::NotifyPostChange() to revert back to calling UpdateSCSPreview(true) only in the "interactive" edit mode case; this ensures that we reconstruct (not respawn) the preview actor instance during an interactive edit; but in all other cases we defer until the following viewport tick.
- modified FBlueprintEditor::UpdatePreviewActor() to call ReregisterAllComponents() prior to the RerunConstructionScripts() call; this was needed in order to ensure that scene proxies are recreated in order to pick up material changes, etc. on the preview actor instance.
- modified FBlueprintEditor::UpdateSCSPreview() to no longer call SSCSEditor::UpdateTree(); this fixes issues with "stuttering" slider, etc. interaction
- removed LastPreviewActor from FSCSEditorViewportClient as we no longer destroy the preview actor instance by default after every change in the base FBlueprintEditor logic
- modified FSCSEditorViewportClient::InvalidatePreview() to update the current preview actor instance via FBlueprintEditor::UpdatePreviewActor() rather than destroy the current instance, as well as to call the base Invalidate() and RefreshPreviewBounds() functions; this matches previous behavior prior to single-screen unification.
- modified FSCSEditorViewport::DeferredUpdatePreview() to no longer call Invalidate(); this has been moved into the VPC's InvalidatePreview() call as it needs to occur after the preview actor instance has been updated
- modified SSCSEditor::Tick() to extend the validation check logic to "BlueprintSCS" mode, so that UCS component display can be handled in the same way it is for "Instanced" mode when the editor setting is enabled. this replaces the previous solution of calling UpdateTree() after every UpdateSCSPreview() call in FBlueprintEditor.
- modified FWidgetBlueprintEditor::OnBlueprintChangedImpl() to call DestroyPreview() prior to calling the base FBlueprintEditor::OnBlueprintChangedImpl() logic; this keeps the UMG editor side functioning the same as it was previously (destroying its preview instance after each BP change notification).

notes:
- also addresses UE-15038

diff --git a/Engine/Source/Editor/Kismet/Private/BlueprintEditor.cpp b/Engine/Source/Editor/Kismet/Private/BlueprintEditor.cpp
--- a/Engine/Source/Editor/Kismet/Private/BlueprintEditor.cpp
+++ b/Engine/Source/Editor/Kismet/Private/BlueprintEditor.cpp
@@ -2939,8 +2939,6 @@ void FBlueprintEditor::OnBlueprintChangedImpl(UBlueprint* InBlueprint, bool bIsJ
 {
 	if (InBlueprint)
 	{
-		DestroyPreview();
-
 		// Notify that the blueprint has been changed (update Content browser, etc)
 		InBlueprint->PostEditChange();
 
@@ -6674,8 +6672,8 @@ void FBlueprintEditor::NotifyPostChange(const FPropertyChangedEvent& PropertyCha
 		FBlueprintEditorUtils::PostEditChangeBlueprintActors(Blueprint);
 	}
 
-	// Force updates to occur immediately, so the change is inside a transaction scope. See: UE-11802
-	UpdateSCSPreview(true);
+	// Force updates to occur immediately during interactive mode (otherwise the preview won't refresh because it won't be ticking)
+	UpdateSCSPreview(PropertyChangedEvent.ChangeType == EPropertyChangeType::Interactive);
 }
 
 void FBlueprintEditor::OnFinishedChangingProperties(const FPropertyChangedEvent& PropertyChangedEvent)
@@ -7324,6 +7322,7 @@ void FBlueprintEditor::UpdatePreviewActor(UBlueprint* InBlueprint, bool bInForce
 	}
 	else if ( PreviewActor )
 	{
+		PreviewActor->ReregisterAllComponents();
 		PreviewActor->RerunConstructionScripts();
 	}
 
@@ -7332,12 +7331,6 @@ void FBlueprintEditor::UpdatePreviewActor(UBlueprint* InBlueprint, bool bInForce
 	{
 		InBlueprint->SimpleConstructionScript->EndEditorComponentConstruction();
 	}
-
-	// Update the tree to include any UCS-constructed components
-	if ( SCSEditor.IsValid() )
-	{
-		SCSEditor->UpdateTree();
-	}
 }
 
 void FBlueprintEditor::DestroyPreview()
@@ -7352,7 +7345,7 @@ void FBlueprintEditor::DestroyPreview()
 	if ( PreviewActor != nullptr )
 	{
 		check(PreviewScene.GetWorld());
-		PreviewScene.GetWorld()->EditorDestroyActor(PreviewActor, true);
+		PreviewScene.GetWorld()->EditorDestroyActor(PreviewActor, false);
 	}
 
 	UBlueprint* PreviewBlueprint = GetBlueprintObj();
diff --git a/Engine/Source/Editor/Kismet/Private/SCSEditorViewportClient.cpp b/Engine/Source/Editor/Kismet/Private/SCSEditorViewportClient.cpp
--- a/Engine/Source/Editor/Kismet/Private/SCSEditorViewportClient.cpp
+++ b/Engine/Source/Editor/Kismet/Private/SCSEditorViewportClient.cpp
@@ -183,14 +183,6 @@ void FSCSEditorViewportClient::Tick(float DeltaSeconds)
 		InvalidatePreview(false);
 	}
 
-	if ( PreviewActor != LastPreviewActor.Get() || PreviewActor == nullptr || IsRealtime() )
-	{
-		LastPreviewActor = PreviewActor;
-
-		Invalidate();
-		RefreshPreviewBounds();
-	}
-
 	// Tick the preview scene world.
 	if (!GIntraFrameDebuggingGameThread)
 	{
@@ -661,8 +653,16 @@ void FSCSEditorViewportClient::InvalidatePreview(bool bResetCamera)
 		return;
 	}
 
-	// Destroy the preview
-	BlueprintEditorPtr.Pin()->DestroyPreview();
+	UBlueprint* Blueprint = BlueprintEditorPtr.Pin()->GetBlueprintObj();
+	check(Blueprint);
+
+	const bool bIsPreviewActorValid = GetPreviewActor() != nullptr;
+
+	// Create or update the Blueprint actor instance in the preview scene
+	BlueprintEditorPtr.Pin()->UpdatePreviewActor(Blueprint, !bIsPreviewActorValid);
+
+	Invalidate();
+	RefreshPreviewBounds();
 	
 	if( bResetCamera )
 	{
@@ -755,7 +755,6 @@ void FSCSEditorViewportClient::ToggleIsSimulateEnabled()
 	PreviewScene->GetWorld()->bBegunPlay = bIsSimulateEnabled;
 	PreviewScene->GetWorld()->bShouldSimulatePhysics = bIsSimulateEnabled;
 
-	AActor* PreviewActor = GetPreviewActor();
 	TSharedPtr<SWidget> SCSEditor = BlueprintEditorPtr.Pin()->GetSCSEditor();
 	TSharedRef<SWidget> Inspector = BlueprintEditorPtr.Pin()->GetInspector();
 
diff --git a/Engine/Source/Editor/Kismet/Private/SCSEditorViewportClient.h b/Engine/Source/Editor/Kismet/Private/SCSEditorViewportClient.h
--- a/Engine/Source/Editor/Kismet/Private/SCSEditorViewportClient.h
+++ b/Engine/Source/Editor/Kismet/Private/SCSEditorViewportClient.h
@@ -131,9 +131,6 @@ private:
 	/** Weak reference to the editor hosting the viewport */
 	TWeakPtr<class FBlueprintEditor> BlueprintEditorPtr;
 
-	/** The preview actor from last tick so we can tell if it changed. */
-	TWeakObjectPtr<AActor> LastPreviewActor;
-
 	/** The full bounds of the preview scene (encompasses all visible components) */
 	FBoxSphereBounds PreviewActorBounds;
 
diff --git a/Engine/Source/Editor/Kismet/Private/SSCSEditor.cpp b/Engine/Source/Editor/Kismet/Private/SSCSEditor.cpp
--- a/Engine/Source/Editor/Kismet/Private/SSCSEditor.cpp
+++ b/Engine/Source/Editor/Kismet/Private/SSCSEditor.cpp
@@ -3519,60 +3519,60 @@ void SSCSEditor::Tick( const FGeometry& AllottedGeometry, const double InCurrent
 {
 	SCompoundWidget::Tick(AllottedGeometry, InCurrentTime, InDeltaTime);
 
-	if (EditorMode == EComponentEditorMode::ActorInstance)
+	TFunction<bool (const TArray<FSCSEditorTreeNodePtrType>&, int32&)> AreAnyNodesInvalidLambda = [&](const TArray<FSCSEditorTreeNodePtrType>& InNodes, int32& OutNumValidNodes) -> bool
 	{
-		TFunction<bool (const TArray<FSCSEditorTreeNodePtrType>&, int32&)> AreAnyNodesInvalidLambda = [&](const TArray<FSCSEditorTreeNodePtrType>& InNodes, int32& OutNumValidNodes) -> bool
+		bool bFoundInvalidNode = false;
+		for(auto NodeIt = InNodes.CreateConstIterator(); NodeIt && !bFoundInvalidNode && !AreAnyNodesInvalidLambda((*NodeIt)->GetChildren(), OutNumValidNodes); ++NodeIt)
 		{
-			bool bFoundInvalidNode = false;
-			for(auto NodeIt = InNodes.CreateConstIterator(); NodeIt && !bFoundInvalidNode && !AreAnyNodesInvalidLambda((*NodeIt)->GetChildren(), OutNumValidNodes); ++NodeIt)
+			const FSCSEditorTreeNodePtrType Node = *NodeIt;
+			bool bIsComponent = Node->GetNodeType() == FSCSEditorTreeNode::ComponentNode;
+			if (bIsComponent)
 			{
-				bool bIsComponent = (*NodeIt)->GetNodeType() == FSCSEditorTreeNode::ComponentNode;
-				if (bIsComponent)
-				{
-					// A component node is considered to be invalid if one of the following conditions holds true:
-					//	(a) The component instance is invalid (i.e. NULL).
-					//	(b) The component instance is pending garbage collection.
-					//	(c) The component instance is a USceneComponent derivative and either is not the scene "root" node or otherwise is not contained within the scene component node hierarchy.
-					//	(d) The component instance is not a USceneComponent derivative and either matches the scene "root" node or is otherwise contained within the scene component node hierarchy.
-					//		=> (c) and (d) could happen if the user changes the parent class of a custom component's Blueprint class, for example.
-					//
-					FSCSEditorTreeNodePtrType ParentNodePtr = (*NodeIt)->GetParent();
-					const UActorComponent* InstancedComponent = (*NodeIt)->GetComponentTemplate();
-					bFoundInvalidNode = !InstancedComponent
-						|| InstancedComponent->IsPendingKill()
-						|| (InstancedComponent->IsA<USceneComponent>() && (*NodeIt != SceneRootNodePtr && (!ParentNodePtr.IsValid() || ParentNodePtr->GetNodeType() != FSCSEditorTreeNode::ComponentNode)))
-						|| (!InstancedComponent->IsA<USceneComponent>() && (*NodeIt == SceneRootNodePtr || (ParentNodePtr.IsValid() && ParentNodePtr->GetNodeType() == FSCSEditorTreeNode::ComponentNode)));
+				// A component node is considered to be invalid if one of the following conditions holds true:
+				//	(a) The component instance is invalid (i.e. NULL).
+				//	(b) The component instance is pending garbage collection.
+				//	(c) The component instance is a USceneComponent derivative and either is not the scene "root" node or otherwise is not contained within the scene component node hierarchy.
+				//	(d) The component instance is not a USceneComponent derivative and either matches the scene "root" node or is otherwise contained within the scene component node hierarchy.
+				//		=> (c) and (d) could happen if the user changes the parent class of a custom component's Blueprint class, for example.
+				//
+				FSCSEditorTreeNodePtrType ParentNodePtr = Node->GetParent();
+				const UActorComponent* InstancedComponent = Node->GetComponentTemplate();
+				bFoundInvalidNode = !InstancedComponent
+					|| InstancedComponent->IsPendingKill()
+					|| (InstancedComponent->IsA<USceneComponent>() && (Node != SceneRootNodePtr && (!ParentNodePtr.IsValid() || ParentNodePtr->GetNodeType() != FSCSEditorTreeNode::ComponentNode)))
+					|| (!InstancedComponent->IsA<USceneComponent>() && (Node == SceneRootNodePtr || (ParentNodePtr.IsValid() && ParentNodePtr->GetNodeType() == FSCSEditorTreeNode::ComponentNode)));
 					
-					if(!bFoundInvalidNode)
-					{
-						++OutNumValidNodes;
-					}
+				if(!bFoundInvalidNode)
+				{
+					++OutNumValidNodes;
 				}
 			}
+		}
 
-			return bFoundInvalidNode;
-		};
+		return bFoundInvalidNode;
+	};
 
-		if (const AActor* ActorInstance = GetActorContext())
+	const AActor* ActorInstance = (EditorMode == EComponentEditorMode::ActorInstance) ? GetActorContext() : PreviewActor.Get(nullptr);
+	if (ActorInstance)
+	{
+		int32 NumComponentInstances = 0;
+		for (auto CompIt = ActorInstance->GetComponents().CreateConstIterator(); CompIt; ++CompIt)
 		{
-			int32 NumComponentInstances = 0;
-			for (auto CompIt = ActorInstance->GetComponents().CreateConstIterator(); CompIt; ++CompIt)
+			// Don't count editor-only components in instanced mode, because we don't show them. Also, don't count UCS-added components if the option to hide them is enabled.
+			const UActorComponent* CompInst = *CompIt;
+			if ((!CompInst->IsEditorOnly() || (EditorMode != EComponentEditorMode::ActorInstance && CompInst->GetArchetype() != CompInst->GetClass()->GetDefaultObject()))
+				&& (CompInst->CreationMethod != EComponentCreationMethod::UserConstructionScript || !GetDefault<UBlueprintEditorSettings>()->bHideConstructionScriptComponentsInDetailsView))
 			{
-				// Don't count editor-only components, because we don't show them. Also, don't count UCS-added components if the option to hide them is enabled.
-				if (!(*CompIt)->IsEditorOnly()
-					&& ((*CompIt)->CreationMethod != EComponentCreationMethod::UserConstructionScript || !GetDefault<UBlueprintEditorSettings>()->bHideConstructionScriptComponentsInDetailsView))
-				{
-					++NumComponentInstances;
-				}
+				++NumComponentInstances;
 			}
+		}
 
-			int32 NumComponentNodes = 0;
-			if (AreAnyNodesInvalidLambda(GetRootNodes(), NumComponentNodes) || NumComponentNodes != NumComponentInstances)
-			{
-				UE_LOG(LogSCSEditor, Log, TEXT("Calling UpdateTree() from Tick()."));
+		int32 NumComponentNodes = 0;
+		if (AreAnyNodesInvalidLambda(GetRootNodes(), NumComponentNodes) || NumComponentNodes != NumComponentInstances)
+		{
+			UE_LOG(LogSCSEditor, Log, TEXT("Calling UpdateTree() from Tick()."));
 
-				UpdateTree();
-			}
+			UpdateTree();
 		}
 	}
 }
diff --git a/Engine/Source/Editor/Kismet/Private/SSCSEditorViewport.cpp b/Engine/Source/Editor/Kismet/Private/SSCSEditorViewport.cpp
--- a/Engine/Source/Editor/Kismet/Private/SSCSEditorViewport.cpp
+++ b/Engine/Source/Editor/Kismet/Private/SSCSEditorViewport.cpp
@@ -442,8 +442,6 @@ void SSCSEditorViewport::RequestRefresh(bool bResetCamera, bool bRefreshNow)
 {
 	if(bRefreshNow)
 	{
-		Invalidate();
-
 		if(ViewportClient.IsValid())
 		{
 			ViewportClient->InvalidatePreview(bResetCamera);
@@ -496,8 +494,6 @@ EActiveTimerReturnType SSCSEditorViewport::DeferredUpdatePreview(double InCurren
 {
 	if (ViewportClient.IsValid())
 	{
-		Invalidate();
-
 		ViewportClient->InvalidatePreview(bResetCamera);
 	}
 
diff --git a/Engine/Source/Editor/UMGEditor/Private/WidgetBlueprintEditor.cpp b/Engine/Source/Editor/UMGEditor/Private/WidgetBlueprintEditor.cpp
--- a/Engine/Source/Editor/UMGEditor/Private/WidgetBlueprintEditor.cpp
+++ b/Engine/Source/Editor/UMGEditor/Private/WidgetBlueprintEditor.cpp
@@ -225,6 +225,8 @@ void FWidgetBlueprintEditor::InvalidatePreview()
 
 void FWidgetBlueprintEditor::OnBlueprintChangedImpl(UBlueprint* InBlueprint, bool bIsJustBeingCompiled )
 {
+	DestroyPreview();
+
 	FBlueprintEditor::OnBlueprintChangedImpl(InBlueprint, bIsJustBeingCompiled);
 
 	if ( InBlueprint )
