The details view in Blueprints of non-member variables (either from a parent class or an external member reference) will no longer show various detail options as editable.

Fixed several cases of showing completely wrong details info as well, where the property info would attempt to be looked up from the current Blueprint instead of the actual Blueprint.

diff --git a/Engine/Source/Editor/Kismet/Private/BlueprintDetailsCustomization.cpp b/Engine/Source/Editor/Kismet/Private/BlueprintDetailsCustomization.cpp
--- a/Engine/Source/Editor/Kismet/Private/BlueprintDetailsCustomization.cpp
+++ b/Engine/Source/Editor/Kismet/Private/BlueprintDetailsCustomization.cpp
@@ -203,6 +203,12 @@ void FBlueprintVarActionDetails::CustomizeDetails( IDetailLayoutBuilder& DetailL
 
 	UProperty* VariableProperty = CachedVariableProperty.Get();
 
+	// Cache the Blueprint which owns this VariableProperty
+	if (UBlueprintGeneratedClass* GeneratedClass = Cast<UBlueprintGeneratedClass>(VariableProperty->GetOuter()))
+	{
+		PropertyOwnerBlueprint = Cast<UBlueprint>(GeneratedClass->ClassGeneratedBy);
+	}
+
 	const UEdGraphSchema_K2* Schema = GetDefault<UEdGraphSchema_K2>();
 	
 	IDetailCategoryBuilder& Category = DetailLayout.EditCategory("Variable", LOCTEXT("VariableDetailsCategory", "Variable"));
@@ -270,6 +276,7 @@ void FBlueprintVarActionDetails::CustomizeDetails( IDetailLayoutBuilder& DetailL
 		SNew(SCheckBox)
 		.IsChecked( this, &FBlueprintVarActionDetails::OnEditableCheckboxState )
 		.OnCheckStateChanged( this, &FBlueprintVarActionDetails::OnEditableChanged )
+		.IsEnabled(IsVariableInBlueprint())
 		.ToolTip(EditableTooltip)
 	];
 
@@ -290,6 +297,7 @@ void FBlueprintVarActionDetails::CustomizeDetails( IDetailLayoutBuilder& DetailL
 		.Text( this, &FBlueprintVarActionDetails::OnGetTooltipText )
 		.ToolTip(ToolTipTooltip)
 		.OnTextCommitted( this, &FBlueprintVarActionDetails::OnTooltipTextCommitted, CachedVariableName )
+		.IsEnabled(IsVariableInBlueprint())
 		.Font( DetailFontInfo )
 	];
 
@@ -310,7 +318,7 @@ void FBlueprintVarActionDetails::CustomizeDetails( IDetailLayoutBuilder& DetailL
 		SNew(SCheckBox)
 		.IsChecked( this, &FBlueprintVarActionDetails::OnCreateWidgetCheckboxState )
 		.OnCheckStateChanged( this, &FBlueprintVarActionDetails::OnCreateWidgetChanged )
-		.IsEnabled(Is3DWidgetEnabled())
+		.IsEnabled(Is3DWidgetEnabled() && IsVariableInBlueprint())
 		.ToolTip(Widget3DTooltip)
 	];
 
@@ -330,6 +338,7 @@ void FBlueprintVarActionDetails::CustomizeDetails( IDetailLayoutBuilder& DetailL
 		SNew(SCheckBox)
 		.IsChecked( this, &FBlueprintVarActionDetails::OnGetExposedToSpawnCheckboxState )
 		.OnCheckStateChanged( this, &FBlueprintVarActionDetails::OnExposedToSpawnChanged )
+		.IsEnabled(IsVariableInBlueprint())
 		.ToolTip(ExposeOnSpawnTooltip)
 	];
 
@@ -349,6 +358,7 @@ void FBlueprintVarActionDetails::CustomizeDetails( IDetailLayoutBuilder& DetailL
 		SNew(SCheckBox)
 		.IsChecked( this, &FBlueprintVarActionDetails::OnGetPrivateCheckboxState )
 		.OnCheckStateChanged( this, &FBlueprintVarActionDetails::OnPrivateChanged )
+		.IsEnabled(IsVariableInBlueprint())
 		.ToolTip(PrivateTooltip)
 	];
 
@@ -368,6 +378,7 @@ void FBlueprintVarActionDetails::CustomizeDetails( IDetailLayoutBuilder& DetailL
 		SNew(SCheckBox)
 		.IsChecked( this, &FBlueprintVarActionDetails::OnGetExposedToMatineeCheckboxState )
 		.OnCheckStateChanged( this, &FBlueprintVarActionDetails::OnExposedToMatineeChanged )
+		.IsEnabled(IsVariableInBlueprint())
 		.ToolTip(ExposeToMatineeTooltip)
 	];
 
@@ -398,6 +409,7 @@ void FBlueprintVarActionDetails::CustomizeDetails( IDetailLayoutBuilder& DetailL
 		.ToolTip( ExposeToConfigTooltip )
 		.IsChecked( this, &FBlueprintVarActionDetails::OnGetConfigVariableCheckboxState )
 		.OnCheckStateChanged( this, &FBlueprintVarActionDetails::OnSetConfigVariableState )
+		.IsEnabled(IsVariableInBlueprint())
 	];
 
 	PopulateCategories(MyBlueprint.Pin().Get(), CategorySource);
@@ -407,7 +419,8 @@ void FBlueprintVarActionDetails::CustomizeDetails( IDetailLayoutBuilder& DetailL
 	TSharedPtr<SToolTip> CategoryTooltip = IDocumentation::Get()->CreateToolTip(LOCTEXT("EditCategoryName_Tooltip", "The category of the variable; editing this will place the variable into another category or create a new one."), NULL, DocLink, TEXT("Category"));
 
 	Category.AddCustomRow( LOCTEXT("CategoryLabel", "Category") )
-	.NameContent()
+		.Visibility(GetPropertyOwnerBlueprint()? EVisibility::Visible : EVisibility::Hidden)
+		.NameContent()
 	[
 		SNew(STextBlock)
 		.Text( LOCTEXT("CategoryLabel", "Category") )
@@ -476,6 +489,7 @@ void FBlueprintVarActionDetails::CustomizeDetails( IDetailLayoutBuilder& DetailL
 			SNew(SEditableTextBox)
 			.Text(this, &FBlueprintVarActionDetails::OnGetMetaKeyValue, UIMin)
 			.OnTextCommitted(this, &FBlueprintVarActionDetails::OnMetaKeyValueChanged, UIMin)
+			.IsEnabled(IsVariableInBlueprint())
 			.Font( DetailFontInfo )
 		]
 		+SHorizontalBox::Slot()
@@ -491,6 +505,7 @@ void FBlueprintVarActionDetails::CustomizeDetails( IDetailLayoutBuilder& DetailL
 			SNew(SEditableTextBox)
 			.Text(this, &FBlueprintVarActionDetails::OnGetMetaKeyValue, UIMax)
 			.OnTextCommitted(this, &FBlueprintVarActionDetails::OnMetaKeyValueChanged, UIMax)
+			.IsEnabled(IsVariableInBlueprint())
 			.Font( DetailFontInfo )
 		]
 	];
@@ -518,6 +533,7 @@ void FBlueprintVarActionDetails::CustomizeDetails( IDetailLayoutBuilder& DetailL
 			SNew(SEditableTextBox)
 			.Text(this, &FBlueprintVarActionDetails::OnGetMetaKeyValue, ClampMin)
 			.OnTextCommitted(this, &FBlueprintVarActionDetails::OnMetaKeyValueChanged, ClampMin)
+			.IsEnabled(IsVariableInBlueprint())
 			.Font(DetailFontInfo)
 		]
 		+ SHorizontalBox::Slot()
@@ -533,6 +549,7 @@ void FBlueprintVarActionDetails::CustomizeDetails( IDetailLayoutBuilder& DetailL
 			SNew(SEditableTextBox)
 			.Text(this, &FBlueprintVarActionDetails::OnGetMetaKeyValue, ClampMax)
 			.OnTextCommitted(this, &FBlueprintVarActionDetails::OnMetaKeyValueChanged, ClampMax)
+			.IsEnabled(IsVariableInBlueprint())
 			.Font(DetailFontInfo)
 		]
 	];
@@ -544,7 +561,7 @@ void FBlueprintVarActionDetails::CustomizeDetails( IDetailLayoutBuilder& DetailL
 	TSharedPtr<SToolTip> ReplicationTooltip = IDocumentation::Get()->CreateToolTip(LOCTEXT("VariableReplicate_Tooltip", "Should this Variable be replicated over the network?"), NULL, DocLink, TEXT("Replication"));
 
 	Category.AddCustomRow( LOCTEXT("VariableReplicationLabel", "Replication") )
-	.Visibility(TAttribute<EVisibility>(this, &FBlueprintVarActionDetails::ReplicationVisibility))
+		.Visibility(TAttribute<EVisibility>(this, &FBlueprintVarActionDetails::ReplicationVisibility))
 	.NameContent()
 	[
 		SNew(STextBlock)
@@ -558,6 +575,7 @@ void FBlueprintVarActionDetails::CustomizeDetails( IDetailLayoutBuilder& DetailL
 		.OptionsSource( &ReplicationOptions )
 		.InitiallySelectedItem(GetVariableReplicationType())
 		.OnSelectionChanged( this, &FBlueprintVarActionDetails::OnChangeReplication )
+		.IsEnabled(IsVariableInBlueprint())
 		.ToolTip(ReplicationTooltip)
 	];
 
@@ -572,7 +590,7 @@ void FBlueprintVarActionDetails::CustomizeDetails( IDetailLayoutBuilder& DetailL
 	// Add in default value editing for properties that can be edited, local properties cannot be edited
 	if ((Blueprint != NULL) && (Blueprint->GeneratedClass != NULL))
 	{
-		if (VariableProperty != NULL)
+		if (VariableProperty != NULL && IsVariableInBlueprint())
 		{
 			const UProperty* OriginalProperty = nullptr;
 			
@@ -696,12 +714,13 @@ void FBlueprintVarActionDetails::CustomizeDetails( IDetailLayoutBuilder& DetailL
 
 				IDetailPropertyRow* Row = DefaultValueCategory.AddExternalProperty(StructData, VariableProperty->GetFName());
 			}
-			else
+			else if (GetPropertyOwnerBlueprint())
 			{
 				// Things are in order, show the property and allow it to be edited
 				TArray<UObject*> ObjectList;
-				ObjectList.Add(Blueprint->GeneratedClass->GetDefaultObject());
+				ObjectList.Add(GetPropertyOwnerBlueprint()->GeneratedClass->GetDefaultObject());
 				IDetailPropertyRow* Row = DefaultValueCategory.AddExternalProperty(ObjectList, VariableProperty->GetFName());
+				Row->IsEnabled(IsVariableInBlueprint());
 			}
 		}
 
@@ -721,6 +740,7 @@ void FBlueprintVarActionDetails::CustomizeDetails( IDetailLayoutBuilder& DetailL
 			SNew(SCheckBox)
 			.IsChecked( this, &FBlueprintVarActionDetails::OnGetTransientCheckboxState )
 			.OnCheckStateChanged( this, &FBlueprintVarActionDetails::OnTransientChanged )
+			.IsEnabled(IsVariableInBlueprint())
 			.ToolTip(TransientTooltip)
 		];
 
@@ -740,6 +760,7 @@ void FBlueprintVarActionDetails::CustomizeDetails( IDetailLayoutBuilder& DetailL
 			SNew(SCheckBox)
 			.IsChecked( this, &FBlueprintVarActionDetails::OnGetSaveGameCheckboxState )
 			.OnCheckStateChanged( this, &FBlueprintVarActionDetails::OnSaveGameChanged )
+			.IsEnabled(IsVariableInBlueprint())
 			.ToolTip(SaveGameTooltip)
 		];
 
@@ -848,7 +869,7 @@ bool FBlueprintVarActionDetails::GetVariableNameChangeEnabled() const
 	check(Blueprint != NULL);
 
 	UProperty* VariableProperty = CachedVariableProperty.Get();
-	if(VariableProperty != nullptr)
+	if(VariableProperty != nullptr && IsVariableInBlueprint())
 	{
 		if(FBlueprintEditorUtils::FindNewVariableIndex(Blueprint, CachedVariableName) != INDEX_NONE)
 		{
@@ -975,7 +996,7 @@ void FBlueprintVarActionDetails::OnVarNameCommitted(const FText& InNewText, ETex
 bool FBlueprintVarActionDetails::GetVariableTypeChangeEnabled() const
 {
 	UProperty* VariableProperty = CachedVariableProperty.Get();
-	if(VariableProperty && !IsALocalVariable(VariableProperty))
+	if(VariableProperty && IsVariableInBlueprint() && !IsALocalVariable(VariableProperty))
 	{
 		if(GetBlueprintObj()->SkeletonGeneratedClass->GetAuthoritativeClass() != VariableProperty->GetOwnerClass()->GetAuthoritativeClass())
 		{
@@ -986,14 +1007,15 @@ bool FBlueprintVarActionDetails::GetVariableTypeChangeEnabled() const
 		{
 			return false;
 		}
+		return true;
 	}
-	return true;
+	return false;
 }
 
 bool FBlueprintVarActionDetails::GetVariableCategoryChangeEnabled() const
 {
 	UProperty* VariableProperty = CachedVariableProperty.Get();
-	if(VariableProperty)
+	if(VariableProperty && IsVariableInBlueprint())
 	{
 		if(UClass* VarSourceClass = Cast<UClass>(VariableProperty->GetOuter()))
 		{
@@ -1055,7 +1077,7 @@ FText FBlueprintVarActionDetails::OnGetTooltipText() const
 	if (VarName != NAME_None)
 	{
 		FString Result;
-		FBlueprintEditorUtils::GetBlueprintVariableMetaData(GetBlueprintObj(), VarName, GetLocalVariableScope(CachedVariableProperty.Get()), TEXT("tooltip"), Result);
+		FBlueprintEditorUtils::GetBlueprintVariableMetaData(GetPropertyOwnerBlueprint(), VarName, GetLocalVariableScope(CachedVariableProperty.Get()), TEXT("tooltip"), Result);
 		return FText::FromString(Result);
 	}
 	return FText();
@@ -1275,10 +1297,10 @@ FText FBlueprintVarActionDetails::OnGetCategoryText() const
 	{
 		const UEdGraphSchema_K2* K2Schema = GetDefault<UEdGraphSchema_K2>();
 
-		FText Category = FBlueprintEditorUtils::GetBlueprintVariableCategory(GetBlueprintObj(), VarName, GetLocalVariableScope(CachedVariableProperty.Get()));
+		FText Category = FBlueprintEditorUtils::GetBlueprintVariableCategory(GetPropertyOwnerBlueprint(), VarName, GetLocalVariableScope(CachedVariableProperty.Get()));
 
 		// Older blueprints will have their name as the default category and whenever it is the same as the default category, display localized text
-		if( Category.EqualTo(FText::FromString(GetBlueprintObj()->GetName())) || Category.EqualTo(K2Schema->VR_DefaultCategory) )
+		if( Category.EqualTo(FText::FromString(GetPropertyOwnerBlueprint()->GetName())) || Category.EqualTo(K2Schema->VR_DefaultCategory) )
 		{
 			return K2Schema->VR_DefaultCategory;
 		}
@@ -1330,7 +1352,7 @@ void FBlueprintVarActionDetails::OnCategorySelectionChanged( TSharedPtr<FText> P
 EVisibility FBlueprintVarActionDetails::ShowEditableCheckboxVisibilty() const
 {
 	UProperty* VariableProperty = CachedVariableProperty.Get();
-	if (VariableProperty)
+	if (VariableProperty && GetPropertyOwnerBlueprint())
 	{
 		if (IsABlueprintVariable(VariableProperty) && !IsAComponentVariable(VariableProperty))
 		{
@@ -1392,7 +1414,7 @@ void FBlueprintVarActionDetails::OnCreateWidgetChanged(ECheckBoxState InNewState
 EVisibility FBlueprintVarActionDetails::Show3DWidgetVisibility() const
 {
 	UProperty* VariableProperty = CachedVariableProperty.Get();
-	if (VariableProperty)
+	if (VariableProperty && GetPropertyOwnerBlueprint())
 	{
 		if (IsABlueprintVariable(VariableProperty) && FEdMode::CanCreateWidgetForProperty(VariableProperty))
 		{
@@ -1442,7 +1464,7 @@ void FBlueprintVarActionDetails::OnExposedToSpawnChanged(ECheckBoxState InNewSta
 EVisibility FBlueprintVarActionDetails::ExposeOnSpawnVisibility() const
 {
 	UProperty* VariableProperty = CachedVariableProperty.Get();
-	if (VariableProperty)
+	if (VariableProperty && GetPropertyOwnerBlueprint())
 	{
 		const UEdGraphSchema_K2* K2Schema = GetDefault<UEdGraphSchema_K2>();
 		FEdGraphPinType VariablePinType;
@@ -1487,7 +1509,7 @@ void FBlueprintVarActionDetails::OnPrivateChanged(ECheckBoxState InNewState)
 EVisibility FBlueprintVarActionDetails::ExposePrivateVisibility() const
 {
 	UProperty* Property = CachedVariableProperty.Get();
-	if (Property)
+	if (Property && GetPropertyOwnerBlueprint())
 	{
 		if (IsABlueprintVariable(Property) && !IsAComponentVariable(Property))
 		{
@@ -1542,7 +1564,7 @@ EVisibility FBlueprintVarActionDetails::ExposeToMatineeVisibility() const
 
 ECheckBoxState FBlueprintVarActionDetails::OnGetConfigVariableCheckboxState() const
 {
-	UBlueprint* Blueprint = GetBlueprintObj();
+	UBlueprint* Blueprint = GetPropertyOwnerBlueprint();
 	const FName VarName = CachedVariableName;
 	ECheckBoxState CheckboxValue = ECheckBoxState::Unchecked;
 
@@ -1601,7 +1623,7 @@ FText FBlueprintVarActionDetails::OnGetMetaKeyValue(FName Key) const
 	if (VarName != NAME_None)
 	{
 		FString Result;
-		FBlueprintEditorUtils::GetBlueprintVariableMetaData(GetBlueprintObj(), VarName, GetLocalVariableScope(CachedVariableProperty.Get()), Key, /*out*/ Result);
+		FBlueprintEditorUtils::GetBlueprintVariableMetaData(GetPropertyOwnerBlueprint(), VarName, GetLocalVariableScope(CachedVariableProperty.Get()), Key, /*out*/ Result);
 
 		return FText::FromString(Result);
 	}
@@ -1644,20 +1666,20 @@ TSharedPtr<FString> FBlueprintVarActionDetails::GetVariableReplicationType() con
 	uint64 PropFlags = 0;
 	UProperty* VariableProperty = CachedVariableProperty.Get();
 
-	if (VariableProperty)
+	if (VariableProperty && IsVariableInBlueprint())
 	{
-		uint64 *PropFlagPtr = FBlueprintEditorUtils::GetBlueprintVariablePropertyFlags(GetBlueprintObj(), VariableProperty->GetFName());
+		uint64 *PropFlagPtr = FBlueprintEditorUtils::GetBlueprintVariablePropertyFlags(GetPropertyOwnerBlueprint(), VariableProperty->GetFName());
 		
 		if (PropFlagPtr != NULL)
 		{
 			PropFlags = *PropFlagPtr;
 			bool IsReplicated = (PropFlags & CPF_Net) > 0;
-			bool bHasRepNotify = FBlueprintEditorUtils::GetBlueprintVariableRepNotifyFunc(GetBlueprintObj(), VariableProperty->GetFName()) != NAME_None;
+			bool bHasRepNotify = FBlueprintEditorUtils::GetBlueprintVariableRepNotifyFunc(GetPropertyOwnerBlueprint(), VariableProperty->GetFName()) != NAME_None;
 			if (bHasRepNotify)
 			{
 				// Verify they actually have a valid rep notify function still
-				UClass* GenClass = GetBlueprintObj()->SkeletonGeneratedClass;
-				UFunction* OnRepFunc = GenClass->FindFunctionByName(FBlueprintEditorUtils::GetBlueprintVariableRepNotifyFunc(GetBlueprintObj(), VariableProperty->GetFName()));
+				UClass* GenClass = GetPropertyOwnerBlueprint()->SkeletonGeneratedClass;
+				UFunction* OnRepFunc = GenClass->FindFunctionByName(FBlueprintEditorUtils::GetBlueprintVariableRepNotifyFunc(GetPropertyOwnerBlueprint(), VariableProperty->GetFName()));
 				if( OnRepFunc == NULL || OnRepFunc->NumParms != 0 || OnRepFunc->GetReturnProperty() != NULL )
 				{
 					bHasRepNotify = false;
diff --git a/Engine/Source/Editor/Kismet/Private/BlueprintDetailsCustomization.h b/Engine/Source/Editor/Kismet/Private/BlueprintDetailsCustomization.h
--- a/Engine/Source/Editor/Kismet/Private/BlueprintDetailsCustomization.h
+++ b/Engine/Source/Editor/Kismet/Private/BlueprintDetailsCustomization.h
@@ -181,6 +181,11 @@ private:
 	/** Refreshes cached data that changes after a Blueprint recompile */
 	void OnPostEditorRefresh();
 
+	/** Returns the Property's Blueprint */
+	UBlueprint* GetPropertyOwnerBlueprint() const { return PropertyOwnerBlueprint.Get(); }
+
+	/** Returns TRUE if the Variable is in the current Blueprint */
+	bool IsVariableInBlueprint() const { return GetPropertyOwnerBlueprint() == GetBlueprintObj(); }
 private:
 	/** Pointer back to my parent tab */
 	TWeakPtr<SMyBlueprint> MyBlueprint;
@@ -211,6 +216,9 @@ private:
 
 	/** Cached name for the variable we are affecting */
 	FName CachedVariableName;
+
+	/** Pointer back to the variable's Blueprint */
+	TWeakObjectPtr<UBlueprint> PropertyOwnerBlueprint;
 };
 
 class FBaseBlueprintGraphActionDetails : public IDetailCustomization
