Compiling a Blueprint while intermediate graphs are open will no longer cause a crash.

CleanInvalidTabs needs to be called whenever the editor is refreshed to make sure the tabs are still valid.

diff --git a/Engine/Source/Editor/Kismet/Private/BlueprintEditor.cpp b/Engine/Source/Editor/Kismet/Private/BlueprintEditor.cpp
--- a/Engine/Source/Editor/Kismet/Private/BlueprintEditor.cpp
+++ b/Engine/Source/Editor/Kismet/Private/BlueprintEditor.cpp
@@ -636,12 +636,13 @@ void FBlueprintEditor::RefreshEditors(ERefreshBlueprintEditorReason::Type Reason
 		StartEditingDefaults(/*bAutoFocus=*/ false, true);
 	}
 
+	// Remove any tabs are that are pending kill or otherwise invalid UObject pointers.
+	DocumentManager->CleanInvalidTabs();
+
 	//@TODO: Should determine when we need to do the invalid/refresh business and if the graph node selection change
 	// under non-compiles is necessary (except when the selection mode is appropriate, as already detected above)
 	if (Reason != ERefreshBlueprintEditorReason::BlueprintCompiled)
 	{
-		DocumentManager->CleanInvalidTabs();
-
 		DocumentManager->RefreshAllTabs();
 
 		bForceFocusOnSelectedNodes = true;
