Fixed a crash when canceling indexing all Blueprints where two cancel operations could be submitted due to lag in refreshing the widgets.

diff --git a/Engine/Source/Editor/Kismet/Private/FindInBlueprintManager.cpp b/Engine/Source/Editor/Kismet/Private/FindInBlueprintManager.cpp
--- a/Engine/Source/Editor/Kismet/Private/FindInBlueprintManager.cpp
+++ b/Engine/Source/Editor/Kismet/Private/FindInBlueprintManager.cpp
@@ -673,19 +673,22 @@ public:
 	/** Cancels caching and destroys this object */
 	void OnCancelCaching(bool bIsImmediate)
 	{
-		ProgressNotification.Pin()->SetText(LOCTEXT("BlueprintIndexCancelled", "Cancelled Indexing Blueprints!"));
+		if (!bIsCancelled)
+		{
+			ProgressNotification.Pin()->SetText(LOCTEXT("BlueprintIndexCancelled", "Cancelled Indexing Blueprints!"));
 
-		ProgressNotification.Pin()->SetCompletionState(SNotificationItem::CS_Fail);
-		ProgressNotification.Pin()->ExpireAndFadeout();
+			ProgressNotification.Pin()->SetCompletionState(SNotificationItem::CS_Fail);
+			ProgressNotification.Pin()->ExpireAndFadeout();
 
-		// Sometimes we can't wait another tick to shutdown, so make the callback immediately.
-		if (bIsImmediate)
-		{
-			FFindInBlueprintSearchManager::Get().FinishedCachingBlueprints(TickCacheIndex, FailedToCacheList);
-		}
-		else
-		{
-			bIsCancelled = true;
+			// Sometimes we can't wait another tick to shutdown, so make the callback immediately.
+			if (bIsImmediate)
+			{
+				FFindInBlueprintSearchManager::Get().FinishedCachingBlueprints(TickCacheIndex, FailedToCacheList);
+			}
+			else
+			{
+				bIsCancelled = true;
+			}
 		}
 	}
 
