Fix crash when dragging blueprint component to event graph

diff --git a/Engine/Source/Editor/Kismet/Private/BlueprintDetailsCustomization.cpp b/Engine/Source/Editor/Kismet/Private/BlueprintDetailsCustomization.cpp
--- a/Engine/Source/Editor/Kismet/Private/BlueprintDetailsCustomization.cpp
+++ b/Engine/Source/Editor/Kismet/Private/BlueprintDetailsCustomization.cpp
@@ -720,7 +720,10 @@ void FBlueprintVarActionDetails::CustomizeDetails( IDetailLayoutBuilder& DetailL
 				TArray<UObject*> ObjectList;
 				ObjectList.Add(GetPropertyOwnerBlueprint()->GeneratedClass->GetDefaultObject());
 				IDetailPropertyRow* Row = DefaultValueCategory.AddExternalProperty(ObjectList, VariableProperty->GetFName());
-				Row->IsEnabled(IsVariableInBlueprint());
+				if (Row != nullptr)
+				{
+					Row->IsEnabled(IsVariableInBlueprint());
+				}
 			}
 		}
 
