Fix a crash when selecting local variables in Blueprints.

diff --git a/Engine/Source/Editor/Kismet/Private/BlueprintDetailsCustomization.cpp b/Engine/Source/Editor/Kismet/Private/BlueprintDetailsCustomization.cpp
--- a/Engine/Source/Editor/Kismet/Private/BlueprintDetailsCustomization.cpp
+++ b/Engine/Source/Editor/Kismet/Private/BlueprintDetailsCustomization.cpp
@@ -204,7 +204,7 @@ void FBlueprintVarActionDetails::CustomizeDetails( IDetailLayoutBuilder& DetailL
 	UProperty* VariableProperty = CachedVariableProperty.Get();
 
 	// Cache the Blueprint which owns this VariableProperty
-	if (UBlueprintGeneratedClass* GeneratedClass = Cast<UBlueprintGeneratedClass>(VariableProperty->GetOuter()))
+	if (UBlueprintGeneratedClass* GeneratedClass = Cast<UBlueprintGeneratedClass>(VariableProperty->GetOwnerClass()))
 	{
 		PropertyOwnerBlueprint = Cast<UBlueprint>(GeneratedClass->ClassGeneratedBy);
 	}
@@ -999,16 +999,19 @@ void FBlueprintVarActionDetails::OnVarNameCommitted(const FText& InNewText, ETex
 bool FBlueprintVarActionDetails::GetVariableTypeChangeEnabled() const
 {
 	UProperty* VariableProperty = CachedVariableProperty.Get();
-	if(VariableProperty && IsVariableInBlueprint() && !IsALocalVariable(VariableProperty))
+	if(VariableProperty && IsVariableInBlueprint())
 	{
-		if(GetBlueprintObj()->SkeletonGeneratedClass->GetAuthoritativeClass() != VariableProperty->GetOwnerClass()->GetAuthoritativeClass())
-		{
-			return false;
-		}
-		// If the variable belongs to this class and cannot be found in the member variable list, it is not editable (it may be a component)
-		if (FBlueprintEditorUtils::FindNewVariableIndex(GetBlueprintObj(), CachedVariableName) == INDEX_NONE)
+		if (!IsALocalVariable(VariableProperty))
 		{
-			return false;
+			if(GetBlueprintObj()->SkeletonGeneratedClass->GetAuthoritativeClass() != VariableProperty->GetOwnerClass()->GetAuthoritativeClass())
+			{
+				return false;
+			}
+			// If the variable belongs to this class and cannot be found in the member variable list, it is not editable (it may be a component)
+			if (FBlueprintEditorUtils::FindNewVariableIndex(GetBlueprintObj(), CachedVariableName) == INDEX_NONE)
+			{
+				return false;
+			}
 		}
 		return true;
 	}
