Let users filter out static functions/macros from their Blueprint context menus.

diff --git a/Engine/Source/Editor/Kismet/Private/BlueprintActionMenuUtils.cpp b/Engine/Source/Editor/Kismet/Private/BlueprintActionMenuUtils.cpp
--- a/Engine/Source/Editor/Kismet/Private/BlueprintActionMenuUtils.cpp
+++ b/Engine/Source/Editor/Kismet/Private/BlueprintActionMenuUtils.cpp
@@ -441,7 +441,13 @@ void FBlueprintActionMenuUtils::MakeContextMenu(FBlueprintActionContext const& C
 	// Composing Filters
 	//--------------------------------------
 
-	FBlueprintActionFilter MainMenuFilter;
+	uint32 FilterFlags = 0x00;
+	if ((ClassTargetMask & EContextTargetFlags::TARGET_BlueprintLibraries) == 0)
+	{
+		FilterFlags |= FBlueprintActionFilter::BPFILTER_RejectGlobalFields;
+	}
+
+	FBlueprintActionFilter MainMenuFilter(FilterFlags);
 	MainMenuFilter.Context = Context;
 	MainMenuFilter.Context.SelectedObjects.Empty();
 
diff --git a/Engine/Source/Editor/Kismet/Private/SBlueprintContextTargetMenu.cpp b/Engine/Source/Editor/Kismet/Private/SBlueprintContextTargetMenu.cpp
--- a/Engine/Source/Editor/Kismet/Private/SBlueprintContextTargetMenu.cpp
+++ b/Engine/Source/Editor/Kismet/Private/SBlueprintContextTargetMenu.cpp
@@ -419,7 +419,14 @@ bool FContextMenuTargetProfile::IsTargetEnabled(EContextTargetFlags::Type Flag)
 //------------------------------------------------------------------------------
 void FContextMenuTargetProfile::SaveProfile() const
 {
-	GConfig->SetInt(*ContextMenuTargetProfileImpl::ConfigSection, *ProfileSaveName, SavedTargetFlags, GEditorIni);
+	// want to save with all bits set (except for ones matching flags that
+	// currently exist)... this is so we can later add flags, and a user's 
+	// saved value doesn't immediately disable them (default to on)
+	uint32 const GreatestUsedFlag = (EContextTargetFlags::ContextTargetFlagsEnd & ~1);
+	uint32 const UnusedFlasgMask  = ((0xFFFFFFFF & GreatestUsedFlag) & ~GreatestUsedFlag);
+	uint32 const SaveValue = UnusedFlasgMask | SavedTargetFlags;
+
+	GConfig->SetInt(*ContextMenuTargetProfileImpl::ConfigSection, *ProfileSaveName, SaveValue, GEditorIni);
 }
 
 //------------------------------------------------------------------------------
@@ -429,6 +436,19 @@ bool FContextMenuTargetProfile::LoadProfile()
 	if (GConfig->GetInt(*ContextMenuTargetProfileImpl::ConfigSection, *ProfileSaveName, SavedFlags, GEditorIni))
 	{
 		SavedTargetFlags = SavedFlags;
+
+		uint32 const GreatestUsedFlag = (EContextTargetFlags::ContextTargetFlagsEnd & ~1);
+		uint32 const LowestUnusedFlag = GreatestUsedFlag << 1;
+		// before we saved values with the unused bits all set (to support future
+		// flags), we saved only the bits that were explicitly set by the user
+		if ((SavedFlags & LowestUnusedFlag) == 0)
+		{
+			uint32 const OriginalFlagsMask = EContextTargetFlags::TARGET_Blueprint | EContextTargetFlags::TARGET_SubComponents |
+				EContextTargetFlags::TARGET_NodeTarget | EContextTargetFlags::TARGET_PinObject | EContextTargetFlags::TARGET_SiblingPinObjects;
+			// add in any new flags that have been added since this profile was last saved
+			SavedTargetFlags = (0xFFFFFFFF & ~OriginalFlagsMask) | SavedFlags;
+
+		}
 		return true;
 	}
 	return false;
diff --git a/Engine/Source/Editor/Kismet/Public/BlueprintActionMenuUtils.h b/Engine/Source/Editor/Kismet/Public/BlueprintActionMenuUtils.h
--- a/Engine/Source/Editor/Kismet/Public/BlueprintActionMenuUtils.h
+++ b/Engine/Source/Editor/Kismet/Public/BlueprintActionMenuUtils.h
@@ -16,6 +16,7 @@ namespace EContextTargetFlags
 		TARGET_NodeTarget			= 0x00000004 UMETA(DisplayName="Node Target", ToolTip="Include functions and variables that belong to the same class that the pin's node does."),
 		TARGET_PinObject			= 0x00000008 UMETA(DisplayName="Pin Type Class", ToolTip="Include functions and variables that belong to this pin type."),
 		TARGET_SiblingPinObjects	= 0x00000010 UMETA(DisplayName="Other Object Outputs", ToolTip="Include functions and variables that belong to any of this node's output types."),
+		TARGET_BlueprintLibraries	= 0x00000020 UMETA(DisplayName="Libraries", ToolTip="Include static functions that are globally accessible (belonging to function/macro libraries, etc.)."),
 
 		ContextTargetFlagsEnd UMETA(Hidden), // +1 to the last flag (so we can easily iterate these flags)
 	};
