[UE-15851] Fix potential log spam in the Blueprint class editor caused by native components that are instanced dynamically outside of an Actor's constructor.

change summary:
- modified SSCSEditor::Tick() to exclude instanced components that have the component class CDO as their archetype during node validation

diff --git a/Engine/Source/Editor/Kismet/Private/SSCSEditor.cpp b/Engine/Source/Editor/Kismet/Private/SSCSEditor.cpp
--- a/Engine/Source/Editor/Kismet/Private/SSCSEditor.cpp
+++ b/Engine/Source/Editor/Kismet/Private/SSCSEditor.cpp
@@ -3558,9 +3558,9 @@ void SSCSEditor::Tick( const FGeometry& AllottedGeometry, const double InCurrent
 		int32 NumComponentInstances = 0;
 		for (auto CompIt = ActorInstance->GetComponents().CreateConstIterator(); CompIt; ++CompIt)
 		{
-			// Don't count editor-only components in instanced mode, because we don't show them. Also, don't count UCS-added components if the option to hide them is enabled.
+			// Don't count editor-only components in instanced mode, because we don't show them. Also, don't count components that aren't based on a non-default template, and don't count UCS-added components if the option to hide them is enabled.
 			const UActorComponent* CompInst = *CompIt;
-			if ((!CompInst->IsEditorOnly() || (EditorMode != EComponentEditorMode::ActorInstance && CompInst->GetArchetype() != CompInst->GetClass()->GetDefaultObject()))
+			if ((!CompInst->IsEditorOnly() || EditorMode != EComponentEditorMode::ActorInstance) && (CompInst->GetArchetype() != CompInst->GetClass()->GetDefaultObject())
 				&& (CompInst->CreationMethod != EComponentCreationMethod::UserConstructionScript || !GetDefault<UBlueprintEditorSettings>()->bHideConstructionScriptComponentsInDetailsView))
 			{
 				++NumComponentInstances;
