Fixing crash with Find-in-Blueprints, fallout from changes to FString.

diff --git a/Engine/Source/Editor/Kismet/Private/FindInBlueprintManager.cpp b/Engine/Source/Editor/Kismet/Private/FindInBlueprintManager.cpp
--- a/Engine/Source/Editor/Kismet/Private/FindInBlueprintManager.cpp
+++ b/Engine/Source/Editor/Kismet/Private/FindInBlueprintManager.cpp
@@ -87,7 +87,7 @@ namespace FiBSerializationHelpers
 		// Read, as a byte string, the number of characters composing the Lookup Table for the Json.
 		FString SizeOfDataAsHex;
 		SizeOfDataAsHex.GetCharArray().AddUninitialized(InBytes + 1);
-		SizeOfDataAsHex[InBytes] = TEXT('\0');
+		SizeOfDataAsHex.GetCharArray()[InBytes] = TEXT('\0');
 		InStream.Serialize((char*)SizeOfDataAsHex.GetCharArray().GetData(), sizeof(TCHAR) * InBytes);
 
 		// Convert the number (which is stored in 1 serialized byte per TChar) into an int32
