[UE-15837] Handling a closed details tab (when attempting to open class defaults) instead of blindly expecting the tab pointer to be valid.

diff --git a/Engine/Source/Editor/Kismet/Private/BlueprintEditor.cpp b/Engine/Source/Editor/Kismet/Private/BlueprintEditor.cpp
--- a/Engine/Source/Editor/Kismet/Private/BlueprintEditor.cpp
+++ b/Engine/Source/Editor/Kismet/Private/BlueprintEditor.cpp
@@ -7606,7 +7606,7 @@ void FBlueprintEditor::TryInvokingDetailsTab(bool bFlash)
 		// being interacted with.  So we make sure the window it's in is focused and the tab is in the foreground.
 		if ( BlueprintTab.IsValid() && BlueprintTab->IsForeground() && BlueprintTab->GetParentWindow()->HasFocusedDescendants()  )
 		{
-			if( !Inspector.IsValid() || (Inspector->GetOwnerTab()->GetDockArea().IsValid()) ) 
+			if (!Inspector.IsValid() || !Inspector->GetOwnerTab().IsValid() || Inspector->GetOwnerTab()->GetDockArea().IsValid())
 			{
 				// Show the details panel if it doesn't exist.
 				TabManager->InvokeTab(FBlueprintEditorTabs::DetailsID);
