Can change a variable's type in the MyBlueprint window by clicking on it's type icon or right clicking to toggle being an array.

diff --git a/Engine/Source/Editor/EditorStyle/Private/SlateEditorStyle.cpp b/Engine/Source/Editor/EditorStyle/Private/SlateEditorStyle.cpp
--- a/Engine/Source/Editor/EditorStyle/Private/SlateEditorStyle.cpp
+++ b/Engine/Source/Editor/EditorStyle/Private/SlateEditorStyle.cpp
@@ -4971,6 +4971,14 @@ void FSlateEditorStyle::FStyle::SetupPersonaStyle()
 // 			.SetCheckedHoveredImage(IMAGE_BRUSH("Common/TreeArrow_Collapsed_Hovered", Icon10x10, FLinearColor(0.4f, 0.4f, 0.4f, 1.f)))
 // 			.SetCheckedPressedImage(IMAGE_BRUSH("Common/TreeArrow_Collapsed_Hovered", Icon10x10, FLinearColor(0.2f, 0.2f, 0.2f, 1.f)));
 		Set("BlueprintEditor.ContextMenu.TargetsButton", BlueprintContextTargetsButtonStyle);
+
+		Set( "BlueprintEditor.CompactPinTypeSelector", FButtonStyle()
+			.SetNormal        ( FSlateNoResource() )
+			.SetPressed       ( BOX_BRUSH( "Common/Button_Pressed", 8.0f/32.0f, SelectionColor_Pressed ) )
+			.SetHovered       ( BOX_BRUSH( "Common/Button_Hovered", 8.0f/32.0f, SelectionColor ) )
+			.SetNormalPadding ( FMargin( 0,0,0,0 ) )
+			.SetPressedPadding( FMargin( 1,1,2,2 ) )
+			);
 	}
 
 	// Kismet linear expression display
diff --git a/Engine/Source/Editor/Kismet/Private/SBlueprintPalette.cpp b/Engine/Source/Editor/Kismet/Private/SBlueprintPalette.cpp
--- a/Engine/Source/Editor/Kismet/Private/SBlueprintPalette.cpp
+++ b/Engine/Source/Editor/Kismet/Private/SBlueprintPalette.cpp
@@ -37,6 +37,7 @@
 #include "Editor/UnrealEd/Public/Kismet2/ComponentEditorUtils.h"
 #include "Engine/SCS_Node.h"
 #include "Components/TimelineComponent.h"
+#include "SPinTypeSelector.h"
 
 #define LOCTEXT_NAMESPACE "BlueprintPalette"
 
@@ -651,6 +652,97 @@ public:
 	}
 };
 
+/*******************************************************************************
+* SPinTypeSelectorHelper
+*******************************************************************************/
+
+class SPinTypeSelectorHelper : public SCompoundWidget
+{
+public:
+	SLATE_BEGIN_ARGS( SPinTypeSelectorHelper ) {}
+	SLATE_END_ARGS()
+
+	/**
+	 * Constructs a PinTypeSelector widget (for variable actions only, so that 
+	 * the user can modify the variable's type without going to the details panel).
+	 * 
+	 * @param  InArgs					A set of slate arguments, defined above.
+	 * @param  InVariableProperty		The variable property to select
+	 * @param  InBlueprintEditor			A pointer to the blueprint editor that the palette belongs to.
+	 */
+	void Construct(const FArguments& InArgs, UProperty* InVariableProperty, TWeakPtr<FBlueprintEditor> InBlueprintEditor)
+	{
+		BlueprintObj = InBlueprintEditor.Pin()->GetBlueprintObj();
+		BlueprintEditorPtr = InBlueprintEditor;
+		VariableProperty = InVariableProperty;
+
+		const UEdGraphSchema_K2* Schema = GetDefault<UEdGraphSchema_K2>();
+		const FSlateFontInfo DetailFontInfo = IDetailLayoutBuilder::GetDetailFont();
+		this->ChildSlot
+		[
+			SNew(SPinTypeSelector, FGetPinTypeTree::CreateUObject(Schema, &UEdGraphSchema_K2::GetVariableTypeTree))
+			.Schema(Schema)
+			.TargetPinType(this, &SPinTypeSelectorHelper::OnGetVarType)
+			.OnPinTypeChanged(this, &SPinTypeSelectorHelper::OnVarTypeChanged)
+			.bAllowExec(false)
+			.Font( DetailFontInfo )
+			.bCompactSelector(true)
+		];
+	}
+
+private:
+	FEdGraphPinType OnGetVarType() const
+	{
+		if (VariableProperty)
+		{
+			const UEdGraphSchema_K2* K2Schema = GetDefault<UEdGraphSchema_K2>();
+			FEdGraphPinType Type;
+			K2Schema->ConvertPropertyToPinType(VariableProperty, Type);
+			return Type;
+		}
+		return FEdGraphPinType();
+	}
+
+	void OnVarTypeChanged(const FEdGraphPinType& InNewPinType)
+	{
+		if (FBlueprintEditorUtils::IsPinTypeValid(InNewPinType))
+		{
+			FName VarName = VariableProperty->GetFName();
+
+			if (VarName != NAME_None)
+			{
+				// Set the MyBP tab's last pin type used as this, for adding lots of variables of the same type
+				BlueprintEditorPtr.Pin()->GetMyBlueprintWidget()->GetLastPinTypeUsed() = InNewPinType;
+
+				if (VariableProperty)
+				{
+					if (UFunction* LocalVariableScope = Cast<UFunction>(VariableProperty->GetOuter()))
+					{
+						FBlueprintEditorUtils::ChangeLocalVariableType(BlueprintObj, LocalVariableScope, VarName, InNewPinType);
+					}
+					else
+					{
+						FBlueprintEditorUtils::ChangeMemberVariableType(BlueprintObj, VarName, InNewPinType);
+					}
+				}
+			}
+		}
+	}
+
+private:
+	/** The action that the owning palette entry represents */
+	TWeakPtr<FEdGraphSchemaAction_K2Var> ActionPtr;
+
+	/** Pointer back to the blueprint that is being displayed: */
+	UBlueprint* BlueprintObj;
+
+	/** Pointer back to the blueprint editor that owns this, optional because of diff and merge views: */
+	TWeakPtr<FBlueprintEditor>     BlueprintEditorPtr;
+
+	/** Variable Property to change the type of */
+	UProperty* VariableProperty;
+};
+
 /*******************************************************************************
 * SPaletteItemVisibilityToggle
 *******************************************************************************/
@@ -896,6 +988,28 @@ void SBlueprintPaletteItem::Construct(const FArguments& InArgs, FCreateWidgetFor
 	bool bIsReadOnly = InBlueprintEditor.IsValid() ? FBlueprintEditorUtils::IsPaletteActionReadOnly(GraphAction, InBlueprintEditor.Pin()) : true;
 	TSharedRef<SWidget> NameSlotWidget = CreateTextSlotWidget( NameFont, InCreateData, bIsReadOnly );
 	
+	// For Variables and Local Variables, we will convert the icon widget into a pin type selector.
+	if (GraphAction->GetTypeId() == FEdGraphSchemaAction_K2Var::StaticGetTypeId() || GraphAction->GetTypeId() == FEdGraphSchemaAction_K2LocalVar::StaticGetTypeId())
+	{
+		UProperty* VariableProp = nullptr;
+
+		if (GraphAction->GetTypeId() == FEdGraphSchemaAction_K2Var::StaticGetTypeId())
+		{
+			VariableProp = StaticCastSharedPtr<FEdGraphSchemaAction_K2Var>(GraphAction)->GetProperty();
+		}
+		else if (GraphAction->GetTypeId() == FEdGraphSchemaAction_K2LocalVar::StaticGetTypeId())
+		{
+			VariableProp = StaticCastSharedPtr<FEdGraphSchemaAction_K2LocalVar>(GraphAction)->GetProperty();
+		}
+
+		// If the variable is not a local variable or created by the current Blueprint, do not use the PinTypeSelector
+		if (FBlueprintEditorUtils::IsVariableCreatedByBlueprint(Blueprint, VariableProp) || Cast<UFunction>(VariableProp->GetOuter()))
+		{
+			const UEdGraphSchema_K2* Schema = GetDefault<UEdGraphSchema_K2>();
+			IconWidget = SNew(SPinTypeSelectorHelper, VariableProp, BlueprintEditorPtr);
+		}
+	}
+
 	// now, create the actual widget
 	ChildSlot
 	[
diff --git a/Engine/Source/Editor/KismetWidgets/Private/SPinTypeSelector.cpp b/Engine/Source/Editor/KismetWidgets/Private/SPinTypeSelector.cpp
--- a/Engine/Source/Editor/KismetWidgets/Private/SPinTypeSelector.cpp
+++ b/Engine/Source/Editor/KismetWidgets/Private/SPinTypeSelector.cpp
@@ -28,10 +28,33 @@ void SPinTypeSelector::Construct(const FArguments& InArgs, FGetPinTypeTree GetPi
 	TreeViewHeight = InArgs._TreeViewHeight;
 
 	TargetPinType = InArgs._TargetPinType;
+	bIsCompactSelector = InArgs._bCompactSelector;
 
-	this->ChildSlot
-	[
-		SNew(SHorizontalBox)
+	bIsRightMousePressed = false;
+
+	// Depending on if this is a compact selector or not, we generate a different compoound widget
+	TSharedPtr<SWidget> Widget;
+
+	if (InArgs._bCompactSelector)
+	{
+		// Only have a combo button with an icon
+		Widget = SAssignNew( TypeComboButton, SComboButton )
+			.OnGetMenuContent(this, &SPinTypeSelector::GetMenuContent)
+			.ContentPadding(0)
+			.ToolTipText(this, &SPinTypeSelector::GetToolTipForComboBoxType)
+			.HasDownArrow(!InArgs._bCompactSelector)
+			.ButtonStyle(FEditorStyle::Get(),  "BlueprintEditor.CompactPinTypeSelector")
+			.ButtonContent()
+			[
+				SNew(SImage)
+				.Image( this, &SPinTypeSelector::GetTypeIconImage )
+				.ColorAndOpacity( this, &SPinTypeSelector::GetTypeIconColor )
+			];
+	}
+	else
+	{
+		// Traditional Pin Type Selector with a combo button, the icon, the current type name, and a toggle button for being an array
+		Widget = SNew(SHorizontalBox)
 		+SHorizontalBox::Slot()
 		[
 			SAssignNew( TypeComboButton, SComboButton )
@@ -80,7 +103,11 @@ void SPinTypeSelector::Construct(const FArguments& InArgs, FGetPinTypeTree GetPi
 				.Image( FEditorStyle::GetBrush(TEXT("Kismet.VariableList.ArrayTypeIcon")) )
 				.ColorAndOpacity( this, &SPinTypeSelector::GetTypeIconColor )
 			]
-		]
+		];
+	}
+	this->ChildSlot
+	[
+		Widget.ToSharedRef()
 	];
 }
 
@@ -127,6 +154,11 @@ void SPinTypeSelector::OnArrayCheckStateChanged(ECheckBoxState NewState)
 	OnTypeChanged.ExecuteIfBound(NewTargetPinType);
 }
 
+void SPinTypeSelector::OnArrayStateToggled()
+{
+	OnArrayCheckStateChanged((IsArrayChecked() == ECheckBoxState::Checked)? ECheckBoxState::Unchecked : ECheckBoxState::Checked);
+}
+
 //=======================================================================
 // Type TreeView Support
 TSharedRef<ITableRow> SPinTypeSelector::GenerateTypeTreeRow(FPinTypeTreeItem InItem, const TSharedRef<STableViewBase>& OwnerTree)
@@ -392,10 +424,17 @@ FText SPinTypeSelector::GetToolTipForComboBoxType() const
 {
 	if(IsEnabled())
 	{
-		return LOCTEXT("PinTypeSelector", "Select the variable's pin type.");
+		if (bIsCompactSelector)
+		{
+			return LOCTEXT("CompactPinTypeSelector", "Left click to select the variable's pin type. Right click to toggle the type as an array.");
+		}
+		else
+		{
+			return LOCTEXT("PinTypeSelector", "Select the variable's pin type.");
+		}
 	}
 
-	return LOCTEXT("PinTypeSelector_Disabled", "Cannot edit variable type while the variable is placed in a graph or inherited from parent.");
+	return LOCTEXT("PinTypeSelector_Disabled", "Cannot edit variable type when they are inherited from parent.");
 }
 
 FText SPinTypeSelector::GetToolTipForArrayWidget() const
@@ -412,4 +451,35 @@ FText SPinTypeSelector::GetToolTipForArrayWidget() const
 
 	return LOCTEXT("ArrayCheckBox_Disabled", "Cannot edit variable type while the variable is placed in a graph or inherited from parent.");
 }
+
+FReply SPinTypeSelector::OnMouseButtonDown( const FGeometry& MyGeometry, const FPointerEvent& MouseEvent )
+{
+	if(bIsCompactSelector && MouseEvent.GetEffectingButton() == EKeys::RightMouseButton)
+	{
+		bIsRightMousePressed = true;
+		return FReply::Handled();
+	}
+
+	return FReply::Unhandled();
+}
+
+FReply SPinTypeSelector::OnMouseButtonUp( const FGeometry& MyGeometry, const FPointerEvent& MouseEvent )
+{
+	if(bIsCompactSelector && MouseEvent.GetEffectingButton() == EKeys::RightMouseButton)
+	{
+		if (bIsRightMousePressed)
+		{
+			OnArrayStateToggled();
+		}
+		return FReply::Handled();
+	}
+
+	return FReply::Unhandled();
+}
+
+void SPinTypeSelector::OnMouseLeave( const FPointerEvent& MouseEvent )
+{
+	bIsRightMousePressed = false;
+}
+
 #undef LOCTEXT_NAMESPACE
diff --git a/Engine/Source/Editor/KismetWidgets/Public/SPinTypeSelector.h b/Engine/Source/Editor/KismetWidgets/Public/SPinTypeSelector.h
--- a/Engine/Source/Editor/KismetWidgets/Public/SPinTypeSelector.h
+++ b/Engine/Source/Editor/KismetWidgets/Public/SPinTypeSelector.h
@@ -28,6 +28,7 @@ public:
 		, _TreeViewWidth(300.f)
 		, _TreeViewHeight(400.f)
 		, _Font( FEditorStyle::GetFontStyle( TEXT("NormalFont") ) )
+		, _bCompactSelector( false )
 		{}
 		SLATE_ATTRIBUTE( FEdGraphPinType, TargetPinType )
 		SLATE_ARGUMENT( const UEdGraphSchema_K2*, Schema )
@@ -39,10 +40,17 @@ public:
 		SLATE_EVENT( FOnPinTypeChanged, OnPinTypePreChanged )
 		SLATE_EVENT( FOnPinTypeChanged, OnPinTypeChanged )
 		SLATE_ATTRIBUTE( FSlateFontInfo, Font )
+		SLATE_ARGUMENT( bool, bCompactSelector )
 	SLATE_END_ARGS()
 public:
 	void Construct(const FArguments& InArgs, FGetPinTypeTree GetPinTypeTreeFunc);
 
+	// SWidget interface
+	virtual FReply OnMouseButtonDown( const FGeometry& MyGeometry, const FPointerEvent& MouseEvent ) override;
+	virtual FReply OnMouseButtonUp( const FGeometry& MyGeometry, const FPointerEvent& MouseEvent ) override;
+	virtual void OnMouseLeave( const FPointerEvent& MouseEvent ) override;
+	// End of SWidget interface
+
 protected:
 	/** Gets the icon (pin or array) for the type being manipulated */
 	const FSlateBrush* GetTypeIconImage() const;
@@ -82,10 +90,19 @@ protected:
 	/** Desired height of the tree view widget */
 	TAttribute<FOptionalSize>	TreeViewHeight;
 
+	/** TRUE when the right mouse button is pressed, keeps from handling a right click that does not begin in the widget */
+	bool bIsRightMousePressed;
+
+	/** TRUE if displaying a compact selector widget, some functionality is enabled in different ways if this is TRUE */
+	bool bIsCompactSelector;
+
 	/** Array checkbox support functions */
 	ECheckBoxState IsArrayChecked() const;
 	void OnArrayCheckStateChanged(ECheckBoxState NewState);
 
+	/** Toggles the variable type as an array */
+	void OnArrayStateToggled();
+
 	/** Array containing the unfiltered list of all supported types this pin could possibly have */
 	TArray<FPinTypeTreeItem>		TypeTreeRoot;
 	/** Array containing a filtered list, according to the text in the searchbox */
diff --git a/Engine/Source/Editor/UnrealEd/Private/Kismet2/BlueprintEditorUtils.cpp b/Engine/Source/Editor/UnrealEd/Private/Kismet2/BlueprintEditorUtils.cpp
--- a/Engine/Source/Editor/UnrealEd/Private/Kismet2/BlueprintEditorUtils.cpp
+++ b/Engine/Source/Editor/UnrealEd/Private/Kismet2/BlueprintEditorUtils.cpp
@@ -3126,6 +3126,16 @@ FString FBlueprintEditorUtils::GetBlueprintTypeDescription(const UBlueprint* Blu
 //////////////////////////////////////////////////////////////////////////
 // Variables
 
+bool FBlueprintEditorUtils::IsVariableCreatedByBlueprint(UBlueprint* InBlueprint, UProperty* InVariableProperty)
+{
+	bool bIsVariableCreatedByBlueprint = false;
+	if (UBlueprintGeneratedClass* GeneratedClass = Cast<UBlueprintGeneratedClass>(InVariableProperty->GetOwnerClass()))
+	{
+		UBlueprint* OwnerBlueprint = Cast<UBlueprint>(GeneratedClass->ClassGeneratedBy);
+		bIsVariableCreatedByBlueprint = (OwnerBlueprint == InBlueprint && FBlueprintEditorUtils::FindNewVariableIndex(InBlueprint, InVariableProperty->GetFName()) != INDEX_NONE);
+	}
+	return bIsVariableCreatedByBlueprint;
+}
 
 // Find the index of a variable first declared in this blueprint. Returns INDEX_NONE if not found.
 int32 FBlueprintEditorUtils::FindNewVariableIndex(const UBlueprint* Blueprint, const FName& InName) 
diff --git a/Engine/Source/Editor/UnrealEd/Public/Kismet2/BlueprintEditorUtils.h b/Engine/Source/Editor/UnrealEd/Public/Kismet2/BlueprintEditorUtils.h
--- a/Engine/Source/Editor/UnrealEd/Public/Kismet2/BlueprintEditorUtils.h
+++ b/Engine/Source/Editor/UnrealEd/Public/Kismet2/BlueprintEditorUtils.h
@@ -898,6 +898,9 @@ public:
 	/** Set RepNotify function of variable */
 	static void SetBlueprintVariableRepNotifyFunc(UBlueprint* Blueprint, const FName& VarName, const FName& RepNotifyFunc);
 
+	/** Returns TRUE if the variable was created by the Blueprint */
+	static bool IsVariableCreatedByBlueprint(UBlueprint* InBlueprint, UProperty* InVariableProperty);
+
 	/**
 	 * Find the index of a variable first declared in this blueprint. Returns INDEX_NONE if not found.
 	 *
