Fixing non-unity build issue caused by CL# 2559203.

diff --git a/Engine/Source/Editor/Kismet/Private/SBlueprintPalette.cpp b/Engine/Source/Editor/Kismet/Private/SBlueprintPalette.cpp
--- a/Engine/Source/Editor/Kismet/Private/SBlueprintPalette.cpp
+++ b/Engine/Source/Editor/Kismet/Private/SBlueprintPalette.cpp
@@ -677,7 +677,6 @@ public:
 		VariableProperty = InVariableProperty;
 
 		const UEdGraphSchema_K2* Schema = GetDefault<UEdGraphSchema_K2>();
-		const FSlateFontInfo DetailFontInfo = IDetailLayoutBuilder::GetDetailFont();
 		this->ChildSlot
 		[
 			SNew(SPinTypeSelector, FGetPinTypeTree::CreateUObject(Schema, &UEdGraphSchema_K2::GetVariableTypeTree))
@@ -685,7 +684,6 @@ public:
 			.TargetPinType(this, &SPinTypeSelectorHelper::OnGetVarType)
 			.OnPinTypeChanged(this, &SPinTypeSelectorHelper::OnVarTypeChanged)
 			.bAllowExec(false)
-			.Font( DetailFontInfo )
 			.bCompactSelector(true)
 		];
 	}
