[UE-16005] Fix a regression that caused additional log spam on add of any Empty Actor instance into a scene.

change summary:
- modified SSCSEditor::Tick() to skip validation checks in SCS (non-instanced) mode unless the option to show UCS-added components is enabled.
- modified SSCSEditor::Tick() to skip the native archetype check entirely when in instanced mode (it's not needed there).
- modified SSCSEditor::Tick() to only periodically log out a warning when we're calling UpdateTree() too frequently due to invalid nodes every tick.

diff --git a/Engine/Source/Editor/Kismet/Private/SSCSEditor.cpp b/Engine/Source/Editor/Kismet/Private/SSCSEditor.cpp
--- a/Engine/Source/Editor/Kismet/Private/SSCSEditor.cpp
+++ b/Engine/Source/Editor/Kismet/Private/SSCSEditor.cpp
@@ -3552,27 +3552,41 @@ void SSCSEditor::Tick( const FGeometry& AllottedGeometry, const double InCurrent
 		return bFoundInvalidNode;
 	};
 
-	const AActor* ActorInstance = (EditorMode == EComponentEditorMode::ActorInstance) ? GetActorContext() : PreviewActor.Get(nullptr);
-	if (ActorInstance)
+	// Only run this in "instanced" mode or otherwise only when the option to show UCS-instanced components is enabled
+	const bool bShowInstancedUCSComponents = !GetDefault<UBlueprintEditorSettings>()->bHideConstructionScriptComponentsInDetailsView;
+	if(EditorMode == EComponentEditorMode::ActorInstance || bShowInstancedUCSComponents)
 	{
-		int32 NumComponentInstances = 0;
-		for (auto CompIt = ActorInstance->GetComponents().CreateConstIterator(); CompIt; ++CompIt)
+		const AActor* ActorInstance = (EditorMode == EComponentEditorMode::ActorInstance) ? GetActorContext() : PreviewActor.Get(nullptr);
+		if (ActorInstance)
 		{
-			// Don't count editor-only components in instanced mode, because we don't show them. Also, don't count components that aren't based on a non-default template, and don't count UCS-added components if the option to hide them is enabled.
-			const UActorComponent* CompInst = *CompIt;
-			if ((!CompInst->IsEditorOnly() || EditorMode != EComponentEditorMode::ActorInstance) && (CompInst->GetArchetype() != CompInst->GetClass()->GetDefaultObject())
-				&& (CompInst->CreationMethod != EComponentCreationMethod::UserConstructionScript || !GetDefault<UBlueprintEditorSettings>()->bHideConstructionScriptComponentsInDetailsView))
+			int32 NumComponentInstances = 0;
+			for (auto CompIt = ActorInstance->GetComponents().CreateConstIterator(); CompIt; ++CompIt)
 			{
-				++NumComponentInstances;
+				// Don't count editor-only components in instanced mode, because we don't show them there. Also, don't count UCS-added components if the option to hide them is enabled, and don't count native components that are based on the CDO in the BP editor context (i.e. preview world).
+				const UActorComponent* CompInst = *CompIt;
+				if ((!CompInst->IsEditorOnly() || EditorMode != EComponentEditorMode::ActorInstance)
+					&& (CompInst->CreationMethod != EComponentCreationMethod::UserConstructionScript || bShowInstancedUCSComponents)
+					&& (EditorMode != EComponentEditorMode::BlueprintSCS || CompInst->CreationMethod != EComponentCreationMethod::Native || CompInst->GetArchetype() != CompInst->GetClass()->GetDefaultObject()))
+				{
+					++NumComponentInstances;
+				}
 			}
-		}
 
-		int32 NumComponentNodes = 0;
-		if (AreAnyNodesInvalidLambda(GetRootNodes(), NumComponentNodes) || NumComponentNodes != NumComponentInstances)
-		{
-			UE_LOG(LogSCSEditor, Log, TEXT("Calling UpdateTree() from Tick()."));
+			int32 NumComponentNodes = 0;
+			static int32 UpdateCounter = 0;
+			if (AreAnyNodesInvalidLambda(GetRootNodes(), NumComponentNodes) || NumComponentNodes != NumComponentInstances)
+			{
+				if(++UpdateCounter % 60)
+				{
+					UE_LOG(LogSCSEditor, Warning, TEXT("Calling UpdateTree() from Tick() more frequently than expected due to invalid tree node state. This might be a performance issue."));
+				}
 
-			UpdateTree();
+				UpdateTree();
+			}
+			else
+			{
+				UpdateCounter = 0;
+			}
 		}
 	}
 }
