Blueprint member variables of type component that are not generated by SCS, now have the details options normally available to variables exposed.

Enhanced FBlueprintVarActionDetails::IsAComponentVariable to look deeper than the property type for determination on if it's a component.

Renamed FBlueprintVarActionDetails::IsAComponentVariable to IsASCSVariable

diff --git a/Engine/Source/Editor/Kismet/Private/BlueprintDetailsCustomization.cpp b/Engine/Source/Editor/Kismet/Private/BlueprintDetailsCustomization.cpp
--- a/Engine/Source/Editor/Kismet/Private/BlueprintDetailsCustomization.cpp
+++ b/Engine/Source/Editor/Kismet/Private/BlueprintDetailsCustomization.cpp
@@ -833,10 +833,15 @@ TSharedRef<ITableRow> FBlueprintVarActionDetails::OnGenerateWidgetForPropertyLis
 		];
 }
 
-bool FBlueprintVarActionDetails::IsAComponentVariable(UProperty* VariableProperty) const
+bool FBlueprintVarActionDetails::IsASCSVariable(UProperty* VariableProperty) const
 {
 	UObjectProperty* VariableObjProp = VariableProperty ? Cast<UObjectProperty>(VariableProperty) : NULL;
-	return (VariableObjProp != NULL) && (VariableObjProp->PropertyClass != NULL) && (VariableObjProp->PropertyClass->IsChildOf(UActorComponent::StaticClass()));
+
+	if (VariableObjProp != NULL && VariableObjProp->PropertyClass != NULL && VariableObjProp->PropertyClass->IsChildOf(UActorComponent::StaticClass()))
+	{
+		return !FBlueprintEditorUtils::IsVariableCreatedByBlueprint(GetBlueprintObj(), VariableObjProp);
+	}
+	return false;
 }
 
 bool FBlueprintVarActionDetails::IsABlueprintVariable(UProperty* VariableProperty) const
@@ -882,7 +887,7 @@ bool FBlueprintVarActionDetails::GetVariableNameChangeEnabled() const
 		{
 			bIsReadOnly = false;
 		}
-		else if(IsAComponentVariable(VariableProperty) && Blueprint->SimpleConstructionScript != NULL)
+		else if(IsASCSVariable(VariableProperty) && Blueprint->SimpleConstructionScript != NULL)
 		{
 			if (USCS_Node* Node = Blueprint->SimpleConstructionScript->FindSCSNode(CachedVariableName))
 			{
@@ -911,7 +916,7 @@ void FBlueprintVarActionDetails::OnVarNameChanged(const FText& InNewText)
 	check(Blueprint != NULL);
 
 	UProperty* VariableProperty = CachedVariableProperty.Get();
-	if(VariableProperty && IsAComponentVariable(VariableProperty) && Blueprint->SimpleConstructionScript != NULL)
+	if(VariableProperty && IsASCSVariable(VariableProperty) && Blueprint->SimpleConstructionScript != NULL)
 	{
 		TArray<USCS_Node*> Nodes = Blueprint->SimpleConstructionScript->GetAllNodes();
 		for (TArray<USCS_Node*>::TConstIterator NodeIt(Nodes); NodeIt; ++NodeIt)
@@ -1360,7 +1365,7 @@ EVisibility FBlueprintVarActionDetails::ShowEditableCheckboxVisibilty() const
 	UProperty* VariableProperty = CachedVariableProperty.Get();
 	if (VariableProperty && GetPropertyOwnerBlueprint())
 	{
-		if (IsABlueprintVariable(VariableProperty) && !IsAComponentVariable(VariableProperty))
+		if (IsABlueprintVariable(VariableProperty) && !IsASCSVariable(VariableProperty))
 		{
 			return EVisibility::Visible;
 		}
@@ -1476,7 +1481,7 @@ EVisibility FBlueprintVarActionDetails::ExposeOnSpawnVisibility() const
 		FEdGraphPinType VariablePinType;
 		K2Schema->ConvertPropertyToPinType(VariableProperty, VariablePinType);
 
-		const bool bShowPrivacySetting = IsABlueprintVariable(VariableProperty) && !IsAComponentVariable(VariableProperty);
+		const bool bShowPrivacySetting = IsABlueprintVariable(VariableProperty) && !IsASCSVariable(VariableProperty);
 		if (bShowPrivacySetting && (K2Schema->FindSetVariableByNameFunction(VariablePinType) != NULL))
 		{
 			return EVisibility::Visible;
@@ -1517,7 +1522,7 @@ EVisibility FBlueprintVarActionDetails::ExposePrivateVisibility() const
 	UProperty* Property = CachedVariableProperty.Get();
 	if (Property && GetPropertyOwnerBlueprint())
 	{
-		if (IsABlueprintVariable(Property) && !IsAComponentVariable(Property))
+		if (IsABlueprintVariable(Property) && !IsASCSVariable(Property))
 		{
 			return EVisibility::Visible;
 		}
@@ -1615,7 +1620,7 @@ EVisibility FBlueprintVarActionDetails::ExposeConfigVisibility() const
 	UProperty* Property = CachedVariableProperty.Get();
 	if (Property)
 	{
-		if (IsABlueprintVariable(Property) && !IsAComponentVariable(Property))
+		if (IsABlueprintVariable(Property) && !IsASCSVariable(Property))
 		{
 			return EVisibility::Visible;
 		}
@@ -1779,7 +1784,7 @@ EVisibility FBlueprintVarActionDetails::ReplicationVisibility() const
 	UProperty* VariableProperty = CachedVariableProperty.Get();
 	if(VariableProperty)
 	{
-		if (!IsAComponentVariable(VariableProperty) && IsABlueprintVariable(VariableProperty))
+		if (!IsASCSVariable(VariableProperty) && IsABlueprintVariable(VariableProperty))
 		{
 			return EVisibility::Visible;
 		}
@@ -1819,7 +1824,7 @@ EVisibility FBlueprintVarActionDetails::GetTransientVisibility() const
 	UProperty* VariableProperty = CachedVariableProperty.Get();
 	if (VariableProperty)
 	{
-		if (IsABlueprintVariable(VariableProperty) && !IsAComponentVariable(VariableProperty))
+		if (IsABlueprintVariable(VariableProperty) && !IsASCSVariable(VariableProperty))
 		{
 			return EVisibility::Visible;
 		}
@@ -1852,7 +1857,7 @@ EVisibility FBlueprintVarActionDetails::GetSaveGameVisibility() const
 	UProperty* VariableProperty = CachedVariableProperty.Get();
 	if (VariableProperty)
 	{
-		if (IsABlueprintVariable(VariableProperty) && !IsAComponentVariable(VariableProperty))
+		if (IsABlueprintVariable(VariableProperty) && !IsASCSVariable(VariableProperty))
 		{
 			return EVisibility::Visible;
 		}
@@ -1885,7 +1890,7 @@ EVisibility FBlueprintVarActionDetails::IsTooltipEditVisible() const
 	UProperty* VariableProperty = CachedVariableProperty.Get();
 	if (VariableProperty)
 	{
-		if ((IsABlueprintVariable(VariableProperty) && !IsAComponentVariable(VariableProperty)) || IsALocalVariable(VariableProperty))
+		if ((IsABlueprintVariable(VariableProperty) && !IsASCSVariable(VariableProperty)) || IsALocalVariable(VariableProperty))
 		{
 			return EVisibility::Visible;
 		}
diff --git a/Engine/Source/Editor/Kismet/Private/BlueprintDetailsCustomization.h b/Engine/Source/Editor/Kismet/Private/BlueprintDetailsCustomization.h
--- a/Engine/Source/Editor/Kismet/Private/BlueprintDetailsCustomization.h
+++ b/Engine/Source/Editor/Kismet/Private/BlueprintDetailsCustomization.h
@@ -92,7 +92,7 @@ private:
 	FName GetVariableName() const;
 
 	/** Commonly queried attributes about the schema action */
-	bool IsAComponentVariable(UProperty* VariableProperty) const;
+	bool IsASCSVariable(UProperty* VariableProperty) const;
 	bool IsABlueprintVariable(UProperty* VariableProperty) const;
 	bool IsALocalVariable(UProperty* VariableProperty) const;
 	UStruct* GetLocalVariableScope(UProperty* VariableProperty) const;
