[UE-15978] Have to copy component transforms directly from the template to the Blueprint preview instance (rather than modifying both in tandem), to keep the two versions from diverging.

diff --git a/Engine/Source/Editor/Kismet/Private/SCSEditorViewportClient.cpp b/Engine/Source/Editor/Kismet/Private/SCSEditorViewportClient.cpp
--- a/Engine/Source/Editor/Kismet/Private/SCSEditorViewportClient.cpp
+++ b/Engine/Source/Editor/Kismet/Private/SCSEditorViewportClient.cpp
@@ -408,17 +408,8 @@ bool FSCSEditorViewportClient::InputWidgetDelta( FViewport* Viewport, EAxisList:
 							}
 							else
 							{
-								// Apply delta to the preview actor's scene component
-								GEditor->ApplyDeltaToComponent(
-										SceneComp,
-										true,
-										&Drag,
-										&Rot,
-										&ModifiedScale,
-										SceneComp->RelativeLocation );
-								UpdatedComponents.Add(SceneComp);
-
-							// Apply delta to the template component object
+								// Apply delta to the template component object 
+								// (the preview scene component will be set in one of the ArchetypeInstances loops below... to keep the two in sync) 
 								GEditor->ApplyDeltaToComponent(
 										SelectedTemplate,
 										true,
@@ -456,7 +447,7 @@ bool FSCSEditorViewportClient::InputWidgetDelta( FViewport* Viewport, EAxisList:
 									for(int32 InstanceIndex = 0; InstanceIndex < ArchetypeInstances.Num(); ++InstanceIndex)
 									{
 										SceneComp = Cast<USceneComponent>(ArchetypeInstances[InstanceIndex]);
-										if(SceneComp && SceneComp->GetOwner() != PreviewActor)
+										if(SceneComp != nullptr)
 										{
 											FComponentEditorUtils::ApplyDefaultValueChange(SceneComp, SceneComp->RelativeLocation, OldRelativeLocation, SelectedTemplate->RelativeLocation);
 											FComponentEditorUtils::ApplyDefaultValueChange(SceneComp, SceneComp->RelativeRotation, OldRelativeRotation, SelectedTemplate->RelativeRotation);
@@ -469,15 +460,12 @@ bool FSCSEditorViewportClient::InputWidgetDelta( FViewport* Viewport, EAxisList:
 									Outer->GetArchetypeInstances(ArchetypeInstances);
 									for(int32 InstanceIndex = 0; InstanceIndex < ArchetypeInstances.Num(); ++InstanceIndex)
 									{
-										if(ArchetypeInstances[InstanceIndex] != PreviewActor)
+										SceneComp = static_cast<USceneComponent*>(FindObjectWithOuter(ArchetypeInstances[InstanceIndex], SelectedTemplate->GetClass(), SelectedTemplate->GetFName()));
+										if(SceneComp)
 										{
-											SceneComp = static_cast<USceneComponent*>(FindObjectWithOuter(ArchetypeInstances[InstanceIndex], SelectedTemplate->GetClass(), SelectedTemplate->GetFName()));
-											if(SceneComp)
-											{
-												FComponentEditorUtils::ApplyDefaultValueChange(SceneComp, SceneComp->RelativeLocation, OldRelativeLocation, SelectedTemplate->RelativeLocation);
-												FComponentEditorUtils::ApplyDefaultValueChange(SceneComp, SceneComp->RelativeRotation, OldRelativeRotation, SelectedTemplate->RelativeRotation);
-												FComponentEditorUtils::ApplyDefaultValueChange(SceneComp, SceneComp->RelativeScale3D,  OldRelativeScale3D,  SelectedTemplate->RelativeScale3D);
-											}
+											FComponentEditorUtils::ApplyDefaultValueChange(SceneComp, SceneComp->RelativeLocation, OldRelativeLocation, SelectedTemplate->RelativeLocation);
+											FComponentEditorUtils::ApplyDefaultValueChange(SceneComp, SceneComp->RelativeRotation, OldRelativeRotation, SelectedTemplate->RelativeRotation);
+											FComponentEditorUtils::ApplyDefaultValueChange(SceneComp, SceneComp->RelativeScale3D, OldRelativeScale3D, SelectedTemplate->RelativeScale3D);
 										}
 									}
 								}
