[UE-15981] Fix undo/redo not applying to the preview instance in the Blueprint class editor.

change summary:
- modified FBlueprintEditor::UpdatePreviewActor() to ensure that it sets RF_Transactional on the spawned instance. this is now needed because we added a UActorComponent::Modify() override back in CL# 2446659 which redirects constructed component Modify() calls to the owning Actor for component data cache purposes.
- also, another minor fix in SSCSEditor::Tick() which covers an edge case I stumbled on that could result in log spam in the Blueprint editor (only when the option to show UCS-added components is turned on, which it is off by default).

diff --git a/Engine/Source/Editor/Kismet/Private/BlueprintEditor.cpp b/Engine/Source/Editor/Kismet/Private/BlueprintEditor.cpp
--- a/Engine/Source/Editor/Kismet/Private/BlueprintEditor.cpp
+++ b/Engine/Source/Editor/Kismet/Private/BlueprintEditor.cpp
@@ -7284,7 +7284,7 @@ void FBlueprintEditor::UpdatePreviewActor(UBlueprint* InBlueprint, bool bInForce
 			FActorSpawnParameters SpawnInfo;
 			SpawnInfo.bNoCollisionFail = true;
 			SpawnInfo.bNoFail = true;
-			SpawnInfo.ObjectFlags = RF_Transient;
+			SpawnInfo.ObjectFlags = RF_Transient|RF_Transactional;
 
 			// Temporarily remove the deprecated flag so we can respawn the Blueprint in the viewport
 			bool bIsClassDeprecated = PreviewBlueprint->GeneratedClass->HasAnyClassFlags(CLASS_Deprecated);
diff --git a/Engine/Source/Editor/Kismet/Private/SSCSEditor.cpp b/Engine/Source/Editor/Kismet/Private/SSCSEditor.cpp
--- a/Engine/Source/Editor/Kismet/Private/SSCSEditor.cpp
+++ b/Engine/Source/Editor/Kismet/Private/SSCSEditor.cpp
@@ -3562,11 +3562,11 @@ void SSCSEditor::Tick( const FGeometry& AllottedGeometry, const double InCurrent
 			int32 NumComponentInstances = 0;
 			for (auto CompIt = ActorInstance->GetComponents().CreateConstIterator(); CompIt; ++CompIt)
 			{
-				// Don't count editor-only components in instanced mode, because we don't show them there. Also, don't count UCS-added components if the option to hide them is enabled, and don't count native components that are based on the CDO in the BP editor context (i.e. preview world).
+				// Don't count editor-only components in instanced mode, because we don't show them there. Also, don't count UCS-added components if the option to hide them is enabled, and don't count components that are based on the CDO in the BP editor context (i.e. preview world).
 				const UActorComponent* CompInst = *CompIt;
 				if ((!CompInst->IsEditorOnly() || EditorMode != EComponentEditorMode::ActorInstance)
 					&& (CompInst->CreationMethod != EComponentCreationMethod::UserConstructionScript || bShowInstancedUCSComponents)
-					&& (EditorMode != EComponentEditorMode::BlueprintSCS || CompInst->CreationMethod != EComponentCreationMethod::Native || CompInst->GetArchetype() != CompInst->GetClass()->GetDefaultObject()))
+					&& (EditorMode != EComponentEditorMode::BlueprintSCS || CompInst->GetArchetype() != CompInst->GetClass()->GetDefaultObject()))
 				{
 					++NumComponentInstances;
 				}
