Added overrides for GetToolkitToolTipText where required

UE-16078 - Crash hovering over texture reference in material

A lot of asset editors were overriding GetToolkitName to implement custom logic (sometimes because the default logic would cause a crash), and these asset editors now also need to override GetToolkitToolTipText too.

diff --git a/Engine/Plugins/2D/Paper2D/Source/Paper2DEditor/Private/FlipbookEditor/FlipbookEditor.cpp b/Engine/Plugins/2D/Paper2D/Source/Paper2DEditor/Private/FlipbookEditor/FlipbookEditor.cpp
--- a/Engine/Plugins/2D/Paper2D/Source/Paper2DEditor/Private/FlipbookEditor/FlipbookEditor.cpp
+++ b/Engine/Plugins/2D/Paper2D/Source/Paper2DEditor/Private/FlipbookEditor/FlipbookEditor.cpp
@@ -419,6 +419,11 @@ FText FFlipbookEditor::GetToolkitName() const
 	return FText::Format(LOCTEXT("FlipbookEditorAppLabel", "{FlipbookName}{DirtyState}"), Args);
 }
 
+FText FFlipbookEditor::GetToolkitToolTipText() const
+{
+	return FAssetEditorToolkit::GetToolTipTextForObject(FlipbookBeingEdited);
+}
+
 FString FFlipbookEditor::GetWorldCentricTabPrefix() const
 {
 	return TEXT("FlipbookEditor");
diff --git a/Engine/Plugins/2D/Paper2D/Source/Paper2DEditor/Private/FlipbookEditor/FlipbookEditor.h b/Engine/Plugins/2D/Paper2D/Source/Paper2DEditor/Private/FlipbookEditor/FlipbookEditor.h
--- a/Engine/Plugins/2D/Paper2D/Source/Paper2DEditor/Private/FlipbookEditor/FlipbookEditor.h
+++ b/Engine/Plugins/2D/Paper2D/Source/Paper2DEditor/Private/FlipbookEditor/FlipbookEditor.h
@@ -26,6 +26,7 @@ public:
 	virtual FName GetToolkitFName() const override;
 	virtual FText GetBaseToolkitName() const override;
 	virtual FText GetToolkitName() const override;
+	virtual FText GetToolkitToolTipText() const override;
 	virtual FLinearColor GetWorldCentricTabColorScale() const override;
 	virtual FString GetWorldCentricTabPrefix() const override;
 	virtual FString GetDocumentationLink() const override;
diff --git a/Engine/Plugins/2D/Paper2D/Source/Paper2DEditor/Private/SpriteEditor/SpriteEditor.cpp b/Engine/Plugins/2D/Paper2D/Source/Paper2DEditor/Private/SpriteEditor/SpriteEditor.cpp
--- a/Engine/Plugins/2D/Paper2D/Source/Paper2DEditor/Private/SpriteEditor/SpriteEditor.cpp
+++ b/Engine/Plugins/2D/Paper2D/Source/Paper2DEditor/Private/SpriteEditor/SpriteEditor.cpp
@@ -467,6 +467,11 @@ FText FSpriteEditor::GetToolkitName() const
 	return FText::Format(LOCTEXT("SpriteEditorAppLabel", "{SpriteName}{DirtyState}"), Args);
 }
 
+FText FSpriteEditor::GetToolkitToolTipText() const
+{
+	return FAssetEditorToolkit::GetToolTipTextForObject(SpriteBeingEdited);
+}
+
 FString FSpriteEditor::GetWorldCentricTabPrefix() const
 {
 	return TEXT("SpriteEditor");
diff --git a/Engine/Plugins/2D/Paper2D/Source/Paper2DEditor/Private/SpriteEditor/SpriteEditor.h b/Engine/Plugins/2D/Paper2D/Source/Paper2DEditor/Private/SpriteEditor/SpriteEditor.h
--- a/Engine/Plugins/2D/Paper2D/Source/Paper2DEditor/Private/SpriteEditor/SpriteEditor.h
+++ b/Engine/Plugins/2D/Paper2D/Source/Paper2DEditor/Private/SpriteEditor/SpriteEditor.h
@@ -37,6 +37,7 @@ public:
 	virtual FName GetToolkitFName() const override;
 	virtual FText GetBaseToolkitName() const override;
 	virtual FText GetToolkitName() const override;
+	virtual FText GetToolkitToolTipText() const override;
 	virtual FLinearColor GetWorldCentricTabColorScale() const override;
 	virtual FString GetWorldCentricTabPrefix() const override;
 	virtual FString GetDocumentationLink() const override;
diff --git a/Engine/Plugins/2D/Paper2D/Source/Paper2DEditor/Private/TileMapEditing/TileMapEdModeToolkit.cpp b/Engine/Plugins/2D/Paper2D/Source/Paper2DEditor/Private/TileMapEditing/TileMapEdModeToolkit.cpp
--- a/Engine/Plugins/2D/Paper2D/Source/Paper2DEditor/Private/TileMapEditing/TileMapEdModeToolkit.cpp
+++ b/Engine/Plugins/2D/Paper2D/Source/Paper2DEditor/Private/TileMapEditing/TileMapEdModeToolkit.cpp
@@ -54,6 +54,15 @@ FText FTileMapEdModeToolkit::GetToolkitName() const
 	return GetBaseToolkitName();
 }
 
+FText FTileMapEdModeToolkit::GetToolkitToolTipText() const
+{
+	if (CurrentTileSetPtr.IsValid())
+	{
+		return FAssetEditorToolkit::GetToolTipTextForObject(CurrentTileSetPtr.Get());
+	}
+	return GetBaseToolkitName();
+}
+
 FEdMode* FTileMapEdModeToolkit::GetEditorMode() const
 {
 	return TileMapEditor;
diff --git a/Engine/Plugins/2D/Paper2D/Source/Paper2DEditor/Private/TileMapEditing/TileMapEdModeToolkit.h b/Engine/Plugins/2D/Paper2D/Source/Paper2DEditor/Private/TileMapEditing/TileMapEdModeToolkit.h
--- a/Engine/Plugins/2D/Paper2D/Source/Paper2DEditor/Private/TileMapEditing/TileMapEdModeToolkit.h
+++ b/Engine/Plugins/2D/Paper2D/Source/Paper2DEditor/Private/TileMapEditing/TileMapEdModeToolkit.h
@@ -18,6 +18,7 @@ public:
 	virtual FName GetToolkitFName() const override;
 	virtual FText GetBaseToolkitName() const override;
 	virtual FText GetToolkitName() const override;
+	virtual FText GetToolkitToolTipText() const override;
 	virtual class FEdMode* GetEditorMode() const override;
 	virtual TSharedPtr<class SWidget> GetInlineContent() const override;
 	// End of IToolkit interface
diff --git a/Engine/Plugins/2D/Paper2D/Source/Paper2DEditor/Private/TileMapEditing/TileMapEditor.cpp b/Engine/Plugins/2D/Paper2D/Source/Paper2DEditor/Private/TileMapEditing/TileMapEditor.cpp
--- a/Engine/Plugins/2D/Paper2D/Source/Paper2DEditor/Private/TileMapEditing/TileMapEditor.cpp
+++ b/Engine/Plugins/2D/Paper2D/Source/Paper2DEditor/Private/TileMapEditing/TileMapEditor.cpp
@@ -390,6 +390,11 @@ FText FTileMapEditor::GetToolkitName() const
 	return FText::Format(LOCTEXT("TileMapEditorAppLabel", "{TileMapName}{DirtyState}"), Args);
 }
 
+FText FTileMapEditor::GetToolkitToolTipText() const
+{
+	return GetToolTipTextForObject(TileMapBeingEdited);
+}
+
 FString FTileMapEditor::GetWorldCentricTabPrefix() const
 {
 	return TEXT("TileMapEditor");
diff --git a/Engine/Plugins/2D/Paper2D/Source/Paper2DEditor/Private/TileMapEditing/TileMapEditor.h b/Engine/Plugins/2D/Paper2D/Source/Paper2DEditor/Private/TileMapEditing/TileMapEditor.h
--- a/Engine/Plugins/2D/Paper2D/Source/Paper2DEditor/Private/TileMapEditing/TileMapEditor.h
+++ b/Engine/Plugins/2D/Paper2D/Source/Paper2DEditor/Private/TileMapEditing/TileMapEditor.h
@@ -20,6 +20,7 @@ public:
 	virtual FName GetToolkitFName() const override;
 	virtual FText GetBaseToolkitName() const override;
 	virtual FText GetToolkitName() const override;
+	virtual FText GetToolkitToolTipText() const override;
 	virtual FLinearColor GetWorldCentricTabColorScale() const override;
 	virtual FString GetWorldCentricTabPrefix() const override;
 	virtual FString GetDocumentationLink() const override;
diff --git a/Engine/Plugins/2D/Paper2D/Source/Paper2DEditor/Private/TileSetEditor.cpp b/Engine/Plugins/2D/Paper2D/Source/Paper2DEditor/Private/TileSetEditor.cpp
--- a/Engine/Plugins/2D/Paper2D/Source/Paper2DEditor/Private/TileSetEditor.cpp
+++ b/Engine/Plugins/2D/Paper2D/Source/Paper2DEditor/Private/TileSetEditor.cpp
@@ -438,6 +438,11 @@ FText FTileSetEditor::GetToolkitName() const
 	return FText::Format(LOCTEXT("TileSetAppLabel", "{TileSetName}{DirtyState}"), Args);
 }
 
+FText FTileSetEditor::GetToolkitToolTipText() const
+{
+	return FAssetEditorToolkit::GetToolTipTextForObject(TileSetBeingEdited);
+}
+
 FString FTileSetEditor::GetWorldCentricTabPrefix() const
 {
 	return TEXT("TileSetEditor");
diff --git a/Engine/Plugins/2D/Paper2D/Source/Paper2DEditor/Private/TileSetEditor.h b/Engine/Plugins/2D/Paper2D/Source/Paper2DEditor/Private/TileSetEditor.h
--- a/Engine/Plugins/2D/Paper2D/Source/Paper2DEditor/Private/TileSetEditor.h
+++ b/Engine/Plugins/2D/Paper2D/Source/Paper2DEditor/Private/TileSetEditor.h
@@ -28,6 +28,7 @@ public:
 	virtual FName GetToolkitFName() const override;
 	virtual FText GetBaseToolkitName() const override;
 	virtual FText GetToolkitName() const override;
+	virtual FText GetToolkitToolTipText() const override;
 	virtual FLinearColor GetWorldCentricTabColorScale() const override;
 	virtual FString GetWorldCentricTabPrefix() const override;
 	virtual FString GetDocumentationLink() const override;
diff --git a/Engine/Plugins/Experimental/CodeEditor/Source/CodeEditor/Private/CodeProjectEditor.cpp b/Engine/Plugins/Experimental/CodeEditor/Source/CodeEditor/Private/CodeProjectEditor.cpp
--- a/Engine/Plugins/Experimental/CodeEditor/Source/CodeEditor/Private/CodeProjectEditor.cpp
+++ b/Engine/Plugins/Experimental/CodeEditor/Source/CodeEditor/Private/CodeProjectEditor.cpp
@@ -267,6 +267,11 @@ FText FCodeProjectEditor::GetToolkitName() const
 	return LOCTEXT("CodeAppLabel", "Code Editor");
 }
 
+FText FCodeProjectEditor::GetToolkitToolTipText() const
+{
+	return LOCTEXT("CodeAppLabel", "Code Editor");
+}
+
 FString FCodeProjectEditor::GetWorldCentricTabPrefix() const
 {
 	return TEXT("CodeEditor");
diff --git a/Engine/Plugins/Experimental/CodeEditor/Source/CodeEditor/Private/CodeProjectEditor.h b/Engine/Plugins/Experimental/CodeEditor/Source/CodeEditor/Private/CodeProjectEditor.h
--- a/Engine/Plugins/Experimental/CodeEditor/Source/CodeEditor/Private/CodeProjectEditor.h
+++ b/Engine/Plugins/Experimental/CodeEditor/Source/CodeEditor/Private/CodeProjectEditor.h
@@ -17,6 +17,7 @@ public:
 	virtual FName GetToolkitFName() const override;
 	virtual FText GetBaseToolkitName() const override;
 	virtual FText GetToolkitName() const override;
+	virtual FText GetToolkitToolTipText() const override;
 	virtual FLinearColor GetWorldCentricTabColorScale() const override;
 	virtual FString GetWorldCentricTabPrefix() const override;
 	// End of FAssetEditorToolkit
diff --git a/Engine/Source/Editor/GameplayAbilitiesEditor/Private/GameplayAbilitiesEditor.cpp b/Engine/Source/Editor/GameplayAbilitiesEditor/Private/GameplayAbilitiesEditor.cpp
--- a/Engine/Source/Editor/GameplayAbilitiesEditor/Private/GameplayAbilitiesEditor.cpp
+++ b/Engine/Source/Editor/GameplayAbilitiesEditor/Private/GameplayAbilitiesEditor.cpp
@@ -86,6 +86,15 @@ FText FGameplayAbilitiesEditor::GetToolkitName() const
 	return FText::Format(LOCTEXT("GameplayAbilitiesEditorAppLabel", "{ObjectName}{DirtyState}"), Args);
 }
 
+FText FGameplayAbilitiesEditor::GetToolkitToolTipText() const
+{
+	const UObject* EditingObject = GetEditingObject();
+
+	check (EditingObject != NULL);
+
+	return FAssetEditorToolkit::GetToolTipTextForObject(EditingObject);
+}
+
 FString FGameplayAbilitiesEditor::GetWorldCentricTabPrefix() const
 {
 	return TEXT("GameplayAbilitiesEditor");
diff --git a/Engine/Source/Editor/GameplayAbilitiesEditor/Public/GameplayAbilitiesEditor.h b/Engine/Source/Editor/GameplayAbilitiesEditor/Public/GameplayAbilitiesEditor.h
--- a/Engine/Source/Editor/GameplayAbilitiesEditor/Public/GameplayAbilitiesEditor.h
+++ b/Engine/Source/Editor/GameplayAbilitiesEditor/Public/GameplayAbilitiesEditor.h
@@ -46,6 +46,7 @@ public:
 	virtual FName GetToolkitFName() const override;
 	virtual FText GetBaseToolkitName() const override;
 	virtual FText GetToolkitName() const override;
+	virtual FText GetToolkitToolTipText() const override;
 	virtual FString GetWorldCentricTabPrefix() const override;
 	virtual FLinearColor GetWorldCentricTabColorScale() const override;
 	// End of IToolkit interface
diff --git a/Engine/Source/Editor/Kismet/Private/BlueprintEditor.cpp b/Engine/Source/Editor/Kismet/Private/BlueprintEditor.cpp
--- a/Engine/Source/Editor/Kismet/Private/BlueprintEditor.cpp
+++ b/Engine/Source/Editor/Kismet/Private/BlueprintEditor.cpp
@@ -6780,6 +6780,52 @@ FText FBlueprintEditor::GetToolkitName() const
 	return FText::Format( NSLOCTEXT("KismetEditor", "ToolkitTitle_UniqueLayerName", "{NumberOfObjects} {ClassName} - Class Defaults"), Args );
 }
 
+FText FBlueprintEditor::GetToolkitToolTipText() const
+{
+	const auto EditingObjects = GetEditingObjects();
+
+	if( IsEditingSingleBlueprint() )
+	{
+		if (FBlueprintEditorUtils::IsLevelScriptBlueprint(GetBlueprintObj()))
+		{
+			const FString& LevelName = FPackageName::GetShortFName( GetBlueprintObj()->GetOutermost()->GetFName().GetPlainNameString() ).GetPlainNameString();	
+
+			FFormatNamedArguments Args;
+			Args.Add( TEXT("LevelName"), FText::FromString( LevelName ) );
+			return FText::Format( NSLOCTEXT("KismetEditor", "LevelScriptAppToolTip", "{LevelName} - Level Blueprint Editor"), Args );
+		}
+		else
+		{
+			return FAssetEditorToolkit::GetToolTipTextForObject( GetBlueprintObj() );
+		}
+	}
+
+	TSubclassOf< UObject > SharedParentClass = NULL;
+
+	for( auto ObjectIter = EditingObjects.CreateConstIterator(); ObjectIter; ++ObjectIter )
+	{
+		UBlueprint* Blueprint = Cast<UBlueprint>( *ObjectIter );;
+		check( Blueprint );
+
+		// Initialize with the class of the first object we encounter.
+		if( *SharedParentClass == NULL )
+		{
+			SharedParentClass = Blueprint->ParentClass;
+		}
+
+		// If we've encountered an object that's not a subclass of the current best baseclass,
+		// climb up a step in the class hierarchy.
+		while( !Blueprint->ParentClass->IsChildOf( SharedParentClass ) )
+		{
+			SharedParentClass = SharedParentClass->GetSuperClass();
+		}
+	}
+
+	FFormatNamedArguments Args;
+	Args.Add( TEXT("NumberOfObjects"), EditingObjects.Num() );
+	Args.Add( TEXT("ObjectName"), FText::FromString( SharedParentClass->GetName() ) );
+	return FText::Format( NSLOCTEXT("KismetEditor", "ToolkitTitle_UniqueLayerName", "{NumberOfObjects} {ClassName} - Class Defaults"), Args );
+}
 
 FLinearColor FBlueprintEditor::GetWorldCentricTabColorScale() const
 {
diff --git a/Engine/Source/Editor/Kismet/Private/UserDefinedEnumEditor.cpp b/Engine/Source/Editor/Kismet/Private/UserDefinedEnumEditor.cpp
--- a/Engine/Source/Editor/Kismet/Private/UserDefinedEnumEditor.cpp
+++ b/Engine/Source/Editor/Kismet/Private/UserDefinedEnumEditor.cpp
@@ -127,6 +127,15 @@ FText FUserDefinedEnumEditor::GetToolkitName() const
 	return GetBaseToolkitName();
 }
 
+FText FUserDefinedEnumEditor::GetToolkitToolTipText() const
+{
+	if (1 == GetEditingObjects().Num())
+	{
+		return FAssetEditorToolkit::GetToolkitToolTipText();
+	}
+	return GetBaseToolkitName();
+}
+
 FString FUserDefinedEnumEditor::GetWorldCentricTabPrefix() const
 {
 	return LOCTEXT("UDEnumWorldCentricTabPrefix", "Enum ").ToString();
diff --git a/Engine/Source/Editor/Kismet/Private/UserDefinedEnumEditor.h b/Engine/Source/Editor/Kismet/Private/UserDefinedEnumEditor.h
--- a/Engine/Source/Editor/Kismet/Private/UserDefinedEnumEditor.h
+++ b/Engine/Source/Editor/Kismet/Private/UserDefinedEnumEditor.h
@@ -35,6 +35,7 @@ public:
 	virtual FName GetToolkitFName() const override;
 	virtual FText GetBaseToolkitName() const override;
 	virtual FText GetToolkitName() const override;
+	virtual FText GetToolkitToolTipText() const override;
 	virtual FString GetWorldCentricTabPrefix() const override;
 	virtual FLinearColor GetWorldCentricTabColorScale() const override;
 
diff --git a/Engine/Source/Editor/Kismet/Private/UserDefinedStructureEditor.cpp b/Engine/Source/Editor/Kismet/Private/UserDefinedStructureEditor.cpp
--- a/Engine/Source/Editor/Kismet/Private/UserDefinedStructureEditor.cpp
+++ b/Engine/Source/Editor/Kismet/Private/UserDefinedStructureEditor.cpp
@@ -243,6 +243,15 @@ FText FUserDefinedStructureEditor::GetToolkitName() const
 	return GetBaseToolkitName();
 }
 
+FText FUserDefinedStructureEditor::GetToolkitToolTipText() const
+{
+	if (1 == GetEditingObjects().Num())
+	{
+		return FAssetEditorToolkit::GetToolkitToolTipText();
+	}
+	return GetBaseToolkitName();
+}
+
 FString FUserDefinedStructureEditor::GetWorldCentricTabPrefix() const
 {
 	return LOCTEXT("UDStructWorldCentricTabPrefix", "Struct ").ToString();
diff --git a/Engine/Source/Editor/Kismet/Private/UserDefinedStructureEditor.h b/Engine/Source/Editor/Kismet/Private/UserDefinedStructureEditor.h
--- a/Engine/Source/Editor/Kismet/Private/UserDefinedStructureEditor.h
+++ b/Engine/Source/Editor/Kismet/Private/UserDefinedStructureEditor.h
@@ -33,6 +33,7 @@ public:
 	virtual FName GetToolkitFName() const override;
 	virtual FText GetBaseToolkitName() const override;
 	virtual FText GetToolkitName() const override;
+	virtual FText GetToolkitToolTipText() const override;
 	virtual FString GetWorldCentricTabPrefix() const override;
 	virtual FLinearColor GetWorldCentricTabColorScale() const override;
 
diff --git a/Engine/Source/Editor/Kismet/Public/BlueprintEditor.h b/Engine/Source/Editor/Kismet/Public/BlueprintEditor.h
--- a/Engine/Source/Editor/Kismet/Public/BlueprintEditor.h
+++ b/Engine/Source/Editor/Kismet/Public/BlueprintEditor.h
@@ -144,6 +144,7 @@ public:
 	virtual FName GetToolkitFName() const override;
 	virtual FText GetBaseToolkitName() const override;
 	virtual FText GetToolkitName() const override;
+	virtual FText GetToolkitToolTipText() const override;
 	virtual FString GetWorldCentricTabPrefix() const override;
 	virtual FLinearColor GetWorldCentricTabColorScale() const override;
 	virtual bool IsBlueprintEditor() const override;
diff --git a/Engine/Source/Editor/MaterialEditor/Private/MaterialEditor.cpp b/Engine/Source/Editor/MaterialEditor/Private/MaterialEditor.cpp
--- a/Engine/Source/Editor/MaterialEditor/Private/MaterialEditor.cpp
+++ b/Engine/Source/Editor/MaterialEditor/Private/MaterialEditor.cpp
@@ -721,6 +721,14 @@ FText FMaterialEditor::GetToolkitName() const
 	return FText::Format( LOCTEXT("MaterialEditorAppLabel", "{ObjectName}{DirtyState}"), Args );
 }
 
+FText FMaterialEditor::GetToolkitToolTipText() const
+{
+	const UObject* EditingObject = GetEditingObjects()[0];
+
+	// Overridden to accommodate editing of multiple objects (original and preview materials)
+	return FAssetEditorToolkit::GetToolTipTextForObject(EditingObject);
+}
+
 FString FMaterialEditor::GetWorldCentricTabPrefix() const
 {
 	return LOCTEXT("WorldCentricTabPrefix", "Material ").ToString();
diff --git a/Engine/Source/Editor/MaterialEditor/Private/MaterialEditor.h b/Engine/Source/Editor/MaterialEditor/Private/MaterialEditor.h
--- a/Engine/Source/Editor/MaterialEditor/Private/MaterialEditor.h
+++ b/Engine/Source/Editor/MaterialEditor/Private/MaterialEditor.h
@@ -214,6 +214,7 @@ public:
 	virtual FName GetToolkitFName() const override;
 	virtual FText GetBaseToolkitName() const override;
 	virtual FText GetToolkitName() const override;
+	virtual FText GetToolkitToolTipText() const override;
 	virtual FString GetWorldCentricTabPrefix() const override;
 
 	/** @return the documentation location for this editor */
diff --git a/Engine/Source/Editor/Persona/Private/Persona.cpp b/Engine/Source/Editor/Persona/Private/Persona.cpp
--- a/Engine/Source/Editor/Persona/Private/Persona.cpp
+++ b/Engine/Source/Editor/Persona/Private/Persona.cpp
@@ -1746,6 +1746,18 @@ FText FPersona::GetToolkitName() const
 	}
 }
 
+FText FPersona::GetToolkitToolTipText() const
+{
+	if (IsEditingSingleBlueprint())
+	{
+		return FAssetEditorToolkit::GetToolTipTextForObject(GetBlueprintObj());
+	}
+	else
+	{
+		check(TargetSkeleton != NULL);
+		return FAssetEditorToolkit::GetToolTipTextForObject(TargetSkeleton);
+	}
+}
 
 FString FPersona::GetWorldCentricTabPrefix() const
 {
diff --git a/Engine/Source/Editor/Persona/Private/Persona.h b/Engine/Source/Editor/Persona/Private/Persona.h
--- a/Engine/Source/Editor/Persona/Private/Persona.h
+++ b/Engine/Source/Editor/Persona/Private/Persona.h
@@ -176,6 +176,7 @@ public:
 	virtual FName GetToolkitFName() const override;
 	virtual FText GetBaseToolkitName() const override;
 	virtual FText GetToolkitName() const override;
+	virtual FText GetToolkitToolTipText() const override;
 	virtual FString GetWorldCentricTabPrefix() const override;
 	virtual FLinearColor GetWorldCentricTabColorScale() const override;
 	virtual void SaveAsset_Execute() override;
diff --git a/Engine/Source/Editor/PropertyEditor/Private/PropertyEditorToolkit.cpp b/Engine/Source/Editor/PropertyEditor/Private/PropertyEditorToolkit.cpp
--- a/Engine/Source/Editor/PropertyEditor/Private/PropertyEditorToolkit.cpp
+++ b/Engine/Source/Editor/PropertyEditor/Private/PropertyEditorToolkit.cpp
@@ -550,6 +550,54 @@ FText FPropertyEditorToolkit::GetToolkitName() const
 }
 
 
+FText FPropertyEditorToolkit::GetToolkitToolTipText() const
+{
+	const auto EditingObjects = GetEditingObjects();
+
+	check( EditingObjects.Num() > 0 );
+
+	if( EditingObjects.Num() == 1 )
+	{
+		const UObject* EditingObject = EditingObjects[ 0 ];
+		return FAssetEditorToolkit::GetToolTipTextForObject(EditingObject);
+	}
+	else
+	{
+		UClass* SharedBaseClass = NULL;
+		for( int32 x = 0; x < EditingObjects.Num(); ++x )
+		{
+			UObject* Obj = EditingObjects[ x ];
+			check( Obj );
+
+			UClass* ObjClass = Cast<UClass>(Obj);
+			if (ObjClass == NULL)
+			{
+				ObjClass = Obj->GetClass();
+			}
+			check( ObjClass );
+
+			// Initialize with the class of the first object we encounter.
+			if( SharedBaseClass == NULL )
+			{
+				SharedBaseClass = ObjClass;
+			}
+
+			// If we've encountered an object that's not a subclass of the current best baseclass,
+			// climb up a step in the class hierarchy.
+			while( !ObjClass->IsChildOf( SharedBaseClass ) )
+			{
+				SharedBaseClass = SharedBaseClass->GetSuperClass();
+			}
+		}
+
+		FFormatNamedArguments Args;
+		Args.Add( TEXT("NumberOfObjects"), EditingObjects.Num() );
+		Args.Add( TEXT("ClassName"), FText::FromString( SharedBaseClass->GetName() ) );
+		return FText::Format( LOCTEXT("ToolkitName_MultiObjectToolTip", "{NumberOfObjects} {ClassName} Objects - Property Matrix Editor"), Args );
+	}
+}
+
+
 FReply FPropertyEditorToolkit::OnToggleColumnClicked( const TWeakPtr< IPropertyTreeRow > Row )
 {
 	if (Row.IsValid())
diff --git a/Engine/Source/Editor/PropertyEditor/Private/PropertyEditorToolkit.h b/Engine/Source/Editor/PropertyEditor/Private/PropertyEditorToolkit.h
--- a/Engine/Source/Editor/PropertyEditor/Private/PropertyEditorToolkit.h
+++ b/Engine/Source/Editor/PropertyEditor/Private/PropertyEditorToolkit.h
@@ -20,6 +20,8 @@ public:
 
 	virtual FText GetToolkitName() const override;
 
+	virtual FText GetToolkitToolTipText() const override;
+
 	virtual FString GetWorldCentricTabPrefix() const override;
 
 	virtual FLinearColor GetWorldCentricTabColorScale() const override;
diff --git a/Engine/Source/Editor/TranslationEditor/Private/TranslationEditor.cpp b/Engine/Source/Editor/TranslationEditor/Private/TranslationEditor.cpp
--- a/Engine/Source/Editor/TranslationEditor/Private/TranslationEditor.cpp
+++ b/Engine/Source/Editor/TranslationEditor/Private/TranslationEditor.cpp
@@ -229,6 +229,19 @@ FText FTranslationEditor::GetToolkitName() const
 	return FText::Format( LOCTEXT("TranslationEditorAppLabel", "{Language}{DirtyState} - {ProjectName} - {ToolkitName}"), Args );
 }
 
+FText FTranslationEditor::GetToolkitToolTipText() const
+{
+	const UObject* EditingObject = GetEditingObject();
+
+	check (EditingObject != NULL);
+
+	FFormatNamedArguments Args;
+	Args.Add(TEXT("Language"), FText::FromString(FPaths::GetBaseFilename(FPaths::GetPath(ArchiveFilePath))));
+	Args.Add(TEXT("ProjectName"), FText::FromString(FPaths::GetBaseFilename(ManifestFilePath)));
+	Args.Add( TEXT("ToolkitName"), GetBaseToolkitName() );
+	return FText::Format( LOCTEXT("TranslationEditorAppToolTip", "{Language} - {ProjectName} - {ToolkitName}"), Args );
+}
+
 FString FTranslationEditor::GetWorldCentricTabPrefix() const
 {
 	return LOCTEXT("WorldCentricTabPrefix", "Translation ").ToString();
diff --git a/Engine/Source/Editor/TranslationEditor/Private/TranslationEditor.h b/Engine/Source/Editor/TranslationEditor/Private/TranslationEditor.h
--- a/Engine/Source/Editor/TranslationEditor/Private/TranslationEditor.h
+++ b/Engine/Source/Editor/TranslationEditor/Private/TranslationEditor.h
@@ -49,6 +49,7 @@ public:
 	virtual FName GetToolkitFName() const override;
 	virtual FText GetBaseToolkitName() const override;
 	virtual FText GetToolkitName() const override;
+	virtual FText GetToolkitToolTipText() const override;
 	virtual FString GetWorldCentricTabPrefix() const override;
 	virtual FLinearColor GetWorldCentricTabColorScale() const override;
 
diff --git a/Engine/Source/Editor/UnrealEd/Private/Toolkits/AssetEditorToolkit.cpp b/Engine/Source/Editor/UnrealEd/Private/Toolkits/AssetEditorToolkit.cpp
--- a/Engine/Source/Editor/UnrealEd/Private/Toolkits/AssetEditorToolkit.cpp
+++ b/Engine/Source/Editor/UnrealEd/Private/Toolkits/AssetEditorToolkit.cpp
@@ -88,7 +88,7 @@ void FAssetEditorToolkit::InitAssetEditor( const EToolkitMode::Type Mode, const
 		NewMajorTab = SNew(SDockTab) 
 			.ContentPadding(0.0f) 
 			.TabRole(ETabRole::MajorTab)
-			.ToolTip(IDocumentation::Get()->CreateToolTip(Label, nullptr, DocLink, GetToolkitFName().ToString()))
+			.ToolTip(IDocumentation::Get()->CreateToolTip(ToolTipText, nullptr, DocLink, GetToolkitFName().ToString()))
 			.Icon( this, &FAssetEditorToolkit::GetDefaultTabIcon )
 			.Label( Label );
 
diff --git a/Engine/Source/Editor/UnrealEd/Private/Toolkits/SimpleAssetEditor.cpp b/Engine/Source/Editor/UnrealEd/Private/Toolkits/SimpleAssetEditor.cpp
--- a/Engine/Source/Editor/UnrealEd/Private/Toolkits/SimpleAssetEditor.cpp
+++ b/Engine/Source/Editor/UnrealEd/Private/Toolkits/SimpleAssetEditor.cpp
@@ -171,6 +171,54 @@ FText FSimpleAssetEditor::GetToolkitName() const
 	}
 }
 
+FText FSimpleAssetEditor::GetToolkitToolTipText() const
+{
+	const auto EditingObjects = GetEditingObjects();
+
+	check( EditingObjects.Num() > 0 );
+
+	FFormatNamedArguments Args;
+	Args.Add( TEXT("ToolkitName"), GetBaseToolkitName() );
+
+	if( EditingObjects.Num() == 1 )
+	{
+		const UObject* EditingObject = EditingObjects[ 0 ];
+		return FAssetEditorToolkit::GetToolTipTextForObject(EditingObject);
+	}
+	else
+	{
+		UClass* SharedBaseClass = NULL;
+		for( int32 x = 0; x < EditingObjects.Num(); ++x )
+		{
+			UObject* Obj = EditingObjects[ x ];
+			check( Obj );
+
+			UClass* ObjClass = Cast<UClass>(Obj);
+			if (ObjClass == NULL)
+			{
+				ObjClass = Obj->GetClass();
+			}
+			check( ObjClass );
+
+			// Initialize with the class of the first object we encounter.
+			if( SharedBaseClass == NULL )
+			{
+				SharedBaseClass = ObjClass;
+			}
+
+			// If we've encountered an object that's not a subclass of the current best baseclass,
+			// climb up a step in the class hierarchy.
+			while( !ObjClass->IsChildOf( SharedBaseClass ) )
+			{
+				SharedBaseClass = SharedBaseClass->GetSuperClass();
+			}
+		}
+
+		Args.Add( TEXT("NumberOfObjects"), EditingObjects.Num() );
+		Args.Add( TEXT("ClassName"), FText::FromString( SharedBaseClass->GetName() ) );
+		return FText::Format( LOCTEXT("ToolkitTitle_EditingMultipleToolTip", "{NumberOfObjects} {ClassName} - {ToolkitName}"), Args );
+	}
+}
 
 FLinearColor FSimpleAssetEditor::GetWorldCentricTabColorScale() const
 {
diff --git a/Engine/Source/Editor/UnrealEd/Public/Toolkits/SimpleAssetEditor.h b/Engine/Source/Editor/UnrealEd/Public/Toolkits/SimpleAssetEditor.h
--- a/Engine/Source/Editor/UnrealEd/Public/Toolkits/SimpleAssetEditor.h
+++ b/Engine/Source/Editor/UnrealEd/Public/Toolkits/SimpleAssetEditor.h
@@ -33,6 +33,7 @@ public:
 	virtual FName GetToolkitFName() const override;
 	virtual FText GetBaseToolkitName() const override;
 	virtual FText GetToolkitName() const override;
+	virtual FText GetToolkitToolTipText() const override;
 	virtual FString GetWorldCentricTabPrefix() const override;
 	virtual FLinearColor GetWorldCentricTabColorScale() const override;
 	virtual bool IsPrimaryEditor() const override { return true; }
