Optimizing FFindInBlueprintSearchManager::CleanCache, we were spending a lot of time copying things around - especially evident when running the compile on load automation tests

diff --git a/Engine/Source/Editor/Kismet/Private/FindInBlueprintManager.cpp b/Engine/Source/Editor/Kismet/Private/FindInBlueprintManager.cpp
--- a/Engine/Source/Editor/Kismet/Private/FindInBlueprintManager.cpp
+++ b/Engine/Source/Editor/Kismet/Private/FindInBlueprintManager.cpp
@@ -1299,7 +1299,7 @@ void FFindInBlueprintSearchManager::CleanCache()
 		if( !SearchArray[SearchValuePair.Value].bMarkedForDeletion && !(SearchArray[SearchValuePair.Value].Blueprint.IsValid() && SearchArray[SearchValuePair.Value].Blueprint->HasAllFlags(RF_PendingKill)) )
 		{
 			// Build the new map/array
-			NewSearchMap.Add(SearchValuePair.Key, NewSearchArray.Add(SearchArray[SearchValuePair.Value]) );
+			NewSearchMap.Add(SearchValuePair.Key, NewSearchArray.Add(MoveTemp(SearchArray[SearchValuePair.Value])) );
 		}
 		else
 		{
@@ -1325,8 +1325,8 @@ void FFindInBlueprintSearchManager::CleanCache()
 		}
 	}
 
-	SearchMap = NewSearchMap;
-	SearchArray = NewSearchArray;
+	SearchMap = MoveTemp( NewSearchMap );
+	SearchArray = MoveTemp( NewSearchArray );
 
 	// After the search, we have to place the active search queries where they belong
 	for( auto& CacheQuery : CacheQueries )
diff --git a/Engine/Source/Editor/Kismet/Public/FindInBlueprintManager.h b/Engine/Source/Editor/Kismet/Public/FindInBlueprintManager.h
--- a/Engine/Source/Editor/Kismet/Public/FindInBlueprintManager.h
+++ b/Engine/Source/Editor/Kismet/Public/FindInBlueprintManager.h
@@ -85,7 +85,57 @@ struct FSearchData
 		: Blueprint(nullptr)
 		, bMarkedForDeletion(false)
 	{
+	}
+
+	// Adding move semantics (on Mac this counts as a user defined constructor, so I have to reimplement the default
+	// copy constructor. PLATFORM_COMPILER_HAS_DEFAULTED_FUNCTIONS is not true on all compilers yet, so I'm implementing manually:
+	FSearchData(FSearchData&& Other)
+		: Blueprint(Other.Blueprint)
+		, BlueprintPath(MoveTemp(Other.BlueprintPath))
+		, Value(MoveTemp(Other.Value))
+		, bMarkedForDeletion(Other.bMarkedForDeletion)
+	{
+	}
 
+	FSearchData& operator=(FSearchData&& RHS)
+	{
+		if (this == &RHS)
+		{
+			return *this;
+		}
+
+		Blueprint = RHS.Blueprint;
+		BlueprintPath = MoveTemp(RHS.BlueprintPath);
+		Value = MoveTemp(RHS.Value);
+		bMarkedForDeletion = RHS.bMarkedForDeletion;
+		return *this;
+	}
+
+	FSearchData(const FSearchData& Other)
+		: Blueprint(Other.Blueprint)
+		, BlueprintPath(Other.BlueprintPath)
+		, Value(Other.Value)
+		, bMarkedForDeletion(Other.bMarkedForDeletion)
+	{
+	}
+
+	FSearchData& operator=(const FSearchData& RHS)
+	{
+		if (this == &RHS)
+		{
+			return *this;
+		}
+
+		Blueprint = RHS.Blueprint;
+		BlueprintPath = RHS.BlueprintPath;
+		Value = RHS.Value;
+		bMarkedForDeletion = RHS.bMarkedForDeletion;
+		return *this;
+	}
+
+
+	~FSearchData()
+	{
 	}
 };
 
