[INTEGRATE] Change 2566786 by [REDACTED_PERSON][REDACTED_HANDLE] on 2015/05/27 12:05:43

	Fix non-unity

diff --git a/Engine/Source/Editor/Kismet/Private/BlueprintDetailsCustomization.h b/Engine/Source/Editor/Kismet/Private/BlueprintDetailsCustomization.h
--- a/Engine/Source/Editor/Kismet/Private/BlueprintDetailsCustomization.h
+++ b/Engine/Source/Editor/Kismet/Private/BlueprintDetailsCustomization.h
@@ -739,7 +739,7 @@ protected:
 
 	void PopulateVariableCategories();
 	
-	void AddExperimentalWarningCategory( IDetailLayoutBuilder& DetailBuilder, const TArray<FSCSEditorTreeNodePtrType>& Nodes );
+	void AddExperimentalWarningCategory( IDetailLayoutBuilder& DetailBuilder, const TArray<TSharedPtr<class FSCSEditorTreeNode>>& Nodes );
 
 private:
 	/** Weak reference to the Blueprint editor */
