Blueprint context menu is back to showing static functions, when context sensitive is disabled.

diff --git a/Engine/Source/Editor/Kismet/Private/BlueprintActionMenuUtils.cpp b/Engine/Source/Editor/Kismet/Private/BlueprintActionMenuUtils.cpp
--- a/Engine/Source/Editor/Kismet/Private/BlueprintActionMenuUtils.cpp
+++ b/Engine/Source/Editor/Kismet/Private/BlueprintActionMenuUtils.cpp
@@ -442,7 +442,7 @@ void FBlueprintActionMenuUtils::MakeContextMenu(FBlueprintActionContext const& C
 	//--------------------------------------
 
 	uint32 FilterFlags = 0x00;
-	if ((ClassTargetMask & EContextTargetFlags::TARGET_BlueprintLibraries) == 0)
+	if ( bIsContextSensitive && ((ClassTargetMask & EContextTargetFlags::TARGET_BlueprintLibraries) == 0) )
 	{
 		FilterFlags |= FBlueprintActionFilter::BPFILTER_RejectGlobalFields;
 	}
