Can right click on any node to search for references of it. Does initial search only in current Blueprint.

Localization fixes for Find-in-Blueprints so it will again search the current and native language.

CallFunction and CustomEvents now store their non-friendly "native" names. These will not actually appear in search results, but will be compared while searching to determine if the node matches the results.

Gathering Blueprint search data will now force friendly names to true while it gathers so gathered data does not differ based on the setting.

diff --git a/Engine/Source/Editor/BlueprintGraph/BlueprintGraph.Build.cs b/Engine/Source/Editor/BlueprintGraph/BlueprintGraph.Build.cs
--- a/Engine/Source/Editor/BlueprintGraph/BlueprintGraph.Build.cs
+++ b/Engine/Source/Editor/BlueprintGraph/BlueprintGraph.Build.cs
@@ -33,6 +33,7 @@ public class BlueprintGraph : ModuleRules
 				"UnrealEd",
                 "GraphEditor",
 				"SlateCore",
+                "Kismet",
 			}
 		);
 
@@ -41,6 +42,7 @@ public class BlueprintGraph : ModuleRules
                 "KismetCompiler",
                 "UnrealEd",
                 "GraphEditor",
+                "Kismet",
             }
 		); 
 	}
diff --git a/Engine/Source/Editor/BlueprintGraph/Classes/K2Node_CallFunction.h b/Engine/Source/Editor/BlueprintGraph/Classes/K2Node_CallFunction.h
--- a/Engine/Source/Editor/BlueprintGraph/Classes/K2Node_CallFunction.h
+++ b/Engine/Source/Editor/BlueprintGraph/Classes/K2Node_CallFunction.h
@@ -79,6 +79,7 @@ public:
 	virtual FName GetPaletteIcon(FLinearColor& OutColor) const override;
 	virtual bool CanPasteHere(const UEdGraph* TargetGraph) const override;
 	virtual void PinDefaultValueChanged(UEdGraphPin* Pin) override;
+	virtual void AddSearchMetaDataInfo(TArray<struct FSearchTagDataPair>& OutTaggedMetaData) const override;
 	// End of UEdGraphNode interface
 
 	// UK2Node interface
diff --git a/Engine/Source/Editor/BlueprintGraph/Classes/K2Node_CustomEvent.h b/Engine/Source/Editor/BlueprintGraph/Classes/K2Node_CustomEvent.h
--- a/Engine/Source/Editor/BlueprintGraph/Classes/K2Node_CustomEvent.h
+++ b/Engine/Source/Editor/BlueprintGraph/Classes/K2Node_CustomEvent.h
@@ -31,6 +31,7 @@ class UK2Node_CustomEvent : public UK2Node_Event
 	virtual FString GetDocumentationExcerptName() const override;
 	virtual FName GetPaletteIcon(FLinearColor& OutColor) const override;
 	virtual void AutowireNewNode(UEdGraphPin* FromPin) override;
+	virtual void AddSearchMetaDataInfo(TArray<struct FSearchTagDataPair>& OutTaggedMetaData) const override;
 	// End UEdGraphNode interface
 
 	// Begin UK2Node interface
diff --git a/Engine/Source/Editor/BlueprintGraph/Private/EdGraphSchema_K2.cpp b/Engine/Source/Editor/BlueprintGraph/Private/EdGraphSchema_K2.cpp
--- a/Engine/Source/Editor/BlueprintGraph/Private/EdGraphSchema_K2.cpp
+++ b/Engine/Source/Editor/BlueprintGraph/Private/EdGraphSchema_K2.cpp
@@ -1580,12 +1580,6 @@ void UEdGraphSchema_K2::GetContextMenuActions(const UEdGraph* CurrentGraph, cons
 						MenuBuilder->AddMenuEntry(FGraphEditorCommands::Get().RemoveOptionPin);
 					}
 
-					// Conditionally add the action to find instances of the node if it is a custom event
-					if (InGraphNode->IsA(UK2Node_CustomEvent::StaticClass()))
-					{
-						MenuBuilder->AddMenuEntry(FGraphEditorCommands::Get().FindInstancesOfCustomEvent);
-					}
-
 					// Don't show the "Assign selected Actor" option if more than one actor is selected
 					if (InGraphNode->IsA(UK2Node_ActorBoundEvent::StaticClass()) && GEditor->GetSelectedActorCount() == 1)
 					{
@@ -1612,10 +1606,11 @@ void UEdGraphSchema_K2::GetContextMenuActions(const UEdGraph* CurrentGraph, cons
 						MenuBuilder->AddMenuEntry(FGraphEditorCommands::Get().GoToDefinition);
 					}
 
+					MenuBuilder->AddMenuEntry(FGraphEditorCommands::Get().FindReferences);
+
 					// show search for references for variable nodes and goto source code action
 					if (InGraphNode->IsA(UK2Node_Variable::StaticClass()))
 					{
-						MenuBuilder->AddMenuEntry(FGraphEditorCommands::Get().FindVariableReferences);
 						MenuBuilder->AddMenuEntry(FGraphEditorCommands::Get().GotoNativeVariableDefinition);
 						GetReplaceVariableMenu(InGraphNode,OwnerBlueprint, MenuBuilder, true);
 					}
diff --git a/Engine/Source/Editor/BlueprintGraph/Private/K2Node_CallFunction.cpp b/Engine/Source/Editor/BlueprintGraph/Private/K2Node_CallFunction.cpp
--- a/Engine/Source/Editor/BlueprintGraph/Private/K2Node_CallFunction.cpp
+++ b/Engine/Source/Editor/BlueprintGraph/Private/K2Node_CallFunction.cpp
@@ -11,6 +11,7 @@
 #include "K2Node_PureAssignmentStatement.h"
 #include "GraphEditorSettings.h"
 #include "BlueprintActionFilter.h"
+#include "Editor/Kismet/Public/FindInBlueprintManager.h"
 
 #define LOCTEXT_NAMESPACE "K2Node"
 
@@ -1296,7 +1297,20 @@ void UK2Node_CallFunction::GeneratePinTooltipFromFunction(UEdGraphPin& Pin, cons
 
 FText UK2Node_CallFunction::GetUserFacingFunctionName(const UFunction* Function)
 {
-	return Function->GetDisplayNameText();
+	FText ReturnDisplayName;
+
+	if( GEditor && GetDefault<UEditorStyleSettings>()->bShowFriendlyNames )
+	{
+		ReturnDisplayName = Function->GetDisplayNameText();
+	}
+	else
+	{
+		static const FString Namespace = TEXT("UObjectDisplayNames");
+		const FString Key = Function->GetFullGroupName(false);
+
+		ReturnDisplayName = Function->GetMetaDataText(TEXT("DisplayName"), Namespace, Key);
+	}
+	return ReturnDisplayName;
 }
 
 FString UK2Node_CallFunction::GetDefaultTooltipForFunction(const UFunction* Function)
@@ -2289,4 +2303,14 @@ bool UK2Node_CallFunction::IsStructureWildcardProperty(const UFunction* Function
 	return false;
 }
 
+void UK2Node_CallFunction::AddSearchMetaDataInfo(TArray<struct FSearchTagDataPair>& OutTaggedMetaData) const
+{
+	Super::AddSearchMetaDataInfo(OutTaggedMetaData);
+
+	if (UFunction* TargetFunction = GetTargetFunction())
+	{
+		OutTaggedMetaData.Add(FSearchTagDataPair(FFindInBlueprintSearchTags::FiB_NativeName, FText::FromString(TargetFunction->GetName())));
+	}
+}
+
 #undef LOCTEXT_NAMESPACE
diff --git a/Engine/Source/Editor/BlueprintGraph/Private/K2Node_CustomEvent.cpp b/Engine/Source/Editor/BlueprintGraph/Private/K2Node_CustomEvent.cpp
--- a/Engine/Source/Editor/BlueprintGraph/Private/K2Node_CustomEvent.cpp
+++ b/Engine/Source/Editor/BlueprintGraph/Private/K2Node_CustomEvent.cpp
@@ -443,4 +443,20 @@ void UK2Node_CustomEvent::AutowireNewNode(UEdGraphPin* FromPin)
 	}
 }
 
+void UK2Node_CustomEvent::AddSearchMetaDataInfo(TArray<struct FSearchTagDataPair>& OutTaggedMetaData) const
+{
+	Super::AddSearchMetaDataInfo(OutTaggedMetaData);
+
+	for (FSearchTagDataPair& SearchData : OutTaggedMetaData)
+	{
+		// Should always be the first item, but there is no guarantee
+		if (SearchData.Key.CompareTo(FFindInBlueprintSearchTags::FiB_Name) == 0)
+		{
+			SearchData.Value = FText::FromString(FName::NameToDisplayString(CustomFunctionName.ToString(), false));
+			break;
+		}
+	}
+	OutTaggedMetaData.Add(FSearchTagDataPair(FFindInBlueprintSearchTags::FiB_NativeName, FText::FromName(CustomFunctionName)));
+}
+
 #undef LOCTEXT_NAMESPACE
diff --git a/Engine/Source/Editor/GraphEditor/Private/GraphEditorActions.cpp b/Engine/Source/Editor/GraphEditor/Private/GraphEditorActions.cpp
--- a/Engine/Source/Editor/GraphEditor/Private/GraphEditorActions.cpp
+++ b/Engine/Source/Editor/GraphEditor/Private/GraphEditorActions.cpp
@@ -44,7 +44,7 @@ void FGraphEditorCommandsImpl::RegisterCommands()
 
 	UI_COMMAND( SelectReferenceInLevel, "Find Actor in Level", "Select the actor referenced by this node in the level", EUserInterfaceActionType::Button, FInputChord() )
 	UI_COMMAND( AssignReferencedActor, "Assign selected Actor", "Assign the selected actor to be this node's referenced object", EUserInterfaceActionType::Button, FInputChord() )
-	UI_COMMAND( FindVariableReferences, "Find Variable References", "Find references of this variable", EUserInterfaceActionType::Button, FInputChord() )
+	UI_COMMAND( FindReferences, "Find References", "Find references of this item", EUserInterfaceActionType::Button, FInputChord() )
 	
 	UI_COMMAND( GotoNativeFunctionDefinition, "Goto Code Definition", "Goto the native code definition of this function", EUserInterfaceActionType::Button, FInputChord() )
 	UI_COMMAND( GotoNativeVariableDefinition, "Goto Code Definition", "Goto the native code definition of this variable", EUserInterfaceActionType::Button, FInputChord() )
@@ -72,8 +72,6 @@ void FGraphEditorCommandsImpl::RegisterCommands()
 
 	UI_COMMAND( CreateComment, "Create Comment", "Create a comment box", EUserInterfaceActionType::Button, FInputChord(EKeys::C))
 
-	UI_COMMAND( FindInstancesOfCustomEvent, "Find Instances of Event", "Find the instances of this custom event", EUserInterfaceActionType::Button, FInputChord() )
-
 	UI_COMMAND( ZoomIn, "Zoom In", "Zoom in on the graph editor", EUserInterfaceActionType::Button, FInputChord(EKeys::Add))
 	UI_COMMAND( ZoomOut, "Zoom Out", "Zoom out from the graph editor", EUserInterfaceActionType::Button, FInputChord(EKeys::Subtract))
 
diff --git a/Engine/Source/Editor/GraphEditor/Public/GraphEditorActions.h b/Engine/Source/Editor/GraphEditor/Public/GraphEditorActions.h
--- a/Engine/Source/Editor/GraphEditor/Public/GraphEditorActions.h
+++ b/Engine/Source/Editor/GraphEditor/Public/GraphEditorActions.h
@@ -64,8 +64,8 @@ public:
 	TSharedPtr< FUICommandInfo > SelectReferenceInLevel;
 	TSharedPtr< FUICommandInfo > AssignReferencedActor;
 
-	// Find variable references
-	TSharedPtr< FUICommandInfo > FindVariableReferences;
+	// Find references
+	TSharedPtr< FUICommandInfo > FindReferences;
 
 	// Goto native code actions
 	TSharedPtr< FUICommandInfo > GotoNativeFunctionDefinition;
@@ -102,9 +102,6 @@ public:
 	//create a comment node
 	TSharedPtr< FUICommandInfo > CreateComment;
 
-	// Find instances of a Custom event node
-	TSharedPtr< FUICommandInfo > FindInstancesOfCustomEvent;
-
 	// Zoom in and out on the graph editor
 	TSharedPtr< FUICommandInfo > ZoomIn;
 	TSharedPtr< FUICommandInfo > ZoomOut;
diff --git a/Engine/Source/Editor/Kismet/Kismet.Build.cs b/Engine/Source/Editor/Kismet/Kismet.Build.cs
--- a/Engine/Source/Editor/Kismet/Kismet.Build.cs
+++ b/Engine/Source/Editor/Kismet/Kismet.Build.cs
@@ -59,5 +59,11 @@ public class Kismet : ModuleRules
                 "BlueprintCompilerCppBackend",
 			}
             );
+
+        CircularlyReferencedDependentModules.AddRange(
+            new string[] {
+                "BlueprintGraph",
+            }
+        ); 
 	}
 }
diff --git a/Engine/Source/Editor/Kismet/Private/BlueprintEditor.cpp b/Engine/Source/Editor/Kismet/Private/BlueprintEditor.cpp
--- a/Engine/Source/Editor/Kismet/Private/BlueprintEditor.cpp
+++ b/Engine/Source/Editor/Kismet/Private/BlueprintEditor.cpp
@@ -1114,13 +1114,9 @@ TSharedRef<SGraphEditor> FBlueprintEditor::CreateGraphEditorWidget(TSharedRef<FT
 				FExecuteAction::CreateSP(this, &FBlueprintEditor::SetPinVisibility, SGraphEditor::Pin_HideNoConnectionNoDefault)
 				);
 
-			GraphEditorCommands->MapAction( FGraphEditorCommands::Get().FindInstancesOfCustomEvent,
-				FExecuteAction::CreateSP( this, &FBlueprintEditor::OnFindInstancesCustomEvent )
-				);
-
-			GraphEditorCommands->MapAction( FGraphEditorCommands::Get().FindVariableReferences,
-				FExecuteAction::CreateSP( this, &FBlueprintEditor::OnFindVariableReferences ),
-				FCanExecuteAction::CreateSP( this, &FBlueprintEditor::CanFindVariableReferences )
+			GraphEditorCommands->MapAction( FGraphEditorCommands::Get().FindReferences,
+				FExecuteAction::CreateSP( this, &FBlueprintEditor::OnFindReferences ),
+				FCanExecuteAction::CreateSP( this, &FBlueprintEditor::CanFindReferences )
 				);
 
 			GraphEditorCommands->MapAction( FGraphEditorCommands::Get().GotoNativeFunctionDefinition,
@@ -7102,37 +7098,6 @@ void FBlueprintEditor::SetPinVisibility(SGraphEditor::EPinVisibility Visibility)
 	OnSetPinVisibility.Broadcast(PinVisibility);
 }
 
-void FBlueprintEditor::OnFindVariableReferences()
-{
-	auto GraphEditor = FocusedGraphEdPtr.Pin();
-	if (GraphEditor.IsValid())
-	{
-		const FGraphPanelSelectionSet SelectedNodes = GraphEditor->GetSelectedNodes();
-		for (FGraphPanelSelectionSet::TConstIterator NodeIt(SelectedNodes); NodeIt; ++NodeIt)
-		{
-			UK2Node_Variable* VariableNode = Cast<UK2Node_Variable>(*NodeIt);
-			if (VariableNode)
-			{
-				SummonSearchUI(true, VariableNode->GetVarName().ToString());
-			}
-		}
-	}
-}
-
-bool FBlueprintEditor::CanFindVariableReferences()
-{
-	auto GraphEditor = FocusedGraphEdPtr.Pin();
-	if (GraphEditor.IsValid())
-	{
-		const FGraphPanelSelectionSet SelectedNodes = GraphEditor->GetSelectedNodes();
-		if( SelectedNodes.Num() == 1 )
-		{
-			return true;
-		}
-	}
-	return false;
-}
-
 void FBlueprintEditor::GotoNativeFunctionDefinition()
 {
 	auto GraphEditor = FocusedGraphEdPtr.Pin();
@@ -7248,27 +7213,55 @@ bool FBlueprintEditor::IsNativeCodeBrowsingAvailable() const
 	return bCodeBasedProject;
 }
 
-void FBlueprintEditor::OnFindInstancesCustomEvent()
+void FBlueprintEditor::OnFindReferences()
 {
 	auto GraphEditor = FocusedGraphEdPtr.Pin();
 	if (GraphEditor.IsValid())
 	{
+		FString SearchTerm;
+
 		const FGraphPanelSelectionSet SelectedNodes = GraphEditor->GetSelectedNodes();
 		for (FGraphPanelSelectionSet::TConstIterator NodeIt(SelectedNodes); NodeIt; ++NodeIt)
 		{
 			UEdGraphNode* SelectedNode = Cast<UEdGraphNode>(*NodeIt);
-			if (SelectedNode != NULL && SelectedNode->bCanRenameNode)
+			if (SelectedNode != NULL)
 			{
-				if (SelectedNode->IsA(UK2Node_CustomEvent::StaticClass()))
+				if (UK2Node_CustomEvent* CustomEvent = Cast<UK2Node_CustomEvent>(SelectedNode))
 				{
-					UK2Node_CustomEvent* CustomEvent = (UK2Node_CustomEvent*)SelectedNode;
-					SummonSearchUI(true, CustomEvent->CustomFunctionName.ToString());
+					SearchTerm =  CustomEvent->CustomFunctionName.ToString();
+				}
+				else if (UK2Node_Variable* VariableNode = Cast<UK2Node_Variable>(SelectedNode))
+				{
+					SearchTerm = FName::NameToDisplayString(VariableNode->GetVarName().ToString(), VariableNode->GetPropertyForVariable()->IsA(UBoolProperty::StaticClass()) );
+				}
+				else
+				{
+					SearchTerm = SelectedNode->GetNodeTitle(ENodeTitleType::ListView).ToString();
 				}
 			}
 		}
+
+		if (!SearchTerm.IsEmpty())
+		{
+			SummonSearchUI(true, FString::Printf(TEXT("\"%s\""), *SearchTerm));
+		}
 	}
 }
 
+bool FBlueprintEditor::CanFindReferences()
+{
+	auto GraphEditor = FocusedGraphEdPtr.Pin();
+	if (GraphEditor.IsValid())
+	{
+		const FGraphPanelSelectionSet SelectedNodes = GraphEditor->GetSelectedNodes();
+		if( SelectedNodes.Num() == 1 )
+		{
+			return true;
+		}
+	}
+	return false;
+}
+
 AActor* FBlueprintEditor::GetPreviewActor() const
 {
 	UBlueprint* PreviewBlueprint = GetBlueprintObj();
diff --git a/Engine/Source/Editor/Kismet/Private/FindInBlueprintManager.cpp b/Engine/Source/Editor/Kismet/Private/FindInBlueprintManager.cpp
--- a/Engine/Source/Editor/Kismet/Private/FindInBlueprintManager.cpp
+++ b/Engine/Source/Editor/Kismet/Private/FindInBlueprintManager.cpp
@@ -31,6 +31,7 @@ const FText FFindInBlueprintSearchTags::FiB_Macros = LOCTEXT("Macros", "Macros")
 const FText FFindInBlueprintSearchTags::FiB_SubGraphs = LOCTEXT("Sub", "Sub");
 
 const FText FFindInBlueprintSearchTags::FiB_Name = LOCTEXT("Name", "Name");
+const FText FFindInBlueprintSearchTags::FiB_NativeName = LOCTEXT("NativeName", "Native Name");
 const FText FFindInBlueprintSearchTags::FiB_ClassName = LOCTEXT("ClassName", "ClassName");
 const FText FFindInBlueprintSearchTags::FiB_NodeGuid = LOCTEXT("NodeGuid", "NodeGuid");
 const FText FFindInBlueprintSearchTags::FiB_Tooltip = LOCTEXT("Tooltip", "Tooltip");
@@ -47,6 +48,51 @@ const FText FFindInBlueprintSearchTags::FiB_IsReference = LOCTEXT("IsReference",
 const FText FFindInBlueprintSearchTags::FiB_Glyph = LOCTEXT("Glyph", "Glyph");
 const FText FFindInBlueprintSearchTags::FiB_GlyphColor = LOCTEXT("GlyphColor", "GlyphColor");
 
+/** Temporarily forces all nodes and pins to use non-friendly names, forces all schema to have nodes clear their cached values so they will re-cache, and then reverts at the end */
+struct FTemporarilyUseFriendlyNodeTitles
+{
+	FTemporarilyUseFriendlyNodeTitles()
+	{
+		UEditorStyleSettings* EditorSettings = GetMutableDefault<UEditorStyleSettings>();
+
+		// Cache the value of bShowFriendlyNames, we will force it to true for gathering BP search data and then restore it
+		bCacheShowFriendlyNames = EditorSettings->bShowFriendlyNames;
+
+		EditorSettings->bShowFriendlyNames = true;
+		ForceVisualizationCacheClear();
+	}
+
+	~FTemporarilyUseFriendlyNodeTitles()
+	{
+		UEditorStyleSettings* EditorSettings = GetMutableDefault<UEditorStyleSettings>();
+		EditorSettings->bShowFriendlyNames = bCacheShowFriendlyNames;
+		ForceVisualizationCacheClear();
+	}
+
+	/** Go through all Schemas and force a visualization cache clear, forcing nodes to refresh their titles */
+	void ForceVisualizationCacheClear()
+	{
+		// Only do the purge if the state was changed
+		if (!bCacheShowFriendlyNames)
+		{
+			// Find all Schemas and force a visualization cache clear
+			for ( TObjectIterator<UClass> ClassIt; ClassIt; ++ClassIt )
+			{
+				UClass* CurrentClass = *ClassIt;
+
+				if (UEdGraphSchema* Schema = Cast<UEdGraphSchema>(CurrentClass->GetDefaultObject()))
+				{
+					Schema->ForceVisualizationCacheClear();
+				}
+			}
+		}
+	}
+
+private:
+	/** Cached state of ShowFriendlyNames in EditorSettings */
+	bool bCacheShowFriendlyNames;
+};
+
 /** Helper functions for serialization of types to and from an FString */
 namespace FiBSerializationHelpers
 {
@@ -337,10 +383,6 @@ namespace BlueprintSearchMetaDataHelpers
 		{ 
 			check( this->CurrentToken == EJsonToken::String ); 
 			// The string value from Json is a Hex value that must be looked up in the LookupTable to find the FText it represents
-			if(const FText* LookupText = LookupTable.Find(FCString::Atoi(*this->StringValue)))
-			{
-				return LookupText->ToString();
-			}
 			return this->StringValue;
 		}
 
@@ -1011,7 +1053,9 @@ void FFindInBlueprintSearchManager::OnAssetLoaded(UObject* InAsset)
 }
 
 FString FFindInBlueprintSearchManager::GatherBlueprintSearchMetadata(const UBlueprint* Blueprint)
-{
+{	
+	FTemporarilyUseFriendlyNodeTitles TemporarilyUseFriendlyNodeTitles;
+
 	FString SearchMetaData;
 
 	// The search registry tags for a Blueprint are all in Json
@@ -1515,7 +1559,7 @@ bool FFindInBlueprintSearchManager::IsCacheInProgress() const
 	return CachingObject != nullptr;
 }
 
-TSharedPtr< FJsonObject > FFindInBlueprintSearchManager::ConvertJsonStringToObject(FString InJsonString)
+TSharedPtr< FJsonObject > FFindInBlueprintSearchManager::ConvertJsonStringToObject(FString InJsonString, TMap<int32, FText>& OutFTextLookupTable)
 {
 	/** The searchable data is more complicated than a Json string, the Json being the main searchable body that is parsed. Below is a diagram of the full data:
 	 *  | int32 "Size" | TMap "Lookup Table" | Json String |
@@ -1536,7 +1580,7 @@ TSharedPtr< FJsonObject > FFindInBlueprintSearchManager::ConvertJsonStringToObje
 
  	// With the size of the TMap in hand, let's serialize JUST that (as a byte string)
 	TMap<int32, FText> LookupTable;
-	LookupTable = FiBSerializationHelpers::Deserialize< TMap<int32, FText> >(ReaderStream, SizeOfData);
+	OutFTextLookupTable = LookupTable = FiBSerializationHelpers::Deserialize< TMap<int32, FText> >(ReaderStream, SizeOfData);
 
 	// The original BufferReader should be positioned at the Json
 	TSharedPtr< FJsonObject > JsonObject = NULL;
diff --git a/Engine/Source/Editor/Kismet/Private/FindInBlueprints.cpp b/Engine/Source/Editor/Kismet/Private/FindInBlueprints.cpp
--- a/Engine/Source/Editor/Kismet/Private/FindInBlueprints.cpp
+++ b/Engine/Source/Editor/Kismet/Private/FindInBlueprints.cpp
@@ -17,6 +17,17 @@
 /** Some utility functions to help with Find-in-Blueprint functionality */
 namespace FindInBlueprintsHelpers
 {
+	/** Looks up a JsonValue's FText from the passed lookup table */
+	FText AsFText(TSharedPtr< FJsonValue > InJsonValue, const TMap<int32, FText>& InLookupTable)
+	{
+		if(const FText* LookupText = InLookupTable.Find(FCString::Atoi(*InJsonValue->AsString())))
+		{
+			return *LookupText;
+		}
+		// Let's never get here.
+		return LOCTEXT("FiBSerializationError", "There was an error in serialization!");
+	}
+
 	/**
 	 * Retrieves the pin type as a string value
 	 *
@@ -83,7 +94,7 @@ namespace FindInBlueprintsHelpers
 	 * @param InGraphType					There is nothing in the Json script to identify graphs, they are all grouped by category, so pass in what type they should be considered
 	 * @param OutBlueprintCategory			The search category these graphs will be children to
 	 */
-	void ExtractGraph( TSharedPtr< FJsonObject > InJsonObject, FString InCategoryTitle, const TArray<FString>& InTokens, EGraphType InGraphType, FSearchResult& OutBlueprintCategory )
+	void ExtractGraph( const TMap<int32, FText>& InLookupTable, TSharedPtr< FJsonObject > InJsonObject, FString InCategoryTitle, const TArray<FString>& InTokens, EGraphType InGraphType, FSearchResult& OutBlueprintCategory )
 	{
 		if(InJsonObject->HasField(InCategoryTitle))
 		{
@@ -93,7 +104,7 @@ namespace FindInBlueprintsHelpers
 			{
 				FSearchResult Result(new FFindInBlueprintsGraph(FText::GetEmpty(), OutBlueprintCategory, InGraphType) );
 
-				if( Result->ExtractContent(FunctionValue->AsObject(), InTokens) )
+				if( Result->ExtractContent(InLookupTable, FunctionValue->AsObject(), InTokens) )
 				{
 					OutBlueprintCategory->Children.Add(Result);
 				}
@@ -117,7 +128,8 @@ namespace FindInBlueprintsHelpers
 		}
 
 		TSharedPtr< FJsonObject > JsonObject = NULL;
-		JsonObject = FFindInBlueprintSearchManager::ConvertJsonStringToObject(InBuffer);
+		TMap<int32, FText> LookupTable;
+		JsonObject = FFindInBlueprintSearchManager::ConvertJsonStringToObject(InBuffer, LookupTable);
 
 		if(JsonObject.IsValid())
 		{
@@ -129,17 +141,17 @@ namespace FindInBlueprintsHelpers
 				{
 					FSearchResult Result(new FFindInBlueprintsProperty(FText::GetEmpty(), OutBlueprintCategory) );
 
-					if( Result->ExtractContent(PropertyValue->AsObject(), InTokens) )
+					if( Result->ExtractContent(LookupTable, PropertyValue->AsObject(), InTokens) )
 					{
 						OutBlueprintCategory->Children.Add(Result);
 					}
 				}
 			}
 
-			ExtractGraph(JsonObject, FFindInBlueprintSearchTags::FiB_Functions.ToString(), InTokens, EGraphType::GT_Function, OutBlueprintCategory);
-			ExtractGraph(JsonObject, FFindInBlueprintSearchTags::FiB_Macros.ToString(), InTokens, EGraphType::GT_Macro, OutBlueprintCategory);
-			ExtractGraph(JsonObject, FFindInBlueprintSearchTags::FiB_UberGraphs.ToString(), InTokens, EGraphType::GT_Ubergraph, OutBlueprintCategory);
-			ExtractGraph(JsonObject, FFindInBlueprintSearchTags::FiB_SubGraphs.ToString(), InTokens, EGraphType::GT_Ubergraph, OutBlueprintCategory);
+			ExtractGraph(LookupTable, JsonObject, FFindInBlueprintSearchTags::FiB_Functions.ToString(), InTokens, EGraphType::GT_Function, OutBlueprintCategory);
+			ExtractGraph(LookupTable, JsonObject, FFindInBlueprintSearchTags::FiB_Macros.ToString(), InTokens, EGraphType::GT_Macro, OutBlueprintCategory);
+			ExtractGraph(LookupTable, JsonObject, FFindInBlueprintSearchTags::FiB_UberGraphs.ToString(), InTokens, EGraphType::GT_Ubergraph, OutBlueprintCategory);
+			ExtractGraph(LookupTable, JsonObject, FFindInBlueprintSearchTags::FiB_SubGraphs.ToString(), InTokens, EGraphType::GT_Ubergraph, OutBlueprintCategory);
 
 			if(JsonObject->HasField(FFindInBlueprintSearchTags::FiB_Components.ToString()))
 			{
@@ -151,7 +163,7 @@ namespace FindInBlueprintsHelpers
 				{
 					FSearchResult Result(new FFindInBlueprintsProperty(FText::GetEmpty(), ComponentsResult) );
 
-					if( Result->ExtractContent(ComponentValue->AsObject(), InTokens) )
+					if( Result->ExtractContent(LookupTable, ComponentValue->AsObject(), InTokens) )
 					{
 						ComponentsResult->Children.Add(Result);
 					}
@@ -175,27 +187,27 @@ namespace FindInBlueprintsHelpers
 	 * @param InOutPinSubCategory	When the pin's sub-category is found, this value will be changed. This value will never be wiped during other translations. It is not thread safe to determine the sub-object.
 	 * @return						TRUE when the parsing is successful
 	 */
-	bool ParsePinType(FString InKey, FText InValue, FEdGraphPinType& InOutPinType, FString& OutPinSubCategory)
+	bool ParsePinType(FText InKey, FText InValue, FEdGraphPinType& InOutPinType, FString& OutPinSubCategory)
 	{
 		bool bParsed = true;
 
-		if(InKey == FFindInBlueprintSearchTags::FiB_PinCategory.ToString())
+		if(InKey.CompareTo(FFindInBlueprintSearchTags::FiB_PinCategory) == 0)
 		{
 			InOutPinType.PinCategory = InValue.ToString();
 		}
-		else if(InKey == FFindInBlueprintSearchTags::FiB_PinSubCategory.ToString())
+		else if(InKey.CompareTo(FFindInBlueprintSearchTags::FiB_PinSubCategory) == 0)
 		{
 			InOutPinType.PinSubCategory = InValue.ToString();
 		}
-		else if(InKey == FFindInBlueprintSearchTags::FiB_ObjectClass.ToString())
+		else if(InKey.CompareTo(FFindInBlueprintSearchTags::FiB_ObjectClass) == 0)
 		{
 			OutPinSubCategory = InValue.ToString();
 		}
-		else if(InKey == FFindInBlueprintSearchTags::FiB_IsArray.ToString())
+		else if(InKey.CompareTo(FFindInBlueprintSearchTags::FiB_IsArray) == 0)
 		{
 			InOutPinType.bIsArray = InValue.ToString().ToBool();
 		}
-		else if(InKey == FFindInBlueprintSearchTags::FiB_IsReference.ToString())
+		else if(InKey.CompareTo(FFindInBlueprintSearchTags::FiB_IsReference) == 0)
 		{
 			InOutPinType.bIsReference = InValue.ToString().ToBool();
 		}
@@ -241,26 +253,30 @@ FReply FFindInBlueprintsResult::OnClick()
 	return FReply::Handled();
 }
 
-bool FFindInBlueprintsResult::ParseSearchInfo(const TArray<FString> &InTokens, FString InKey, FText InValue, TSharedPtr< FFindInBlueprintsResult > InParent)
+bool FFindInBlueprintsResult::ParseSearchInfo(const TArray<FString> &InTokens, FText InKey, FText InValue, TSharedPtr< FFindInBlueprintsResult > InParent)
 {
 	// By default, nothing much gets processed, just check to see if it matches the search requirements
 	return FindInBlueprintsHelpers::TextMatchesSearchTokens(InTokens, InValue);
 }
 
-void FFindInBlueprintsResult::AddExtraSearchInfo(FString InKey, FText InValue, TSharedPtr< FFindInBlueprintsResult > InParent)
+void FFindInBlueprintsResult::AddExtraSearchInfo(FText InKey, FText InValue, TSharedPtr< FFindInBlueprintsResult > InParent)
 {
 	FText ExtraSearchInfoText;
 
 	FFormatNamedArguments Args;
-	Args.Add(TEXT("Key"), FText::FromString(InKey));
+	Args.Add(TEXT("Key"), InKey);
 	Args.Add(TEXT("Value"), InValue);
 	ExtraSearchInfoText = FText::Format(LOCTEXT("ExtraSearchInfo", "{Key}: {Value}"), Args);
 
 	TSharedPtr< FFindInBlueprintsResult > SearchResult(new FFindInBlueprintsResult(ExtraSearchInfoText, InParent) );
-	InParent->Children.Add(SearchResult);
+
+	if (InKey.CompareTo(FFindInBlueprintSearchTags::FiB_NativeName) != 0)
+	{
+		InParent->Children.Add(SearchResult);
+	}
 }
 
-bool FFindInBlueprintsResult::ExtractJsonValue(const TArray<FString>& InTokens, FText InKey, TSharedPtr< FJsonValue > InValue, TSharedPtr< FFindInBlueprintsResult > InParentOverride)
+bool FFindInBlueprintsResult::ExtractJsonValue(const TMap<int32, FText>& InLookupTable, const TArray<FString>& InTokens, FText InKey, TSharedPtr< FJsonValue > InValue, TSharedPtr< FFindInBlueprintsResult > InParentOverride)
 {
 	bool bMatchesSearchResults = false;
 
@@ -276,11 +292,11 @@ bool FFindInBlueprintsResult::ExtractJsonValue(const TArray<FString>& InTokens,
 		{
 			if(Array[ArrayIdx]->Type == EJson::String)
 			{
-				bArrayHasMatches |= ParseSearchInfo(InTokens, FString::FromInt(ArrayIdx), FText::FromString(Array[ArrayIdx]->AsString()), ArrayHeader);
+				bArrayHasMatches |= ParseSearchInfo(InTokens, FText::FromString(FString::FromInt(ArrayIdx)), FindInBlueprintsHelpers::AsFText(Array[ArrayIdx], InLookupTable), ArrayHeader);
 			}
 			else
 			{
-				bArrayHasMatches |= ExtractJsonValue(InTokens, FText::FromString(FString::FromInt(ArrayIdx)), Array[ArrayIdx], ArrayHeader);
+				bArrayHasMatches |= ExtractJsonValue(InLookupTable, InTokens, FText::FromString(FString::FromInt(ArrayIdx)), Array[ArrayIdx], ArrayHeader);
 			}
 		}
 
@@ -295,7 +311,7 @@ bool FFindInBlueprintsResult::ExtractJsonValue(const TArray<FString>& InTokens,
 	{
 		TSharedPtr< FFindInBlueprintsResult > ObjectResultHeader = MakeShareable(new FFindInBlueprintsResult(InKey, ParentResultPtr) );
 
-		bMatchesSearchResults |= ExtractContent(InValue->AsObject(), InTokens, ObjectResultHeader);
+		bMatchesSearchResults |= ExtractContent(InLookupTable, InValue->AsObject(), InTokens, ObjectResultHeader);
 
 		if(bMatchesSearchResults)
 		{
@@ -304,17 +320,17 @@ bool FFindInBlueprintsResult::ExtractJsonValue(const TArray<FString>& InTokens,
 	}
 	else if(InValue->Type == EJson::String)
 	{
-		bMatchesSearchResults |= ParseSearchInfo(InTokens, InKey.ToString(), FText::FromString(InValue->AsString()), ParentResultPtr);
+		bMatchesSearchResults |= ParseSearchInfo(InTokens, InKey, FindInBlueprintsHelpers::AsFText(InValue, InLookupTable), ParentResultPtr);
 	}
 	else
 	{
-		bMatchesSearchResults |= ParseSearchInfo(InTokens, InKey.ToString(), FText::FromString(InValue->AsString()), ParentResultPtr);
+		bMatchesSearchResults |= ParseSearchInfo(InTokens, InKey, FindInBlueprintsHelpers::AsFText(InValue, InLookupTable), ParentResultPtr);
 	}
 
 	return bMatchesSearchResults;
 }
 
-bool FFindInBlueprintsResult::ExtractContent(TSharedPtr< FJsonObject > InJsonNode, const TArray<FString>& InTokens, TSharedPtr< FFindInBlueprintsResult > InParentOverride)
+bool FFindInBlueprintsResult::ExtractContent(const TMap<int32, FText>& InLookupTable, TSharedPtr< FJsonObject > InJsonNode, const TArray<FString>& InTokens, TSharedPtr< FFindInBlueprintsResult > InParentOverride)
 {
 	bool bMatchesSearchResults = false;
 
@@ -322,7 +338,7 @@ bool FFindInBlueprintsResult::ExtractContent(TSharedPtr< FJsonObject > InJsonNod
 
 	for( auto MapValues : InJsonNode->Values )
 	{
-		bMatchesSearchResults |= ExtractJsonValue(InTokens, FText::FromString(MapValues.Key), MapValues.Value, ParentResultPtr);
+		bMatchesSearchResults |= ExtractJsonValue(InLookupTable, InTokens, FText::FromString(MapValues.Key), MapValues.Value, ParentResultPtr);
 	}
 	return bMatchesSearchResults;
 }
@@ -427,9 +443,9 @@ TSharedRef<SWidget> FFindInBlueprintsGraphNode::CreateIcon() const
 		.ToolTipText( GetCategory() );
 }
 
-bool FFindInBlueprintsGraphNode::ParseSearchInfo(const TArray<FString> &InTokens, FString InKey, FText InValue, TSharedPtr< FFindInBlueprintsResult > InParent)
+bool FFindInBlueprintsGraphNode::ParseSearchInfo(const TArray<FString> &InTokens, FText InKey, FText InValue, TSharedPtr< FFindInBlueprintsResult > InParent)
 {
-	if(InKey == FFindInBlueprintSearchTags::FiB_NodeGuid.ToString())
+	if(InKey.CompareTo(FFindInBlueprintSearchTags::FiB_NodeGuid) == 0)
 	{
 		FString NodeGUIDAsString = InValue.ToString();
 		FGuid::Parse(NodeGUIDAsString, NodeGuid);
@@ -440,25 +456,25 @@ bool FFindInBlueprintsGraphNode::ParseSearchInfo(const TArray<FString> &InTokens
 
 	bool bMatchesTokens = FindInBlueprintsHelpers::TextMatchesSearchTokens(InTokens, InValue);
 
-	if(InKey == FFindInBlueprintSearchTags::FiB_ClassName.ToString())
+	if(InKey.CompareTo(FFindInBlueprintSearchTags::FiB_ClassName) == 0)
 	{
 		ClassName = InValue.ToString();
 		bMatchesTokens = false;
 	}
-	else if(InKey == FFindInBlueprintSearchTags::FiB_Name.ToString())
+	else if(InKey.CompareTo(FFindInBlueprintSearchTags::FiB_Name) == 0)
 	{
 		DisplayText = InValue;
 	}
-	else if(InKey == FFindInBlueprintSearchTags::FiB_Comment.ToString())
+	else if(InKey.CompareTo(FFindInBlueprintSearchTags::FiB_Comment) == 0)
 	{
 		CommentText = InValue.ToString();
 	}
-	else if(InKey == FFindInBlueprintSearchTags::FiB_Glyph.ToString())
+	else if(InKey.CompareTo(FFindInBlueprintSearchTags::FiB_Glyph) == 0)
 	{
 		GlyphBrush = FEditorStyle::GetBrush(*InValue.ToString());
 		bMatchesTokens = false;
 	}
-	else if(InKey == FFindInBlueprintSearchTags::FiB_GlyphColor.ToString())
+	else if(InKey.CompareTo(FFindInBlueprintSearchTags::FiB_GlyphColor) == 0)
 	{
 		GlyphColor.InitFromString(InValue.ToString());
 		bMatchesTokens = false;
@@ -475,14 +491,14 @@ bool FFindInBlueprintsGraphNode::ParseSearchInfo(const TArray<FString> &InTokens
 	return bMatchesTokens;
 }
 
-bool FFindInBlueprintsGraphNode::ExtractContent(TSharedPtr< FJsonObject > InJsonNode, const TArray<FString>& InTokens)
+bool FFindInBlueprintsGraphNode::ExtractContent(const TMap<int32, FText>& InLookupTable, TSharedPtr< FJsonObject > InJsonNode, const TArray<FString>& InTokens)
 {
 	// Very important to get the schema first, other bits of data depend on it
 	FString SchemaName;
 	TSharedPtr< FJsonValue > SchemaNameValue = InJsonNode->GetField< EJson::String >(FFindInBlueprintSearchTags::FiB_SchemaName.ToString());
 	if(SchemaNameValue.IsValid())
 	{
-		SchemaName = SchemaNameValue->AsString();
+		SchemaName = FindInBlueprintsHelpers::AsFText(SchemaNameValue, InLookupTable).ToString();
 	}
 
 	bool bMatchesSearchResults = false;
@@ -499,7 +515,7 @@ bool FFindInBlueprintsGraphNode::ExtractContent(TSharedPtr< FJsonObject > InJson
 			{
 				// The pin
 				TSharedPtr< FFindInBlueprintsPin > PinResult(new FFindInBlueprintsPin(FText::GetEmpty(), SharedThis( this ), SchemaName) );
-				if(PinResult->ExtractContent(Pin->AsObject(), InTokens))
+				if(PinResult->ExtractContent(InLookupTable, Pin->AsObject(), InTokens))
 				{
 					Children.Add(PinResult);
 					bMatchesSearchResults = true;
@@ -513,7 +529,7 @@ bool FFindInBlueprintsGraphNode::ExtractContent(TSharedPtr< FJsonObject > InJson
 		}
 		else
 		{
-			bMatchesSearchResults |= ExtractJsonValue(InTokens, FText::FromString(MapValues.Key), MapValues.Value, SharedThis( this ));
+			bMatchesSearchResults |= ExtractJsonValue(InLookupTable, InTokens, FText::FromString(MapValues.Key), MapValues.Value, SharedThis( this ));
 		}
 	}
 	return bMatchesSearchResults;
@@ -587,11 +603,11 @@ TSharedRef<SWidget> FFindInBlueprintsPin::CreateIcon() const
 		.ToolTipText(FText::FromString(FindInBlueprintsHelpers::GetPinTypeAsString(PinType)));
 }
 
-bool FFindInBlueprintsPin::ParseSearchInfo(const TArray<FString> &InTokens, FString InKey, FText InValue, TSharedPtr< FFindInBlueprintsResult > InParent)
+bool FFindInBlueprintsPin::ParseSearchInfo(const TArray<FString> &InTokens, FText InKey, FText InValue, TSharedPtr< FFindInBlueprintsResult > InParent)
 {
 	bool bMatchesTokens = FindInBlueprintsHelpers::TextMatchesSearchTokens(InTokens, InValue);
 
-	if(InKey == FFindInBlueprintSearchTags::FiB_Name.ToString())
+	if(InKey.CompareTo(FFindInBlueprintSearchTags::FiB_Name) == 0)
 	{
 		DisplayText = InValue;
 		bMatchesTokens = FindInBlueprintsHelpers::StringMatchesSearchTokens(InTokens, DisplayText.ToString());
@@ -599,7 +615,7 @@ bool FFindInBlueprintsPin::ParseSearchInfo(const TArray<FString> &InTokens, FStr
 	else if(FindInBlueprintsHelpers::ParsePinType(InKey, InValue, PinType, PinSubCategory))
 	{
 		// For pins, the only information in the Pin Type we care about is the object class, so return false if the Key is not that
-		if(InKey != FFindInBlueprintSearchTags::FiB_ObjectClass.ToString())
+		if(InKey.CompareTo(FFindInBlueprintSearchTags::FiB_ObjectClass) != 0)
 		{
 			return false;
 		}
@@ -667,23 +683,23 @@ TSharedRef<SWidget> FFindInBlueprintsProperty::CreateIcon() const
 		.ToolTipText( FText::FromString(FindInBlueprintsHelpers::GetPinTypeAsString(PinType)) );
 }
 
-bool FFindInBlueprintsProperty::ParseSearchInfo(const TArray<FString> &InTokens, FString InKey, FText InValue, TSharedPtr< FFindInBlueprintsResult > InParent)
+bool FFindInBlueprintsProperty::ParseSearchInfo(const TArray<FString> &InTokens, FText InKey, FText InValue, TSharedPtr< FFindInBlueprintsResult > InParent)
 {
 	bool bMatchesTokens = FindInBlueprintsHelpers::TextMatchesSearchTokens(InTokens, InValue);
 
-	if(InKey == FFindInBlueprintSearchTags::FiB_Name.ToString())
+	if(InKey.CompareTo(FFindInBlueprintSearchTags::FiB_Name) == 0)
 	{
 		DisplayText = InValue;
 	}
 	else if(FindInBlueprintsHelpers::ParsePinType(InKey, InValue, PinType, PinSubCategory))
 	{
 		// For properties, we care about all the information in the pin type, except whether its a reference or an array
-		if(InKey == FFindInBlueprintSearchTags::FiB_IsArray.ToString() || InKey == FFindInBlueprintSearchTags::FiB_IsReference.ToString())
+		if(InKey.CompareTo(FFindInBlueprintSearchTags::FiB_IsArray) == 0 || InKey.CompareTo(FFindInBlueprintSearchTags::FiB_IsReference) == 0)
 		{
 			return false;
 		}
 	}
-	else if(InKey == FFindInBlueprintSearchTags::FiB_IsSCSComponent.ToString())
+	else if(InKey.CompareTo(FFindInBlueprintSearchTags::FiB_IsSCSComponent) == 0)
 	{
 		bIsSCSComponent = true;
 
@@ -776,7 +792,7 @@ TSharedRef<SWidget> FFindInBlueprintsGraph::CreateIcon() const
 		.ToolTipText( GetCategory() );
 }
 
-bool FFindInBlueprintsGraph::ExtractContent(TSharedPtr< FJsonObject > InJsonNode, const TArray<FString>& InTokens)
+bool FFindInBlueprintsGraph::ExtractContent(const TMap<int32, FText>& InLookupTable, TSharedPtr< FJsonObject > InJsonNode, const TArray<FString>& InTokens)
 {
 	bool bMatchesSearchResults = false;
 
@@ -792,7 +808,7 @@ bool FFindInBlueprintsGraph::ExtractContent(TSharedPtr< FJsonObject > InJsonNode
 			{
 				TSharedPtr< FFindInBlueprintsGraphNode > NodeResult(new FFindInBlueprintsGraphNode(FText::GetEmpty(), SharedThis( this )) );
 
-				if( NodeResult->ExtractContent(NodeValue->AsObject(), InTokens) )
+				if( NodeResult->ExtractContent(InLookupTable, NodeValue->AsObject(), InTokens) )
 				{
 					Children.Add(NodeResult);
 					bMatchesSearchResults = true;
@@ -812,7 +828,7 @@ bool FFindInBlueprintsGraph::ExtractContent(TSharedPtr< FJsonObject > InJsonNode
 			{
 				FSearchResult Result(new FFindInBlueprintsProperty(FText::GetEmpty(), SharedThis( this )) );
 
-				if( Result->ExtractContent(PropertyValue->AsObject(), InTokens) )
+				if( Result->ExtractContent(InLookupTable, PropertyValue->AsObject(), InTokens) )
 				{
 					Children.Add(Result);
 					bMatchesSearchResults = true;
@@ -821,17 +837,17 @@ bool FFindInBlueprintsGraph::ExtractContent(TSharedPtr< FJsonObject > InJsonNode
 		}
 		else
 		{
-			bMatchesSearchResults |= ExtractJsonValue( InTokens, FText::FromString(MapValues.Key), MapValues.Value, SharedThis( this ) );
+			bMatchesSearchResults |= ExtractJsonValue( InLookupTable, InTokens, FText::FromString(MapValues.Key), MapValues.Value, SharedThis( this ) );
 		}
 	}
 	return bMatchesSearchResults;
 }
 
-bool FFindInBlueprintsGraph::ParseSearchInfo(const TArray<FString> &InTokens, FString InKey, FText InValue, TSharedPtr< FFindInBlueprintsResult > InParent)
+bool FFindInBlueprintsGraph::ParseSearchInfo(const TArray<FString> &InTokens, FText InKey, FText InValue, TSharedPtr< FFindInBlueprintsResult > InParent)
 {
 	bool bMatchesTokens = FindInBlueprintsHelpers::TextMatchesSearchTokens(InTokens, InValue);
 
-	if(InKey == FFindInBlueprintSearchTags::FiB_Name.ToString())
+	if(InKey.CompareTo(FFindInBlueprintSearchTags::FiB_Name) == 0)
 	{
 		DisplayText = InValue;
 	}
diff --git a/Engine/Source/Editor/Kismet/Private/SMyBlueprint.cpp b/Engine/Source/Editor/Kismet/Private/SMyBlueprint.cpp
--- a/Engine/Source/Editor/Kismet/Private/SMyBlueprint.cpp
+++ b/Engine/Source/Editor/Kismet/Private/SMyBlueprint.cpp
@@ -21,6 +21,7 @@
 #include "SBlueprintPalette.h"
 #include "SGraphActionMenu.h"
 #include "BlueprintEditorCommands.h"
+#include "GraphEditorActions.h"
 
 #include "Editor/AnimGraph/Classes/AnimationGraph.h"
 
@@ -54,7 +55,6 @@ void FMyBlueprintCommands::RegisterCommands()
 	UI_COMMAND( FocusNode, "Focus", "Focuses on the associated node", EUserInterfaceActionType::Button, FInputChord() );
 	UI_COMMAND( FocusNodeInNewTab, "Focus in New Tab", "Focuses on the associated node in a new tab", EUserInterfaceActionType::Button, FInputChord() );
 	UI_COMMAND( ImplementFunction, "Implement Function", "Implements this overridable function as a new function.", EUserInterfaceActionType::Button, FInputChord() );
-	UI_COMMAND( FindEntry, "Find References", "Searches for all references of this function or variable.", EUserInterfaceActionType::Button, FInputChord() );
 	UI_COMMAND(DeleteEntry, "Delete", "Deletes this function or variable from this blueprint.", EUserInterfaceActionType::Button, FInputChord(EKeys::Platform_Delete));
 	UI_COMMAND( GotoNativeVarDefinition, "Goto Code Definition", "Goto the native code definition of this variable", EUserInterfaceActionType::Button, FInputChord() );
 }
@@ -169,10 +169,10 @@ void SMyBlueprint::Construct(const FArguments& InArgs, TWeakPtr<FBlueprintEditor
 			FCanExecuteAction(), FIsActionChecked(),
 			FIsActionButtonVisible::CreateSP(this, &SMyBlueprint::CanImplementFunction) );
 	
-		ToolKitCommandList->MapAction( FMyBlueprintCommands::Get().FindEntry,
-			FExecuteAction::CreateSP(this, &SMyBlueprint::OnFindEntry),
+		ToolKitCommandList->MapAction( FGraphEditorCommands::Get().FindReferences,
+			FExecuteAction::CreateSP(this, &SMyBlueprint::OnFindReference),
 			FCanExecuteAction(),
-			FIsActionButtonVisible::CreateSP(this, &SMyBlueprint::CanFindEntry) );
+			FIsActionButtonVisible::CreateSP(this, &SMyBlueprint::CanFindReference) );
 	
 		ToolKitCommandList->MapAction( FMyBlueprintCommands::Get().DeleteEntry,
 			FExecuteAction::CreateSP(this, &SMyBlueprint::OnDeleteEntry),
@@ -1831,7 +1831,7 @@ TSharedPtr<SWidget> SMyBlueprint::OnContextMenuOpening()
 			MenuBuilder.AddMenuEntry(FMyBlueprintCommands::Get().FocusNodeInNewTab);
 			MenuBuilder.AddMenuEntry(FGenericCommands::Get().Rename, NAME_None, LOCTEXT("Rename", "Rename"), LOCTEXT("Rename_Tooltip", "Renames this function or variable from blueprint.") );
 			MenuBuilder.AddMenuEntry(FMyBlueprintCommands::Get().ImplementFunction);
-			MenuBuilder.AddMenuEntry(FMyBlueprintCommands::Get().FindEntry);
+			MenuBuilder.AddMenuEntry(FGraphEditorCommands::Get().FindReferences);
 			MenuBuilder.AddMenuEntry(FMyBlueprintCommands::Get().GotoNativeVarDefinition);
 			MenuBuilder.AddMenuEntry(FGenericCommands::Get().Duplicate);
 			MenuBuilder.AddMenuEntry(FMyBlueprintCommands::Get().DeleteEntry);
@@ -2078,7 +2078,7 @@ void SMyBlueprint::ImplementFunction(FEdGraphSchemaAction_K2Graph* GraphAction)
 	}
 }
 
-void SMyBlueprint::OnFindEntry()
+void SMyBlueprint::OnFindReference()
 {
 	FString SearchTerm;
 	if (FEdGraphSchemaAction_K2Graph* GraphAction = SelectionAsGraph())
@@ -2103,7 +2103,7 @@ void SMyBlueprint::OnFindEntry()
 	}
 	else if (FEdGraphSchemaAction_K2Struct* StructAction = SelectionAsStruct())
 	{
-		BlueprintEditorPtr.Pin()->SummonSearchUI(true, StructAction->Struct->GetName());
+		SearchTerm = StructAction->Struct->GetName();
 	}
 	else if (FEdGraphSchemaAction_K2Event* EventAction = SelectionAsEvent())
 	{
@@ -2122,7 +2122,7 @@ void SMyBlueprint::OnFindEntry()
 	}
 }
 
-bool SMyBlueprint::CanFindEntry() const
+bool SMyBlueprint::CanFindReference() const
 {
 	// Nothing relevant to the category will ever be found, unless the name of the category overlaps with another item
 	if (SelectionIsCategory())
diff --git a/Engine/Source/Editor/Kismet/Private/SMyBlueprint.h b/Engine/Source/Editor/Kismet/Private/SMyBlueprint.h
--- a/Engine/Source/Editor/Kismet/Private/SMyBlueprint.h
+++ b/Engine/Source/Editor/Kismet/Private/SMyBlueprint.h
@@ -40,7 +40,6 @@ public:
 	TSharedPtr<FUICommandInfo> FocusNode;
 	TSharedPtr<FUICommandInfo> FocusNodeInNewTab;
 	TSharedPtr<FUICommandInfo> ImplementFunction;
-	TSharedPtr<FUICommandInfo> FindEntry;
 	TSharedPtr<FUICommandInfo> DeleteEntry;
 	TSharedPtr<FUICommandInfo> GotoNativeVarDefinition;
 	// Add New Item
@@ -193,8 +192,8 @@ private:
 	void ImplementFunction(TSharedPtr<FEdGraphSchemaAction_K2Graph> GraphAction);
 	void ImplementFunction(FEdGraphSchemaAction_K2Graph* GraphAction);
 	bool CanImplementFunction() const;
-	void OnFindEntry();
-	bool CanFindEntry() const;
+	void OnFindReference();
+	bool CanFindReference() const;
 	void OnDeleteEntry();
 	bool CanDeleteEntry() const;
 	FReply OnAddNewLocalVariable();
diff --git a/Engine/Source/Editor/Kismet/Public/BlueprintEditor.h b/Engine/Source/Editor/Kismet/Public/BlueprintEditor.h
--- a/Engine/Source/Editor/Kismet/Public/BlueprintEditor.h
+++ b/Engine/Source/Editor/Kismet/Public/BlueprintEditor.h
@@ -459,9 +459,6 @@ public:
 
 	virtual void ClearSelectionStateFor(FName SelectionOwner);
 
-	/** Find all instances of the selected custom event. */
-	void OnFindInstancesCustomEvent();
-
 	/** Handles spawning a graph node in the current graph using the passed in chord */
 	FReply OnSpawnGraphNodeByShortcut(FInputChord InChord, const FVector2D& InPosition, UEdGraph* InGraph);
 
@@ -920,11 +917,11 @@ private:
 	/** Helper function to navigate the current tab */
 	void NavigateTab(FDocumentTracker::EOpenDocumentCause InCause);
 
-	/** Find all references of the selected variable. */
-	void OnFindVariableReferences();
+	/** Find all references of the selected node. */
+	void OnFindReferences();
 
-	/** Checks if we can currently find all references of the variable selection. */
-	bool CanFindVariableReferences();
+	/** Checks if we can currently find all references of the node selection. */
+	bool CanFindReferences();
 
 	/** Called when the user generates a warning tooltip because a connection was invalid */
 	void OnDisallowedPinConnection(const class UEdGraphPin* PinA, const class UEdGraphPin* PinB);
diff --git a/Engine/Source/Editor/Kismet/Public/FindInBlueprintManager.h b/Engine/Source/Editor/Kismet/Public/FindInBlueprintManager.h
--- a/Engine/Source/Editor/Kismet/Public/FindInBlueprintManager.h
+++ b/Engine/Source/Editor/Kismet/Public/FindInBlueprintManager.h
@@ -31,6 +31,8 @@ struct KISMET_API FFindInBlueprintSearchTags
 
 	/** Name tag */
 	static const FText FiB_Name;
+	/** Native Name tag */
+	static const FText FiB_NativeName;
 	/** Class Name tag */
 	static const FText FiB_ClassName;
 	/** NodeGuid tag */
@@ -261,7 +263,7 @@ public:
 	static FString ConvertFStringToHexString(FString InValue);
 
 	/** Given a fully constructed Find-in-Blueprint FString of searchable data, will parse and construct a JsonObject */
-	static TSharedPtr< class FJsonObject > ConvertJsonStringToObject(FString InJsonString);
+	static TSharedPtr< class FJsonObject > ConvertJsonStringToObject(FString InJsonString, TMap<int32, FText>& OutFTextLookupTable);
 
 private:
 	/** Initializes the FiB manager */
diff --git a/Engine/Source/Editor/Kismet/Public/FindInBlueprints.h b/Engine/Source/Editor/Kismet/Public/FindInBlueprints.h
--- a/Engine/Source/Editor/Kismet/Public/FindInBlueprints.h
+++ b/Engine/Source/Editor/Kismet/Public/FindInBlueprints.h
@@ -38,21 +38,23 @@ public:
 	/**
 	 * Extracts the content of a JsonNode, building the sub-hierarchy for display
 	 *
+	 * @param InLookupTable		Lookup table for converting Json strings into an FText
 	 * @param InJsonNode		Node to extract data from
 	 * @param InTokens			Tokens to match searchable content against
 	 * @return					Returns TRUE if the content matches search results and was extracted
 	 */
-	virtual bool ExtractContent(TSharedPtr< FJsonObject > InJsonNode, const TArray<FString>& InTokens) { return ExtractContent(InJsonNode, InTokens, NULL); };
+	virtual bool ExtractContent(const TMap<int32, FText>& InLookupTable, TSharedPtr< FJsonObject > InJsonNode, const TArray<FString>& InTokens) { return ExtractContent(InLookupTable, InJsonNode, InTokens, NULL); };
 
 	/**
 	 * Extracts the content of a JsonNode, building the sub-hierarchy for display
 	 *
+	 * @param InLookupTable		Lookup table for converting Json strings into an FText
 	 * @param InJsonNode		Node to extract data from
 	 * @param InTokens			Tokens to match searchable content against
 	 * @param InParentOverride	Parent search result to attach sub-content to
 	 * @return					Returns TRUE if the content matches search results and was extracted
 	 */
-	virtual bool ExtractContent(TSharedPtr< FJsonObject > InJsonNode, const TArray<FString>& InTokens, TSharedPtr< FFindInBlueprintsResult > InParentOverride);
+	virtual bool ExtractContent(const TMap<int32, FText>& InLookupTable, TSharedPtr< FJsonObject > InJsonNode, const TArray<FString>& InTokens, TSharedPtr< FFindInBlueprintsResult > InParentOverride);
 
 	/**
 	 * Parses search info for specific data important for displaying the search result in an easy to understand format
@@ -64,7 +66,7 @@ public:
 	 *
 	 * @return					TRUE if the item passes the search filter or is otherwise accepted
 	 */
-	virtual bool ParseSearchInfo(const TArray<FString> &InTokens, FString InKey, FText InValue, TSharedPtr< FFindInBlueprintsResult > InParent);
+	virtual bool ParseSearchInfo(const TArray<FString> &InTokens, FText InKey, FText InValue, TSharedPtr< FFindInBlueprintsResult > InParent);
 
 	/**
 	 * Adds extra search info, anything that does not have a predestined place in the search result. Adds a sub-item to the searches and formats its description so the tag displays
@@ -73,7 +75,7 @@ public:
 	 * @param	InValue			Compared against search query to see if it passes the filter, sometimes data is rejected because it is deemed unsearchable
 	 * @param	InParent		The parent search result
 	 */
-	void AddExtraSearchInfo(FString InKey, FText InValue, TSharedPtr< FFindInBlueprintsResult > InParent);
+	void AddExtraSearchInfo(FText InKey, FText InValue, TSharedPtr< FFindInBlueprintsResult > InParent);
 
 	/**
 	 * Iterates through all this node's children, and tells the tree view to expand them
@@ -88,12 +90,13 @@ protected:
 	/**
 	 * Extracts the details of a JsonValue in the search results, pushing arrays into sub-objects correctly, breaking up objects, and passing on the rest of the types
 	 *
+	 * @param	InLookupTable			Lookup table for converting Json strings into an FText
 	 * @param	InTokens				The search tokens to check results against
 	 * @param	InKey					This is the tag for the data, describing what it is so special handling can occur if needed
 	 * @param	InValue					Has all the data for this object that needs to be deciphered or broken up into base types to be deciphered
 	 * @param	InParentOverride		The parent to attach all children data to, can be NULL.
 	 */
-	bool ExtractJsonValue(const TArray<FString>& InTokens, FText InKey, TSharedPtr< FJsonValue > InValue, TSharedPtr< FFindInBlueprintsResult > InParentOverride);
+	bool ExtractJsonValue(const TMap<int32, FText>& InLookupTable, const TArray<FString>& InTokens, FText InKey, TSharedPtr< FJsonValue > InValue, TSharedPtr< FFindInBlueprintsResult > InParentOverride);
 
 public:
 	/*Any children listed under this category */
@@ -119,8 +122,8 @@ public:
 	/** FFindInBlueprintsResult Interface */
 	virtual FReply OnClick() override;
 	virtual TSharedRef<SWidget>	CreateIcon() const override;
-	virtual bool ParseSearchInfo(const TArray<FString> &InTokens, FString InKey, FText InValue, TSharedPtr< FFindInBlueprintsResult > InParentOverride = NULL) override;
-	virtual bool ExtractContent(TSharedPtr< FJsonObject > InJsonNode, const TArray<FString>& InTokens) override;
+	virtual bool ParseSearchInfo(const TArray<FString> &InTokens, FText InKey, FText InValue, TSharedPtr< FFindInBlueprintsResult > InParentOverride = NULL) override;
+	virtual bool ExtractContent(const TMap<int32, FText>& InLookupTable, TSharedPtr< FJsonObject > InJsonNode, const TArray<FString>& InTokens) override;
 	virtual FText GetCategory() const override;
 	virtual void FinalizeSearchData() override;
 	/** End FFindInBlueprintsResult Interface */
@@ -151,7 +154,7 @@ public:
 
 	/** FFindInBlueprintsResult Interface */
 	virtual TSharedRef<SWidget>	CreateIcon() const override;
-	virtual bool ParseSearchInfo(const TArray<FString> &InTokens, FString InKey, FText InValue, TSharedPtr< FFindInBlueprintsResult > InParentOverride = NULL) override;
+	virtual bool ParseSearchInfo(const TArray<FString> &InTokens, FText InKey, FText InValue, TSharedPtr< FFindInBlueprintsResult > InParentOverride = NULL) override;
 	virtual FText GetCategory() const override;
 	virtual void FinalizeSearchData() override;
 	/** End FFindInBlueprintsResult Interface */
@@ -178,7 +181,7 @@ public:
 
 	/** FFindInBlueprintsResult Interface */
 	virtual TSharedRef<SWidget>	CreateIcon() const override;
-	virtual bool ParseSearchInfo(const TArray<FString> &InTokens, FString InKey, FText InValue, TSharedPtr< FFindInBlueprintsResult > InParentOverride = NULL) override;
+	virtual bool ParseSearchInfo(const TArray<FString> &InTokens, FText InKey, FText InValue, TSharedPtr< FFindInBlueprintsResult > InParentOverride = NULL) override;
 	virtual FText GetCategory() const override;
 	virtual void FinalizeSearchData() override;
 	/** End FFindInBlueprintsResult Interface */
@@ -206,9 +209,9 @@ public:
 	/** FFindInBlueprintsResult Interface */
 	virtual FReply OnClick() override;
 	virtual TSharedRef<SWidget>	CreateIcon() const override;
-	virtual bool ParseSearchInfo(const TArray<FString> &InTokens, FString InKey, FText InValue, TSharedPtr< FFindInBlueprintsResult > InParentOverride = NULL) override;
+	virtual bool ParseSearchInfo(const TArray<FString> &InTokens, FText InKey, FText InValue, TSharedPtr< FFindInBlueprintsResult > InParentOverride = NULL) override;
 	virtual FText GetCategory() const override;
-	virtual bool ExtractContent(TSharedPtr< FJsonObject > InJsonNode, const TArray<FString>& InTokens) override;
+	virtual bool ExtractContent(const TMap<int32, FText>& InLookupTable, TSharedPtr< FJsonObject > InJsonNode, const TArray<FString>& InTokens) override;
 	/** End FFindInBlueprintsResult Interface */
 protected:
 
diff --git a/Engine/Source/Runtime/Engine/Private/EdGraph/EdGraphSchema.cpp b/Engine/Source/Runtime/Engine/Private/EdGraph/EdGraphSchema.cpp
--- a/Engine/Source/Runtime/Engine/Private/EdGraph/EdGraphSchema.cpp
+++ b/Engine/Source/Runtime/Engine/Private/EdGraph/EdGraphSchema.cpp
@@ -509,7 +509,7 @@ FText UEdGraphSchema::GetPinDisplayName(const UEdGraphPin* Pin) const
 	else
 	{
 		ResultPinName = Pin->PinFriendlyName;
-		bool bShowNodesAndPinsUnlocalized;
+		bool bShowNodesAndPinsUnlocalized = false;
 		GConfig->GetBool( TEXT("Internationalization"), TEXT("ShowNodesAndPinsUnlocalized"), bShowNodesAndPinsUnlocalized, GEditorSettingsIni );
 		if (bShowNodesAndPinsUnlocalized)
 		{
