Cannot name Timeline tracks the same as any of the native pins on the node.

Fixed a minor localization issue with the display of track name errors.

diff --git a/Engine/Source/Editor/Kismet/Private/STimelineEditor.cpp b/Engine/Source/Editor/Kismet/Private/STimelineEditor.cpp
--- a/Engine/Source/Editor/Kismet/Private/STimelineEditor.cpp
+++ b/Engine/Source/Editor/Kismet/Private/STimelineEditor.cpp
@@ -1320,9 +1320,33 @@ bool STimelineEditor::OnVerifyTrackNameCommit(const FText& TrackName, FText& Out
 	else if(TrackBase->TrackName != RequestedName && 
 		false == TimelineObj->IsNewTrackNameValid(RequestedName))
 	{
-		OutErrorMessage = FText::FromString(FString(TEXT("\"")) + *TrackName.ToString() + LOCTEXT( "AlreadyInUse", "\" is already in use." ).ToString());
+		FFormatNamedArguments Args;
+		Args.Add(TEXT("TrackName"), TrackName);
+		OutErrorMessage = FText::Format(LOCTEXT("AlreadyInUse", "\"{TrackName}\" is already in use."), Args);
 		bValid = false;
 	}
+	else
+	{
+		TSharedPtr<FBlueprintEditor> Kismet2 = Kismet2Ptr.Pin();
+		UBlueprint* Blueprint = Kismet2->GetBlueprintObj();
+		UK2Node_Timeline* TimelineNode = FBlueprintEditorUtils::FindNodeForTimeline(Blueprint, TimelineObj);
+		if (TimelineNode)
+		{
+			for(TArray<UEdGraphPin*>::TIterator PinIt(TimelineNode->Pins);PinIt;++PinIt)
+			{
+				UEdGraphPin* Pin = *PinIt;
+
+				if (Pin->PinName == TrackName.ToString())
+				{
+					FFormatNamedArguments Args;
+					Args.Add(TEXT("TrackName"), TrackName);
+					OutErrorMessage = FText::Format(LOCTEXT("PinAlreadyInUse", "\"{TrackName}\" is already in use as a default pin!"), Args);
+					bValid = false;
+					break;
+				}
+			}
+		}
+	}
 
 	return bValid;
 }
@@ -1340,10 +1364,12 @@ void STimelineEditor::OnTrackNameCommitted( const FText& StringName, ETextCommit
 		
 		if (TimelineNode)
 		{
-			//rename pin directly to avoid breaking links
-			for(TArray<UEdGraphPin*>::TIterator PinIt(TimelineNode->Pins);PinIt;++PinIt)
+			// Start looking from the bottom of the list of pins, where user defined ones are stored.
+			// It should not be possible to name pins to be the same as default pins, 
+			// but in the case (fixes broken nodes) that they happen to be the same, this protects them
+			for (int32 PinIdx = TimelineNode->Pins.Num() - 1; PinIdx >= 0; --PinIdx)
 			{
-				UEdGraphPin* Pin = *PinIt;
+				UEdGraphPin* Pin = TimelineNode->Pins[PinIdx];
 			
 				if (Pin->PinName == TrackBase->TrackName.ToString())
 				{
