Fixed crash when displaying tooltips for blueprint context target menus
Also fixed display names appearing as "None"

The issue here was due to passing in bit-mask enum values as indices to the UEnum name array.

diff --git a/Engine/Source/Editor/Kismet/Private/SBlueprintContextTargetMenu.cpp b/Engine/Source/Editor/Kismet/Private/SBlueprintContextTargetMenu.cpp
--- a/Engine/Source/Editor/Kismet/Private/SBlueprintContextTargetMenu.cpp
+++ b/Engine/Source/Editor/Kismet/Private/SBlueprintContextTargetMenu.cpp
@@ -461,17 +461,14 @@ bool FContextMenuTargetProfile::LoadProfile()
 /**  */
 namespace BlueprintContextTargetMenuImpl
 {
-	static FText GetContextTargetDisplayName(UEnum* Enum, EContextTargetFlags::Type ContextTarget);
-}
-
-//------------------------------------------------------------------------------
-static FText BlueprintContextTargetMenuImpl::GetContextTargetDisplayName(UEnum* Enum, EContextTargetFlags::Type ContextTarget)
-{
-	if (Enum != nullptr)
+	static FText GetContextTargetDisplayName(UEnum* Enum, int32 EnumIndex)
 	{
-		return Enum->GetDisplayNameText(ContextTarget);
+		if (Enum != nullptr)
+		{
+			return Enum->GetDisplayNameText(EnumIndex);
+		}
+		return LOCTEXT("UnrecognizedTarget", "Error: <UNRECOGNIZED>");
 	}
-	return LOCTEXT("UnrecognizedTarget", "Error: <UNRECOGNIZED>");
 }
 
 //------------------------------------------------------------------------------
@@ -539,18 +536,16 @@ void SBlueprintContextTargetMenu::Construct(const FArguments& InArgs, const FBlu
 	uint32 const GreatestFlag = (EContextTargetFlags::ContextTargetFlagsEnd & ~1);
 
 	int32  ColIndex = 0;
-	uint32 TargetFlag = (1<<0);
-	while (TargetFlag <= GreatestFlag)
+	for (int32 BitMaskOffset = 0; (1 << BitMaskOffset) <= GreatestFlag; ++BitMaskOffset)
 	{
-		EContextTargetFlags::Type ContextTarget = (EContextTargetFlags::Type)TargetFlag;
-		TargetFlag <<= 1;
+		EContextTargetFlags::Type ContextTarget = (EContextTargetFlags::Type)(1 << BitMaskOffset);
 
 		if (TargetEnum && TargetEnum->HasMetaData(TEXT("Hidden"), ContextTarget))
 		{
 			continue;
 		}
 		
-		FText const MenuName = BlueprintContextTargetMenuImpl::GetContextTargetDisplayName(TargetEnum, ContextTarget);
+		FText const MenuName = BlueprintContextTargetMenuImpl::GetContextTargetDisplayName(TargetEnum, BitMaskOffset);
 		const FContextMenuTargetProfile& ProfileRef = TargetProfile;
 
 		Columns[ColIndex]->AddSlot()
@@ -562,7 +557,7 @@ void SBlueprintContextTargetMenu::Construct(const FArguments& InArgs, const FBlu
 				.IsEnabled_Raw(&TargetProfile, &FContextMenuTargetProfile::IsTargetEnabled, ContextTarget)
 				.IsChecked(this, &SBlueprintContextTargetMenu::GetTargetCheckedState, ContextTarget)
 				.OnCheckStateChanged(this, &SBlueprintContextTargetMenu::OnTargetCheckStateChanged, ContextTarget)
-				.ToolTipText_Lambda([TargetEnum, ContextTarget, &ProfileRef]()->FText
+				.ToolTipText_Lambda([TargetEnum, BitMaskOffset, ContextTarget, &ProfileRef]()->FText
 				{
 					if (!ProfileRef.IsTargetEnabled(ContextTarget))
 					{
@@ -570,7 +565,7 @@ void SBlueprintContextTargetMenu::Construct(const FArguments& InArgs, const FBlu
 					}
 					else if (TargetEnum != nullptr)
 					{
-						return TargetEnum->GetToolTipText(ContextTarget);
+						return TargetEnum->GetToolTipText(BitMaskOffset);
 					}
 					return LOCTEXT("GenericTargetTooltip", "Include variables/functions that belong to this target.");
 				})
