Changing a pin type on a Select node will recombine all subpins on the option and result pins before setting the new pin type.

diff --git a/Engine/Source/Editor/BlueprintGraph/Private/K2Node_Select.cpp b/Engine/Source/Editor/BlueprintGraph/Private/K2Node_Select.cpp
--- a/Engine/Source/Editor/BlueprintGraph/Private/K2Node_Select.cpp
+++ b/Engine/Source/Editor/BlueprintGraph/Private/K2Node_Select.cpp
@@ -777,12 +777,28 @@ void UK2Node_Select::PinTypeChanged(UEdGraphPin* Pin)
 		UEdGraphPin* ReturnPin = GetReturnValuePin();
 		if (ReturnPin->PinType != Pin->PinType)
 		{
+			// Recombine the sub pins back into the ReturnPin
+			if (ReturnPin->SubPins.Num() > 0)
+			{
+				Schema->RecombinePin(ReturnPin->SubPins[0]);
+			}
 			ReturnPin->PinType = Pin->PinType;
 			Schema->SetPinDefaultValueBasedOnType(ReturnPin);
 		}
 
-		// Set the options
+		// Recombine all option pins back into their root
 		TArray<UEdGraphPin*> OptionPins;
+ 		GetOptionPins(OptionPins);
+		for (UEdGraphPin* OptionPin : OptionPins)
+		{
+			// Recombine the sub pins back into the OptionPin
+			if (OptionPin->ParentPin == nullptr && OptionPin->SubPins.Num() > 0)
+			{
+				Schema->RecombinePin(OptionPin->SubPins[0]);
+			}
+		}
+
+		// Get the options again and set them
 		GetOptionPins(OptionPins);
 		for (auto It = OptionPins.CreateConstIterator(); It; It++)
 		{
diff --git a/Engine/Source/Editor/Kismet/Private/BlueprintEditor.cpp b/Engine/Source/Editor/Kismet/Private/BlueprintEditor.cpp
--- a/Engine/Source/Editor/Kismet/Private/BlueprintEditor.cpp
+++ b/Engine/Source/Editor/Kismet/Private/BlueprintEditor.cpp
@@ -3838,53 +3838,52 @@ bool FBlueprintEditor::CanRemoveOptionPin() const
 
 void FBlueprintEditor::OnChangePinType()
 {
-	if (UEdGraphPin* Pin = GetCurrentlySelectedPin())
+	if (UEdGraphPin* SelectedPin = GetCurrentlySelectedPin())
 	{
-		TSharedPtr<SGraphEditor> FocusedGraphEd = FocusedGraphEdPtr.Pin();
-		if (FocusedGraphEd.IsValid())
+		// Grab the root pin, that is what we want to edit
+		UEdGraphPin* RootPin = SelectedPin;
+		while(RootPin->ParentPin)
 		{
-			UEdGraphPin* SelectedPin = FocusedGraphEd->GetGraphPinForMenu();
-			if (SelectedPin)
-			{
-				const UEdGraphSchema_K2* Schema = GetDefault<UEdGraphSchema_K2>();
+			RootPin = RootPin->ParentPin;
+		}
 
-				// If this is the index node of the select node, we need to use the index list of types
-				UK2Node_Select* SelectNode = Cast<UK2Node_Select>(SelectedPin->GetOwningNode());
-				if (SelectNode && SelectNode->GetIndexPin() == SelectedPin)
-				{
-					TSharedRef<SCompoundWidget> PinChange = SNew(SPinTypeSelector, FGetPinTypeTree::CreateUObject(Schema, &UEdGraphSchema_K2::GetVariableIndexTypeTree))
-						.TargetPinType(this, &FBlueprintEditor::OnGetPinType, SelectedPin)
-						.OnPinTypeChanged(this, &FBlueprintEditor::OnChangePinTypeFinished, SelectedPin)
-						.Schema(Schema)
-						.bAllowExec(false)
-						.IsEnabled(true)
-						.bAllowArrays(false);
-
-					PinTypeChangePopupWindow = FSlateApplication::Get().PushMenu( 
-						GetToolkitHost()->GetParentWidget(), // Parent widget should be k2 not the menu thats open or it will be closed when the menu is dismissed
-						PinChange,
-						FSlateApplication::Get().GetCursorPos(), // summon location
-						FPopupTransitionEffect(FPopupTransitionEffect::TypeInPopup)
-						);
-				}
-				else
-				{
-					TSharedRef<SCompoundWidget> PinChange = SNew(SPinTypeSelector, FGetPinTypeTree::CreateUObject(Schema, &UEdGraphSchema_K2::GetVariableTypeTree))
-						.TargetPinType(this, &FBlueprintEditor::OnGetPinType, SelectedPin)
-						.OnPinTypeChanged(this, &FBlueprintEditor::OnChangePinTypeFinished, SelectedPin)
-						.Schema(Schema)
-						.bAllowExec(false)
-						.IsEnabled(true)
-						.bAllowArrays(false);
-
-					PinTypeChangePopupWindow = FSlateApplication::Get().PushMenu( 
-						GetToolkitHost()->GetParentWidget(), // Parent widget should be k2 not the menu thats open or it will be closed when the menu is dismissed
-						PinChange,
-						FSlateApplication::Get().GetCursorPos(), // summon location
-						FPopupTransitionEffect(FPopupTransitionEffect::TypeInPopup)
-						);
-				}
-			}
+		const UEdGraphSchema_K2* Schema = GetDefault<UEdGraphSchema_K2>();
+
+		// If this is the index node of the select node, we need to use the index list of types
+		UK2Node_Select* SelectNode = Cast<UK2Node_Select>(SelectedPin->GetOwningNode());
+		if (SelectNode && SelectNode->GetIndexPin() == SelectedPin)
+		{
+			TSharedRef<SCompoundWidget> PinChange = SNew(SPinTypeSelector, FGetPinTypeTree::CreateUObject(Schema, &UEdGraphSchema_K2::GetVariableIndexTypeTree))
+				.TargetPinType(this, &FBlueprintEditor::OnGetPinType, RootPin)
+				.OnPinTypeChanged(this, &FBlueprintEditor::OnChangePinTypeFinished, SelectedPin)
+				.Schema(Schema)
+				.bAllowExec(false)
+				.IsEnabled(true)
+				.bAllowArrays(false);
+
+			PinTypeChangePopupWindow = FSlateApplication::Get().PushMenu( 
+				GetToolkitHost()->GetParentWidget(), // Parent widget should be k2 not the menu thats open or it will be closed when the menu is dismissed
+				PinChange,
+				FSlateApplication::Get().GetCursorPos(), // summon location
+				FPopupTransitionEffect(FPopupTransitionEffect::TypeInPopup)
+				);
+		}
+		else
+		{
+			TSharedRef<SCompoundWidget> PinChange = SNew(SPinTypeSelector, FGetPinTypeTree::CreateUObject(Schema, &UEdGraphSchema_K2::GetVariableTypeTree))
+				.TargetPinType(this, &FBlueprintEditor::OnGetPinType, RootPin)
+				.OnPinTypeChanged(this, &FBlueprintEditor::OnChangePinTypeFinished, SelectedPin)
+				.Schema(Schema)
+				.bAllowExec(false)
+				.IsEnabled(true)
+				.bAllowArrays(false);
+
+			PinTypeChangePopupWindow = FSlateApplication::Get().PushMenu( 
+				GetToolkitHost()->GetParentWidget(), // Parent widget should be k2 not the menu thats open or it will be closed when the menu is dismissed
+				PinChange,
+				FSlateApplication::Get().GetCursorPos(), // summon location
+				FPopupTransitionEffect(FPopupTransitionEffect::TypeInPopup)
+				);
 		}
 	}
 }
@@ -3899,12 +3898,9 @@ void FBlueprintEditor::OnChangePinTypeFinished(const FEdGraphPinType& PinType, U
 	if (FBlueprintEditorUtils::IsPinTypeValid(PinType))
 	{
 		InSelectedPin->PinType = PinType;
-		if (UEdGraphPin* SelectedPin = GetCurrentlySelectedPin())
+		if (UK2Node_Select* SelectNode = Cast<UK2Node_Select>(InSelectedPin->GetOwningNode()))
 		{
-			if (UK2Node_Select* SelectNode = Cast<UK2Node_Select>(SelectedPin->GetOwningNode()))
-			{
-				SelectNode->ChangePinType(SelectedPin);
-			}
+			SelectNode->ChangePinType(InSelectedPin);
 		}
 	}
 
