C++ Functions that are tagged as BlueprintImplementableEvent and are overriden as a function graph in Blueprints will display an error when attempting to drag and drop into the graph to inform the user that the function is not BlueprintCallable.

diff --git a/Engine/Source/Editor/Kismet/Private/SMyBlueprint.cpp b/Engine/Source/Editor/Kismet/Private/SMyBlueprint.cpp
--- a/Engine/Source/Editor/Kismet/Private/SMyBlueprint.cpp
+++ b/Engine/Source/Editor/Kismet/Private/SMyBlueprint.cpp
@@ -1423,6 +1423,16 @@ FReply SMyBlueprint::OnActionDragged( const TArray< TSharedPtr<FEdGraphSchemaAct
 			
 			if (FuncAction->GraphType == EEdGraphSchemaAction_K2Graph::Function ||FuncAction->GraphType == EEdGraphSchemaAction_K2Graph::Interface)
 			{
+				// Callback function to report that the user cannot drop this function in the graph
+				auto CanDragDropAction = [](TSharedPtr<FEdGraphSchemaAction> /*DropAction*/, UEdGraph* /*HoveredGraphIn*/, FText& ImpededReasonOut, bool bIsBlueprintCallableFunction)->bool
+				{
+					if (!bIsBlueprintCallableFunction)
+					{
+						ImpededReasonOut = LOCTEXT("NonBlueprintCallable", "This function was not marked as Blueprint Callable and cannot be placed in a graph!");
+					}
+					return bIsBlueprintCallableFunction;
+				};
+
 				bool bIsBlueprintCallableFunction = false;
 				if (FuncAction->EdGraph != NULL)
 				{
@@ -1439,10 +1449,7 @@ FReply SMyBlueprint::OnActionDragged( const TArray< TSharedPtr<FEdGraphSchemaAct
 					}
 				}
 
-				if (bIsBlueprintCallableFunction)
-				{
-					return FReply::Handled().BeginDragDrop(FKismetFunctionDragDropAction::New(InAction, FuncAction->FuncName, BlueprintEditorPtr.Pin()->GetBlueprintObj()->SkeletonGeneratedClass, FMemberReference(), AnalyticsDelegate));
-				}
+				return FReply::Handled().BeginDragDrop(FKismetFunctionDragDropAction::New(InAction, FuncAction->FuncName, BlueprintEditorPtr.Pin()->GetBlueprintObj()->SkeletonGeneratedClass, FMemberReference(), AnalyticsDelegate, FKismetDragDropAction::FCanBeDroppedDelegate::CreateLambda(CanDragDropAction, bIsBlueprintCallableFunction)));
 			}
 			else if (FuncAction->GraphType == EEdGraphSchemaAction_K2Graph::Macro)
 			{
