Can set the "AdvancedDisplay" flag on Blueprint member variables.

PR #1222

diff --git a/Engine/Source/Editor/Kismet/Private/BlueprintDetailsCustomization.cpp b/Engine/Source/Editor/Kismet/Private/BlueprintDetailsCustomization.cpp
--- a/Engine/Source/Editor/Kismet/Private/BlueprintDetailsCustomization.cpp
+++ b/Engine/Source/Editor/Kismet/Private/BlueprintDetailsCustomization.cpp
@@ -767,6 +767,26 @@ void FBlueprintVarActionDetails::CustomizeDetails( IDetailLayoutBuilder& DetailL
 			.ToolTip(SaveGameTooltip)
 		];
 
+		TSharedPtr<SToolTip> AdvancedDisplayTooltip = IDocumentation::Get()->CreateToolTip(LOCTEXT("VariableAdvancedDisplay_Tooltip", "Hide this variable in Class Defaults windows by default"), NULL, DocLink, TEXT("AdvancedDisplay"));
+
+		Category.AddCustomRow(LOCTEXT("VariableAdvancedDisplay", "Advanced Display"), true)
+			.Visibility(TAttribute<EVisibility>(this, &FBlueprintVarActionDetails::GetAdvancedDisplayVisibility))
+			.NameContent()
+			[
+				SNew(STextBlock)
+				.ToolTip(AdvancedDisplayTooltip)
+				.Text(LOCTEXT("VariableAdvancedDisplay", "Advanced Display"))
+				.Font(DetailFontInfo)
+			]
+		.ValueContent()
+			[
+				SNew(SCheckBox)
+				.IsChecked(this, &FBlueprintVarActionDetails::OnGetAdvancedDisplayCheckboxState)
+				.OnCheckStateChanged(this, &FBlueprintVarActionDetails::OnAdvancedDisplayChanged)
+				.IsEnabled(IsVariableInBlueprint())
+				.ToolTip(AdvancedDisplayTooltip)
+			];
+
 		TSharedPtr<SToolTip> PropertyFlagsTooltip = IDocumentation::Get()->CreateToolTip(LOCTEXT("DefinedPropertyFlags_Tooltip", "List of defined flags for this property"), NULL, DocLink, TEXT("PropertyFlags"));
 
 		Category.AddCustomRow(LOCTEXT("DefinedPropertyFlags", "Defined Property Flags"), true)
@@ -1885,6 +1905,39 @@ void FBlueprintVarActionDetails::OnSaveGameChanged(ECheckBoxState InNewState)
 	}
 }
 
+EVisibility FBlueprintVarActionDetails::GetAdvancedDisplayVisibility() const
+{
+	UProperty* VariableProperty = CachedVariableProperty.Get();
+	if (VariableProperty)
+	{
+		if (IsABlueprintVariable(VariableProperty))
+		{
+			return EVisibility::Visible;
+		}
+	}
+	return EVisibility::Collapsed;
+}
+
+ECheckBoxState FBlueprintVarActionDetails::OnGetAdvancedDisplayCheckboxState() const
+{
+	UProperty* Property = CachedVariableProperty.Get();
+	if (Property)
+	{
+		return (Property && Property->HasAnyPropertyFlags(CPF_AdvancedDisplay)) ? ECheckBoxState::Checked : ECheckBoxState::Unchecked;
+	}
+	return ECheckBoxState::Unchecked;
+}
+
+void FBlueprintVarActionDetails::OnAdvancedDisplayChanged(ECheckBoxState InNewState)
+{
+	UProperty* Property = CachedVariableProperty.Get();
+	if (Property)
+	{
+		const bool bAdvancedFlag = (InNewState == ECheckBoxState::Checked);
+		FBlueprintEditorUtils::SetVariableAdvancedDisplayFlag(GetBlueprintObj(), Property->GetFName(), bAdvancedFlag);
+	}
+}
+
 EVisibility FBlueprintVarActionDetails::IsTooltipEditVisible() const
 {
 	UProperty* VariableProperty = CachedVariableProperty.Get();
diff --git a/Engine/Source/Editor/Kismet/Private/BlueprintDetailsCustomization.h b/Engine/Source/Editor/Kismet/Private/BlueprintDetailsCustomization.h
--- a/Engine/Source/Editor/Kismet/Private/BlueprintDetailsCustomization.h
+++ b/Engine/Source/Editor/Kismet/Private/BlueprintDetailsCustomization.h
@@ -169,6 +169,10 @@ private:
 	ECheckBoxState OnGetSaveGameCheckboxState() const;
 	void OnSaveGameChanged(ECheckBoxState InNewState);
 
+	EVisibility GetAdvancedDisplayVisibility() const;
+	ECheckBoxState OnGetAdvancedDisplayCheckboxState() const;
+	void OnAdvancedDisplayChanged(ECheckBoxState InNewState);
+
 	/** Refresh the property flags list */
 	void RefreshPropertyFlags();
 
diff --git a/Engine/Source/Editor/UnrealEd/Private/Kismet2/BlueprintEditorUtils.cpp b/Engine/Source/Editor/UnrealEd/Private/Kismet2/BlueprintEditorUtils.cpp
--- a/Engine/Source/Editor/UnrealEd/Private/Kismet2/BlueprintEditorUtils.cpp
+++ b/Engine/Source/Editor/UnrealEd/Private/Kismet2/BlueprintEditorUtils.cpp
@@ -3389,6 +3389,25 @@ void FBlueprintEditorUtils::SetVariableSaveGameFlag(UBlueprint* InBlueprint, con
 	FBlueprintEditorUtils::MarkBlueprintAsStructurallyModified(InBlueprint);
 }
 
+void FBlueprintEditorUtils::SetVariableAdvancedDisplayFlag(UBlueprint* InBlueprint, const FName& InVarName, const bool bInIsAdvancedDisplay)
+{
+	const int32 VarIndex = FBlueprintEditorUtils::FindNewVariableIndex(InBlueprint, InVarName);
+
+	if (VarIndex != INDEX_NONE)
+	{
+		if( bInIsAdvancedDisplay )
+		{
+			InBlueprint->NewVariables[VarIndex].PropertyFlags |= CPF_AdvancedDisplay;
+		}
+		else
+		{
+			InBlueprint->NewVariables[VarIndex].PropertyFlags &= ~CPF_AdvancedDisplay;
+		}
+	}
+
+	FBlueprintEditorUtils::MarkBlueprintAsStructurallyModified(InBlueprint);
+}
+
 struct FMetaDataDependencyHelper
 {
 	static void OnChange(UBlueprint* Blueprint, FName MetaDataKey)
diff --git a/Engine/Source/Editor/UnrealEd/Public/Kismet2/BlueprintEditorUtils.h b/Engine/Source/Editor/UnrealEd/Public/Kismet2/BlueprintEditorUtils.h
--- a/Engine/Source/Editor/UnrealEd/Public/Kismet2/BlueprintEditorUtils.h
+++ b/Engine/Source/Editor/UnrealEd/Public/Kismet2/BlueprintEditorUtils.h
@@ -845,6 +845,14 @@ public:
 	 */
 	static void SetVariableSaveGameFlag(UBlueprint* InBlueprint, const FName& InVarName, const bool bInIsSaveGame);
 
+	/**
+	 * Sets the Advanced Display flag on the variable with the specified name
+	 *
+	 * @param	InVarName				Name of the var to set the flag on
+	 * @param	bInIsAdvancedDisplay	The new value to set the bitflag to
+	 */
+	static void SetVariableAdvancedDisplayFlag(UBlueprint* InBlueprint, const FName& InVarName, const bool bInIsAdvancedDisplay);
+
 	/** Sets a metadata key/value on the specified variable
 	 *
 	 * @param Blueprint				The Blueprint to find the variable in
