The eye icon appears now for marking a user created component variable as editable in the My Blueprints window. Does not appear for SCS components.

diff --git a/Engine/Source/Editor/Kismet/Private/SBlueprintPalette.cpp b/Engine/Source/Editor/Kismet/Private/SBlueprintPalette.cpp
--- a/Engine/Source/Editor/Kismet/Private/SBlueprintPalette.cpp
+++ b/Engine/Source/Editor/Kismet/Private/SBlueprintPalette.cpp
@@ -775,7 +775,7 @@ public:
 			UStruct* VarSourceScope = (VariableProp != NULL) ? CastChecked<UStruct>(VariableProp->GetOuter()) : NULL;
 			const bool bIsBlueprintVariable = (VarSourceScope == BlueprintObj->SkeletonGeneratedClass);
 			const bool bIsComponentVar = (VariableObjProp != NULL) && (VariableObjProp->PropertyClass != NULL) && (VariableObjProp->PropertyClass->IsChildOf(UActorComponent::StaticClass()));
-			bShouldHaveAVisibilityToggle = bIsBlueprintVariable && !bIsComponentVar;
+			bShouldHaveAVisibilityToggle = bIsBlueprintVariable && (!bIsComponentVar || FBlueprintEditorUtils::IsVariableCreatedByBlueprint(BlueprintObj, VariableObjProp));
 		}
 
 		this->ChildSlot[
