Added missing file

diff --git a/Engine/Source/Editor/Kismet/Public/BlueprintEditor.h b/Engine/Source/Editor/Kismet/Public/BlueprintEditor.h
--- a/Engine/Source/Editor/Kismet/Public/BlueprintEditor.h
+++ b/Engine/Source/Editor/Kismet/Public/BlueprintEditor.h
@@ -869,7 +869,7 @@ protected:
 	void RenameGraph(class UEdGraphNode* GraphNode, const FString& NewName);
 
 	/** Called when a node's title is being committed for a rename so it can be verified */
-	bool OnNodeVerifyTitleCommit(const FText& NewText, UEdGraphNode* NodeBeingChanged);
+	bool OnNodeVerifyTitleCommit(const FText& NewText, UEdGraphNode* NodeBeingChanged, FText& OutErrorMessage);
 
 	/**Load macro & function blueprint libraries from asset registry*/
 	void LoadLibrariesFromAssetRegistry();
