Dragging and dropping variables from the MyBlueprint tab will no longer default to assuming that the variable is readonly.

diff --git a/Engine/Source/Editor/Kismet/Private/BPVariableDragDropAction.cpp b/Engine/Source/Editor/Kismet/Private/BPVariableDragDropAction.cpp
--- a/Engine/Source/Editor/Kismet/Private/BPVariableDragDropAction.cpp
+++ b/Engine/Source/Editor/Kismet/Private/BPVariableDragDropAction.cpp
@@ -59,7 +59,7 @@ void FKismetVariableDragDropAction::HoverTargetChanged()
 
 	UEdGraphPin* PinUnderCursor = GetHoveredPin();
 
-	bool bCanMakeSetter = false;
+	bool bCanMakeSetter = true;
 	bool bBadSchema = false;
 	bool bBadGraph = false;
 	UEdGraph* HoveredGraph = GetHoveredGraph();
@@ -301,7 +301,7 @@ void FKismetVariableDragDropAction::HoverTargetChanged()
 		Args.Add(TEXT("VariableName"), FText::FromString(VariableString));
 
 		StatusSymbol = FEditorStyle::GetBrush(TEXT("Graph.ConnectorFeedback.Error"));
-		Message = FText::Format(LOCTEXT("CannotPlaceSetter", "Variable '{VariableName}' is readonly you cannot set this variable."), Args);
+		Message = FText::Format(LOCTEXT("CannotPlaceSetter", "Variable '{VariableName}' is readonly, you cannot set this variable."), Args);
 	}
 	// Draw variable icon
 	else
