Can promote pins to local variables in Blueprint function graphs (or their child graphs).

Option only appears when available, will pull the default value over.

diff --git a/Engine/Source/Editor/BlueprintGraph/Private/EdGraphSchema_K2.cpp b/Engine/Source/Editor/BlueprintGraph/Private/EdGraphSchema_K2.cpp
--- a/Engine/Source/Editor/BlueprintGraph/Private/EdGraphSchema_K2.cpp
+++ b/Engine/Source/Editor/BlueprintGraph/Private/EdGraphSchema_K2.cpp
@@ -1487,6 +1487,11 @@ void UEdGraphSchema_K2::GetContextMenuActions(const UEdGraph* CurrentGraph, cons
 				if (InGraphPin->PinType.PinCategory != PC_Exec)
 				{
 					MenuBuilder->AddMenuEntry( FGraphEditorCommands::Get().PromoteToVariable );
+
+					if (FBlueprintEditorUtils::DoesSupportLocalVariables(CurrentGraph))
+					{
+						MenuBuilder->AddMenuEntry( FGraphEditorCommands::Get().PromoteToLocalVariable );
+					}
 				}
 	
 				if (PinHasSplittableStructType(InGraphPin) && InGraphNode->AllowSplitPins())
diff --git a/Engine/Source/Editor/GraphEditor/Private/GraphEditorActions.cpp b/Engine/Source/Editor/GraphEditor/Private/GraphEditorActions.cpp
--- a/Engine/Source/Editor/GraphEditor/Private/GraphEditorActions.cpp
+++ b/Engine/Source/Editor/GraphEditor/Private/GraphEditorActions.cpp
@@ -53,6 +53,7 @@ void FGraphEditorCommandsImpl::RegisterCommands()
 
 	UI_COMMAND( BreakPinLinks, "Break Link(s)", "Breaks pin links", EUserInterfaceActionType::Button, FInputChord() )
 	UI_COMMAND( PromoteToVariable, "Promote to Variable", "Promotes something to a variable", EUserInterfaceActionType::Button, FInputChord() )
+	UI_COMMAND( PromoteToLocalVariable, "Promote to Local Variable", "Promotes something to a local variable of the current function", EUserInterfaceActionType::Button, FInputChord() )
 	UI_COMMAND( SplitStructPin, "Split Struct Pin", "Breaks a struct pin in to a separate pin per element", EUserInterfaceActionType::Button, FInputChord() )
 	UI_COMMAND( RecombineStructPin, "Recombine Struct Pin", "Takes struct pins that have been broken in to composite elements and combines them back to a single struct pin", EUserInterfaceActionType::Button, FInputChord() )
 	UI_COMMAND( StartWatchingPin, "Watch this value", "Adds this pin or variable to the watch list", EUserInterfaceActionType::Button, FInputChord() )
diff --git a/Engine/Source/Editor/GraphEditor/Public/GraphEditorActions.h b/Engine/Source/Editor/GraphEditor/Public/GraphEditorActions.h
--- a/Engine/Source/Editor/GraphEditor/Public/GraphEditorActions.h
+++ b/Engine/Source/Editor/GraphEditor/Public/GraphEditorActions.h
@@ -77,6 +77,7 @@ public:
 	// Pin-specific actions
 	TSharedPtr< FUICommandInfo > BreakPinLinks;
 	TSharedPtr< FUICommandInfo > PromoteToVariable;
+	TSharedPtr< FUICommandInfo > PromoteToLocalVariable;
 	TSharedPtr< FUICommandInfo > SplitStructPin;
 	TSharedPtr< FUICommandInfo > RecombineStructPin;
 	TSharedPtr< FUICommandInfo > StartWatchingPin;
diff --git a/Engine/Source/Editor/Kismet/Private/BlueprintEditor.cpp b/Engine/Source/Editor/Kismet/Private/BlueprintEditor.cpp
--- a/Engine/Source/Editor/Kismet/Private/BlueprintEditor.cpp
+++ b/Engine/Source/Editor/Kismet/Private/BlueprintEditor.cpp
@@ -872,8 +872,13 @@ TSharedRef<SGraphEditor> FBlueprintEditor::CreateGraphEditorWidget(TSharedRef<FT
 		GraphEditorCommands = MakeShareable( new FUICommandList );
 		{
 			GraphEditorCommands->MapAction(FGraphEditorCommands::Get().PromoteToVariable,
-				FExecuteAction::CreateSP( this, &FBlueprintEditor::OnPromoteToVariable ),
-				FCanExecuteAction::CreateSP( this, &FBlueprintEditor::CanPromoteToVariable )
+				FExecuteAction::CreateSP( this, &FBlueprintEditor::OnPromoteToVariable, true ),
+				FCanExecuteAction::CreateSP( this, &FBlueprintEditor::CanPromoteToVariable, true )
+				);
+
+			GraphEditorCommands->MapAction(FGraphEditorCommands::Get().PromoteToLocalVariable,
+				FExecuteAction::CreateSP( this, &FBlueprintEditor::OnPromoteToVariable, false ),
+				FCanExecuteAction::CreateSP( this, &FBlueprintEditor::CanPromoteToVariable, false )
 				);
 
 			GraphEditorCommands->MapAction(FGraphEditorCommands::Get().SplitStructPin,
@@ -3393,30 +3398,34 @@ void FBlueprintEditor::DumpMessagesToCompilerLog(const TArray<TSharedRef<FTokeni
 	}
 }
 
-void FBlueprintEditor::DoPromoteToVariable( UBlueprint* InBlueprint, UEdGraphPin* InTargetPin )
+void FBlueprintEditor::DoPromoteToVariable( UBlueprint* InBlueprint, UEdGraphPin* InTargetPin, bool bInToMemberVariable )
 {
 	UEdGraphNode* PinNode = InTargetPin->GetOwningNode();
 	check(PinNode);
 	UEdGraph* GraphObj = PinNode->GetGraph();
 	check(GraphObj);
 
-	const FScopedTransaction Transaction( LOCTEXT("PromoteToVariable", "Promote To Variable") );
+	// Used for promoting to local variable
+	UEdGraph* FunctionGraph = nullptr;
+
+	const FScopedTransaction Transaction( bInToMemberVariable? LOCTEXT("PromoteToVariable", "Promote To Variable") : LOCTEXT("PromoteToLocalVariable", "Promote to Local Variable") );
 	InBlueprint->Modify();
 	GraphObj->Modify();
 
-	FString VarNameString = TEXT("NewVar");
-	FName VarName = FName(*VarNameString);
-
-	// Make sure the new name is valid
-	TSharedPtr<INameValidatorInterface> NameValidator = MakeShareable(new FKismetNameValidator(InBlueprint));
-	int32 Index = 0;
-	while (NameValidator->IsValid(VarName) != Ok)
+	FName VarName;
+	bool bWasSuccessful = false;
+	if (bInToMemberVariable)
 	{
-		VarName = FName(*FString::Printf(TEXT("%s%i"), *VarNameString, Index));
-		++Index;
+		VarName = FBlueprintEditorUtils::FindUniqueKismetName(GetBlueprintObj(), TEXT("NewVar"));
+		bWasSuccessful = FBlueprintEditorUtils::AddMemberVariable( GetBlueprintObj(), VarName, InTargetPin->PinType, InTargetPin->GetDefaultAsString() );
+	}
+	else
+	{
+		ensure(FBlueprintEditorUtils::DoesSupportLocalVariables(GraphObj));
+		VarName = FBlueprintEditorUtils::FindUniqueKismetName(GetBlueprintObj(), TEXT("NewLocalVar"));
+		FunctionGraph = FBlueprintEditorUtils::GetTopLevelGraph(GraphObj);
+		bWasSuccessful = FBlueprintEditorUtils::AddLocalVariable( GetBlueprintObj(), FunctionGraph, VarName, InTargetPin->PinType, InTargetPin->GetDefaultAsString() );
 	}
-
-	const bool bWasSuccessful = FBlueprintEditorUtils::AddMemberVariable( GetBlueprintObj(), VarName, InTargetPin->PinType, InTargetPin->GetDefaultAsString() );
 
 	if (bWasSuccessful)
 	{
@@ -3434,7 +3443,14 @@ void FBlueprintEditor::DoPromoteToVariable( UBlueprint* InBlueprint, UEdGraphPin
 			TemplateNode = NewObject<UK2Node_VariableSet>();
 		}
 
-		TemplateNode->VariableReference.SetSelfMember(VarName);
+		if (bInToMemberVariable)
+		{
+			TemplateNode->VariableReference.SetSelfMember(VarName);
+		}
+		else
+		{
+			TemplateNode->VariableReference.SetLocalMember(VarName, FunctionGraph->GetName(), FBlueprintEditorUtils::FindLocalVariableGuidByName(InBlueprint, FunctionGraph, VarName));
+		}
 		NodeInfo.NodeTemplate = TemplateNode;
 
 		// Set position of new node to be close to node we clicked on
@@ -3448,7 +3464,7 @@ void FBlueprintEditor::DoPromoteToVariable( UBlueprint* InBlueprint, UEdGraphPin
 	}
 }
 
-void FBlueprintEditor::OnPromoteToVariable()
+void FBlueprintEditor::OnPromoteToVariable(bool bInToMemberVariable)
 {
 	TSharedPtr<SGraphEditor> FocusedGraphEd = FocusedGraphEdPtr.Pin();
 	if (FocusedGraphEd.IsValid())
@@ -3459,11 +3475,11 @@ void FBlueprintEditor::OnPromoteToVariable()
 		check(GetBlueprintObj()->SkeletonGeneratedClass);
 		check(TargetPin);
 
-		DoPromoteToVariable( GetBlueprintObj(), TargetPin );
+		DoPromoteToVariable( GetBlueprintObj(), TargetPin, bInToMemberVariable );
 	}
 }
 
-bool FBlueprintEditor::CanPromoteToVariable() const
+bool FBlueprintEditor::CanPromoteToVariable(bool bInToMemberVariable) const
 {
 	const UEdGraphSchema_K2* K2Schema = GetDefault<UEdGraphSchema_K2>();
 
@@ -3473,7 +3489,10 @@ bool FBlueprintEditor::CanPromoteToVariable() const
 	{
 		if (UEdGraphPin* Pin = FocusedGraphEd->GetGraphPinForMenu())
 		{
-			bCanPromote = K2Schema->CanPromotePinToVariable(*Pin);
+			if (bInToMemberVariable || (!bInToMemberVariable && FBlueprintEditorUtils::DoesSupportLocalVariables(FocusedGraphEd->GetCurrentGraph())))
+			{
+				bCanPromote = K2Schema->CanPromotePinToVariable(*Pin);
+			}
 		}
 	}
 	
diff --git a/Engine/Source/Editor/Kismet/Private/SBlueprintActionMenu.cpp b/Engine/Source/Editor/Kismet/Private/SBlueprintActionMenu.cpp
--- a/Engine/Source/Editor/Kismet/Private/SBlueprintActionMenu.cpp
+++ b/Engine/Source/Editor/Kismet/Private/SBlueprintActionMenu.cpp
@@ -23,11 +23,12 @@
 USTRUCT()
 struct FBlueprintAction_PromoteVariable : public FEdGraphSchemaAction
 {
-	FBlueprintAction_PromoteVariable()
+	FBlueprintAction_PromoteVariable(bool bInToMemberVariable)
 		: FEdGraphSchemaAction()
+		, bToMemberVariable(bInToMemberVariable)
 	{
-		MenuDescription = LOCTEXT("PromoteToVariable", "Promote to variable");
-		TooltipDescription = LOCTEXT("PromoteToVariable", "Promote to variable").ToString();
+		MenuDescription = bInToMemberVariable? LOCTEXT("PromoteToVariable", "Promote to variable") : LOCTEXT("PromoteToLocalVariable", "Promote to local variable");
+		TooltipDescription = bInToMemberVariable? LOCTEXT("PromoteToVariable", "Promote to variable").ToString() : LOCTEXT("PromoteToLocalVariable", "Promote to local variable").ToString();
 		Grouping = 1;
 	}
 
@@ -39,7 +40,7 @@ struct FBlueprintAction_PromoteVariable : public FEdGraphSchemaAction
 			UBlueprint* Blueprint = FBlueprintEditorUtils::FindBlueprintForGraph(ParentGraph);
 			if( ( MyBlueprintEditor.IsValid() == true ) && ( Blueprint != NULL ) )
 			{
-				MyBlueprintEditor.Pin()->DoPromoteToVariable( Blueprint, FromPin );
+				MyBlueprintEditor.Pin()->DoPromoteToVariable( Blueprint, FromPin, bToMemberVariable );
 			}
 		}
 		return NULL;		
@@ -48,6 +49,9 @@ struct FBlueprintAction_PromoteVariable : public FEdGraphSchemaAction
 
 	/* Pointer to the blueprint editor containing the blueprint in which we will promote the variable. */
 	TWeakPtr<class FBlueprintEditor> MyBlueprintEditor;
+
+	/* TRUE if promoting to member variable, FALSE if promoting to local variable */
+	bool bToMemberVariable;
 };
 
 /**
@@ -536,9 +540,16 @@ void SBlueprintActionMenu::TryInsertPromoteToVariable(FBlueprintActionContext co
 	{
 		if (K2Schema->CanPromotePinToVariable(*MenuContext.Pins[0]))
 		{
-			TSharedPtr<FBlueprintAction_PromoteVariable> PromoteAction = TSharedPtr<FBlueprintAction_PromoteVariable>(new FBlueprintAction_PromoteVariable());
+			TSharedPtr<FBlueprintAction_PromoteVariable> PromoteAction = TSharedPtr<FBlueprintAction_PromoteVariable>(new FBlueprintAction_PromoteVariable(true));
 			PromoteAction->MyBlueprintEditor = EditorPtr;
 			OutAllActions.AddAction( PromoteAction );
+
+			if (MenuContext.Graphs.Num() == 1 && FBlueprintEditorUtils::DoesSupportLocalVariables(MenuContext.Graphs[0]))
+			{
+				TSharedPtr<FBlueprintAction_PromoteVariable> LocalPromoteAction = TSharedPtr<FBlueprintAction_PromoteVariable>(new FBlueprintAction_PromoteVariable(false));
+				LocalPromoteAction->MyBlueprintEditor = EditorPtr;
+				OutAllActions.AddAction( LocalPromoteAction );
+			}
 		}
 	}
 }
diff --git a/Engine/Source/Editor/Kismet/Public/BlueprintEditor.h b/Engine/Source/Editor/Kismet/Public/BlueprintEditor.h
--- a/Engine/Source/Editor/Kismet/Public/BlueprintEditor.h
+++ b/Engine/Source/Editor/Kismet/Public/BlueprintEditor.h
@@ -465,10 +465,11 @@ public:
 	/** 
 	 * Perform the actual promote to variable action on the given pin in the given blueprint.
 	 *
-	 * @param	InBlueprint	The blueprint in which to create the variable.
-	 * @param	InTargetPin The pin on which to base the variable.
+	 * @param	InBlueprint				The blueprint in which to create the variable.
+	 * @param	InTargetPin				The pin on which to base the variable.
+	 * @param	bInToMemberVariable		TRUE if attempting to create a member variable, FALSE if the variable should be local
 	 */
-	void DoPromoteToVariable( UBlueprint* InBlueprint, UEdGraphPin* InTargetPin );		
+	void DoPromoteToVariable( UBlueprint* InBlueprint, UEdGraphPin* InTargetPin, bool bInToMemberVariable );		
 
 	/** Called when node is spawned by keymap */
 	void OnNodeSpawnedByKeymap();
@@ -641,8 +642,8 @@ protected:
 	UEdGraphPin* GetCurrentlySelectedPin() const;
 
 	// UI Action functionality
-	void OnPromoteToVariable();
-	bool CanPromoteToVariable() const;
+	void OnPromoteToVariable(bool bInToMemberVariable);
+	bool CanPromoteToVariable(bool bInToMemberVariable) const;
 
 	void OnSplitStructPin();
 	bool CanSplitStructPin() const;
diff --git a/Engine/Source/Editor/UnrealEd/Private/Kismet2/BlueprintEditorUtils.cpp b/Engine/Source/Editor/UnrealEd/Private/Kismet2/BlueprintEditorUtils.cpp
--- a/Engine/Source/Editor/UnrealEd/Private/Kismet2/BlueprintEditorUtils.cpp
+++ b/Engine/Source/Editor/UnrealEd/Private/Kismet2/BlueprintEditorUtils.cpp
@@ -4564,7 +4564,7 @@ FBPVariableDescription FBlueprintEditorUtils::DuplicateVariableDescription(UBlue
 	return NewVar;
 }
 
-bool FBlueprintEditorUtils::AddLocalVariable(UBlueprint* Blueprint, UEdGraph* InTargetGraph, const FName InNewVarName, const FEdGraphPinType& InNewVarType)
+bool FBlueprintEditorUtils::AddLocalVariable(UBlueprint* Blueprint, UEdGraph* InTargetGraph, const FName InNewVarName, const FEdGraphPinType& InNewVarType, const FString& DefaultValue/*= FString()*/)
 {
 	if(InTargetGraph != NULL && InTargetGraph->GetSchema()->GetGraphType(InTargetGraph) == GT_Function)
 	{
@@ -4585,6 +4585,7 @@ bool FBlueprintEditorUtils::AddLocalVariable(UBlueprint* Blueprint, UEdGraph* In
 		NewVar.VarType = InNewVarType;
 		NewVar.FriendlyName = FName::NameToDisplayString( NewVar.VarName.ToString(), (NewVar.VarType.PinCategory == K2Schema->PC_Boolean) ? true : false );
 		NewVar.Category = K2Schema->VR_DefaultCategory;
+		NewVar.DefaultValue = DefaultValue;
 
 		PostSetupObjectPinType(Blueprint, NewVar);
 
diff --git a/Engine/Source/Editor/UnrealEd/Private/Tests/EditorBuildPromotionTests.cpp b/Engine/Source/Editor/UnrealEd/Private/Tests/EditorBuildPromotionTests.cpp
--- a/Engine/Source/Editor/UnrealEd/Private/Tests/EditorBuildPromotionTests.cpp
+++ b/Engine/Source/Editor/UnrealEd/Private/Tests/EditorBuildPromotionTests.cpp
@@ -495,7 +495,7 @@ namespace EditorBuildPromotionTestUtils
 		FBlueprintEditor* CurrentBlueprintEditor = (FBlueprintEditor*)OpenEditor;
 
 		UEdGraphPin* PinToPromote = Node->FindPin(PinName);
-		CurrentBlueprintEditor->DoPromoteToVariable(InBlueprint, PinToPromote);
+		CurrentBlueprintEditor->DoPromoteToVariable(InBlueprint, PinToPromote, true);
 	}
 
 	/**
diff --git a/Engine/Source/Editor/UnrealEd/Public/Kismet2/BlueprintEditorUtils.h b/Engine/Source/Editor/UnrealEd/Public/Kismet2/BlueprintEditorUtils.h
--- a/Engine/Source/Editor/UnrealEd/Public/Kismet2/BlueprintEditorUtils.h
+++ b/Engine/Source/Editor/UnrealEd/Public/Kismet2/BlueprintEditorUtils.h
@@ -710,7 +710,7 @@ public:
 	 *
 	 * @return	true if it succeeds, false if it fails.
 	 */
-	static bool AddLocalVariable(UBlueprint* Blueprint, UEdGraph* InTargetGraph, const FName InNewVarName, const FEdGraphPinType& InNewVarType);
+	static bool AddLocalVariable(UBlueprint* Blueprint, UEdGraph* InTargetGraph, const FName InNewVarName, const FEdGraphPinType& InNewVarType, const FString& DefaultValue = FString());
 
 	/**
 	 * Removes a member variable if it was declared in this blueprint and not in a base class.
