[UE-16629] New code did not treat the blueprint editor as optional, crashed when executed by the diff tool

diff --git a/Engine/Source/Editor/Kismet/Private/SBlueprintPalette.cpp b/Engine/Source/Editor/Kismet/Private/SBlueprintPalette.cpp
--- a/Engine/Source/Editor/Kismet/Private/SBlueprintPalette.cpp
+++ b/Engine/Source/Editor/Kismet/Private/SBlueprintPalette.cpp
@@ -670,9 +670,9 @@ public:
 	 * @param  InVariableProperty		The variable property to select
 	 * @param  InBlueprintEditor			A pointer to the blueprint editor that the palette belongs to.
 	 */
-	void Construct(const FArguments& InArgs, UProperty* InVariableProperty, TWeakPtr<FBlueprintEditor> InBlueprintEditor)
+	void Construct(const FArguments& InArgs, UProperty* InVariableProperty, UBlueprint* InBlueprint, TWeakPtr<FBlueprintEditor> InBlueprintEditor)
 	{
-		BlueprintObj = InBlueprintEditor.Pin()->GetBlueprintObj();
+		BlueprintObj = InBlueprint;
 		BlueprintEditorPtr = InBlueprintEditor;
 		VariableProperty = InVariableProperty;
 
@@ -1004,7 +1004,7 @@ void SBlueprintPaletteItem::Construct(const FArguments& InArgs, FCreateWidgetFor
 		if (FBlueprintEditorUtils::IsVariableCreatedByBlueprint(Blueprint, VariableProp) || Cast<UFunction>(VariableProp->GetOuter()))
 		{
 			const UEdGraphSchema_K2* Schema = GetDefault<UEdGraphSchema_K2>();
-			IconWidget = SNew(SPinTypeSelectorHelper, VariableProp, BlueprintEditorPtr);
+			IconWidget = SNew(SPinTypeSelectorHelper, VariableProp, Blueprint, BlueprintEditorPtr);
 		}
 	}
 
