[UE-15960] Ensure that we use the inherited component override template for SCS parent nodes when browsing sockets for attachment in the Blueprint editor.

change summary:
- modified FBlueprintComponentDetails::OnBrowseSocket() to invoke FSCSEditorTreeNodeComponent::GetEditableComponentTemplate() rather than GetComponentTemplate()

diff --git a/Engine/Source/Editor/Kismet/Private/BlueprintDetailsCustomization.cpp b/Engine/Source/Editor/Kismet/Private/BlueprintDetailsCustomization.cpp
--- a/Engine/Source/Editor/Kismet/Private/BlueprintDetailsCustomization.cpp
+++ b/Engine/Source/Editor/Kismet/Private/BlueprintDetailsCustomization.cpp
@@ -4863,7 +4863,7 @@ void FBlueprintComponentDetails::OnBrowseSocket()
 
 		if (ParentFNode.IsValid())
 		{
-			if (USceneComponent* ParentSceneComponent = Cast<USceneComponent>(ParentFNode->GetComponentTemplate()))
+			if (USceneComponent* ParentSceneComponent = Cast<USceneComponent>(ParentFNode->GetEditableComponentTemplate(Editor->GetBlueprint())))
 			{
 				if (ParentSceneComponent->HasAnySockets())
 				{
