UE-16863 Crash upon using the undo shortcut twice after expanding a Math Expression node

diff --git a/Engine/Source/Editor/Kismet/Private/SBlueprintPalette.cpp b/Engine/Source/Editor/Kismet/Private/SBlueprintPalette.cpp
--- a/Engine/Source/Editor/Kismet/Private/SBlueprintPalette.cpp
+++ b/Engine/Source/Editor/Kismet/Private/SBlueprintPalette.cpp
@@ -1540,7 +1540,10 @@ TSharedPtr<SToolTip> SBlueprintPaletteItem::ConstructToolTipWidget() const
 			if (GraphAction->EdGraph != NULL)
 			{
 				FGraphDisplayInfo DisplayInfo;
-				GraphAction->EdGraph->GetSchema()->GetGraphDisplayInformation(*(GraphAction->EdGraph), DisplayInfo);
+				if (auto GraphSchema = GraphAction->EdGraph->GetSchema())
+				{
+					GraphSchema->GetGraphDisplayInformation(*(GraphAction->EdGraph), DisplayInfo);
+				}
 
 				DocExcerptRef.DocLink = DisplayInfo.DocLink;
 				DocExcerptRef.DocExcerptName = DisplayInfo.DocExcerptName;
diff --git a/Engine/Source/Editor/Kismet/Private/SMyBlueprint.cpp b/Engine/Source/Editor/Kismet/Private/SMyBlueprint.cpp
--- a/Engine/Source/Editor/Kismet/Private/SMyBlueprint.cpp
+++ b/Engine/Source/Editor/Kismet/Private/SMyBlueprint.cpp
@@ -818,12 +818,12 @@ TSharedRef<SWidget> SMyBlueprint::OnCreateWidgetForAction(FCreateWidgetForAction
 
 void SMyBlueprint::GetChildGraphs(UEdGraph* InEdGraph, int32 const SectionId, FGraphActionListBuilderBase& OutAllActions, FText ParentCategory)
 {
-	// Grab subgraphs
-	const UEdGraphSchema* Schema = InEdGraph->GetSchema();
-
 	// Grab display info
 	FGraphDisplayInfo EdGraphDisplayInfo;
-	Schema->GetGraphDisplayInformation(*InEdGraph, EdGraphDisplayInfo);
+	if (const UEdGraphSchema* Schema = InEdGraph ? InEdGraph->GetSchema() : nullptr)
+	{
+		Schema->GetGraphDisplayInformation(*InEdGraph, EdGraphDisplayInfo);
+	}
 	const FText EdGraphDisplayName = EdGraphDisplayInfo.DisplayName;
 
 	// Grab children graphs
@@ -832,9 +832,11 @@ void SMyBlueprint::GetChildGraphs(UEdGraph* InEdGraph, int32 const SectionId, FG
 		UEdGraph* Graph = *It;
 		check(Graph);
 
-		const UEdGraphSchema* ChildSchema = Graph->GetSchema();
 		FGraphDisplayInfo ChildGraphDisplayInfo;
-		ChildSchema->GetGraphDisplayInformation(*Graph, ChildGraphDisplayInfo);
+		if (const UEdGraphSchema* ChildSchema = Graph->GetSchema())
+		{
+			ChildSchema->GetGraphDisplayInformation(*Graph, ChildGraphDisplayInfo);
+		}
 
 		FText DisplayText = ChildGraphDisplayInfo.DisplayName;
 
@@ -893,9 +895,11 @@ void SMyBlueprint::GetChildEvents(UEdGraph const* InEdGraph, int32 const Section
 	}
 
 	// grab the parent graph's name
-	UEdGraphSchema const* Schema = InEdGraph->GetSchema();
 	FGraphDisplayInfo EdGraphDisplayInfo;
-	Schema->GetGraphDisplayInformation(*InEdGraph, EdGraphDisplayInfo);
+	if (UEdGraphSchema const* Schema = InEdGraph->GetSchema())
+	{
+		Schema->GetGraphDisplayInformation(*InEdGraph, EdGraphDisplayInfo);
+	}
 	FText const EdGraphDisplayName = EdGraphDisplayInfo.DisplayName;
 	FText ActionCategory;
 	if (!ParentCategory.IsEmpty())
@@ -919,9 +923,11 @@ void SMyBlueprint::GetLocalVariables(FGraphActionListBuilderBase& OutAllActions)
 	if( TopLevelGraph )
 	{
 		// grab the parent graph's name
-		UEdGraphSchema const* Schema = TopLevelGraph->GetSchema();
 		FGraphDisplayInfo EdGraphDisplayInfo;
-		Schema->GetGraphDisplayInformation(*TopLevelGraph, EdGraphDisplayInfo);
+		if (UEdGraphSchema const* Schema = TopLevelGraph->GetSchema())
+		{
+			Schema->GetGraphDisplayInformation(*TopLevelGraph, EdGraphDisplayInfo);
+		}
 		FText const EdGraphDisplayName = EdGraphDisplayInfo.DisplayName;
 
 		TArray<UK2Node_FunctionEntry*> FunctionEntryNodes;
diff --git a/Engine/Source/Editor/UnrealEd/Private/EdGraphUtilities.cpp b/Engine/Source/Editor/UnrealEd/Private/EdGraphUtilities.cpp
--- a/Engine/Source/Editor/UnrealEd/Private/EdGraphUtilities.cpp
+++ b/Engine/Source/Editor/UnrealEd/Private/EdGraphUtilities.cpp
@@ -218,11 +218,13 @@ void FEdGraphUtilities::MergeChildrenGraphsIn(UEdGraph* MergeTarget, UEdGraph* P
 	{
 		UEdGraph* ChildGraph = ParentGraph->SubGraphs[Index];
 
-		auto NodeOwner = Cast<const UEdGraphNode>(ChildGraph->GetOuter());
+		auto NodeOwner = Cast<const UEdGraphNode>(ChildGraph ? ChildGraph->GetOuter() : nullptr);
 		const bool bNonVirtualGraph = NodeOwner ? NodeOwner->ShouldMergeChildGraphs() : true;
 
 		// Only merge children in with the same schema as the parent
-		const bool bSchemaMatches = ChildGraph->GetSchema()->GetClass()->IsChildOf(MergeTarget->GetSchema()->GetClass());
+		auto ChildSchema = ChildGraph ? ChildGraph->GetSchema() : nullptr;
+		auto TargetSchema = MergeTarget ? MergeTarget->GetSchema() : nullptr;
+		const bool bSchemaMatches = ChildSchema && TargetSchema && ChildSchema->GetClass()->IsChildOf(TargetSchema->GetClass());
 		const bool bDoMerge = bNonVirtualGraph && (!bRequireSchemaMatch || bSchemaMatches);
 		if (bDoMerge)
 		{
