Fixed typo in message dialog

diff --git a/Engine/Source/Editor/Kismet/Private/FindInBlueprintManager.cpp b/Engine/Source/Editor/Kismet/Private/FindInBlueprintManager.cpp
--- a/Engine/Source/Editor/Kismet/Private/FindInBlueprintManager.cpp
+++ b/Engine/Source/Editor/Kismet/Private/FindInBlueprintManager.cpp
@@ -1464,7 +1464,7 @@ void FFindInBlueprintSearchManager::CacheAllUncachedBlueprints(TWeakPtr< SFindIn
 		FText DialogTitle = LOCTEXT("ConfirmIndexAll_Title", "Indexing All");
 		FFormatNamedArguments Args;
 		Args.Add(TEXT("PackageCount"), UncachedBlueprints.Num());
-		const EAppReturnType::Type ReturnValue = FMessageDialog::Open(EAppMsgType::YesNoCancel, FText::Format(LOCTEXT("CacheAllConfirmationMessage", "This process can take a long time and the editor may become unresponsive; there are {PackageCount} Blueprints to load.\n\nWould you like to checkout, load, and save all Blueprints to make this indexing permanant? Otherwise, all Blueprints will still be loaded but you will be required to re-index the next time you start the editor!"), Args), &DialogTitle);
+		const EAppReturnType::Type ReturnValue = FMessageDialog::Open(EAppMsgType::YesNoCancel, FText::Format(LOCTEXT("CacheAllConfirmationMessage", "This process can take a long time and the editor may become unresponsive; there are {PackageCount} Blueprints to load.\n\nWould you like to checkout, load, and save all Blueprints to make this indexing permanent? Otherwise, all Blueprints will still be loaded but you will be required to re-index the next time you start the editor!"), Args), &DialogTitle);
 
 		// If Yes is chosen, checkout and save all Blueprints, if No is chosen, only load all Blueprints
 		if (ReturnValue != EAppReturnType::Cancel)
