BP - return node improvements:
- UE-14092 Cannot create a return node with an execution pin only for a function
- UE-7849 Being able to have multiple return nodes for a function
- minor fixes

diff --git a/Engine/Source/Editor/BlueprintGraph/Classes/K2Node_EditablePinBase.h b/Engine/Source/Editor/BlueprintGraph/Classes/K2Node_EditablePinBase.h
--- a/Engine/Source/Editor/BlueprintGraph/Classes/K2Node_EditablePinBase.h
+++ b/Engine/Source/Editor/BlueprintGraph/Classes/K2Node_EditablePinBase.h
@@ -134,7 +134,7 @@ class UK2Node_EditablePinBase : public UK2Node
 	 * @param	InPinType				The type info for the pin to create
 	 * @param	InDesiredDirection		Desired direction of the pin, will auto-correct if the direction is not allowed on the pin.
 	 */
-	BLUEPRINTGRAPH_API UEdGraphPin* CreateUserDefinedPin(const FString& InPinName, const FEdGraphPinType& InPinType, EEdGraphPinDirection InDesiredDirection);
+	BLUEPRINTGRAPH_API UEdGraphPin* CreateUserDefinedPin(const FString& InPinName, const FEdGraphPinType& InPinType, EEdGraphPinDirection InDesiredDirection, bool bUseUniqueName = true);
 
 	/**
 	 * Removes a pin from the user-defined array, and removes the pin with the same name from the Pins array
@@ -143,6 +143,13 @@ class UK2Node_EditablePinBase : public UK2Node
 	 */
 	BLUEPRINTGRAPH_API void RemoveUserDefinedPin( TSharedPtr<FUserPinInfo> PinToRemove );
 
+	/**
+	 * Removes from the user-defined array, and removes the pin with the same name from the Pins array
+	 *
+	 * @param	PinName name of pin to remove
+	 */
+	BLUEPRINTGRAPH_API void RemoveUserDefinedPinByName(const FString& PinName);
+
 	/**
 	 * Creates a new pin on the node from the specified user pin info.
 	 * Must be overridden so each type of node can ensure that the pin is created in the proper direction, etc
diff --git a/Engine/Source/Editor/BlueprintGraph/Classes/K2Node_FunctionResult.h b/Engine/Source/Editor/BlueprintGraph/Classes/K2Node_FunctionResult.h
--- a/Engine/Source/Editor/BlueprintGraph/Classes/K2Node_FunctionResult.h
+++ b/Engine/Source/Editor/BlueprintGraph/Classes/K2Node_FunctionResult.h
@@ -13,18 +13,25 @@ class UK2Node_FunctionResult : public UK2Node_FunctionTerminator
 	// Begin UEdGraphNode interface
 	virtual void AllocateDefaultPins() override;
 	virtual FText GetNodeTitle(ENodeTitleType::Type TitleType) const override;
-	virtual bool CanUserDeleteNode() const override { return false; }
+	virtual FText GetTooltipText() const override;
 	// End UEdGraphNode interface
 
 	// Begin UK2Node interface
 	virtual bool DrawNodeAsExit() const override { return true; }
 	virtual bool ShouldShowNodeProperties() const override { return true; }
 	virtual class FNodeHandlingFunctor* CreateNodeHandler(class FKismetCompilerContext& CompilerContext) const override;
+
+	virtual void GetMenuActions(FBlueprintActionDatabaseRegistrar& ActionRegistrar) const override;
+	virtual bool IsCompatibleWithGraph(UEdGraph const* Graph) const;
+	virtual void PostPlacedNewNode();
+	virtual bool CanDuplicateNode() const { return false; }
 	// End UK2Node interface
 
 	// Begin K2Node_FunctionTerminator interface
 	virtual bool CanCreateUserDefinedPin(const FEdGraphPinType& InPinType, EEdGraphPinDirection InDesiredDirection, FText& OutErrorMessage) override;
 	virtual UEdGraphPin* CreatePinFromUserDefinition(const TSharedPtr<FUserPinInfo> NewPinInfo) override;
 	// End K2Node_FunctionTerminator interface
+
+	BLUEPRINTGRAPH_API TArray<UK2Node_FunctionResult*> GetAllResultNodes() const;
 };
 
diff --git a/Engine/Source/Editor/BlueprintGraph/Private/K2Node_EditablePinBase.cpp b/Engine/Source/Editor/BlueprintGraph/Private/K2Node_EditablePinBase.cpp
--- a/Engine/Source/Editor/BlueprintGraph/Private/K2Node_EditablePinBase.cpp
+++ b/Engine/Source/Editor/BlueprintGraph/Private/K2Node_EditablePinBase.cpp
@@ -19,10 +19,10 @@ void UK2Node_EditablePinBase::AllocateDefaultPins()
 	}
 }
 
-UEdGraphPin* UK2Node_EditablePinBase::CreateUserDefinedPin(const FString& InPinName, const FEdGraphPinType& InPinType, EEdGraphPinDirection InDesiredDirection)
+UEdGraphPin* UK2Node_EditablePinBase::CreateUserDefinedPin(const FString& InPinName, const FEdGraphPinType& InPinType, EEdGraphPinDirection InDesiredDirection, bool bUseUniqueName)
 {
 	// Sanitize the name, if needed
-	const FString NewPinName = CreateUniquePinName(InPinName);
+	const FString NewPinName = bUseUniqueName ? CreateUniquePinName(InPinName) : InPinName;
 
 	// First, add this pin to the user-defined pins
 	TSharedPtr<FUserPinInfo> NewPinInfo = MakeShareable( new FUserPinInfo() );
@@ -61,6 +61,31 @@ void UK2Node_EditablePinBase::RemoveUserDefinedPin(TSharedPtr<FUserPinInfo> PinT
 	UserDefinedPins.Remove(PinToRemove);
 }
 
+void UK2Node_EditablePinBase::RemoveUserDefinedPinByName(const FString& PinName)
+{
+	for (int32 i = 0; i < Pins.Num(); i++)
+	{
+		UEdGraphPin* Pin = Pins[i];
+		if (Pin->PinName == PinName)
+		{
+			Pin->BreakAllPinLinks();
+			Pins.Remove(Pin);
+			Pin->MarkPendingKill();
+
+			if (UBlueprint* Blueprint = GetBlueprint())
+			{
+				FKismetDebugUtilities::RemovePinWatch(Blueprint, Pin);
+			}
+		}
+	}
+
+	// Remove the description from the user-defined pins array
+	UserDefinedPins.RemoveAll([&](const TSharedPtr<FUserPinInfo>& UDPin)
+	{
+		return UDPin.IsValid() && (UDPin->PinName == PinName);
+	});
+}
+
 void UK2Node_EditablePinBase::ExportCustomProperties(FOutputDevice& Out, uint32 Indent)
 {
 	for (int32 PinIndex = 0; PinIndex < UserDefinedPins.Num(); ++PinIndex)
diff --git a/Engine/Source/Editor/BlueprintGraph/Private/K2Node_FunctionResult.cpp b/Engine/Source/Editor/BlueprintGraph/Private/K2Node_FunctionResult.cpp
--- a/Engine/Source/Editor/BlueprintGraph/Private/K2Node_FunctionResult.cpp
+++ b/Engine/Source/Editor/BlueprintGraph/Private/K2Node_FunctionResult.cpp
@@ -64,6 +64,14 @@ public:
 		// Do not register as a default any Pin that comes from being Split
 		if (Net->ParentPin == nullptr)
 		{
+			for (auto& ResultTerm : Context.Results)
+			{
+				if ((ResultTerm.Name == Net->PinName) && (ResultTerm.Type == Net->PinType))
+				{
+					Context.NetMap.Add(Net, &ResultTerm);
+					return;
+				}
+			}
 			FBPTerminal* Term = new (Context.Results) FBPTerminal();
 			Term->CopyFromPin(Net, Net->PinName);
 			Context.NetMap.Add(Net, Term);
@@ -104,7 +112,11 @@ UK2Node_FunctionResult::UK2Node_FunctionResult(const FObjectInitializer& ObjectI
 
 FText UK2Node_FunctionResult::GetNodeTitle(ENodeTitleType::Type TitleType) const
 {
-	return NSLOCTEXT("K2Node", "ReturnNode", "ReturnNode");
+	if (ENodeTitleType::MenuTitle == TitleType)
+	{
+		return NSLOCTEXT("K2Node", "ReturnNodeMenuTitle", "Add Return Node...");
+	}
+	return NSLOCTEXT("K2Node", "ReturnNode", "Return Node");
 }
 
 void UK2Node_FunctionResult::AllocateDefaultPins()
@@ -148,3 +160,90 @@ FNodeHandlingFunctor* UK2Node_FunctionResult::CreateNodeHandler(FKismetCompilerC
 {
 	return new FKCHandler_FunctionResult(CompilerContext);
 }
+
+FText UK2Node_FunctionResult::GetTooltipText() const
+{
+	return NSLOCTEXT("K2Node", "ReturnNodeTooltip", "The node terminates the function's execution. It returns output parameters.");
+}
+
+void UK2Node_FunctionResult::GetMenuActions(FBlueprintActionDatabaseRegistrar& ActionRegistrar) const
+{
+	// actions get registered under specific object-keys; the idea is that 
+	// actions might have to be updated (or deleted) if their object-key is  
+	// mutated (or removed)... here we use the node's class (so if the node 
+	// type disappears, then the action should go with it)
+	UClass* ActionKey = GetClass();
+	// to keep from needlessly instantiating a UBlueprintNodeSpawner, first   
+	// check to make sure that the registrar is looking for actions of this type
+	// (could be regenerating actions for a specific asset, and therefore the 
+	// registrar would only accept actions corresponding to that asset)
+	if (ActionRegistrar.IsOpenForRegistration(ActionKey))
+	{
+		UBlueprintNodeSpawner* NodeSpawner = UBlueprintNodeSpawner::Create(GetClass());
+		check(NodeSpawner != nullptr);
+
+		ActionRegistrar.AddBlueprintAction(ActionKey, NodeSpawner);
+	}
+}
+
+bool UK2Node_FunctionResult::IsCompatibleWithGraph(UEdGraph const* Graph) const
+{
+	const EGraphType GraphType = Graph->GetSchema()->GetGraphType(Graph);
+	const bool bIsCompatible = GraphType == EGraphType::GT_Function;
+	return bIsCompatible && Super::IsCompatibleWithGraph(Graph);
+}
+
+TArray<UK2Node_FunctionResult*> UK2Node_FunctionResult::GetAllResultNodes() const
+{
+	TArray<UK2Node_FunctionResult*> AllResultNodes;
+	if (auto Graph = GetGraph())
+	{
+		Graph->GetNodesOfClass(AllResultNodes);
+	}
+	return AllResultNodes;
+}
+
+void UK2Node_FunctionResult::PostPlacedNewNode()
+{
+	Super::PostPlacedNewNode();
+
+	UK2Node_FunctionResult* PrimaryNode = nullptr;
+	TArray<UK2Node_FunctionResult*> AllResultNodes = GetAllResultNodes();
+	for (auto ResultNode : AllResultNodes)
+	{
+		if (ResultNode && (this != ResultNode))
+		{
+			PrimaryNode = ResultNode;
+			break;
+		}
+	}
+
+	if (PrimaryNode)
+	{
+		{
+			TArray< TSharedPtr<FUserPinInfo> > UDPinsCopy = UserDefinedPins;
+			for (auto UDPin : UDPinsCopy)
+			{
+				RemoveUserDefinedPin(UDPin);
+			}
+			UserDefinedPins.Empty();
+		}
+
+		SignatureClass = PrimaryNode->SignatureClass;
+		SignatureName = PrimaryNode->SignatureName;
+		bIsEditable = PrimaryNode->bIsEditable;
+
+		for (auto UDPin : PrimaryNode->UserDefinedPins)
+		{
+			if (UDPin.IsValid())
+			{
+				TSharedPtr<FUserPinInfo> NewPinInfo = MakeShareable(new FUserPinInfo());
+				NewPinInfo->PinName = UDPin->PinName;
+				NewPinInfo->PinType = UDPin->PinType;
+				NewPinInfo->DesiredPinDirection = UDPin->DesiredPinDirection;
+				UserDefinedPins.Add(NewPinInfo);
+				CreatePinFromUserDefinition(NewPinInfo);
+			}
+		}
+	}
+}
diff --git a/Engine/Source/Editor/Kismet/Private/BlueprintDetailsCustomization.cpp b/Engine/Source/Editor/Kismet/Private/BlueprintDetailsCustomization.cpp
--- a/Engine/Source/Editor/Kismet/Private/BlueprintDetailsCustomization.cpp
+++ b/Engine/Source/Editor/Kismet/Private/BlueprintDetailsCustomization.cpp
@@ -1192,10 +1192,8 @@ void FBlueprintVarActionDetails::PopulateCategories(SMyBlueprint* MyBlueprint, T
 			}
 		}
 
-		TWeakObjectPtr<UK2Node_EditablePinBase> EntryNode;
-		TWeakObjectPtr<UK2Node_EditablePinBase> ResultNode;
-		FBlueprintEditorUtils::GetEntryAndResultNodes(FunctionGraph, EntryNode, ResultNode);
-		if (UK2Node_FunctionEntry* FunctionEntryNode = Cast<UK2Node_FunctionEntry>(EntryNode.Get()))
+		auto EntryNode = FBlueprintEditorUtils::GetEntryNode(FunctionGraph);
+		if (UK2Node_FunctionEntry* FunctionEntryNode = Cast<UK2Node_FunctionEntry>(EntryNode))
 		{
 			for (FBPVariableDescription& Variable : FunctionEntryNode->LocalVariables)
 			{
@@ -1214,10 +1212,8 @@ void FBlueprintVarActionDetails::PopulateCategories(SMyBlueprint* MyBlueprint, T
 
 	for (UEdGraph* MacroGraph : Blueprint->MacroGraphs)
 	{
-		TWeakObjectPtr<UK2Node_EditablePinBase> EntryNode;
-		TWeakObjectPtr<UK2Node_EditablePinBase> ResultNode;
-		FBlueprintEditorUtils::GetEntryAndResultNodes(MacroGraph, EntryNode, ResultNode);
-		if (UK2Node_Tunnel* TypedEntryNode = ExactCast<UK2Node_Tunnel>(EntryNode.Get()))
+		auto EntryNode = FBlueprintEditorUtils::GetEntryNode(MacroGraph);
+		if (UK2Node_Tunnel* TypedEntryNode = ExactCast<UK2Node_Tunnel>(EntryNode))
 		{
 			bool bNewCategory = true;
 			for (int32 j = 0; j < CategorySource.Num() && bNewCategory; ++j)
@@ -2171,22 +2167,42 @@ void FBlueprintGraphArgumentLayout::GenerateChildContent( IDetailChildrenBuilder
 	}
 }
 
+namespace {
+
+	static TArray<UK2Node_EditablePinBase*> GatherAllResultNodes(UK2Node_EditablePinBase* TargetNode)
+	{
+		if (auto ResultNode = Cast<UK2Node_FunctionResult>(TargetNode))
+		{
+			return (TArray<UK2Node_EditablePinBase*>)ResultNode->GetAllResultNodes();
+		}
+		TArray<UK2Node_EditablePinBase*> Result;
+		if (TargetNode)
+		{
+			Result.Add(TargetNode);
+		}
+		return Result;
+	}
+
+} // namespace
+
 void FBlueprintGraphArgumentLayout::OnRemoveClicked()
 {
-	if (ParamItemPtr.IsValid())
+	auto ParamItem = ParamItemPtr.Pin();
+	if (ParamItem.IsValid())
 	{
 		const FScopedTransaction Transaction( LOCTEXT( "RemoveParam", "Remove Parameter" ) );
-		TargetNode->Modify();
-
-		TargetNode->RemoveUserDefinedPin(ParamItemPtr.Pin());
 
-		auto MyBlueprint = GraphActionDetailsPtr.Pin()->GetMyBlueprint();
-		auto Graph = GraphActionDetailsPtr.Pin()->GetGraph();
-		bool bNodeWasCleanedUp = GraphActionDetailsPtr.Pin()->ConditionallyCleanUpResultNode();
-		GraphActionDetailsPtr.Pin()->OnParamsChanged(TargetNode, true);
-		if (bNodeWasCleanedUp && MyBlueprint.IsValid() && Graph)
+		auto GraphActionDetails = GraphActionDetailsPtr.Pin();
+		auto TargetNodes = GatherAllResultNodes(TargetNode);
+		for (auto Node : TargetNodes)
 		{
-			MyBlueprint.Pin()->SelectItemByName(Graph->GetFName());
+			Node->Modify();
+			Node->RemoveUserDefinedPinByName(ParamItem->PinName);
+
+			if (GraphActionDetails.IsValid())
+			{
+				GraphActionDetails->OnParamsChanged(Node, true);
+			}
 		}
 	}
 }
@@ -2198,10 +2214,18 @@ FReply FBlueprintGraphArgumentLayout::OnArgMoveUp()
 	if (ThisParamIndex != INDEX_NONE && NewParamIndex >= 0)
 	{
 		const FScopedTransaction Transaction( LOCTEXT("K2_MovePinUp", "Move Pin Up") );
-		TargetNode->Modify();
+		auto TargetNodes = GatherAllResultNodes(TargetNode);
+		for (auto Node : TargetNodes)
+		{
+			Node->Modify();
+			Node->UserDefinedPins.Swap(ThisParamIndex, NewParamIndex);
 
-		TargetNode->UserDefinedPins.Swap( ThisParamIndex, NewParamIndex );
-		GraphActionDetailsPtr.Pin()->OnParamsChanged(TargetNode, true);
+			auto GraphActionDetails = GraphActionDetailsPtr.Pin();
+			if (GraphActionDetails.IsValid())
+			{
+				GraphActionDetails->OnParamsChanged(Node, true);
+			}
+		}
 	}
 	return FReply::Handled();
 }
@@ -2212,11 +2236,19 @@ FReply FBlueprintGraphArgumentLayout::OnArgMoveDown()
 	const int32 NewParamIndex = ThisParamIndex+1;
 	if (ThisParamIndex != INDEX_NONE && NewParamIndex < TargetNode->UserDefinedPins.Num())
 	{
-		const FScopedTransaction Transaction( LOCTEXT("K2_MovePinUp", "Move Pin Up") );
-		TargetNode->Modify();
-
-		TargetNode->UserDefinedPins.Swap( ThisParamIndex, NewParamIndex );
-		GraphActionDetailsPtr.Pin()->OnParamsChanged(TargetNode, true);
+		const FScopedTransaction Transaction( LOCTEXT("K2_MovePinDown", "Move Pin Down") );
+		auto TargetNodes = GatherAllResultNodes(TargetNode);
+		for (auto Node : TargetNodes)
+		{
+			Node->Modify();
+			Node->UserDefinedPins.Swap(ThisParamIndex, NewParamIndex);
+			
+			auto GraphActionDetails = GraphActionDetailsPtr.Pin();
+			if (GraphActionDetails.IsValid())
+			{
+				GraphActionDetails->OnParamsChanged(Node, true);
+			}
+		}
 	}
 	return FReply::Handled();
 }
@@ -2295,14 +2327,11 @@ void FBlueprintGraphArgumentLayout::OnArgNameTextCommitted(const FText& NewText,
 {
 	if (!NewText.IsEmpty() && TargetNode && ParamItemPtr.IsValid() && GraphActionDetailsPtr.IsValid() && !ShouldPinBeReadOnly())
 	{
-		const FString& OldName = ParamItemPtr.Pin()->PinName;
+		const FString OldName = ParamItemPtr.Pin()->PinName;
 		const FString& NewName = NewText.ToString();
 		if(OldName != NewName)
 		{
-			if(GraphActionDetailsPtr.Pin()->OnPinRenamed(TargetNode, OldName, NewName))
-			{
-				ParamItemPtr.Pin()->PinName = NewName;
-			}
+			GraphActionDetailsPtr.Pin()->OnPinRenamed(TargetNode, OldName, NewName);
 		}
 	}
 }
@@ -2333,13 +2362,30 @@ void FBlueprintGraphArgumentLayout::PinInfoChanged(const FEdGraphPinType& PinTyp
 {
 	if (ParamItemPtr.IsValid() && FBlueprintEditorUtils::IsPinTypeValid(PinType))
 	{
-		ParamItemPtr.Pin()->PinType = PinType;
-		if (GraphActionDetailsPtr.IsValid())
+		const FString PinName = ParamItemPtr.Pin()->PinName;
+		auto GraphActionDetailsPinned = GraphActionDetailsPtr.Pin();
+		if (GraphActionDetailsPinned.IsValid())
 		{
-			GraphActionDetailsPtr.Pin()->GetMyBlueprint().Pin()->GetLastFunctionPinTypeUsed() = PinType;
+			auto MyBPPinned = GraphActionDetailsPinned->GetMyBlueprint().Pin();
+			if (MyBPPinned.IsValid())
+			{
+				MyBPPinned->GetLastFunctionPinTypeUsed() = PinType;
+			}
 			if( !ShouldPinBeReadOnly(true) )
 			{
-				GraphActionDetailsPtr.Pin()->OnParamsChanged(TargetNode);
+				auto TargetNodes = GatherAllResultNodes(TargetNode);
+				for (auto Node : TargetNodes)
+				{
+					auto UDPinPtr = Node->UserDefinedPins.FindByPredicate([&](TSharedPtr<FUserPinInfo>& UDPin)
+					{
+						return UDPin.IsValid() && (UDPin->PinName == PinName);
+					});
+					if (UDPinPtr)
+					{
+						(*UDPinPtr)->PinType = PinType;
+					}
+					GraphActionDetailsPinned->OnParamsChanged(Node);
+				}
 			}
 		}
 	}
@@ -2347,9 +2393,16 @@ void FBlueprintGraphArgumentLayout::PinInfoChanged(const FEdGraphPinType& PinTyp
 
 void FBlueprintGraphArgumentLayout::OnPrePinInfoChange(const FEdGraphPinType& PinType)
 {
-	if( !ShouldPinBeReadOnly(true) && TargetNode )
+	if( !ShouldPinBeReadOnly(true))
 	{
-		TargetNode->Modify();
+		auto TargetNodes = GatherAllResultNodes(TargetNode);
+		for (auto Node : TargetNodes)
+		{
+			if (Node)
+			{
+				Node->Modify();
+			}
+		}
 	}
 }
 
@@ -2879,22 +2932,6 @@ FText FBlueprintGraphActionDetails::GetCurrentReplicatedEventString() const
 	return ReplicationText;
 }
 
-bool FBaseBlueprintGraphActionDetails::ConditionallyCleanUpResultNode()
-{
-	UEdGraph* Graph = GetGraph();
-	UK2Node_EditablePinBase * FunctionResultNode = FunctionResultNodePtr.Get();
-
-	if( Graph && FunctionResultNode && FunctionResultNode->UserDefinedPins.Num() == 0 &&
-		!Cast<UK2Node_Tunnel>(FunctionResultNode))
-	{
-		Graph->RemoveNode(FunctionResultNode);
-		FunctionResultNodePtr = NULL;
-
-		return true;
-	}
-	return false;
-}
-
 bool FBaseBlueprintGraphActionDetails::AttemptToCreateResultNode()
 {
 	if (!FunctionResultNodePtr.IsValid())
@@ -3354,9 +3391,13 @@ bool FBaseBlueprintGraphActionDetails::OnPinRenamed(UK2Node_EditablePinBase* Tar
 			PinRenamedHelper.NodesToRename.Add(FunctionEntryNodePtr.Get());
 		}
 
-		if (FunctionResultNodePtr.IsValid())
+		const FScopedTransaction Transaction(LOCTEXT("RenameParam", "Rename Parameter"));
+
+		auto ResultNodes = GatherAllResultNodes(FunctionResultNodePtr.Get());
+		for (auto ResultNode : ResultNodes)
 		{
-			PinRenamedHelper.NodesToRename.Add(FunctionResultNodePtr.Get());
+			ResultNode->Modify();
+			PinRenamedHelper.NodesToRename.Add(ResultNode);
 		}
 
 		PinRenamedHelper.ModifiedBlueprints.Add(GetBlueprintObj());
@@ -3379,6 +3420,18 @@ bool FBaseBlueprintGraphActionDetails::OnPinRenamed(UK2Node_EditablePinBase* Tar
 		{
 			(*NodeIter)->RenameUserDefinedPin(OldName, NewName, false);
 		}
+
+		for (auto ResultNode : ResultNodes)
+		{
+			auto UDPinPtr = ResultNode->UserDefinedPins.FindByPredicate([&](TSharedPtr<FUserPinInfo>& Pin)
+			{
+				return Pin.IsValid() && (Pin->PinName == OldName);
+			});
+			if (UDPinPtr)
+			{
+				(*UDPinPtr)->PinName = NewName;
+			}
+		}
 	}
 	return true;
 }
@@ -4059,10 +4112,8 @@ FReply FBlueprintGraphActionDetails::OnAddNewOutputClicked()
 	UK2Node_EditablePinBase* FunctionResultNode = FunctionResultNodePtr.Get();
 	if( FunctionResultNode )
 	{
-		FunctionResultNode->Modify();
 		FEdGraphPinType PinType = MyBlueprint.Pin()->GetLastFunctionPinTypeUsed();
 		PinType.bIsReference = false;
-
 		// Make sure that if this is an exec node we are allowed one.
 		const UEdGraphSchema_K2* Schema = GetDefault<UEdGraphSchema_K2>();
 		if ((PinType.PinCategory == Schema->PC_Exec) && (!FunctionResultNode->CanModifyExecutionWires()))
@@ -4070,20 +4121,34 @@ FReply FBlueprintGraphActionDetails::OnAddNewOutputClicked()
 			MyBlueprint.Pin()->ResetLastPinType();
 			PinType = MyBlueprint.Pin()->GetLastFunctionPinTypeUsed();
 		}
-		FString NewPinName = TEXT("NewParam");
-		if (FunctionResultNode->CreateUserDefinedPin(NewPinName, PinType, EGPD_Input))
+
+		const FString NewPinName = FunctionResultNode->CreateUniquePinName(TEXT("NewParam"));
+		auto TargetNodes = GatherAllResultNodes(FunctionResultNode);
+		bool bAllChanged = TargetNodes.Num() > 0;
+		for (auto Node : TargetNodes)
 		{
-			OnParamsChanged(FunctionResultNode, true);
+			Node->Modify();
+			auto NewPin = Node->CreateUserDefinedPin(NewPinName, PinType, EGPD_Input, false);
+			bAllChanged &= nullptr != NewPin;
 
-			if ( !PreviousResultNode )
+			if (bAllChanged)
 			{
-				DetailsLayoutPtr->ForceRefreshDetails();
+				OnParamsChanged(Node, true);
+			}
+			else
+			{
+				break;
 			}
 		}
-		else
+		if (!bAllChanged)
 		{
 			Transaction.Cancel();
 		}
+
+		if (!PreviousResultNode)
+		{
+			DetailsLayoutPtr->ForceRefreshDetails();
+		}
 	}
 	else
 	{
diff --git a/Engine/Source/Editor/Kismet/Private/BlueprintDetailsCustomization.h b/Engine/Source/Editor/Kismet/Private/BlueprintDetailsCustomization.h
--- a/Engine/Source/Editor/Kismet/Private/BlueprintDetailsCustomization.h
+++ b/Engine/Source/Editor/Kismet/Private/BlueprintDetailsCustomization.h
@@ -252,9 +252,6 @@ public:
 
 	/** Refreshes the graph and ensures the target node is up to date */
 	bool OnPinRenamed(UK2Node_EditablePinBase* TargetNode, const FString& OldName, const FString& NewName);
-
-	/** Called to potentially remove the result node (if there are no output args), returns true if it was cleaned up */
-	bool ConditionallyCleanUpResultNode();
 	
 	/** Gets the blueprint we're editing */
 	TWeakPtr<SMyBlueprint> GetMyBlueprint() const {return MyBlueprint.Pin();}
diff --git a/Engine/Source/Editor/Kismet/Private/SMyBlueprint.cpp b/Engine/Source/Editor/Kismet/Private/SMyBlueprint.cpp
--- a/Engine/Source/Editor/Kismet/Private/SMyBlueprint.cpp
+++ b/Engine/Source/Editor/Kismet/Private/SMyBlueprint.cpp
@@ -465,15 +465,13 @@ void SMyBlueprint::OnCategoryNameCommitted(const FText& InNewText, ETextCommit::
 					// Don't allow changing the category of a graph who's parent is not the current Blueprint
 					if(GraphAction && !FBlueprintEditorUtils::IsPaletteActionReadOnly(Actions[i], BlueprintEditorPtr.Pin()) && FBlueprintEditorUtils::FindBlueprintForGraph(GraphAction->EdGraph) == GetBlueprintObj())
 					{
-						TWeakObjectPtr<UK2Node_EditablePinBase> EntryNode;
-						TWeakObjectPtr<UK2Node_EditablePinBase> ResultNode;
-						FBlueprintEditorUtils::GetEntryAndResultNodes(GraphAction->EdGraph, EntryNode, ResultNode);
+						auto EntryNode = FBlueprintEditorUtils::GetEntryNode(GraphAction->EdGraph);
 						EntryNode->Modify();
-						if (UK2Node_FunctionEntry* FunctionEntryNode = Cast<UK2Node_FunctionEntry>(EntryNode.Get()))
+						if (UK2Node_FunctionEntry* FunctionEntryNode = Cast<UK2Node_FunctionEntry>(EntryNode))
 						{
 							FunctionEntryNode->MetaData.Category = CategoryName;
 						}
-						else if (UK2Node_Tunnel* TypedEntryNode = ExactCast<UK2Node_Tunnel>(EntryNode.Get()))
+						else if (UK2Node_Tunnel* TypedEntryNode = ExactCast<UK2Node_Tunnel>(EntryNode))
 						{
 							TypedEntryNode->MetaData.Category = CategoryName;
 						}
@@ -1781,14 +1779,12 @@ FText SMyBlueprint::GetGraphCategory(UEdGraph* InGraph) const
 	FText ReturnCategory;
 
 	// Pull the category from the required metadata based on the types of nodes we can discover in the graph
-	TWeakObjectPtr<UK2Node_EditablePinBase> EntryNode;
-	TWeakObjectPtr<UK2Node_EditablePinBase> ResultNode;
-	FBlueprintEditorUtils::GetEntryAndResultNodes(InGraph, EntryNode, ResultNode);
-	if (UK2Node_FunctionEntry* FunctionEntryNode = Cast<UK2Node_FunctionEntry>(EntryNode.Get()))
+	auto EntryNode = FBlueprintEditorUtils::GetEntryNode(InGraph);
+	if (UK2Node_FunctionEntry* FunctionEntryNode = Cast<UK2Node_FunctionEntry>(EntryNode))
 	{
 		ReturnCategory = FunctionEntryNode->MetaData.Category;
 	}
-	else if (UK2Node_Tunnel* TypedEntryNode = ExactCast<UK2Node_Tunnel>(EntryNode.Get()))
+	else if (UK2Node_Tunnel* TypedEntryNode = ExactCast<UK2Node_Tunnel>(EntryNode))
 	{
 		ReturnCategory = TypedEntryNode->MetaData.Category;
 	}
diff --git a/Engine/Source/Editor/UnrealEd/Private/Kismet2/BlueprintEditorUtils.cpp b/Engine/Source/Editor/UnrealEd/Private/Kismet2/BlueprintEditorUtils.cpp
--- a/Engine/Source/Editor/UnrealEd/Private/Kismet2/BlueprintEditorUtils.cpp
+++ b/Engine/Source/Editor/UnrealEd/Private/Kismet2/BlueprintEditorUtils.cpp
@@ -2115,11 +2115,8 @@ void FBlueprintEditorUtils::MarkBlueprintAsStructurallyModified(UBlueprint* Blue
 		Blueprint->GetAllGraphs(AllGraphs);
 		for (int32 i = 0; i < AllGraphs.Num(); i++)
 		{
-			TWeakObjectPtr<UK2Node_EditablePinBase> EntryNode;
-			TWeakObjectPtr<UK2Node_EditablePinBase> ResultNode;
-			GetEntryAndResultNodes(AllGraphs[i], EntryNode, ResultNode);
-
-			if(UK2Node_Tunnel* TunnelNode = ExactCast<UK2Node_Tunnel>(EntryNode.Get()))
+			auto EntryNode = GetEntryNode(AllGraphs[i]);
+			if(UK2Node_Tunnel* TunnelNode = ExactCast<UK2Node_Tunnel>(EntryNode))
 			{
 				// Remove data marking graphs as latent, this will be re-cache'd as needed
 				TunnelNode->MetaData.HasLatentFunctions = INDEX_NONE;
@@ -2169,11 +2166,8 @@ void FBlueprintEditorUtils::MarkBlueprintAsModified(UBlueprint* Blueprint)
 		Blueprint->GetAllGraphs(AllGraphs);
 		for (int32 i = 0; i < AllGraphs.Num(); i++)
 		{
-			TWeakObjectPtr<UK2Node_EditablePinBase> EntryNode;
-			TWeakObjectPtr<UK2Node_EditablePinBase> ResultNode;
-			GetEntryAndResultNodes(AllGraphs[i], EntryNode, ResultNode);
-
-			if(UK2Node_Tunnel* TunnelNode = ExactCast<UK2Node_Tunnel>(EntryNode.Get()))
+			auto EntryNode = GetEntryNode(AllGraphs[i]);
+			if(UK2Node_Tunnel* TunnelNode = ExactCast<UK2Node_Tunnel>(EntryNode))
 			{
 				// Remove data marking graphs as latent, this will be re-cache'd as needed
 				TunnelNode->MetaData.HasLatentFunctions = INDEX_NONE;
@@ -7754,6 +7748,47 @@ bool FBlueprintEditorUtils::GetFunctionGuidFromClassByFieldName(const UClass* In
 	return false;
 }
 
+UK2Node_EditablePinBase* FBlueprintEditorUtils::GetEntryNode(const UEdGraph* InGraph)
+{
+	UK2Node_EditablePinBase* Result = nullptr;
+	if (InGraph)
+	{
+		TArray<UK2Node_FunctionEntry*> EntryNodes;
+		InGraph->GetNodesOfClass(EntryNodes);
+		if (EntryNodes.Num() > 0)
+		{
+			if (EntryNodes[0]->IsEditable())
+			{
+				Result = EntryNodes[0];
+			}
+		}
+		else
+		{
+			TArray<UK2Node_Tunnel*> TunnelNodes;
+			InGraph->GetNodesOfClass(TunnelNodes);
+
+			if (TunnelNodes.Num() > 0)
+			{
+				// Iterate over the tunnel nodes, and try to find an entry and exit
+				for (int32 i = 0; i < TunnelNodes.Num(); i++)
+				{
+					UK2Node_Tunnel* Node = TunnelNodes[i];
+					// Composite nodes should never be considered for function entry / exit, since we're searching for a graph's terminals
+					if (Node->IsEditable() && !Node->IsA(UK2Node_Composite::StaticClass()))
+					{
+						if (Node->bCanHaveOutputs)
+						{
+							Result = Node;
+							break;
+						}
+					}
+				}
+			}
+		}
+	}
+	return Result;
+}
+
 void FBlueprintEditorUtils::GetEntryAndResultNodes(const UEdGraph* InGraph, TWeakObjectPtr<class UK2Node_EditablePinBase>& OutEntryNode, TWeakObjectPtr<class UK2Node_EditablePinBase>& OutResultNode)
 {
 	if (InGraph)
@@ -7776,7 +7811,6 @@ void FBlueprintEditorUtils::GetEntryAndResultNodes(const UEdGraph* InGraph, TWea
 				TArray<UK2Node_FunctionResult*> ResultNodes;
 				InGraph->GetNodesOfClass(ResultNodes);
 
-				check(ResultNodes.Num() <= 1);
 				UK2Node_FunctionResult* ResultNode = ResultNodes.Num() ? ResultNodes[0] : NULL;
 				// Note:  we assume that if the entry is editable, the result is too (since the entry node is guaranteed to be there on graph creation, but the result isn't)
 				if( ResultNode )
@@ -7836,11 +7870,7 @@ FKismetUserDeclaredFunctionMetadata* FBlueprintEditorUtils::GetGraphFunctionMeta
 
 FText FBlueprintEditorUtils::GetGraphDescription(const UEdGraph* InGraph)
 {
-	TWeakObjectPtr<class UK2Node_EditablePinBase> EntryNode;
-	TWeakObjectPtr<class UK2Node_EditablePinBase> ResultNode;
-	GetEntryAndResultNodes(InGraph, EntryNode, ResultNode);
-
-	UK2Node_EditablePinBase * FunctionEntryNode = EntryNode.Get();
+	auto FunctionEntryNode = GetEntryNode(InGraph);
 	if (UK2Node_FunctionEntry* TypedEntryNode = Cast<UK2Node_FunctionEntry>(FunctionEntryNode))
 	{
 		return FText::FromString(TypedEntryNode->MetaData.ToolTip);
@@ -7860,11 +7890,9 @@ bool FBlueprintEditorUtils::CheckIfGraphHasLatentFunctions(UEdGraph* InGraph)
 	{
 		static bool CheckIfGraphHasLatentFunctions(UEdGraph* InGraphToCheck, TArray<UEdGraph*>& InspectedGraphList)
 		{
-			TWeakObjectPtr<UK2Node_EditablePinBase> EntryNode;
-			TWeakObjectPtr<UK2Node_EditablePinBase> ResultNode;
-			GetEntryAndResultNodes(InGraphToCheck, EntryNode, ResultNode);
+			auto EntryNode = GetEntryNode(InGraphToCheck);
 
-			UK2Node_Tunnel* TunnelNode = ExactCast<UK2Node_Tunnel>(EntryNode.Get());
+			UK2Node_Tunnel* TunnelNode = ExactCast<UK2Node_Tunnel>(EntryNode);
 			if(!TunnelNode)
 			{
 				// No tunnel, no metadata.
diff --git a/Engine/Source/Editor/UnrealEd/Public/Kismet2/BlueprintEditorUtils.h b/Engine/Source/Editor/UnrealEd/Public/Kismet2/BlueprintEditorUtils.h
--- a/Engine/Source/Editor/UnrealEd/Public/Kismet2/BlueprintEditorUtils.h
+++ b/Engine/Source/Editor/UnrealEd/Public/Kismet2/BlueprintEditorUtils.h
@@ -1182,6 +1182,15 @@ public:
 	 */
 	static void GetEntryAndResultNodes(const UEdGraph* InGraph, TWeakObjectPtr<class UK2Node_EditablePinBase>& OutEntryNode, TWeakObjectPtr<class UK2Node_EditablePinBase>& OutResultNode);
 
+
+	/**
+	 * Finds the entry node for a function or macro graph
+	 *
+	 * @param InGraph			The graph to search through
+	 * @return		The found entry node for the graph
+	 */
+	static class UK2Node_EditablePinBase* GetEntryNode(const UEdGraph* InGraph);
+
 	/**
 	 * Returns the function meta data block for the graph entry node.
 	 *
