Fixed incorrect behavior with checking out all un-cached Blueprints when specified by the user.

diff --git a/Engine/Source/Editor/Kismet/Private/FindInBlueprintManager.cpp b/Engine/Source/Editor/Kismet/Private/FindInBlueprintManager.cpp
--- a/Engine/Source/Editor/Kismet/Private/FindInBlueprintManager.cpp
+++ b/Engine/Source/Editor/Kismet/Private/FindInBlueprintManager.cpp
@@ -1447,7 +1447,7 @@ void FFindInBlueprintSearchManager::OnCacheAllUncachedBlueprints(bool bInSourceC
 	// Multiple threads can be adding to this at the same time
 	FScopeLock ScopeLock(&SafeModifyCacheCriticalSection);
 
-	if(ISourceControlModule::Get().IsEnabled() && bCheckoutAndSave)
+	if(bInSourceControlActive && bCheckoutAndSave)
 	{
 		FEditorFileUtils::CheckoutPackages(UncachedBlueprints);
 	}
@@ -1477,14 +1477,15 @@ void FFindInBlueprintSearchManager::CacheAllUncachedBlueprints(TWeakPtr< SFindIn
 			CachingObject = new FCacheAllBlueprintsTickableObject(UncachedBlueprints, bCheckoutAndSave);
 			OutActiveTimerDelegate.BindRaw(CachingObject, &FCacheAllBlueprintsTickableObject::Tick);
 
-			if(!ISourceControlModule::Get().IsEnabled() && bCheckoutAndSave)
+			const bool bIsSourceControlEnabled = ISourceControlModule::Get().IsEnabled();
+			if(!bIsSourceControlEnabled && bCheckoutAndSave)
 			{
 				// Offer to start up Source Control
 				ISourceControlModule::Get().ShowLoginDialog(FSourceControlLoginClosed::CreateRaw(this, &FFindInBlueprintSearchManager::OnCacheAllUncachedBlueprints, bCheckoutAndSave), ELoginWindowMode::Modeless, EOnLoginWindowStartup::PreserveProvider);
 			}
 			else
 			{
-				OnCacheAllUncachedBlueprints(true, bCheckoutAndSave);
+				OnCacheAllUncachedBlueprints(bIsSourceControlEnabled, bCheckoutAndSave);
 			}
 
 			SourceCachingWidget = InSourceWidget;
