Crash fix when compiling a Blueprint while Find-in-Blueprints is searching all Blueprints.

diff --git a/Engine/Source/Editor/Kismet/Private/FindInBlueprintManager.cpp b/Engine/Source/Editor/Kismet/Private/FindInBlueprintManager.cpp
--- a/Engine/Source/Editor/Kismet/Private/FindInBlueprintManager.cpp
+++ b/Engine/Source/Editor/Kismet/Private/FindInBlueprintManager.cpp
@@ -1128,8 +1128,6 @@ FString FFindInBlueprintSearchManager::GatherBlueprintSearchMetadata(const UBlue
 
 void FFindInBlueprintSearchManager::AddOrUpdateBlueprintSearchMetadata(UBlueprint* InBlueprint, bool bInForceReCache/* = false*/)
 {
-	double StartSeconds = FPlatformTime::Seconds();
-
 	check(InBlueprint);
 
 	// Allow only one thread modify the search data at a time
@@ -1329,8 +1327,7 @@ void FFindInBlueprintSearchManager::CleanCache()
 			FSearchData searchData;
 	 		ContinueSearchQuery(ActiveSearch, searchData);
 
-			// We will be looking the path up in the map, which uses the package path instead of the Blueprint, so rebuild the package path from the Blueprint path
-			FString CachePath = FPaths::GetPath(searchData.BlueprintPath) / FPaths::GetBaseFilename(searchData.BlueprintPath);
+			FString CachePath = searchData.BlueprintPath;
 	 		CacheQueries.Add(ActiveSearch, CachePath);
 	 	}
 	}
