Making it so external modules can add their own Blueprint toolbar buttons through extenders (and not just use them to add options to the menu bar).

diff --git a/Engine/Source/Editor/Kismet/Private/BlueprintEditor.cpp b/Engine/Source/Editor/Kismet/Private/BlueprintEditor.cpp
--- a/Engine/Source/Editor/Kismet/Private/BlueprintEditor.cpp
+++ b/Engine/Source/Editor/Kismet/Private/BlueprintEditor.cpp
@@ -1586,8 +1586,11 @@ void FBlueprintEditor::InitBlueprintEditor(
 	FKismet2Menu::SetupBlueprintEditorMenu( MenuExtender, *this );
 	AddMenuExtender(MenuExtender);
 
-	FBlueprintEditorModule* BlueprintEditorModule = &FModuleManager::LoadModuleChecked<FBlueprintEditorModule>( "Kismet" );
-	AddMenuExtender(BlueprintEditorModule->GetMenuExtensibilityManager()->GetAllExtenders(GetToolkitCommands(), GetEditingObjects()));
+	FBlueprintEditorModule* BlueprintEditorModule = &FModuleManager::LoadModuleChecked<FBlueprintEditorModule>("Kismet");
+	TSharedPtr<FExtender> CustomExtenders = BlueprintEditorModule->GetMenuExtensibilityManager()->GetAllExtenders(GetToolkitCommands(), GetEditingObjects());
+
+	AddMenuExtender(CustomExtenders);
+	AddToolbarExtender(CustomExtenders);
 	
 	RegenerateMenusAndToolbars();
 
