[UE-17366] Stability improvemnets to the SMyBlueprint widget that's instatiated by the diff/merge tools. Disabling various editing functionality

diff --git a/Engine/Source/Editor/GraphEditor/Private/SGraphActionMenu.cpp b/Engine/Source/Editor/GraphEditor/Private/SGraphActionMenu.cpp
--- a/Engine/Source/Editor/GraphEditor/Private/SGraphActionMenu.cpp
+++ b/Engine/Source/Editor/GraphEditor/Private/SGraphActionMenu.cpp
@@ -164,7 +164,7 @@ class SGraphActionCategoryWidget : public SCompoundWidget
 	SLATE_END_ARGS()
 
 	TWeakPtr<FGraphActionNode> ActionNode;
-
+	TAttribute<bool> IsReadOnly;
 public:
 	TWeakPtr<SInlineEditableTextBlock> InlineWidget;
 
@@ -177,6 +177,7 @@ public:
 		FEditorCategoryUtils::GetCategoryTooltipInfo(*InActionNode->GetDisplayName().ToString(), CategoryTooltip, CategoryLink, CategoryExcerpt);
 
 		TSharedRef<SToolTip> ToolTipWidget = IDocumentation::Get()->CreateToolTip(CategoryTooltip, NULL, CategoryLink, CategoryExcerpt);
+		IsReadOnly = InArgs._IsReadOnly;
 
 		this->ChildSlot
 		[
@@ -211,6 +212,11 @@ public:
 
 	virtual void OnDragEnter( const FGeometry& MyGeometry, const FDragDropEvent& DragDropEvent ) override
 	{
+		if (IsReadOnly.Get())
+		{
+			return;
+		}
+
 		TSharedPtr<FGraphEditorDragDropAction> GraphDropOp = DragDropEvent.GetOperationAs<FGraphEditorDragDropAction>();
 		if (GraphDropOp.IsValid())
 		{
diff --git a/Engine/Source/Editor/Kismet/Private/SBlueprintPalette.cpp b/Engine/Source/Editor/Kismet/Private/SBlueprintPalette.cpp
--- a/Engine/Source/Editor/Kismet/Private/SBlueprintPalette.cpp
+++ b/Engine/Source/Editor/Kismet/Private/SBlueprintPalette.cpp
@@ -966,6 +966,8 @@ void SBlueprintPaletteItem::Construct(const FArguments& InArgs, FCreateWidgetFor
 	ActionPtr = InCreateData->Action;
 	BlueprintEditorPtr = InBlueprintEditor;
 
+	const bool bIsFullyReadOnly = !InBlueprintEditor.IsValid();
+	
 	// construct the icon widget
 	FSlateBrush const* IconBrush   = FEditorStyle::GetBrush(TEXT("NoBrush"));
 	FSlateColor        IconColor   = FSlateColor::UseForeground();
@@ -973,6 +975,7 @@ void SBlueprintPaletteItem::Construct(const FArguments& InArgs, FCreateWidgetFor
 	FString			   IconDocLink, IconDocExcerpt;
 	GetPaletteItemIcon(GraphAction, Blueprint, IconBrush, IconColor, IconToolTip, IconDocLink, IconDocExcerpt);
 	TSharedRef<SWidget> IconWidget = CreateIconWidget(FText::FromString(IconToolTip), IconBrush, IconColor, IconDocLink, IconDocExcerpt);
+	IconWidget->SetEnabled(!bIsFullyReadOnly);
 
 	// Setup a meta tag for this node
 	FTutorialMetaData TagMeta("PaletteItem"); 
@@ -983,7 +986,7 @@ void SBlueprintPaletteItem::Construct(const FArguments& InArgs, FCreateWidgetFor
 	}
 	// construct the text widget
 	FSlateFontInfo NameFont = FSlateFontInfo(FPaths::EngineContentDir() / TEXT("Slate/Fonts/Roboto-Regular.ttf"), 10);
-	bool bIsReadOnly = InBlueprintEditor.IsValid() ? FBlueprintEditorUtils::IsPaletteActionReadOnly(GraphAction, InBlueprintEditor.Pin()) : true;
+	const bool bIsReadOnly = bIsFullyReadOnly || FBlueprintEditorUtils::IsPaletteActionReadOnly(GraphAction, InBlueprintEditor.Pin());
 	TSharedRef<SWidget> NameSlotWidget = CreateTextSlotWidget( NameFont, InCreateData, bIsReadOnly );
 	
 	// For Variables and Local Variables, we will convert the icon widget into a pin type selector.
@@ -1004,7 +1007,8 @@ void SBlueprintPaletteItem::Construct(const FArguments& InArgs, FCreateWidgetFor
 		if (FBlueprintEditorUtils::IsVariableCreatedByBlueprint(Blueprint, VariableProp) || Cast<UFunction>(VariableProp->GetOuter()))
 		{
 			const UEdGraphSchema_K2* Schema = GetDefault<UEdGraphSchema_K2>();
-			IconWidget = SNew(SPinTypeSelectorHelper, VariableProp, Blueprint, BlueprintEditorPtr);
+			IconWidget = SNew(SPinTypeSelectorHelper, VariableProp, Blueprint, BlueprintEditorPtr)
+				.IsEnabled(!bIsFullyReadOnly);
 		}
 	}
 
@@ -1034,11 +1038,21 @@ void SBlueprintPaletteItem::Construct(const FArguments& InArgs, FCreateWidgetFor
 			.VAlign(VAlign_Center)
 		[
 			SNew(SPaletteItemVisibilityToggle, ActionPtr, InBlueprintEditor, InBlueprint)
+			.IsEnabled(!bIsFullyReadOnly)
 		]
 	];
 }
+
 END_SLATE_FUNCTION_BUILD_OPTIMIZATION
 
+void SBlueprintPaletteItem::OnDragEnter(const FGeometry& MyGeometry, const FDragDropEvent& DragDropEvent)
+{
+	if (BlueprintEditorPtr.IsValid())
+	{
+		SGraphPaletteItem::OnDragEnter(MyGeometry, DragDropEvent);
+	}
+}
+
 /*******************************************************************************
 * SBlueprintPaletteItem Private Methods
 *******************************************************************************/
diff --git a/Engine/Source/Editor/Kismet/Private/SBlueprintPalette.h b/Engine/Source/Editor/Kismet/Private/SBlueprintPalette.h
--- a/Engine/Source/Editor/Kismet/Private/SBlueprintPalette.h
+++ b/Engine/Source/Editor/Kismet/Private/SBlueprintPalette.h
@@ -36,6 +36,10 @@ public:
 private:
 	void Construct(const FArguments& InArgs, FCreateWidgetForActionData* const InCreateData, UBlueprint* InBlueprint, TWeakPtr<FBlueprintEditor> InBlueprintEditor);
 
+	// SWidget Interface
+	virtual void OnDragEnter(const FGeometry& MyGeometry, const FDragDropEvent& DragDropEvent) override;
+	// End of SWidget Interface
+
 	// SGraphPaletteItem Interface
 	virtual TSharedRef<SWidget> CreateTextSlotWidget( const FSlateFontInfo& NameFont,  FCreateWidgetForActionData* const InCreateData, bool bIsReadOnly ) override;
 	virtual FText GetDisplayText() const override;
diff --git a/Engine/Source/Editor/Kismet/Private/SMyBlueprint.cpp b/Engine/Source/Editor/Kismet/Private/SMyBlueprint.cpp
--- a/Engine/Source/Editor/Kismet/Private/SMyBlueprint.cpp
+++ b/Engine/Source/Editor/Kismet/Private/SMyBlueprint.cpp
@@ -1416,6 +1416,11 @@ bool SMyBlueprint::IsShowingEmptySections() const
 
 FReply SMyBlueprint::OnActionDragged( const TArray< TSharedPtr<FEdGraphSchemaAction> >& InActions, const FPointerEvent& MouseEvent )
 {
+	if (!BlueprintEditorPtr.IsValid())
+	{
+		return FReply::Unhandled();
+	}
+
 	TSharedPtr<FEdGraphSchemaAction> InAction( InActions.Num() > 0 ? InActions[0] : NULL );
 	if(InAction.IsValid())
 	{
@@ -1453,13 +1458,13 @@ FReply SMyBlueprint::OnActionDragged( const TArray< TSharedPtr<FEdGraphSchemaAct
 					}
 				}
 
-				return FReply::Handled().BeginDragDrop(FKismetFunctionDragDropAction::New(InAction, FuncAction->FuncName, BlueprintEditorPtr.Pin()->GetBlueprintObj()->SkeletonGeneratedClass, FMemberReference(), AnalyticsDelegate, FKismetDragDropAction::FCanBeDroppedDelegate::CreateLambda(CanDragDropAction, bIsBlueprintCallableFunction)));
+				return FReply::Handled().BeginDragDrop(FKismetFunctionDragDropAction::New(InAction, FuncAction->FuncName, GetBlueprintObj()->SkeletonGeneratedClass, FMemberReference(), AnalyticsDelegate, FKismetDragDropAction::FCanBeDroppedDelegate::CreateLambda(CanDragDropAction, bIsBlueprintCallableFunction)));
 			}
 			else if (FuncAction->GraphType == EEdGraphSchemaAction_K2Graph::Macro)
 			{
-				if ((FuncAction->EdGraph != NULL) && BlueprintEditorPtr.Pin()->GetBlueprintObj()->BlueprintType != BPTYPE_MacroLibrary)
+				if ((FuncAction->EdGraph != NULL) && GetBlueprintObj()->BlueprintType != BPTYPE_MacroLibrary)
 				{
-					return FReply::Handled().BeginDragDrop(FKismetMacroDragDropAction::New(FuncAction->FuncName, BlueprintEditorPtr.Pin()->GetBlueprintObj(), FuncAction->EdGraph, AnalyticsDelegate));
+					return FReply::Handled().BeginDragDrop(FKismetMacroDragDropAction::New(FuncAction->FuncName, GetBlueprintObj(), FuncAction->EdGraph, AnalyticsDelegate));
 				}
 			}
 		}
