Renaming input pins on Blueprint functions will no longer revert the name back to the default.

diff --git a/Engine/Source/Editor/Kismet/Private/BlueprintDetailsCustomization.cpp b/Engine/Source/Editor/Kismet/Private/BlueprintDetailsCustomization.cpp
--- a/Engine/Source/Editor/Kismet/Private/BlueprintDetailsCustomization.cpp
+++ b/Engine/Source/Editor/Kismet/Private/BlueprintDetailsCustomization.cpp
@@ -2331,7 +2331,10 @@ void FBlueprintGraphArgumentLayout::OnArgNameTextCommitted(const FText& NewText,
 		const FString& NewName = NewText.ToString();
 		if(OldName != NewName)
 		{
-			GraphActionDetailsPtr.Pin()->OnPinRenamed(TargetNode, OldName, NewName);
+			if(GraphActionDetailsPtr.Pin()->OnPinRenamed(TargetNode, OldName, NewName))
+			{
+				ParamItemPtr.Pin()->PinName = NewName;
+			}
 		}
 	}
 }
