UE-17522 Unable to rename Function and Macro input parameters (Details panel)

diff --git a/Engine/Source/Editor/Kismet/Private/BlueprintDetailsCustomization.cpp b/Engine/Source/Editor/Kismet/Private/BlueprintDetailsCustomization.cpp
--- a/Engine/Source/Editor/Kismet/Private/BlueprintDetailsCustomization.cpp
+++ b/Engine/Source/Editor/Kismet/Private/BlueprintDetailsCustomization.cpp
@@ -2331,10 +2331,7 @@ void FBlueprintGraphArgumentLayout::OnArgNameTextCommitted(const FText& NewText,
 		const FString& NewName = NewText.ToString();
 		if(OldName != NewName)
 		{
-			if(GraphActionDetailsPtr.Pin()->OnPinRenamed(TargetNode, OldName, NewName))
-			{
-				ParamItemPtr.Pin()->PinName = NewName;
-			}
+			GraphActionDetailsPtr.Pin()->OnPinRenamed(TargetNode, OldName, NewName);
 		}
 	}
 }
@@ -3389,18 +3386,17 @@ bool FBaseBlueprintGraphActionDetails::OnPinRenamed(UK2Node_EditablePinBase* Tar
 	{
 		FPinRenamedHelper PinRenamedHelper;
 
-		if (FunctionEntryNodePtr.IsValid())
-		{
-			PinRenamedHelper.NodesToRename.Add(FunctionEntryNodePtr.Get());
-		}
-
 		const FScopedTransaction Transaction(LOCTEXT("RenameParam", "Rename Parameter"));
 
-		auto ResultNodes = GatherAllResultNodes(FunctionResultNodePtr.Get());
-		for (auto ResultNode : ResultNodes)
+		auto TerminalNodes = GatherAllResultNodes(FunctionResultNodePtr.Get());
+		if (auto EntryNode = FunctionEntryNodePtr.Get())
+		{
+			TerminalNodes.Add(EntryNode);
+		}
+		for (auto TerminalNode : TerminalNodes)
 		{
-			ResultNode->Modify();
-			PinRenamedHelper.NodesToRename.Add(ResultNode);
+			TerminalNode->Modify();
+			PinRenamedHelper.NodesToRename.Add(TerminalNode);
 		}
 
 		PinRenamedHelper.ModifiedBlueprints.Add(GetBlueprintObj());
@@ -3424,9 +3420,9 @@ bool FBaseBlueprintGraphActionDetails::OnPinRenamed(UK2Node_EditablePinBase* Tar
 			(*NodeIter)->RenameUserDefinedPin(OldName, NewName, false);
 		}
 
-		for (auto ResultNode : ResultNodes)
+		for (auto TerminalNode : TerminalNodes)
 		{
-			auto UDPinPtr = ResultNode->UserDefinedPins.FindByPredicate([&](TSharedPtr<FUserPinInfo>& Pin)
+			auto UDPinPtr = TerminalNode->UserDefinedPins.FindByPredicate([&](TSharedPtr<FUserPinInfo>& Pin)
 			{
 				return Pin.IsValid() && (Pin->PinName == OldName);
 			});
