Blueprints - Timeline UI updates
+ Make individual tracks collapsable.
+ Add an option to turn off the view synchronization for an individual track.
+ General padding and layout tweaks.

Jira: UE-16872
Jira: UE-16873

diff --git a/Engine/Source/Editor/Kismet/Private/STimelineEditor.cpp b/Engine/Source/Editor/Kismet/Private/STimelineEditor.cpp
--- a/Engine/Source/Editor/Kismet/Private/STimelineEditor.cpp
+++ b/Engine/Source/Editor/Kismet/Private/STimelineEditor.cpp
@@ -111,51 +111,80 @@ void STimelineEdTrack::Construct(const FArguments& InArgs, TSharedPtr<FTimelineE
 
 	this->ChildSlot
 	[
-		SNew(SHorizontalBox)
-		+SHorizontalBox::Slot()
-		.AutoWidth()
-		.Padding(2)
+		SNew(SVerticalBox)
+		
+		// Heading Slot
+		+ SVerticalBox::Slot()
+		.AutoHeight()
 		[
-			SNew(SVerticalBox)
-			+SVerticalBox::Slot()
-			.AutoHeight()
-			.Padding(0,0,0,0)
+			SNew(SBorder)
+			.BorderImage(FEditorStyle::GetBrush("DetailsView.CategoryTop_Hovered"))
+			.ForegroundColor(FLinearColor::White)
 			[
-				SNew(SBorder)
-				.VAlign(VAlign_Fill)
+				SNew(SHorizontalBox)
+
+				// Expander Button
+				+ SHorizontalBox::Slot()
+				.AutoWidth()
 				[
-					SNew(SBox)
-					.WidthOverride(128)
-					.HeightOverride(32)
-					[
-						// Name of track
-						SAssignNew(InlineTextBlock, SInlineEditableTextBlock)
-							.Text(FText::FromName(TrackBase->TrackName))
-							.ToolTipText(LOCTEXT("TrackNameTooltip", "Enter track name"))
-							.OnVerifyTextChanged(TimelineRef, &STimelineEditor::OnVerifyTrackNameCommit,TrackBase, this)
-							.OnTextCommitted(TimelineRef, &STimelineEditor::OnTrackNameCommitted,TrackBase, this)
-					]
+					SNew(SCheckBox)
+					.IsChecked(this, &STimelineEdTrack::GetIsExpandedState)
+					.OnCheckStateChanged(this, &STimelineEdTrack::OnIsExpandedStateChanged)
+					.CheckedImage(FEditorStyle::GetBrush("TreeArrow_Expanded"))
+					.CheckedHoveredImage(FEditorStyle::GetBrush("TreeArrow_Expanded_Hovered"))
+					.CheckedPressedImage(FEditorStyle::GetBrush("TreeArrow_Expanded"))
+					.UncheckedImage(FEditorStyle::GetBrush("TreeArrow_Collapsed"))
+					.UncheckedHoveredImage(FEditorStyle::GetBrush("TreeArrow_Collapsed_Hovered"))
+					.UncheckedPressedImage(FEditorStyle::GetBrush("TreeArrow_Collapsed"))
+				]
+
+				// Track Name
+				+ SHorizontalBox::Slot()
+				.FillWidth(1)
+				[
+					// Name of track
+					SAssignNew(InlineTextBlock, SInlineEditableTextBlock)
+					.Text(FText::FromName(TrackBase->TrackName))
+					.ToolTipText(LOCTEXT("TrackNameTooltip", "Enter track name"))
+					.OnVerifyTextChanged(TimelineRef, &STimelineEditor::OnVerifyTrackNameCommit, TrackBase, this)
+					.OnTextCommitted(TimelineRef, &STimelineEditor::OnTrackNameCommitted, TrackBase, this)
 				]
 			]
+		]
 
-			+SVerticalBox::Slot()
-			.AutoHeight()
-			.Padding(0,16,0,0)
+		// Content Slot
+		+ SVerticalBox::Slot()
+		[
+			// Box for content visibility
+			SNew(SBox)
+			.Visibility(this, &STimelineEdTrack::GetContentVisibility)
 			[
-				SNew(SBorder)
-				.VAlign(VAlign_Fill)
+				SNew(SHorizontalBox)
+
+				// Label Area
+				+SHorizontalBox::Slot()
+				.AutoWidth()
 				[
-					SNew(SBox)
-					.WidthOverride(128)
-					.HeightOverride(36)
-					.VAlign(VAlign_Center)
-					.HAlign(HAlign_Fill)
+					SNew(SVerticalBox)
+
+					// External Curve Label
+					+SVerticalBox::Slot()
+					.AutoHeight()
+					.Padding(2)
+					[
+						SNew(STextBlock)
+						.Text(LOCTEXT("ExternalCurveLabel", "External Curve"))
+					]
+
+					// External Curve Controls
+					+SVerticalBox::Slot()
+					.AutoHeight()
+					.Padding(2, 0, 2, 4)
 					[
 						SNew(SHorizontalBox)
 						+SHorizontalBox::Slot()
-						.AutoWidth()
-						.Padding(1,0)
-						.MaxWidth(80)
+						.Padding(0, 0, 1, 0)
+						.FillWidth(1)
 						[
 							//External curve name display box
 							SNew(SEditableTextBox)
@@ -183,6 +212,7 @@ void STimelineEdTrack::Construct(const FArguments& InArgs, TSharedPtr<FTimelineE
 								.Image( FEditorStyle::GetBrush(TEXT("PropertyWindow.Button_Use")) )
 							]
 						]
+
 						// Browse external curve button
 						+SHorizontalBox::Slot()
 						.AutoWidth()
@@ -199,6 +229,7 @@ void STimelineEdTrack::Construct(const FArguments& InArgs, TSharedPtr<FTimelineE
 								.Image( FEditorStyle::GetBrush(TEXT("PropertyWindow.Button_Browse")) )
 							]
 						]
+
 						// Convert to internal curve button
 						+SHorizontalBox::Slot()
 						.AutoWidth()
@@ -216,29 +247,46 @@ void STimelineEdTrack::Construct(const FArguments& InArgs, TSharedPtr<FTimelineE
 							]
 						]
 					]
+
+					// Synchronize curve view checkbox.
+					+ SVerticalBox::Slot()
+					.AutoHeight()
+					.Padding(2, 0, 2, 0)
+					[
+						SNew(SCheckBox)
+						.IsChecked(this, &STimelineEdTrack::GetIsCurveViewSynchronizedState)
+						.OnCheckStateChanged(this, &STimelineEdTrack::OnIsCurveViewSynchronizedStateChanged)
+						.ToolTipText(LOCTEXT("SynchronizeViewToolTip", "Keep the zoom and pan of this curve synchronized with other curves."))
+						[
+							SNew(STextBlock)
+							.Text(LOCTEXT("SynchronizeViewLabel", "Synchronize View"))
+						]
+					]
 				]
-			]
-		]
-		+SHorizontalBox::Slot()
-			.FillWidth(1)
-			[
-				SNew(SBorder)
-				.VAlign(VAlign_Fill)
+
+				// Graph Area
+				+SHorizontalBox::Slot()
+				.FillWidth(1)
 				[
-					SAssignNew(TrackWidget, SCurveEditor)
-					.ViewMinInput(TimelineRef, &STimelineEditor::GetViewMinInput)
-					.ViewMaxInput(TimelineRef, &STimelineEditor::GetViewMaxInput)
-					.ViewMinOutput(TimelineRef, &STimelineEditor::GetViewMinOutput)
-					.ViewMaxOutput(TimelineRef, &STimelineEditor::GetViewMaxOutput)
-					.TimelineLength(TimelineRef, &STimelineEditor::GetTimelineLength)
-					.OnSetInputViewRange(TimelineRef, &STimelineEditor::SetInputViewRange)
-					.OnSetOutputViewRange(TimelineRef, &STimelineEditor::SetOutputViewRange)
-					.DesiredSize(TimelineRef, &STimelineEditor::GetTimelineDesiredSize)
-					.DrawCurve(bDrawCurve)
-					.HideUI(false)
-					.OnCreateAsset(this, &STimelineEdTrack::OnCreateExternalCurve )
+					SNew(SBorder)
+					.VAlign(VAlign_Fill)
+					[
+						SAssignNew(TrackWidget, SCurveEditor)
+						.ViewMinInput(this, &STimelineEdTrack::GetMinInput)
+						.ViewMaxInput(this, &STimelineEdTrack::GetMaxInput)
+						.ViewMinOutput(this, &STimelineEdTrack::GetMinOutput)
+						.ViewMaxOutput(this, &STimelineEdTrack::GetMaxOutput)
+						.TimelineLength(TimelineRef, &STimelineEditor::GetTimelineLength)
+						.OnSetInputViewRange(this, &STimelineEdTrack::OnSetInputViewRange)
+						.OnSetOutputViewRange(this, &STimelineEdTrack::OnSetOutputViewRange)
+						.DesiredSize(TimelineRef, &STimelineEditor::GetTimelineDesiredSize)
+						.DrawCurve(bDrawCurve)
+						.HideUI(false)
+						.OnCreateAsset(this, &STimelineEdTrack::OnCreateExternalCurve )
+					]
 				]
 			]
+		]
 	];
 
 	if( TrackBase )
@@ -644,6 +692,93 @@ void STimelineEdTrack::CopyCurveData( const FRichCurve* SrcCurve, FRichCurve* De
 	}
 }
 
+ECheckBoxState STimelineEdTrack::GetIsExpandedState() const
+{
+	return Track->bIsExpanded ? ECheckBoxState::Checked : ECheckBoxState::Unchecked;
+}
+
+void STimelineEdTrack::OnIsExpandedStateChanged(ECheckBoxState IsExpandedState)
+{
+	Track->bIsExpanded = IsExpandedState == ECheckBoxState::Checked;
+}
+
+EVisibility STimelineEdTrack::GetContentVisibility() const
+{
+	return Track->bIsExpanded ? EVisibility::Visible : EVisibility::Collapsed;
+}
+
+ECheckBoxState  STimelineEdTrack::GetIsCurveViewSynchronizedState() const
+{
+	return Track->bIsCurveViewSynchronized ? ECheckBoxState::Checked : ECheckBoxState::Unchecked;
+}
+
+void  STimelineEdTrack::OnIsCurveViewSynchronizedStateChanged(ECheckBoxState IsCurveViewSynchronizedState)
+{
+	Track->bIsCurveViewSynchronized = IsCurveViewSynchronizedState == ECheckBoxState::Checked;
+	if (Track->bIsCurveViewSynchronized == false)
+	{
+		TSharedPtr<STimelineEditor> TimelineEditor = TimelineEdPtr.Pin();
+		LocalInputMin = TimelineEditor->GetViewMinInput();
+		LocalInputMax = TimelineEditor->GetViewMaxInput();
+		LocalOutputMin = TimelineEditor->GetViewMinOutput();
+		LocalOutputMax = TimelineEditor->GetViewMaxOutput();
+	}
+}
+
+float STimelineEdTrack::GetMinInput() const
+{
+	return Track->bIsCurveViewSynchronized
+		? TimelineEdPtr.Pin()->GetViewMinInput()
+		: LocalInputMin;
+}
+
+float STimelineEdTrack::GetMaxInput() const
+{
+	return Track->bIsCurveViewSynchronized
+		? TimelineEdPtr.Pin()->GetViewMaxInput()
+		: LocalInputMax;
+}
+
+float STimelineEdTrack::GetMinOutput() const
+{
+	return Track->bIsCurveViewSynchronized
+		? TimelineEdPtr.Pin()->GetViewMinOutput()
+		: LocalOutputMin;
+}
+
+float STimelineEdTrack::GetMaxOutput() const
+{
+	return Track->bIsCurveViewSynchronized
+		? TimelineEdPtr.Pin()->GetViewMaxOutput()
+		: LocalOutputMax;
+}
+
+void STimelineEdTrack::OnSetInputViewRange(float Min, float Max)
+{
+	if (Track->bIsCurveViewSynchronized)
+	{
+		TimelineEdPtr.Pin()->SetInputViewRange(Min, Max);
+	}
+	else
+	{
+		LocalInputMin = Min;
+		LocalInputMax = Max;
+	}
+}
+
+void STimelineEdTrack::OnSetOutputViewRange(float Min, float Max)
+{
+	if (Track->bIsCurveViewSynchronized)
+	{
+		TimelineEdPtr.Pin()->SetOutputViewRange(Min, Max);
+	}
+	else
+	{
+		LocalOutputMin = Min;
+		LocalOutputMax = Max;
+	}
+}
+
 void STimelineEdTrack::ResetExternalCurveInfo( )
 {
 	ExternalCurveName = FString( TEXT( "None" ) );
@@ -944,6 +1079,7 @@ TSharedRef<ITableRow> STimelineEditor::MakeTrackWidget( TSharedPtr<FTimelineEdTr
 
 	return
 	SNew(STableRow< TSharedPtr<FTimelineEdTrack> >, OwnerTable )
+	.Padding(FMargin(0, 0, 0, 2))
 	[
 		SNew(STimelineEdTrack, Track, SharedThis(this))
 	];
diff --git a/Engine/Source/Editor/Kismet/Private/STimelineEditor.h b/Engine/Source/Editor/Kismet/Private/STimelineEditor.h
--- a/Engine/Source/Editor/Kismet/Private/STimelineEditor.h
+++ b/Engine/Source/Editor/Kismet/Private/STimelineEditor.h
@@ -29,6 +29,11 @@ public:
 	int32 TrackIndex;
 	/** Trigger when a rename is requested on the track */
 	FSimpleDelegate OnRenameRequest;
+	/** Whether or not this track is expanded in the UI. */
+	bool bIsExpanded;
+	/** Whether or not this track has its curve's view synchronized with the other curve views. */
+	bool bIsCurveViewSynchronized;
+
 public:
 	static TSharedRef<FTimelineEdTrack> Make(ETrackType InType, int32 InIndex)
 	{
@@ -40,7 +45,10 @@ private:
 	FTimelineEdTrack(ETrackType InType, int32 InIndex)
 		: TrackType(InType)
 		, TrackIndex(InIndex)
-	{}
+	{
+		bIsExpanded = true;
+		bool bIsCurveViewSynchronized = true;
+	}
 
 	/** Hidden constructor, always use Make above */
 	FTimelineEdTrack() {}
@@ -79,6 +87,18 @@ private:
 	/** String to display external curve path as tooltip*/
 	FString ExternalCurvePath;
 
+	/** The local curve input min to use when the this tracks curve view isn't synchronized. */
+	float LocalInputMin;
+
+	/** The local curve input max to use when the this tracks curve view isn't synchronized. */
+	float LocalInputMax;
+
+	/** The local curve output min to use when the this tracks curve view isn't synchronized. */
+	float LocalOutputMin;
+
+	/** The local curve output max to use when the this tracks curve view isn't synchronized. */
+	float LocalOutputMax;
+
 	/**Function to destroy popup window*/
 	void OnCloseCreateCurveWindow();
 
@@ -121,6 +141,34 @@ private:
 	/** Function to copy data from one curve to another*/
 	static void CopyCurveData( const FRichCurve* SrcCurve, FRichCurve* DestCurve );
 
+	/** Gets whether this track is expanded. */
+	ECheckBoxState GetIsExpandedState() const;
+	/** Callback for when the expanded state for this track changes. */
+	void OnIsExpandedStateChanged(ECheckBoxState IsExpandedState);
+
+	/** Gets whether or not the content of this track should be visible, based on whether or not it's expanded. */
+	EVisibility GetContentVisibility() const;
+
+	/** Gets a check box state representing whether or not this track's curve view is synchronized with other tracks. */
+	ECheckBoxState GetIsCurveViewSynchronizedState() const;
+	/** Callback for when the check box state representing whether or not this track's curve view is synchronized with other tracks changes. */
+	void OnIsCurveViewSynchronizedStateChanged(ECheckBoxState IsCurveViewSynchronized);
+
+	/** Get the minimum input for the curve view. */
+	float GetMinInput() const;
+	/** Get the maximum input for the curve view. */
+	float GetMaxInput() const;
+
+	/** Get the minimum output for the curve view. */
+	float GetMinOutput() const;
+	/** Get the maximum output for the curve view. */
+	float GetMaxOutput() const;
+
+	/** Callback to set the input view range for the curve editor. */
+	void OnSetInputViewRange(float Min, float Max);
+	/** Callback to set the output view range for the curve editor. */
+	void OnSetOutputViewRange(float Min, float Max);
+
 public:
 	/** Inline block for changing name of track */
 	TSharedPtr<SInlineEditableTextBlock> InlineNameBlock;
