Fixed variable initialisation in STimelineEditor.h. (Was causing as shadow warning, but suspect wasnt correct anyway)

diff --git a/Engine/Source/Editor/Kismet/Private/STimelineEditor.h b/Engine/Source/Editor/Kismet/Private/STimelineEditor.h
--- a/Engine/Source/Editor/Kismet/Private/STimelineEditor.h
+++ b/Engine/Source/Editor/Kismet/Private/STimelineEditor.h
@@ -47,7 +47,7 @@ private:
 		, TrackIndex(InIndex)
 	{
 		bIsExpanded = true;
-		bool bIsCurveViewSynchronized = true;
+		bIsCurveViewSynchronized = true;
 	}
 
 	/** Hidden constructor, always use Make above */
