When reparenting Blueprints, all nodes will be regenerated to resolve Target pin issues.

diff --git a/Engine/Source/Editor/Kismet/Private/BlueprintEditor.cpp b/Engine/Source/Editor/Kismet/Private/BlueprintEditor.cpp
--- a/Engine/Source/Editor/Kismet/Private/BlueprintEditor.cpp
+++ b/Engine/Source/Editor/Kismet/Private/BlueprintEditor.cpp
@@ -2553,6 +2553,7 @@ void FBlueprintEditor::ReparentBlueprint_NewParentChosen(UClass* ChosenClass)
 
 			// Ensure that the Blueprint is up-to-date (valid SCS etc.) before compiling
 			EnsureBlueprintIsUpToDate(BlueprintObj);
+			FBlueprintEditorUtils::RefreshAllNodes(GetBlueprintObj());
 			FBlueprintEditorUtils::MarkBlueprintAsModified(BlueprintObj);
 
 			Compile();
