Drag and dropping a component from the SCS editor to a graph will no longer incorrectly report that the graph is outside of scope when it is a part of the same Blueprint.

Drag and drop no longer uses the SCS node's Blueprint, instead looking at the Blueprint Editor it is in for the information.

diff --git a/Engine/Source/Editor/Kismet/Private/SSCSEditor.cpp b/Engine/Source/Editor/Kismet/Private/SSCSEditor.cpp
--- a/Engine/Source/Editor/Kismet/Private/SSCSEditor.cpp
+++ b/Engine/Source/Editor/Kismet/Private/SSCSEditor.cpp
@@ -1875,7 +1875,8 @@ FReply SSCS_RowWidget::HandleOnDragDetected( const FGeometry& MyGeometry, const
 		TSharedPtr<FSCSEditorTreeNode> FirstNode = SelectedNodePtrs[0];
 		if (FirstNode->GetNodeType() == FSCSEditorTreeNode::ComponentNode)
 		{
-			UBlueprint* Blueprint = FirstNode->GetBlueprint();
+			// Do not use the Blueprint from FirstNode, it may still be referencing the parent.
+			UBlueprint* Blueprint = GetBlueprint();
 			const FName VariableName = FirstNode->GetVariableName();
 			UStruct* VariableScope = (Blueprint != nullptr) ? Blueprint->SkeletonGeneratedClass : nullptr;
 
