Harden code against crashes accessing undo transactions in the blueprint editor

https://jira.ol.epicgames.net/browse/UE-18063

diff --git a/Engine/Source/Editor/Kismet/Private/BlueprintEditor.cpp b/Engine/Source/Editor/Kismet/Private/BlueprintEditor.cpp
--- a/Engine/Source/Editor/Kismet/Private/BlueprintEditor.cpp
+++ b/Engine/Source/Editor/Kismet/Private/BlueprintEditor.cpp
@@ -2714,14 +2714,17 @@ void FBlueprintEditor::PostUndo(bool bSuccess)
 
 		// Look at the transaction this function is responding to, see if any object in it has an outermost of the Blueprint
 		const FTransaction* Transaction = GEditor->Trans->GetTransaction(GEditor->Trans->GetQueueLength() - GEditor->Trans->GetUndoCount());
-		TArray<UObject*> TransactionObjects;
-		Transaction->GetTransactionObjects(TransactionObjects);
-		for (UObject* Object : TransactionObjects)
+		if( Transaction != nullptr )
 		{
-			if (Object->GetOutermost() == BlueprintOutermost)
+			TArray<UObject*> TransactionObjects;
+			Transaction->GetTransactionObjects(TransactionObjects);
+			for (UObject* Object : TransactionObjects)
 			{
-				bAffectsBlueprint = true;
-				break;
+				if (Object->GetOutermost() == BlueprintOutermost)
+				{
+					bAffectsBlueprint = true;
+					break;
+				}
 			}
 		}
 
