[UE-17380] Fix issues with duplication of components within the component tree in the Blueprint editor.

change summary:
- added SSCSEditor::bAllowTreeUpdates to avoid making unnecessary (redundant) UpdateTree() calls in specific situations
- modified SSCSEditor::AddNewComponent() to defer Blueprint class regeneration until AFTER properties have been copied from the source template; this also defers any calls to RerunConstructionScripts() on Actor instances so that they don't cache defaults to Component instance data

notes:
- also addresses UE-18152

diff --git a/Engine/Source/Editor/Kismet/Private/SSCSEditor.cpp b/Engine/Source/Editor/Kismet/Private/SSCSEditor.cpp
--- a/Engine/Source/Editor/Kismet/Private/SSCSEditor.cpp
+++ b/Engine/Source/Editor/Kismet/Private/SSCSEditor.cpp
@@ -3242,6 +3242,7 @@ void SSCSEditor::Construct( const FArguments& InArgs )
 	OnHighlightPropertyInDetailsView = InArgs._OnHighlightPropertyInDetailsView;
 	bUpdatingSelection = false;
 	bHasAddedSceneAndBehaviorComponentSeparator = false;
+	bAllowTreeUpdates = true;
 
 	CommandList = MakeShareable( new FUICommandList );
 	CommandList->MapAction( FGenericCommands::Get().Cut,
@@ -4188,6 +4189,12 @@ void SSCSEditor::UpdateTree(bool bRegenerateTreeNodes)
 {
 	check(SCSTreeWidget.IsValid());
 
+	// Early exit if we're deferring tree updates
+	if(!bAllowTreeUpdates)
+	{
+		return;
+	}
+
 	if(bRegenerateTreeNodes)
 	{
 		// Obtain the set of expandable tree nodes that are currently collapsed
@@ -4612,8 +4619,15 @@ UActorComponent* SSCSEditor::AddNewComponent( UClass* NewComponentClass, UObject
 		Blueprint->Modify();
 		SaveSCSCurrentState(Blueprint->SimpleConstructionScript);
 
+		// Defer Blueprint class regeneration and tree updates if we need to copy object properties from a source template.
+		const bool bMarkBlueprintModified = !ComponentTemplate;
+		if(!bMarkBlueprintModified)
+		{
+			bAllowTreeUpdates = false;
+		}
+		
 		const FName NewVariableName = (Asset ? FName(*FComponentEditorUtils::GenerateValidVariableNameFromAsset(Asset, nullptr)) : NAME_None);
-		NewComponent = AddNewNode(Blueprint->SimpleConstructionScript->CreateNode(NewComponentClass, NewVariableName), Asset, true);
+		NewComponent = AddNewNode(Blueprint->SimpleConstructionScript->CreateNode(NewComponentClass, NewVariableName), Asset, bMarkBlueprintModified);
 
 		if (ComponentTemplate)
 		{
@@ -4622,6 +4636,10 @@ UActorComponent* SSCSEditor::AddNewComponent( UClass* NewComponentClass, UObject
 			FObjectWriter Writer(ComponentTemplate, SavedProperties);
 			FObjectReader(NewComponent, SavedProperties);
 			NewComponent->UpdateComponentToWorld();
+
+			// Wait until here to mark as structurally modified because we don't want any RerunConstructionScript() calls to happen until AFTER we've serialized properties from the source object.
+			bAllowTreeUpdates = true;
+			FBlueprintEditorUtils::MarkBlueprintAsStructurallyModified(Blueprint);
 		}
 	}
 	else    // EComponentEditorMode::ActorInstance
diff --git a/Engine/Source/Editor/Kismet/Public/SSCSEditor.h b/Engine/Source/Editor/Kismet/Public/SSCSEditor.h
--- a/Engine/Source/Editor/Kismet/Public/SSCSEditor.h
+++ b/Engine/Source/Editor/Kismet/Public/SSCSEditor.h
@@ -1079,4 +1079,7 @@ private:
 
 	/** true if we've added the separator between the scene and behavior components to the root nodes */
 	bool bHasAddedSceneAndBehaviorComponentSeparator;
+
+	/** Controls whether or not to allow calls to UpdateTree() */
+	bool bAllowTreeUpdates;
 };
