Can drag and drop custom events from the MyBlueprint window to add a call function for the event.

Double clicking will still focus the event.

diff --git a/Engine/Source/Editor/Kismet/Private/BPFunctionDragDropAction.cpp b/Engine/Source/Editor/Kismet/Private/BPFunctionDragDropAction.cpp
--- a/Engine/Source/Editor/Kismet/Private/BPFunctionDragDropAction.cpp
+++ b/Engine/Source/Editor/Kismet/Private/BPFunctionDragDropAction.cpp
@@ -146,13 +146,25 @@ bool FKismetDragDropAction::ActionWillShowExistingNode() const
 	if (ActionNode.IsValid() && (HoveredGraph != NULL))
 	{
 		bWillFocusOnExistingNode = (ActionNode->GetTypeId() == FEdGraphSchemaAction_K2TargetNode::StaticGetTypeId()) ||
-			(ActionNode->GetTypeId() == FEdGraphSchemaAction_K2Event::StaticGetTypeId()) ||
 			(ActionNode->GetTypeId() == FEdGraphSchemaAction_K2InputAction::StaticGetTypeId());
 
-		if (!bWillFocusOnExistingNode && (ActionNode->GetTypeId() == FEdGraphSchemaAction_K2AddEvent::StaticGetTypeId()))
+		if (!bWillFocusOnExistingNode)
 		{
-			FEdGraphSchemaAction_K2AddEvent* AddEventAction = (FEdGraphSchemaAction_K2AddEvent*)ActionNode.Get();
-			bWillFocusOnExistingNode = AddEventAction->EventHasAlreadyBeenPlaced(FBlueprintEditorUtils::FindBlueprintForGraph(HoveredGraph));
+			if (ActionNode->GetTypeId() == FEdGraphSchemaAction_K2AddEvent::StaticGetTypeId())
+			{
+				FEdGraphSchemaAction_K2AddEvent* AddEventAction = (FEdGraphSchemaAction_K2AddEvent*)ActionNode.Get();
+				bWillFocusOnExistingNode = AddEventAction->EventHasAlreadyBeenPlaced(FBlueprintEditorUtils::FindBlueprintForGraph(HoveredGraph));
+			}
+			else if (ActionNode->GetTypeId() == FEdGraphSchemaAction_K2Event::StaticGetTypeId())
+			{
+				FEdGraphSchemaAction_K2Event* FuncAction = (FEdGraphSchemaAction_K2Event*)ActionNode.Get();
+				UK2Node_CustomEvent* CustomEvent = Cast<UK2Node_CustomEvent>(FuncAction->NodeTemplate);
+				// Drag and dropping custom event's will place a Call Function and will not focus the existing event
+				if (CustomEvent == nullptr)
+				{
+					bWillFocusOnExistingNode = true;
+				}
+			}
 		}
 	}
 
diff --git a/Engine/Source/Editor/Kismet/Private/SMyBlueprint.cpp b/Engine/Source/Editor/Kismet/Private/SMyBlueprint.cpp
--- a/Engine/Source/Editor/Kismet/Private/SMyBlueprint.cpp
+++ b/Engine/Source/Editor/Kismet/Private/SMyBlueprint.cpp
@@ -1517,12 +1517,21 @@ FReply SMyBlueprint::OnActionDragged( const TArray< TSharedPtr<FEdGraphSchemaAct
 			}
 		}
 		else if (InAction->GetTypeId() == FEdGraphSchemaAction_K2Event::StaticGetTypeId())
-		{			
-			// don't need a valid FCanBeDroppedDelegate because this entry means we already have this 
-			// event placed (so this action will just focus it)
-			TSharedRef<FKismetDragDropAction> DragOperation = FKismetDragDropAction::New(InAction, AnalyticsDelegate, FKismetDragDropAction::FCanBeDroppedDelegate());
+		{	
+			// Check if it's a custom event, it is preferable to drop a call function for custom events than to focus on the node
+			FEdGraphSchemaAction_K2Event* FuncAction = (FEdGraphSchemaAction_K2Event*)InAction.Get();
+			if (UK2Node_CustomEvent* CustomEvent = Cast<UK2Node_CustomEvent>(FuncAction->NodeTemplate))
+			{
+				return FReply::Handled().BeginDragDrop(FKismetFunctionDragDropAction::New(InAction, CustomEvent->GetFunctionName(), GetBlueprintObj()->SkeletonGeneratedClass, FMemberReference(), AnalyticsDelegate, FKismetDragDropAction::FCanBeDroppedDelegate()));
+			}
+			else
+			{
+				// don't need a valid FCanBeDroppedDelegate because this entry means we already have this 
+				// event placed (so this action will just focus it)
+				TSharedRef<FKismetDragDropAction> DragOperation = FKismetDragDropAction::New(InAction, AnalyticsDelegate, FKismetDragDropAction::FCanBeDroppedDelegate());
 
-			return FReply::Handled().BeginDragDrop(DragOperation);
+				return FReply::Handled().BeginDragDrop(DragOperation);
+			}
 		}
 	}
 
