Duplicating function or macro graphs that contain nodes with Text pins will no longer incorrectly link the original and duplicated pins inexplicably so that changes in one affect the other.

During duplication of entire graphs, node Guids are regenerated on the new graph.

diff --git a/Engine/Source/Editor/Kismet/Private/SMyBlueprint.cpp b/Engine/Source/Editor/Kismet/Private/SMyBlueprint.cpp
--- a/Engine/Source/Editor/Kismet/Private/SMyBlueprint.cpp
+++ b/Engine/Source/Editor/Kismet/Private/SMyBlueprint.cpp
@@ -2505,7 +2505,16 @@ void SMyBlueprint::OnDuplicateAction()
 
 		DuplicatedGraph->Modify();
 
-		// Only function duplication is supported
+		// Generate new Guids for all nodes in the graph
+		// *NOTE* this cannot occur during PostDuplicate, node Guids need to remain static during duplication for Blueprint compilation
+		for (UEdGraphNode* EdGraphNode : DuplicatedGraph->Nodes)
+		{
+			if (EdGraphNode)
+			{
+				EdGraphNode->CreateNewGuid();
+			}
+		}
+		// Only function and macro duplication is supported
 		EGraphType GraphType = DuplicatedGraph->GetSchema()->GetGraphType(GraphAction->EdGraph);
 		check(GraphType == GT_Function || GraphType == GT_Macro);
 
