[UE-6303] Removing the old components/graph/defaults buttons. Can use the diff list to navigate these controls and they no longer match the full blueprint editor. Removing them also makes it easier to support animation and UMG blueprints (which don't have components)

diff --git a/Engine/Source/Editor/Kismet/Private/SBlueprintDiff.cpp b/Engine/Source/Editor/Kismet/Private/SBlueprintDiff.cpp
--- a/Engine/Source/Editor/Kismet/Private/SBlueprintDiff.cpp
+++ b/Engine/Source/Editor/Kismet/Private/SBlueprintDiff.cpp
@@ -649,63 +649,9 @@ void SBlueprintDiff::Construct( const FArguments& InArgs)
 
 	bLockViews = true;
 
-	TArray<UEdGraph*> GraphsOld,GraphsNew;
-	PanelOld.Blueprint->GetAllGraphs(GraphsOld);
-	PanelNew.Blueprint->GetAllGraphs(GraphsNew);
-	
-	//Add Graphs that exist in both blueprints, or in blueprint 1 only
-	for(auto It(GraphsOld.CreateConstIterator());It;It++)
-	{
-		UEdGraph* GraphOld = *It;
-		UEdGraph* GraphNew = NULL;
-		for(auto It2(GraphsNew.CreateIterator());It2;It2++)
-		{
-			UEdGraph* TestGraph = *It2;
-			if(TestGraph && GraphOld->GetName() == TestGraph->GetName())
-			{
-				GraphNew = TestGraph;
-				*It2 = NULL;
-				break;
-			}
-		}
-		// Do not worry about graphs that are contained in MathExpression nodes, they are recreated each compile
-		if (IsGraphDiffNeeded(GraphOld))
-		{
-			CreateGraphEntry(GraphOld,GraphNew);
-		}
-	}
-
-	//Add graphs that only exist in 2nd(new) blueprint
-	for(auto It2(GraphsNew.CreateIterator());It2;It2++)
-	{
-		UEdGraph* GraphNew = *It2;
-		if(GraphNew != NULL)
-		{
-			// Do not worry about graphs that are contained in MathExpression nodes, they are recreated each compile
-			if (IsGraphDiffNeeded(GraphNew))
-			{
-				CreateGraphEntry(NULL,GraphNew);
-			}
-		}
-	}
-
 	TAttribute<FName> GetActiveMode(this, &SBlueprintDiff::GetCurrentMode);
 	FOnModeChangeRequested SetActiveMode = FOnModeChangeRequested::CreateRaw(this, &SBlueprintDiff::SetCurrentMode);
 
-	auto Widgets = FBlueprintEditorToolbar::GenerateToolbarWidgets( InArgs._BlueprintNew, GetActiveMode, SetActiveMode );
-
-	TSharedRef<SHorizontalBox> MiscWidgets = SNew(SHorizontalBox);
-
-	for (int32 WidgetIdx = 0; WidgetIdx < Widgets.Num(); ++WidgetIdx)
-	{
-		MiscWidgets->AddSlot()
-			.AutoWidth()
-			.VAlign(VAlign_Center)
-			[
-				Widgets[WidgetIdx].ToSharedRef()
-			];
-	}
-
 	FToolBarBuilder ToolbarBuilder(TSharedPtr< const FUICommandList >(), FMultiBoxCustomization::None);
 	ToolbarBuilder.AddToolBarButton(
 		FUIAction(
@@ -715,8 +661,7 @@ void SBlueprintDiff::Construct( const FArguments& InArgs)
 		, NAME_None
 		, LOCTEXT("PrevDiffLabel", "Prev")
 		, LOCTEXT("PrevDiffTooltip", "Go to previous difference")
-		, FSlateIcon(FEditorStyle::GetStyleSetName()
-		, "BlueprintDif.PrevDiff")
+		, FSlateIcon(FEditorStyle::GetStyleSetName(), "BlueprintDif.PrevDiff")
 	);
 	ToolbarBuilder.AddToolBarButton(
 		FUIAction(
@@ -726,12 +671,11 @@ void SBlueprintDiff::Construct( const FArguments& InArgs)
 		, NAME_None
 		, LOCTEXT("NextDiffLabel", "Next")
 		, LOCTEXT("NextDiffTooltip", "Go to next difference")
-		, FSlateIcon(FEditorStyle::GetStyleSetName(
-		), "BlueprintDif.NextDiff")
+		, FSlateIcon(FEditorStyle::GetStyleSetName(), "BlueprintDif.NextDiff")
 	);
 	ToolbarBuilder.AddSeparator();
 	ToolbarBuilder.AddToolBarButton(
-		FUIAction(FExecuteAction::CreateSP(this, &SBlueprintDiff::OnToggleLockView2))
+		FUIAction(FExecuteAction::CreateSP(this, &SBlueprintDiff::OnToggleLockView))
 		, NAME_None
 		, LOCTEXT("LockGraphsLabel", "Lock/Unlock")
 		, LOCTEXT("LockGraphsTooltip", "Force all graph views to change together, or allow independent scrolling/zooming")
@@ -739,43 +683,8 @@ void SBlueprintDiff::Construct( const FArguments& InArgs)
 	);
 
 	GraphPanel = GenerateGraphPanel();
-	DefaultsPanel = GenerateDefaultsPanel();
-	ComponentsPanel = GenerateComponentsPanel();
-	// Unfortunately we can't perform the diff until the UI is generated, the primary reason for this is that
-	// details customizations determine what is actually editable:
-	for (const auto& Graph : Graphs)
-	{
-		TArray< TSharedPtr<FBlueprintDifferenceTreeEntry> > Children;
-		for (const auto& Difference : Graph->DiffListSource)
-		{
-			auto ChildEntry = TSharedPtr<FBlueprintDifferenceTreeEntry>(
-				new FBlueprintDifferenceTreeEntry(
-				FOnDiffEntryFocused::CreateRaw(this, &SBlueprintDiff::OnDiffListSelectionChanged, Difference)
-				, FGenerateDiffEntryWidget::CreateSP(Difference.ToSharedRef(), &FDiffResultItem::GenerateWidget)
-				, TArray< TSharedPtr<FBlueprintDifferenceTreeEntry> >()
-				)
-			);
-			Children.Push(ChildEntry);
-			RealDifferences.Push(ChildEntry);
-			(void)ChildEntry;
-		}
-
-		if (Children.Num() == 0)
-		{
-			// make one child informing the user that there are no differences:
-			Children.Push(FBlueprintDifferenceTreeEntry::NoDifferencesEntry());
-		}
-
-		auto Entry = TSharedPtr<FBlueprintDifferenceTreeEntry>(
-			new FBlueprintDifferenceTreeEntry(
-				FOnDiffEntryFocused::CreateRaw(this, &SBlueprintDiff::OnSelectionChanged, Graph, ESelectInfo::Direct )
-				, FGenerateDiffEntryWidget::CreateSP(Graph.ToSharedRef(), &FListItemGraphToDiff::GenerateWidget)
-				, Children)
-			);
-		MasterDifferencesList.Push(Entry);
-	}
 
-	DifferencesTreeView = DiffTreeView::CreateTreeView(&MasterDifferencesList);
+	GenerateDifferencesList();
 
 	const auto TextBlock = [](FText Text) -> TSharedRef<SWidget>
 	{
@@ -807,7 +716,7 @@ void SBlueprintDiff::Construct( const FArguments& InArgs)
 				[
 					SNew(SHorizontalBox)
 					+ SHorizontalBox::Slot()
-					.Padding(16.f)
+					.Padding(4.f)
 					.AutoWidth()
 					[
 						ToolbarBuilder.MakeWidget()
@@ -816,22 +725,6 @@ void SBlueprintDiff::Construct( const FArguments& InArgs)
 					[
 						SNew(SSpacer)
 					]
-					+ SHorizontalBox::Slot()
-					.AutoWidth()
-					[
-						// Nested box just give us padding similar to the main blueprint editor,
-						// one border has the dark grey "Toolbar.Background" style:
-						SNew(SHorizontalBox)
-						+ SHorizontalBox::Slot()
-						.Padding(0.0f, 4.f, 0.0f, 12.f)
-						[
-							SNew(SBorder)
-							.BorderImage(FEditorStyle::GetBrush(TEXT("Toolbar.Background")))
-							[
-								MiscWidgets
-							]
-						]
-					]
 				]
 				+ SVerticalBox::Slot()
 				[
@@ -1031,16 +924,10 @@ void SBlueprintDiff::OnDiffListSelectionChanged(TSharedPtr<FDiffResultItem> TheD
 	}
 }
 
-void	SBlueprintDiff::OnToggleLockView2()
-{
-	OnToggleLockView();
-}
-
-FReply	SBlueprintDiff::OnToggleLockView()
+void SBlueprintDiff::OnToggleLockView()
 {
 	bLockViews = !bLockViews;
 	ResetGraphEditors();
-	return FReply::Handled();
 }
 
 FSlateIcon SBlueprintDiff::GetLockViewImage() const
@@ -1227,6 +1114,88 @@ void SBlueprintDiff::HandleGraphChanged( const FName GraphName )
 	PanelNew.GeneratePanel(GraphNew, GraphOld);
 }
 
+void SBlueprintDiff::GenerateDifferencesList()
+{
+	MasterDifferencesList.Empty();
+	Graphs.Empty();
+
+	TArray<UEdGraph*> GraphsOld, GraphsNew;
+	PanelOld.Blueprint->GetAllGraphs(GraphsOld);
+	PanelNew.Blueprint->GetAllGraphs(GraphsNew);
+
+	//Add Graphs that exist in both blueprints, or in blueprint 1 only
+	for (auto It(GraphsOld.CreateConstIterator()); It; It++)
+	{
+		UEdGraph* GraphOld = *It;
+		UEdGraph* GraphNew = NULL;
+		for (auto It2(GraphsNew.CreateIterator()); It2; It2++)
+		{
+			UEdGraph* TestGraph = *It2;
+			if (TestGraph && GraphOld->GetName() == TestGraph->GetName())
+			{
+				GraphNew = TestGraph;
+				*It2 = NULL;
+				break;
+			}
+		}
+		// Do not worry about graphs that are contained in MathExpression nodes, they are recreated each compile
+		if (IsGraphDiffNeeded(GraphOld))
+		{
+			CreateGraphEntry(GraphOld,GraphNew);
+		}
+	}
+
+	//Add graphs that only exist in 2nd(new) blueprint
+	for (auto It2(GraphsNew.CreateIterator()); It2; It2++)
+	{
+		UEdGraph* GraphNew = *It2;
+		if (GraphNew != NULL && IsGraphDiffNeeded(GraphNew))
+		{
+			CreateGraphEntry(NULL, GraphNew);
+		}
+	}
+
+	// Unfortunately we can't perform the diff until the UI is generated, the primary reason for this is that
+	// details customizations determine what is actually editable:
+	DefaultsPanel = GenerateDefaultsPanel();
+	ComponentsPanel = GenerateComponentsPanel();
+
+
+	for (const auto& Graph : Graphs)
+	{
+		TArray< TSharedPtr<FBlueprintDifferenceTreeEntry> > Children;
+		for (const auto& Difference : Graph->DiffListSource)
+		{
+			auto ChildEntry = TSharedPtr<FBlueprintDifferenceTreeEntry>(
+				new FBlueprintDifferenceTreeEntry(
+				FOnDiffEntryFocused::CreateRaw(this, &SBlueprintDiff::OnDiffListSelectionChanged, Difference)
+				, FGenerateDiffEntryWidget::CreateSP(Difference.ToSharedRef(), &FDiffResultItem::GenerateWidget)
+				, TArray< TSharedPtr<FBlueprintDifferenceTreeEntry> >()
+				)
+				);
+			Children.Push(ChildEntry);
+			RealDifferences.Push(ChildEntry);
+			(void)ChildEntry;
+		}
+
+		if (Children.Num() == 0)
+		{
+			// make one child informing the user that there are no differences:
+			Children.Push(FBlueprintDifferenceTreeEntry::NoDifferencesEntry());
+		}
+
+		auto Entry = TSharedPtr<FBlueprintDifferenceTreeEntry>(
+			new FBlueprintDifferenceTreeEntry(
+			FOnDiffEntryFocused::CreateRaw(this, &SBlueprintDiff::OnSelectionChanged, Graph, ESelectInfo::Direct)
+			, FGenerateDiffEntryWidget::CreateSP(Graph.ToSharedRef(), &FListItemGraphToDiff::GenerateWidget)
+			, Children)
+			);
+		MasterDifferencesList.Push(Entry);
+	}
+
+	DifferencesTreeView = DiffTreeView::CreateTreeView(&MasterDifferencesList);
+}
+
 SBlueprintDiff::FDiffControl SBlueprintDiff::GenerateGraphPanel()
 {
 	const TSharedRef<FTabManager::FLayout> DefaultLayout = FTabManager::NewLayout("BlueprintDiff_Layout_v1")
diff --git a/Engine/Source/Editor/Kismet/Public/SBlueprintDiff.h b/Engine/Source/Editor/Kismet/Public/SBlueprintDiff.h
--- a/Engine/Source/Editor/Kismet/Public/SBlueprintDiff.h
+++ b/Engine/Source/Editor/Kismet/Public/SBlueprintDiff.h
@@ -157,8 +157,7 @@ protected:
 	void DisablePinDiffFocus();
 
 	/*User toggles the option to lock the views between the two blueprints */
-	void	OnToggleLockView2();
-	FReply	OnToggleLockView();
+	void OnToggleLockView();
 
 	/*Reset the graph editor, called when user switches graphs to display*/
 	void ResetGraphEditors();
@@ -175,6 +174,9 @@ protected:
 	/** Event handler that updates the graph view when user selects a new graph */
 	void HandleGraphChanged( const FName GraphName );
 
+	/** Function used to generate the list of differences and the widgets needed to calculate that list */
+	void GenerateDifferencesList();
+
 	/** Checks if a graph is valid for diff'ing */
 	bool IsGraphDiffNeeded(class UEdGraph* InGraph) const;
 
