Blueprints: Improved the logic for when to reuse an existing tab versus open a new one
- All subgraphs will now open in the parent when double-clicking on a node in a graph (e.g., animation states and rules, instead of only collapsed subgraphs in the ubergraph)
- Holding down Shift while clicking on an item in the My Blueprint tree or in the graph will force the creation of a new tab regardless of what is already opened

diff --git a/Engine/Source/Editor/Kismet/Private/BlueprintEditor.cpp b/Engine/Source/Editor/Kismet/Private/BlueprintEditor.cpp
--- a/Engine/Source/Editor/Kismet/Private/BlueprintEditor.cpp
+++ b/Engine/Source/Editor/Kismet/Private/BlueprintEditor.cpp
@@ -3180,7 +3180,7 @@ void FBlueprintEditor::JumpToHyperlink(const UObject* ObjectReference, bool bReq
 
 	if (const UEdGraphNode* Node = Cast<const UEdGraphNode>(ObjectReference))
 	{
-		if( bRequestRename )
+		if (bRequestRename)
 		{
 			IsNodeTitleVisible(Node, bRequestRename);
 		}
@@ -3195,15 +3195,51 @@ void FBlueprintEditor::JumpToHyperlink(const UObject* ObjectReference, bool bReq
 	}
 	else if (const UEdGraph* Graph = Cast<const UEdGraph>(ObjectReference))
 	{
-		// Only ubergraphs should re-use the current tab
-		if(Graph->GetSchema()->GetGraphType(Graph) == GT_Ubergraph || Cast<UK2Node_Composite>(Graph->GetOuter()))
+		// Navigating into things should re-use the current tab when it makes sense
+		FDocumentTracker::EOpenDocumentCause OpenMode = FDocumentTracker::OpenNewDocument;
+		if ((Graph->GetSchema()->GetGraphType(Graph) == GT_Ubergraph) || Cast<UK2Node_Composite>(Graph->GetOuter()))
 		{
-			OpenDocument(const_cast<UEdGraph*>(Graph), FDocumentTracker::NavigatingCurrentDocument);
+			// Ubergraphs directly reuse the current graph
+			OpenMode = FDocumentTracker::NavigatingCurrentDocument;
 		}
 		else
 		{
-			OpenDocument(const_cast<UEdGraph*>(Graph), FDocumentTracker::OpenNewDocument);
+			// Walk up the outer chain to see if any tabs have a parent of this document open for edit, and if so
+			// we should reuse that one and drill in deeper instead
+			for (UObject* WalkPtr = const_cast<UEdGraph*>(Graph); WalkPtr != nullptr; WalkPtr = WalkPtr->GetOuter())
+			{
+				TArray< TSharedPtr<SDockTab> > TabResults;
+				if (FindOpenTabsContainingDocument(WalkPtr, /*out*/ TabResults))
+				{
+					// See if the parent was active
+					bool bIsActive = false;
+					for (TSharedPtr<SDockTab> Tab : TabResults)
+					{
+						if (Tab->IsActive())
+						{
+							bIsActive = true;
+							break;
+						}
+					}
+
+					if (bIsActive)
+					{
+						OpenMode = FDocumentTracker::NavigatingCurrentDocument;
+						break;
+					}
+				}
+			}
 		}
+
+		// Force it to open in a new document if shift is pressed
+		const bool bIsShiftPressed = FSlateApplication::Get().GetModifierKeys().IsShiftDown();
+		if (bIsShiftPressed)
+		{
+			OpenMode = FDocumentTracker::ForceOpenNewDocument;
+		}
+
+		// Open the document
+		OpenDocument(const_cast<UEdGraph*>(Graph), OpenMode);
 	}
 	else if (const AActor* ReferencedActor = Cast<const AActor>(ObjectReference))
 	{
@@ -3216,8 +3252,7 @@ void FBlueprintEditor::JumpToHyperlink(const UObject* ObjectReference, bool bReq
 	}
 	else if(const UFunction* Function = Cast<const UFunction>(ObjectReference))
 	{
-		UEdGraph* FunctionGraph = FBlueprintEditorUtils::FindScopeGraph(GetBlueprintObj(), Function);
-		if(FunctionGraph)
+		if (UEdGraph* FunctionGraph = FBlueprintEditorUtils::FindScopeGraph(GetBlueprintObj(), Function))
 		{
 			OpenDocument(const_cast<UEdGraph*>(FunctionGraph), FDocumentTracker::OpenNewDocument);
 		}
diff --git a/Engine/Source/Editor/Kismet/Private/SMyBlueprint.cpp b/Engine/Source/Editor/Kismet/Private/SMyBlueprint.cpp
--- a/Engine/Source/Editor/Kismet/Private/SMyBlueprint.cpp
+++ b/Engine/Source/Editor/Kismet/Private/SMyBlueprint.cpp
@@ -1661,6 +1661,10 @@ void SMyBlueprint::OnActionDoubleClicked(const TArray< TSharedPtr<FEdGraphSchema
 
 void SMyBlueprint::ExecuteAction(TSharedPtr<FEdGraphSchemaAction> InAction)
 {
+	// Force it to open in a new document if shift is pressed
+	const bool bIsShiftPressed = FSlateApplication::Get().GetModifierKeys().IsShiftDown();
+	FDocumentTracker::EOpenDocumentCause OpenMode = bIsShiftPressed ? FDocumentTracker::ForceOpenNewDocument : FDocumentTracker::OpenNewDocument;
+
 	UBlueprint* BlueprintObj = BlueprintEditorPtr.Pin()->GetBlueprintObj();
 	if(InAction.IsValid())
 	{
@@ -1670,7 +1674,7 @@ void SMyBlueprint::ExecuteAction(TSharedPtr<FEdGraphSchemaAction> InAction)
 
 			if (GraphAction->EdGraph)
 			{
-				BlueprintEditorPtr.Pin()->OpenDocument(GraphAction->EdGraph, FDocumentTracker::OpenNewDocument);
+				BlueprintEditorPtr.Pin()->OpenDocument(GraphAction->EdGraph, OpenMode);
 			}
 		}
 		if (InAction->GetTypeId() == FEdGraphSchemaAction_K2Delegate::StaticGetTypeId())
@@ -1679,7 +1683,7 @@ void SMyBlueprint::ExecuteAction(TSharedPtr<FEdGraphSchemaAction> InAction)
 
 			if (DelegateAction->EdGraph)
 			{
-				BlueprintEditorPtr.Pin()->OpenDocument(DelegateAction->EdGraph, FDocumentTracker::OpenNewDocument);
+				BlueprintEditorPtr.Pin()->OpenDocument(DelegateAction->EdGraph, OpenMode);
 			}
 		}
 		else if(InAction->GetTypeId() == FEdGraphSchemaAction_K2Var::StaticGetTypeId())
@@ -1697,7 +1701,7 @@ void SMyBlueprint::ExecuteAction(TSharedPtr<FEdGraphSchemaAction> InAction)
 					// Convert the Timeline's name to a variable name before comparing it to the variable
 					if (FName(*UTimelineTemplate::TimelineTemplateNameToVariableName(BlueprintObj->Timelines[i]->GetFName())) == VarAction->GetVariableName())
 					{
-						BlueprintEditorPtr.Pin()->OpenDocument(BlueprintObj->Timelines[i], FDocumentTracker::OpenNewDocument);
+						BlueprintEditorPtr.Pin()->OpenDocument(BlueprintObj->Timelines[i], OpenMode);
 					}
 				}
 			}
