Made the SCS Editor Viewport camera speed a per-project editor user setting.

diff --git a/Engine/Source/Editor/Kismet/Private/SCSEditorViewportClient.cpp b/Engine/Source/Editor/Kismet/Private/SCSEditorViewportClient.cpp
--- a/Engine/Source/Editor/Kismet/Private/SCSEditorViewportClient.cpp
+++ b/Engine/Source/Editor/Kismet/Private/SCSEditorViewportClient.cpp
@@ -633,6 +633,16 @@ FMatrix FSCSEditorViewportClient::GetWidgetCoordSystem() const
 	return Matrix;
 }
 
+int32 FSCSEditorViewportClient::GetCameraSpeedSetting() const
+{
+	return GetDefault<UEditorPerProjectUserSettings>()->SCSViewportCameraSpeed;
+}
+
+void FSCSEditorViewportClient::SetCameraSpeedSetting(int32 SpeedSetting)
+{
+	GetMutableDefault<UEditorPerProjectUserSettings>()->SCSViewportCameraSpeed = SpeedSetting;
+}
+
 void FSCSEditorViewportClient::InvalidatePreview(bool bResetCamera)
 {
 	// Ensure that the editor is valid before continuing
diff --git a/Engine/Source/Editor/Kismet/Private/SCSEditorViewportClient.h b/Engine/Source/Editor/Kismet/Private/SCSEditorViewportClient.h
--- a/Engine/Source/Editor/Kismet/Private/SCSEditorViewportClient.h
+++ b/Engine/Source/Editor/Kismet/Private/SCSEditorViewportClient.h
@@ -39,6 +39,8 @@ public:
 	virtual FVector GetWidgetLocation() const override;
 	virtual FMatrix GetWidgetCoordSystem() const override;
 	virtual ECoordSystem GetWidgetCoordSystemSpace() const override { return WidgetCoordSystem; }
+	virtual int32 GetCameraSpeedSetting() const override;
+	virtual void SetCameraSpeedSetting(int32 SpeedSetting) override;
 
 
 	/** 
diff --git a/Engine/Source/Editor/UnrealEd/Classes/Editor/EditorPerProjectUserSettings.h b/Engine/Source/Editor/UnrealEd/Classes/Editor/EditorPerProjectUserSettings.h
--- a/Engine/Source/Editor/UnrealEd/Classes/Editor/EditorPerProjectUserSettings.h
+++ b/Engine/Source/Editor/UnrealEd/Classes/Editor/EditorPerProjectUserSettings.h
@@ -69,6 +69,10 @@ class UEditorPerProjectUserSettings : public UObject
 	UPROPERTY(config)
 	bool bSCSEditorShowFloor;
 
+	/** How fast the SCS viewport camera moves */
+	UPROPERTY(config, meta=(UIMin = "1", UIMax = "8", ClampMin="1", ClampMax="8"))
+	int32 SCSViewportCameraSpeed;
+
 	// Color curve:
 	//   Release->Attack happens instantly
 	//   Attack holds for AttackHoldPeriod, then
diff --git a/Engine/Source/Editor/UnrealEd/Private/Settings/EditorPerProjectUserSettings.cpp b/Engine/Source/Editor/UnrealEd/Private/Settings/EditorPerProjectUserSettings.cpp
--- a/Engine/Source/Editor/UnrealEd/Private/Settings/EditorPerProjectUserSettings.cpp
+++ b/Engine/Source/Editor/UnrealEd/Private/Settings/EditorPerProjectUserSettings.cpp
@@ -13,6 +13,7 @@ UEditorPerProjectUserSettings::UEditorPerProjectUserSettings(const FObjectInitia
 	MaterialQualityLevel = 1;
 	bMonitorEditorPerformance = true;
 	BlueprintFavorites = CreateDefaultSubobject<UBlueprintPaletteFavorites>(TEXT("BlueprintFavorites"));
+	SCSViewportCameraSpeed = 4;
 }
 
 void UEditorPerProjectUserSettings::PostInitProperties()
