[UE-18670][UE-18614] Making sure deleted components are renamed on removal (so they don't conflict with new component names).

diff --git a/Engine/Source/Editor/Kismet/Private/SSCSEditor.cpp b/Engine/Source/Editor/Kismet/Private/SSCSEditor.cpp
--- a/Engine/Source/Editor/Kismet/Private/SSCSEditor.cpp
+++ b/Engine/Source/Editor/Kismet/Private/SSCSEditor.cpp
@@ -5127,6 +5127,19 @@ void SSCSEditor::RemoveComponentNode(FSCSEditorTreeNodePtrType InNodePtr)
 
 			// Clear the delegate
 			SCS_Node->SetOnNameChanged(FSCSNodeNameChanged());
+
+			// on removal, since we don't move the template from the 
+			// GeneratedClass (which we shouldn't, as it would create a 
+			// discrepancy with existing instances), we rename it instead so that 
+			// we can re-use the name without having to compile (we still have a 
+			// problem if they attempt to name it to what ever we choose here, 
+			// but that is unlikely)
+			if (SCS_Node->ComponentTemplate != nullptr)
+			{
+				SCS_Node->ComponentTemplate->Modify();
+				const FString RemovedName = SCS_Node->GetVariableName().ToString() + TEXT("_REMOVED_") + FGuid::NewGuid().ToString();
+				SCS_Node->ComponentTemplate->Rename(*RemovedName, /*NewOuter =*/nullptr, REN_DontCreateRedirectors);
+			}
 		}
 	}
 	else    // EComponentEditorMode::ActorInstance
