[UE-17622] Cleaning up the DefaultSceneRoot component's tooltip to reflect that it will no longer be auto replaced.

diff --git a/Engine/Source/Editor/Kismet/Private/SSCSEditor.cpp b/Engine/Source/Editor/Kismet/Private/SSCSEditor.cpp
--- a/Engine/Source/Editor/Kismet/Private/SSCSEditor.cpp
+++ b/Engine/Source/Editor/Kismet/Private/SSCSEditor.cpp
@@ -2812,7 +2812,7 @@ FText SSCS_RowWidget::GetTooltipText() const
 		}
 		else
 		{
-			return LOCTEXT("DefaultSceneRootToolTip", "This is the default scene root component. It cannot be copied, renamed or deleted.\nAdding a new scene component will automatically replace it as the new root.");
+			return LOCTEXT("DefaultSceneRootToolTip", "This is the default scene root component. It cannot be copied, renamed or deleted.\nIt can be replaced by drag/dropping another scene component over it.");
 		}
 	}
 	else
