[UE-18695] Duplicate, copy, rename, etc. commands now work with the Blueprint viewport focused.

diff --git a/Engine/Source/Editor/Kismet/Private/SSCSEditorViewport.cpp b/Engine/Source/Editor/Kismet/Private/SSCSEditorViewport.cpp
--- a/Engine/Source/Editor/Kismet/Private/SSCSEditorViewport.cpp
+++ b/Engine/Source/Editor/Kismet/Private/SSCSEditorViewport.cpp
@@ -371,12 +371,12 @@ void SSCSEditorViewport::BindCommands()
 		FCanExecuteAction::CreateSP(SCSEditorWidget, &SSCSEditor::CanDeleteNodes)
 	);
 
-	const FBlueprintEditorCommands& Commands = FBlueprintEditorCommands::Get();
-
-	CommandList->Append(BlueprintEditor->GetToolkitCommands());
 	CommandList->Append(BlueprintEditor->GetSCSEditor()->CommandList.ToSharedRef());
+	CommandList->Append(BlueprintEditor->GetToolkitCommands());
 	SEditorViewport::BindCommands();
 
+	const FBlueprintEditorCommands& Commands = FBlueprintEditorCommands::Get();
+
 	BlueprintEditorPtr.Pin()->GetToolkitCommands()->MapAction(
 		Commands.EnableSimulation,
 		FExecuteAction::CreateSP(this, &SSCSEditorViewport::ToggleIsSimulateEnabled),
