Undoing edits on instanced sub-objects in the Blueprints defaults will no longer crash.

Fix for an issue still present after CL# 2632686.

diff --git a/Engine/Source/Editor/Kismet/Private/BlueprintEditor.cpp b/Engine/Source/Editor/Kismet/Private/BlueprintEditor.cpp
--- a/Engine/Source/Editor/Kismet/Private/BlueprintEditor.cpp
+++ b/Engine/Source/Editor/Kismet/Private/BlueprintEditor.cpp
@@ -2725,10 +2725,13 @@ void FBlueprintEditor::FixSubObjectReferencesPostUndoRedo(UObject* InObject)
 			// The property and sub-objects should have the same name.
 			if (PropertySubObject->GetFName() == SubObject->GetFName())
 			{
+				// We found a matching property, we do not want to re-make the property
+				bFoundMatchingSubObject = true;
+
+				// Check if the properties have different outers so we can map old-to-new
 				if (PropertySubObject->GetOuter() != InObject)
 				{
 					OldToNewInstanceMap.Add(PropertySubObject, SubObject);
-					bFoundMatchingSubObject = true;
 				}
 				// Recurse on the SubObject to correct any sub-object/property references
 				FixSubObjectReferencesPostUndoRedo(SubObject);
