Fix redundant component tree updates in the level editor's Details view when the Actor's root component is both native and editor-only.

diff --git a/Engine/Source/Editor/Kismet/Private/SSCSEditor.cpp b/Engine/Source/Editor/Kismet/Private/SSCSEditor.cpp
--- a/Engine/Source/Editor/Kismet/Private/SSCSEditor.cpp
+++ b/Engine/Source/Editor/Kismet/Private/SSCSEditor.cpp
@@ -3568,9 +3568,9 @@ void SSCSEditor::Tick( const FGeometry& AllottedGeometry, const double InCurrent
 			int32 NumComponentInstances = 0;
 			for (auto CompIt = ActorInstance->GetComponents().CreateConstIterator(); CompIt; ++CompIt)
 			{
-				// Don't count editor-only components in instanced mode, because we don't show them there. Also, don't count UCS-added components if the option to hide them is enabled, and don't count components that are based on the CDO in the BP editor context (i.e. preview world).
+				// Don't count editor-only components in instanced mode (unless it's the root component), because we don't show them there. Also, don't count UCS-added components if the option to hide them is enabled, and don't count components that are based on the CDO in the BP editor context (i.e. preview world).
 				const UActorComponent* CompInst = *CompIt;
-				if ((!CompInst->IsEditorOnly() || EditorMode != EComponentEditorMode::ActorInstance)
+				if ((!CompInst->IsEditorOnly() || EditorMode != EComponentEditorMode::ActorInstance || CompInst == ActorInstance->GetRootComponent())
 					&& (CompInst->CreationMethod != EComponentCreationMethod::UserConstructionScript || bShowInstancedUCSComponents)
 					&& (EditorMode != EComponentEditorMode::BlueprintSCS || CompInst->GetArchetype() != CompInst->GetClass()->GetDefaultObject()))
 				{
