Functions correctly inherit purity in Blueprints and will update to reflect the current purity state of the inherited function.

Made accessors to get the Function Entry's flags for the generated UFunction and for retrieval of the final skeleton class's UFunction's flags.

diff --git a/Engine/Source/Editor/BlueprintGraph/Classes/K2Node_FunctionEntry.h b/Engine/Source/Editor/BlueprintGraph/Classes/K2Node_FunctionEntry.h
--- a/Engine/Source/Editor/BlueprintGraph/Classes/K2Node_FunctionEntry.h
+++ b/Engine/Source/Editor/BlueprintGraph/Classes/K2Node_FunctionEntry.h
@@ -13,10 +13,6 @@ class UK2Node_FunctionEntry : public UK2Node_FunctionTerminator
 	UPROPERTY()
 	FName CustomGeneratedFunctionName;
 
-	/** Any extra flags that the function may need */
-	UPROPERTY()
-	int32 ExtraFlags;
-
 	/** Function metadata */
 	UPROPERTY()
 	struct FKismetUserDeclaredFunctionMetadata MetaData;
@@ -63,5 +59,25 @@ class UK2Node_FunctionEntry : public UK2Node_FunctionTerminator
 
 	// Returns pin for the automatically added WorldContext parameter (used only by BlueprintFunctionLibrary).
 	BLUEPRINTGRAPH_API UEdGraphPin* GetAutoWorldContextPin() const;
+
+	/** Retrieves the function flags from the UFunction that this function entry node represents */
+	BLUEPRINTGRAPH_API int32 GetFunctionFlags() const;
+
+	/** Retrieves the extra flags set on this node. */
+	BLUEPRINTGRAPH_API int32 GetExtraFlags(struct FKismetFunctionContext& CompilerContext) const
+	{
+		return ExtraFlags;
+	}
+
+	/** Set the extra flags on this node */
+	BLUEPRINTGRAPH_API void SetExtraFlags(int32 InFlags)
+	{
+		ExtraFlags = InFlags;
+	}
+
+protected:
+	/** Any extra flags that the function may need */
+	UPROPERTY()
+	int32 ExtraFlags;
 };
 
diff --git a/Engine/Source/Editor/BlueprintGraph/Private/EdGraphSchema_K2.cpp b/Engine/Source/Editor/BlueprintGraph/Private/EdGraphSchema_K2.cpp
--- a/Engine/Source/Editor/BlueprintGraph/Private/EdGraphSchema_K2.cpp
+++ b/Engine/Source/Editor/BlueprintGraph/Private/EdGraphSchema_K2.cpp
@@ -941,7 +941,7 @@ void UEdGraphSchema_K2::AddExtraFunctionFlags(const UEdGraph* CurrentGraph, int3
 	{
 		if (UK2Node_FunctionEntry* Node = Cast<UK2Node_FunctionEntry>(*It))
 		{
-			Node->ExtraFlags |= ExtraFlags;
+			Node->SetExtraFlags(Node->GetFunctionFlags() | ExtraFlags);
 		}
 	}
 }
@@ -2993,7 +2993,7 @@ void UEdGraphSchema_K2::CreateFunctionGraphTerminators(UEdGraph& Graph, UClass*
 	if (InterfaceToImplement)
 	{
 		// Add modifier flags from the declaration
-		EntryNode->ExtraFlags |= InterfaceToImplement->FunctionFlags & (FUNC_Const | FUNC_Static | FUNC_BlueprintPure);
+		EntryNode->SetExtraFlags(EntryNode->GetFunctionFlags() | InterfaceToImplement->FunctionFlags & (FUNC_Const | FUNC_Static | FUNC_BlueprintPure));
 
 		UK2Node* NextNode = EntryNode;
 		UEdGraphPin* NextExec = FindExecutionPin(*EntryNode, EGPD_Output);
@@ -4504,7 +4504,7 @@ bool UEdGraphSchema_K2::IsConstFunctionGraph( const UEdGraph* TestEdGraph, bool*
 					*bOutIsEnforcingConstCorrectness = EntryNode->bEnforceConstCorrectness;
 				}
 
-				return (EntryNode->ExtraFlags & FUNC_Const) != 0;
+				return (EntryNode->GetFunctionFlags() & FUNC_Const) != 0;
 			}
 		}
 	}
@@ -4536,7 +4536,7 @@ bool UEdGraphSchema_K2::IsStaticFunctionGraph( const UEdGraph* TestEdGraph ) con
 			UEdGraphNode* Node = *I;
 			if(auto EntryNode = Cast<UK2Node_FunctionEntry>(Node))
 			{
-				return (EntryNode->ExtraFlags & FUNC_Static) != 0;
+				return (EntryNode->GetFunctionFlags() & FUNC_Static) != 0;
 			}
 		}
 	}
diff --git a/Engine/Source/Editor/BlueprintGraph/Private/K2Node_FunctionEntry.cpp b/Engine/Source/Editor/BlueprintGraph/Private/K2Node_FunctionEntry.cpp
--- a/Engine/Source/Editor/BlueprintGraph/Private/K2Node_FunctionEntry.cpp
+++ b/Engine/Source/Editor/BlueprintGraph/Private/K2Node_FunctionEntry.cpp
@@ -272,6 +272,26 @@ FText UK2Node_FunctionEntry::GetTooltipText() const
 	return Super::GetTooltipText();
 }
 
+int32 UK2Node_FunctionEntry::GetFunctionFlags() const
+{
+	int32 ReturnFlags = 0;
+
+	UClass* ClassToLookup = SignatureClass;
+
+	if (SignatureClass && SignatureClass->ClassGeneratedBy)
+	{
+		UBlueprint* GeneratingBP = CastChecked<UBlueprint>(SignatureClass->ClassGeneratedBy);
+		ClassToLookup = GeneratingBP->SkeletonGeneratedClass;
+	}
+
+	UFunction* Function = FindField<UFunction>(ClassToLookup, SignatureName);
+	if (Function)
+	{
+		ReturnFlags = Function->FunctionFlags;
+	}
+	return ReturnFlags | ExtraFlags;
+}
+
 void UK2Node_FunctionEntry::ExpandNode(class FKismetCompilerContext& CompilerContext, UEdGraph* SourceGraph)
 {
 	Super::ExpandNode(CompilerContext, SourceGraph);
diff --git a/Engine/Source/Editor/Kismet/Private/BlueprintDetailsCustomization.cpp b/Engine/Source/Editor/Kismet/Private/BlueprintDetailsCustomization.cpp
--- a/Engine/Source/Editor/Kismet/Private/BlueprintDetailsCustomization.cpp
+++ b/Engine/Source/Editor/Kismet/Private/BlueprintDetailsCustomization.cpp
@@ -2889,8 +2889,10 @@ void FBlueprintGraphActionDetails::SetNetFlags( TWeakObjectPtr<UK2Node_EditableP
 
 			if (UK2Node_FunctionEntry* TypedEntryNode = Cast<UK2Node_FunctionEntry>(FunctionEntryNode.Get()))
 			{
-				TypedEntryNode->ExtraFlags &= ~FlagsToClear;
-				TypedEntryNode->ExtraFlags |= FlagsToSet;
+				int32 ExtraFlags = TypedEntryNode->GetFunctionFlags();
+				ExtraFlags &= ~FlagsToClear;
+				ExtraFlags |= FlagsToSet;
+				TypedEntryNode->SetExtraFlags(ExtraFlags);
 				bBlueprintModified = true;
 			}
 			if (UK2Node_CustomEvent * CustomEventNode = Cast<UK2Node_CustomEvent>(FunctionEntryNode.Get()))
@@ -3773,7 +3775,7 @@ FText FBlueprintGraphActionDetails::GetCurrentAccessSpecifierName() const
 	UK2Node_EditablePinBase * FunctionEntryNode = FunctionEntryNodePtr.Get();
 	if(UK2Node_FunctionEntry* EntryNode = Cast<UK2Node_FunctionEntry>(FunctionEntryNode))
 	{
-		AccessSpecifierFlag = FUNC_AccessSpecifiers & EntryNode->ExtraFlags;
+		AccessSpecifierFlag = FUNC_AccessSpecifiers & EntryNode->GetFunctionFlags();
 	}
 	else if(UK2Node_CustomEvent* CustomEventNode = Cast<UK2Node_CustomEvent>(FunctionEntryNode))
 	{
@@ -3820,8 +3822,10 @@ void FBlueprintGraphActionDetails::OnAccessSpecifierSelected( TSharedPtr<FAccess
 		const uint32 ClearAccessSpecifierMask = ~FUNC_AccessSpecifiers;
 		if(UK2Node_FunctionEntry* EntryNode = Cast<UK2Node_FunctionEntry>(FunctionEntryNode))
 		{
-			EntryNode->ExtraFlags &= ClearAccessSpecifierMask;
-			EntryNode->ExtraFlags |= SpecifierName->SpecifierFlag;
+			int32 ExtraFlags = EntryNode->GetFunctionFlags();
+			ExtraFlags &= ClearAccessSpecifierMask;
+			ExtraFlags |= SpecifierName->SpecifierFlag;
+			EntryNode->SetExtraFlags(ExtraFlags);
 		}
 		else if(UK2Node_Event* EventNode = Cast<UK2Node_Event>(FunctionEntryNode))
 		{
@@ -3912,7 +3916,7 @@ void FBlueprintGraphActionDetails::OnIsReliableReplicationFunctionModified(const
 		{
 			if (UK2Node_FunctionEntry* TypedEntryNode = Cast<UK2Node_FunctionEntry>(FunctionEntryNode))
 			{
-				TypedEntryNode->ExtraFlags |= FUNC_NetReliable;
+				TypedEntryNode->SetExtraFlags(TypedEntryNode->GetFunctionFlags() | FUNC_NetReliable);
 			}
 			if (UK2Node_CustomEvent * CustomEventNode = Cast<UK2Node_CustomEvent>(FunctionEntryNode))
 			{
@@ -3923,7 +3927,7 @@ void FBlueprintGraphActionDetails::OnIsReliableReplicationFunctionModified(const
 		{
 			if (UK2Node_FunctionEntry* TypedEntryNode = Cast<UK2Node_FunctionEntry>(FunctionEntryNode))
 			{
-				TypedEntryNode->ExtraFlags &= ~FUNC_NetReliable;
+				TypedEntryNode->SetExtraFlags(TypedEntryNode->GetFunctionFlags() & ~FUNC_NetReliable);
 			}
 			if (UK2Node_CustomEvent * CustomEventNode = Cast<UK2Node_CustomEvent>(FunctionEntryNode))
 			{
@@ -3981,8 +3985,8 @@ void FBlueprintGraphActionDetails::OnIsPureFunctionModified( const ECheckBoxStat
 		Function->Modify();
 
 		//set flags on function entry node also
-		EntryNode->ExtraFlags	^= FUNC_BlueprintPure;
 		Function->FunctionFlags ^= FUNC_BlueprintPure;
+		EntryNode->SetExtraFlags(Function->FunctionFlags);
 		OnParamsChanged(FunctionEntryNode);
 	}
 }
@@ -3995,7 +3999,7 @@ ECheckBoxState FBlueprintGraphActionDetails::GetIsPureFunction() const
 	{
 		return ECheckBoxState::Undetermined;
 	}
-	return (EntryNode->ExtraFlags & FUNC_BlueprintPure) ? ECheckBoxState::Checked :  ECheckBoxState::Unchecked;
+	return (EntryNode->GetFunctionFlags() & FUNC_BlueprintPure) ? ECheckBoxState::Checked :  ECheckBoxState::Unchecked;
 }
 
 bool FBlueprintGraphActionDetails::IsConstFunctionVisible() const
@@ -4025,8 +4029,8 @@ void FBlueprintGraphActionDetails::OnIsConstFunctionModified( const ECheckBoxSta
 		Function->Modify();
 
 		//set flags on function entry node also
-		EntryNode->ExtraFlags	^= FUNC_Const;
 		Function->FunctionFlags ^= FUNC_Const;
+		EntryNode->SetExtraFlags(Function->FunctionFlags);
 		OnParamsChanged(FunctionEntryNode);
 	}
 }
@@ -4039,7 +4043,7 @@ ECheckBoxState FBlueprintGraphActionDetails::GetIsConstFunction() const
 	{
 		return ECheckBoxState::Undetermined;
 	}
-	return (EntryNode->ExtraFlags & FUNC_Const) ? ECheckBoxState::Checked :  ECheckBoxState::Unchecked;
+	return (EntryNode->GetFunctionFlags() & FUNC_Const) ? ECheckBoxState::Checked :  ECheckBoxState::Unchecked;
 }
 
 FReply FBaseBlueprintGraphActionDetails::OnAddNewInputClicked()
diff --git a/Engine/Source/Editor/Kismet/Private/SMyBlueprint.cpp b/Engine/Source/Editor/Kismet/Private/SMyBlueprint.cpp
--- a/Engine/Source/Editor/Kismet/Private/SMyBlueprint.cpp
+++ b/Engine/Source/Editor/Kismet/Private/SMyBlueprint.cpp
@@ -1450,7 +1450,7 @@ FReply SMyBlueprint::OnActionDragged( const TArray< TSharedPtr<FEdGraphSchemaAct
 						if (UK2Node_FunctionEntry* Node = Cast<UK2Node_FunctionEntry>(*It))
 						{
 							// See whether this node is a blueprint callable function
-							if (Node->ExtraFlags & (FUNC_BlueprintCallable|FUNC_BlueprintPure))
+							if (Node->GetFunctionFlags() & (FUNC_BlueprintCallable|FUNC_BlueprintPure))
 							{
 								bIsBlueprintCallableFunction = true;
 							}
diff --git a/Engine/Source/Editor/KismetCompiler/Private/KismetCompiler.cpp b/Engine/Source/Editor/KismetCompiler/Private/KismetCompiler.cpp
--- a/Engine/Source/Editor/KismetCompiler/Private/KismetCompiler.cpp
+++ b/Engine/Source/Editor/KismetCompiler/Private/KismetCompiler.cpp
@@ -1174,8 +1174,9 @@ void FKismetCompilerContext::PrecompileFunction(FKismetFunctionContext& Context)
 
 		// Determine if this is a new function or if it overrides a parent function
 			//@TODO: Does not support multiple overloads for a parent virtual function
-		UFunction* ParentFunction = Blueprint->ParentClass->FindFunctionByName(NewFunctionName);
-
+		UClass* SuperClass = Context.NewClass->GetSuperClass();
+		UFunction* ParentFunction = Context.NewClass->GetSuperClass()->FindFunctionByName(NewFunctionName);
+		
 		const FString NewFunctionNameString = NewFunctionName.ToString();
 		if (CreatedFunctionNames.Contains(NewFunctionNameString))
 		{
@@ -1246,7 +1247,7 @@ void FKismetCompilerContext::PrecompileFunction(FKismetFunctionContext& Context)
 		// Inherit extra flags from the entry node
 		if (Context.EntryPoint)
 		{
-			Context.Function->FunctionFlags |= Context.EntryPoint->ExtraFlags;
+			Context.Function->FunctionFlags |= Context.EntryPoint->GetExtraFlags(Context);
 		}
 
 		// First try to get the overriden function from the super class
@@ -1265,7 +1266,7 @@ void FKismetCompilerContext::PrecompileFunction(FKismetFunctionContext& Context)
 		// Inherit flags and validate against overridden function if it exists
 		if (OverridenFunction)
 		{
-			Context.Function->FunctionFlags |= (OverridenFunction->FunctionFlags & (FUNC_FuncInherit | FUNC_Public | FUNC_Protected | FUNC_Private));
+			Context.Function->FunctionFlags |= (OverridenFunction->FunctionFlags & (FUNC_FuncInherit | FUNC_Public | FUNC_Protected | FUNC_Private | FUNC_BlueprintPure));
 
 			if ((Context.Function->FunctionFlags & FUNC_AccessSpecifiers) != (OverridenFunction->FunctionFlags & FUNC_AccessSpecifiers))
 			{
@@ -1354,7 +1355,7 @@ void FKismetCompilerContext::PrecompileFunction(FKismetFunctionContext& Context)
 		CreateLocalVariablesForFunction(Context);
 
 		//Validate AccessSpecifier
-		const uint32 AccessSpecifierFlag = FUNC_AccessSpecifiers & Context.EntryPoint->ExtraFlags;
+		const uint32 AccessSpecifierFlag = FUNC_AccessSpecifiers & Context.EntryPoint->GetExtraFlags(Context);
 		const bool bAcceptedAccessSpecifier = 
 			(0 == AccessSpecifierFlag) || (FUNC_Public == AccessSpecifierFlag) || (FUNC_Protected == AccessSpecifierFlag) || (FUNC_Private == AccessSpecifierFlag);
 		if(!bAcceptedAccessSpecifier)
@@ -2187,7 +2188,7 @@ void FKismetCompilerContext::CreateFunctionStubForEvent(UK2Node_Event* SrcEventN
 
 	if(!SrcEventNode->bOverrideFunction && SrcEventNode->IsUsedByAuthorityOnlyDelegate())
 	{
-		EntryNode->ExtraFlags |= FUNC_BlueprintAuthorityOnly;
+		EntryNode->SetExtraFlags(EntryNode->GetFunctionFlags() | FUNC_BlueprintAuthorityOnly);
 	}
 
 	// If this is a customizable event, make sure to copy over the user defined pins
diff --git a/Engine/Source/Editor/UnrealEd/Private/Kismet2/BlueprintEditorUtils.cpp b/Engine/Source/Editor/UnrealEd/Private/Kismet2/BlueprintEditorUtils.cpp
--- a/Engine/Source/Editor/UnrealEd/Private/Kismet2/BlueprintEditorUtils.cpp
+++ b/Engine/Source/Editor/UnrealEd/Private/Kismet2/BlueprintEditorUtils.cpp
@@ -333,6 +333,7 @@ void FBasePinChangeHelper::Broadcast(UBlueprint* InBlueprint, UK2Node_EditablePi
 		const UClass* SignatureClass = Func
 			? Func->GetOwnerClass()
 			: (UClass*)(FunctionDefNode ? FunctionDefNode->SignatureClass : nullptr);
+
 		// Reconstruct all function call sites that call this function (in open blueprints)
 		for (TObjectIterator<UK2Node_CallFunction> It(RF_Transient); It; ++It)
 		{
@@ -7150,7 +7151,7 @@ void FBlueprintEditorUtils::UpdateOldPureFunctions( UBlueprint* Blueprint )
 						{
 							UK2Node_FunctionEntry* Entry = EntryNodes[0];
 							Entry->Modify();
-							Entry->ExtraFlags |= FUNC_BlueprintPure;
+							Entry->SetExtraFlags(Entry->GetFunctionFlags() | FUNC_BlueprintPure);
 						}
 					}
 				}
