Typo fix.

diff --git a/Engine/Source/Editor/Kismet/Private/BlueprintDetailsCustomization.cpp b/Engine/Source/Editor/Kismet/Private/BlueprintDetailsCustomization.cpp
--- a/Engine/Source/Editor/Kismet/Private/BlueprintDetailsCustomization.cpp
+++ b/Engine/Source/Editor/Kismet/Private/BlueprintDetailsCustomization.cpp
@@ -390,7 +390,7 @@ void FBlueprintVarActionDetails::CustomizeDetails( IDetailLayoutBuilder& DetailL
 		ConfigTooltipArgs.Add(TEXT("ConfigPath"), FText::FromString(OwnerClass->GetDefaultConfigFilename()));
 		ConfigTooltipArgs.Add(TEXT("ConfigSection"), FText::FromString(OwnerClass->GetPathName()));
 	}
-	const FText LocalisedTooltip = FText::Format(LOCTEXT("VariableExposeToConfig_Tooltip", "Should this variable read it's default value from a config file if it is present?\r\n\r\nThis is used for customising variable default values and behavior between different projects and configurations.\r\n\r\nConfig file [{ConfigPath}]\r\nConfig section [{ConfigSection}]"), ConfigTooltipArgs); 
+	const FText LocalisedTooltip = FText::Format(LOCTEXT("VariableExposeToConfig_Tooltip", "Should this variable read its default value from a config file if it is present?\r\n\r\nThis is used for customising variable default values and behavior between different projects and configurations.\r\n\r\nConfig file [{ConfigPath}]\r\nConfig section [{ConfigSection}]"), ConfigTooltipArgs); 
 
 	TSharedPtr<SToolTip> ExposeToConfigTooltip = IDocumentation::Get()->CreateToolTip(LocalisedTooltip, NULL, DocLink, TEXT("ExposeToConfig"));
 
