Fixed crash typing long names into blueprint function input and output pins

diff --git a/Engine/Source/Editor/Kismet/Private/BlueprintDetailsCustomization.cpp b/Engine/Source/Editor/Kismet/Private/BlueprintDetailsCustomization.cpp
--- a/Engine/Source/Editor/Kismet/Private/BlueprintDetailsCustomization.cpp
+++ b/Engine/Source/Editor/Kismet/Private/BlueprintDetailsCustomization.cpp
@@ -3359,6 +3359,12 @@ bool FBaseBlueprintGraphActionDetails::OnVerifyPinRename(UK2Node_EditablePinBase
 		return true;
 	}
 
+	if( InNewName.Len() > NAME_SIZE )
+	{
+		OutErrorMessage = FText::Format( LOCTEXT("PinNameTooLong", "The name you entered is too long. Names must be less than {0} characters"), FText::AsNumber( NAME_SIZE ) );
+		return false;
+	}
+
 	if (InTargetNode)
 	{
 		// Check if the name conflicts with any of the other internal UFunction's property names (local variables and parameters).
