FFindInBlueprintSearchManager now caches the asset class pointers it's already looked up

FFindInBlueprintSearchManager::OnAssetAdded was showing up as a hot-spot due to the number of times it was calling FAssetData::GetClass.

It now caches these look-ups so that it can avoid calling FindObject most of the time. This has quite the impact on asset registry processing time when dealing with large projects.

diff --git a/Engine/Source/Editor/Kismet/Private/FindInBlueprintManager.cpp b/Engine/Source/Editor/Kismet/Private/FindInBlueprintManager.cpp
--- a/Engine/Source/Editor/Kismet/Private/FindInBlueprintManager.cpp
+++ b/Engine/Source/Editor/Kismet/Private/FindInBlueprintManager.cpp
@@ -6,6 +6,7 @@
 #include "FindInBlueprints.h"
 #include "Editor/UnrealEd/Public/Kismet2/BlueprintEditorUtils.h"
 #include "AssetRegistryModule.h"
+#include "HotReloadInterface.h"
 
 #include "JsonUtilities.h"
 #include "SNotificationList.h"
@@ -904,6 +905,12 @@ FFindInBlueprintSearchManager::~FFindInBlueprintSearchManager()
 	FCoreUObjectDelegates::PreGarbageCollect.RemoveAll(this);
 	FCoreUObjectDelegates::PostGarbageCollect.RemoveAll(this);
 	FCoreUObjectDelegates::OnAssetLoaded.RemoveAll(this);
+
+	if(FModuleManager::Get().IsModuleLoaded("HotReload"))
+	{
+		IHotReloadInterface& HotReloadSupport = FModuleManager::GetModuleChecked<IHotReloadInterface>("HotReload");
+		HotReloadSupport.OnHotReload().RemoveAll(this);
+	}
 }
 
 void FFindInBlueprintSearchManager::Initialize()
@@ -915,6 +922,10 @@ void FFindInBlueprintSearchManager::Initialize()
 	FCoreUObjectDelegates::PreGarbageCollect.AddRaw(this, &FFindInBlueprintSearchManager::PauseFindInBlueprintSearch);
 	FCoreUObjectDelegates::PostGarbageCollect.AddRaw(this, &FFindInBlueprintSearchManager::UnpauseFindInBlueprintSearch);
 	FCoreUObjectDelegates::OnAssetLoaded.AddRaw(this, &FFindInBlueprintSearchManager::OnAssetLoaded);
+	
+	// Register to be notified of hot reloads
+	IHotReloadInterface& HotReloadSupport = FModuleManager::LoadModuleChecked<IHotReloadInterface>("HotReload");
+	HotReloadSupport.OnHotReload().AddRaw(this, &FFindInBlueprintSearchManager::OnHotReload);
 
 	if(!GIsSavingPackage)
 	{
@@ -925,7 +936,24 @@ void FFindInBlueprintSearchManager::Initialize()
 
 void FFindInBlueprintSearchManager::OnAssetAdded(const FAssetData& InAssetData)
 {
-	if(InAssetData.GetClass()->IsChildOf(UBlueprint::StaticClass()) || InAssetData.GetClass()->IsChildOf(UWorld::StaticClass()))
+	const UClass* AssetClass = nullptr;
+	{
+		const UClass** FoundClass = CachedAssetClasses.Find(InAssetData.AssetClass);
+		if (FoundClass)
+		{
+			AssetClass = *FoundClass;
+		}
+		else
+		{
+			AssetClass = InAssetData.GetClass();
+			if (AssetClass)
+			{
+				CachedAssetClasses.Add(InAssetData.AssetClass, AssetClass);
+			}
+		}
+	}
+
+	if(AssetClass && AssetClass->IsChildOf(UBlueprint::StaticClass()) || AssetClass->IsChildOf(UWorld::StaticClass()))
 	{
 		FString BlueprintPackagePath = FPaths::GetPath(InAssetData.ObjectPath.ToString()) / FPaths::GetBaseFilename(InAssetData.ObjectPath.ToString());
 
@@ -935,7 +963,7 @@ void FFindInBlueprintSearchManager::OnAssetAdded(const FAssetData& InAssetData)
 		{
 			if(InAssetData.IsAssetLoaded())
 			{
-				if(InAssetData.GetClass()->IsChildOf(UBlueprint::StaticClass()))
+				if(AssetClass->IsChildOf(UBlueprint::StaticClass()))
 				{
 					if(UBlueprint* BlueprintAsset = Cast<UBlueprint>(InAssetData.GetAsset()))
 					{
@@ -943,7 +971,7 @@ void FFindInBlueprintSearchManager::OnAssetAdded(const FAssetData& InAssetData)
 						AddOrUpdateBlueprintSearchMetadata(BlueprintAsset);
 					}
 				}
-				else if(InAssetData.GetClass()->IsChildOf(UWorld::StaticClass()))
+				else if(AssetClass->IsChildOf(UWorld::StaticClass()))
 				{
 					UWorld* WorldAsset = Cast<UWorld>(InAssetData.GetAsset());
 					if(WorldAsset->PersistentLevel)
@@ -1055,6 +1083,11 @@ void FFindInBlueprintSearchManager::OnAssetLoaded(UObject* InAsset)
 	}
 }
 
+void FFindInBlueprintSearchManager::OnHotReload(bool bWasTriggeredAutomatically)
+{
+	CachedAssetClasses.Reset();
+}
+
 FString FFindInBlueprintSearchManager::GatherBlueprintSearchMetadata(const UBlueprint* Blueprint)
 {	
 	FTemporarilyUseFriendlyNodeTitles TemporarilyUseFriendlyNodeTitles;
diff --git a/Engine/Source/Editor/Kismet/Public/FindInBlueprintManager.h b/Engine/Source/Editor/Kismet/Public/FindInBlueprintManager.h
--- a/Engine/Source/Editor/Kismet/Public/FindInBlueprintManager.h
+++ b/Engine/Source/Editor/Kismet/Public/FindInBlueprintManager.h
@@ -287,6 +287,9 @@ private:
 	/** Callback hook from the Asset Registry when an asset is loaded */
 	void OnAssetLoaded(class UObject* InAsset);
 
+	/** Callback hook from the Hot Reload manager that indicates that a module has been hot-reloaded */
+	void OnHotReload(bool bWasTriggeredAutomatically);
+
 	/** Helper to gathers the Blueprint's search metadata */
 	FString GatherBlueprintSearchMetadata(const UBlueprint* Blueprint);
 
@@ -348,4 +351,7 @@ protected:
 
 	/** Tickable object that does the caching of uncached Blueprints at a rate of once per tick */
 	class FCacheAllBlueprintsTickableObject* CachingObject;
+
+	/** Mapping between a class name and its UClass instance - used for faster look up in FFindInBlueprintSearchManager::OnAssetAdded */
+	TMap<FName, const UClass*> CachedAssetClasses;
 };
