CIS fix

diff --git a/Engine/Source/Editor/Kismet/Private/FindInBlueprintManager.cpp b/Engine/Source/Editor/Kismet/Private/FindInBlueprintManager.cpp
--- a/Engine/Source/Editor/Kismet/Private/FindInBlueprintManager.cpp
+++ b/Engine/Source/Editor/Kismet/Private/FindInBlueprintManager.cpp
@@ -953,7 +953,7 @@ void FFindInBlueprintSearchManager::OnAssetAdded(const FAssetData& InAssetData)
 		}
 	}
 
-	if(AssetClass && AssetClass->IsChildOf(UBlueprint::StaticClass()) || AssetClass->IsChildOf(UWorld::StaticClass()))
+	if(AssetClass && (AssetClass->IsChildOf(UBlueprint::StaticClass()) || AssetClass->IsChildOf(UWorld::StaticClass())))
 	{
 		FString BlueprintPackagePath = FPaths::GetPath(InAssetData.ObjectPath.ToString()) / FPaths::GetBaseFilename(InAssetData.ObjectPath.ToString());
 
