Made all paths in the FiB manager be stored as an FName

This avoids some FName -> FString conversion (and FString copying) while responding to asset registry events, and allows the asset registry to process its results faster.

Also changed some places that were copying FSearchData to move it instead (eg, when putting it into the search array/map).

diff --git a/Engine/Source/Editor/Kismet/Private/FindInBlueprintManager.cpp b/Engine/Source/Editor/Kismet/Private/FindInBlueprintManager.cpp
--- a/Engine/Source/Editor/Kismet/Private/FindInBlueprintManager.cpp
+++ b/Engine/Source/Editor/Kismet/Private/FindInBlueprintManager.cpp
@@ -667,7 +667,7 @@ class FCacheAllBlueprintsTickableObject
 
 public:
 
-	FCacheAllBlueprintsTickableObject(TArray<FString>& InUncachedBlueprints, bool bInCheckOutAndSave)
+	FCacheAllBlueprintsTickableObject(TArray<FName>& InUncachedBlueprints, bool bInCheckOutAndSave)
 		: TickCacheIndex(0)
 		, UncachedBlueprints(InUncachedBlueprints)
 		, bIsStarted(false)
@@ -701,13 +701,13 @@ public:
 	}
 
 	/** Returns the name of the current Blueprint being cached */
-	FString GetCurrentCacheBlueprintName() const
+	FName GetCurrentCacheBlueprintName() const
 	{
 		if(UncachedBlueprints.Num() && TickCacheIndex >= 0)
 		{
 			return UncachedBlueprints[TickCacheIndex];
 		}
-		return FString();
+		return NAME_None;
 	}
 
 	/** Returns the progress as a percent */
@@ -764,7 +764,7 @@ public:
 		else
 		{
 			FAssetRegistryModule* AssetRegistryModule = &FModuleManager::LoadModuleChecked<FAssetRegistryModule>(TEXT("AssetRegistry"));
-			FAssetData AssetData = AssetRegistryModule->Get().GetAssetByObjectPath(FName(*UncachedBlueprints[TickCacheIndex]));
+			FAssetData AssetData = AssetRegistryModule->Get().GetAssetByObjectPath(UncachedBlueprints[TickCacheIndex]);
 
 			// Cache whether the Blueprint is already loaded
 			bool bIsLoaded = AssetData.IsAssetLoaded();
@@ -859,13 +859,13 @@ private:
 	int32 TickCacheIndex;
 
 	/** The list of uncached Blueprints that are in the process of being cached */
-	TArray<FString> UncachedBlueprints;
+	TArray<FName> UncachedBlueprints;
 
 	/** Notification that appears and details progress */
 	TWeakPtr<SNotificationItem> ProgressNotification;
 
 	/** List of Blueprints that failed to be saved */
-	TArray<FString> FailedToCacheList;
+	TArray<FName> FailedToCacheList;
 
 	/** TRUE if the caching process is started */
 	bool bIsStarted;
@@ -955,10 +955,8 @@ void FFindInBlueprintSearchManager::OnAssetAdded(const FAssetData& InAssetData)
 
 	if(AssetClass && (AssetClass->IsChildOf(UBlueprint::StaticClass()) || AssetClass->IsChildOf(UWorld::StaticClass())))
 	{
-		FString BlueprintPackagePath = FPaths::GetPath(InAssetData.ObjectPath.ToString()) / FPaths::GetBaseFilename(InAssetData.ObjectPath.ToString());
-
 		// Confirm that the Blueprint has not been added already, this can occur during duplication of Blueprints.
-		int32* IndexPtr = SearchMap.Find(BlueprintPackagePath);
+		int32* IndexPtr = SearchMap.Find(InAssetData.ObjectPath);
 		if(!IndexPtr)
 		{
 			if(InAssetData.IsAssetLoaded())
@@ -992,33 +990,35 @@ void FFindInBlueprintSearchManager::OnAssetAdded(const FAssetData& InAssetData)
 			{
 				FSearchData NewSearchData;
 
-				NewSearchData.BlueprintPath = InAssetData.ObjectPath.ToString();
+				NewSearchData.BlueprintPath = InAssetData.ObjectPath;
 				NewSearchData.bMarkedForDeletion = false;
 
 				// Since the asset was not loaded, pull out the searchable data stored in the asset
 				NewSearchData.Value = *FiBSearchData;
-				AddSearchDataToDatabase(NewSearchData);
+				AddSearchDataToDatabase(MoveTemp(NewSearchData));
 			}
 			else
 			{
 				// The asset is uncached, we will want to inform the user that this is the case
-				UncachedBlueprints.Add(InAssetData.ObjectPath.ToString());
+				UncachedBlueprints.Add(InAssetData.ObjectPath);
 			}
 		}
 	}
 }
 
-int32 FFindInBlueprintSearchManager::AddSearchDataToDatabase(FSearchData& InSearchData)
+int32 FFindInBlueprintSearchManager::AddSearchDataToDatabase(FSearchData InSearchData)
 {
-	int32 ArrayIndex = SearchArray.Add(InSearchData);
+	FName BlueprintPath = InSearchData.BlueprintPath; // Copy before we move the data into the array
+
+	int32 ArrayIndex = SearchArray.Add(MoveTemp(InSearchData));
 
 	// Add the asset file path to the map along with the index into the array
-	SearchMap.Add(InSearchData.BlueprintPath, ArrayIndex);
+	SearchMap.Add(BlueprintPath, ArrayIndex);
 
 	return ArrayIndex;
 }
 
-void FFindInBlueprintSearchManager::RemoveBlueprintByPath(FString InPath)
+void FFindInBlueprintSearchManager::RemoveBlueprintByPath(FName InPath)
 {
 	int32* SearchIdx = SearchMap.Find(InPath);
 
@@ -1031,7 +1031,7 @@ void FFindInBlueprintSearchManager::OnAssetRemoved(const class FAssetData& InAss
 {
 	if(InAssetData.IsAssetLoaded())
 	{
-		RemoveBlueprintByPath(InAssetData.ObjectPath.ToString());
+		RemoveBlueprintByPath(InAssetData.ObjectPath);
 	}
 }
 
@@ -1040,14 +1040,14 @@ void FFindInBlueprintSearchManager::OnAssetRenamed(const class FAssetData& InAss
 	// Renaming removes the item from the manager, it will be re-added in the OnAssetAdded event under the new name.
 	if(InAssetData.IsAssetLoaded())
 	{
-		RemoveBlueprintByPath(InOldName);
+		RemoveBlueprintByPath(FName(*InOldName));
 	}
 }
 
 void FFindInBlueprintSearchManager::OnAssetLoaded(UObject* InAsset)
 {
 	UBlueprint* BlueprintAsset = nullptr;
-	FString BlueprintPath = InAsset->GetPathName();
+	FName BlueprintPath = *InAsset->GetPathName();
 
 	if(UWorld* WorldAsset = Cast<UWorld>(InAsset))
 	{
@@ -1056,7 +1056,7 @@ void FFindInBlueprintSearchManager::OnAssetLoaded(UObject* InAsset)
 			BlueprintAsset = Cast<UBlueprint>((UObject*)WorldAsset->PersistentLevel->GetLevelScriptBlueprint(true));
 			if(BlueprintAsset)
 			{
-				BlueprintPath = BlueprintAsset->GetPathName();
+				BlueprintPath = *BlueprintAsset->GetPathName();
 			}
 		}
 	}
@@ -1076,8 +1076,8 @@ void FFindInBlueprintSearchManager::OnAssetLoaded(UObject* InAsset)
 		if(IndexPtr)
 		{
 			// That index should never have a Blueprint already, but if it does, it should be the same Blueprint!
-			ensureMsgf(!SearchArray[*IndexPtr].Blueprint.IsValid() || SearchArray[*IndexPtr].Blueprint == BlueprintAsset, TEXT("Blueprint in database has path %s and is being stomped by %s"), *(SearchArray[*IndexPtr].BlueprintPath), *BlueprintPath);
-			ensureMsgf(!SearchArray[*IndexPtr].Blueprint.IsValid() || SearchArray[*IndexPtr].BlueprintPath == BlueprintPath, TEXT("Blueprint in database has path %s and is being stomped by %s"), *(SearchArray[*IndexPtr].BlueprintPath), *BlueprintPath);
+			ensureMsgf(!SearchArray[*IndexPtr].Blueprint.IsValid() || SearchArray[*IndexPtr].Blueprint == BlueprintAsset, TEXT("Blueprint in database has path %s and is being stomped by %s"), *(SearchArray[*IndexPtr].BlueprintPath.ToString()), *BlueprintPath.ToString());
+			ensureMsgf(!SearchArray[*IndexPtr].Blueprint.IsValid() || SearchArray[*IndexPtr].BlueprintPath == BlueprintPath, TEXT("Blueprint in database has path %s and is being stomped by %s"), *(SearchArray[*IndexPtr].BlueprintPath.ToString()), *BlueprintPath.ToString());
 			SearchArray[*IndexPtr].Blueprint = BlueprintAsset;
 		}
 	}
@@ -1166,17 +1166,17 @@ void FFindInBlueprintSearchManager::AddOrUpdateBlueprintSearchMetadata(UBlueprin
 	// Allow only one thread modify the search data at a time
 	FScopeLock ScopeLock(&SafeModifyCacheCriticalSection);
 
-	FString BlueprintPath;
+	FName BlueprintPath;
 	if(FBlueprintEditorUtils::IsLevelScriptBlueprint(InBlueprint))
 	{
 		if(UWorld* World = InBlueprint->GetTypedOuter<UWorld>())
 		{
-			BlueprintPath = World->GetPathName();
+			BlueprintPath = *World->GetPathName();
 		}
 	}
 	else
 	{
-		BlueprintPath = InBlueprint->GetPathName();
+		BlueprintPath = *InBlueprint->GetPathName();
 	}
 
 	int32* IndexPtr = SearchMap.Find(BlueprintPath);
@@ -1187,7 +1187,7 @@ void FFindInBlueprintSearchManager::AddOrUpdateBlueprintSearchMetadata(UBlueprin
 		SearchData.Blueprint = InBlueprint;
 		SearchData.BlueprintPath = BlueprintPath;
 
-		Index = AddSearchDataToDatabase(SearchData);
+		Index = AddSearchDataToDatabase(MoveTemp(SearchData));
 	}
 	else
 	{
@@ -1302,12 +1302,12 @@ FString FFindInBlueprintSearchManager::QuerySingleBlueprint(UBlueprint* InBluepr
 		AddOrUpdateBlueprintSearchMetadata(InBlueprint, true);
 	}
 
-	FString Key = InBlueprint->GetPathName();
+	FName Key = *InBlueprint->GetPathName();
 	if (ULevel* LevelOuter = Cast<ULevel>(InBlueprint->GetOuter()))
 	{
 		if (UWorld* WorldOuter = Cast<UWorld>(LevelOuter->GetOuter()))
 		{
-			Key = WorldOuter->GetPathName();
+			Key = *WorldOuter->GetPathName();
 		}
 	}
 	int32* ArrayIdx = SearchMap.Find(Key);
@@ -1351,21 +1351,21 @@ void FFindInBlueprintSearchManager::CleanCache()
 	// *NOTE* SaveCache is a thread safe operation by design, all searching threads are paused during the operation so there is no critical section locking
 
 	// We need to cache where the active queries are so that we can put them back in a safe and expected position
-	TMap< const FStreamSearch*, FString > CacheQueries;
+	TMap< const FStreamSearch*, FName > CacheQueries;
 	for( auto It = ActiveSearchQueries.CreateIterator() ; It ; ++It )
 	{
 	 	const FStreamSearch* ActiveSearch = It.Key();
 	 	check(ActiveSearch);
 	 	{
-			FSearchData searchData;
-	 		ContinueSearchQuery(ActiveSearch, searchData);
+			FSearchData SearchData;
+	 		ContinueSearchQuery(ActiveSearch, SearchData);
 
-			FString CachePath = searchData.BlueprintPath;
+			FName CachePath = SearchData.BlueprintPath;
 	 		CacheQueries.Add(ActiveSearch, CachePath);
 	 	}
 	}
 
-	TMap<FString, int> NewSearchMap;
+	TMap<FName, int32> NewSearchMap;
 	TArray<FSearchData> NewSearchArray;
 
 	for(auto& SearchValuePair : SearchMap)
@@ -1388,7 +1388,7 @@ void FFindInBlueprintSearchManager::CleanCache()
 				AssetRegistryModule = &FModuleManager::LoadModuleChecked<FAssetRegistryModule>(TEXT("AssetRegistry"));
 
 				// The asset was not user deleted, so this should usually find the asset. New levels can be deleted if they were not saved
-				FAssetData AssetData = AssetRegistryModule->Get().GetAssetByObjectPath(*SearchArray[SearchValuePair.Value].BlueprintPath);
+				FAssetData AssetData = AssetRegistryModule->Get().GetAssetByObjectPath(SearchArray[SearchValuePair.Value].BlueprintPath);
 				if(AssetData.IsValid())
 				{
 					if(const FString* FiBSearchData = AssetData.TagsAndValues.Find("FiB"))
@@ -1409,8 +1409,8 @@ void FFindInBlueprintSearchManager::CleanCache()
 	for( auto& CacheQuery : CacheQueries )
 	{
 	 	int32 NewMappedIndex = 0;
-	 	// If the CachePath has a length, it means it is valid, otherwise we are at the end and there are no more search results, leave the query there so it can handle shutdown on it's own
-	 	if(CacheQuery.Value.Len() > 0)
+	 	// Is the CachePath is valid? Otherwise we are at the end and there are no more search results, leave the query there so it can handle shutdown on it's own
+	 	if(!CacheQuery.Value.IsNone())
 	 	{
 	 		int32* NewMappedIndexPtr = SearchMap.Find(CacheQuery.Value);
 	 		check(NewMappedIndexPtr);
@@ -1479,7 +1479,13 @@ void FFindInBlueprintSearchManager::OnCacheAllUncachedBlueprints(bool bInSourceC
 
 	if(bInSourceControlActive && bCheckoutAndSave)
 	{
-		FEditorFileUtils::CheckoutPackages(UncachedBlueprints);
+		TArray<FString> UncachedBlueprintStrings;
+		UncachedBlueprintStrings.Reserve(UncachedBlueprints.Num());
+		for (const FName& UncachedBlueprint : UncachedBlueprints)
+		{
+			UncachedBlueprintStrings.Add(UncachedBlueprint.ToString());
+		}
+		FEditorFileUtils::CheckoutPackages(UncachedBlueprintStrings);
 	}
 
 	// Start the cache process.
@@ -1543,9 +1549,9 @@ int32 FFindInBlueprintSearchManager::GetCurrentCacheIndex() const
 	return CachingIndex;
 }
 
-FString FFindInBlueprintSearchManager::GetCurrentCacheBlueprintName() const
+FName FFindInBlueprintSearchManager::GetCurrentCacheBlueprintName() const
 {
-	FString CachingBPName;
+	FName CachingBPName;
 	if(CachingObject)
 	{
 		CachingBPName = CachingObject->GetCurrentCacheBlueprintName();
@@ -1566,7 +1572,7 @@ float  FFindInBlueprintSearchManager::GetCacheProgress() const
 	return ReturnCacheValue;
 }
 
-void FFindInBlueprintSearchManager::FinishedCachingBlueprints(int32 InNumberCached, TArray<FString>& InFailedToCacheList)
+void FFindInBlueprintSearchManager::FinishedCachingBlueprints(int32 InNumberCached, TArray<FName>& InFailedToCacheList)
 {
 	// Multiple threads could be adding to this at the same time
 	FScopeLock ScopeLock(&SafeModifyCacheCriticalSection);
diff --git a/Engine/Source/Editor/Kismet/Private/FindInBlueprints.cpp b/Engine/Source/Editor/Kismet/Private/FindInBlueprints.cpp
--- a/Engine/Source/Editor/Kismet/Private/FindInBlueprints.cpp
+++ b/Engine/Source/Editor/Kismet/Private/FindInBlueprints.cpp
@@ -1069,10 +1069,10 @@ void SFindInBlueprints::ConditionallyAddCacheBar()
 		{
 			// Create a single string of all the Blueprint paths that failed to cache, on separate lines
 			FString PackageList;
-			TArray<FString> FailedToCacheList = FFindInBlueprintSearchManager::Get().GetFailedToCachePathList();
-			for (FString Package : FailedToCacheList)
+			TArray<FName> FailedToCacheList = FFindInBlueprintSearchManager::Get().GetFailedToCachePathList();
+			for (FName Package : FailedToCacheList)
 			{
-				PackageList += Package + TEXT("\n");
+				PackageList += Package.ToString() + TEXT("\n");
 			}
 
 			// Lambda to put together the popup menu detailing the failed to cache paths
@@ -1697,7 +1697,7 @@ FText SFindInBlueprints::GetUncachedBlueprintWarningText() const
 
 FText SFindInBlueprints::GetCurrentCacheBlueprintName() const
 {
-	return FText::FromString(FFindInBlueprintSearchManager::Get().GetCurrentCacheBlueprintName());
+	return FText::FromName(FFindInBlueprintSearchManager::Get().GetCurrentCacheBlueprintName());
 }
 
 void SFindInBlueprints::OnCacheComplete()
@@ -1896,7 +1896,7 @@ void FExtractionTask::DoWork()
 {
 	for (const auto& QueryResult: QueryResultBatch)
 	{
-		FSearchResult BlueprintCategory = FSearchResult(new FFindInBlueprintsResult(FText::FromString(QueryResult.BlueprintPath)));
+		FSearchResult BlueprintCategory = FSearchResult(new FFindInBlueprintsResult(FText::FromName(QueryResult.BlueprintPath)));
 		FindInBlueprintsHelpers::Extract(QueryResult.Value, *Tokens, BlueprintCategory);
 
 		// If there are children, add the item to the search results
diff --git a/Engine/Source/Editor/Kismet/Public/FindInBlueprintManager.h b/Engine/Source/Editor/Kismet/Public/FindInBlueprintManager.h
--- a/Engine/Source/Editor/Kismet/Public/FindInBlueprintManager.h
+++ b/Engine/Source/Editor/Kismet/Public/FindInBlueprintManager.h
@@ -75,7 +75,7 @@ struct FSearchData
 	TWeakObjectPtr<UBlueprint> Blueprint;
 
 	/** The full Blueprint path this search data is associated with */
-	FString BlueprintPath;
+	FName BlueprintPath;
 
 	/** Search data block for the Blueprint */
 	FString Value;
@@ -236,13 +236,13 @@ public:
 	int32 GetCurrentCacheIndex() const;
 
 	/** Returns the name of the current Blueprint being cached */
-	FString GetCurrentCacheBlueprintName() const;
+	FName GetCurrentCacheBlueprintName() const;
 
 	/** Returns the progress complete on the caching */
 	float GetCacheProgress() const;
 
 	/** Returns the list of Blueprint paths that failed to cache */
-	TArray<FString> GetFailedToCachePathList() const { return FailedToCachePaths; }
+	TArray<FName> GetFailedToCachePathList() const { return FailedToCachePaths; }
 
 	/** Returns the number of Blueprints that failed to cache */
 	int32 GetFailedToCacheCount() const { return FailedToCachePaths.Num(); }
@@ -254,7 +254,7 @@ public:
 	 *
 	 * @param InNumberCached		The number of Blueprints cached, to be chopped off the existing array so the rest (if any) can be finished later
 	 */
-	void FinishedCachingBlueprints(int32 InNumberCached, TArray<FString>& InFailedToCacheList);
+	void FinishedCachingBlueprints(int32 InNumberCached, TArray<FName>& InFailedToCacheList);
 
 	/** Returns TRUE if Blueprints are being cached. */
 	bool IsCacheInProgress() const;
@@ -305,16 +305,14 @@ private:
 	 * @param InSearchData		Data to add to the database
 	 * @return					Index into the SearchArray for looking up the added item
 	 */
-	int32 AddSearchDataToDatabase(FSearchData& InSearchData);
+	int32 AddSearchDataToDatabase(FSearchData InSearchData);
 
 	/** Removes a Blueprint from being managed by the FiB system by passing in the UBlueprint's path */
-	void RemoveBlueprintByPath(FString InPath);
+	void RemoveBlueprintByPath(FName InPath);
 
-	/** Removes a World Blueprint with compound FiB searchable data from being managed by the FiB system by passing in the World's path */
-	void RemoveWorldByPath(FString InPath);
 protected:
 	/** Maps the Blueprint paths to their index in the SearchArray */
-	TMap<FString, int> SearchMap;
+	TMap<FName, int32> SearchMap;
 
 	/** Stores the Blueprint search data and is used to iterate over in small chunks */
 	TArray<FSearchData> SearchArray;
@@ -344,10 +342,10 @@ protected:
 	TWeakPtr<SFindInBlueprints> SourceCachingWidget;
 
 	/** Blueprint paths that have not been cached for searching due to lack of data, this means that they are either older Blueprints, or the DDC cannot find the data */
-	TArray<FString> UncachedBlueprints;
+	TArray<FName> UncachedBlueprints;
 
 	/** List of paths for Blueprints that failed to cache */
-	TArray<FString> FailedToCachePaths;
+	TArray<FName> FailedToCachePaths;
 
 	/** Tickable object that does the caching of uncached Blueprints at a rate of once per tick */
 	class FCacheAllBlueprintsTickableObject* CachingObject;
