Added alignment commands to blueprint graphs

Standard options are available:
 - Align top/middle/bottom/left/center/right
 - Distribute horizontally/vertically

diff --git a/Engine/Source/Editor/BlueprintGraph/Private/EdGraphSchema_K2.cpp b/Engine/Source/Editor/BlueprintGraph/Private/EdGraphSchema_K2.cpp
--- a/Engine/Source/Editor/BlueprintGraph/Private/EdGraphSchema_K2.cpp
+++ b/Engine/Source/Editor/BlueprintGraph/Private/EdGraphSchema_K2.cpp
@@ -1487,7 +1487,26 @@ void UEdGraphSchema_K2::GetContextMenuActions(const UEdGraph* CurrentGraph, cons
 						MenuBuilder->AddMenuEntry( FGraphEditorCommands::Get().PromoteSelectionToFunction );
 						MenuBuilder->AddMenuEntry( FGraphEditorCommands::Get().PromoteSelectionToMacro );
 					}
+
+					MenuBuilder->AddSubMenu(LOCTEXT("AlignmentHeader", "Alignment"), FText(), FNewMenuDelegate::CreateLambda([](FMenuBuilder& MenuBuilder){
+						
+						MenuBuilder.BeginSection("EdGraphSchemaAlignment", LOCTEXT("AlignHeader", "Align"));
+						MenuBuilder.AddMenuEntry( FGraphEditorCommands::Get().AlignNodesTop );
+						MenuBuilder.AddMenuEntry( FGraphEditorCommands::Get().AlignNodesMiddle );
+						MenuBuilder.AddMenuEntry( FGraphEditorCommands::Get().AlignNodesBottom );
+						MenuBuilder.AddMenuEntry( FGraphEditorCommands::Get().AlignNodesLeft );
+						MenuBuilder.AddMenuEntry( FGraphEditorCommands::Get().AlignNodesCenter );
+						MenuBuilder.AddMenuEntry( FGraphEditorCommands::Get().AlignNodesRight );
+						MenuBuilder.EndSection();
+
+						MenuBuilder.BeginSection("EdGraphSchemaDistribution", LOCTEXT("DistributionHeader", "Distribution"));
+						MenuBuilder.AddMenuEntry( FGraphEditorCommands::Get().DistributeNodesHorizontally );
+						MenuBuilder.AddMenuEntry( FGraphEditorCommands::Get().DistributeNodesVertically );
+						MenuBuilder.EndSection();
+						
+					}));
 				}
+				
 				MenuBuilder->EndSection();
 			}
 
diff --git a/Engine/Source/Editor/GraphEditor/Private/GraphEditorActions.cpp b/Engine/Source/Editor/GraphEditor/Private/GraphEditorActions.cpp
--- a/Engine/Source/Editor/GraphEditor/Private/GraphEditorActions.cpp
+++ b/Engine/Source/Editor/GraphEditor/Private/GraphEditorActions.cpp
@@ -42,6 +42,16 @@ void FGraphEditorCommandsImpl::RegisterCommands()
 	UI_COMMAND( CollapseSelectionToFunction, "Collapse to Function", "Collapses selected nodes into a single function node.", EUserInterfaceActionType::Button, FInputChord() )
 	UI_COMMAND( CollapseSelectionToMacro, "Collapse to Macro", "Collapses selected nodes into a single macro node.", EUserInterfaceActionType::Button, FInputChord() )
 
+	UI_COMMAND( AlignNodesTop, "Align Top", "Aligns the top edges of the selected nodes", EUserInterfaceActionType::Button, FInputChord() )
+	UI_COMMAND( AlignNodesMiddle, "Align Middle", "Aligns the vertical middles of the selected nodes", EUserInterfaceActionType::Button, FInputChord() )
+	UI_COMMAND( AlignNodesBottom, "Align Bottom", "Aligns the bottom edges of the selected nodes", EUserInterfaceActionType::Button, FInputChord() )
+	UI_COMMAND( AlignNodesLeft, "Align Left", "Aligns the left edges of the selected nodes", EUserInterfaceActionType::Button, FInputChord() )
+	UI_COMMAND( AlignNodesCenter, "Align Center", "Aligns the horizontal centers of the selected nodes", EUserInterfaceActionType::Button, FInputChord() )
+	UI_COMMAND( AlignNodesRight, "Align Right", "Aligns the right edges of the selected nodes", EUserInterfaceActionType::Button, FInputChord() )
+
+	UI_COMMAND( DistributeNodesHorizontally, "Distribute Horizontally", "Evenly distributes the selected nodes horizontally", EUserInterfaceActionType::Button, FInputChord() )
+	UI_COMMAND( DistributeNodesVertically, "Distribute Vertically", "Evenly distributes the selected nodes vertically", EUserInterfaceActionType::Button, FInputChord() )
+
 	UI_COMMAND( SelectReferenceInLevel, "Find Actor in Level", "Select the actor referenced by this node in the level", EUserInterfaceActionType::Button, FInputChord() )
 	UI_COMMAND( AssignReferencedActor, "Assign selected Actor", "Assign the selected actor to be this node's referenced object", EUserInterfaceActionType::Button, FInputChord() )
 	UI_COMMAND( FindReferences, "Find References", "Find references of this item", EUserInterfaceActionType::Button, FInputChord() )
diff --git a/Engine/Source/Editor/GraphEditor/Private/SGraphEditorImpl.cpp b/Engine/Source/Editor/GraphEditor/Private/SGraphEditorImpl.cpp
--- a/Engine/Source/Editor/GraphEditor/Private/SGraphEditorImpl.cpp
+++ b/Engine/Source/Editor/GraphEditor/Private/SGraphEditorImpl.cpp
@@ -159,6 +159,20 @@ bool SGraphEditorImpl::GetBoundsForSelectedNodes( class FSlateRect& Rect, float
 	return GraphPanel->GetBoundsForSelectedNodes(Rect, Padding);
 }
 
+bool SGraphEditorImpl::GetBoundsForNode( const UEdGraphNode* InNode, class FSlateRect& Rect, float Padding) const
+{
+	FVector2D TopLeft, BottomRight;
+	if (GraphPanel->GetBoundsForNode(InNode, TopLeft, BottomRight, Padding))
+	{
+		Rect.Left = TopLeft.X;
+		Rect.Top = TopLeft.Y;
+		Rect.Bottom = BottomRight.Y;
+		Rect.Right = BottomRight.X;
+		return true;
+	}
+	return false;
+}
+
 void SGraphEditorImpl::Construct( const FArguments& InArgs )
 {
 	Commands = MakeShareable( new FUICommandList() );
diff --git a/Engine/Source/Editor/GraphEditor/Private/SGraphEditorImpl.h b/Engine/Source/Editor/GraphEditor/Private/SGraphEditorImpl.h
--- a/Engine/Source/Editor/GraphEditor/Private/SGraphEditorImpl.h
+++ b/Engine/Source/Editor/GraphEditor/Private/SGraphEditorImpl.h
@@ -131,6 +131,7 @@ public:
 	virtual UEdGraphPin* GetGraphPinForMenu() override;
 	virtual void ZoomToFit(bool bOnlySelection) override;
 	virtual bool GetBoundsForSelectedNodes( class FSlateRect& Rect, float Padding) override;
+	virtual bool GetBoundsForNode( const UEdGraphNode* InNode, class FSlateRect& Rect, float Padding) const override;
 	virtual void NotifyGraphChanged() override;
 	virtual TSharedPtr<SWidget> GetTitleBar() const override;
 	virtual void SetViewLocation(const FVector2D& Location, float ZoomAmount) override;
diff --git a/Engine/Source/Editor/GraphEditor/Private/SGraphPanel.cpp b/Engine/Source/Editor/GraphEditor/Private/SGraphPanel.cpp
--- a/Engine/Source/Editor/GraphEditor/Private/SGraphPanel.cpp
+++ b/Engine/Source/Editor/GraphEditor/Private/SGraphPanel.cpp
@@ -1081,6 +1081,11 @@ void SGraphPanel::AttachGraphEvents(TSharedPtr<SGraphNode> CreatedSubNode)
 	CreatedSubNode->SetTextCommittedEvent(OnTextCommitted);
 }
 
+bool SGraphPanel::GetBoundsForNode(const UObject* InNode, FVector2D& MinCorner, FVector2D& MaxCorner, float Padding) const
+{
+	return SNodePanel::GetBoundsForNode(InNode, MinCorner, MaxCorner, Padding);
+}
+
 const TSharedRef<SGraphNode> SGraphPanel::GetChild(int32 ChildIndex)
 {
 	return StaticCastSharedRef<SGraphNode>(Children[ChildIndex]);
diff --git a/Engine/Source/Editor/GraphEditor/Private/SGraphPanel.h b/Engine/Source/Editor/GraphEditor/Private/SGraphPanel.h
--- a/Engine/Source/Editor/GraphEditor/Private/SGraphPanel.h
+++ b/Engine/Source/Editor/GraphEditor/Private/SGraphPanel.h
@@ -148,6 +148,9 @@ public:
 	/** Returns if this graph is editable */
 	bool IsGraphEditable() const { return IsEditable.Get(); }
 
+	/** Attempt to retrieve the bounds for the specified node */
+	bool GetBoundsForNode(const UObject* InNode, FVector2D& MinCorner, FVector2D& MaxCorner, float Padding = 0.0f) const;
+
 protected:
 
 	void NotifyGraphChanged ( const struct FEdGraphEditAction& InAction);
diff --git a/Engine/Source/Editor/GraphEditor/Private/SNodePanel.cpp b/Engine/Source/Editor/GraphEditor/Private/SNodePanel.cpp
--- a/Engine/Source/Editor/GraphEditor/Private/SNodePanel.cpp
+++ b/Engine/Source/Editor/GraphEditor/Private/SNodePanel.cpp
@@ -1378,17 +1378,17 @@ bool SNodePanel::IsNodeCulled(const TSharedRef<SNode>& Node, const FGeometry& Al
 
 }
 
-bool SNodePanel::GetBoundsForNode(const UObject* InNode, /*out*/ FVector2D& MinCorner, /*out*/ FVector2D& MaxCorner, float Padding)
+bool SNodePanel::GetBoundsForNode(const UObject* InNode, /*out*/ FVector2D& MinCorner, /*out*/ FVector2D& MaxCorner, float Padding) const
 {
 	MinCorner = FVector2D(MAX_FLT, MAX_FLT);
 	MaxCorner = FVector2D(-MAX_FLT, -MAX_FLT);
 
 	bool bValid = false;
 
-	TSharedRef<SNode>* pWidget = (InNode) ? NodeToWidgetLookup.Find(InNode) : nullptr;
+	const TSharedRef<SNode>* pWidget = (InNode) ? NodeToWidgetLookup.Find(InNode) : nullptr;
 	if (pWidget)
 	{
-		SNode& Widget = pWidget->Get();
+		const SNode& Widget = pWidget->Get();
 		const FVector2D Lower = Widget.GetPosition();
 		const FVector2D Upper = Lower + Widget.GetDesiredSize();
 
diff --git a/Engine/Source/Editor/GraphEditor/Public/GraphEditorActions.h b/Engine/Source/Editor/GraphEditor/Public/GraphEditorActions.h
--- a/Engine/Source/Editor/GraphEditor/Public/GraphEditorActions.h
+++ b/Engine/Source/Editor/GraphEditor/Public/GraphEditorActions.h
@@ -60,6 +60,18 @@ public:
 	TSharedPtr< FUICommandInfo > CollapseSelectionToFunction;
 	TSharedPtr< FUICommandInfo > CollapseSelectionToMacro;
 
+	// Alignment commands
+	TSharedPtr< FUICommandInfo > AlignNodesTop;
+	TSharedPtr< FUICommandInfo > AlignNodesMiddle;
+	TSharedPtr< FUICommandInfo > AlignNodesBottom;
+
+	TSharedPtr< FUICommandInfo > AlignNodesLeft;
+	TSharedPtr< FUICommandInfo > AlignNodesCenter;
+	TSharedPtr< FUICommandInfo > AlignNodesRight;
+
+	TSharedPtr< FUICommandInfo > DistributeNodesHorizontally;
+	TSharedPtr< FUICommandInfo > DistributeNodesVertically;
+
 	//
 	TSharedPtr< FUICommandInfo > SelectReferenceInLevel;
 	TSharedPtr< FUICommandInfo > AssignReferencedActor;
diff --git a/Engine/Source/Editor/GraphEditor/Public/SNodePanel.h b/Engine/Source/Editor/GraphEditor/Public/SNodePanel.h
--- a/Engine/Source/Editor/GraphEditor/Public/SNodePanel.h
+++ b/Engine/Source/Editor/GraphEditor/Public/SNodePanel.h
@@ -778,7 +778,7 @@ protected:
 	 * Get the bounds of the given node
 	 * @return True if successful
 	 */
-	bool GetBoundsForNode(const UObject* InNode, /*out*/ FVector2D& MinCorner, /*out*/ FVector2D& MaxCorner, float Padding = 0.0f);
+	bool GetBoundsForNode(const UObject* InNode, /*out*/ FVector2D& MinCorner, /*out*/ FVector2D& MaxCorner, float Padding = 0.0f) const;
 
 	/**
 	 * Get the bounds of the selected nodes 
diff --git a/Engine/Source/Editor/Kismet/Private/BlueprintEditor.cpp b/Engine/Source/Editor/Kismet/Private/BlueprintEditor.cpp
--- a/Engine/Source/Editor/Kismet/Private/BlueprintEditor.cpp
+++ b/Engine/Source/Editor/Kismet/Private/BlueprintEditor.cpp
@@ -1006,6 +1006,40 @@ TSharedRef<SGraphEditor> FBlueprintEditor::CreateGraphEditorWidget(TSharedRef<FT
 				FIsActionChecked(),
 				FIsActionButtonVisible::CreateSP( this, &FBlueprintEditor::NewDocument_IsVisibleForType, CGT_NewMacroGraph )
 				);
+			
+			// Alignment Commands
+			GraphEditorCommands->MapAction( FGraphEditorCommands::Get().AlignNodesTop,
+				FExecuteAction::CreateSP( this, &FBlueprintEditor::OnAlignTop )
+				);
+
+			GraphEditorCommands->MapAction( FGraphEditorCommands::Get().AlignNodesMiddle,
+				FExecuteAction::CreateSP( this, &FBlueprintEditor::OnAlignMiddle )
+				);
+
+			GraphEditorCommands->MapAction( FGraphEditorCommands::Get().AlignNodesBottom,
+				FExecuteAction::CreateSP( this, &FBlueprintEditor::OnAlignBottom )
+				);
+
+			GraphEditorCommands->MapAction( FGraphEditorCommands::Get().AlignNodesLeft,
+				FExecuteAction::CreateSP( this, &FBlueprintEditor::OnAlignLeft )
+				);
+
+			GraphEditorCommands->MapAction( FGraphEditorCommands::Get().AlignNodesCenter,
+				FExecuteAction::CreateSP( this, &FBlueprintEditor::OnAlignCenter )
+				);
+
+			GraphEditorCommands->MapAction( FGraphEditorCommands::Get().AlignNodesRight,
+				FExecuteAction::CreateSP( this, &FBlueprintEditor::OnAlignRight )
+				);
+
+			// Distribution Commands
+			GraphEditorCommands->MapAction( FGraphEditorCommands::Get().DistributeNodesHorizontally,
+				FExecuteAction::CreateSP( this, &FBlueprintEditor::OnDistributeNodesH )
+				);
+
+			GraphEditorCommands->MapAction( FGraphEditorCommands::Get().DistributeNodesVertically,
+				FExecuteAction::CreateSP( this, &FBlueprintEditor::OnDistributeNodesV )
+				);
 
 			GraphEditorCommands->MapAction( FGenericCommands::Get().Rename,
 				FExecuteAction::CreateSP( this, &FBlueprintEditor::OnRenameNode ),
@@ -4849,6 +4883,281 @@ bool FBlueprintEditor::CanExpandNodes() const
 	return false;
 }
 
+/** Struct used for generically aligning nodes */
+struct FAlignmentData
+{
+	FAlignmentData(UEdGraphNode* InNode, int32& InTargetProperty, float InTargetOffset)
+		: Node(InNode), TargetProperty(InTargetProperty), TargetOffset(InTargetOffset)
+	{}
+
+	/** The node to position */
+	UEdGraphNode* Node;
+	/** The property within the node to read/write */
+	int32& TargetProperty;
+	/** The offset from the property to consider for alignment */
+	float TargetOffset;
+
+	/** Get the destination target from this alignment data (property + offset) */
+	float GetTarget() const
+	{
+		return float(TargetProperty) + TargetOffset;
+	}
+};
+
+/** Align the specified array of node data so all the target properties line up */
+void AlignNodes(const TArray<FAlignmentData>& InData)
+{
+	// Calculate the average target property...
+	float MeanPropertyValue = 0.f;
+	for (auto& Entry : InData)
+	{
+		MeanPropertyValue += Entry.GetTarget();
+	}
+	MeanPropertyValue /= InData.Num();
+
+	// ...and find the closest *actual* property to this average
+	float ThisTarget = InData[0].GetTarget();
+	float ClosestTarget = ThisTarget;
+	for (int32 Index = 1; Index < InData.Num(); ++Index)
+	{
+		ThisTarget = InData[Index].GetTarget();
+		if (FMath::Abs(ThisTarget - MeanPropertyValue) < FMath::Abs(ClosestTarget - MeanPropertyValue))
+		{
+			ClosestTarget = ThisTarget;
+		}
+	}
+
+	// Now set all the properties to the target
+	for (auto& Entry : InData)
+	{
+		Entry.Node->Modify();
+		Entry.TargetProperty = ClosestTarget - Entry.TargetOffset;
+	}
+}
+
+FVector2D GetNodeSize(const SGraphEditor& GraphEditor, const UEdGraphNode* Node)
+{
+	FSlateRect Rect;
+	if (GraphEditor.GetBoundsForNode(Node, Rect, 0.f))
+	{
+		return FVector2D(Rect.Right - Rect.Left, Rect.Bottom - Rect.Top);
+	}
+
+	return FVector2D(Node->NodeWidth, Node->NodeHeight);
+}
+
+void FBlueprintEditor::OnAlignTop()
+{
+	TArray<FAlignmentData> AlignData;
+	for (UObject* It : GetSelectedNodes())
+	{
+		if (UEdGraphNode* Node = Cast<UEdGraphNode>(It))
+		{
+			AlignData.Add(FAlignmentData(Node, Node->NodePosY, 0.f));
+		}
+	}
+
+	if (AlignData.Num())
+	{
+		const FScopedTransaction Transaction(FGraphEditorCommands::Get().AlignNodesTop->GetLabel());
+		AlignNodes(AlignData);
+	}
+}
+
+void FBlueprintEditor::OnAlignMiddle()
+{
+	TSharedPtr<SGraphEditor> FocusedGraphEd = FocusedGraphEdPtr.Pin();
+	if (!FocusedGraphEd.IsValid())
+	{
+		return;
+	}
+
+	TArray<FAlignmentData> AlignData;
+	for (UObject* It : GetSelectedNodes())
+	{
+		if (UEdGraphNode* Node = Cast<UEdGraphNode>(It))
+		{
+			AlignData.Add(FAlignmentData(Node, Node->NodePosY, GetNodeSize(*FocusedGraphEd, Node).Y*.5f));
+		}
+	}
+
+	if (AlignData.Num())
+	{
+		const FScopedTransaction Transaction(FGraphEditorCommands::Get().AlignNodesMiddle->GetLabel());
+		AlignNodes(AlignData);
+	}
+}
+
+void FBlueprintEditor::OnAlignBottom()
+{
+	TSharedPtr<SGraphEditor> FocusedGraphEd = FocusedGraphEdPtr.Pin();
+	if (!FocusedGraphEd.IsValid())
+	{
+		return;
+	}
+
+	TArray<FAlignmentData> AlignData;
+	for (UObject* It : GetSelectedNodes())
+	{
+		if (UEdGraphNode* Node = Cast<UEdGraphNode>(It))
+		{
+			AlignData.Add(FAlignmentData(Node, Node->NodePosY, GetNodeSize(*FocusedGraphEd, Node).Y));
+		}
+	}
+
+	if (AlignData.Num())
+	{
+		const FScopedTransaction Transaction(FGraphEditorCommands::Get().AlignNodesBottom->GetLabel());
+		AlignNodes(AlignData);
+	}
+}
+
+void FBlueprintEditor::OnAlignLeft()
+{
+	TArray<FAlignmentData> AlignData;
+	for (UObject* It : GetSelectedNodes())
+	{
+		if (UEdGraphNode* Node = Cast<UEdGraphNode>(It))
+		{
+			AlignData.Add(FAlignmentData(Node, Node->NodePosX, 0.f));
+		}
+	}
+
+	if (AlignData.Num())
+	{
+		const FScopedTransaction Transaction(FGraphEditorCommands::Get().AlignNodesLeft->GetLabel());
+		AlignNodes(AlignData);
+	}
+}
+
+void FBlueprintEditor::OnAlignCenter()
+{
+	TSharedPtr<SGraphEditor> FocusedGraphEd = FocusedGraphEdPtr.Pin();
+	if (!FocusedGraphEd.IsValid())
+	{
+		return;
+	}
+	
+	TArray<FAlignmentData> AlignData;
+	for (UObject* It : GetSelectedNodes())
+	{
+		if (UEdGraphNode* Node = Cast<UEdGraphNode>(It))
+		{
+			AlignData.Add(FAlignmentData(Node, Node->NodePosX, GetNodeSize(*FocusedGraphEd, Node).X*.5f));
+		}
+	}
+
+	if (AlignData.Num())
+	{
+		const FScopedTransaction Transaction(FGraphEditorCommands::Get().AlignNodesCenter->GetLabel());
+		AlignNodes(AlignData);
+	}
+}
+
+void FBlueprintEditor::OnAlignRight()
+{
+	TSharedPtr<SGraphEditor> FocusedGraphEd = FocusedGraphEdPtr.Pin();
+	if (!FocusedGraphEd.IsValid())
+	{
+		return;
+	}
+	
+	TArray<FAlignmentData> AlignData;
+	for (UObject* It : GetSelectedNodes())
+	{
+		if (UEdGraphNode* Node = Cast<UEdGraphNode>(It))
+		{
+			AlignData.Add(FAlignmentData(Node, Node->NodePosX, GetNodeSize(*FocusedGraphEd, Node).X));
+		}
+	}
+
+	if (AlignData.Num())
+	{
+		const FScopedTransaction Transaction(FGraphEditorCommands::Get().AlignNodesRight->GetLabel());
+		AlignNodes(AlignData);
+	}
+}
+
+/** Distribute the specified array of node data evenly */
+void DistributeNodes(TArray<FAlignmentData>& InData)
+{
+	// Sort the data
+	InData.Sort([](const FAlignmentData& A, const FAlignmentData& B){
+		return A.TargetProperty + A.TargetOffset/2 < B.TargetProperty + B.TargetOffset/2;
+	});
+
+	// Measure the available space
+	float TotalWidthOfNodes = 0.f;
+	for (int32 Index = 1; Index < InData.Num() - 1; ++Index)
+	{
+		TotalWidthOfNodes += InData[Index].TargetOffset;
+	}
+
+	const float SpaceToDistributeIn = InData.Last().TargetProperty - InData[0].GetTarget();
+	const float PaddingAmount = ((SpaceToDistributeIn - TotalWidthOfNodes) / (InData.Num() - 1));
+
+	float TargetPosition = InData[0].GetTarget() + PaddingAmount;
+
+	// Now set all the properties on the target
+	for (int32 Index = 1; Index < InData.Num() - 1; ++Index)
+	{
+		auto& Entry = InData[Index];
+
+		Entry.Node->Modify();
+		Entry.TargetProperty = TargetPosition;
+
+		TargetPosition = Entry.GetTarget() + PaddingAmount;
+	}
+}
+
+void FBlueprintEditor::OnDistributeNodesH()
+{
+	TSharedPtr<SGraphEditor> FocusedGraphEd = FocusedGraphEdPtr.Pin();
+	if (!FocusedGraphEd.IsValid())
+	{
+		return;
+	}
+	
+	TArray<FAlignmentData> AlignData;
+	for (UObject* It : GetSelectedNodes())
+	{
+		if (UEdGraphNode* Node = Cast<UEdGraphNode>(It))
+		{
+			AlignData.Add(FAlignmentData(Node, Node->NodePosX, GetNodeSize(*FocusedGraphEd, Node).X));
+		}
+	}
+
+	if (AlignData.Num() > 2)
+	{
+		const FScopedTransaction Transaction(FGraphEditorCommands::Get().DistributeNodesHorizontally->GetLabel());
+		DistributeNodes(AlignData);
+	}
+}
+
+void FBlueprintEditor::OnDistributeNodesV()
+{
+	TSharedPtr<SGraphEditor> FocusedGraphEd = FocusedGraphEdPtr.Pin();
+	if (!FocusedGraphEd.IsValid())
+	{
+		return;
+	}
+	
+	TArray<FAlignmentData> AlignData;
+	for (UObject* It : GetSelectedNodes())
+	{
+		if (UEdGraphNode* Node = Cast<UEdGraphNode>(It))
+		{
+			AlignData.Add(FAlignmentData(Node, Node->NodePosY, GetNodeSize(*FocusedGraphEd, Node).Y));
+		}
+	}
+
+	if (AlignData.Num() > 2)
+	{
+		const FScopedTransaction Transaction(FGraphEditorCommands::Get().DistributeNodesVertically->GetLabel());
+		DistributeNodes(AlignData);
+	}
+}
+
 void FBlueprintEditor::SelectAllNodes()
 {
 	TSharedPtr<SGraphEditor> FocusedGraphEd = FocusedGraphEdPtr.Pin();
@@ -7772,3 +8081,4 @@ void FBlueprintEditor::TryInvokingDetailsTab(bool bFlash)
 /////////////////////////////////////////////////////
 
 #undef LOCTEXT_NAMESPACE
+
diff --git a/Engine/Source/Editor/Kismet/Public/BlueprintEditor.h b/Engine/Source/Editor/Kismet/Public/BlueprintEditor.h
--- a/Engine/Source/Editor/Kismet/Public/BlueprintEditor.h
+++ b/Engine/Source/Editor/Kismet/Public/BlueprintEditor.h
@@ -714,6 +714,16 @@ protected:
 	void OnExpandNodes();
 	bool CanExpandNodes() const;
 
+	void OnAlignTop();
+	void OnAlignMiddle();
+	void OnAlignBottom();
+	void OnAlignLeft();
+	void OnAlignCenter();
+	void OnAlignRight();
+	
+	void OnDistributeNodesH();
+	void OnDistributeNodesV();
+
 	void SelectAllNodes();
 	bool CanSelectAllNodes() const;
 
diff --git a/Engine/Source/Editor/UnrealEd/Public/GraphEditor.h b/Engine/Source/Editor/UnrealEd/Public/GraphEditor.h
--- a/Engine/Source/Editor/UnrealEd/Public/GraphEditor.h
+++ b/Engine/Source/Editor/UnrealEd/Public/GraphEditor.h
@@ -322,6 +322,16 @@ public:
 		return false;
 	}
 
+	/** Get Bounds for the specified node, returns false on failure */
+	virtual bool GetBoundsForNode( const UEdGraphNode* InNode, class FSlateRect& Rect, float Padding ) const
+	{
+		if (Implementation.IsValid())
+		{
+			return Implementation->GetBoundsForNode(InNode, Rect, Padding);
+		}
+		return false;
+	}
+
 	// Invoked to let this widget know that the GraphEditor module has been reloaded
 	UNREALED_API void OnModuleReloaded();
 
