[UE-12263] Add ability to mark BlueprintCallable functions as "DevelopmentOnly" - by default, these nodes will be compiled out of Blueprint class function graphs at cook/package time.

change summary:
- deprecated the 'UEdGraphNode::bIsNodeEnabled' flag & added fixup code for serialization (load)
- added support for pruning "disabled" nodes at compile time w/o breaking existing flow in the source graph.
- modified function call nodes to look for new "DevelopmentOnly" UFUNCTION metadata and disable when found.
- added UEdGraphNode::IsNodeEnabled(), Enable() and Disable() APIs. also added a new 'bUserSetEnabledState' flag.
- added the UEdGraphNode::IsDevelopmentMode() API; subclasses can override to provide additional functionality if necessary.
- added a new 'UCookerSettings::bCompileBlueprintsInDevelopmentMode' flag to control whether or not Blueprints are compiled in development mode at cook time (accessible in Project Settings->Cooker); defaults to OFF.
- (experimental) added a "Compile Options" section to the Blueprint graph node context menu to allow users to explicitly enable or disable nodes.
- (experimental) added a BPGC setting to allow users to optionally override the cooker default setting for the development/release compile switch (UE-12270, WiP).
- the experimental features above are gated by a new 'UBlueprintEditorSettings::bAllowExplicitImpureNodeDisabling flag' (accessible in Editor Settings->Blueprint Editor); defaults to OFF for now.

diff --git a/Engine/Source/Editor/BlueprintGraph/Classes/EdGraphSchema_K2.h b/Engine/Source/Editor/BlueprintGraph/Classes/EdGraphSchema_K2.h
--- a/Engine/Source/Editor/BlueprintGraph/Classes/EdGraphSchema_K2.h
+++ b/Engine/Source/Editor/BlueprintGraph/Classes/EdGraphSchema_K2.h
@@ -116,6 +116,9 @@ public:
 	/** For functions that have the MD_WorldContext metadata but are safe to be called from contexts that do not have the ability to provide the world context (either through GetWorld() or ShowWorldContextPin class metadata */
 	static const FName MD_CallableWithoutWorldContext;
 
+	/** For functions that should be compiled in development mode only */
+	static const FName MD_DevelopmentOnly;
+
 	/** If true, an unconnected pin will generate a UPROPERTY under the hood to connect as the input, which will be set to the literal value for the pin.  Only valid for reference parameters. */
 	static const FName MD_AutoCreateRefTerm;
 
diff --git a/Engine/Source/Editor/BlueprintGraph/Classes/K2Node.h b/Engine/Source/Editor/BlueprintGraph/Classes/K2Node.h
--- a/Engine/Source/Editor/BlueprintGraph/Classes/K2Node.h
+++ b/Engine/Source/Editor/BlueprintGraph/Classes/K2Node.h
@@ -109,6 +109,7 @@ class UK2Node : public UEdGraphNode
 	BLUEPRINTGRAPH_API virtual bool ShowPaletteIconOnNode() const override { return true; }
 	BLUEPRINTGRAPH_API virtual bool AllowSplitPins() const override;
 	BLUEPRINTGRAPH_API virtual UEdGraphPin* GetPassThroughPin(const UEdGraphPin* FromPin) const override;
+	BLUEPRINTGRAPH_API virtual bool IsInDevelopmentMode() const override;
 	// End of UEdGraphNode interface
 
 	// K2Node interface
diff --git a/Engine/Source/Editor/BlueprintGraph/Private/BlueprintEditorSettings.cpp b/Engine/Source/Editor/BlueprintGraph/Private/BlueprintEditorSettings.cpp
--- a/Engine/Source/Editor/BlueprintGraph/Private/BlueprintEditorSettings.cpp
+++ b/Engine/Source/Editor/BlueprintGraph/Private/BlueprintEditorSettings.cpp
@@ -26,6 +26,7 @@ UBlueprintEditorSettings::UBlueprintEditorSettings(const FObjectInitializer& Obj
 	// Compiler Settings
 	, SaveOnCompile(SoC_Never)
 	, bJumpToNodeErrors(false)
+	, bAllowExplicitImpureNodeDisabling(false)
 	// Developer Settings
 	, bShowActionMenuItemSignatures(false)
 	// Perf Settings
diff --git a/Engine/Source/Editor/BlueprintGraph/Private/BlueprintEventNodeSpawner.cpp b/Engine/Source/Editor/BlueprintGraph/Private/BlueprintEventNodeSpawner.cpp
--- a/Engine/Source/Editor/BlueprintGraph/Private/BlueprintEventNodeSpawner.cpp
+++ b/Engine/Source/Editor/BlueprintGraph/Private/BlueprintEventNodeSpawner.cpp
@@ -58,7 +58,7 @@ UK2Node_Event* UBlueprintEventNodeSpawnerImpl::FindCustomEventNode(UBlueprint* B
 //------------------------------------------------------------------------------
 static void UBlueprintEventNodeSpawnerImpl::RemoveAllDisabledNodes(UEdGraphNode* InNode, UEdGraph* InParentGraph)
 {
-	if(InNode && !InNode->bIsNodeEnabled)
+	if(InNode && !InNode->IsNodeEnabled() && !InNode->bUserSetEnabledState)
 	{
 		// Go through all pin connections and consume any disabled nodes so we do not leave garbage.
 		for (UEdGraphPin* Pin : InNode->Pins)
@@ -191,7 +191,7 @@ UEdGraphNode* UBlueprintEventNodeSpawner::Invoke(UEdGraph* ParentGraph, FBinding
 
 	// This Event node might already be present in the Blueprint in a disabled state, 
 	// remove it and allow the user to successfully place the node where they want it.
-	if(EventNode && !EventNode->bIsNodeEnabled)
+	if(EventNode && !EventNode->IsNodeEnabled() && !EventNode->bUserSetEnabledState)
 	{
 		UBlueprintEventNodeSpawnerImpl::RemoveAllDisabledNodes(EventNode, ParentGraph);
 		EventNode = nullptr;
diff --git a/Engine/Source/Editor/BlueprintGraph/Private/EdGraphSchema_K2.cpp b/Engine/Source/Editor/BlueprintGraph/Private/EdGraphSchema_K2.cpp
--- a/Engine/Source/Editor/BlueprintGraph/Private/EdGraphSchema_K2.cpp
+++ b/Engine/Source/Editor/BlueprintGraph/Private/EdGraphSchema_K2.cpp
@@ -65,6 +65,7 @@ const FName FBlueprintMetadata::MD_HideSelfPin(TEXT("HideSelfPin"));
 const FName FBlueprintMetadata::MD_DefaultToSelf(TEXT("DefaultToSelf"));
 const FName FBlueprintMetadata::MD_WorldContext(TEXT("WorldContext"));
 const FName FBlueprintMetadata::MD_CallableWithoutWorldContext(TEXT("CallableWithoutWorldContext"));
+const FName FBlueprintMetadata::MD_DevelopmentOnly(TEXT("DevelopmentOnly"));
 const FName FBlueprintMetadata::MD_AutoCreateRefTerm(TEXT("AutoCreateRefTerm"));
 
 const FName FBlueprintMetadata::MD_ShowWorldContextPin(TEXT("ShowWorldContextPin"));
@@ -1510,11 +1511,61 @@ void UEdGraphSchema_K2::GetContextMenuActions(const UEdGraph* CurrentGraph, cons
 				MenuBuilder->EndSection();
 			}
 
-			// Add breakpoint actions
 			if (const UK2Node* K2Node = Cast<const UK2Node>(InGraphNode))
 			{
 				if (!K2Node->IsNodePure())
 				{
+					const UBlueprintEditorSettings* EditorSettings = GetDefault<UBlueprintEditorSettings>();
+					if (EditorSettings && EditorSettings->bAllowExplicitImpureNodeDisabling)
+					{
+						// Don't expose the enabled state for disabled nodes that were not explicitly disabled by the user
+						if (!bIsDebugging && (K2Node->bUserSetEnabledState || K2Node->EnabledState != ENodeEnabledState::Disabled))
+						{
+							// Add compile options
+							MenuBuilder->BeginSection("EdGraphSchemaCompileOptions", LOCTEXT("CompileOptionsHeader", "Compile Options"));
+							{
+								MenuBuilder->AddMenuEntry(
+									FGraphEditorCommands::Get().DisableNodes,
+									NAME_None,
+									LOCTEXT("DisableCompile", "Disable (Do Not Compile)"),
+									LOCTEXT("DisableCompileToolTip", "Selected node(s) will not be compiled."));
+
+								TSharedPtr<const FUICommandList> MenuCommandList = MenuBuilder->GetTopCommandList();
+								if(ensure(MenuCommandList.IsValid()))
+								{
+									const FUIAction* SubMenuUIAction = MenuCommandList->GetActionForCommand(FGraphEditorCommands::Get().EnableNodes);
+									if(ensure(SubMenuUIAction))
+									{
+										MenuBuilder->AddSubMenu(
+											LOCTEXT("EnableCompileSubMenu", "Enable Compile"),
+											LOCTEXT("EnableCompileSubMenuToolTip", "Options to enable selected node(s) for compile."),
+											FNewMenuDelegate::CreateLambda([MenuCommandList](FMenuBuilder& SubMenuBuilder)
+											{
+												SubMenuBuilder.PushCommandList(MenuCommandList.ToSharedRef());
+
+												SubMenuBuilder.AddMenuEntry(
+													FGraphEditorCommands::Get().EnableNodes_Always,
+													NAME_None,
+													LOCTEXT("EnableCompileAlways", "Always"),
+													LOCTEXT("EnableCompileAlwaysToolTip", "Always compile selected node(s)."));
+												SubMenuBuilder.AddMenuEntry(
+													FGraphEditorCommands::Get().EnableNodes_DevelopmentOnly,
+													NAME_None,
+													LOCTEXT("EnableCompileDevelopmentOnly", "Development Only"),
+													LOCTEXT("EnableCompileDevelopmentOnlyToolTip", "Compile selected node(s) for development only."));
+
+												SubMenuBuilder.PopCommandList();
+											}),
+											*SubMenuUIAction,
+											NAME_None, FGraphEditorCommands::Get().EnableNodes->GetUserInterfaceType());
+									}
+								}
+							}
+							MenuBuilder->EndSection();
+						}
+					}
+					
+					// Add breakpoint actions
 					MenuBuilder->BeginSection("EdGraphSchemaBreakpoints", LOCTEXT("BreakpointsHeader", "Breakpoints"));
 					{
 						MenuBuilder->AddMenuEntry( FGraphEditorCommands::Get().ToggleBreakpoint );
diff --git a/Engine/Source/Editor/BlueprintGraph/Private/K2Node.cpp b/Engine/Source/Editor/BlueprintGraph/Private/K2Node.cpp
--- a/Engine/Source/Editor/BlueprintGraph/Private/K2Node.cpp
+++ b/Engine/Source/Editor/BlueprintGraph/Private/K2Node.cpp
@@ -8,6 +8,7 @@
 #include "KismetDebugUtilities.h" // for HasDebuggingData(), GetWatchText()
 #include "KismetCompiler.h"
 #include "GraphEditorSettings.h"
+#include "BlueprintEditorSettings.h"
 
 #include "ObjectEditorUtils.h"
 
@@ -1014,24 +1015,39 @@ UEdGraphPin* UK2Node::GetExecPin() const
 
 UEdGraphPin* UK2Node::GetPassThroughPin(const UEdGraphPin* FromPin) const
 {
-	UEdGraphPin* MatchedPin = nullptr;
+	UEdGraphPin* PassThroughPin = nullptr;
 	if(FromPin && Pins.Contains(FromPin))
 	{
 		const UEdGraphSchema_K2* K2Schema = GetDefault<UEdGraphSchema_K2>();
 		if(K2Schema->IsExecPin(*FromPin))
 		{
-			if(FromPin->Direction == EGPD_Input)
+			// Locate the first exec pin that's opposite the given exec pin, if any.
+			for(auto Pin : Pins)
 			{
-				MatchedPin = FindPin(K2Schema->PN_Then);
-			}
-			else
-			{
-				MatchedPin = FindPin(K2Schema->PN_Execute);
+				if(Pin && Pin != FromPin && Pin->Direction != FromPin->Direction && K2Schema->IsExecPin(*Pin))
+				{
+					PassThroughPin = Pin;
+					break;
+				}
 			}
 		}
 	}
 
-	return MatchedPin;
+	return PassThroughPin;
+}
+
+bool UK2Node::IsInDevelopmentMode() const
+{
+	// Check class setting (which can override the default setting)
+	const UBlueprint* OwningBP = GetBlueprint();
+	if(OwningBP != nullptr
+		&& OwningBP->CompileMode != EBlueprintCompileMode::Default)
+	{
+		return OwningBP->CompileMode == EBlueprintCompileMode::Development;
+	}
+
+	// Check default setting
+	return Super::IsInDevelopmentMode();
 }
 
 bool UK2Node::CanCreateUnderSpecifiedSchema(const UEdGraphSchema* DesiredSchema) const
diff --git a/Engine/Source/Editor/BlueprintGraph/Private/K2Node_CallFunction.cpp b/Engine/Source/Editor/BlueprintGraph/Private/K2Node_CallFunction.cpp
--- a/Engine/Source/Editor/BlueprintGraph/Private/K2Node_CallFunction.cpp
+++ b/Engine/Source/Editor/BlueprintGraph/Private/K2Node_CallFunction.cpp
@@ -1814,6 +1814,24 @@ void UK2Node_CallFunction::Serialize(FArchive& Ar)
 				FunctionReference.SetDirect(FunctionReference.GetMemberName(), FunctionGuid, (bSelf ? NULL : FunctionReference.GetMemberParentClass((UClass*)NULL)), bSelf);
 			}
 		}
+
+		// Don't validate the enabled state if the user has explicitly set it.
+		if (!bUserSetEnabledState)
+		{
+			if (const UFunction* Function = GetTargetFunction())
+			{
+				// Enable as development-only if specified in metadata. This way existing functions that have the metadata added to them will get their enabled state fixed up on load.
+				if (EnabledState == ENodeEnabledState::Enabled && Function->HasMetaData(FBlueprintMetadata::MD_DevelopmentOnly))
+				{
+					EnabledState = ENodeEnabledState::DevelopmentOnly;
+				}
+				// Ensure that if the metadata is removed, we also fix up the enabled state to avoid leaving it set as development-only in that case.
+				else if (EnabledState == ENodeEnabledState::DevelopmentOnly && !Function->HasMetaData(FBlueprintMetadata::MD_DevelopmentOnly))
+				{
+					EnabledState = ENodeEnabledState::Enabled;
+				}
+			}
+		}
 	}
 }
 
@@ -1823,6 +1841,16 @@ void UK2Node_CallFunction::PostPlacedNewNode()
 
 	// Try re-setting the function given our new parent scope, in case it turns an external to an internal, or vis versa
 	FunctionReference.RefreshGivenNewSelfScope<UFunction>(GetBlueprintClassFromNode());
+
+	// Re-enable for development only if specified in metadata.
+	if(EnabledState == ENodeEnabledState::Enabled && !bUserSetEnabledState)
+	{
+		const UFunction* Function = GetTargetFunction();
+		if (Function && Function->HasMetaData(FBlueprintMetadata::MD_DevelopmentOnly))
+		{
+			EnabledState = ENodeEnabledState::DevelopmentOnly;
+		}
+	}
 }
 
 FNodeHandlingFunctor* UK2Node_CallFunction::CreateNodeHandler(FKismetCompilerContext& CompilerContext) const
diff --git a/Engine/Source/Editor/BlueprintGraph/Private/K2Node_Event.cpp b/Engine/Source/Editor/BlueprintGraph/Private/K2Node_Event.cpp
--- a/Engine/Source/Editor/BlueprintGraph/Private/K2Node_Event.cpp
+++ b/Engine/Source/Editor/BlueprintGraph/Private/K2Node_Event.cpp
@@ -439,7 +439,7 @@ bool UK2Node_Event::CanPasteHere(const UEdGraph* TargetGraph) const
 						// If the event function is already handled in this Blueprint, don't paste this event
 						for(int32 i = 0; i < ExistingEventNodes.Num() && !bDisallowPaste; ++i)
 						{
-							bDisallowPaste = ExistingEventNodes[i]->bOverrideFunction && ExistingEventNodes[i]->bIsNodeEnabled && AreEventNodesIdentical(this, ExistingEventNodes[i]);
+							bDisallowPaste = ExistingEventNodes[i]->bOverrideFunction && ExistingEventNodes[i]->IsNodeEnabled() && AreEventNodesIdentical(this, ExistingEventNodes[i]);
 						}
 
 						// We need to also check for 'const' BPIE methods that might already be implemented as functions with a read-only 'self' context (these were previously implemented as events)
diff --git a/Engine/Source/Editor/BlueprintGraph/Public/BlueprintEditorSettings.h b/Engine/Source/Editor/BlueprintGraph/Public/BlueprintEditorSettings.h
--- a/Engine/Source/Editor/BlueprintGraph/Public/BlueprintEditorSettings.h
+++ b/Engine/Source/Editor/BlueprintGraph/Public/BlueprintEditorSettings.h
@@ -88,6 +88,10 @@ public:
 	UPROPERTY(EditAnywhere, config, Category=Compiler)
 	bool bJumpToNodeErrors;
 
+	/** If enabled, nodes can be explicitly disabled via context menu when right-clicking on impure nodes in the Blueprint editor. Disabled nodes will not be compiled, but also will not break existing connections. */
+	UPROPERTY(EditAnywhere, config, Category=Experimental, AdvancedDisplay)
+	bool bAllowExplicitImpureNodeDisabling;
+
 // Developer Settings
 public:
 	/** If enabled, tooltips on action menu items will show the associated action's signature id (can be used to setup custom favorites menus).*/
diff --git a/Engine/Source/Editor/GraphEditor/Private/GraphEditorActions.cpp b/Engine/Source/Editor/GraphEditor/Private/GraphEditorActions.cpp
--- a/Engine/Source/Editor/GraphEditor/Private/GraphEditorActions.cpp
+++ b/Engine/Source/Editor/GraphEditor/Private/GraphEditorActions.cpp
@@ -51,6 +51,11 @@ void FGraphEditorCommandsImpl::RegisterCommands()
 
 	UI_COMMAND( DistributeNodesHorizontally, "Distribute Horizontally", "Evenly distributes the selected nodes horizontally", EUserInterfaceActionType::Button, FInputChord() )
 	UI_COMMAND( DistributeNodesVertically, "Distribute Vertically", "Evenly distributes the selected nodes vertically", EUserInterfaceActionType::Button, FInputChord() )
+	
+	UI_COMMAND( EnableNodes, "Enable Nodes", "Selected node(s) will be enabled.", EUserInterfaceActionType::Check, FInputChord() )
+	UI_COMMAND( DisableNodes, "Disable Nodes", "Selected node(s) will be disabled.", EUserInterfaceActionType::Check, FInputChord() )
+	UI_COMMAND( EnableNodes_Always, "Enable Nodes (Always)", "Selected node(s) will always be enabled.", EUserInterfaceActionType::RadioButton, FInputChord() )
+	UI_COMMAND( EnableNodes_DevelopmentOnly, "Enable Nodes (Development Only)", "Selected node(s) will be enabled in development mode only.", EUserInterfaceActionType::RadioButton, FInputChord() )
 
 	UI_COMMAND( SelectReferenceInLevel, "Find Actor in Level", "Select the actor referenced by this node in the level", EUserInterfaceActionType::Button, FInputChord() )
 	UI_COMMAND( AssignReferencedActor, "Assign selected Actor", "Assign the selected actor to be this node's referenced object", EUserInterfaceActionType::Button, FInputChord() )
diff --git a/Engine/Source/Editor/GraphEditor/Private/SCommentBubble.cpp b/Engine/Source/Editor/GraphEditor/Private/SCommentBubble.cpp
--- a/Engine/Source/Editor/GraphEditor/Private/SCommentBubble.cpp
+++ b/Engine/Source/Editor/GraphEditor/Private/SCommentBubble.cpp
@@ -373,7 +373,7 @@ FSlateColor SCommentBubble::GetBubbleColor() const
 {
 	FLinearColor ReturnColor = ColorAndOpacity.Get().GetSpecifiedColor();
 
-	if(!GraphNode->bIsNodeEnabled)
+	if(!GraphNode->IsNodeEnabled())
 	{
 		ReturnColor.A *= 0.6f;
 	}
diff --git a/Engine/Source/Editor/GraphEditor/Private/SGraphNode.cpp b/Engine/Source/Editor/GraphEditor/Private/SGraphNode.cpp
--- a/Engine/Source/Editor/GraphEditor/Private/SGraphNode.cpp
+++ b/Engine/Source/Editor/GraphEditor/Private/SGraphNode.cpp
@@ -523,7 +523,7 @@ FSlateColor SGraphNode::GetNodeTitleColor() const
 {
 	FLinearColor ReturnTitleColor = GraphNode->IsDeprecated() ? FLinearColor::Red : GetNodeObj()->GetNodeTitleColor();
 
-	if(!GraphNode->bIsNodeEnabled)
+	if(!GraphNode->IsNodeEnabled())
 	{
 		ReturnTitleColor *= FLinearColor(0.5f, 0.5f, 0.5f, 0.4f);
 	}
@@ -537,7 +537,7 @@ FSlateColor SGraphNode::GetNodeTitleColor() const
 FSlateColor SGraphNode::GetNodeBodyColor() const
 {
 	FLinearColor ReturnBodyColor = FLinearColor::White;
-	if(!GraphNode->bIsNodeEnabled)
+	if(!GraphNode->IsNodeEnabled())
 	{
 		ReturnBodyColor *= FLinearColor(1.0f, 1.0f, 1.0f, 0.5f); 
 	}
@@ -547,7 +547,7 @@ FSlateColor SGraphNode::GetNodeBodyColor() const
 FSlateColor SGraphNode::GetNodeTitleIconColor() const
 {
 	FLinearColor ReturnIconColor = IconColor;
-	if(!GraphNode->bIsNodeEnabled)
+	if(!GraphNode->IsNodeEnabled())
 	{
 		ReturnIconColor *= FLinearColor(1.0f, 1.0f, 1.0f, 0.3f); 
 	}
@@ -557,7 +557,7 @@ FSlateColor SGraphNode::GetNodeTitleIconColor() const
 FLinearColor SGraphNode::GetNodeTitleTextColor() const
 {
 	FLinearColor ReturnTextColor = FLinearColor::White;
-	if(!GraphNode->bIsNodeEnabled)
+	if(!GraphNode->IsNodeEnabled())
 	{
 		ReturnTextColor *= FLinearColor(1.0f, 1.0f, 1.0f, 0.3f); 
 	}
diff --git a/Engine/Source/Editor/GraphEditor/Private/SGraphPin.cpp b/Engine/Source/Editor/GraphEditor/Private/SGraphPin.cpp
--- a/Engine/Source/Editor/GraphEditor/Private/SGraphPin.cpp
+++ b/Engine/Source/Editor/GraphEditor/Private/SGraphPin.cpp
@@ -891,7 +891,7 @@ FSlateColor SGraphPin::GetPinColor() const
 	}
 	if (const UEdGraphSchema* Schema = GraphPinObj->GetSchema())
 	{
-		if(!GetPinObj()->GetOwningNode()->bIsNodeEnabled)
+		if(!GetPinObj()->GetOwningNode()->IsNodeEnabled())
 		{
 			return Schema->GetPinTypeColor(GraphPinObj->PinType) * FLinearColor(1.0f, 1.0f, 1.0f, 0.5f);
 		}
@@ -907,7 +907,7 @@ FSlateColor SGraphPin::GetPinTextColor() const
 	// If there is no schema there is no owning node (or basically this is a deleted node)
 	if (UEdGraphNode* GraphNode = GraphPinObj->GetOwningNodeUnchecked())
 	{
-		if(!GraphNode->bIsNodeEnabled)
+		if(!GraphNode->IsNodeEnabled())
 		{
 			return FLinearColor(1.0f, 1.0f, 1.0f, 0.5f);
 		}
diff --git a/Engine/Source/Editor/GraphEditor/Public/GraphEditorActions.h b/Engine/Source/Editor/GraphEditor/Public/GraphEditorActions.h
--- a/Engine/Source/Editor/GraphEditor/Public/GraphEditorActions.h
+++ b/Engine/Source/Editor/GraphEditor/Public/GraphEditorActions.h
@@ -71,6 +71,12 @@ public:
 
 	TSharedPtr< FUICommandInfo > DistributeNodesHorizontally;
 	TSharedPtr< FUICommandInfo > DistributeNodesVertically;
+	
+	// Enable/disable commands
+	TSharedPtr< FUICommandInfo > EnableNodes;
+	TSharedPtr< FUICommandInfo > DisableNodes;
+	TSharedPtr< FUICommandInfo > EnableNodes_Always;
+	TSharedPtr< FUICommandInfo > EnableNodes_DevelopmentOnly;
 
 	//
 	TSharedPtr< FUICommandInfo > SelectReferenceInLevel;
diff --git a/Engine/Source/Editor/Kismet/Private/BlueprintDetailsCustomization.cpp b/Engine/Source/Editor/Kismet/Private/BlueprintDetailsCustomization.cpp
--- a/Engine/Source/Editor/Kismet/Private/BlueprintDetailsCustomization.cpp
+++ b/Engine/Source/Editor/Kismet/Private/BlueprintDetailsCustomization.cpp
@@ -8,6 +8,7 @@
 #include "PropertyRestriction.h"
 #include "BlueprintEditor.h"
 #include "BlueprintEditorModes.h"
+#include "BlueprintEditorSettings.h"
 #include "Editor/PropertyEditor/Public/PropertyEditing.h"
 #include "SColorPicker.h"
 #include "SKismetInspector.h"
@@ -4534,6 +4535,14 @@ void FBlueprintGlobalOptionsDetails::CustomizeDetails(IDetailLayoutBuilder& Deta
 		static FName DeprecatePropName(TEXT("bDeprecate"));
 		DetailLayout.HideProperty(DetailLayout.GetProperty(DeprecatePropName));
 
+		// Hide the experimental CompileMode setting (if not enabled)
+		const UBlueprintEditorSettings* EditorSettings = GetDefault<UBlueprintEditorSettings>();
+		if (EditorSettings && !EditorSettings->bAllowExplicitImpureNodeDisabling)
+		{
+			static FName CompileModePropertyName(TEXT("CompileMode"));
+			DetailLayout.HideProperty(DetailLayout.GetProperty(CompileModePropertyName));
+		}
+
 		// Hide 'run on drag' for LevelBP
 		if (bIsLevelScriptBP)
 		{
diff --git a/Engine/Source/Editor/Kismet/Private/BlueprintEditor.cpp b/Engine/Source/Editor/Kismet/Private/BlueprintEditor.cpp
--- a/Engine/Source/Editor/Kismet/Private/BlueprintEditor.cpp
+++ b/Engine/Source/Editor/Kismet/Private/BlueprintEditor.cpp
@@ -302,7 +302,7 @@ static bool BlueprintEditorImpl::GraphHasDefaultNode(UEdGraph const* InGraph)
 			continue;
 		}
 
-		if (Node->GetOutermost()->GetMetaData()->HasValue(Node, FNodeMetadata::DefaultGraphNode) && Node->bIsNodeEnabled)
+		if (Node->GetOutermost()->GetMetaData()->HasValue(Node, FNodeMetadata::DefaultGraphNode) && Node->IsNodeEnabled())
 		{
 			bHasDefaultNodes = true;
 			break;
@@ -1184,6 +1184,30 @@ TSharedRef<SGraphEditor> FBlueprintEditor::CreateGraphEditorWidget(TSharedRef<FT
 				FExecuteAction::CreateSP(this, &FBlueprintEditor::OnGoToDocumentation),
 				FCanExecuteAction::CreateSP(this, &FBlueprintEditor::CanGoToDocumentation)
 				);
+
+			GraphEditorCommands->MapAction(FGraphEditorCommands::Get().EnableNodes,
+				FExecuteAction(),
+				FCanExecuteAction(),
+				FGetActionCheckState::CreateSP(this, &FBlueprintEditor::GetEnabledCheckBoxStateForSelectedNodes)
+				);
+
+			GraphEditorCommands->MapAction(FGraphEditorCommands::Get().DisableNodes,
+				FExecuteAction::CreateSP(this, &FBlueprintEditor::OnSetEnabledStateForSelectedNodes, ENodeEnabledState::Disabled),
+				FCanExecuteAction(),
+				FGetActionCheckState::CreateSP(this, &FBlueprintEditor::CheckEnabledStateForSelectedNodes, ENodeEnabledState::Disabled)
+				);
+
+			GraphEditorCommands->MapAction(FGraphEditorCommands::Get().EnableNodes_Always,
+				FExecuteAction::CreateSP(this, &FBlueprintEditor::OnSetEnabledStateForSelectedNodes, ENodeEnabledState::Enabled),
+				FCanExecuteAction(),
+				FGetActionCheckState::CreateSP(this, &FBlueprintEditor::CheckEnabledStateForSelectedNodes, ENodeEnabledState::Enabled)
+				);
+
+			GraphEditorCommands->MapAction(FGraphEditorCommands::Get().EnableNodes_DevelopmentOnly,
+				FExecuteAction::CreateSP(this, &FBlueprintEditor::OnSetEnabledStateForSelectedNodes, ENodeEnabledState::DevelopmentOnly),
+				FCanExecuteAction(),
+				FGetActionCheckState::CreateSP(this, &FBlueprintEditor::CheckEnabledStateForSelectedNodes, ENodeEnabledState::DevelopmentOnly)
+				);
 		}
 	}
 
@@ -5633,8 +5657,9 @@ void FBlueprintEditor::PasteNodesHere(class UEdGraph* DestinationGraph, const FV
 				// Check if the nodes are identical, if they are we need to delete the original because it is disabled. Identical nodes that are in an enabled state will never make it this far and still be enabled.
 				if(bIdenticalNode)
 				{
-					// Should not have made it to being a pasted node if the pre-existing node wasn't disabled.
-					ensure(!ExistingEventNode->bIsNodeEnabled);
+					// Should not have made it to being a pasted node if the pre-existing node wasn't disabled or was otherwise explicitly disabled by the user.
+					ensure(!ExistingEventNode->IsNodeEnabled());
+					ensure(!ExistingEventNode->bUserSetEnabledState);
 
 					// Destroy the pre-existing node, we do not need it.
 					ExistingEventNode->DestroyNode();
@@ -6055,6 +6080,77 @@ bool FBlueprintEditor::CanGoToDocumentation()
 	return !DocumentationLink.IsEmpty();
 }
 
+void FBlueprintEditor::OnSetEnabledStateForSelectedNodes(ENodeEnabledState NewState)
+{
+	const FScopedTransaction Transaction(LOCTEXT("SetNodeEnabledState", "Set Node Enabled State"));
+
+	const FGraphPanelSelectionSet SelectedNodes = GetSelectedNodes();
+	for (auto SelectedNode : SelectedNodes)
+	{
+		UEdGraphNode* SelectedGraphNode = Cast<UEdGraphNode>(SelectedNode);
+		if(SelectedGraphNode)
+		{
+			SelectedGraphNode->Modify();
+			SelectedGraphNode->EnabledState = NewState;
+			SelectedGraphNode->bUserSetEnabledState = true;
+		}
+	}
+
+	FBlueprintEditorUtils::MarkBlueprintAsStructurallyModified(GetBlueprintObj());
+}
+
+ECheckBoxState FBlueprintEditor::GetEnabledCheckBoxStateForSelectedNodes()
+{
+	const FGraphPanelSelectionSet SelectedNodes = GetSelectedNodes();
+	ECheckBoxState Result = SelectedNodes.Num() > 0 ? ECheckBoxState::Undetermined : ECheckBoxState::Unchecked;
+	for (auto SelectedNode : SelectedNodes)
+	{
+		UEdGraphNode* SelectedGraphNode = Cast<UEdGraphNode>(SelectedNode);
+		if(SelectedGraphNode)
+		{
+			const bool bIsSelectedNodeEnabled = SelectedGraphNode->IsNodeEnabled();
+			if(Result == ECheckBoxState::Undetermined)
+			{
+				Result = bIsSelectedNodeEnabled ? ECheckBoxState::Checked : ECheckBoxState::Unchecked;
+			}
+			else if((!bIsSelectedNodeEnabled && Result == ECheckBoxState::Checked)
+				|| (bIsSelectedNodeEnabled && Result == ECheckBoxState::Unchecked))
+			{
+				Result = ECheckBoxState::Undetermined;
+				break;
+			}
+		}
+	}
+
+	return Result;
+}
+
+ECheckBoxState FBlueprintEditor::CheckEnabledStateForSelectedNodes(ENodeEnabledState CheckState)
+{
+	const FGraphPanelSelectionSet SelectedNodes = GetSelectedNodes();
+	ECheckBoxState Result = SelectedNodes.Num() > 0 ? ECheckBoxState::Undetermined : ECheckBoxState::Unchecked;
+	for (auto SelectedNode : SelectedNodes)
+	{
+		UEdGraphNode* SelectedGraphNode = Cast<UEdGraphNode>(SelectedNode);
+		if(SelectedGraphNode)
+		{
+			const ENodeEnabledState NodeState = SelectedGraphNode->EnabledState;
+			if(Result == ECheckBoxState::Undetermined)
+			{
+				Result = (NodeState == CheckState) ? ECheckBoxState::Checked : ECheckBoxState::Unchecked;
+			}
+			else if((NodeState != CheckState && Result == ECheckBoxState::Checked)
+				|| (NodeState == CheckState && Result == ECheckBoxState::Unchecked))
+			{
+				Result = ECheckBoxState::Undetermined;
+				break;
+			}
+		}
+	}
+
+	return Result;
+}
+
 void FBlueprintEditor::ToggleSaveIntermediateBuildProducts()
 {
 	bSaveIntermediateBuildProducts = !bSaveIntermediateBuildProducts;
diff --git a/Engine/Source/Editor/Kismet/Public/BlueprintEditor.h b/Engine/Source/Editor/Kismet/Public/BlueprintEditor.h
--- a/Engine/Source/Editor/Kismet/Public/BlueprintEditor.h
+++ b/Engine/Source/Editor/Kismet/Public/BlueprintEditor.h
@@ -965,9 +965,20 @@ private:
 	/** Util to try and get doc link for the currently selected node */
 	FString GetDocLinkForSelectedNode();
 
+	/** Set the enabled state for currently-selected nodes */
+	void OnSetEnabledStateForSelectedNodes(ENodeEnabledState NewState);
+
+	/** Returns the appropriate check box state representing whether or not the selected nodes are enabled */
+	ECheckBoxState GetEnabledCheckBoxStateForSelectedNodes();
+
+	/** Attempt to match the given enabled state for currently-selected nodes */
+	ECheckBoxState CheckEnabledStateForSelectedNodes(ENodeEnabledState CheckState);
+
 	/** Fixes SubObject references of the passed object so they match up to sub-object UProperty references */
 	void FixSubObjectReferencesPostUndoRedo(UObject* InObject);
 
+
+
 public://@TODO
 	TSharedPtr<FDocumentTracker> DocumentManager;
 
diff --git a/Engine/Source/Editor/Persona/Private/PersonaModule.cpp b/Engine/Source/Editor/Persona/Private/PersonaModule.cpp
--- a/Engine/Source/Editor/Persona/Private/PersonaModule.cpp
+++ b/Engine/Source/Editor/Persona/Private/PersonaModule.cpp
@@ -114,7 +114,7 @@ void FPersonaModule::OnNewBlueprintCreated(UBlueprint* InBlueprint)
 		GetOwnerNode->NodePosX = SafeXPosition;
 		GetOwnerNode->NodePosY = SafeYPosition;
 		UEdGraphSchema_K2::SetNodeMetaData(GetOwnerNode, FNodeMetadata::DefaultGraphNode);
-		GetOwnerNode->bIsNodeEnabled = false;
+		GetOwnerNode->DisableNode();
 
 		EventGraph->AddNode(GetOwnerNode);
 	}
diff --git a/Engine/Source/Editor/UnrealEd/Private/EdGraphUtilities.cpp b/Engine/Source/Editor/UnrealEd/Private/EdGraphUtilities.cpp
--- a/Engine/Source/Editor/UnrealEd/Private/EdGraphUtilities.cpp
+++ b/Engine/Source/Editor/UnrealEd/Private/EdGraphUtilities.cpp
@@ -179,6 +179,17 @@ UEdGraph* FEdGraphUtilities::CloneGraph(UEdGraph* InSource, UObject* NewOuter, F
 			UObject* const Dest = It.Value();
 
 			MessageLog->NotifyIntermediateObjectCreation(Dest, Source);
+
+			// During compilation, set cloned nodes to a non-conditional enabled state.
+			if (bCloningForCompile)
+			{
+				const UEdGraphNode* SrcNode = Cast<UEdGraphNode>(Source);
+				UEdGraphNode* DstNode = Cast<UEdGraphNode>(Dest);
+				if(SrcNode && DstNode)
+				{
+					DstNode->EnabledState = SrcNode->IsNodeEnabled() ? ENodeEnabledState::Enabled : ENodeEnabledState::Disabled;
+				}
+			}
 		}
 	}
 
diff --git a/Engine/Source/Editor/UnrealEd/Private/Kismet2/BlueprintEditorUtils.cpp b/Engine/Source/Editor/UnrealEd/Private/Kismet2/BlueprintEditorUtils.cpp
--- a/Engine/Source/Editor/UnrealEd/Private/Kismet2/BlueprintEditorUtils.cpp
+++ b/Engine/Source/Editor/UnrealEd/Private/Kismet2/BlueprintEditorUtils.cpp
@@ -2979,8 +2979,8 @@ bool FBlueprintEditorUtils::IsDataOnlyBlueprint(const UBlueprint* Blueprint)
 	{
 		for(UEdGraphNode* GraphNode : EventGraph->Nodes)
 		{
-			// If there is an enabled node in the event graph, the Blueprint is not data only
-			if (GraphNode && GraphNode->bIsNodeEnabled)
+			// If there is an enabled node in the event graph (or a node that was explicitly disabled by the user), the Blueprint is not data only
+			if (GraphNode && (GraphNode->IsNodeEnabled() || GraphNode->bUserSetEnabledState))
 			{
 				return false;
 			}
diff --git a/Engine/Source/Editor/UnrealEd/Private/Kismet2/Kismet2.cpp b/Engine/Source/Editor/UnrealEd/Private/Kismet2/Kismet2.cpp
--- a/Engine/Source/Editor/UnrealEd/Private/Kismet2/Kismet2.cpp
+++ b/Engine/Source/Editor/UnrealEd/Private/Kismet2/Kismet2.cpp
@@ -584,7 +584,7 @@ UK2Node_Event* FKismetEditorUtilities::AddDefaultEventNode(UBlueprint* InBluepri
 		NewEventNode->PostPlacedNewNode();
 		NewEventNode->SetFlags(RF_Transactional);
 		NewEventNode->AllocateDefaultPins();
-		NewEventNode->bIsNodeEnabled = false;
+		NewEventNode->DisableNode();
 		NewEventNode->NodeComment = LOCTEXT("DisabledNodeComment", "This node is disabled and will not be called.\nDrag off pins to build functionality.").ToString();
 		NewEventNode->bCommentBubblePinned = true;
 		NewEventNode->bCommentBubbleVisible = true;
@@ -611,10 +611,10 @@ UK2Node_Event* FKismetEditorUtilities::AddDefaultEventNode(UBlueprint* InBluepri
 			UEdGraphSchema_K2::SetNodeMetaData(ParentFunctionNode, FNodeMetadata::DefaultGraphNode);
 			FunctionNodeCreator.Finalize();
 
-			ParentFunctionNode->bIsNodeEnabled = false;
+			ParentFunctionNode->DisableNode();
 
 			// Adding the call to parent and connecting it will reset this value
-			NewEventNode->bIsNodeEnabled = false;
+			NewEventNode->DisableNode();
 			NewEventNode->NodeComment = LOCTEXT("DisabledNodeComment", "This node is disabled and will not be called.\nDrag off pins to build functionality.").ToString();
 		}
 	}
diff --git a/Engine/Source/Editor/UnrealEd/Public/CookerSettings.h b/Engine/Source/Editor/UnrealEd/Public/CookerSettings.h
--- a/Engine/Source/Editor/UnrealEd/Public/CookerSettings.h
+++ b/Engine/Source/Editor/UnrealEd/Public/CookerSettings.h
@@ -21,6 +21,10 @@ public:
 	UPROPERTY(GlobalConfig, EditAnywhere, Category = Cooker, meta = (DisplayName = "Enable cooking via network in the background of the editor, launch on uses this setting, requires device to have network access to editor"))
 	bool bEnableCookOnTheSide;
 
+	/** Whether or not to compile Blueprints in development mode when cooking. */
+	UPROPERTY(GlobalConfig, EditAnywhere, Category = Cooker, AdvancedDisplay)
+	bool bCompileBlueprintsInDevelopmentMode;
+
 
 	/** Quality of 0 means fastest, 4 means best quality */
 	UPROPERTY(GlobalConfig, EditAnywhere, Category = Textures, meta = (DisplayName = "PVRTC Compression Quality (0-4, 0 is fastest)"))
diff --git a/Engine/Source/Runtime/Engine/Classes/EdGraph/EdGraphNode.h b/Engine/Source/Runtime/Engine/Classes/EdGraph/EdGraphNode.h
--- a/Engine/Source/Runtime/Engine/Classes/EdGraph/EdGraphNode.h
+++ b/Engine/Source/Runtime/Engine/Classes/EdGraph/EdGraphNode.h
@@ -55,6 +55,18 @@ namespace ENodeAdvancedPins
 	};
 }
 
+/** Enum to indicate a node's enabled state. */
+UENUM()
+enum class ENodeEnabledState : uint8
+{
+	/** Node is enabled. */
+	Enabled,
+	/** Node is disabled. */
+	Disabled,
+	/** Node is enabled for development only. */
+	DevelopmentOnly
+};
+
 /** Holds metadata keys, so as to discourage text duplication throughout the engine. */
 struct ENGINE_API FNodeMetadata
 {
@@ -152,9 +164,42 @@ class ENGINE_API UEdGraphNode : public UObject
 	UPROPERTY()
 	TEnumAsByte<ENodeAdvancedPins::Type> AdvancedPinDisplay;
 
-	/** FALSE if the node is a disabled, which eliminates it from being compiled */
+	/** Indicates in what state the node is enabled, which may eliminate it from being compiled */
+	UPROPERTY()
+	ENodeEnabledState EnabledState;
+
+	/** Indicates whether or not the user explicitly set the enabled state */
 	UPROPERTY()
-	bool bIsNodeEnabled;
+	bool bUserSetEnabledState;
+
+private:
+	/** (DEPRECATED) FALSE if the node is a disabled, which eliminates it from being compiled */
+	UPROPERTY()
+	bool bIsNodeEnabled_DEPRECATED;
+
+public:
+	/** Enables this node. */
+	FORCEINLINE void EnableNode()
+	{
+		bUserSetEnabledState = false;
+		EnabledState = ENodeEnabledState::Enabled;
+	}
+
+	/** Disables this node. */
+	FORCEINLINE void DisableNode()
+	{
+		bUserSetEnabledState = false;
+		EnabledState = ENodeEnabledState::Disabled;
+	}
+
+	/** Determines whether or not the node is enabled. */
+	FORCEINLINE bool IsNodeEnabled() const
+	{
+		return (EnabledState == ENodeEnabledState::Enabled) || (EnabledState == ENodeEnabledState::DevelopmentOnly && IsInDevelopmentMode());
+	}
+
+	/** Determines whether or not the node will compile in development mode. */
+	virtual bool IsInDevelopmentMode() const;
 
 #if WITH_EDITOR
 
diff --git a/Engine/Source/Runtime/Engine/Classes/Engine/Blueprint.h b/Engine/Source/Runtime/Engine/Classes/Engine/Blueprint.h
--- a/Engine/Source/Runtime/Engine/Classes/Engine/Blueprint.h
+++ b/Engine/Source/Runtime/Engine/Classes/Engine/Blueprint.h
@@ -64,6 +64,14 @@ namespace EKismetCompileType
 	};
 };
 
+/** Compile modes. */
+UENUM()
+enum class EBlueprintCompileMode : uint8
+{
+	Default UMETA(DisplayName="Use Default", ToolTip="Use the default setting."),
+	Development UMETA(ToolTip="Always compile in development mode (even when cooking)."),
+	FinalRelease UMETA(ToolTip="Always compile in final release mode.")
+};
 
 struct FKismetCompilerOptions
 {
@@ -331,6 +339,10 @@ class ENGINE_API UBlueprint : public UBlueprintCore
 	/** Deprecates the Blueprint, marking the generated class with the CLASS_Deprecated flag */
 	UPROPERTY(EditAnywhere, Category=ClassOptions, AdvancedDisplay)
 	bool bDeprecate;
+
+	/** The mode that will be used when compiling this class. */
+	UPROPERTY(EditAnywhere, Category=ClassOptions, AdvancedDisplay)
+	EBlueprintCompileMode CompileMode;
 #endif //WITH_EDITORONLY_DATA
 
 	/** 'Simple' construction script - graph of components to instance */
diff --git a/Engine/Source/Runtime/Engine/Classes/Kismet/KismetSystemLibrary.h b/Engine/Source/Runtime/Engine/Classes/Kismet/KismetSystemLibrary.h
--- a/Engine/Source/Runtime/Engine/Classes/Kismet/KismetSystemLibrary.h
+++ b/Engine/Source/Runtime/Engine/Classes/Kismet/KismetSystemLibrary.h
@@ -220,7 +220,7 @@ class ENGINE_API UKismetSystemLibrary : public UBlueprintFunctionLibrary
 	 * @param	TextColor		Whether or not to print the output to the console
 	 * @param	Duration		The display duration (if Print to Screen is True). Using negative number will result in loading the duration time from the config.
 	 */
-	UFUNCTION(BlueprintCallable, meta=(WorldContext="WorldContextObject", CallableWithoutWorldContext, Keywords = "log print", AdvancedDisplay = "2"), Category="Utilities|String")
+	UFUNCTION(BlueprintCallable, meta=(WorldContext="WorldContextObject", CallableWithoutWorldContext, Keywords = "log print", AdvancedDisplay = "2", DevelopmentOnly), Category="Utilities|String")
 	static void PrintString(UObject* WorldContextObject, const FString& InString = FString(TEXT("Hello")), bool bPrintToScreen = true, bool bPrintToLog = true, FLinearColor TextColor = FLinearColor(0.0, 0.66, 1.0), float Duration = 2.f);
 
 	/**
@@ -234,7 +234,7 @@ class ENGINE_API UKismetSystemLibrary : public UBlueprintFunctionLibrary
 	 * @param	TextColor		Whether or not to print the output to the console
 	 * @param	Duration		The display duration (if Print to Screen is True). Using negative number will result in loading the duration time from the config.
 	 */
-	UFUNCTION(BlueprintCallable, meta=(WorldContext="WorldContextObject", CallableWithoutWorldContext, Keywords = "log", AdvancedDisplay = "2"), Category="Utilities|Text")
+	UFUNCTION(BlueprintCallable, meta=(WorldContext="WorldContextObject", CallableWithoutWorldContext, Keywords = "log", AdvancedDisplay = "2", DevelopmentOnly), Category="Utilities|Text")
 	static void PrintText(UObject* WorldContextObject, const FText InText = FText::FromString(TEXT("Hello")), bool bPrintToScreen = true, bool bPrintToLog = true, FLinearColor TextColor = FLinearColor(0.0, 0.66, 1.0), float Duration = 2.f);
 
 	/**
diff --git a/Engine/Source/Runtime/Engine/Private/EdGraph/EdGraph.cpp b/Engine/Source/Runtime/Engine/Private/EdGraph/EdGraph.cpp
--- a/Engine/Source/Runtime/Engine/Private/EdGraph/EdGraph.cpp
+++ b/Engine/Source/Runtime/Engine/Private/EdGraph/EdGraph.cpp
@@ -213,8 +213,28 @@ void UEdGraph::MoveNodesToAnotherGraph(UEdGraph* DestinationGraph, bool bIsLoadi
 		{
 #if WITH_EDITOR
 			// During compilation, do not move ghost nodes, they are not used during compilation.
-			if (bInIsCompiling && !Node->bIsNodeEnabled)
+			if (bInIsCompiling && !Node->IsNodeEnabled())
 			{
+				// Pass existing connections through non-enabled nodes
+				for (auto Pin : Node->Pins)
+				{
+					if (Pin->Direction == EGPD_Input && Pin->LinkedTo.Num() > 0)
+					{
+						UEdGraphPin* PassThroughPin = Node->GetPassThroughPin(Pin);
+						if (PassThroughPin != nullptr && PassThroughPin->LinkedTo.Num() > 0)
+						{
+							for (auto OutputPin : Pin->LinkedTo)
+							{
+								for (auto InputPin : PassThroughPin->LinkedTo)
+								{
+									InputPin->LinkedTo.Add(OutputPin);
+									OutputPin->LinkedTo.Add(InputPin);
+								}
+							}
+						}
+					}
+				}
+
 				// Break all node links, if any exist, do not move the node
 				Node->BreakAllNodeLinks();
 				continue;
diff --git a/Engine/Source/Runtime/Engine/Private/EdGraph/EdGraphNode.cpp b/Engine/Source/Runtime/Engine/Private/EdGraph/EdGraphNode.cpp
--- a/Engine/Source/Runtime/Engine/Private/EdGraph/EdGraphNode.cpp
+++ b/Engine/Source/Runtime/Engine/Private/EdGraph/EdGraphNode.cpp
@@ -4,6 +4,7 @@
 #include "EdGraph/EdGraph.h"
 #include "BlueprintUtilities.h"
 #if WITH_EDITOR
+#include "Editor/UnrealEd/Public/CookerSettings.h"
 #include "Editor/UnrealEd/Public/Kismet2/BlueprintEditorUtils.h"
 #include "SlateBasics.h"
 #include "ScopedTransaction.h"
@@ -42,7 +43,9 @@ FGraphNodeContextMenuBuilder::FGraphNodeContextMenuBuilder(const UEdGraph* InGra
 UEdGraphNode::UEdGraphNode(const FObjectInitializer& ObjectInitializer)
 	: Super(ObjectInitializer)
 	, AdvancedPinDisplay(ENodeAdvancedPins::NoPins)
-	, bIsNodeEnabled(true)
+	, EnabledState(ENodeEnabledState::Enabled)
+	, bUserSetEnabledState(false)
+	, bIsNodeEnabled_DEPRECATED(true)
 {
 
 #if WITH_EDITORONLY_DATA
@@ -278,6 +281,12 @@ void UEdGraphNode::PostLoad()
 	{
 		bCommentBubbleVisible = !NodeComment.IsEmpty();
 	}
+
+	// If this was an older version, ensure that we update the enabled state to match
+	if(!bIsNodeEnabled_DEPRECATED)
+	{
+		EnabledState = ENodeEnabledState::Disabled;
+	}
 }
 
 void UEdGraphNode::PostEditChangeProperty(FPropertyChangedEvent& PropertyChangedEvent)
@@ -394,6 +403,16 @@ FText UEdGraphNode::GetKeywords() const
 
 #endif	//#if WITH_EDITOR
 
+bool UEdGraphNode::IsInDevelopmentMode() const
+{
+#if WITH_EDITOR
+	// By default, development mode is implied when running in the editor and not cooking via commandlet, unless enabled in the project settings.
+	return !IsRunningCommandlet() || GetDefault<UCookerSettings>()->bCompileBlueprintsInDevelopmentMode;
+#else
+	return false;
+#endif
+}
+
 /////////////////////////////////////////////////////
 
 #undef LOCTEXT_NAMESPACE
diff --git a/Engine/Source/Runtime/Engine/Private/EdGraph/EdGraphPin.cpp b/Engine/Source/Runtime/Engine/Private/EdGraph/EdGraphPin.cpp
--- a/Engine/Source/Runtime/Engine/Private/EdGraph/EdGraphPin.cpp
+++ b/Engine/Source/Runtime/Engine/Private/EdGraph/EdGraphPin.cpp
@@ -81,11 +81,12 @@ namespace GraphPinHelpers
 {
 	void EnableAllConnectedNodes(UEdGraphNode* InNode)
 	{
-		if(InNode && !InNode->bIsNodeEnabled)
+		// Only enable when it has not been explicitly user-disabled
+		if(InNode && InNode->EnabledState == ENodeEnabledState::Disabled && !InNode->bUserSetEnabledState)
 		{
 			// Enable the node and clear the comment
 			InNode->Modify();
-			InNode->bIsNodeEnabled = true;
+			InNode->EnableNode();
 			InNode->NodeComment.Empty();
 
 			// Go through all pin connections and enable the nodes. Enabled nodes will prevent further iteration
