Further optimizations to blueprint context menu constrution. Avoiding copying strings, recalculation of categories.

diff --git a/Engine/Source/Editor/AnimGraph/Private/AnimGraphNode_BlendSpaceBase.cpp b/Engine/Source/Editor/AnimGraph/Private/AnimGraphNode_BlendSpaceBase.cpp
--- a/Engine/Source/Editor/AnimGraph/Private/AnimGraphNode_BlendSpaceBase.cpp
+++ b/Engine/Source/Editor/AnimGraph/Private/AnimGraphNode_BlendSpaceBase.cpp
@@ -15,31 +15,30 @@
 struct FNewBlendSpacePlayerAction : public FEdGraphSchemaAction_K2NewNode
 {
 	FNewBlendSpacePlayerAction(class UBlendSpaceBase* BlendSpace)
+		: FEdGraphSchemaAction_K2NewNode()
 	{
 		check(BlendSpace);
 
 		const bool bIsAimOffset = BlendSpace->IsA(UAimOffsetBlendSpace::StaticClass()) || BlendSpace->IsA(UAimOffsetBlendSpace1D::StaticClass());
+		FString NewTooltipDescription;
 		if (bIsAimOffset)
 		{
 			UAnimGraphNode_RotationOffsetBlendSpace* Template = NewObject<UAnimGraphNode_RotationOffsetBlendSpace>();
 			Template->Node.BlendSpace = BlendSpace;
 			NodeTemplate = Template;
-			TooltipDescription = TEXT("Evaluates an aim offset at a particular coordinate to produce a pose");
+			NewTooltipDescription = TEXT("Evaluates an aim offset at a particular coordinate to produce a pose");
 		}
 		else
 		{
 			UAnimGraphNode_BlendSpacePlayer* Template = NewObject<UAnimGraphNode_BlendSpacePlayer>();
 			Template->Node.BlendSpace = BlendSpace;
 			NodeTemplate = Template;
-			TooltipDescription = TEXT("Evaluates a blend space at a particular coordinate to produce a pose");
+			NewTooltipDescription = TEXT("Evaluates a blend space at a particular coordinate to produce a pose");
 		}
 
-		MenuDescription = NodeTemplate->GetNodeTitle(ENodeTitleType::ListView);
+		FText NewMenuDescription = NodeTemplate->GetNodeTitle(ENodeTitleType::ListView);
 
-		Category = LOCTEXT("Animation", "Animations");
-
-		// Grab extra keywords
-		Keywords = FText::FromString(BlendSpace->GetPathName());
+		UpdateSearchData(NewMenuDescription, NewTooltipDescription, LOCTEXT("Animation", "Animations"), FText::FromString(BlendSpace->GetPathName()));
 	}
 };
 
diff --git a/Engine/Source/Editor/AnimGraph/Private/AnimGraphNode_SequencePlayer.cpp b/Engine/Source/Editor/AnimGraph/Private/AnimGraphNode_SequencePlayer.cpp
--- a/Engine/Source/Editor/AnimGraph/Private/AnimGraphNode_SequencePlayer.cpp
+++ b/Engine/Source/Editor/AnimGraph/Private/AnimGraphNode_SequencePlayer.cpp
@@ -24,18 +24,12 @@ protected:
 	FAssetData AssetInfo;
 public:
 	FNewSequencePlayerAction(const FAssetData& InAssetInfo, FText Title)
+		: FEdGraphSchemaAction_K2NewNode(LOCTEXT("Animation", "Animations"), Title, TEXT("Evaluates an animation sequence to produce a pose"), 0, FText::FromName(InAssetInfo.ObjectPath))
 	{
 		AssetInfo = InAssetInfo;
 
 		UAnimGraphNode_SequencePlayer* Template = NewObject<UAnimGraphNode_SequencePlayer>();
 		NodeTemplate = Template;
-
-		MenuDescription = Title;
-		TooltipDescription = TEXT("Evaluates an animation sequence to produce a pose");
-		Category = LOCTEXT("Animation", "Animations");
-
-		// Grab extra keywords
-		Keywords = FText::FromName(InAssetInfo.ObjectPath);
 	}
 
 	virtual UEdGraphNode* PerformAction(class UEdGraph* ParentGraph, UEdGraphPin* FromPin, const FVector2D Location, bool bSelectNewNode = true) override
diff --git a/Engine/Source/Editor/BehaviorTreeEditor/Private/SBehaviorTreeBlackboardView.cpp b/Engine/Source/Editor/BehaviorTreeEditor/Private/SBehaviorTreeBlackboardView.cpp
--- a/Engine/Source/Editor/BehaviorTreeEditor/Private/SBehaviorTreeBlackboardView.cpp
+++ b/Engine/Source/Editor/BehaviorTreeEditor/Private/SBehaviorTreeBlackboardView.cpp
@@ -50,9 +50,8 @@ FEdGraphSchemaAction_BlackboardEntry::FEdGraphSchemaAction_BlackboardEntry( UBla
 
 void FEdGraphSchemaAction_BlackboardEntry::Update()
 {
-	MenuDescription = FText::FromName(Key.EntryName);
-	TooltipDescription = FText::Format(LOCTEXT("BlackboardEntryFormat", "{0} '{1}'"), Key.KeyType ? Key.KeyType->GetClass()->GetDisplayNameText() : LOCTEXT("NullKeyDesc", "None"), FText::FromName(Key.EntryName)).ToString();
-	SectionID = bIsInherited ? EBlackboardSectionTitles::InheritedKeys : EBlackboardSectionTitles::Keys;	
+	UpdateSearchData(FText::FromName(Key.EntryName), FText::Format(LOCTEXT("BlackboardEntryFormat", "{0} '{1}'"), Key.KeyType ? Key.KeyType->GetClass()->GetDisplayNameText() : LOCTEXT("NullKeyDesc", "None"), FText::FromName(Key.EntryName)).ToString(), FText(), FText());
+	SectionID = bIsInherited ? EBlackboardSectionTitles::InheritedKeys : EBlackboardSectionTitles::Keys;
 }
 
 class SBehaviorTreeBlackboardItem : public SGraphPaletteItem
@@ -87,7 +86,7 @@ class SBehaviorTreeBlackboardItem : public SGraphPaletteItem
 		FSlateBrush const* IconBrush   = FEditorStyle::GetBrush(TEXT("NoBrush"));
 		GetPaletteItemIcon(GraphAction, IconBrush);
 
-		TSharedRef<SWidget> IconWidget = CreateIconWidget( FText::FromString(GraphAction->TooltipDescription), IconBrush, FLinearColor::White );
+		TSharedRef<SWidget> IconWidget = CreateIconWidget( FText::FromString(GraphAction->GetTooltipDescription()), IconBrush, FLinearColor::White );
 		TSharedRef<SWidget> NameSlotWidget = CreateTextSlotWidget( NameFont, InCreateData, BlackboardEntryAction->bIsInherited );
 		TSharedRef<SWidget> DebugSlotWidget = CreateDebugSlotWidget( NameFont );
 
@@ -171,7 +170,7 @@ private:
 
 	virtual FText GetItemTooltip() const override
 	{
-		return FText::FromString(ActionPtr.Pin()->TooltipDescription);
+		return FText::FromString(ActionPtr.Pin()->GetTooltipDescription());
 	}
 
 	virtual void OnNameTextCommitted(const FText& NewText, ETextCommit::Type InTextCommit) override
diff --git a/Engine/Source/Editor/BlueprintGraph/Classes/EdGraphSchema_K2_Actions.h b/Engine/Source/Editor/BlueprintGraph/Classes/EdGraphSchema_K2_Actions.h
--- a/Engine/Source/Editor/BlueprintGraph/Classes/EdGraphSchema_K2_Actions.h
+++ b/Engine/Source/Editor/BlueprintGraph/Classes/EdGraphSchema_K2_Actions.h
@@ -31,8 +31,8 @@ struct BLUEPRINTGRAPH_API FEdGraphSchemaAction_K2NewNode : public FEdGraphSchema
 		, bGotoNode(false)
 	{}
 
-	FEdGraphSchemaAction_K2NewNode(const FText& InNodeCategory, const FText& InMenuDesc, const FString& InToolTip, const int32 InGrouping)
-		: FEdGraphSchemaAction(InNodeCategory, InMenuDesc, InToolTip, InGrouping) 
+	FEdGraphSchemaAction_K2NewNode(const FText& InNodeCategory, const FText& InMenuDesc, const FString& InToolTip, const int32 InGrouping, const FText& InKeywords = FText())
+		: FEdGraphSchemaAction(InNodeCategory, InMenuDesc, InToolTip, InGrouping, InKeywords)
 		, NodeTemplate(NULL)
 		, bGotoNode(false)
 	{}
@@ -349,17 +349,13 @@ struct FEdGraphSchemaAction_K2AddComment : public FEdGraphSchemaAction
 	virtual FName GetTypeId() const override { return StaticGetTypeId(); } 
 
 	FEdGraphSchemaAction_K2AddComment()
-		: FEdGraphSchemaAction()
+		: FEdGraphSchemaAction(FText(), NSLOCTEXT("K2AddComment", "AddComment", "Add Comment..."), NSLOCTEXT("K2AddComment", "AddComment_Tooltip", "Create a resizable comment box.").ToString(), 0)
 	{
-		MenuDescription = NSLOCTEXT("K2AddComment", "AddComment", "Add Comment...");
-		TooltipDescription = NSLOCTEXT("K2AddComment", "AddComment_Tooltip", "Create a resizable comment box.").ToString();
 	}
 
 	FEdGraphSchemaAction_K2AddComment(const FText& InDescription, const FString& InToolTip)
-		: FEdGraphSchemaAction()
+		: FEdGraphSchemaAction(FText(), InDescription, InToolTip, 0)
 	{
-		MenuDescription = InDescription;
-		TooltipDescription = InToolTip;
 	}
 
 	// FEdGraphSchemaAction interface
@@ -382,17 +378,13 @@ struct FEdGraphSchemaAction_K2AddDocumentation : public FEdGraphSchemaAction
 	virtual FName GetTypeId() const override { return StaticGetTypeId(); } 
 
 	FEdGraphSchemaAction_K2AddDocumentation()
-		: FEdGraphSchemaAction()
+		: FEdGraphSchemaAction(FText(), NSLOCTEXT("K2AddDocumentation", "AddDocumentation", "Add Documentation ..."), NSLOCTEXT("K2AddDocumentation", "AddDocumentation_Tooltip", "Creates a Documentation Node.").ToString(), 0)
 	{
-		MenuDescription = NSLOCTEXT("K2AddDocumentation", "AddDocumentation", "Add Documentation ...");
-		TooltipDescription = NSLOCTEXT("K2AddDocumentation", "AddDocumentation_Tooltip", "Creates a Documentation Node.").ToString();
 	}
 
 	FEdGraphSchemaAction_K2AddDocumentation(const FString& InDescription, const FString& InToolTip)
-		: FEdGraphSchemaAction()
+		: FEdGraphSchemaAction(FText(), FText::FromString(InDescription), InToolTip, 0)
 	{
-		MenuDescription = FText::FromString( InDescription );
-		TooltipDescription = InToolTip;
 	}
 	
 	// FEdGraphSchemaAction interface
@@ -531,8 +523,8 @@ public:
 		: FEdGraphSchemaAction()
 	{}
 
-	FEdGraphSchemaAction_K2Var (const FText& InNodeCategory, const FText& InMenuDesc, const FString& InToolTip, const int32 InGrouping)
-		: FEdGraphSchemaAction(InNodeCategory, InMenuDesc, InToolTip, InGrouping)
+	FEdGraphSchemaAction_K2Var (const FText& InNodeCategory, const FText& InMenuDesc, const FString& InToolTip, const int32 InGrouping, const int32 InSectionID)
+		: FEdGraphSchemaAction(InNodeCategory, InMenuDesc, InToolTip, InGrouping, FText(), InSectionID)
 	{}
 
 	FName GetVariableName() const
@@ -598,8 +590,8 @@ public:
 		: FEdGraphSchemaAction()
 	{}
 
-	FEdGraphSchemaAction_K2LocalVar (const FText& InNodeCategory, const FText& InMenuDesc, const FString& InToolTip, const int32 InGrouping)
-		: FEdGraphSchemaAction(InNodeCategory, InMenuDesc, InToolTip, InGrouping)
+	FEdGraphSchemaAction_K2LocalVar(const FText& InNodeCategory, const FText& InMenuDesc, const FString& InToolTip, const int32 InGrouping, const int32 InSectionID)
+		: FEdGraphSchemaAction(InNodeCategory, InMenuDesc, InToolTip, InGrouping, FText(), InSectionID)
 	{}
 
 	FName GetVariableName() const
@@ -669,8 +661,8 @@ struct BLUEPRINTGRAPH_API FEdGraphSchemaAction_K2Graph : public FEdGraphSchemaAc
 		: FEdGraphSchemaAction()
 	{}
 
-	FEdGraphSchemaAction_K2Graph (EEdGraphSchemaAction_K2Graph::Type InType, const FText& InNodeCategory, const FText& InMenuDesc, const FString& InToolTip, const int32 InGrouping)
-		: FEdGraphSchemaAction(InNodeCategory, InMenuDesc, InToolTip, InGrouping), GraphType(InType), EdGraph(NULL)
+	FEdGraphSchemaAction_K2Graph (EEdGraphSchemaAction_K2Graph::Type InType, const FText& InNodeCategory, const FText& InMenuDesc, const FString& InToolTip, const int32 InGrouping, const int32 InSectionID = 0)
+		: FEdGraphSchemaAction(InNodeCategory, InMenuDesc, InToolTip, InGrouping, FText(), InSectionID), GraphType(InType), EdGraph(NULL)
 	{}
 
 	// FEdGraphSchemaAction interface
@@ -815,8 +807,8 @@ public:
 		: FEdGraphSchemaAction(), EdGraph(NULL)
 	{}
 
-	FEdGraphSchemaAction_K2Delegate(const FText& InNodeCategory, const FText& InMenuDesc, const FString& InToolTip, const int32 InGrouping)
-		: FEdGraphSchemaAction(InNodeCategory, InMenuDesc, InToolTip, InGrouping), EdGraph(NULL)
+	FEdGraphSchemaAction_K2Delegate(const FText& InNodeCategory, const FText& InMenuDesc, const FString& InToolTip, const int32 InGrouping, const int32 InSectionID)
+		: FEdGraphSchemaAction(InNodeCategory, InMenuDesc, InToolTip, InGrouping, FText(), InSectionID), EdGraph(NULL)
 	{}
 
 	FName GetDelegateName() const
diff --git a/Engine/Source/Editor/BlueprintGraph/Private/BlueprintActionFilter.cpp b/Engine/Source/Editor/BlueprintGraph/Private/BlueprintActionFilter.cpp
--- a/Engine/Source/Editor/BlueprintGraph/Private/BlueprintActionFilter.cpp
+++ b/Engine/Source/Editor/BlueprintGraph/Private/BlueprintActionFilter.cpp
@@ -788,9 +788,9 @@ static bool BlueprintActionFilterImpl::IsRejectedGlobalField(FBlueprintActionFil
 		UClass* FieldClass = Field->GetOwnerClass();
 		if (bIsFilteredOut && (FieldClass != nullptr))
 		{
-			for (UClass const* Class : Filter.TargetClasses)
+			for (const auto& ClassData : Filter.TargetClasses)
 			{
-				bool const bIsInternalMemberField = IsClassOfType(Class, FieldClass);
+				bool const bIsInternalMemberField = IsClassOfType(ClassData.TargetClass, FieldClass);
 				if (bIsInternalMemberField)
 				{
 					bIsFilteredOut = false;
@@ -823,9 +823,9 @@ static bool BlueprintActionFilterImpl::IsNonTargetMember(FBlueprintActionFilter
 		{
 			bIsFilteredOut = Filter.TargetClasses.Num() > 0;
 
-			for (UClass const* Class : Filter.TargetClasses)
+			for (const auto& ClassData : Filter.TargetClasses)
 			{
-				bool const bIsTargetOwnedField = IsClassOfType(Class, ActionClass);
+				bool const bIsTargetOwnedField = IsClassOfType(ClassData.TargetClass, ActionClass);
 				if (bIsTargetOwnedField)
 				{
 					bIsFilteredOut = false;
@@ -852,11 +852,11 @@ static bool BlueprintActionFilterImpl::IsFieldCategoryHidden(FBlueprintActionFil
 	{
 		bIsFilteredOut = (Filter.TargetClasses.Num() > 0);
 
-		for (UClass* TargetClass : Filter.TargetClasses)
+		for (const auto& ClassData : Filter.TargetClasses)
 		{
 			// Use the UiSpec to get the category
 			FBlueprintActionUiSpec UiSpec = BlueprintAction.NodeSpawner->GetUiSpec(Filter.Context, BlueprintAction.GetBindings());
-			if (!FEditorCategoryUtils::IsCategoryHiddenFromClass(TargetClass->GetAuthoritativeClass(), UiSpec.Category.ToString()))
+			if (!FEditorCategoryUtils::IsCategoryHiddenFromClass(ClassData.HiddenCategories, ClassData.TargetClass, UiSpec.Category.ToString()))
 			{
 				bIsFilteredOut = false;
 				break;
@@ -1522,13 +1522,13 @@ static bool BlueprintActionFilterImpl::IsExtraneousInterfaceCall(FBlueprintActio
 		bool const bCanBeAddedToBlueprints = !InterfaceClass->HasMetaData(FBlueprintMetadata::MD_CannotImplementInterfaceInBlueprint);
 
 		bIsFilteredOut = (Filter.TargetClasses.Num() > 0);
-		for (const UClass* TargetClass : Filter.TargetClasses)
+		for (const auto& ClassData : Filter.TargetClasses)
 		{
-			bool const bImplementsInterface = IsClassOfType(TargetClass, InterfaceClass);
+			bool const bImplementsInterface = IsClassOfType(ClassData.TargetClass, InterfaceClass);
 			// if this is a blueprint class, and "CannotImplementInterfaceInBlueprint" 
 			// is set on the interface, then we know sub-classes cannot have
 			// the interface either (so there's no point to offering a message node)
-			bool const bIsBlueprintClass = (Cast<UBlueprintGeneratedClass>(TargetClass) != nullptr);
+			bool const bIsBlueprintClass = (Cast<UBlueprintGeneratedClass>(ClassData.TargetClass) != nullptr);
 
 			// if the class doesn't directly implement the interface (and it is 
 			// a possibility that some sub-class does), then we want to offer 
@@ -1549,9 +1549,9 @@ static bool BlueprintActionFilterImpl::IsExtraneousInterfaceCall(FBlueprintActio
 		if (bIsInterfaceAction && !NodeClass->IsChildOf<UK2Node_Event>())
 		{
 			bIsFilteredOut = (Filter.TargetClasses.Num() > 0);
-			for (const UClass* TargetClass : Filter.TargetClasses)
+			for (const auto& ClassData : Filter.TargetClasses)
 			{
-				UClass const* InterfaceImplementingClass = FindInheritedInterfaceClass(TargetClass, FuncClass);
+				UClass const* InterfaceImplementingClass = FindInheritedInterfaceClass(ClassData.TargetClass, FuncClass);
 				// interfaces that are added directly to a Blueprint (even in 
 				// the case of an interface on a parent blueprint) have their 
 				// functions stubbed-out/added to the blueprint class directly;
@@ -1762,6 +1762,8 @@ FBlueprintActionFilter::FBlueprintActionFilter(uint32 Flags/*= 0x00*/)
 {
 	using namespace BlueprintActionFilterImpl;
 
+	BluprintGraphModule = &FModuleManager::LoadModuleChecked<FBlueprintGraphModule>("BlueprintGraph");
+
 	//
 	// NOTE: The order of these tests can have perf implications, the more one
 	//       rejects on average the later it should be added (they're executed
@@ -1816,6 +1818,28 @@ FBlueprintActionFilter::FBlueprintActionFilter(uint32 Flags/*= 0x00*/)
 	AddRejectionTest(FRejectionTestDelegate::CreateStatic(IsStaleFieldAction));
 }
 
+//------------------------------------------------------------------------------
+void FBlueprintActionFilter::AddUnique(TArray<FTargetClassFilterData>& ToArray, UClass* TargetClass)
+{
+	for (auto ClassData : ToArray)
+	{
+		if (ClassData.TargetClass == TargetClass)
+		{
+			return;
+		}
+	}
+	FBlueprintActionFilter::Add(ToArray, TargetClass);
+}
+
+//------------------------------------------------------------------------------
+void FBlueprintActionFilter::Add(TArray<FTargetClassFilterData>& ToArray, UClass* TargetClass)
+{
+	TArray<FString> ClassHideCategories;
+	FEditorCategoryUtils::GetClassHideCategories(TargetClass, ClassHideCategories);
+	FTargetClassFilterData Data = { TargetClass, MoveTemp(ClassHideCategories) };
+	ToArray.Add(Data);
+}
+
 //------------------------------------------------------------------------------
 void FBlueprintActionFilter::AddRejectionTest(FRejectionTestDelegate IsFilteredDelegate)
 {
@@ -1896,9 +1920,7 @@ bool FBlueprintActionFilter::IsFilteredByThis(FBlueprintActionInfo& BlueprintAct
 
 	if (!bIsFiltered)
 	{
-		FBlueprintGraphModule& BluprintGraphModule = FModuleManager::LoadModuleChecked<FBlueprintGraphModule>("BlueprintGraph");
-
-		for (auto& ExtraRejectionTest : BluprintGraphModule.GetExtendedActionMenuFilters())
+		for (auto& ExtraRejectionTest : BluprintGraphModule->GetExtendedActionMenuFilters())
 		{
 			if (ExtraRejectionTest.Execute(FilterRef, BlueprintAction))
 			{
diff --git a/Engine/Source/Editor/BlueprintGraph/Public/BlueprintActionFilter.h b/Engine/Source/Editor/BlueprintGraph/Public/BlueprintActionFilter.h
--- a/Engine/Source/Editor/BlueprintGraph/Public/BlueprintActionFilter.h
+++ b/Engine/Source/Editor/BlueprintGraph/Public/BlueprintActionFilter.h
@@ -207,9 +207,24 @@ public:
 	/**
 	 * A list of classes that you want members for. If an action would produce
 	 * a node with a TargetPin, and that pin is incompatible with one of these
-	 * classes, then the action is filtered out.
+	 * classes, then the action is filtered out. Extra data is cached in 
+	 * FTargetClassFilterData so that it can be reused by filters efficiently.
 	 */
-	TArray<UClass*> TargetClasses;
+	struct FTargetClassFilterData
+	{
+		UClass* TargetClass;
+		TArray<FString> HiddenCategories;
+	};
+	TArray<FTargetClassFilterData> TargetClasses;
+
+	/** Helper to add a class to the TargetClasses, TargetClass may already be in the array */
+	static void AddUnique(TArray<FTargetClassFilterData>& ToArray, UClass* TargetClass);
+
+	/** Helper to add a class to the TargetClasses, fills out FTargetClassFilterData */
+	static void Add(TArray<FTargetClassFilterData>& ToArray, UClass* TargetClass);
+
+	/** Cached reference to the BluprintGraphModule, which has extra rejection tests: */
+	class FBlueprintGraphModule* BluprintGraphModule;
 
 	/**
 	 * Users can extend the filter and add their own rejection tests with this
diff --git a/Engine/Source/Editor/BlueprintGraph/Public/BlueprintNodeSpawner.h b/Engine/Source/Editor/BlueprintGraph/Public/BlueprintNodeSpawner.h
--- a/Engine/Source/Editor/BlueprintGraph/Public/BlueprintNodeSpawner.h
+++ b/Engine/Source/Editor/BlueprintGraph/Public/BlueprintNodeSpawner.h
@@ -21,6 +21,36 @@ struct FBlueprintActionUiSpec
 		: IconTint(FLinearColor::White)
 	{}
 
+	FBlueprintActionUiSpec(FBlueprintActionUiSpec&& Other)
+	{
+		Move(Other);
+	}
+
+	FBlueprintActionUiSpec& operator=(FBlueprintActionUiSpec&& Other)
+	{
+		if (this != &Other)
+		{
+			Move(Other);
+		}
+		return *this;
+	}
+
+	FBlueprintActionUiSpec(const FBlueprintActionUiSpec& Other)
+	{
+		Copy(Other);
+	}
+
+	FBlueprintActionUiSpec& operator=(const FBlueprintActionUiSpec& Other)
+	{
+		if (this != &Other)
+		{
+			Copy(Other);
+		}
+		return *this;
+	}
+
+	~FBlueprintActionUiSpec() {}
+
 	FText   MenuName;
 	FText   Category;
 	FText   Tooltip;
@@ -29,6 +59,31 @@ struct FBlueprintActionUiSpec
 	FLinearColor IconTint;
 	FString DocLink;
 	FString DocExcerptTag;
+
+private:
+	void Move(FBlueprintActionUiSpec& Other)
+	{
+		MenuName = MoveTemp(Other.MenuName);
+		Category = MoveTemp(Other.Category);
+		Tooltip = MoveTemp(Other.Tooltip);
+		Keywords = MoveTemp(Other.Keywords);
+		IconName = MoveTemp(Other.IconName);
+		IconTint = MoveTemp(Other.IconTint);
+		DocLink = MoveTemp(Other.DocLink);
+		DocExcerptTag = MoveTemp(Other.DocExcerptTag);
+	}
+
+	void Copy(const FBlueprintActionUiSpec& Other)
+	{
+		MenuName = Other.MenuName;
+		Category = Other.Category;
+		Tooltip = Other.Tooltip;
+		Keywords = Other.Keywords;
+		IconName = Other.IconName;
+		IconTint = Other.IconTint;
+		DocLink = Other.DocLink;
+		DocExcerptTag = Other.DocExcerptTag;
+	}
 };
 
 /*******************************************************************************
diff --git a/Engine/Source/Editor/GraphEditor/Private/GraphActionNode.cpp b/Engine/Source/Editor/GraphEditor/Private/GraphActionNode.cpp
--- a/Engine/Source/Editor/GraphEditor/Private/GraphActionNode.cpp
+++ b/Engine/Source/Editor/GraphEditor/Private/GraphActionNode.cpp
@@ -185,11 +185,10 @@ FGraphActionNode::FGraphActionNode(TArray< TSharedPtr<FEdGraphSchemaAction> > co
 //------------------------------------------------------------------------------
 TSharedPtr<FGraphActionNode> FGraphActionNode::AddChild(FGraphActionListBuilderBase::ActionGroup const& ActionSet)
 {
-	TArray<FString> CategoryStack;
-	ActionSet.GetCategoryChain(CategoryStack);
+	const TArray<FString>& CategoryStack = ActionSet.GetCategoryChain();
 
 	TSharedPtr<FGraphActionNode> ActionNode = FGraphActionNode::NewActionNode(ActionSet.Actions);
-	AddChildRecursively(CategoryStack, ActionNode);
+	AddChildRecursively(CategoryStack, 0, ActionNode);
 	
 	return ActionNode;
 }
@@ -421,18 +420,18 @@ TSharedPtr<FGraphActionNode> FGraphActionNode::NewActionNode(TArray< TSharedPtr<
 
 	for (TSharedPtr<FEdGraphSchemaAction> const& Action : ActionList)
 	{
-		Grouping = FMath::Max(Grouping, Action->Grouping);
+		Grouping = FMath::Max(Grouping, Action->GetGrouping());
 		if (SectionID == INVALID_SECTION_ID)
 		{
 			// take the first non-zero section ID
-			SectionID = Action->SectionID;
+			SectionID = Action->GetSectionID();
 		}
 	}
 
 	FGraphActionNode* ActionNode = new FGraphActionNode(ActionList, Grouping, SectionID);
 	TSharedPtr<FEdGraphSchemaAction> PrimeAction = ActionNode->GetPrimaryAction();
 	checkSlow(PrimeAction.IsValid());
-	ActionNode->DisplayText = PrimeAction->MenuDescription;
+	ActionNode->DisplayText = PrimeAction->GetMenuDescription();
 
 	return MakeShareable(ActionNode);
 }
@@ -448,7 +447,7 @@ TSharedPtr<FGraphActionNode> FGraphActionNode::NewGroupDividerNode(TWeakPtr<FGra
 }
 
 //------------------------------------------------------------------------------
-void FGraphActionNode::AddChildRecursively(TArray<FString>& CategoryStack, TSharedPtr<FGraphActionNode> NodeToAdd)
+void FGraphActionNode::AddChildRecursively(const TArray<FString>& CategoryStack, int32 Idx, TSharedPtr<FGraphActionNode> NodeToAdd)
 {
 	if (NodeToAdd->SectionID != INVALID_SECTION_ID)
 	{
@@ -464,27 +463,27 @@ void FGraphActionNode::AddChildRecursively(TArray<FString>& CategoryStack, TShar
 
 		if ( FoundSectionNode.IsValid() )
 		{
-			FoundSectionNode->AddChildRecursively(CategoryStack, NodeToAdd);
+			FoundSectionNode->AddChildRecursively(CategoryStack, Idx, NodeToAdd);
 			return;
 		}
 	}
 
-	if ( CategoryStack.Num() > 0 )
+	if ( Idx < CategoryStack.Num() )
 	{
-		FString CategorySection = CategoryStack[0];
-		CategoryStack.RemoveAt(0, 1);
+		const FString& CategorySection = CategoryStack[Idx];
+		++Idx;
 
 		// make sure we don't already have a child that this can nest under
 		TSharedPtr<FGraphActionNode> ExistingNode = FindMatchingParent(CategorySection, NodeToAdd);
 		if ( ExistingNode.IsValid() )
 		{
-			ExistingNode->AddChildRecursively(CategoryStack, NodeToAdd);
+			ExistingNode->AddChildRecursively(CategoryStack, Idx, NodeToAdd);
 		}
 		else
 		{
 			TSharedPtr<FGraphActionNode> CategoryNode = NewCategoryNode(CategorySection, NodeToAdd->Grouping, NodeToAdd->SectionID);
 			InsertChild(CategoryNode);
-			CategoryNode->AddChildRecursively(CategoryStack, NodeToAdd);
+			CategoryNode->AddChildRecursively(CategoryStack, Idx, NodeToAdd);
 		}
 	}
 	else
@@ -547,7 +546,7 @@ void FGraphActionNode::InsertChild(TSharedPtr<FGraphActionNode> NodeToAdd)
 			// make sure we already haven't already added a heading for this section
 			!ChildSections.Contains(NodeToAdd->SectionID) &&
 			// if this node also has a category, use that over a section heading
-			(!NodeToAdd->IsActionNode() || NodeToAdd->GetPrimaryAction()->Category.IsEmpty());
+			(!NodeToAdd->IsActionNode() || NodeToAdd->GetPrimaryAction()->GetCategory().IsEmpty());
 
 		if (bAddSectionHeading)
 		{
diff --git a/Engine/Source/Editor/GraphEditor/Private/GraphActionNode.h b/Engine/Source/Editor/GraphEditor/Private/GraphActionNode.h
--- a/Engine/Source/Editor/GraphEditor/Private/GraphActionNode.h
+++ b/Engine/Source/Editor/GraphEditor/Private/GraphActionNode.h
@@ -221,13 +221,14 @@ private:
 	static TSharedPtr<FGraphActionNode> NewGroupDividerNode(TWeakPtr<FGraphActionNode> Parent, int32 Grouping);
 
 	/**
-	 * Whittles down the CategoryStack, adding category-nodes as needed. The
+	 * Iterates the CategoryStack, adding category-nodes as needed. The
 	 * last category is what the node will be inserted under.
 	 *
 	 * @param  CategoryStack	A list of categories denoting where to nest the new node (the first element is the highest category)
+	 * @param  Idx				Current point in the category stack that we have iterated to
 	 * @param  NodeToAdd		The node you want inserted.
 	 */
-	void AddChildRecursively(TArray<FString>& CategoryStack, TSharedPtr<FGraphActionNode> NodeToAdd);
+	void AddChildRecursively(const TArray<FString>& CategoryStack, int32 Idx, TSharedPtr<FGraphActionNode> NodeToAdd);
 
 	/**
 	 * Looks through this node's children to see if a there already exists a 
diff --git a/Engine/Source/Editor/GraphEditor/Private/GraphEditorDragDropAction.cpp b/Engine/Source/Editor/GraphEditor/Private/GraphEditorDragDropAction.cpp
--- a/Engine/Source/Editor/GraphEditor/Private/GraphEditorDragDropAction.cpp
+++ b/Engine/Source/Editor/GraphEditor/Private/GraphEditorDragDropAction.cpp
@@ -183,7 +183,7 @@ void FGraphSchemaActionDragDropAction::HoverTargetChanged()
 		const FSlateBrush* StatusSymbol = FEditorStyle::GetBrush(TEXT("Graph.ConnectorFeedback.NewNode"));
 
 		//Create feedback message with the function name.
-		SetSimpleFeedbackMessage(StatusSymbol, FLinearColor::White, ActionNode->MenuDescription);
+		SetSimpleFeedbackMessage(StatusSymbol, FLinearColor::White, ActionNode->GetMenuDescription());
 	}
 }
 
diff --git a/Engine/Source/Editor/GraphEditor/Private/SGraphActionMenu.cpp b/Engine/Source/Editor/GraphEditor/Private/SGraphActionMenu.cpp
--- a/Engine/Source/Editor/GraphEditor/Private/SGraphActionMenu.cpp
+++ b/Engine/Source/Editor/GraphEditor/Private/SGraphActionMenu.cpp
@@ -128,14 +128,14 @@ void SDefaultGraphActionWidget::Construct(const FArguments& InArgs, const FCreat
 	this->ChildSlot
 	[
 		SNew(SHorizontalBox)
-		.ToolTipText( FText::FromString(InCreateData->Action->TooltipDescription) )
+		.ToolTipText( FText::FromString(InCreateData->Action->GetTooltipDescription()) )
 		+ SHorizontalBox::Slot()
 		.AutoWidth()
 		.VAlign(VAlign_Center)
 		[
 			SNew(STextBlock)
 			.Font(FSlateFontInfo( FPaths::EngineContentDir() / TEXT("Slate/Fonts/Roboto-Regular.ttf"), 9 ))
-			.Text(InCreateData->Action->MenuDescription)
+			.Text(InCreateData->Action->GetMenuDescription())
 			.HighlightText(InArgs._HighlightText)
 		]
 	];
@@ -586,7 +586,7 @@ static bool CompareGraphActionNode(TSharedPtr<FGraphActionNode> A, TSharedPtr<FG
 
 	if (A->HasValidAction() && B->HasValidAction())
 	{
-		return A->GetPrimaryAction()->MenuDescription.CompareTo(B->GetPrimaryAction()->MenuDescription) == 0;
+		return A->GetPrimaryAction()->GetMenuDescription().CompareTo(B->GetPrimaryAction()->GetMenuDescription()) == 0;
 	}
 	else if(!A->HasValidAction() && !B->HasValidAction())
 	{
diff --git a/Engine/Source/Editor/GraphEditor/Private/SGraphPalette.cpp b/Engine/Source/Editor/GraphEditor/Private/SGraphPalette.cpp
--- a/Engine/Source/Editor/GraphEditor/Private/SGraphPalette.cpp
+++ b/Engine/Source/Editor/GraphEditor/Private/SGraphPalette.cpp
@@ -24,7 +24,7 @@ void SGraphPaletteItem::Construct(const FArguments& InArgs, FCreateWidgetForActi
 	// Find icons
 	const FSlateBrush* IconBrush = FEditorStyle::GetBrush(TEXT("NoBrush"));
 	FSlateColor IconColor = FSlateColor::UseForeground();
-	FString ToolTip = GraphAction->TooltipDescription;
+	FString ToolTip = GraphAction->GetTooltipDescription();
 	FText IconToolTip = FText::FromString(ToolTip);
 	bool bIsReadOnly = false;
 
@@ -157,12 +157,12 @@ void SGraphPaletteItem::OnNameTextCommitted(const FText& NewText, ETextCommit::T
 
 FText SGraphPaletteItem::GetDisplayText() const
 {
-	return ActionPtr.Pin()->MenuDescription;
+	return ActionPtr.Pin()->GetMenuDescription();
 }
 
 FText SGraphPaletteItem::GetItemTooltip() const
 {
-	return ActionPtr.Pin()->MenuDescription;
+	return ActionPtr.Pin()->GetMenuDescription();
 }
 
 
diff --git a/Engine/Source/Editor/Kismet/Private/BPFunctionDragDropAction.cpp b/Engine/Source/Editor/Kismet/Private/BPFunctionDragDropAction.cpp
--- a/Engine/Source/Editor/Kismet/Private/BPFunctionDragDropAction.cpp
+++ b/Engine/Source/Editor/Kismet/Private/BPFunctionDragDropAction.cpp
@@ -103,7 +103,7 @@ void FKismetDragDropAction::HoverTargetChanged()
 	if (ActionWillShowExistingNode())
 	{
 		FSlateBrush const* ShowsExistingIcon = FEditorStyle::GetBrush(TEXT("Graph.ConnectorFeedback.ShowNode"));
-		FText DragingText = FText::Format(LOCTEXT("ShowExistingNode", "Show '{0}'"), ActionNode->MenuDescription);
+		FText DragingText = FText::Format(LOCTEXT("ShowExistingNode", "Show '{0}'"), ActionNode->GetMenuDescription());
 		SetSimpleFeedbackMessage(ShowsExistingIcon, FLinearColor::White, DragingText);
 	}
 	// it should be obvious that we can't drop on anything but a graph, so no need to point that out
diff --git a/Engine/Source/Editor/Kismet/Private/BlueprintActionMenuBuilder.cpp b/Engine/Source/Editor/Kismet/Private/BlueprintActionMenuBuilder.cpp
--- a/Engine/Source/Editor/Kismet/Private/BlueprintActionMenuBuilder.cpp
+++ b/Engine/Source/Editor/Kismet/Private/BlueprintActionMenuBuilder.cpp
@@ -14,8 +14,11 @@
 #include "SMyBlueprint.h"				// for SelectionAsVar()
 #include "BlueprintEditorUtils.h"		// for FindBlueprintForGraphChecked()
 #include "BlueprintEditor.h"			// for GetFocusedGraph()
+#include "EditorCategoryUtils.h"
+#include "ObjectEditorUtils.h"
 
 #define LOCTEXT_NAMESPACE "BlueprintActionMenuBuilder"
+DEFINE_LOG_CATEGORY_STATIC(LogBlueprintActionMenuItemFactory, Log, All);
 
 /*******************************************************************************
  * FBlueprintActionMenuItemFactory
@@ -107,10 +110,7 @@ TSharedPtr<FBlueprintActionMenuItem> FBlueprintActionMenuItemFactory::MakeAction
 	FBlueprintActionUiSpec UiSignature = GetActionUiSignature(EditorContext, ActionInfo);
 
 	UBlueprintNodeSpawner const* Action = ActionInfo.NodeSpawner;
-	FBlueprintActionMenuItem* NewMenuItem = new FBlueprintActionMenuItem(Action, UiSignature);
-
-	NewMenuItem->Grouping = MenuGrouping;
-	NewMenuItem->Category = FText::FromString(RootCategory.ToString() + TEXT('|') + NewMenuItem->Category.ToString());
+	FBlueprintActionMenuItem* NewMenuItem = new FBlueprintActionMenuItem(Action, UiSignature, IBlueprintNodeBinder::FBindingSet(), FText::FromString(RootCategory.ToString() + TEXT('|') + UiSignature.Category.ToString()), MenuGrouping);
 
 	return MakeShareable(NewMenuItem);
 }
@@ -119,9 +119,67 @@ TSharedPtr<FBlueprintActionMenuItem> FBlueprintActionMenuItemFactory::MakeAction
 TSharedPtr<FBlueprintDragDropMenuItem> FBlueprintActionMenuItemFactory::MakeDragDropMenuItem(UBlueprintNodeSpawner const* SampleAction)
 {
 	// FBlueprintDragDropMenuItem takes care of its own menu MenuDescription, etc.
-	FBlueprintDragDropMenuItem* NewMenuItem = new FBlueprintDragDropMenuItem(Context, SampleAction, MenuGrouping);
+	UProperty const* SampleProperty = nullptr;
+	if (UBlueprintDelegateNodeSpawner const* DelegateSpawner = Cast<UBlueprintDelegateNodeSpawner const>(SampleAction))
+	{
+		SampleProperty = DelegateSpawner->GetDelegateProperty();
+	}
+	else if (UBlueprintVariableNodeSpawner const* VariableSpawner = Cast<UBlueprintVariableNodeSpawner const>(SampleAction))
+	{
+		SampleProperty = VariableSpawner->GetVarProperty();
+	}
+
+	FText MenuDescription;
+	FString TooltipDescription;
+	FText Category;
+	if (SampleProperty != nullptr)
+	{
+		bool const bShowFriendlyNames = GetDefault<UEditorStyleSettings>()->bShowFriendlyNames;
+		MenuDescription = bShowFriendlyNames ? FText::FromString(UEditorEngine::GetFriendlyName(SampleProperty)) : FText::FromName(SampleProperty->GetFName());
+
+		TooltipDescription = SampleProperty->GetToolTipText().ToString();
+		Category = FObjectEditorUtils::GetCategoryText(SampleProperty);
+
+		bool const bIsDelegateProperty = SampleProperty->IsA<UMulticastDelegateProperty>();
+		if (bIsDelegateProperty && Category.IsEmpty())
+		{
+			Category = FEditorCategoryUtils::GetCommonCategory(FCommonEditorCategory::Delegates);
+		}
+		else if (!bIsDelegateProperty)
+		{
+			check(Context.Blueprints.Num() > 0);
+			UBlueprint const* Blueprint = Context.Blueprints[0];
+			UClass const*     BlueprintClass = (Blueprint->SkeletonGeneratedClass != nullptr) ? Blueprint->SkeletonGeneratedClass : Blueprint->ParentClass;
+
+			UClass const* PropertyClass = SampleProperty->GetOwnerClass();
+			checkSlow(PropertyClass != nullptr);
+			bool const bIsMemberProperty = BlueprintClass->IsChildOf(PropertyClass);
+
+			FText TextCategory;
+			if (Category.IsEmpty())
+			{
+				TextCategory = FEditorCategoryUtils::GetCommonCategory(FCommonEditorCategory::Variables);
+			}
+			else if (bIsMemberProperty)
+			{
+				TextCategory = FEditorCategoryUtils::BuildCategoryString(FCommonEditorCategory::Variables, Category);
+			}
+
+			if (!bIsMemberProperty)
+			{
+				TextCategory = FText::Format(LOCTEXT("NonMemberVarCategory", "Class|{0}|{1}"), PropertyClass->GetDisplayNameText(), TextCategory);
+			}
+			Category = TextCategory;
+		}
+	}
+	else
+	{
+		UE_LOG(LogBlueprintActionMenuItemFactory, Warning, TEXT("Unhandled (or invalid) spawner: '%s'"), *SampleAction->GetName());
+	}
+
+	Category = FText::FromString(RootCategory.ToString() + TEXT('|') + Category.ToString());
 
-	NewMenuItem->Category = FText::FromString(RootCategory.ToString() + TEXT('|') + NewMenuItem->Category.ToString());
+	FBlueprintDragDropMenuItem* NewMenuItem = new FBlueprintDragDropMenuItem(Context, SampleAction, MenuGrouping, Category, MenuDescription, TooltipDescription);
 	return MakeShareable(NewMenuItem);
 }
 
@@ -131,10 +189,7 @@ TSharedPtr<FBlueprintActionMenuItem> FBlueprintActionMenuItemFactory::MakeBoundM
 	FBlueprintActionUiSpec UiSignature = GetActionUiSignature(EditorContext, ActionInfo);
 
 	UBlueprintNodeSpawner const* Action = ActionInfo.NodeSpawner;
-	FBlueprintActionMenuItem* NewMenuItem = new FBlueprintActionMenuItem(Action, UiSignature, ActionInfo.GetBindings());
-
-	NewMenuItem->Grouping = MenuGrouping;
-	NewMenuItem->Category = FText::FromString( RootCategory.ToString() + TEXT('|') + NewMenuItem->Category.ToString());
+	FBlueprintActionMenuItem* NewMenuItem = new FBlueprintActionMenuItem(Action, UiSignature, ActionInfo.GetBindings(), FText::FromString(RootCategory.ToString() + TEXT('|') + UiSignature.Category.ToString()), MenuGrouping);
 
 	return MakeShareable(NewMenuItem);
 }
@@ -331,7 +386,7 @@ void FBlueprintActionMenuBuilderImpl::FMenuSectionDefinition::AddBoundMenuItems(
 
 					if (Flags & FBlueprintActionMenuBuilder::FlattenCategoryHierarcy)
 					{
-						LastMadeMenuItem->Category = ItemFactory.RootCategory;
+						LastMadeMenuItem->UpdateCategory( ItemFactory.RootCategory );
 					}
 				}
 				else
@@ -396,7 +451,7 @@ FBlueprintActionMenuBuilderImpl::MenuItemList FBlueprintActionMenuBuilderImpl::F
 		UnBoundMenuEntry = ItemFactory.MakeActionMenuItem(EditorContext, DatabaseAction);
 		if (Flags & FBlueprintActionMenuBuilder::FlattenCategoryHierarcy)
 		{
-			UnBoundMenuEntry->Category = ItemFactory.RootCategory;
+			UnBoundMenuEntry->UpdateCategory( ItemFactory.RootCategory );
 		}
 	}
 
diff --git a/Engine/Source/Editor/Kismet/Private/BlueprintActionMenuItem.cpp b/Engine/Source/Editor/Kismet/Private/BlueprintActionMenuItem.cpp
--- a/Engine/Source/Editor/Kismet/Private/BlueprintActionMenuItem.cpp
+++ b/Engine/Source/Editor/Kismet/Private/BlueprintActionMenuItem.cpp
@@ -199,18 +199,15 @@ static UEdGraphNode* FBlueprintMenuActionItemImpl::AutowireSpawnedNodes(UEdGraph
  ******************************************************************************/
 
 //------------------------------------------------------------------------------
-FBlueprintActionMenuItem::FBlueprintActionMenuItem(UBlueprintNodeSpawner const* NodeSpawner, FBlueprintActionUiSpec const& UiSpec, IBlueprintNodeBinder::FBindingSet const& InBindings)
-	: Action(NodeSpawner)
+FBlueprintActionMenuItem::FBlueprintActionMenuItem(UBlueprintNodeSpawner const* NodeSpawner, FBlueprintActionUiSpec const& UiSpec, IBlueprintNodeBinder::FBindingSet const& InBindings, const FText& InNodeCategory, int32 InGrouping)
+	: FEdGraphSchemaAction(InNodeCategory, UiSpec.MenuName, UiSpec.Tooltip.ToString(), InGrouping, UiSpec.Keywords)
+	, Action(NodeSpawner)
 	, IconTint(UiSpec.IconTint)
 	, IconBrush(FEditorStyle::GetBrush(UiSpec.IconName))
 	, Bindings(InBindings)
 {
 	check(Action != nullptr);
 
-	MenuDescription = UiSpec.MenuName;
-	Category = UiSpec.Category;
-	TooltipDescription = UiSpec.Tooltip.ToString();
-	Keywords = UiSpec.Keywords;
 	DocExcerptRef.DocLink = UiSpec.DocLink;
 	DocExcerptRef.DocExcerptName = UiSpec.DocExcerptTag;
 	// we may fill out the UiSpec's DocLink with whitespace (so we can tell the
@@ -339,13 +336,14 @@ void FBlueprintActionMenuItem::AppendBindings(const FBlueprintActionContext& Con
 	Bindings.Append(BindingSet);
 
 	FBlueprintActionUiSpec UiSpec = Action->GetUiSpec(Context, Bindings);
+
 	// ui signature could be dynamic, and change as bindings change
-	MenuDescription = UiSpec.MenuName;
 	// @TODO: would invalidate any category pre-pending that was done at the 
 	//        MenuBuilder level
 	//Category  = UiSpec.Category.ToString();
-	TooltipDescription = UiSpec.Tooltip.ToString();
-	Keywords  = UiSpec.Keywords;	
+	
+	UpdateSearchData(UiSpec.MenuName, UiSpec.Tooltip.ToString(), FText(), UiSpec.Keywords);
+
 	IconBrush = FEditorStyle::GetBrush(UiSpec.IconName);
 	IconTint  = UiSpec.IconTint;
 	DocExcerptRef.DocLink = UiSpec.DocLink;
diff --git a/Engine/Source/Editor/Kismet/Private/BlueprintActionMenuUtils.cpp b/Engine/Source/Editor/Kismet/Private/BlueprintActionMenuUtils.cpp
--- a/Engine/Source/Editor/Kismet/Private/BlueprintActionMenuUtils.cpp
+++ b/Engine/Source/Editor/Kismet/Private/BlueprintActionMenuUtils.cpp
@@ -293,7 +293,7 @@ static FBlueprintActionFilter BlueprintActionMenuUtilsImpl::MakeCallOnMemberFilt
 
 	bool bForceAddComponents = ((ContextTargetMask & EContextTargetFlags::TARGET_SubComponents) != 0);
 
-	TArray<UClass*> TargetClasses = MainMenuFilter.TargetClasses;
+	auto TargetClasses = MainMenuFilter.TargetClasses;
 	if (bForceAddComponents && (TargetClasses.Num() == 0))
 	{
 		for (UBlueprint const* TargetBlueprint : MainMenuFilter.Context.Blueprints)
@@ -301,13 +301,14 @@ static FBlueprintActionFilter BlueprintActionMenuUtilsImpl::MakeCallOnMemberFilt
 			UClass* BpClass = TargetBlueprint->SkeletonGeneratedClass;
 			if (BpClass != nullptr)
 			{
-				TargetClasses.AddUnique(BpClass);
+				FBlueprintActionFilter::AddUnique(TargetClasses, BpClass);
 			}
 		}
 	}
 
-	for (UClass const* TargetClass : TargetClasses)
+	for ( const auto& ClassData : TargetClasses)
 	{
+		UClass const* TargetClass = ClassData.TargetClass;
 		for (TFieldIterator<UObjectProperty> PropertyIt(TargetClass, EFieldIteratorFlags::IncludeSuper); PropertyIt; ++PropertyIt)
 		{
  			UObjectProperty* ObjectProperty = *PropertyIt;
@@ -425,7 +426,7 @@ void FBlueprintActionMenuUtils::MakePaletteMenu(FBlueprintActionContext const& C
 	
 	if (FilterClass != nullptr)
 	{
-		MenuFilter.TargetClasses.Add(FilterClass);
+		FBlueprintActionFilter::Add(MenuFilter.TargetClasses, FilterClass);
 	}
 
 	MenuOut.AddMenuSection(MenuFilter, LOCTEXT("PaletteRoot", "Library"), /*MenuOrder =*/0, FBlueprintActionMenuBuilder::ConsolidatePropertyActions);
@@ -480,7 +481,7 @@ void FBlueprintActionMenuUtils::MakeContextMenu(FBlueprintActionContext const& C
 			bCanOperateOnLevelActors &= BlueprintClass->IsChildOf<ALevelScriptActor>();
 			if (bIsContextSensitive && (ClassTargetMask & EContextTargetFlags::TARGET_Blueprint))
 			{
-				MainMenuFilter.TargetClasses.AddUnique(BlueprintClass);
+				FBlueprintActionFilter::AddUnique(MainMenuFilter.TargetClasses, BlueprintClass);
 			}
 		}
 		bCanHaveActorComponents &= FBlueprintEditorUtils::DoesSupportComponents(Blueprint);
@@ -567,7 +568,7 @@ void FBlueprintActionMenuUtils::MakeContextMenu(FBlueprintActionContext const& C
 			{
 				if (ClassTargetMask & EContextTargetFlags::TARGET_PinObject)
 				{
-					MainMenuFilter.TargetClasses.AddUnique(PinObjClass);
+					FBlueprintActionFilter::AddUnique(MainMenuFilter.TargetClasses, PinObjClass);
 				}
 			}
 
@@ -579,7 +580,7 @@ void FBlueprintActionMenuUtils::MakeContextMenu(FBlueprintActionContext const& C
 				{
 					if (UClass* TargetClass = GetPinClassType(TargetPin))
 					{
-						MainMenuFilter.TargetClasses.AddUnique(TargetClass);
+						FBlueprintActionFilter::AddUnique(MainMenuFilter.TargetClasses, TargetClass);
 					}
 				}
 			}
@@ -595,7 +596,7 @@ void FBlueprintActionMenuUtils::MakeContextMenu(FBlueprintActionContext const& C
 				{
 					if (UClass* PinClass = GetPinClassType(NodePin))
 					{
-						MainMenuFilter.TargetClasses.AddUnique(PinClass);
+						FBlueprintActionFilter::AddUnique(MainMenuFilter.TargetClasses, PinClass);
 					}
 				}
 			}
diff --git a/Engine/Source/Editor/Kismet/Private/BlueprintDragDropMenuItem.cpp b/Engine/Source/Editor/Kismet/Private/BlueprintDragDropMenuItem.cpp
--- a/Engine/Source/Editor/Kismet/Private/BlueprintDragDropMenuItem.cpp
+++ b/Engine/Source/Editor/Kismet/Private/BlueprintDragDropMenuItem.cpp
@@ -21,66 +21,11 @@
 DEFINE_LOG_CATEGORY_STATIC(LogBlueprintDragDropMenuItem, Log, All);
 
 //------------------------------------------------------------------------------
-FBlueprintDragDropMenuItem::FBlueprintDragDropMenuItem(FBlueprintActionContext const& Context, UBlueprintNodeSpawner const* SampleAction, int32 MenuGrouping/* = 0*/)
+FBlueprintDragDropMenuItem::FBlueprintDragDropMenuItem(FBlueprintActionContext const& Context, UBlueprintNodeSpawner const* SampleAction, int32 MenuGrouping, const FText& InNodeCategory, const FText& InMenuDesc, const FString& InToolTip)
+: FEdGraphSchemaAction(InNodeCategory, InMenuDesc, InToolTip, MenuGrouping)
 {
 	AppendAction(SampleAction);
 	check(SampleAction != nullptr);
-	Grouping  = MenuGrouping;
-
-	UProperty const* SampleProperty = nullptr;
-	if (UBlueprintDelegateNodeSpawner const* DelegateSpawner = Cast<UBlueprintDelegateNodeSpawner const>(SampleAction))
-	{
-		SampleProperty = DelegateSpawner->GetDelegateProperty();
-	}
-	else if (UBlueprintVariableNodeSpawner const* VariableSpawner = Cast<UBlueprintVariableNodeSpawner const>(SampleAction))
-	{
-		SampleProperty = VariableSpawner->GetVarProperty();
-	}
-
-	if (SampleProperty != nullptr)
-	{
-		bool const bShowFriendlyNames = GetDefault<UEditorStyleSettings>()->bShowFriendlyNames;
-		MenuDescription = bShowFriendlyNames ? FText::FromString(UEditorEngine::GetFriendlyName(SampleProperty)) : FText::FromName(SampleProperty->GetFName());
-
-		TooltipDescription = SampleProperty->GetToolTipText().ToString();
-		Category = FObjectEditorUtils::GetCategoryText(SampleProperty);
-
-		bool const bIsDelegateProperty = SampleProperty->IsA<UMulticastDelegateProperty>();
-		if (bIsDelegateProperty && Category.IsEmpty())
-		{
-			Category = FEditorCategoryUtils::GetCommonCategory(FCommonEditorCategory::Delegates);
-		}
-		else if (!bIsDelegateProperty)
-		{
-			check(Context.Blueprints.Num() > 0);
-			UBlueprint const* Blueprint = Context.Blueprints[0];
-			UClass const*     BlueprintClass = (Blueprint->SkeletonGeneratedClass != nullptr) ? Blueprint->SkeletonGeneratedClass : Blueprint->ParentClass;
-
-			UClass const* PropertyClass = SampleProperty->GetOwnerClass();
-			checkSlow(PropertyClass != nullptr);
-			bool const bIsMemberProperty = BlueprintClass->IsChildOf(PropertyClass);
-
-			FText TextCategory;
-			if (Category.IsEmpty())
-			{
-				TextCategory = FEditorCategoryUtils::GetCommonCategory(FCommonEditorCategory::Variables);
-			}
-			else if (bIsMemberProperty)
-			{
-				TextCategory = FEditorCategoryUtils::BuildCategoryString(FCommonEditorCategory::Variables, Category);
-			}
-			
-			if (!bIsMemberProperty)
-			{
-				TextCategory = FText::Format(LOCTEXT("NonMemberVarCategory", "Class|{0}|{1}"), PropertyClass->GetDisplayNameText(), TextCategory);
-			}
-			Category = TextCategory;
-		}
-	}
-	else
-	{
-		UE_LOG(LogBlueprintDragDropMenuItem, Warning, TEXT("Unhandled (or invalid) spawner: '%s'"), *SampleAction->GetName());
-	}
 }
 
 //------------------------------------------------------------------------------
diff --git a/Engine/Source/Editor/Kismet/Private/SBlueprintActionMenu.cpp b/Engine/Source/Editor/Kismet/Private/SBlueprintActionMenu.cpp
--- a/Engine/Source/Editor/Kismet/Private/SBlueprintActionMenu.cpp
+++ b/Engine/Source/Editor/Kismet/Private/SBlueprintActionMenu.cpp
@@ -24,12 +24,12 @@ USTRUCT()
 struct FBlueprintAction_PromoteVariable : public FEdGraphSchemaAction
 {
 	FBlueprintAction_PromoteVariable(bool bInToMemberVariable)
-		: FEdGraphSchemaAction()
+		: FEdGraphSchemaAction(	FText(), 
+								bInToMemberVariable? LOCTEXT("PromoteToVariable", "Promote to variable") : LOCTEXT("PromoteToLocalVariable", "Promote to local variable"),
+								bInToMemberVariable ? LOCTEXT("PromoteToVariable", "Promote to variable").ToString() : LOCTEXT("PromoteToLocalVariable", "Promote to local variable").ToString(),
+								1)
 		, bToMemberVariable(bInToMemberVariable)
 	{
-		MenuDescription = bInToMemberVariable? LOCTEXT("PromoteToVariable", "Promote to variable") : LOCTEXT("PromoteToLocalVariable", "Promote to local variable");
-		TooltipDescription = bInToMemberVariable? LOCTEXT("PromoteToVariable", "Promote to variable").ToString() : LOCTEXT("PromoteToLocalVariable", "Promote to local variable").ToString();
-		Grouping = 1;
 	}
 
 	// FEdGraphSchemaAction interface
diff --git a/Engine/Source/Editor/Kismet/Private/SBlueprintPalette.cpp b/Engine/Source/Editor/Kismet/Private/SBlueprintPalette.cpp
--- a/Engine/Source/Editor/Kismet/Private/SBlueprintPalette.cpp
+++ b/Engine/Source/Editor/Kismet/Private/SBlueprintPalette.cpp
@@ -281,7 +281,7 @@ static void GetSubGraphIcon(FEdGraphSchemaAction_K2Graph const* const ActionIn,
 static void GetPaletteItemIcon(TSharedPtr<FEdGraphSchemaAction> ActionIn, UBlueprint const* BlueprintIn, FSlateBrush const*& BrushOut, FSlateColor& ColorOut, FString& ToolTipOut, FString& DocLinkOut, FString& DocExcerptOut)
 {
 	// Default to tooltip based on action supplied
-	ToolTipOut = (ActionIn->TooltipDescription.Len() > 0) ? ActionIn->TooltipDescription : ActionIn->MenuDescription.ToString();
+	ToolTipOut = (ActionIn->GetTooltipDescription().Len() > 0) ? ActionIn->GetTooltipDescription() : ActionIn->GetMenuDescription().ToString();
 
 	if (ActionIn->GetTypeId() == FBlueprintActionMenuItem::StaticGetTypeId())
 	{
@@ -971,7 +971,7 @@ void SBlueprintPaletteItem::Construct(const FArguments& InArgs, FCreateWidgetFor
 	// construct the icon widget
 	FSlateBrush const* IconBrush   = FEditorStyle::GetBrush(TEXT("NoBrush"));
 	FSlateColor        IconColor   = FSlateColor::UseForeground();
-	FString            IconToolTip = GraphAction->TooltipDescription;
+	FString            IconToolTip = GraphAction->GetTooltipDescription();
 	FString			   IconDocLink, IconDocExcerpt;
 	GetPaletteItemIcon(GraphAction, Blueprint, IconBrush, IconColor, IconToolTip, IconDocLink, IconDocExcerpt);
 	TSharedRef<SWidget> IconWidget = CreateIconWidget(FText::FromString(IconToolTip), IconBrush, IconColor, IconDocLink, IconDocExcerpt);
@@ -981,8 +981,8 @@ void SBlueprintPaletteItem::Construct(const FArguments& InArgs, FCreateWidgetFor
 	FTutorialMetaData TagMeta("PaletteItem"); 
 	if( ActionPtr.IsValid() )
 	{
-		TagMeta.Tag = *FString::Printf(TEXT("PaletteItem,%s,%d"), *GraphAction->MenuDescription.ToString(), GraphAction->SectionID);
-		TagMeta.FriendlyName = GraphAction->MenuDescription.ToString();
+		TagMeta.Tag = *FString::Printf(TEXT("PaletteItem,%s,%d"), *GraphAction->GetMenuDescription().ToString(), GraphAction->GetSectionID());
+		TagMeta.FriendlyName = GraphAction->GetMenuDescription().ToString();
 	}
 	// construct the text widget
 	FSlateFontInfo NameFont = FSlateFontInfo(FPaths::EngineContentDir() / TEXT("Slate/Fonts/Roboto-Regular.ttf"), 10);
@@ -1162,7 +1162,7 @@ FText SBlueprintPaletteItem::GetDisplayText() const
 		}
 		else
 		{
-			MenuDescriptionCache.SetCachedText(ActionPtr.Pin()->MenuDescription, K2Schema);
+			MenuDescriptionCache.SetCachedText(ActionPtr.Pin()->GetMenuDescription(), K2Schema);
 		}
 	}
 
@@ -1403,7 +1403,7 @@ FText SBlueprintPaletteItem::GetToolTipText() const
 	if (PaletteAction.IsValid())
 	{
 		// Default tooltip is taken from the action
-		ToolTipText = (PaletteAction->TooltipDescription.Len() > 0) ? FText::FromString(PaletteAction->TooltipDescription) : PaletteAction->MenuDescription;
+		ToolTipText = (PaletteAction->GetTooltipDescription().Len() > 0) ? FText::FromString(PaletteAction->GetTooltipDescription()) : PaletteAction->GetMenuDescription();
 
 		if(PaletteAction->GetTypeId() == FEdGraphSchemaAction_K2AddComponent::StaticGetTypeId())
 		{
diff --git a/Engine/Source/Editor/Kismet/Private/SMyBlueprint.cpp b/Engine/Source/Editor/Kismet/Private/SMyBlueprint.cpp
--- a/Engine/Source/Editor/Kismet/Private/SMyBlueprint.cpp
+++ b/Engine/Source/Editor/Kismet/Private/SMyBlueprint.cpp
@@ -458,7 +458,7 @@ void SMyBlueprint::OnCategoryNameCommitted(const FText& InNewText, ETextCommit::
 			else if (Actions[i]->GetTypeId() == FEdGraphSchemaAction_K2Graph::StaticGetTypeId())
 			{
 				// Do not allow renaming of any graph actions outside of the following
-				if(Actions[i]->SectionID == NodeSectionID::FUNCTION || Actions[i]->SectionID == NodeSectionID::MACRO)
+				if(Actions[i]->GetSectionID() == NodeSectionID::FUNCTION || Actions[i]->GetSectionID() == NodeSectionID::MACRO)
 				{
 					FEdGraphSchemaAction_K2Graph* GraphAction = (FEdGraphSchemaAction_K2Graph*)Actions[i].Get();
 
@@ -727,8 +727,8 @@ void SMyBlueprint::BuildOverridableFunctionsMenu(FMenuBuilder& MenuBuilder)
 		for ( auto& OverrideAction : OverridableFunctionActions )
 		{
 			MenuBuilder.AddMenuEntry(
-				OverrideAction->MenuDescription,
-				FText::FromString(OverrideAction->TooltipDescription),
+				OverrideAction->GetMenuDescription(),
+				FText::FromString(OverrideAction->GetTooltipDescription()),
 				FSlateIcon(),
 				FUIAction(FExecuteAction::CreateSP(this, &SMyBlueprint::ImplementFunction, OverrideAction)),
 				NAME_None,
@@ -774,7 +774,7 @@ bool SMyBlueprint::CanRequestRenameOnActionNode(TWeakPtr<FGraphActionNode> InSel
 				FEdGraphSchemaAction_K2Graph* GraphAction = (FEdGraphSchemaAction_K2Graph*)(Actions[i].Get());
 
 				// Do not allow renaming of any graph actions outside of the following
-				if(Actions[i]->SectionID == NodeSectionID::FUNCTION || Actions[i]->SectionID == NodeSectionID::MACRO)
+				if(Actions[i]->GetSectionID() == NodeSectionID::FUNCTION || Actions[i]->GetSectionID() == NodeSectionID::MACRO)
 				{
 					if(!FBlueprintEditorUtils::IsPaletteActionReadOnly(Actions[i], BlueprintEditorPtr.Pin()) && FBlueprintEditorUtils::FindBlueprintForGraph(GraphAction->EdGraph) == GetBlueprintObj())
 					{
@@ -852,10 +852,9 @@ void SMyBlueprint::GetChildGraphs(UEdGraph* InEdGraph, int32 const SectionId, FG
 		FText ChildDesc = DisplayText;
 		const FName DisplayName =  FName(*DisplayText.ToString());
 
-		TSharedPtr<FEdGraphSchemaAction_K2Graph> NewChildAction = MakeShareable(new FEdGraphSchemaAction_K2Graph(EEdGraphSchemaAction_K2Graph::Subgraph, Category, ChildDesc, ChildTooltip, 1));
+		TSharedPtr<FEdGraphSchemaAction_K2Graph> NewChildAction = MakeShareable(new FEdGraphSchemaAction_K2Graph(EEdGraphSchemaAction_K2Graph::Subgraph, Category, ChildDesc, ChildTooltip, 1, SectionId));
 		NewChildAction->FuncName = DisplayName;
 		NewChildAction->EdGraph = Graph;
-		NewChildAction->SectionID = SectionId;
 		OutAllActions.AddAction(NewChildAction);
 		
 		GetChildGraphs(Graph, SectionId, OutAllActions, Category);
@@ -947,9 +946,8 @@ void SMyBlueprint::GetLocalVariables(FGraphActionListBuilderBase& OutAllActions)
 				UFunction* Func = FindField<UFunction>(GetBlueprintObj()->SkeletonGeneratedClass, TopLevelGraph->GetFName());
 				if (Func)
 				{
-					TSharedPtr<FEdGraphSchemaAction_K2LocalVar> NewVarAction = MakeShareable(new FEdGraphSchemaAction_K2LocalVar(Category, FText::FromName(Variable.VarName), TEXT(""), 0));
+					TSharedPtr<FEdGraphSchemaAction_K2LocalVar> NewVarAction = MakeShareable(new FEdGraphSchemaAction_K2LocalVar(Category, FText::FromName(Variable.VarName), TEXT(""), 0, NodeSectionID::LOCAL_VARIABLE));
 					NewVarAction->SetVariableInfo(Variable.VarName, Func, Variable.VarType.PinCategory == K2Schema->PC_Boolean);
-					NewVarAction->SectionID = NodeSectionID::LOCAL_VARIABLE;
 					OutAllActions.AddAction(NewVarAction);
 				}
 			}
@@ -1001,7 +999,7 @@ void SMyBlueprint::CollectAllActions(FGraphActionListBuilderBase& OutAllActions)
 				int32 RootCategoryDelim = Category.Find( TEXT( "|" ));
 				FName RootCategory = RootCategoryDelim == INDEX_NONE ?	*Category : *Category.Left( RootCategoryDelim );
 				// Get root sort index
-				const int32 SortIndex = GetSortIndex( RootCategory ) + Action->SectionID;
+				const int32 SortIndex = GetSortIndex( RootCategory ) + Action->GetSectionID();
 
 				static const FName DefaultCategory( TEXT( "Default" ));
 				//if( RootCategory != DefaultCategory )
@@ -1114,11 +1112,10 @@ void SMyBlueprint::CollectAllActions(FGraphActionListBuilderBase& OutAllActions)
 				PropertyCategory = LOCTEXT("Components", "Components");
 			}
 
-			TSharedPtr<FEdGraphSchemaAction_K2Var> NewVarAction = MakeShareable(new FEdGraphSchemaAction_K2Var(PropertyCategory, PropertyDesc, PropertyTooltip.ToString(), 0));
+			TSharedPtr<FEdGraphSchemaAction_K2Var> NewVarAction = MakeShareable(new FEdGraphSchemaAction_K2Var(PropertyCategory, PropertyDesc, PropertyTooltip.ToString(), 0, NodeSectionID::VARIABLE));
 			const UArrayProperty* ArrayProperty = Cast<const UArrayProperty>(Property);
 			const UProperty* TestProperty = ArrayProperty ? ArrayProperty->Inner : Property;
 			NewVarAction->SetVariableInfo(PropertyName, BlueprintObj->SkeletonGeneratedClass, Cast<UBoolProperty>(TestProperty) != nullptr);
-			NewVarAction->SectionID = NodeSectionID::VARIABLE;
 			SortList.AddAction( UserCategoryName, NewVarAction );
 		}
 		else if (bShouldShowAsDelegate)
@@ -1127,10 +1124,9 @@ void SMyBlueprint::CollectAllActions(FGraphActionListBuilderBase& OutAllActions)
 			// Delegate is visible in MyBlueprint when not-native or its category name is not empty.
 			if (Property->HasAllPropertyFlags(CPF_Edit) || !PropertyCategory.IsEmpty())
 			{
-				NewFuncAction = MakeShareable(new FEdGraphSchemaAction_K2Delegate(PropertyCategory, PropertyDesc, PropertyTooltip.ToString(), 0));
+				NewFuncAction = MakeShareable(new FEdGraphSchemaAction_K2Delegate(PropertyCategory, PropertyDesc, PropertyTooltip.ToString(), 0, NodeSectionID::DELEGATE));
 				NewFuncAction->SetDelegateInfo(PropertyName, BlueprintObj->SkeletonGeneratedClass);
 				NewFuncAction->EdGraph = NULL;
-				NewFuncAction->SectionID = NodeSectionID::DELEGATE;
 				SortList.AddAction( UserCategoryName, NewFuncAction );
 			}
 
@@ -1168,21 +1164,16 @@ void SMyBlueprint::CollectAllActions(FGraphActionListBuilderBase& OutAllActions)
 			FunctionCategory = Function->GetMetaDataText(FBlueprintMetadata::MD_FunctionCategory, TEXT("UObjectCategory"), Function->GetFullGroupName(false));
 		}
 
-		TSharedPtr<FEdGraphSchemaAction_K2Graph> NewFuncAction = MakeShareable(new FEdGraphSchemaAction_K2Graph(EEdGraphSchemaAction_K2Graph::Function, FunctionCategory, DisplayInfo.PlainName, DisplayInfo.Tooltip, bIsConstructionScript ? 2 : 1));
+		//@TODO: Should be a bit more generic (or the AnimGraph shouldn't be stored as a FunctionGraph...)
+		int32 SectionID = Graph->IsA<UAnimationGraph>() ? NodeSectionID::GRAPH : NodeSectionID::FUNCTION;
+		TSharedPtr<FEdGraphSchemaAction_K2Graph> NewFuncAction = MakeShareable(new FEdGraphSchemaAction_K2Graph(EEdGraphSchemaAction_K2Graph::Function, FunctionCategory, DisplayInfo.PlainName, DisplayInfo.Tooltip, bIsConstructionScript ? 2 : 1, SectionID));
 		NewFuncAction->FuncName = Graph->GetFName();
 		NewFuncAction->EdGraph = Graph;
-		NewFuncAction->SectionID = NodeSectionID::FUNCTION;
-
-		//@TODO: Should be a bit more generic (or the AnimGraph shouldn't be stored as a FunctionGraph...)
-		if (Graph->IsA<UAnimationGraph>())
-		{
-			NewFuncAction->SectionID = NodeSectionID::GRAPH;
-		}
 
 		OutAllActions.AddAction(NewFuncAction);
 
-		GetChildGraphs(Graph, NewFuncAction->SectionID, OutAllActions, FunctionCategory);
-		GetChildEvents(Graph, NewFuncAction->SectionID, OutAllActions, FunctionCategory);
+		GetChildGraphs(Graph, NewFuncAction->GetSectionID(), OutAllActions, FunctionCategory);
+		GetChildEvents(Graph, NewFuncAction->GetSectionID(), OutAllActions, FunctionCategory);
 
 		ImplementedFunctionCache.Add(Graph->GetFName());
 	}
@@ -1199,14 +1190,13 @@ void SMyBlueprint::CollectAllActions(FGraphActionListBuilderBase& OutAllActions)
 
 		FText MacroCategory = GetGraphCategory(Graph);
 
-		TSharedPtr<FEdGraphSchemaAction_K2Graph> NewMacroAction = MakeShareable(new FEdGraphSchemaAction_K2Graph(EEdGraphSchemaAction_K2Graph::Macro, MacroCategory, DisplayInfo.PlainName, DisplayInfo.Tooltip, 1));
-		NewMacroAction->SectionID = NodeSectionID::MACRO;
+		TSharedPtr<FEdGraphSchemaAction_K2Graph> NewMacroAction = MakeShareable(new FEdGraphSchemaAction_K2Graph(EEdGraphSchemaAction_K2Graph::Macro, MacroCategory, DisplayInfo.PlainName, DisplayInfo.Tooltip, 1, NodeSectionID::MACRO));
 		NewMacroAction->FuncName = MacroName;
 		NewMacroAction->EdGraph = Graph;
 		OutAllActions.AddAction(NewMacroAction);
 
-		GetChildGraphs(Graph, NewMacroAction->SectionID, OutAllActions, MacroCategory);
-		GetChildEvents(Graph, NewMacroAction->SectionID, OutAllActions, MacroCategory);
+		GetChildGraphs(Graph, NewMacroAction->GetSectionID(), OutAllActions, MacroCategory);
+		GetChildEvents(Graph, NewMacroAction->GetSectionID(), OutAllActions, MacroCategory);
 
 		ImplementedFunctionCache.Add(MacroName);
 	}
@@ -1237,9 +1227,8 @@ void SMyBlueprint::CollectAllActions(FGraphActionListBuilderBase& OutAllActions)
 
 			FText FunctionCategory = Function->GetMetaDataText(FBlueprintMetadata::MD_FunctionCategory, TEXT("UObjectCategory"), Function->GetFullGroupName(false));
 
-			TSharedPtr<FEdGraphSchemaAction_K2Graph> NewFuncAction = MakeShareable(new FEdGraphSchemaAction_K2Graph(EEdGraphSchemaAction_K2Graph::Function, FunctionCategory, FunctionDesc, FunctionTooltip, 1));
+			TSharedPtr<FEdGraphSchemaAction_K2Graph> NewFuncAction = MakeShareable(new FEdGraphSchemaAction_K2Graph(EEdGraphSchemaAction_K2Graph::Function, FunctionCategory, FunctionDesc, FunctionTooltip, 1, NodeSectionID::FUNCTION_OVERRIDABLE));
 			NewFuncAction->FuncName = FunctionName;
-			NewFuncAction->SectionID = NodeSectionID::FUNCTION_OVERRIDABLE;
 
 			OverridableFunctionActions.Add(NewFuncAction);
 			OverridableFunctionNames.Add(FunctionName);
@@ -1265,14 +1254,13 @@ void SMyBlueprint::CollectAllActions(FGraphActionListBuilderBase& OutAllActions)
 				FunctionCategory = Function->GetMetaDataText(FBlueprintMetadata::MD_FunctionCategory, TEXT("UObjectCategory"), Function->GetFullGroupName(false));
 			}
 
-			TSharedPtr<FEdGraphSchemaAction_K2Graph> NewFuncAction = MakeShareable(new FEdGraphSchemaAction_K2Graph(EEdGraphSchemaAction_K2Graph::Interface, FunctionCategory, FText::FromString(FunctionDesc), FunctionTooltip, 1));
+			TSharedPtr<FEdGraphSchemaAction_K2Graph> NewFuncAction = MakeShareable(new FEdGraphSchemaAction_K2Graph(EEdGraphSchemaAction_K2Graph::Interface, FunctionCategory, FText::FromString(FunctionDesc), FunctionTooltip, 1, NodeSectionID::INTERFACE));
 			NewFuncAction->FuncName = FunctionName;
 			NewFuncAction->EdGraph = Graph;
-			NewFuncAction->SectionID = NodeSectionID::INTERFACE;
 			OutAllActions.AddAction(NewFuncAction);
 
-			GetChildGraphs(Graph, NewFuncAction->SectionID, OutAllActions, FunctionCategory);
-			GetChildEvents(Graph, NewFuncAction->SectionID, OutAllActions, FunctionCategory);
+			GetChildGraphs(Graph, NewFuncAction->GetSectionID(), OutAllActions, FunctionCategory);
+			GetChildEvents(Graph, NewFuncAction->GetSectionID(), OutAllActions, FunctionCategory);
 		}
 	}
 
@@ -1315,14 +1303,13 @@ void SMyBlueprint::CollectAllActions(FGraphActionListBuilderBase& OutAllActions)
 		FGraphDisplayInfo DisplayInfo;
 		Graph->GetSchema()->GetGraphDisplayInformation(*Graph, DisplayInfo);
 
-		TSharedPtr<FEdGraphSchemaAction_K2Graph> NeUbergraphAction = MakeShareable(new FEdGraphSchemaAction_K2Graph(EEdGraphSchemaAction_K2Graph::Graph, FText::GetEmpty(), DisplayInfo.PlainName, DisplayInfo.Tooltip, 2));
+		TSharedPtr<FEdGraphSchemaAction_K2Graph> NeUbergraphAction = MakeShareable(new FEdGraphSchemaAction_K2Graph(EEdGraphSchemaAction_K2Graph::Graph, FText::GetEmpty(), DisplayInfo.PlainName, DisplayInfo.Tooltip, 2, NodeSectionID::GRAPH));
 		NeUbergraphAction->FuncName = Graph->GetFName();
 		NeUbergraphAction->EdGraph = Graph;
-		NeUbergraphAction->SectionID = NodeSectionID::GRAPH;
 		OutAllActions.AddAction(NeUbergraphAction);
 
-		GetChildGraphs(Graph, NeUbergraphAction->SectionID, OutAllActions);
-		GetChildEvents(Graph, NeUbergraphAction->SectionID, OutAllActions);
+		GetChildGraphs(Graph, NeUbergraphAction->GetSectionID(), OutAllActions);
+		GetChildEvents(Graph, NeUbergraphAction->GetSectionID(), OutAllActions);
 	}
 
 	// Grab intermediate pages
@@ -1339,8 +1326,8 @@ void SMyBlueprint::CollectAllActions(FGraphActionListBuilderBase& OutAllActions)
 		NewIntermediateAction->EdGraph = Graph;
 		OutAllActions.AddAction(NewIntermediateAction);
 
-		GetChildGraphs(Graph, NewIntermediateAction->SectionID, OutAllActions);
-		GetChildEvents(Graph, NewIntermediateAction->SectionID, OutAllActions);
+		GetChildGraphs(Graph, NewIntermediateAction->GetSectionID(), OutAllActions);
+		GetChildEvents(Graph, NewIntermediateAction->GetSectionID(), OutAllActions);
 	}
 
 	if(GetLocalActionsListVisibility().IsVisible())
@@ -2134,13 +2121,13 @@ void SMyBlueprint::OnFindReference()
 	}
 	else if (FEdGraphSchemaAction_K2Event* EventAction = SelectionAsEvent())
 	{
-		SearchTerm = EventAction->MenuDescription.ToString();
+		SearchTerm = EventAction->GetMenuDescription().ToString();
 	}
 	else if (FEdGraphSchemaAction_K2InputAction* InputAction = SelectionAsInputAction())
 	{
 		SearchTerm = InputAction->NodeTemplate ? 
 			InputAction->NodeTemplate->GetNodeTitle(ENodeTitleType::FullTitle).ToString() :
-			InputAction->MenuDescription.ToString();
+			InputAction->GetMenuDescription().ToString();
 	}
 
 	if(!SearchTerm.IsEmpty())
@@ -2327,7 +2314,7 @@ void SMyBlueprint::OnDeleteEntry()
 		{
 			FText TransactionTitle;
 
-			switch((NodeSectionID::Type)Actions[0]->SectionID)
+			switch((NodeSectionID::Type)Actions[0]->GetSectionID())
 			{
 			case NodeSectionID::VARIABLE:
 			case NodeSectionID::LOCAL_VARIABLE:
diff --git a/Engine/Source/Editor/Kismet/Public/BlueprintActionMenuItem.h b/Engine/Source/Editor/Kismet/Public/BlueprintActionMenuItem.h
--- a/Engine/Source/Editor/Kismet/Public/BlueprintActionMenuItem.h
+++ b/Engine/Source/Editor/Kismet/Public/BlueprintActionMenuItem.h
@@ -28,7 +28,7 @@ struct FBlueprintActionMenuItem : public FEdGraphSchemaAction
 public:	
 	/** Constructors */
 	FBlueprintActionMenuItem(UBlueprintNodeSpawner const* NodeSpawner = nullptr) : Action(NodeSpawner), IconTint(FLinearColor::White), IconBrush(nullptr) {}
-	FBlueprintActionMenuItem(UBlueprintNodeSpawner const* NodeSpawner, FBlueprintActionUiSpec const& UiDef, IBlueprintNodeBinder::FBindingSet const& Bindings = IBlueprintNodeBinder::FBindingSet());
+	FBlueprintActionMenuItem(UBlueprintNodeSpawner const* NodeSpawner, FBlueprintActionUiSpec const& UiDef, IBlueprintNodeBinder::FBindingSet const& Bindings = IBlueprintNodeBinder::FBindingSet(), const FText& InNodeCategory = FText(), int32 InGrouping = 0);
 	
 	// FEdGraphSchemaAction interface
 	virtual FName         GetTypeId() const final { return StaticGetTypeId(); }
diff --git a/Engine/Source/Editor/Kismet/Public/BlueprintDragDropMenuItem.h b/Engine/Source/Editor/Kismet/Public/BlueprintDragDropMenuItem.h
--- a/Engine/Source/Editor/Kismet/Public/BlueprintDragDropMenuItem.h
+++ b/Engine/Source/Editor/Kismet/Public/BlueprintDragDropMenuItem.h
@@ -43,7 +43,7 @@ public:
 	 * @param  SampleAction	One of the node-spawners that this "consolidated" menu item represents (shouldn't matter which one).
 	 * @param  MenuGrouping	Defines the order in which this entry is listed in the menu.
 	 */
-	FBlueprintDragDropMenuItem(FBlueprintActionContext const& Context, UBlueprintNodeSpawner const* SampleAction, int32 MenuGrouping = 0);
+	FBlueprintDragDropMenuItem(FBlueprintActionContext const& Context, UBlueprintNodeSpawner const* SampleAction, int32 MenuGrouping, const FText& InNodeCategory, const FText& InMenuDesc, const FString& InToolTip );
 	
 	// FEdGraphSchemaAction interface
 	virtual FName GetTypeId() const final { return StaticGetTypeId(); }
diff --git a/Engine/Source/Editor/MaterialEditor/Private/MaterialEditorUtilities.cpp b/Engine/Source/Editor/MaterialEditor/Private/MaterialEditorUtilities.cpp
--- a/Engine/Source/Editor/MaterialEditor/Private/MaterialEditorUtilities.cpp
+++ b/Engine/Source/Editor/MaterialEditor/Private/MaterialEditorUtilities.cpp
@@ -577,9 +577,8 @@ void FMaterialEditorUtilities::AddMaterialExpressionCategory(FGraphActionMenuBui
 				TSharedPtr<FMaterialGraphSchemaAction_NewNode> NewNodeAction(new FMaterialGraphSchemaAction_NewNode(
 					CategoryName,
 					FText::FromString(MaterialExpression.Name),
-					ToolTip.ToString(), 0));
+					ToolTip.ToString(), 0, CastChecked<UMaterialExpression>(MaterialExpression.MaterialClass->GetDefaultObject())->GetKeywords()));
 				NewNodeAction->MaterialExpressionClass = MaterialExpression.MaterialClass;
-				NewNodeAction->Keywords = CastChecked<UMaterialExpression>(MaterialExpression.MaterialClass->GetDefaultObject())->GetKeywords();
 				ActionMenuBuilder.AddAction(NewNodeAction);
 			}
 		}
diff --git a/Engine/Source/Editor/MaterialEditor/Private/SMaterialPalette.cpp b/Engine/Source/Editor/MaterialEditor/Private/SMaterialPalette.cpp
--- a/Engine/Source/Editor/MaterialEditor/Private/SMaterialPalette.cpp
+++ b/Engine/Source/Editor/MaterialEditor/Private/SMaterialPalette.cpp
@@ -38,7 +38,7 @@ void SMaterialPaletteItem::Construct(const FArguments& InArgs, FCreateWidgetForA
 	// Find icons
 	const FSlateBrush* IconBrush = FEditorStyle::GetBrush(TEXT("NoBrush"));
 	FSlateColor IconColor = FSlateColor::UseForeground();
-	FString ToolTip = GraphAction->TooltipDescription;
+	FString ToolTip = GraphAction->GetTooltipDescription();
 	FText IconToolTip = FText::FromString(ToolTip);
 	bool bIsReadOnly = false;
 
diff --git a/Engine/Source/Editor/SoundClassEditor/Private/SSoundClassActionMenu.cpp b/Engine/Source/Editor/SoundClassEditor/Private/SSoundClassActionMenu.cpp
--- a/Engine/Source/Editor/SoundClassEditor/Private/SSoundClassActionMenu.cpp
+++ b/Engine/Source/Editor/SoundClassEditor/Private/SSoundClassActionMenu.cpp
@@ -26,21 +26,21 @@ void SSoundClassActionMenuItem::Construct(const FArguments& InArgs, TSharedPtr<F
 		this->ChildSlot
 			[
 				SNew(SHorizontalBox)
-				.ToolTipText( FText::FromString(InAction->TooltipDescription) )
+				.ToolTipText( FText::FromString(InAction->GetTooltipDescription()) )
 				+ SHorizontalBox::Slot()
 				.AutoWidth()
 				.VAlign(VAlign_Center)
 				[
 					SNew(STextBlock)
 					.Font(FSlateFontInfo( FPaths::EngineContentDir() / TEXT("Slate/Fonts/Roboto-Regular.ttf"), 9 ))
-					.Text(InAction->MenuDescription)
+					.Text(InAction->GetMenuDescription())
 					.HighlightText(InArgs._HighlightText)
 				]
 			];
 	}
 	else
 	{
-		TSharedRef<SWidget> NewSoundClassWidget = CreateNewSoundClassWidget( InAction->MenuDescription, FText::FromString(InAction->TooltipDescription), FSlateFontInfo( FPaths::EngineContentDir() / TEXT("Slate/Fonts/Roboto-Regular.ttf"), 9), InAction );
+		TSharedRef<SWidget> NewSoundClassWidget = CreateNewSoundClassWidget(InAction->GetMenuDescription(), FText::FromString(InAction->GetTooltipDescription()), FSlateFontInfo(FPaths::EngineContentDir() / TEXT("Slate/Fonts/Roboto-Regular.ttf"), 9), InAction);
 		// Promote requires 2 'slots'
 		this->ChildSlot
 			[
diff --git a/Engine/Source/Editor/UMGEditor/Private/Customizations/UMGDetailCustomizations.cpp b/Engine/Source/Editor/UMGEditor/Private/Customizations/UMGDetailCustomizations.cpp
--- a/Engine/Source/Editor/UMGEditor/Private/Customizations/UMGDetailCustomizations.cpp
+++ b/Engine/Source/Editor/UMGEditor/Private/Customizations/UMGDetailCustomizations.cpp
@@ -42,7 +42,7 @@ public:
 			.TextStyle(FEditorStyle::Get(), "NormalText")
 			.HAlign(HAlign_Center)
 			.ForegroundColor(FSlateColor::UseForeground())
-			.ToolTipText(FText::FromString(Action->TooltipDescription))
+			.ToolTipText(FText::FromString(Action->GetTooltipDescription()))
 			.OnClicked(this, &SGraphSchemaActionButton::AddOrViewEventBinding)
 			[
 				InArgs._Content.Widget
diff --git a/Engine/Source/Editor/UnrealEd/Classes/MaterialGraph/MaterialGraphSchema.h b/Engine/Source/Editor/UnrealEd/Classes/MaterialGraph/MaterialGraphSchema.h
--- a/Engine/Source/Editor/UnrealEd/Classes/MaterialGraph/MaterialGraphSchema.h
+++ b/Engine/Source/Editor/UnrealEd/Classes/MaterialGraph/MaterialGraphSchema.h
@@ -23,8 +23,8 @@ struct UNREALED_API FMaterialGraphSchemaAction_NewNode : public FEdGraphSchemaAc
 		, MaterialExpressionClass(NULL)
 	{}
 
-	FMaterialGraphSchemaAction_NewNode(const FText& InNodeCategory, const FText& InMenuDesc, const FString& InToolTip, const int32 InGrouping)
-		: FEdGraphSchemaAction(InNodeCategory, InMenuDesc, InToolTip, InGrouping) 
+	FMaterialGraphSchemaAction_NewNode(const FText& InNodeCategory, const FText& InMenuDesc, const FString& InToolTip, const int32 InGrouping, const FText& InKeywords)
+		: FEdGraphSchemaAction(InNodeCategory, InMenuDesc, InToolTip, InGrouping, InKeywords)
 		, MaterialExpressionClass(NULL)
 	{}
 
diff --git a/Engine/Source/Editor/UnrealEd/Private/Commandlets/DumpBlueprintsInfoCommandlet.cpp b/Engine/Source/Editor/UnrealEd/Private/Commandlets/DumpBlueprintsInfoCommandlet.cpp
--- a/Engine/Source/Editor/UnrealEd/Private/Commandlets/DumpBlueprintsInfoCommandlet.cpp
+++ b/Engine/Source/Editor/UnrealEd/Private/Commandlets/DumpBlueprintsInfoCommandlet.cpp
@@ -1016,8 +1016,7 @@ static FString DumpBlueprintInfoUtils::GetGraphTypeString(UEdGraph const* Graph)
 //------------------------------------------------------------------------------
 static FString DumpBlueprintInfoUtils::GetActionKey(FGraphActionListBuilderBase::ActionGroup const& Action)
 {
-	TArray<FString> MenuHierarchy;
-	Action.GetCategoryChain(MenuHierarchy);
+	const TArray<FString>& MenuHierarchy = Action.GetCategoryChain();
 
 	FString ActionKey;
 	for (FString const& SubCategory : MenuHierarchy)
@@ -1030,7 +1029,7 @@ static FString DumpBlueprintInfoUtils::GetActionKey(FGraphActionListBuilderBase:
 	}
 
 	TSharedPtr<FEdGraphSchemaAction> MainAction = Action.Actions[0];
-	ActionKey += MainAction->MenuDescription.ToString();
+	ActionKey += MainAction->GetMenuDescription().ToString();
 
 	return ActionKey;
 }
@@ -1607,9 +1606,9 @@ static void DumpBlueprintInfoUtils::DumpActionList(uint32 Indent, FGraphActionLi
 			TSharedPtr<FEdGraphSchemaAction> LHSAction = LHS.Actions[0];
 			TSharedPtr<FEdGraphSchemaAction> RHSAction = RHS.Actions[0];
 			
-			if (LHSAction->Grouping != RHSAction->Grouping)
+			if (LHSAction->GetGrouping() != RHSAction->GetGrouping())
 			{
-				return LHSAction->Grouping > RHSAction->Grouping;
+				return LHSAction->GetGrouping() > RHSAction->GetGrouping();
 			}
 			
 			FString LhKey = GetActionKey(LHS);
@@ -1647,8 +1646,7 @@ static void DumpBlueprintInfoUtils::DumpActionMenuItem(uint32 Indent, FGraphActi
 	check(Action.Actions.Num() > 0);
 
 	// Get action category info
-	TArray<FString> MenuHierarchy;
-	Action.GetCategoryChain(MenuHierarchy);
+	TArray<FString> MenuHierarchy = Action.GetCategoryChain();
 
 	FString ActionCategory = TEXT("");
 
@@ -1662,7 +1660,7 @@ static void DumpBlueprintInfoUtils::DumpActionMenuItem(uint32 Indent, FGraphActi
 	}
 
 	TSharedPtr<FEdGraphSchemaAction> PrimeAction = Action.Actions[0];
-	FString const ActionName = PrimeAction->MenuDescription.ToString();
+	FString const ActionName = PrimeAction->GetMenuDescription().ToString();
 
 	FString const ActionEntryIndent = BuildIndentString(Indent);
 	FString ActionEntry = ActionEntryIndent + "\"" + ActionCategory + ActionName + "\"";;
@@ -1684,14 +1682,14 @@ static void DumpBlueprintInfoUtils::DumpActionMenuItem(uint32 Indent, FGraphActi
 		}
 		ActionEntry += "\","; // end action category data
 
-		FString TooltipStr = PrimeAction->TooltipDescription.Replace(TEXT("\\\""), TEXT("'")).Replace(TEXT("\""), TEXT("'"));
+		FString TooltipStr = PrimeAction->GetTooltipDescription().Replace(TEXT("\\\""), TEXT("'")).Replace(TEXT("\""), TEXT("'"));
 		FString const TooltipFieldLabel("\"Tooltip\"     : \"");
 		TooltipStr = TooltipStr.Replace(TEXT("\n"), *(IndentedNewline + BuildIndentString(TooltipFieldLabel.Len(), /*bUseSpaces =*/true)));
 
 		ActionEntry += IndentedNewline + TooltipFieldLabel + TooltipStr + "\",";
-		ActionEntry += IndentedNewline + "\"Keywords\"    : \"" + PrimeAction->Keywords.ToString() + "\",";
-		ActionEntry += IndentedNewline + "\"SearchTitle\" : \"" + PrimeAction->GetSearchTitle() + "\",";
-		ActionEntry += IndentedNewline + FString::Printf(TEXT("\"Grouping\"    : %d"), PrimeAction->Grouping);
+		ActionEntry += IndentedNewline + "\"Keywords\"    : \"" + PrimeAction->GetKeywords().ToString() + "\",";
+		ActionEntry += IndentedNewline + "\"SearchTitle\" : \"" + PrimeAction->GetMenuDescription().ToString() + "\",";
+		ActionEntry += IndentedNewline + FString::Printf(TEXT("\"Grouping\"    : %d"), PrimeAction->GetGrouping());
 		
 		// Get action node type info
 		UK2Node const* NodeTemplate = FBlueprintActionMenuUtils::ExtractNodeTemplateFromAction(PrimeAction);
diff --git a/Engine/Source/Editor/UnrealEd/Private/Commandlets/GenerateBlueprintAPICommandlet.cpp b/Engine/Source/Editor/UnrealEd/Private/Commandlets/GenerateBlueprintAPICommandlet.cpp
--- a/Engine/Source/Editor/UnrealEd/Private/Commandlets/GenerateBlueprintAPICommandlet.cpp
+++ b/Engine/Source/Editor/UnrealEd/Private/Commandlets/GenerateBlueprintAPICommandlet.cpp
@@ -519,8 +519,7 @@ static FString GenerateBlueprintAPIUtils::BuildIndentString(uint32 IndentCount,
 //------------------------------------------------------------------------------
 static FString GenerateBlueprintAPIUtils::GetActionKey(FGraphActionListBuilderBase::ActionGroup const& Action)
 {
-	TArray<FString> MenuHierarchy;
-	Action.GetCategoryChain(MenuHierarchy);
+	TArray<FString> MenuHierarchy = Action.GetCategoryChain();
 
 	FString ActionKey;
 	for (FString const& SubCategory : MenuHierarchy)
@@ -533,7 +532,7 @@ static FString GenerateBlueprintAPIUtils::GetActionKey(FGraphActionListBuilderBa
 	}
 
 	TSharedPtr<FEdGraphSchemaAction> MainAction = Action.Actions[0];
-	ActionKey += MainAction->MenuDescription.ToString();
+	ActionKey += MainAction->GetMenuDescription().ToString();
 
 	return ActionKey;
 }
@@ -734,9 +733,9 @@ static void GenerateBlueprintAPIUtils::DumpActionList(uint32 Indent, FGraphActio
 			TSharedPtr<FEdGraphSchemaAction> LHSAction = LHS.Actions[0];
 			TSharedPtr<FEdGraphSchemaAction> RHSAction = RHS.Actions[0];
 			
-			if (LHSAction->Grouping != RHSAction->Grouping)
+			if (LHSAction->GetGrouping() != RHSAction->GetGrouping())
 			{
-				return LHSAction->Grouping > RHSAction->Grouping;
+				return LHSAction->GetGrouping() > RHSAction->GetGrouping();
 			}
 			
 			FString LhKey = GetActionKey(LHS);
@@ -767,8 +766,7 @@ static void GenerateBlueprintAPIUtils::DumpActionMenuItem(uint32 Indent, FGraphA
 	check(Action.Actions.Num() > 0);
 
 	// Get action category info
-	TArray<FString> MenuHierarchy;
-	Action.GetCategoryChain(MenuHierarchy);
+	const TArray<FString>& MenuHierarchy = Action.GetCategoryChain();
 
 	FString ActionCategory = TEXT("");
 
@@ -798,7 +796,7 @@ static void GenerateBlueprintAPIUtils::DumpActionMenuItem(uint32 Indent, FGraphA
 	}
 
 	TSharedPtr<FEdGraphSchemaAction> PrimeAction = Action.Actions[0];
-	const FString ActionName = PrimeAction->MenuDescription.ToString();
+	const FString ActionName = PrimeAction->GetMenuDescription().ToString();
 
 	const FString ActionEntryIndent = BuildIndentString(Indent);
 	FString ActionEntry = ActionEntryIndent + "\"" + MakeJsonString(ActionCategory + ActionName) + "\"";
@@ -809,7 +807,7 @@ static void GenerateBlueprintAPIUtils::DumpActionMenuItem(uint32 Indent, FGraphA
 	ActionEntry += " : {";
 
 	const FString TooltipFieldLabel("\"Tooltip\"      : \"");
-	const FString TooltipStr = PrimeAction->TooltipDescription.Replace(TEXT("\n"), *(IndentedNewline + BuildIndentString(TooltipFieldLabel.Len(), /*bUseSpaces =*/true)));
+	const FString TooltipStr = PrimeAction->GetTooltipDescription().Replace(TEXT("\n"), *(IndentedNewline + BuildIndentString(TooltipFieldLabel.Len(), /*bUseSpaces =*/true)));
 
 	ActionEntry += IndentedNewline + TooltipFieldLabel + MakeJsonString(TooltipStr) + "\"";
 		
diff --git a/Engine/Source/Editor/UnrealEd/Private/EditorCategoryUtils.cpp b/Engine/Source/Editor/UnrealEd/Private/EditorCategoryUtils.cpp
--- a/Engine/Source/Editor/UnrealEd/Private/EditorCategoryUtils.cpp
+++ b/Engine/Source/Editor/UnrealEd/Private/EditorCategoryUtils.cpp
@@ -336,15 +336,20 @@ bool FEditorCategoryUtils::IsCategoryHiddenFromClass(UClass const* Class, FText
 //------------------------------------------------------------------------------
 bool FEditorCategoryUtils::IsCategoryHiddenFromClass(UClass const* Class, FString const& Category)
 {
-	bool bIsHidden = false;
-
 	TArray<FString> ClassHideCategories;
 	GetClassHideCategories(Class, ClassHideCategories);
+	return IsCategoryHiddenFromClass(ClassHideCategories, Class, Category);
+}
+
+//------------------------------------------------------------------------------
+bool FEditorCategoryUtils::IsCategoryHiddenFromClass(const TArray<FString>& ClassHideCategories, UClass const* Class, const FString& Category)
+{
+	bool bIsHidden = false;
 
 	// run the category through sanitization so we can ensure compares will hit
 	FString const DisplayCategory = GetCategoryDisplayString(Category);
 
-	for (FString& HideCategory : ClassHideCategories)
+	for (const FString& HideCategory : ClassHideCategories)
 	{
 		bIsHidden = (HideCategory == DisplayCategory);
 		if (bIsHidden)
diff --git a/Engine/Source/Editor/UnrealEd/Public/EditorCategoryUtils.h b/Engine/Source/Editor/UnrealEd/Public/EditorCategoryUtils.h
--- a/Engine/Source/Editor/UnrealEd/Public/EditorCategoryUtils.h
+++ b/Engine/Source/Editor/UnrealEd/Public/EditorCategoryUtils.h
@@ -154,7 +154,7 @@ namespace FEditorCategoryUtils
 	 * @param  Category A category path that you want to check.
 	 * @return True if the category is hidden, false if not.
 	 */
-	UNREALED_API bool IsCategoryHiddenFromClass(UClass const* Class, FText const& Category);
+	UNREALED_API bool IsCategoryHiddenFromClass(const UClass* Class, const FText& Category);
 
 	/**
 	 * Checks to see if the specified category is hidden from the supplied class.
@@ -163,8 +163,18 @@ namespace FEditorCategoryUtils
 	 * @param  Category A category path that you want to check.
 	 * @return True if the category is hidden, false if not.
 	 */
-	UNREALED_API bool IsCategoryHiddenFromClass(UClass const* Class, FString const& Category);
+	UNREALED_API bool IsCategoryHiddenFromClass(const UClass* Class, const FString& Category);
 
+	/**
+	 * Checks to see if the specified category is hidden from the supplied Class, avoids recalculation of ClassHideCategories.
+	 * Useful when checking the same class over and over again with different categories.
+	 *
+	 * @param ClassHideCategories	The categories tht have been hidden for the class
+	 * @param  Class	The class you want to query.
+	 * @param  Category A category path that you want to check.
+	 * @return True if the category is hidden, false if not.
+	 */
+	UNREALED_API bool IsCategoryHiddenFromClass(const TArray<FString>& ClassHideCategories, UClass const* Class, const FString& Category);
 	/**
 	 * Returns tooltip information for the specified category
 	 * 
diff --git a/Engine/Source/Runtime/Engine/Classes/EdGraph/EdGraphSchema.h b/Engine/Source/Runtime/Engine/Classes/EdGraph/EdGraphSchema.h
--- a/Engine/Source/Runtime/Engine/Classes/EdGraph/EdGraphSchema.h
+++ b/Engine/Source/Runtime/Engine/Classes/EdGraph/EdGraphSchema.h
@@ -56,6 +56,7 @@ struct ENGINE_API FEdGraphSchemaAction
 	static FName StaticGetTypeId() {static FName Type("FEdGraphSchemaAction"); return Type;}
 	virtual FName GetTypeId() const { return StaticGetTypeId(); }
 
+private:
 	/** The menu text that should be displayed for this node in the creation menu. */
 	UPROPERTY()
 	FText MenuDescription;
@@ -72,6 +73,7 @@ struct ENGINE_API FEdGraphSchemaAction
 	UPROPERTY()
 	FText Keywords;
 
+public:
 	/** This is a priority number for overriding alphabetical order in the action list (higher value  == higher in the list). */
 	UPROPERTY()
 	int32 Grouping;
@@ -80,6 +82,20 @@ struct ENGINE_API FEdGraphSchemaAction
 	UPROPERTY()
 	int32 SectionID;
 
+	UPROPERTY()
+	TArray<FString> MenuDescriptionArray;
+
+	UPROPERTY()
+	TArray<FString> FullSearchTitlesArray;
+
+	UPROPERTY()
+	TArray<FString>  FullSearchKeywordsArray;
+
+	UPROPERTY()
+	TArray<FString>  FullSearchCategoryArray;
+
+	UPROPERTY()
+	FString SearchText;
 	FEdGraphSchemaAction() 
 		: Grouping(0)
 		, SectionID(0)
@@ -87,13 +103,11 @@ struct ENGINE_API FEdGraphSchemaAction
 	
 	virtual ~FEdGraphSchemaAction() {}
 
-	FEdGraphSchemaAction(const FText& InNodeCategory, const FText& InMenuDesc, const FString& InToolTip, const int32 InGrouping)
-		: MenuDescription(InMenuDesc)
-		, TooltipDescription(InToolTip)
-		, Category(InNodeCategory)
-		, Grouping(InGrouping)
-		, SectionID(0)
+	FEdGraphSchemaAction(const FText& InNodeCategory, const FText& InMenuDesc, const FString& InToolTip, const int32 InGrouping, const FText& InKeywords = FText(), int32 InSectionID = 0)
+		: Grouping(InGrouping)
+		, SectionID(InSectionID)
 	{
+		UpdateSearchData(InMenuDesc, InToolTip, InNodeCategory, InKeywords);
 	}
 
 	/** Whether or not this action can be parented to other actions of the same type. */
@@ -118,28 +132,66 @@ struct ENGINE_API FEdGraphSchemaAction
 		return NewNode;
 	}
 
+	void UpdateCategory(const FText& NewCategory);
+
+	void UpdateSearchData(const FText& NewMenuDescription, const FString& NewToolTipDescription, const FText& NewCategory, const FText& NewKeywords);
+
+	int32 GetSectionID() const
+	{
+		return SectionID;
+	}
+
+	int32 GetGrouping() const 
+	{
+		return Grouping;
+	}
+
+	const FText& GetMenuDescription() const
+	{
+		return MenuDescription;
+	}
+
+	const FString& GetTooltipDescription() const
+	{
+		return TooltipDescription;
+	}
+
+	const FText& GetCategory() const
+	{
+		return Category;
+	}
+
+	const FText& GetKeywords() const
+	{
+		return Keywords;
+	}
+
+	const TArray<FString>& GetMenuDescriptionArray() const
+	{
+		return MenuDescriptionArray;
+	}
+
 	/** Retrieves the full searchable title for this action. */
-	FString GetSearchTitle()
+	const TArray<FString>& GetSearchTitleArray() const
 	{
-		return (MenuDescription.ToString() +
-			TEXT(' ') +
-			MenuDescription.BuildSourceString()).ToLower();
+		return FullSearchTitlesArray;
 	}
 
 	/** Retrieves the full searchable keywords for this action. */
-	FString GetSearchKeywords()
+	const TArray<FString>& GetSearchKeywordsArray() const
 	{
-		return (Keywords.ToString() +
-			TEXT(' ') +
-			Keywords.BuildSourceString()).ToLower();
+		return FullSearchKeywordsArray;
 	}
 
 	/** Retrieves the full searchable categories for this action. */
-	FString GetSearchCategory()
+	const TArray<FString>& GetSearchCategoryArray() const
+	{
+		return FullSearchCategoryArray;
+	}
+
+	const FString& GetFullSearchText() const
 	{
-		return (Category.ToString() +
-			TEXT(' ') +
-			Category.BuildSourceString()).ToLower();
+		return SearchText;
 	}
 
 	// GC.
@@ -280,10 +332,9 @@ public:
 
 		ENGINE_API ~ActionGroup();
 		/**
-		 * 
-		 * @param  HierarchyOut	A list of the category tiers that this action should be listed under.
+		 * @return  A reference to the array of strings that represent the category chain
 		 */
-		ENGINE_API void GetCategoryChain(TArray<FString>& HierarchyOut) const;
+		 ENGINE_API const TArray<FString>& GetCategoryChain() const;
 
 		/**
 		 * Goes through all actions and calls PerformAction on them individually
@@ -296,16 +347,16 @@ public:
 		/**
 		 * Returns a the string that should be used when searching for matching actions. Looks only at the first action.
 		 */
-		ENGINE_API const FString& GetSearchTextForFirstAction() const { return SearchText; }
+		ENGINE_API const FString& GetSearchTextForFirstAction() const { return Actions[0]->GetFullSearchText(); }
 
 		/** Returns the SearchKeywordsArray */
-		ENGINE_API const TArray<FString>& GetSearchKeywordsArrayForFirstAction() const { return SearchKeywordsArray; }
+		ENGINE_API const TArray<FString>& GetSearchKeywordsArrayForFirstAction() const { return Actions[0]->GetSearchKeywordsArray(); }
 		/** Returns the MenuDescriptionArray */
-		ENGINE_API const TArray<FString>& GetMenuDescriptionArrayForFirstAction() const { return MenuDescriptionArray; }
+		ENGINE_API const TArray<FString>& GetMenuDescriptionArrayForFirstAction() const { return Actions[0]->GetMenuDescriptionArray(); }
 		/** Returns the SearchTitleArray */
-		ENGINE_API const TArray<FString>& GetSearchTitleArrayForFirstAction() const { return SearchTitleArray; }
+		ENGINE_API const TArray<FString>& GetSearchTitleArrayForFirstAction() const { return Actions[0]->GetSearchTitleArray(); }
 		/** Returns the SearchCategoryArray */
-		ENGINE_API const TArray<FString>& GetSearchCategoryArrayForFirstAction() const { return SearchCategoryArray; }
+		ENGINE_API const TArray<FString>& GetSearchCategoryArrayForFirstAction() const { return Actions[0]->GetSearchCategoryArray(); }
 
 		/** All of the actions this entry contains */
 		TArray< TSharedPtr<FEdGraphSchemaAction> > Actions;
@@ -326,21 +377,10 @@ public:
 		 */
 		void InitSearchText();
 
-		void InitScoringData();
-
 		/** The category to list this entry under (could be left empty, as it gets concatenated with the first sub-action's category) */
 		FString RootCategory;
 		/** The chain of categories */
 		TArray<FString> CategoryChain;
-
-		/** Scoring data */
-		TArray<FString> SearchKeywordsArray;
-		TArray<FString> MenuDescriptionArray;
-		TArray<FString> SearchTitleArray;
-		TArray<FString> SearchCategoryArray;
-
-		/** The text that can be used to find this entry */
-		FString SearchText;
 	};
 private:
 
diff --git a/Engine/Source/Runtime/Engine/Private/EdGraph/EdGraphSchema.cpp b/Engine/Source/Runtime/Engine/Private/EdGraph/EdGraphSchema.cpp
--- a/Engine/Source/Runtime/Engine/Private/EdGraph/EdGraphSchema.cpp
+++ b/Engine/Source/Runtime/Engine/Private/EdGraph/EdGraphSchema.cpp
@@ -15,6 +15,74 @@
 
 #define LOCTEXT_NAMESPACE "EdGraph"
 
+void FEdGraphSchemaAction::UpdateCategory(const FText& NewCategory)
+{
+	Category = NewCategory;
+
+	TArray<FString> Scratch;
+	Category.ToString().ParseIntoArray(FullSearchCategoryArray, TEXT(" "), true);
+	Category.BuildSourceString().ParseIntoArray(Scratch, TEXT(" "), true);
+	FullSearchCategoryArray.Append(Scratch);
+
+	SearchText.Empty();
+	for (const auto& Entry : FullSearchTitlesArray)
+	{
+		SearchText += Entry;
+	}
+	SearchText.Append(LINE_TERMINATOR);
+	for (const auto& Entry : FullSearchKeywordsArray)
+	{
+		SearchText += Entry;
+	}
+	SearchText.Append(LINE_TERMINATOR);
+	for (const auto& Entry : FullSearchCategoryArray)
+	{
+		SearchText += Entry;
+	}
+}
+
+void FEdGraphSchemaAction::UpdateSearchData(const FText& NewMenuDescription, const FString& NewToolTipDescription, const FText& NewCategory, const FText& NewKeywords)
+{
+	MenuDescription = NewMenuDescription;
+	TooltipDescription = NewToolTipDescription;
+	Category = NewCategory;
+	Keywords = NewKeywords;
+
+	MenuDescription.ToString().ParseIntoArray(MenuDescriptionArray, TEXT(" "), true);
+
+	FullSearchTitlesArray = MenuDescriptionArray;
+	TArray<FString> Scratch;
+	MenuDescription.BuildSourceString().ParseIntoArray(Scratch, TEXT(" "), true);
+	FullSearchTitlesArray.Append(Scratch);
+
+	Keywords.ToString().ParseIntoArray(FullSearchKeywordsArray, TEXT(" "), true);
+	Keywords.BuildSourceString().ParseIntoArray(Scratch, TEXT(" "), true);
+	FullSearchKeywordsArray.Append(Scratch);
+
+	Category.ToString().ParseIntoArray(FullSearchCategoryArray, TEXT(" "), true);
+	Category.BuildSourceString().ParseIntoArray(Scratch, TEXT(" "), true);
+	FullSearchCategoryArray.Append(Scratch);
+
+	// Glob search text together, we use the SearchText string for basic filtering:
+	for (auto& Entry : FullSearchTitlesArray)
+	{
+		Entry.ToLowerInline();
+		SearchText += Entry;
+	}
+	SearchText.Append(LINE_TERMINATOR);
+	for (auto& Entry : FullSearchKeywordsArray)
+	{
+		Entry.ToLowerInline();
+		SearchText += Entry;
+	}
+	SearchText.Append(LINE_TERMINATOR);
+	for (auto& Entry : FullSearchCategoryArray)
+	{
+		Entry.ToLowerInline();
+		SearchText += Entry;
+	}
+}
+
 /////////////////////////////////////////////////////
 // FGraphActionListBuilderBase
 
@@ -56,7 +124,6 @@ FGraphActionListBuilderBase::ActionGroup::ActionGroup( TSharedPtr<FEdGraphSchema
 {
 	Actions.Add( InAction );
 	InitCategoryChain();
-	InitScoringData();
 }
 
 FGraphActionListBuilderBase::ActionGroup::ActionGroup( const TArray< TSharedPtr<FEdGraphSchemaAction> >& InActions, FString const& CategoryPrefix/* = TEXT("") */ )
@@ -64,7 +131,6 @@ FGraphActionListBuilderBase::ActionGroup::ActionGroup( const TArray< TSharedPtr<
 {
 	Actions = InActions;
 	InitCategoryChain();
-	InitScoringData();
 }
 
 FGraphActionListBuilderBase::ActionGroup::ActionGroup(FGraphActionListBuilderBase::ActionGroup && Other)
@@ -99,10 +165,13 @@ FGraphActionListBuilderBase::ActionGroup::~ActionGroup()
 {
 }
 
-void FGraphActionListBuilderBase::ActionGroup::GetCategoryChain(TArray<FString>& HierarchyOut) const
+const TArray<FString>& FGraphActionListBuilderBase::ActionGroup::GetCategoryChain() const
 {
 #if WITH_EDITOR
-	HierarchyOut = CategoryChain;
+	return CategoryChain;
+#else
+	static TArray<FString> Dummy;
+	return Dummy;
 #endif
 }
 
@@ -123,11 +192,6 @@ void FGraphActionListBuilderBase::ActionGroup::Move(FGraphActionListBuilderBase:
 	Actions = MoveTemp(Other.Actions);
 	RootCategory = MoveTemp(Other.RootCategory);
 	CategoryChain = MoveTemp(Other.CategoryChain);
-	SearchKeywordsArray = MoveTemp(Other.SearchKeywordsArray);
-	MenuDescriptionArray = MoveTemp(Other.MenuDescriptionArray);
-	SearchTitleArray = MoveTemp(Other.SearchTitleArray);
-	SearchCategoryArray = MoveTemp(Other.SearchCategoryArray);
-	SearchText = MoveTemp(Other.SearchText);
 }
 
 void FGraphActionListBuilderBase::ActionGroup::Copy(const ActionGroup& Other)
@@ -135,11 +199,6 @@ void FGraphActionListBuilderBase::ActionGroup::Copy(const ActionGroup& Other)
 	Actions = Other.Actions;
 	RootCategory = Other.RootCategory;
 	CategoryChain = Other.CategoryChain;
-	SearchKeywordsArray = Other.SearchKeywordsArray;
-	MenuDescriptionArray = Other.MenuDescriptionArray;
-	SearchTitleArray = Other.SearchTitleArray;
-	SearchCategoryArray = Other.SearchCategoryArray;
-	SearchText = Other.SearchText;
 }
 
 void FGraphActionListBuilderBase::ActionGroup::InitCategoryChain()
@@ -152,7 +211,7 @@ void FGraphActionListBuilderBase::ActionGroup::InitCategoryChain()
 	{
 		TArray<FString> SubCategoryChain;
 
-		FString SubCategory = FEditorCategoryUtils::GetCategoryDisplayString(Actions[0]->Category.ToString());
+		FString SubCategory = FEditorCategoryUtils::GetCategoryDisplayString(Actions[0]->GetCategory().ToString());
 		SubCategory.ParseIntoArray(SubCategoryChain, *CategoryDelim, true);
 
 		CategoryChain.Append(SubCategoryChain);
@@ -165,36 +224,6 @@ void FGraphActionListBuilderBase::ActionGroup::InitCategoryChain()
 #endif
 }
 
-void FGraphActionListBuilderBase::ActionGroup::InitScoringData()
-{
-	if (Actions.Num() > 0)
-	{
-		FEdGraphSchemaAction& FirstAction = *Actions[0];
-
-		// We keep these individual arrays so that they can be weighted differently by the scoring algorithm in SGraphActionMenu::GetActionFilteredWeight
-		FirstAction.GetSearchKeywords().ParseIntoArray(SearchKeywordsArray, TEXT(" "), true);
-		FirstAction.MenuDescription.ToString().ToLower().ParseIntoArray(MenuDescriptionArray, TEXT(" "), true);
-		FirstAction.GetSearchTitle().ParseIntoArray(SearchTitleArray, TEXT(" "), true);
-		FirstAction.GetSearchCategory().ParseIntoArray(SearchCategoryArray, TEXT(" "), true);
-
-		// Glob search text together, we use the SearchText string for basic filtering:
-		for (const auto& Entry : SearchTitleArray)
-		{
-			SearchText += Entry;
-		}
-		SearchText.Append(LINE_TERMINATOR);
-		for (const auto& Entry : SearchKeywordsArray)
-		{
-			SearchText += Entry;
-		}
-		SearchText.Append(LINE_TERMINATOR);
-		for (const auto& Entry : SearchCategoryArray)
-		{
-			SearchText += Entry;
-		}
-	}
-}
-
 /////////////////////////////////////////////////////
 // FCategorizedGraphActionListBuilder
 
