Added advanced searching syntax to Find-in-Blueprints, leveraging the system used for the CB.

TextFilterExpressionEvaluator has support for identifying functions.

Find-in-all-Blueprints uses a single separate thread again, it cannot utilize the task graph system for lack of heap space.

Find-in-Blueprints now supports searching using tags and filters for limiting searches to certain objects inside of Blueprints.

Interfaces, paths, and parent class name are now searchable when using tags.

Added a system for systematically loading and identifying unloaded Blueprint data.

diff --git a/Engine/Source/Editor/Kismet/Private/FiBSearchInstance.cpp b/Engine/Source/Editor/Kismet/Private/FiBSearchInstance.cpp
new file mode 100644
--- /dev/null
+++ b/Engine/Source/Editor/Kismet/Private/FiBSearchInstance.cpp
@@ -0,0 +1,709 @@
+// Copyright 1998-2015 Epic Games, Inc. All Rights Reserved.
+#include "BlueprintEditorPrivatePCH.h"
+#include "FiBSearchInstance.h"
+
+/** All operators when evaluating FiB searched expressions must return this token, it helps to manage
+ *	the results from functions as well as the specific components that were matched, and allows
+ *	for combining those results through complex operator combinations that may eliminate entire sections
+ *	of search results.
+ */
+class FFiBToken
+{
+public:
+	FFiBToken(bool bInValue)
+		: bValue(bInValue)
+	{
+	}
+
+	FFiBToken(bool bInValue, TMultiMap< const FImaginaryFiBData*, FComponentUniqueDisplay >& InMatchingSearchComponents)
+		: MatchingSearchComponents(MoveTemp(InMatchingSearchComponents))
+		, bValue(bInValue)
+	{
+	}
+
+	FFiBToken(bool bInValue, TArray< const FImaginaryFiBData* >& InMatchesSearchQuery)
+		: MatchesSearchQuery(InMatchesSearchQuery)
+		, bValue(bInValue)
+	{
+	}
+
+	FFiBToken(const FFiBToken& Other)
+		: MatchesSearchQuery(Other.MatchesSearchQuery)
+		, MatchingSearchComponents(Other.MatchingSearchComponents)
+		, bValue(Other.bValue)
+	{
+	}
+
+	FFiBToken(FFiBToken&& Other)
+		: MatchesSearchQuery(MoveTemp(Other.MatchesSearchQuery))
+		, MatchingSearchComponents(MoveTemp(Other.MatchingSearchComponents))
+		, bValue(Other.bValue)
+	{
+	}
+
+	FFiBToken& operator=(const FFiBToken& Other)
+	{
+		MatchesSearchQuery = Other.MatchesSearchQuery;
+		bValue = Other.bValue;
+		return *this;
+	}
+
+	FFiBToken& operator=(FFiBToken&& Other)
+	{
+		MatchesSearchQuery = MoveTemp(Other.MatchesSearchQuery);
+		bValue = Other.bValue;
+		return *this;
+	}
+
+	/** Combines another token into this one, merging all collected data */
+	void CombineToken(const FFiBToken& InToken)
+	{
+		MergeMatchesSearchQuery(InToken.MatchesSearchQuery);
+		MergeMatchingSearchComponents(InToken.MatchingSearchComponents);
+	}
+
+	/** Helper to only merge the MatchesSearchQuery data with this token */
+	void MergeMatchesSearchQuery(const TArray< const FImaginaryFiBData* >& InMatchesSearchQuery)
+	{
+		for (const FImaginaryFiBData* MatchesItem : InMatchesSearchQuery)
+		{
+			MatchesSearchQuery.AddUnique(MatchesItem);
+		}
+	}
+
+	/** Helper to only merge the MatchingSearchComponents data with this token */
+	void MergeMatchingSearchComponents(const TMultiMap< const FImaginaryFiBData*, FComponentUniqueDisplay >& InMatchingSearchComponents)
+	{
+		for (auto& MatchesItem : InMatchingSearchComponents)
+		{
+			MatchingSearchComponents.AddUnique(MatchesItem.Key, MatchesItem.Value);
+		}
+	}
+
+public:
+	/** A going list of all imaginary items that matched the search query at the time of this result token's creation */
+	TArray< const FImaginaryFiBData* > MatchesSearchQuery;
+
+	/** A mapping of items and their components that match the search query at the time of this result token's creation */
+	TMultiMap< const FImaginaryFiBData*, FComponentUniqueDisplay > MatchingSearchComponents;
+
+	/** Whether this result token should be considered TRUE or FALSE for purposes of further evaluation */
+	bool bValue;
+};
+
+DEFINE_EXPRESSION_NODE_TYPE(FFiBToken,	0x03378490, 0x42D14E26, 0x8E95AD2F, 0x74567868)
+
+/////////////////////////////////////
+// FFiBContextHelper
+
+/** Helper class to reroute testing of expressions against the context so 
+ *	that a mapping of the components in the context can be prepared and returned
+ */
+class FFiBContextHelper : public ITextFilterExpressionContext
+{
+public:
+	FFiBContextHelper()
+	{
+
+	}
+
+	FFiBContextHelper(TWeakPtr< FImaginaryFiBData > InContext)
+		: Context(InContext)
+	{
+	}
+
+	virtual bool TestBasicStringExpression(const FTextFilterString& InValue, const ETextFilterTextComparisonMode InTextComparisonMode) const override
+	{
+		if (Context.IsValid())
+		{
+			return Context.Pin()->TestBasicStringExpression(InValue, InTextComparisonMode, MatchingSearchComponents);
+		}
+
+		return false;
+	}
+
+	virtual bool TestComplexExpression(const FName& InKey, const FTextFilterString& InValue, const ETextFilterComparisonOperation InComparisonOperation, const ETextFilterTextComparisonMode InTextComparisonMode) const override
+	{
+		if (Context.IsValid())
+		{
+			return Context.Pin()->TestComplexExpression(InKey, InValue, InComparisonOperation, InTextComparisonMode, MatchingSearchComponents);
+		}
+
+		return false;
+	}
+
+	/** Context that is actually being tested */
+	TWeakPtr< FImaginaryFiBData > Context;
+
+	/** Modified in a const function callback, this is a going list of all search components that matched the expression */
+	mutable TMultiMap< const FImaginaryFiBData*, FComponentUniqueDisplay > MatchingSearchComponents;
+};
+
+////////////////////////
+// FFiBSearchInstance
+
+FSearchResult FFiBSearchInstance::StartSearchQuery(const FString& InSearchString, TSharedPtr<FImaginaryBlueprint> InImaginaryBlueprintRoot)
+{
+	PendingSearchables.Add(InImaginaryBlueprintRoot);
+	DoSearchQuery(InSearchString);
+
+	FSearchResult CachedSearchResult = FImaginaryFiBData::CreateSearchTree(nullptr, InImaginaryBlueprintRoot, MatchesSearchQuery, MatchingSearchComponents);
+	return MatchesSearchQuery.Num() > 0? CachedSearchResult : nullptr;
+}
+
+bool FFiBSearchInstance::DoSearchQuery(const FString& InSearchString, bool bInComplete/* = true*/)
+{
+	TSharedPtr< FFindInBlueprintExpressionEvaluator> ExpressionEvaluator(new FFindInBlueprintExpressionEvaluator(ETextFilterExpressionEvaluatorMode::Complex, this));
+
+	// Add all the required function bindings
+	ExpressionEvaluator->AddFunctionTokenCallback( TEXT("All"), FTokenFunctionHandler::CreateRaw(this, &FFiBSearchInstance::OnFilterFunction, ESearchQueryFilter::AllFilter));
+	ExpressionEvaluator->AddFunctionTokenCallback( TEXT("Blueprint"), FTokenFunctionHandler::CreateRaw(this, &FFiBSearchInstance::OnFilterFunction, ESearchQueryFilter::BlueprintFilter));
+	ExpressionEvaluator->AddFunctionTokenCallback( TEXT("Graphs"), FTokenFunctionHandler::CreateRaw(this, &FFiBSearchInstance::OnFilterFunction, ESearchQueryFilter::GraphsFilter));
+	ExpressionEvaluator->AddFunctionTokenCallback( TEXT("EventGraphs"), FTokenFunctionHandler::CreateRaw(this, &FFiBSearchInstance::OnFilterFunction, ESearchQueryFilter::UberGraphsFilter));
+	ExpressionEvaluator->AddFunctionTokenCallback( TEXT("Functions"), FTokenFunctionHandler::CreateRaw(this, &FFiBSearchInstance::OnFilterFunction, ESearchQueryFilter::FunctionsFilter));
+	ExpressionEvaluator->AddFunctionTokenCallback( TEXT("Macros"), FTokenFunctionHandler::CreateRaw(this, &FFiBSearchInstance::OnFilterFunction, ESearchQueryFilter::MacrosFilter));
+	ExpressionEvaluator->AddFunctionTokenCallback( TEXT("Properties"), FTokenFunctionHandler::CreateRaw(this, &FFiBSearchInstance::OnFilterFunction, ESearchQueryFilter::PropertiesFilter));
+	ExpressionEvaluator->AddFunctionTokenCallback( TEXT("Variables"), FTokenFunctionHandler::CreateRaw(this, &FFiBSearchInstance::OnFilterFunction, ESearchQueryFilter::PropertiesFilter));
+	ExpressionEvaluator->AddFunctionTokenCallback( TEXT("Components"), FTokenFunctionHandler::CreateRaw(this, &FFiBSearchInstance::OnFilterFunction, ESearchQueryFilter::ComponentsFilter));
+	ExpressionEvaluator->AddFunctionTokenCallback( TEXT("Nodes"), FTokenFunctionHandler::CreateRaw(this, &FFiBSearchInstance::OnFilterFunction, ESearchQueryFilter::NodesFilter));
+	ExpressionEvaluator->AddFunctionTokenCallback( TEXT("Pins"), FTokenFunctionHandler::CreateRaw(this, &FFiBSearchInstance::OnFilterFunction, ESearchQueryFilter::PinsFilter));
+	ExpressionEvaluator->SetFilterText(FText::FromString(InSearchString));
+
+	for (int SearchableIdx = 0; SearchableIdx < PendingSearchables.Num(); ++SearchableIdx)
+	{
+		CurrentSearchable = PendingSearchables[SearchableIdx];
+		TSharedPtr< FImaginaryFiBData > CurrentSearchablePinned = CurrentSearchable.Pin();
+		CurrentSearchablePinned->ParseAllChildData();
+		if (ExpressionEvaluator->TestTextFilter(*CurrentSearchablePinned.Get()))
+		{
+			MatchesSearchQuery.AddUnique(CurrentSearchablePinned.Get());
+		}
+
+		if (bInComplete || CurrentSearchablePinned->IsCategory())
+		{
+			// Any children that are not treated as a TagAndValue Category should be added for independent searching
+			TArray< TSharedPtr< FImaginaryFiBData > > Children = CurrentSearchablePinned->GetAllParsedChildData();
+			for (TSharedPtr< FImaginaryFiBData > Child : Children)
+			{
+				if (!Child->IsTagAndValueCategory())
+				{
+					PendingSearchables.Add(Child);
+				}
+			}
+		}
+	}
+	CurrentSearchable.Reset();
+
+	return MatchesSearchQuery.Num() > 0;
+}
+
+void FFiBSearchInstance::BuildFunctionTargets(TSharedPtr<FImaginaryFiBData> InRootData, ESearchQueryFilter InSearchQueryFilter, TArray< TWeakPtr< FImaginaryFiBData > >& OutTargetPendingSearchables)
+{
+	for (TSharedPtr<FImaginaryFiBData> ChildData : InRootData->GetAllParsedChildData())
+	{
+		if (ChildData->IsCompatibleWithFilter(InSearchQueryFilter))
+		{
+			OutTargetPendingSearchables.Add(ChildData);
+		}
+		else if (ChildData->CanCallFilter(InSearchQueryFilter))
+		{
+			ChildData->ParseAllChildData();
+			BuildFunctionTargets(ChildData, InSearchQueryFilter, OutTargetPendingSearchables);
+		}
+	}
+}
+
+bool FFiBSearchInstance::OnFilterFunction(const FTextFilterString& A, ESearchQueryFilter InSearchQueryFilter)
+{
+	if (CurrentSearchable.IsValid())
+	{
+		TSharedPtr< FFiBSearchInstance > SubSearchInstance(new FFiBSearchInstance);
+		bool bSearchSuccess = false;
+
+		TSharedPtr< FImaginaryFiBData > CurrentSearchablePinned = CurrentSearchable.Pin();
+		if (CurrentSearchablePinned->CanCallFilter(InSearchQueryFilter))
+		{
+			CurrentSearchablePinned->ParseAllChildData();
+			BuildFunctionTargets(CurrentSearchablePinned, InSearchQueryFilter, SubSearchInstance->PendingSearchables);
+		}
+		else if (InSearchQueryFilter == ESearchQueryFilter::BlueprintFilter && CurrentSearchablePinned->IsCompatibleWithFilter(ESearchQueryFilter::BlueprintFilter))
+		{
+			// We are filtering by Blueprint, since this is a Blueprint just add the CurrentSearchable to the PendingSearchables and do the sub-search on it
+			SubSearchInstance->PendingSearchables.Add(CurrentSearchable);
+		}
+
+		// Proceed to doing a sub-search
+		if (SubSearchInstance->PendingSearchables.Num() > 0)
+		{
+			bSearchSuccess = SubSearchInstance->DoSearchQuery(A.AsString(), InSearchQueryFilter == ESearchQueryFilter::AllFilter);
+			if (bSearchSuccess)
+			{
+				for (auto& MatchesItem : SubSearchInstance->MatchesSearchQuery)
+				{
+					LastFunctionResultMatchesSearchQuery.AddUnique(MatchesItem);
+				}
+
+				for (auto& MatchesItem : SubSearchInstance->MatchingSearchComponents)
+				{
+					LastFunctionMatchingSearchComponents.AddUnique(MatchesItem.Key, MatchesItem.Value);
+				}
+			}
+		}
+		return bSearchSuccess;
+	}
+
+	return false;
+}
+
+////////////////////////////////////////
+// FFindInBlueprintExpressionEvaluator
+
+bool FFindInBlueprintExpressionEvaluator::EvaluateCompiledExpression(const ExpressionParser::CompileResultType& InCompiledResult, const ITextFilterExpressionContext& InContext, FText* OutErrorText) const
+{
+	using namespace TextFilterExpressionParser;
+	using TextFilterExpressionParser::FTextToken;
+
+	if (InCompiledResult.IsValid())
+	{
+		auto EvalResult = ExpressionParser::Evaluate(InCompiledResult.GetValue(), JumpTable, &InContext);
+		if (EvalResult.IsValid())
+		{
+			if (const bool* BoolResult = EvalResult.GetValue().Cast<bool>())
+			{
+				return *BoolResult;
+			}
+			else if (const FTextToken* TextResult = EvalResult.GetValue().Cast<FTextToken>())
+			{
+				FFiBContextHelper ContextHelper(SearchInstance->CurrentSearchable);
+				bool bResult = TextResult->EvaluateAsBasicStringExpression(&ContextHelper);
+				if (bResult)
+				{
+					for (auto& MatchesItem : ContextHelper.MatchingSearchComponents)
+					{
+						SearchInstance->MatchingSearchComponents.AddUnique(MatchesItem.Key, MatchesItem.Value);
+					}
+				}
+				return bResult;
+			}
+			else if (const FFiBToken* FiBToken = EvalResult.GetValue().Cast<FFiBToken>())
+			{
+				if (FiBToken->bValue)
+				{
+					for (auto& MatchesItem : FiBToken->MatchesSearchQuery)
+					{
+						SearchInstance->MatchesSearchQuery.AddUnique(MatchesItem);
+					}
+
+					for (auto& MatchesItem : FiBToken->MatchingSearchComponents)
+					{
+						SearchInstance->MatchingSearchComponents.AddUnique(MatchesItem.Key, MatchesItem.Value);
+					}
+				}
+				return FiBToken->bValue;
+			}
+		}
+		else if (OutErrorText)
+		{
+			*OutErrorText = EvalResult.GetError().Text;
+		}
+	}
+
+	return false;
+}
+
+void FFindInBlueprintExpressionEvaluator::MapBasicJumps()
+{
+	using namespace TextFilterExpressionParser;
+	using TextFilterExpressionParser::FTextToken;
+
+	JumpTable.MapBinary<FLessOrEqual>(
+		[this](const FTextToken& A, const FTextToken& B, const ITextFilterExpressionContext* InContext)
+	{ 
+		FFiBContextHelper ContextHelper(SearchInstance->CurrentSearchable);
+		bool bResult = B.EvaluateAsComplexExpression(&ContextHelper, A.GetString(), ETextFilterComparisonOperation::LessOrEqual);
+		return FFiBToken(bResult, ContextHelper.MatchingSearchComponents);
+	});
+
+	JumpTable.MapBinary<FLess>(
+		[this](const FTextToken& A, const FTextToken& B, const ITextFilterExpressionContext* InContext)				
+	{ 
+		FFiBContextHelper ContextHelper(SearchInstance->CurrentSearchable);
+		bool bResult = B.EvaluateAsComplexExpression(&ContextHelper, A.GetString(), ETextFilterComparisonOperation::Less); 
+		return FFiBToken(bResult, ContextHelper.MatchingSearchComponents);
+	});
+
+	JumpTable.MapBinary<FGreaterOrEqual>(
+		[this](const FTextToken& A, const FTextToken& B, const ITextFilterExpressionContext* InContext)	
+	{ 
+		FFiBContextHelper ContextHelper(SearchInstance->CurrentSearchable);
+		bool bResult = B.EvaluateAsComplexExpression(&ContextHelper, A.GetString(), ETextFilterComparisonOperation::GreaterOrEqual); 
+		return FFiBToken(bResult, ContextHelper.MatchingSearchComponents);
+	});
+
+	JumpTable.MapBinary<FGreater>(
+		[this](const FTextToken& A, const FTextToken& B, const ITextFilterExpressionContext* InContext)			
+	{ 
+		FFiBContextHelper ContextHelper(SearchInstance->CurrentSearchable);
+		bool bResult = B.EvaluateAsComplexExpression(&ContextHelper, A.GetString(), ETextFilterComparisonOperation::Greater); 
+		return FFiBToken(bResult, ContextHelper.MatchingSearchComponents);
+	});
+
+	JumpTable.MapBinary<FNotEqual>(
+		[this](const FTextToken& A, const FTextToken& B, const ITextFilterExpressionContext* InContext)			
+	{ 
+		FFiBContextHelper ContextHelper(SearchInstance->CurrentSearchable);
+		bool bResult = B.EvaluateAsComplexExpression(&ContextHelper, A.GetString(), ETextFilterComparisonOperation::NotEqual); 
+		return FFiBToken(bResult, ContextHelper.MatchingSearchComponents);
+	});
+
+	JumpTable.MapBinary<FEqual>(
+		[this](const FTextToken& A, const FTextToken& B, const ITextFilterExpressionContext* InContext)				
+	{ 
+		FFiBContextHelper ContextHelper(SearchInstance->CurrentSearchable);
+		bool bResult = B.EvaluateAsComplexExpression(&ContextHelper, A.GetString(), ETextFilterComparisonOperation::Equal); 
+		return FFiBToken(bResult, ContextHelper.MatchingSearchComponents);
+	});
+
+	JumpTable.MapPreUnary<FNot>(
+		[this](const FTextToken& V, const ITextFilterExpressionContext* InContext)									
+	{ 
+		FFiBContextHelper ContextHelper(SearchInstance->CurrentSearchable);
+		bool bResult = !V.EvaluateAsBasicStringExpression(&ContextHelper); 
+		return FFiBToken(bResult, ContextHelper.MatchingSearchComponents);
+	});
+
+	JumpTable.MapPreUnary<FNot>(
+		[](bool V, const ITextFilterExpressionContext* InContext)
+	{
+		return !V;
+	});
+}
+
+void FFindInBlueprintExpressionEvaluator::MapOrBinaryJumps()
+{
+	using namespace TextFilterExpressionParser;
+
+	// Core Updated
+	JumpTable.MapBinary<FOr>(
+		[this](const TextFilterExpressionParser::FTextToken& A, const TextFilterExpressionParser::FTextToken& B, const ITextFilterExpressionContext* InContext)				
+	{ 
+		FFiBContextHelper ContextHelperA(SearchInstance->CurrentSearchable);
+		bool bAResult = A.EvaluateAsBasicStringExpression(&ContextHelperA);
+
+		FFiBContextHelper ContextHelperB(SearchInstance->CurrentSearchable);
+		bool bBResult = B.EvaluateAsBasicStringExpression(&ContextHelperB); 
+		bool bResult = bAResult || bBResult;
+
+		FFiBToken ResultToken(bResult);
+		if (bResult)
+		{
+			ResultToken.MatchesSearchQuery.AddUnique(static_cast<const FImaginaryFiBData*>(InContext));
+
+			if (bAResult)
+			{
+				ResultToken.MergeMatchingSearchComponents(ContextHelperA.MatchingSearchComponents);
+			}
+
+			if (bBResult)
+			{
+				ResultToken.MergeMatchingSearchComponents(ContextHelperB.MatchingSearchComponents);
+			}
+		}
+		return ResultToken;
+	});
+
+	JumpTable.MapBinary<FOr>(
+		[this](const TextFilterExpressionParser::FTextToken& A, bool B, const ITextFilterExpressionContext* InContext)							
+	{ 
+		FFiBContextHelper ContextHelper(SearchInstance->CurrentSearchable);
+		bool bAResult = A.EvaluateAsBasicStringExpression(&ContextHelper);
+		bool bResult = bAResult || B;
+
+		FFiBToken ResultToken(bResult);
+		if (bResult)
+		{
+			ResultToken.MatchesSearchQuery.AddUnique(static_cast<const FImaginaryFiBData*>(InContext));
+
+			if (bAResult)
+			{
+				ResultToken.MergeMatchingSearchComponents(ContextHelper.MatchingSearchComponents);
+			}
+		}
+		return ResultToken;
+	});
+
+	JumpTable.MapBinary<FOr>(
+		[this](bool A, const TextFilterExpressionParser::FTextToken& B, const ITextFilterExpressionContext* InContext)							
+	{ 
+		FFiBContextHelper ContextHelper(SearchInstance->CurrentSearchable);
+		bool bBResult = B.EvaluateAsBasicStringExpression(&ContextHelper);
+		bool bResult = A || bBResult;
+
+		FFiBToken ResultToken(bResult);
+		if (bResult)
+		{
+			ResultToken.MatchesSearchQuery.AddUnique(static_cast<const FImaginaryFiBData*>(InContext));
+
+			if (bBResult)
+			{
+				ResultToken.MergeMatchingSearchComponents(ContextHelper.MatchingSearchComponents);
+			}
+		}
+		return ResultToken;
+	});
+
+	JumpTable.MapBinary<FOr>(
+		[this](bool A, bool B, const ITextFilterExpressionContext* InContext)											
+	{ 
+		return A || B; 
+	});
+
+	// FiB Specific
+	JumpTable.MapBinary<FOr>(
+		[](const FFiBToken& A, const FFiBToken& B, const ITextFilterExpressionContext* InContext)
+	{
+		bool bResult = A.bValue || B.bValue;
+		FFiBToken ResultToken(bResult);
+
+		if (A.bValue)
+		{
+			ResultToken.CombineToken(A);
+		}
+
+		if (B.bValue)
+		{
+			ResultToken.CombineToken(B);
+		}
+		return ResultToken;
+	});
+	JumpTable.MapBinary<FOr>(
+		[](const FFiBToken& A, bool B, const ITextFilterExpressionContext* InContext)								
+	{ 
+		bool bResult = A.bValue || B;
+
+		FFiBToken ResultToken(bResult);
+		if (A.bValue)
+		{
+			ResultToken.CombineToken(A);
+		}
+		if (B)
+		{
+			ResultToken.MatchesSearchQuery.AddUnique(static_cast<const FImaginaryFiBData*>(InContext));
+		}
+		return ResultToken;
+	});
+	JumpTable.MapBinary<FOr>(
+		[](bool A, const FFiBToken& B, const ITextFilterExpressionContext* InContext)								
+	{
+		bool bResult = A || B.bValue;
+
+		FFiBToken ResultToken(bResult);
+
+		if (A)
+		{
+			ResultToken.MatchesSearchQuery.AddUnique(static_cast<const FImaginaryFiBData*>(InContext));
+		}
+		if (B.bValue)
+		{
+			ResultToken.CombineToken(B);
+		}
+		return ResultToken;
+	});
+	JumpTable.MapBinary<FOr>(
+		[this](const FFiBToken& A, TextFilterExpressionParser::FTextToken B, const ITextFilterExpressionContext* InContext)								
+	{
+		FFiBContextHelper ContextHelper(SearchInstance->CurrentSearchable);
+		bool bBResult = B.EvaluateAsBasicStringExpression(&ContextHelper);
+		bool bResult = A.bValue || bBResult;
+		FFiBToken ResultToken(bResult);
+		if (A.bValue)
+		{
+			ResultToken.CombineToken(A);
+		}
+		if (bBResult)
+		{
+			ResultToken.MatchesSearchQuery.AddUnique(static_cast<const FImaginaryFiBData*>(InContext));
+			ResultToken.MergeMatchingSearchComponents(ContextHelper.MatchingSearchComponents);
+		}
+		return ResultToken;
+	});
+	JumpTable.MapBinary<FOr>(
+		[this](TextFilterExpressionParser::FTextToken A, const FFiBToken& B, const ITextFilterExpressionContext* InContext)
+	{
+		FFiBContextHelper ContextHelper(SearchInstance->CurrentSearchable);
+		bool bAResult = A.EvaluateAsBasicStringExpression(&ContextHelper);
+		bool bResult = bAResult || B.bValue;
+		FFiBToken ResultToken(bResult);
+
+		if (bAResult)
+		{
+			ResultToken.MatchesSearchQuery.AddUnique(static_cast<const FImaginaryFiBData*>(InContext));
+			ResultToken.MergeMatchingSearchComponents(ContextHelper.MatchingSearchComponents);
+		}
+		if (B.bValue)
+		{
+			ResultToken.CombineToken(B);
+		}
+		return ResultToken;
+	});
+}
+
+void FFindInBlueprintExpressionEvaluator::MapAndBinaryJumps()
+{
+	using namespace TextFilterExpressionParser;
+
+	// Core Updated
+	JumpTable.MapBinary<FAnd>(
+		[this](const TextFilterExpressionParser::FTextToken& A, const TextFilterExpressionParser::FTextToken& B, const ITextFilterExpressionContext* InContext)				
+	{ 
+		FFiBContextHelper ContextHelper(SearchInstance->CurrentSearchable);
+		bool bResult = A.EvaluateAsBasicStringExpression(&ContextHelper) && B.EvaluateAsBasicStringExpression(&ContextHelper); 
+
+		FFiBToken ResultToken(bResult);
+		if (bResult)
+		{
+			ResultToken.MatchesSearchQuery.AddUnique(static_cast<const FImaginaryFiBData*>(InContext));
+			ResultToken.MergeMatchingSearchComponents(ContextHelper.MatchingSearchComponents);
+		}
+		return ResultToken;
+	});
+
+	JumpTable.MapBinary<FAnd>(
+		[this](const TextFilterExpressionParser::FTextToken& A, bool B, const ITextFilterExpressionContext* InContext)							
+	{ 
+		FFiBContextHelper ContextHelper(SearchInstance->CurrentSearchable);
+		bool bResult = A.EvaluateAsBasicStringExpression(&ContextHelper) && B;
+
+		FFiBToken ResultToken(bResult);
+		if (bResult)
+		{
+			ResultToken.MatchesSearchQuery.AddUnique(static_cast<const FImaginaryFiBData*>(InContext));
+			ResultToken.MergeMatchingSearchComponents(ContextHelper.MatchingSearchComponents);
+		}
+		return ResultToken;
+	});
+
+	JumpTable.MapBinary<FAnd>(
+		[this](bool A, const TextFilterExpressionParser::FTextToken& B, const ITextFilterExpressionContext* InContext)							
+	{ 
+		FFiBContextHelper ContextHelper(SearchInstance->CurrentSearchable);
+		bool bResult = A && B.EvaluateAsBasicStringExpression(&ContextHelper);
+
+		FFiBToken ResultToken(bResult);
+		if (bResult)
+		{
+			ResultToken.MatchesSearchQuery.AddUnique(static_cast<const FImaginaryFiBData*>(InContext));
+			ResultToken.MergeMatchingSearchComponents(ContextHelper.MatchingSearchComponents);
+		}
+		return ResultToken;
+	});
+
+	JumpTable.MapBinary<FAnd>(
+		[](bool A, bool B, const ITextFilterExpressionContext* InContext)											
+	{ 
+		return A && B; 
+	});
+
+	// FiB Specific
+	JumpTable.MapBinary<FAnd>(
+		[](const FFiBToken& A, const FFiBToken& B, const ITextFilterExpressionContext* InContext)
+	{
+		bool bResult = A.bValue && B.bValue;
+		FFiBToken ResultToken(bResult);
+		if (bResult)
+		{
+			ResultToken.CombineToken(A);
+			ResultToken.CombineToken(B);
+		}
+		return ResultToken;
+	});
+	JumpTable.MapBinary<FAnd>(
+		[](const FFiBToken& A, bool B, const ITextFilterExpressionContext* InContext)							
+	{ 
+		bool bResult = A.bValue && B;
+		FFiBToken ResultToken(bResult);
+		if (bResult)
+		{
+			ResultToken.CombineToken(A);
+		}
+		if (B)
+		{
+			ResultToken.MatchesSearchQuery.AddUnique(static_cast<const FImaginaryFiBData*>(InContext));
+		}
+		return ResultToken;
+	});
+	JumpTable.MapBinary<FAnd>(
+		[](bool A, const FFiBToken& B, const ITextFilterExpressionContext* InContext)							
+	{
+		bool bResult = A && B.bValue;
+		FFiBToken ResultToken(bResult);
+		if (A)
+		{
+			ResultToken.MatchesSearchQuery.AddUnique(static_cast<const FImaginaryFiBData*>(InContext));
+		}
+		if (bResult)
+		{
+			ResultToken.CombineToken(B);
+		}
+		return ResultToken;
+	});
+	JumpTable.MapBinary<FAnd>(
+		[this](const FFiBToken& A, TextFilterExpressionParser::FTextToken B, const ITextFilterExpressionContext* InContext)							
+	{ 
+		FFiBContextHelper ContextHelper(SearchInstance->CurrentSearchable);
+		bool bBResult = B.EvaluateAsBasicStringExpression(&ContextHelper);
+		bool bResult = A.bValue && bBResult;
+		FFiBToken ResultToken(bResult);
+		if (bResult)
+		{
+			ResultToken.CombineToken(A);
+		}
+		if (bBResult)
+		{
+			ResultToken.MatchesSearchQuery.AddUnique(static_cast<const FImaginaryFiBData*>(InContext));
+			ResultToken.MergeMatchingSearchComponents(ContextHelper.MatchingSearchComponents);
+		}
+		return ResultToken;
+	});
+	JumpTable.MapBinary<FAnd>(
+		[this](TextFilterExpressionParser::FTextToken A, const FFiBToken& B, const ITextFilterExpressionContext* InContext)							
+	{ 
+		FFiBContextHelper ContextHelper(SearchInstance->CurrentSearchable);
+		bool bAResult = A.EvaluateAsBasicStringExpression(&ContextHelper);
+		bool bResult = bAResult && B.bValue;
+		FFiBToken ResultToken(bResult);
+		if (bAResult)
+		{
+			ResultToken.MatchesSearchQuery.AddUnique(static_cast<const FImaginaryFiBData*>(InContext));
+			ResultToken.MergeMatchingSearchComponents(ContextHelper.MatchingSearchComponents);
+		}
+		if (bResult)
+		{
+			ResultToken.CombineToken(B);
+		}
+		return ResultToken;
+	});
+}
+
+void FFindInBlueprintExpressionEvaluator::ConstructExpressionParser()
+{
+	SetupGrammar();
+	MapBasicJumps();
+	MapOrBinaryJumps();
+	MapAndBinaryJumps();
+
+	using namespace TextFilterExpressionParser;
+
+	JumpTable.MapBinary<TextFilterExpressionParser::FFunction>(
+		[this](const TextFilterExpressionParser::FTextToken& A, const TextFilterExpressionParser::FTextToken& B, const ITextFilterExpressionContext* InContext)	
+	{
+		bool bResult = false;
+		if (FTokenFunctionHandler* FunctionCallback = TokenFunctionHandlers.Find(A.GetString().AsString()))
+		{
+			bResult = FunctionCallback->Execute(B.GetString());
+		}
+		FFiBToken Result(bResult, SearchInstance->LastFunctionResultMatchesSearchQuery);
+		Result.MatchingSearchComponents = SearchInstance->LastFunctionMatchingSearchComponents;
+		return Result;
+	});
+}
\ No newline at end of file
diff --git a/Engine/Source/Editor/Kismet/Private/FiBSearchInstance.h b/Engine/Source/Editor/Kismet/Private/FiBSearchInstance.h
new file mode 100644
--- /dev/null
+++ b/Engine/Source/Editor/Kismet/Private/FiBSearchInstance.h
@@ -0,0 +1,98 @@
+// Copyright 1998-2015 Epic Games, Inc. All Rights Reserved.
+#pragma once
+
+#include "TextFilterExpressionEvaluator.h"
+#include "ImaginaryBlueprintData.h"
+
+class FFiBSearchInstance;
+
+/** Evaluates the expression the user submitted to be searched for */
+class FFindInBlueprintExpressionEvaluator : public FTextFilterExpressionEvaluator
+{
+public:
+	/** Construction and assignment */
+	FFindInBlueprintExpressionEvaluator(const ETextFilterExpressionEvaluatorMode InMode, FFiBSearchInstance* InSearchInstance)
+		: FTextFilterExpressionEvaluator(InMode)
+		, SearchInstance(InSearchInstance)
+	{
+		ConstructExpressionParser();
+	}
+
+protected:
+	/** FTextFilterExpressionEvaluator Interface */
+	virtual void ConstructExpressionParser() override;
+	virtual bool EvaluateCompiledExpression(const ExpressionParser::CompileResultType& InCompiledResult, const ITextFilterExpressionContext& InContext, FText* OutErrorText) const override;
+	/** End FTextFilterExpressionEvaluator Interface */
+
+	/** Helper function to make a mapping of all basic jump operations */
+	void MapBasicJumps();
+
+	/** Helper function to make a mapping of all "Or" binary jump operations */
+	void MapOrBinaryJumps();
+
+	/** Helper function to make a mapping of all "And" binary jump operations */
+	void MapAndBinaryJumps();
+
+protected:
+	/** Referenced SearchInstance that is powering this search */
+	FFiBSearchInstance* SearchInstance;
+};
+
+/** Used to manage searches through imaginary Blueprints */
+class FFiBSearchInstance
+{
+public:
+	/**
+	 * Starts a search query given a string and an imaginary Blueprint
+	 *
+	 * @param InSearchString				The string to search using
+	 * @param InImaginaryBlueprintRoot		The imaginary Blueprint to search through
+	 * @return								Search result shared pointer, can be used for display in the search results window
+	 */
+	FSearchResult StartSearchQuery(const FString& InSearchString, TSharedPtr<FImaginaryBlueprint> InImaginaryBlueprintRoot);
+	
+	/**
+	 * Runs a search query on any pending imaginary data
+	 *
+	 * @param InSearchString		The string to search using
+	 * @param bInComplete			TRUE if a complete search of all child items should also be done, this should be FALSE when you only want to compare the string against the immediate items and not their children
+	 * @return						TRUE if the search was successful.
+	 */
+	bool DoSearchQuery(const FString& InSearchString, bool bInComplete = true);
+
+	/**
+	 * Builds a list of imaginary items that can be targeted by a function
+	 *
+	 * @param InRootData						Item to find the child imaginary data that can have be sub-searched for the function
+	 * @param InSearchQueryFilter				A filter to decide if the imaginary item is compatible with the function
+	 * @param OutTargetPendingSearchables		List of pending imaginary items that will be sub-searched, this gets filled out by the function
+	 */
+	void BuildFunctionTargets(TSharedPtr<FImaginaryFiBData> InRootData, ESearchQueryFilter InSearchQueryFilter, TArray< TWeakPtr< FImaginaryFiBData > >& OutTargetPendingSearchables);
+
+	/**
+	 * Callback when a function is called by the evaluator
+	 *
+	 * @param InParameterString		The parameter string to sub-search using
+	 * @param InSearchQueryFilter	Effectively this is the function being called, the filter is used to determine if the current item can have the function called on it and whether it's sub items can have it called on them
+	 * @return						TRUE if the function was successful at finding valid results
+	 */
+	bool OnFilterFunction(const FTextFilterString& A, ESearchQueryFilter InSearchQueryFilter);
+
+	/** Current item being searched in the Imaginary Blueprint */
+	TWeakPtr< FImaginaryFiBData > CurrentSearchable;
+
+	/** A going list of all imaginary items that match the search query */
+	TArray< const FImaginaryFiBData* > MatchesSearchQuery;
+
+	/** A mapping of items and their components that match the search query */
+	TMultiMap< const FImaginaryFiBData*, FComponentUniqueDisplay > MatchingSearchComponents;
+
+	/** A list of imaginary items that still need to be searched */
+	TArray< TWeakPtr< FImaginaryFiBData > > PendingSearchables;
+
+	/** When a function returns on an item, this is the list of items that matched the sub-search query */
+	TArray< const FImaginaryFiBData* > LastFunctionResultMatchesSearchQuery;
+
+	/** When a function returns on an item, this is the mapping of imaginary items to components that matched the sub-search query */
+	TMultiMap< const FImaginaryFiBData*, FComponentUniqueDisplay > LastFunctionMatchingSearchComponents;
+};
\ No newline at end of file
diff --git a/Engine/Source/Editor/Kismet/Private/FindInBlueprintManager.cpp b/Engine/Source/Editor/Kismet/Private/FindInBlueprintManager.cpp
--- a/Engine/Source/Editor/Kismet/Private/FindInBlueprintManager.cpp
+++ b/Engine/Source/Editor/Kismet/Private/FindInBlueprintManager.cpp
@@ -39,6 +39,9 @@ const FText FFindInBlueprintSearchTags::FiB_Tooltip = LOCTEXT("Tooltip", "Toolti
 const FText FFindInBlueprintSearchTags::FiB_DefaultValue = LOCTEXT("DefaultValue", "DefaultValue");
 const FText FFindInBlueprintSearchTags::FiB_Description = LOCTEXT("Description", "Description");
 const FText FFindInBlueprintSearchTags::FiB_Comment = LOCTEXT("Comment", "Comment");
+const FText FFindInBlueprintSearchTags::FiB_Path = LOCTEXT("Path", "Path");
+const FText FFindInBlueprintSearchTags::FiB_ParentClass = LOCTEXT("ParentClass", "ParentClass");
+const FText FFindInBlueprintSearchTags::FiB_Interfaces = LOCTEXT("Interfaces", "Interfaces");
 
 const FText FFindInBlueprintSearchTags::FiB_Pins = LOCTEXT("Pins", "Pins");
 const FText FFindInBlueprintSearchTags::FiB_PinCategory = LOCTEXT("PinCategory", "PinCategory");
@@ -991,6 +994,32 @@ void FFindInBlueprintSearchManager::OnAssetAdded(const FAssetData& InAssetData)
 				FSearchData NewSearchData;
 
 				NewSearchData.BlueprintPath = InAssetData.ObjectPath;
+
+				if (const FString* ParentClass = InAssetData.TagsAndValues.Find(TEXT("ParentClass")))
+				{
+					NewSearchData.ParentClass = *ParentClass;
+				}
+
+				const FString* ImplementedInterfaces = InAssetData.TagsAndValues.Find("ImplementedInterfaces");
+				if(ImplementedInterfaces)
+				{
+					FString FullInterface;
+					FString RemainingString;
+					FString InterfaceName;
+					FString CurrentString = *ImplementedInterfaces;
+					while(CurrentString.Split(TEXT(","), &FullInterface, &RemainingString))
+					{
+						if(FullInterface.Split(TEXT("."), &CurrentString, &InterfaceName, ESearchCase::CaseSensitive, ESearchDir::FromEnd))
+						{
+							if(!CurrentString.StartsWith(TEXT("Graphs=(")))
+							{
+								NewSearchData.Interfaces.Add(InterfaceName);
+							}
+						}
+						CurrentString = RemainingString;
+					}
+				}
+
 				NewSearchData.bMarkedForDeletion = false;
 
 				// Since the asset was not loaded, pull out the searchable data stored in the asset
@@ -1186,7 +1215,6 @@ void FFindInBlueprintSearchManager::AddOrUpdateBlueprintSearchMetadata(UBlueprin
 		FSearchData SearchData;
 		SearchData.Blueprint = InBlueprint;
 		SearchData.BlueprintPath = BlueprintPath;
-
 		Index = AddSearchDataToDatabase(MoveTemp(SearchData));
 	}
 	else
@@ -1196,7 +1224,10 @@ void FFindInBlueprintSearchManager::AddOrUpdateBlueprintSearchMetadata(UBlueprin
 
 	// Build the search data
 	SearchArray[Index].BlueprintPath = BlueprintPath;
-
+	if (UProperty* ParentClassProp = InBlueprint->GetClass()->FindPropertyByName(GET_MEMBER_NAME_CHECKED(UBlueprint, ParentClass)))
+	{
+		ParentClassProp->ExportTextItem(SearchArray[Index].ParentClass, ParentClassProp->ContainerPtrToValuePtr<uint8>(InBlueprint), nullptr, InBlueprint, 0);
+	}
 	// Cannot successfully gather most searchable data if there is no SkeletonGeneratedClass, so don't try, leave it as whatever it was last set to
 	if(InBlueprint->SkeletonGeneratedClass != nullptr)
 	{
diff --git a/Engine/Source/Editor/Kismet/Private/FindInBlueprints.cpp b/Engine/Source/Editor/Kismet/Private/FindInBlueprints.cpp
--- a/Engine/Source/Editor/Kismet/Private/FindInBlueprints.cpp
+++ b/Engine/Source/Editor/Kismet/Private/FindInBlueprints.cpp
@@ -11,215 +11,85 @@
 #include "IDocumentation.h"
 #include "SSearchBox.h"
 #include "GenericCommands.h"
+#include "ImaginaryBlueprintData.h"
+#include "FiBSearchInstance.h"
 
 #define LOCTEXT_NAMESPACE "FindInBlueprints"
 
-/** Some utility functions to help with Find-in-Blueprint functionality */
-namespace FindInBlueprintsHelpers
+
+FText FindInBlueprintsHelpers::AsFText(TSharedPtr< FJsonValue > InJsonValue, const TMap<int32, FText>& InLookupTable)
 {
-	/** Looks up a JsonValue's FText from the passed lookup table */
-	FText AsFText(TSharedPtr< FJsonValue > InJsonValue, const TMap<int32, FText>& InLookupTable)
+	if (const FText* LookupText = InLookupTable.Find(FCString::Atoi(*InJsonValue->AsString())))
 	{
-		if(const FText* LookupText = InLookupTable.Find(FCString::Atoi(*InJsonValue->AsString())))
-		{
-			return *LookupText;
-		}
-		// Let's never get here.
-		return LOCTEXT("FiBSerializationError", "There was an error in serialization!");
+		return *LookupText;
 	}
+	// Let's never get here.
+	return LOCTEXT("FiBSerializationError", "There was an error in serialization!");
+}
 
-	/**
-	 * Retrieves the pin type as a string value
-	 *
-	 * @param InPinType		The pin type to look at
-	 *
-	 * @return				The pin type as a string in format [category]'[sub-category object]'
-	 */
-	FString GetPinTypeAsString(const FEdGraphPinType& InPinType)
+FText FindInBlueprintsHelpers::AsFText(int32 InValue, const TMap<int32, FText>& InLookupTable)
+{
+	if (const FText* LookupText = InLookupTable.Find(InValue))
 	{
-		FString Result = InPinType.PinCategory;
-		if(UObject* SubCategoryObject = InPinType.PinSubCategoryObject.Get()) 
-		{
-			Result += FString(" '") + SubCategoryObject->GetName() + "'";
-		}
-
-		return Result;
+		return *LookupText;
 	}
+	// Let's never get here.
+	return LOCTEXT("FiBSerializationError", "There was an error in serialization!");
+}
 
-	/**
-	 * Determines if a string matches the search tokens
-	 *
-	 * @param InTokens					Search tokens to compare with
-	 * @param InComparisonString		The string to check for matching tokens in
-	 *
-	 * @return							TRUE if the string matches any of the search tokens
-	 */
-	bool StringMatchesSearchTokens(const TArray<FString>& InTokens, const FString& InComparisonString)
-	{
-		const FString SafeSearchString = InComparisonString.Replace( TEXT( " " ), TEXT( "" ) );
-
-		bool bFoundAllTokens = true;
-		//search the entry for each token, it must have all of them to pass
-		for(auto TokIT(InTokens.CreateConstIterator());TokIT;++TokIT)
-		{
-			const FString& Token = *TokIT;
-			if(!SafeSearchString.Contains(Token))
-			{
-				bFoundAllTokens = false;
-				break;
-			}
-		}
-		return bFoundAllTokens;
-	}
+bool FindInBlueprintsHelpers::IsTextEqualToString(FText InText, FString InString)
+{
+	return InString == InText.ToString() || InString == InText.BuildSourceString();
+}
 
-	/**
-	 * Determines if an FText matches the search tokens, comparing both the (localized) display string and the (unlocalized) source string
-	 *
-	 * @param InTokens					Search tokens to compare with
-	 * @param InComparisonText			The FText to check for matching tokens in
-	 *
-	 * @return							TRUE if the FText matches any of the search tokens
-	 */
-	bool TextMatchesSearchTokens(const TArray<FString>& InTokens, const FText& InComparisonText)
+FString FindInBlueprintsHelpers::GetPinTypeAsString(const FEdGraphPinType& InPinType)
+{
+	FString Result = InPinType.PinCategory;
+	if(UObject* SubCategoryObject = InPinType.PinSubCategoryObject.Get()) 
 	{
-		return StringMatchesSearchTokens(InTokens, InComparisonText.ToString()) || StringMatchesSearchTokens(InTokens, InComparisonText.BuildSourceString());
+		Result += FString(" '") + SubCategoryObject->GetName() + "'";
 	}
-
-	/**
-	 * Helper function to extract all graph data from a Json object and into a search result. There is little type safety, it assumes all Json objects found in the InCategoryTitle are graphs
-	 *
-	 * @param InJsonObject					The JsonObject to extract the graphs from
-	 * @param InCategoryTitle				The category the graphs are titled under.
-	 * @param InTokens						Filtering tokens to use
-	 * @param InGraphType					There is nothing in the Json script to identify graphs, they are all grouped by category, so pass in what type they should be considered
-	 * @param OutBlueprintCategory			The search category these graphs will be children to
-	 */
-	void ExtractGraph( const TMap<int32, FText>& InLookupTable, TSharedPtr< FJsonObject > InJsonObject, FString InCategoryTitle, const TArray<FString>& InTokens, EGraphType InGraphType, FSearchResult& OutBlueprintCategory )
+	else
 	{
-		if(InJsonObject->HasField(InCategoryTitle))
-		{
-			// Pulls out all functions for this Blueprint
-			TArray<TSharedPtr< FJsonValue > > FunctionList = InJsonObject->GetArrayField(InCategoryTitle);
-			for( TSharedPtr< FJsonValue > FunctionValue : FunctionList )
-			{
-				FSearchResult Result(new FFindInBlueprintsGraph(FText::GetEmpty(), OutBlueprintCategory, InGraphType) );
-
-				if( Result->ExtractContent(InLookupTable, FunctionValue->AsObject(), InTokens) )
-				{
-					OutBlueprintCategory->Children.Add(Result);
-				}
-			}
-		}
+		Result += FString(" '") + InPinType.PinSubCategory + "'";
 	}
 
-	/**
-	 * Utility function to begin a search
-	 *
-	 * @param InBuffer					The search buffer, a Json string to be parsed and converted into search result
-	 * @param InTokens					Search tokens to filter the results by
-	 * @param OutBlueprintCategory		The root category for the search results
-	 */
-	void Extract( const FString& InBuffer, const TArray<FString>& InTokens, FSearchResult& OutBlueprintCategory )
-	{
-		// Skip empty buffers
-		if(InBuffer.Len() == 0)
-		{
-			return;
-		}
-
-		TSharedPtr< FJsonObject > JsonObject = NULL;
-		TMap<int32, FText> LookupTable;
-		JsonObject = FFindInBlueprintSearchManager::ConvertJsonStringToObject(InBuffer, LookupTable);
-
-		if(JsonObject.IsValid())
-		{
-			if(JsonObject->HasField(FFindInBlueprintSearchTags::FiB_Properties.ToString()))
-			{
-				// Pulls out all properties (variables) for this Blueprint
-				TArray<TSharedPtr< FJsonValue > > PropertyList = JsonObject->GetArrayField(FFindInBlueprintSearchTags::FiB_Properties.ToString());
-				for( TSharedPtr< FJsonValue > PropertyValue : PropertyList )
-				{
-					FSearchResult Result(new FFindInBlueprintsProperty(FText::GetEmpty(), OutBlueprintCategory) );
-
-					if( Result->ExtractContent(LookupTable, PropertyValue->AsObject(), InTokens) )
-					{
-						OutBlueprintCategory->Children.Add(Result);
-					}
-				}
-			}
-
-			ExtractGraph(LookupTable, JsonObject, FFindInBlueprintSearchTags::FiB_Functions.ToString(), InTokens, EGraphType::GT_Function, OutBlueprintCategory);
-			ExtractGraph(LookupTable, JsonObject, FFindInBlueprintSearchTags::FiB_Macros.ToString(), InTokens, EGraphType::GT_Macro, OutBlueprintCategory);
-			ExtractGraph(LookupTable, JsonObject, FFindInBlueprintSearchTags::FiB_UberGraphs.ToString(), InTokens, EGraphType::GT_Ubergraph, OutBlueprintCategory);
-			ExtractGraph(LookupTable, JsonObject, FFindInBlueprintSearchTags::FiB_SubGraphs.ToString(), InTokens, EGraphType::GT_Ubergraph, OutBlueprintCategory);
-
-			if(JsonObject->HasField(FFindInBlueprintSearchTags::FiB_Components.ToString()))
-			{
-				FSearchResult ComponentsResult(new FFindInBlueprintsResult(FFindInBlueprintSearchTags::FiB_Components, OutBlueprintCategory) );
-
-				// Pulls out all properties (variables) for this Blueprint
-				TArray<TSharedPtr< FJsonValue > > ComponentList = JsonObject->GetArrayField(FFindInBlueprintSearchTags::FiB_Components.ToString());
-				for( TSharedPtr< FJsonValue > ComponentValue : ComponentList )
-				{
-					FSearchResult Result(new FFindInBlueprintsProperty(FText::GetEmpty(), ComponentsResult) );
+	return Result;
+}
 
-					if( Result->ExtractContent(LookupTable, ComponentValue->AsObject(), InTokens) )
-					{
-						ComponentsResult->Children.Add(Result);
-					}
-				}
+bool FindInBlueprintsHelpers::ParsePinType(FText InKey, FText InValue, FEdGraphPinType& InOutPinType)
+{
+	bool bParsed = true;
 
-				// Only add the component list category if there were components
-				if(ComponentsResult->Children.Num())
-				{
-					OutBlueprintCategory->Children.Add(ComponentsResult);
-				}
-			}
-		}
+	if(InKey.CompareTo(FFindInBlueprintSearchTags::FiB_PinCategory) == 0)
+	{
+		InOutPinType.PinCategory = InValue.ToString();
 	}
-
-	/**
-	 * Parses a pin type from passed in key names and values
-	 *
-	 * @param InKey					The key name for what the data should be translated as
-	 * @param InValue				Value to be be translated
-	 * @param InOutPinType			Modifies the PinType based on the passed parameters, building it up over multiple calls
-	 * @param InOutPinSubCategory	When the pin's sub-category is found, this value will be changed. This value will never be wiped during other translations. It is not thread safe to determine the sub-object.
-	 * @return						TRUE when the parsing is successful
-	 */
-	bool ParsePinType(FText InKey, FText InValue, FEdGraphPinType& InOutPinType, FString& OutPinSubCategory)
+	else if(InKey.CompareTo(FFindInBlueprintSearchTags::FiB_PinSubCategory) == 0)
 	{
-		bool bParsed = true;
-
-		if(InKey.CompareTo(FFindInBlueprintSearchTags::FiB_PinCategory) == 0)
-		{
-			InOutPinType.PinCategory = InValue.ToString();
-		}
-		else if(InKey.CompareTo(FFindInBlueprintSearchTags::FiB_PinSubCategory) == 0)
-		{
-			InOutPinType.PinSubCategory = InValue.ToString();
-		}
-		else if(InKey.CompareTo(FFindInBlueprintSearchTags::FiB_ObjectClass) == 0)
-		{
-			OutPinSubCategory = InValue.ToString();
-		}
-		else if(InKey.CompareTo(FFindInBlueprintSearchTags::FiB_IsArray) == 0)
-		{
-			InOutPinType.bIsArray = InValue.ToString().ToBool();
-		}
-		else if(InKey.CompareTo(FFindInBlueprintSearchTags::FiB_IsReference) == 0)
-		{
-			InOutPinType.bIsReference = InValue.ToString().ToBool();
-		}
-		else
-		{
-			bParsed = false;
-		}
-
-		return bParsed;
+		InOutPinType.PinSubCategory = InValue.ToString();
+	}
+	else if(InKey.CompareTo(FFindInBlueprintSearchTags::FiB_ObjectClass) == 0)
+	{
+		InOutPinType.PinSubCategory = InValue.ToString();
+	}
+	else if(InKey.CompareTo(FFindInBlueprintSearchTags::FiB_IsArray) == 0)
+	{
+		InOutPinType.bIsArray = InValue.ToString().ToBool();
 	}
+	else if(InKey.CompareTo(FFindInBlueprintSearchTags::FiB_IsReference) == 0)
+	{
+		InOutPinType.bIsReference = InValue.ToString().ToBool();
+	}
+	else
+	{
+		bParsed = false;
+	}
+
+	return bParsed;
 }
 
+
 //////////////////////////////////////////////////////////////////////////
 // FBlueprintSearchResult
 
@@ -253,96 +123,6 @@ FReply FFindInBlueprintsResult::OnClick()
 	return FReply::Handled();
 }
 
-bool FFindInBlueprintsResult::ParseSearchInfo(const TArray<FString> &InTokens, FText InKey, FText InValue, TSharedPtr< FFindInBlueprintsResult > InParent)
-{
-	// By default, nothing much gets processed, just check to see if it matches the search requirements
-	return FindInBlueprintsHelpers::TextMatchesSearchTokens(InTokens, InValue);
-}
-
-void FFindInBlueprintsResult::AddExtraSearchInfo(FText InKey, FText InValue, TSharedPtr< FFindInBlueprintsResult > InParent)
-{
-	FText ExtraSearchInfoText;
-
-	FFormatNamedArguments Args;
-	Args.Add(TEXT("Key"), InKey);
-	Args.Add(TEXT("Value"), InValue);
-	ExtraSearchInfoText = FText::Format(LOCTEXT("ExtraSearchInfo", "{Key}: {Value}"), Args);
-
-	TSharedPtr< FFindInBlueprintsResult > SearchResult(new FFindInBlueprintsResult(ExtraSearchInfoText, InParent) );
-
-	if (InKey.CompareTo(FFindInBlueprintSearchTags::FiB_NativeName) != 0)
-	{
-		InParent->Children.Add(SearchResult);
-	}
-}
-
-bool FFindInBlueprintsResult::ExtractJsonValue(const TMap<int32, FText>& InLookupTable, const TArray<FString>& InTokens, FText InKey, TSharedPtr< FJsonValue > InValue, TSharedPtr< FFindInBlueprintsResult > InParentOverride)
-{
-	bool bMatchesSearchResults = false;
-
-	TSharedPtr< FFindInBlueprintsResult > ParentResultPtr = InParentOverride.IsValid()? InParentOverride : SharedThis( this );
-
-	if( InValue->Type == EJson::Array)
-	{
-		TSharedPtr< FFindInBlueprintsResult > ArrayHeader = MakeShareable(new FFindInBlueprintsResult(InKey, ParentResultPtr) );
-		bool bArrayHasMatches = false;
-
-		TArray< TSharedPtr< FJsonValue > > Array = InValue->AsArray();
-		for( int ArrayIdx = 0; ArrayIdx < Array.Num(); ++ArrayIdx )
-		{
-			if(Array[ArrayIdx]->Type == EJson::String)
-			{
-				bArrayHasMatches |= ParseSearchInfo(InTokens, FText::FromString(FString::FromInt(ArrayIdx)), FindInBlueprintsHelpers::AsFText(Array[ArrayIdx], InLookupTable), ArrayHeader);
-			}
-			else
-			{
-				bArrayHasMatches |= ExtractJsonValue(InLookupTable, InTokens, FText::FromString(FString::FromInt(ArrayIdx)), Array[ArrayIdx], ArrayHeader);
-			}
-		}
-
-		bMatchesSearchResults |= bArrayHasMatches;
-
-		if(bArrayHasMatches)
-		{
-			ParentResultPtr->Children.Add(ArrayHeader);
-		}
-	}
-	else if(InValue->Type == EJson::Object)
-	{
-		TSharedPtr< FFindInBlueprintsResult > ObjectResultHeader = MakeShareable(new FFindInBlueprintsResult(InKey, ParentResultPtr) );
-
-		bMatchesSearchResults |= ExtractContent(InLookupTable, InValue->AsObject(), InTokens, ObjectResultHeader);
-
-		if(bMatchesSearchResults)
-		{
-			ParentResultPtr->Children.Add(ObjectResultHeader);
-		}
-	}
-	else if(InValue->Type == EJson::String)
-	{
-		bMatchesSearchResults |= ParseSearchInfo(InTokens, InKey, FindInBlueprintsHelpers::AsFText(InValue, InLookupTable), ParentResultPtr);
-	}
-	else
-	{
-		bMatchesSearchResults |= ParseSearchInfo(InTokens, InKey, FindInBlueprintsHelpers::AsFText(InValue, InLookupTable), ParentResultPtr);
-	}
-
-	return bMatchesSearchResults;
-}
-
-bool FFindInBlueprintsResult::ExtractContent(const TMap<int32, FText>& InLookupTable, TSharedPtr< FJsonObject > InJsonNode, const TArray<FString>& InTokens, TSharedPtr< FFindInBlueprintsResult > InParentOverride)
-{
-	bool bMatchesSearchResults = false;
-
-	TSharedPtr< FFindInBlueprintsResult > ParentResultPtr = InParentOverride.IsValid()? InParentOverride : SharedThis( this );
-
-	for( auto MapValues : InJsonNode->Values )
-	{
-		bMatchesSearchResults |= ExtractJsonValue(InLookupTable, InTokens, FText::FromString(MapValues.Key), MapValues.Value, ParentResultPtr);
-	}
-	return bMatchesSearchResults;
-}
-
 FText FFindInBlueprintsResult::GetCategory() const
 {
 	return FText::GetEmpty();
@@ -443,23 +223,17 @@ TSharedRef<SWidget> FFindInBlueprintsGraphNode::CreateIcon() const
 		.ToolTipText( GetCategory() );
 }
 
-bool FFindInBlueprintsGraphNode::ParseSearchInfo(const TArray<FString> &InTokens, FText InKey, FText InValue, TSharedPtr< FFindInBlueprintsResult > InParent)
+void FFindInBlueprintsGraphNode::ParseSearchInfo(FText InKey, FText InValue)
 {
 	if(InKey.CompareTo(FFindInBlueprintSearchTags::FiB_NodeGuid) == 0)
 	{
 		FString NodeGUIDAsString = InValue.ToString();
 		FGuid::Parse(NodeGUIDAsString, NodeGuid);
-
-		// NodeGuid's are not searchable
-		return false;
 	}
 
-	bool bMatchesTokens = FindInBlueprintsHelpers::TextMatchesSearchTokens(InTokens, InValue);
-
 	if(InKey.CompareTo(FFindInBlueprintSearchTags::FiB_ClassName) == 0)
 	{
 		ClassName = InValue.ToString();
-		bMatchesTokens = false;
 	}
 	else if(InKey.CompareTo(FFindInBlueprintSearchTags::FiB_Name) == 0)
 	{
@@ -472,67 +246,11 @@ bool FFindInBlueprintsGraphNode::ParseSearchInfo(const TArray<FString> &InTokens
 	else if(InKey.CompareTo(FFindInBlueprintSearchTags::FiB_Glyph) == 0)
 	{
 		GlyphBrush = FEditorStyle::GetBrush(*InValue.ToString());
-		bMatchesTokens = false;
 	}
 	else if(InKey.CompareTo(FFindInBlueprintSearchTags::FiB_GlyphColor) == 0)
 	{
 		GlyphColor.InitFromString(InValue.ToString());
-		bMatchesTokens = false;
 	}
-	else
-	{
-		// Any generic keys and value pairs get put in as "extra" search info
-		if(bMatchesTokens)
-		{
-			AddExtraSearchInfo(InKey, InValue, InParent);
-			return true;
-		}
-	}
-	return bMatchesTokens;
-}
-
-bool FFindInBlueprintsGraphNode::ExtractContent(const TMap<int32, FText>& InLookupTable, TSharedPtr< FJsonObject > InJsonNode, const TArray<FString>& InTokens)
-{
-	// Very important to get the schema first, other bits of data depend on it
-	FString SchemaName;
-	TSharedPtr< FJsonValue > SchemaNameValue = InJsonNode->GetField< EJson::String >(FFindInBlueprintSearchTags::FiB_SchemaName.ToString());
-	if(SchemaNameValue.IsValid())
-	{
-		SchemaName = FindInBlueprintsHelpers::AsFText(SchemaNameValue, InLookupTable).ToString();
-	}
-
-	bool bMatchesSearchResults = false;
-
-	for( auto MapValues : InJsonNode->Values )
-	{
-		FText Key = FText::FromString(MapValues.Key);
-
-		if(Key.CompareTo(FFindInBlueprintSearchTags::FiB_Pins) == 0)
-		{
-			TArray< TSharedPtr< FJsonValue > > PinsList = MapValues.Value->AsArray();
-
-			for( TSharedPtr< FJsonValue > Pin : PinsList )
-			{
-				// The pin
-				TSharedPtr< FFindInBlueprintsPin > PinResult(new FFindInBlueprintsPin(FText::GetEmpty(), SharedThis( this ), SchemaName) );
-				if(PinResult->ExtractContent(InLookupTable, Pin->AsObject(), InTokens))
-				{
-					Children.Add(PinResult);
-					bMatchesSearchResults = true;
-				}
-			}
-		}
-		else if(Key.CompareTo(FFindInBlueprintSearchTags::FiB_SchemaName) == 0)
-		{
-			// Previously extracted
-			continue;
-		}
-		else
-		{
-			bMatchesSearchResults |= ExtractJsonValue(InLookupTable, InTokens, FText::FromString(MapValues.Key), MapValues.Value, SharedThis( this ));
-		}
-	}
-	return bMatchesSearchResults;
 }
 
 FText FFindInBlueprintsGraphNode::GetCategory() const
@@ -603,32 +321,16 @@ TSharedRef<SWidget> FFindInBlueprintsPin::CreateIcon() const
 		.ToolTipText(FText::FromString(FindInBlueprintsHelpers::GetPinTypeAsString(PinType)));
 }
 
-bool FFindInBlueprintsPin::ParseSearchInfo(const TArray<FString> &InTokens, FText InKey, FText InValue, TSharedPtr< FFindInBlueprintsResult > InParent)
+void FFindInBlueprintsPin::ParseSearchInfo(FText InKey, FText InValue)
 {
-	bool bMatchesTokens = FindInBlueprintsHelpers::TextMatchesSearchTokens(InTokens, InValue);
-
 	if(InKey.CompareTo(FFindInBlueprintSearchTags::FiB_Name) == 0)
 	{
 		DisplayText = InValue;
-		bMatchesTokens = FindInBlueprintsHelpers::StringMatchesSearchTokens(InTokens, DisplayText.ToString());
-	}
-	else if(FindInBlueprintsHelpers::ParsePinType(InKey, InValue, PinType, PinSubCategory))
-	{
-		// For pins, the only information in the Pin Type we care about is the object class, so return false if the Key is not that
-		if(InKey.CompareTo(FFindInBlueprintSearchTags::FiB_ObjectClass) != 0)
-		{
-			return false;
-		}
 	}
 	else
 	{
-		// Any generic keys and value pairs get put in as "extra" search info
-		if(bMatchesTokens)
-		{
-			AddExtraSearchInfo(InKey, InValue, InParent);
-		}
+		FindInBlueprintsHelpers::ParsePinType(InKey, InValue, PinType);
 	}
-	return bMatchesTokens;
 }
 
 FText FFindInBlueprintsPin::GetCategory() const
@@ -638,14 +340,18 @@ FText FFindInBlueprintsPin::GetCategory() const
 
 void FFindInBlueprintsPin::FinalizeSearchData()
 {
-	if(!PinSubCategory.IsEmpty())
+	if(!PinType.PinSubCategory.IsEmpty())
 	{
-		PinType.PinSubCategoryObject = FindObject<UClass>(ANY_PACKAGE, *PinSubCategory, true);
+		PinType.PinSubCategoryObject = FindObject<UClass>(ANY_PACKAGE, *PinType.PinSubCategory, true);
 		if(!PinType.PinSubCategoryObject.IsValid())
 		{
-			PinType.PinSubCategoryObject = FindObject<UScriptStruct>(UObject::StaticClass(), *PinSubCategory);
+			PinType.PinSubCategoryObject = FindObject<UScriptStruct>(UObject::StaticClass(), *PinType.PinSubCategory);
+		}
+
+		if (PinType.PinSubCategoryObject.IsValid())
+		{
+			PinType.PinSubCategory.Empty();
 		}
-		PinSubCategory.Empty();
 	}
 
 	if(!SchemaName.IsEmpty())
@@ -683,38 +389,20 @@ TSharedRef<SWidget> FFindInBlueprintsProperty::CreateIcon() const
 		.ToolTipText( FText::FromString(FindInBlueprintsHelpers::GetPinTypeAsString(PinType)) );
 }
 
-bool FFindInBlueprintsProperty::ParseSearchInfo(const TArray<FString> &InTokens, FText InKey, FText InValue, TSharedPtr< FFindInBlueprintsResult > InParent)
+void FFindInBlueprintsProperty::ParseSearchInfo(FText InKey, FText InValue)
 {
-	bool bMatchesTokens = FindInBlueprintsHelpers::TextMatchesSearchTokens(InTokens, InValue);
-
 	if(InKey.CompareTo(FFindInBlueprintSearchTags::FiB_Name) == 0)
 	{
 		DisplayText = InValue;
 	}
-	else if(FindInBlueprintsHelpers::ParsePinType(InKey, InValue, PinType, PinSubCategory))
-	{
-		// For properties, we care about all the information in the pin type, except whether its a reference or an array
-		if(InKey.CompareTo(FFindInBlueprintSearchTags::FiB_IsArray) == 0 || InKey.CompareTo(FFindInBlueprintSearchTags::FiB_IsReference) == 0)
-		{
-			return false;
-		}
-	}
 	else if(InKey.CompareTo(FFindInBlueprintSearchTags::FiB_IsSCSComponent) == 0)
 	{
 		bIsSCSComponent = true;
-
-		// This value is not searchable
-		bMatchesTokens = false;
 	}
 	else
 	{
-		// Any generic keys and value pairs get put in as "extra" search info
-		if(bMatchesTokens)
-		{
-			AddExtraSearchInfo(InKey, InValue, InParent);
-		}
+		FindInBlueprintsHelpers::ParsePinType(InKey, InValue, PinType);
 	}
-	return bMatchesTokens;
 }
 
 FText FFindInBlueprintsProperty::GetCategory() const
@@ -728,14 +416,18 @@ FText FFindInBlueprintsProperty::GetCategory() const
 
 void FFindInBlueprintsProperty::FinalizeSearchData()
 {
-	if(!PinSubCategory.IsEmpty())
+	if(!PinType.PinSubCategory.IsEmpty())
 	{
-		PinType.PinSubCategoryObject = FindObject<UClass>(ANY_PACKAGE, *PinSubCategory, true);
+		PinType.PinSubCategoryObject = FindObject<UClass>(ANY_PACKAGE, *PinType.PinSubCategory, true);
 		if(!PinType.PinSubCategoryObject.IsValid())
 		{
-			PinType.PinSubCategoryObject = FindObject<UScriptStruct>(UObject::StaticClass(), *PinSubCategory);
+			PinType.PinSubCategoryObject = FindObject<UScriptStruct>(UObject::StaticClass(), *PinType.PinSubCategory);
+		}
+
+		if (PinType.PinSubCategoryObject.IsValid())
+		{
+			PinType.PinSubCategory.Empty();
 		}
-		PinSubCategory.Empty();
 	}
 }
 
@@ -792,74 +484,12 @@ TSharedRef<SWidget> FFindInBlueprintsGraph::CreateIcon() const
 		.ToolTipText( GetCategory() );
 }
 
-bool FFindInBlueprintsGraph::ExtractContent(const TMap<int32, FText>& InLookupTable, TSharedPtr< FJsonObject > InJsonNode, const TArray<FString>& InTokens)
+void FFindInBlueprintsGraph::ParseSearchInfo(FText InKey, FText InValue)
 {
-	bool bMatchesSearchResults = false;
-
-	for( auto MapValues : InJsonNode->Values )
-	{
-		FText Key = FText::FromString(MapValues.Key);
-
-		if(Key.CompareTo(FFindInBlueprintSearchTags::FiB_Nodes) == 0)
-		{
-			TArray< TSharedPtr< FJsonValue > > NodeList = MapValues.Value->AsArray();
-
-			for( TSharedPtr< FJsonValue > NodeValue : NodeList )
-			{
-				TSharedPtr< FFindInBlueprintsGraphNode > NodeResult(new FFindInBlueprintsGraphNode(FText::GetEmpty(), SharedThis( this )) );
-
-				if( NodeResult->ExtractContent(InLookupTable, NodeValue->AsObject(), InTokens) )
-				{
-					Children.Add(NodeResult);
-					bMatchesSearchResults = true;
-				}
-			}
-		}
-		else if(Key.CompareTo(FFindInBlueprintSearchTags::FiB_SchemaName) == 0)
-		{
-			// Previously extracted
-			continue;
-		}
-		else if(Key.CompareTo(FFindInBlueprintSearchTags::FiB_Properties) == 0)
-		{
-			// Pulls out all properties (variables) for this Blueprint
-			TArray<TSharedPtr< FJsonValue > > PropertyList = MapValues.Value->AsArray();
-			for( TSharedPtr< FJsonValue > PropertyValue : PropertyList )
-			{
-				FSearchResult Result(new FFindInBlueprintsProperty(FText::GetEmpty(), SharedThis( this )) );
-
-				if( Result->ExtractContent(InLookupTable, PropertyValue->AsObject(), InTokens) )
-				{
-					Children.Add(Result);
-					bMatchesSearchResults = true;
-				}
-			}
-		}
-		else
-		{
-			bMatchesSearchResults |= ExtractJsonValue( InLookupTable, InTokens, FText::FromString(MapValues.Key), MapValues.Value, SharedThis( this ) );
-		}
-	}
-	return bMatchesSearchResults;
-}
-
-bool FFindInBlueprintsGraph::ParseSearchInfo(const TArray<FString> &InTokens, FText InKey, FText InValue, TSharedPtr< FFindInBlueprintsResult > InParent)
-{
-	bool bMatchesTokens = FindInBlueprintsHelpers::TextMatchesSearchTokens(InTokens, InValue);
-
 	if(InKey.CompareTo(FFindInBlueprintSearchTags::FiB_Name) == 0)
 	{
 		DisplayText = InValue;
 	}
-	else
-	{
-		if(bMatchesTokens)
-		{
-			AddExtraSearchInfo(InKey, InValue, InParent);
-		}
-	}
-
-	return bMatchesTokens;
 }
 
 FText FFindInBlueprintsGraph::GetCategory() const
@@ -885,8 +515,8 @@ class FStreamSearch : public FRunnable
 {
 public:
 	/** Constructor */
-	FStreamSearch(const TArray<FString>& InSearchTokens )
-		: SearchTokens(InSearchTokens)
+	FStreamSearch(const FString& InSearchValue )
+		: SearchValue(InSearchValue)
 		, bThreadCompleted(false)
 		, StopTaskCounter(0)
 	{
@@ -900,7 +530,44 @@ public:
 		return true;
 	}
 
-	virtual uint32 Run() override;
+	virtual uint32 Run() override
+	{
+		FFindInBlueprintSearchManager::Get().BeginSearchQuery(this);
+
+		TFunction<void(const FSearchResult&)> OnResultReady = [this](const FSearchResult& Result) {
+			FScopeLock ScopeLock(&SearchCriticalSection);
+			ItemsFound.Add(Result);
+		};
+
+		// Searching comes to an end if it is requested using the StopTaskCounter or continuing the search query yields no results
+		FSearchData QueryResult;
+		while (FFindInBlueprintSearchManager::Get().ContinueSearchQuery(this, QueryResult))
+		{
+			// This work cannot be spun off to a use the TaskGraph system for spreading the work over multiple threads/cores as the stack limit is too small for the work being done.
+			if (QueryResult.Value.Len() > 0)
+			{
+				TSharedPtr< FImaginaryBlueprint> ImaginaryBlueprint(new FImaginaryBlueprint(FPaths::GetBaseFilename(QueryResult.BlueprintPath.ToString()), QueryResult.BlueprintPath.ToString(), QueryResult.ParentClass, QueryResult.Interfaces, QueryResult.Value));
+				TSharedPtr< FFiBSearchInstance > SearchInstance(new FFiBSearchInstance);
+				FSearchResult SearchResult = SearchInstance->StartSearchQuery(*SearchValue, ImaginaryBlueprint);
+
+				// If there are children, add the item to the search results
+				if(SearchResult.IsValid() && SearchResult->Children.Num() != 0)
+				{
+					OnResultReady(SearchResult);
+				}
+			}
+
+			if (StopTaskCounter.GetValue())
+			{
+				// Ensure that the FiB Manager knows that we are done searching
+				FFindInBlueprintSearchManager::Get().EnsureSearchQueryEnds(this);
+			}
+		}
+
+		bThreadCompleted = true;
+
+		return 0;
+	}
 
 	virtual void Stop() override
 	{
@@ -913,8 +580,6 @@ public:
 	}
 	/** End FRunnable Interface */
 
-	TArray<FSearchData> MakeExtractionBatch() const;
-
 	/** Brings the thread to a safe stop before continuing. */
 	void EnsureCompletion()
 	{
@@ -936,7 +601,7 @@ public:
 	}
 
 	/**
-	 * Appends the items filtered through the search tokens to the passed array
+	 * Appends the items filtered through the search filter to the passed array
 	 *
 	 * @param OutItemsFound		All the items found since last queried
 	 */
@@ -960,8 +625,8 @@ public:
 	/** A list of items found, cleared whenever the main thread pulls them to display to screen */
 	TArray<FSearchResult> ItemsFound;
 
-	/** The search tokens to filter results by */
-	TArray<FString> SearchTokens;
+	/** The search value to filter results by */
+	FString SearchValue;
 
 	/** Prevents searching while other threads are pulling search results */
 	FCriticalSection SearchCriticalSection;
@@ -1375,34 +1040,6 @@ ECheckBoxState SFindInBlueprints::OnGetFindModeChecked() const
 
 void SFindInBlueprints::InitiateSearch()
 {
-	TArray<FString> Tokens;
-	if(SearchValue.Contains("\"") && SearchValue.ParseIntoArray(Tokens, TEXT("\""), true)>0)
-	{
-		for( auto &TokenIt : Tokens )
-		{
-			// we have the token, we don't need the quotes anymore, they'll just confused the comparison later on
-			TokenIt = TokenIt.TrimQuotes();
-			// We remove the spaces as all later comparison strings will also be de-spaced
-			TokenIt = TokenIt.Replace( TEXT( " " ), TEXT( "" ) );
-		}
-
-		// due to being able to handle multiple quoted blocks like ("Make Epic" "Game Now") we can end up with
-		// and empty string between (" ") blocks so this simply removes them
-		struct FRemoveMatchingStrings
-		{ 
-			bool operator()(const FString& RemovalCandidate) const
-			{
-				return RemovalCandidate.IsEmpty();
-			}
-		};
-		Tokens.RemoveAll( FRemoveMatchingStrings() );
-	}
-	else
-	{
-		// unquoted search equivalent to a match-any-of search
-		SearchValue.ParseIntoArray(Tokens, TEXT(" "), true);
-	}
-
 	if(ItemsFound.Num())
 	{
 		// Reset the scroll to the top
@@ -1411,7 +1048,7 @@ void SFindInBlueprints::InitiateSearch()
 
 	ItemsFound.Empty();
 
-	if (Tokens.Num() > 0)
+	if (SearchValue.Len() > 0)
 	{
 		OnRemoveCacheBar();
 
@@ -1420,55 +1057,59 @@ void SFindInBlueprints::InitiateSearch()
 		HighlightText = FText::FromString( SearchValue );
 		if (bIsInFindWithinBlueprintMode)
 		{
-			MatchTokensWithinBlueprint(Tokens);
-		}
-		else
-		{
-			LaunchStreamThread(Tokens);
-		}
-	}
-}
+			if(StreamSearch.IsValid() && !StreamSearch->IsComplete())
+			{
+				StreamSearch->Stop();
+				StreamSearch->EnsureCompletion();
+				StreamSearch.Reset();
+			}
 
-void SFindInBlueprints::MatchTokensWithinBlueprint(const TArray<FString>& Tokens)
-{
-	if(StreamSearch.IsValid() && !StreamSearch->IsComplete())
-	{
-		StreamSearch->Stop();
-		StreamSearch->EnsureCompletion();
-		StreamSearch->GetFilteredItems(ItemsFound);
-		ItemsFound.Empty();
-		StreamSearch.Reset();
-	}
+			UBlueprint* Blueprint = BlueprintEditorPtr.Pin()->GetBlueprintObj();
+			FString ParentClass;
+			if (UProperty* ParentClassProp = Blueprint->GetClass()->FindPropertyByName(GET_MEMBER_NAME_CHECKED(UBlueprint, ParentClass)))
+			{
+				ParentClassProp->ExportTextItem(ParentClass, ParentClassProp->ContainerPtrToValuePtr<uint8>(Blueprint), nullptr, Blueprint, 0);
+			}
 
-	// No multi-threading for single Blueprint search, just query for the data block and process it.
-	UBlueprint* Blueprint = BlueprintEditorPtr.Pin()->GetBlueprintObj();
+			TArray<FString> Interfaces;
 
-	RootSearchResult = FSearchResult(new FFindInBlueprintsResult(FText::FromString(Blueprint->GetPathName())));
-	
-	FindInBlueprintsHelpers::Extract(FFindInBlueprintSearchManager::Get().QuerySingleBlueprint(Blueprint), Tokens, RootSearchResult);
+			for (FBPInterfaceDescription& InterfaceDesc : Blueprint->ImplementedInterfaces)
+			{
+				Interfaces.Add(InterfaceDesc.Interface->GetPathName());
+			}
+			TSharedPtr< FImaginaryBlueprint> ImaginaryBlueprint(new FImaginaryBlueprint(Blueprint->GetName(), Blueprint->GetPathName(), ParentClass, Interfaces, FFindInBlueprintSearchManager::Get().QuerySingleBlueprint(Blueprint)));
+			TSharedPtr< FFiBSearchInstance > SearchInstance(new FFiBSearchInstance);
+			FSearchResult SearchResult = SearchInstance->StartSearchQuery(SearchValue, ImaginaryBlueprint);
 
-	// We do not use the root, the user already knows what Blueprint they are searching within
-	ItemsFound = RootSearchResult->Children;
-	
-	if(ItemsFound.Num() == 0)
-	{
-		// Insert a fake result to inform user if none found
-		ItemsFound.Add(FSearchResult(new FFindInBlueprintsResult(LOCTEXT("BlueprintSearchNoResults", "No Results found"))));
-		HighlightText = FText::GetEmpty();
-		TreeView->RequestTreeRefresh();
-	}
-	else
-	{
-		for(auto Item : ItemsFound)
+			if (SearchResult.IsValid())
+			{
+				ItemsFound = SearchResult->Children;
+			}
+
+			if(ItemsFound.Num() == 0)
+			{
+				// Insert a fake result to inform user if none found
+				ItemsFound.Add(FSearchResult(new FFindInBlueprintsResult(LOCTEXT("BlueprintSearchNoResults", "No Results found"))));
+				HighlightText = FText::GetEmpty();
+			}
+			else
+			{
+				for(auto Item : ItemsFound)
+				{
+					Item->ExpandAllChildren(TreeView);
+				}
+			}
+
+			TreeView->RequestTreeRefresh();
+		}
+		else
 		{
-			Item->ExpandAllChildren(TreeView);
+			LaunchStreamThread(SearchValue);
 		}
 	}
-
-	TreeView->RequestTreeRefresh();
 }
 
-void SFindInBlueprints::LaunchStreamThread(const TArray<FString>& InTokens)
+void SFindInBlueprints::LaunchStreamThread(const FString& InSearchValue)
 {
 	if(StreamSearch.IsValid() && !StreamSearch->IsComplete())
 	{
@@ -1481,7 +1122,7 @@ void SFindInBlueprints::LaunchStreamThread(const TArray<FString>& InTokens)
 		RegisterActiveTimer( 0.f, FWidgetActiveTimerDelegate::CreateSP( this, &SFindInBlueprints::UpdateSearchResults ) );
 	}
 
-	StreamSearch = MakeShareable(new FStreamSearch(InTokens));
+	StreamSearch = MakeShareable(new FStreamSearch(InSearchValue));
 }
 
 TSharedRef<ITableRow> SFindInBlueprints::OnGenerateRow( FSearchResult InItem, const TSharedRef<STableViewBase>& OwnerTable )
@@ -1529,7 +1170,7 @@ TSharedRef<ITableRow> SFindInBlueprints::OnGenerateRow( FSearchResult InItem, co
 			[
 				SNew(SHorizontalBox)
 				+SHorizontalBox::Slot()
-				.VAlign(VAlign_Top)
+				.VAlign(VAlign_Center)
 				.AutoWidth()
 				[
 					InItem->CreateIcon()
@@ -1771,140 +1412,4 @@ void SFindInBlueprints::OnCopyAction()
 	FPlatformMisc::ClipboardCopy( *SelectedText );
 }
 
-////////////////////////////////////
-// FExtractionTask
-
-/**
- * Task to search a batch of blueprints
- *
- * Each small batch is non-abandonable, but if the search is cancelled, no more batches are queued
- */
-class FExtractionTask : public FNonAbandonableTask
-{
-public:
-	/** Async task function */
-	void DoWork();
-
-	/** Batch of blueprints and associated data to search */
-	TArray<FSearchData> QueryResultBatch;
-
-	/** Pointer to array of tokens to search for */
-	const TArray<FString>* Tokens;
-
-	/** Function to call once the search result is ready */
-	TFunction<void(const FSearchResult&)> OnResultReady;
-
-	FORCEINLINE TStatId GetStatId() const
-	{
-		RETURN_QUICK_DECLARE_CYCLE_STAT( FExtractionTask, STATGROUP_ThreadPoolAsyncTasks );
-	}
-};
-
-TArray<FSearchData> FStreamSearch::MakeExtractionBatch() const
-{
-	// Batch sizes between 10 and 30 seem fastest
-	static const int BatchSize = 15;
-
-	TArray<FSearchData> Batch;
-	do
-	{
-		if (StopTaskCounter.GetValue())
-			return TArray<FSearchData>();
-
-		FSearchData QueryResult;
-		if (!FFindInBlueprintSearchManager::Get().ContinueSearchQuery(this, QueryResult))
-			return Batch;
-
-		Batch.Push(QueryResult);
-	}
-	while (Batch.Num() <= BatchSize);
-
-	return Batch;
-}
-
-uint32 FStreamSearch::Run()
-{
-	FFindInBlueprintSearchManager::Get().BeginSearchQuery(this);
-
-	// Seems to go fastest at or slightly above one thread per physical core
-	TArray<FAsyncTask<FExtractionTask>> ExtractionTasks;
-	ExtractionTasks.SetNum(FPlatformMisc::NumberOfCores());
-
-	TFunction<void(const FSearchResult&)> OnResultReady = [this](const FSearchResult& Result) {
-		FScopeLock ScopeLock(&SearchCriticalSection);
-		ItemsFound.Add(Result);
-	};
-
-	// Searching comes to an end if it is requested using the StopTaskCounter or continuing the search query yields no results
-	TArray<FSearchData> QueryResultBatch;
-	for (;;)
-	{
-		if (QueryResultBatch.Num() == 0)
-		{
-			QueryResultBatch = MakeExtractionBatch();
-		}
-		else
-		{
-			// Spinning until a task becomes available
-			FPlatformProcess::Sleep(0);
-		}
-
-		if (QueryResultBatch.Num() == 0)
-		{
-			// All search work has been dispatched
-			break;
-		}
-
-		for (auto& Task: ExtractionTasks)
-		{
-			if (Task.IsWorkDone())
-			{
-				Task.EnsureCompletion();
-
-				auto& worker = Task.GetTask();
-				worker.QueryResultBatch = MoveTemp(QueryResultBatch);
-				worker.Tokens = &SearchTokens;
-				worker.OnResultReady = OnResultReady;
-
-				Task.StartBackgroundTask();
-
-				// Ensure batch array is empty (probably already is due to move above)
-				QueryResultBatch.Reset();
-				break;
-			}
-		}
-	}
-
-	// Wait for all work to finish
-	for (auto& Task: ExtractionTasks)
-	{
-		Task.EnsureCompletion();
-	}
-
-	if (StopTaskCounter.GetValue())
-	{
-		// Ensure that the FiB Manager knows that we are done searching
-		FFindInBlueprintSearchManager::Get().EnsureSearchQueryEnds(this);
-	}
-
-	bThreadCompleted = true;
-
-	return 0;
-}
-
-void FExtractionTask::DoWork()
-{
-	for (const auto& QueryResult: QueryResultBatch)
-	{
-		FSearchResult BlueprintCategory = FSearchResult(new FFindInBlueprintsResult(FText::FromName(QueryResult.BlueprintPath)));
-		FindInBlueprintsHelpers::Extract(QueryResult.Value, *Tokens, BlueprintCategory);
-
-		// If there are children, add the item to the search results
-		if(BlueprintCategory->Children.Num() != 0)
-		{
-			OnResultReady(BlueprintCategory);
-		}
-	}
-}
-
 #undef LOCTEXT_NAMESPACE
diff --git a/Engine/Source/Editor/Kismet/Private/ImaginaryBlueprintData.cpp b/Engine/Source/Editor/Kismet/Private/ImaginaryBlueprintData.cpp
new file mode 100644
--- /dev/null
+++ b/Engine/Source/Editor/Kismet/Private/ImaginaryBlueprintData.cpp
@@ -0,0 +1,666 @@
+// Copyright 1998-2015 Epic Games, Inc. All Rights Reserved.
+#include "BlueprintEditorPrivatePCH.h"
+#include "ImaginaryBlueprintData.h"
+#include "Json.h"
+#include "FindInBlueprints.h"
+#include "FindInBlueprintManager.h"
+
+#include "TextFilter.h"
+
+#define LOCTEXT_NAMESPACE "FindInBlueprints"
+
+///////////////////////
+// FSearchableValueInfo
+
+FText FSearchableValueInfo::GetDisplayText(const TMap<int32, FText>& InLookupTable) const
+{
+	if (!DisplayText.IsEmpty() || LookupTableKey == -1)
+	{
+		return DisplayText;
+	}
+	return FindInBlueprintsHelpers::AsFText(LookupTableKey, InLookupTable);
+}
+
+////////////////////////////
+// FComponentUniqueDisplay
+
+bool FComponentUniqueDisplay::operator==(const FComponentUniqueDisplay& Other)
+{
+	// Two search results in the same object/sub-object should never have the same display string ({Key}: {Value} pairing)
+	return SearchResult.IsValid() && Other.SearchResult.IsValid() && SearchResult->GetDisplayString().CompareTo(Other.SearchResult->GetDisplayString()) == 0;
+}
+
+///////////////////////
+// FImaginaryFiBData
+
+FImaginaryFiBData::FImaginaryFiBData() : LookupTablePtr(nullptr)
+{
+}
+
+FImaginaryFiBData::FImaginaryFiBData(TSharedPtr< FJsonObject > InUnparsedJsonObject, TMap<int32, FText>* InLookupTablePtr)
+	: UnparsedJsonObject(InUnparsedJsonObject)
+	, LookupTablePtr(InLookupTablePtr)
+{
+}
+
+FSearchResult FImaginaryFiBData::CreateSearchResult(FSearchResult InParent)
+{
+	FSearchResult ReturnSearchResult = CreateSearchResult_Internal(InParent);
+
+	FText DisplayName;
+	for( auto TagsAndValues : ParsedTagsAndValues )
+	{
+		if (TagsAndValues.Value.IsCoreDisplay() || !TagsAndValues.Value.IsSearchable())
+		{
+			FText Value = TagsAndValues.Value.GetDisplayText(*LookupTablePtr);
+			ReturnSearchResult->ParseSearchInfo(FText::FromString(TagsAndValues.Key), Value);
+		}
+	}
+	return ReturnSearchResult;
+}
+
+FSearchResult FImaginaryFiBData::CreateSearchTree(FSearchResult InParentSearchResult, TWeakPtr< FImaginaryFiBData > InCurrentPointer, TArray< const FImaginaryFiBData* >& InValidSearchResults, TMultiMap< const FImaginaryFiBData*, FComponentUniqueDisplay >& InMatchingSearchComponents)
+{
+	FSearchResult CurrentSearchResult = InCurrentPointer.Pin()->CreateSearchResult(InParentSearchResult);
+	bool bValidSearchResults = false;
+
+	// Check all children first, to see if they are valid in the search results
+	for (TSharedPtr<FImaginaryFiBData> ChildData : InCurrentPointer.Pin()->ParsedChildData)
+	{
+		FSearchResult Result = CreateSearchTree(CurrentSearchResult, ChildData, InValidSearchResults, InMatchingSearchComponents);
+		if (Result.IsValid())
+		{
+			bValidSearchResults = true;
+			CurrentSearchResult->Children.Add(Result);
+		}
+	}
+
+	// If the children did not match the search results but this item does, then we will want to return true
+	if (!bValidSearchResults && (InValidSearchResults.Find(InCurrentPointer.Pin().Get()) != INDEX_NONE || InMatchingSearchComponents.Find(InCurrentPointer.Pin().Get())))
+	{
+		bValidSearchResults = true;
+	}
+
+	if (bValidSearchResults)
+	{
+		TArray< FComponentUniqueDisplay > SearchResultList;
+		InMatchingSearchComponents.MultiFind(InCurrentPointer.Pin().Get(), SearchResultList, true);
+		CurrentSearchResult->Children.Reserve(CurrentSearchResult->Children.Num() + SearchResultList.Num());
+
+		// Add any data that matched the search results as a child of our search result
+		for (FComponentUniqueDisplay& SearchResultWrapper : SearchResultList)
+		{
+			SearchResultWrapper.SearchResult->Parent = CurrentSearchResult;
+			CurrentSearchResult->Children.Add(SearchResultWrapper.SearchResult);
+		}
+		return CurrentSearchResult;
+	}
+	return nullptr;
+}
+
+bool FImaginaryFiBData::IsCompatibleWithFilter(ESearchQueryFilter InSearchQueryFilter)
+{
+	return true;
+}
+
+bool FImaginaryFiBData::CanCallFilter(ESearchQueryFilter InSearchQueryFilter)
+{
+	// Always compatible with the AllFilter
+	return InSearchQueryFilter == ESearchQueryFilter::AllFilter;
+}
+
+void FImaginaryFiBData::ParseAllChildData()
+{
+	if (UnparsedJsonObject.IsValid())
+	{
+		for( auto MapValues : UnparsedJsonObject->Values )
+		{
+			TSharedPtr< FJsonValue > JsonValue = MapValues.Value;
+
+			if (!TrySpecialHandleJsonValue(MapValues.Key, JsonValue))
+			{
+				FText Key = FText::FromString(MapValues.Key);
+				ParseJsonValue(Key, Key, JsonValue);
+			}
+		}
+	}
+
+	UnparsedJsonObject.Reset();
+}
+
+void FImaginaryFiBData::ParseJsonValue(FText InKey, FText InDisplayKey, TSharedPtr< FJsonValue > InJsonValue)
+{
+	if( InJsonValue->Type == EJson::String)
+	{
+		ParsedTagsAndValues.Add(InKey.ToString(), FSearchableValueInfo(InDisplayKey, FCString::Atoi(*InJsonValue->AsString()), GetSearchabilityStatus(InKey.ToString())));
+	}
+	else if (InJsonValue->Type == EJson::Boolean)
+	{
+		ParsedTagsAndValues.Add(InKey.ToString(), FSearchableValueInfo(InDisplayKey, FText::FromString(InJsonValue->AsString()), GetSearchabilityStatus(InKey.ToString())));
+	}
+	else if( InJsonValue->Type == EJson::Array)
+	{
+		TSharedPtr< FCategorySectionHelper > ArrayCategory = MakeShareable(new FCategorySectionHelper(LookupTablePtr, InKey, true));
+		ParsedChildData.Add(ArrayCategory);
+		TArray<TSharedPtr< FJsonValue > > ArrayList = InJsonValue->AsArray();
+		for( int32 ArrayIdx = 0; ArrayIdx < ArrayList.Num(); ++ArrayIdx)
+		{
+			TSharedPtr< FJsonValue > ArrayValue = ArrayList[ArrayIdx];
+			ArrayCategory->ParseJsonValue(InKey, FText::FromString(FString::FromInt(ArrayIdx)), ArrayValue);
+		}		
+	}
+	else if (InJsonValue->Type == EJson::Object)
+	{
+		TSharedPtr< FCategorySectionHelper > SubObjectCategory = MakeShareable(new FCategorySectionHelper(InJsonValue->AsObject(), LookupTablePtr, InDisplayKey, false));
+		SubObjectCategory->ParseAllChildData();
+		ParsedChildData.Add(SubObjectCategory);
+	}
+	else
+	{
+		// For everything else, there's this. Numbers come here and will be treated as strings
+		ParsedTagsAndValues.Add(InKey.ToString(), FSearchableValueInfo(InDisplayKey, FText::FromString(InJsonValue->AsString()), GetSearchabilityStatus(InKey.ToString())));
+	}
+}
+
+FText FImaginaryFiBData::CreateSearchComponentDisplayText(FText InKey, FText InValue) const
+{
+	FFormatNamedArguments Args;
+	Args.Add(TEXT("Key"), InKey);
+	Args.Add(TEXT("Value"), InValue);
+	return FText::Format(LOCTEXT("ExtraSearchInfo", "{Key}: {Value}"), Args);
+}
+
+bool FImaginaryFiBData::TestBasicStringExpression(const FTextFilterString& InValue, const ETextFilterTextComparisonMode InTextComparisonMode, TMultiMap< const FImaginaryFiBData*, FComponentUniqueDisplay >& InOutMatchingSearchComponents) const
+{
+	bool bMatchesSearchQuery = false;
+	for( auto ParsedValues : ParsedTagsAndValues )
+	{
+		if (ParsedValues.Value.IsSearchable() && !ParsedValues.Value.IsExplicitSearchable())
+		{
+			FText Value = ParsedValues.Value.GetDisplayText(*LookupTablePtr);
+			FString ValueAsString = Value.ToString();
+			ValueAsString.ReplaceInline(TEXT(" "), TEXT(""));
+			bool bMatchesSearch = TextFilterUtils::TestBasicStringExpression(ValueAsString, InValue, InTextComparisonMode) || TextFilterUtils::TestBasicStringExpression(Value.BuildSourceString(), InValue, InTextComparisonMode);
+			
+			if (bMatchesSearch && !ParsedValues.Value.IsCoreDisplay())
+			{
+				FSearchResult SearchResult(new FFindInBlueprintsResult(CreateSearchComponentDisplayText(ParsedValues.Value.GetDisplayKey(), Value), nullptr ));
+				InOutMatchingSearchComponents.Add(this, FComponentUniqueDisplay(SearchResult));
+			}
+
+			bMatchesSearchQuery |= bMatchesSearch;
+		}
+	}
+	// Any children that are treated as a TagAndValue Category should be added for independent searching
+	for (const TSharedPtr< FImaginaryFiBData > Child : ParsedChildData)
+	{
+		if (Child->IsTagAndValueCategory())
+		{
+			bMatchesSearchQuery |= Child->TestBasicStringExpression(InValue, InTextComparisonMode, InOutMatchingSearchComponents);
+		}
+	}
+
+	return bMatchesSearchQuery;
+}
+
+bool FImaginaryFiBData::TestComplexExpression(const FName& InKey, const FTextFilterString& InValue, const ETextFilterComparisonOperation InComparisonOperation, const ETextFilterTextComparisonMode InTextComparisonMode, TMultiMap< const FImaginaryFiBData*, FComponentUniqueDisplay >& InOutMatchingSearchComponents) const
+{
+	bool bMatchesSearchQuery = false;
+
+	TArray<FSearchableValueInfo> ValueList;
+	ParsedTagsAndValues.MultiFind(InKey.ToString(), ValueList, true);
+	for (const FSearchableValueInfo& ItemValue : ValueList)
+	{
+		if (ItemValue.IsSearchable())
+		{
+			FText Value = ItemValue.GetDisplayText(*LookupTablePtr);
+			FString ValueAsString = Value.ToString();
+			ValueAsString.ReplaceInline(TEXT(" "), TEXT(""));
+			bool bMatchesSearch = TextFilterUtils::TestComplexExpression(ValueAsString, InValue, InComparisonOperation, InTextComparisonMode) || TextFilterUtils::TestBasicStringExpression(Value.BuildSourceString(), InValue, InTextComparisonMode);
+
+			if (bMatchesSearch && !ItemValue.IsCoreDisplay())
+			{
+				FSearchResult SearchResult(new FFindInBlueprintsResult(CreateSearchComponentDisplayText(ItemValue.GetDisplayKey(), Value), nullptr ));
+				InOutMatchingSearchComponents.Add(this, FComponentUniqueDisplay(SearchResult));
+			}
+			bMatchesSearchQuery |= bMatchesSearch;
+		}
+	}
+
+	// Any children that are treated as a TagAndValue Category should be added for independent searching
+	for (const TSharedPtr< FImaginaryFiBData > Child : ParsedChildData)
+	{
+		if (Child->IsTagAndValueCategory())
+		{
+			bMatchesSearchQuery |= Child->TestComplexExpression(InKey, InValue, InComparisonOperation, InTextComparisonMode, InOutMatchingSearchComponents);
+		}
+	}
+	return bMatchesSearchQuery;
+}
+
+///////////////////////////
+// FCategorySectionHelper
+
+FCategorySectionHelper::FCategorySectionHelper(TMap<int32, FText>* InLookupTablePtr, FText InCategoryName, bool bInTagAndValueCategory)
+	: FImaginaryFiBData(nullptr, InLookupTablePtr)
+	, CategoryName(InCategoryName)
+	, bIsTagAndValue(bInTagAndValueCategory)
+{
+}
+
+FCategorySectionHelper::FCategorySectionHelper(TSharedPtr< FJsonObject > InUnparsedJsonObject, TMap<int32, FText>* InLookupTablePtr, FText InCategoryName, bool bInTagAndValueCategory)
+	: FImaginaryFiBData(InUnparsedJsonObject, InLookupTablePtr)
+	, CategoryName(InCategoryName)
+	, bIsTagAndValue(bInTagAndValueCategory)
+{
+}
+
+FCategorySectionHelper::FCategorySectionHelper(TSharedPtr< FJsonObject > InUnparsedJsonObject, TMap<int32, FText>* InLookupTablePtr, FText InCategoryName, bool bInTagAndValueCategory, FCategorySectionHelperCallback InSpecialHandlingCallback)
+	: FImaginaryFiBData(InUnparsedJsonObject, InLookupTablePtr)
+	, SpecialHandlingCallback(InSpecialHandlingCallback)
+	, CategoryName(InCategoryName)
+	, bIsTagAndValue(bInTagAndValueCategory)
+{
+
+}
+
+bool FCategorySectionHelper::CanCallFilter(ESearchQueryFilter InSearchQueryFilter)
+{
+	return true;
+}
+
+FSearchResult FCategorySectionHelper::CreateSearchResult_Internal(FSearchResult InParent) const
+{
+	return FSearchResult(new FFindInBlueprintsResult(CategoryName, InParent));
+}
+
+void FCategorySectionHelper::ParseAllChildData()
+{
+	if (UnparsedJsonObject.IsValid() && SpecialHandlingCallback.IsBound())
+	{
+		SpecialHandlingCallback.Execute(UnparsedJsonObject, ParsedChildData);
+		UnparsedJsonObject.Reset();
+	}
+	else
+	{
+		FImaginaryFiBData::ParseAllChildData();
+	}
+}
+
+//////////////////////////////////////////
+// FImaginaryBlueprint
+
+FImaginaryBlueprint::FImaginaryBlueprint(FString InBlueprintName, FString InBlueprintPath, FString InBlueprintParentClass, TArray<FString>& InInterfaces, FString InUnparsedStringData)
+	: BlueprintPath(InBlueprintPath)
+	, UnparsedStringData(InUnparsedStringData)
+{
+	ParseToJson();
+	ParsedTagsAndValues.Add(FFindInBlueprintSearchTags::FiB_Name.ToString(), FSearchableValueInfo(FFindInBlueprintSearchTags::FiB_Name, FText::FromString(InBlueprintName), ESearchableValueStatus::ExplicitySearchable));
+	ParsedTagsAndValues.Add(FFindInBlueprintSearchTags::FiB_Path.ToString(), FSearchableValueInfo(FFindInBlueprintSearchTags::FiB_Path, FText::FromString(InBlueprintPath), ESearchableValueStatus::ExplicitySearchable));
+	ParsedTagsAndValues.Add(FFindInBlueprintSearchTags::FiB_ParentClass.ToString(), FSearchableValueInfo(FFindInBlueprintSearchTags::FiB_ParentClass, FText::FromString(InBlueprintParentClass), ESearchableValueStatus::ExplicitySearchable));
+
+	TSharedPtr< FCategorySectionHelper > InterfaceCategory = MakeShareable(new FCategorySectionHelper(nullptr, &LookupTable, FFindInBlueprintSearchTags::FiB_Interfaces, true));
+	for( int32 InterfaceIdx = 0; InterfaceIdx < InInterfaces.Num(); ++InterfaceIdx)
+	{
+		FString& Interface = InInterfaces[InterfaceIdx];
+		FText Key = FText::FromString(FString::FromInt(InterfaceIdx));
+		FSearchableValueInfo Value(Key, FText::FromString(Interface), ESearchableValueStatus::ExplicitySearchable);
+		InterfaceCategory->AddKeyValuePair(FFindInBlueprintSearchTags::FiB_Interfaces, Value);
+	}		
+	ParsedChildData.Add(InterfaceCategory);
+}
+
+FSearchResult FImaginaryBlueprint::CreateSearchResult_Internal(FSearchResult InParent) const
+{
+	return FSearchResult(new FFindInBlueprintsResult(ParsedTagsAndValues.Find(FFindInBlueprintSearchTags::FiB_Name.ToString())->GetDisplayText(LookupTable)));
+}
+
+bool FImaginaryBlueprint::IsCompatibleWithFilter(ESearchQueryFilter InSearchQueryFilter)
+{
+	return InSearchQueryFilter == ESearchQueryFilter::AllFilter || InSearchQueryFilter == ESearchQueryFilter::BlueprintFilter;
+}
+
+bool FImaginaryBlueprint::CanCallFilter(ESearchQueryFilter InSearchQueryFilter)
+{
+	return InSearchQueryFilter == ESearchQueryFilter::NodesFilter ||
+		InSearchQueryFilter == ESearchQueryFilter::PinsFilter ||
+		InSearchQueryFilter == ESearchQueryFilter::GraphsFilter ||
+		InSearchQueryFilter == ESearchQueryFilter::UberGraphsFilter ||
+		InSearchQueryFilter == ESearchQueryFilter::FunctionsFilter ||
+		InSearchQueryFilter == ESearchQueryFilter::MacrosFilter ||
+		InSearchQueryFilter == ESearchQueryFilter::PropertiesFilter ||
+		InSearchQueryFilter == ESearchQueryFilter::VariablesFilter ||
+		InSearchQueryFilter == ESearchQueryFilter::ComponentsFilter ||
+		FImaginaryFiBData::CanCallFilter(InSearchQueryFilter);
+}
+
+void FImaginaryBlueprint::ParseToJson()
+{
+	UnparsedJsonObject = FFindInBlueprintSearchManager::ConvertJsonStringToObject(UnparsedStringData, LookupTable);
+}
+
+bool FImaginaryBlueprint::TrySpecialHandleJsonValue(FString InKey, TSharedPtr< FJsonValue > InJsonValue)
+{
+	bool bResult = false;
+
+	if(InKey == FFindInBlueprintSearchTags::FiB_Properties.ToString())
+	{
+		// Pulls out all properties (variables) for this Blueprint
+		TArray<TSharedPtr< FJsonValue > > PropertyList = InJsonValue->AsArray();
+		for( TSharedPtr< FJsonValue > PropertyValue : PropertyList )
+		{
+			ParsedChildData.Add(MakeShareable(new FImaginaryProperty(PropertyValue->AsObject(), &LookupTable)));
+		}
+		bResult = true;
+	}
+	else if (InKey == FFindInBlueprintSearchTags::FiB_Functions.ToString())
+	{
+		ParseGraph(InJsonValue, FFindInBlueprintSearchTags::FiB_Functions.ToString(), GT_Function);
+		bResult = true;
+	}
+	else if (InKey == FFindInBlueprintSearchTags::FiB_Macros.ToString())
+	{
+		ParseGraph(InJsonValue, FFindInBlueprintSearchTags::FiB_Macros.ToString(), GT_Macro);
+		bResult = true;
+	}
+	else if (InKey == FFindInBlueprintSearchTags::FiB_UberGraphs.ToString())
+	{
+		ParseGraph(InJsonValue, FFindInBlueprintSearchTags::FiB_UberGraphs.ToString(), GT_Ubergraph);
+		bResult = true;
+	}
+	else if (InKey == FFindInBlueprintSearchTags::FiB_SubGraphs.ToString())
+	{
+		ParseGraph(InJsonValue, FFindInBlueprintSearchTags::FiB_SubGraphs.ToString(), GT_Ubergraph);
+		bResult = true;
+	}
+	else if(InKey == FFindInBlueprintSearchTags::FiB_Components.ToString())
+	{
+		TArray<TSharedPtr< FJsonValue > > ComponentList = InJsonValue->AsArray();
+		TSharedPtr< FJsonObject > ComponentsWrapperObject(new FJsonObject);
+		ComponentsWrapperObject->Values.Add(FFindInBlueprintSearchTags::FiB_Components.ToString(), InJsonValue);
+		ParsedChildData.Add(MakeShareable(new FCategorySectionHelper(ComponentsWrapperObject, &LookupTable, FFindInBlueprintSearchTags::FiB_Components, false, FCategorySectionHelper::FCategorySectionHelperCallback::CreateRaw(this, &FImaginaryBlueprint::ParseComponents))));
+		bResult = true;
+	}
+
+	if (!bResult)
+	{
+		bResult = FImaginaryFiBData::TrySpecialHandleJsonValue(InKey, InJsonValue);
+	}
+	return bResult;
+}
+
+void FImaginaryBlueprint::ParseGraph( TSharedPtr< FJsonValue > InJsonValue, FString InCategoryTitle, EGraphType InGraphType )
+{
+	TArray<TSharedPtr< FJsonValue > > GraphList = InJsonValue->AsArray();
+	for( TSharedPtr< FJsonValue > GraphValue : GraphList )
+	{
+		ParsedChildData.Add(MakeShareable(new FImaginaryGraph(GraphValue->AsObject(), &LookupTable, InGraphType)));
+	}
+}
+
+void FImaginaryBlueprint::ParseComponents(TSharedPtr< FJsonObject > InJsonObject, TArray< TSharedPtr<FImaginaryFiBData> >& OutParsedChildData)
+{
+	// Pulls out all properties (variables) for this Blueprint
+	TArray<TSharedPtr< FJsonValue > > ComponentList = InJsonObject->GetArrayField(FFindInBlueprintSearchTags::FiB_Components.ToString());
+	for( TSharedPtr< FJsonValue > ComponentValue : ComponentList )
+	{
+		OutParsedChildData.Add(MakeShareable(new FImaginaryComponent(ComponentValue->AsObject(), &LookupTable)));
+	}
+}
+
+//////////////////////////
+// FImaginaryGraph
+
+FImaginaryGraph::FImaginaryGraph(TSharedPtr< FJsonObject > InUnparsedJsonObject, TMap<int32, FText>* InLookupTablePtr, EGraphType InGraphType)
+	: FImaginaryFiBData(InUnparsedJsonObject, InLookupTablePtr)
+	, GraphType(InGraphType)
+{
+}
+
+FSearchResult FImaginaryGraph::CreateSearchResult_Internal(FSearchResult InParent) const
+{
+	return FSearchResult(new FFindInBlueprintsGraph(FText::GetEmpty(), InParent, GraphType));
+}
+
+bool FImaginaryGraph::IsCompatibleWithFilter(ESearchQueryFilter InSearchQueryFilter)
+{
+	return InSearchQueryFilter == ESearchQueryFilter::AllFilter || 
+		InSearchQueryFilter == ESearchQueryFilter::GraphsFilter ||
+		(GraphType == GT_Ubergraph && InSearchQueryFilter == ESearchQueryFilter::UberGraphsFilter) ||
+		(GraphType == GT_Function && InSearchQueryFilter == ESearchQueryFilter::FunctionsFilter) ||
+		(GraphType == GT_Macro && InSearchQueryFilter == ESearchQueryFilter::MacrosFilter);
+}
+
+bool FImaginaryGraph::CanCallFilter(ESearchQueryFilter InSearchQueryFilter)
+{
+	return InSearchQueryFilter == ESearchQueryFilter::PinsFilter ||
+		InSearchQueryFilter == ESearchQueryFilter::NodesFilter ||
+		(GraphType == GT_Function && InSearchQueryFilter == ESearchQueryFilter::PropertiesFilter) ||
+		(GraphType == GT_Function && InSearchQueryFilter == ESearchQueryFilter::VariablesFilter) ||
+		FImaginaryFiBData::CanCallFilter(InSearchQueryFilter);
+}
+
+ESearchableValueStatus FImaginaryGraph::GetSearchabilityStatus(FString InKey)
+{
+	// This is a non-ideal way to assign searchability vs being a core display item and will be resolved in future versions of the FiB data in the AR
+	if (FindInBlueprintsHelpers::IsTextEqualToString(FFindInBlueprintSearchTags::FiB_Name, InKey)
+		|| FindInBlueprintsHelpers::IsTextEqualToString(FFindInBlueprintSearchTags::FiB_NativeName, InKey)
+		)
+	{
+		return ESearchableValueStatus::CoreDisplayItem;
+	}
+
+	return ESearchableValueStatus::Searchable;
+}
+
+bool FImaginaryGraph::TrySpecialHandleJsonValue(FString InKey, TSharedPtr< FJsonValue > InJsonValue)
+{
+	if (FindInBlueprintsHelpers::IsTextEqualToString(FFindInBlueprintSearchTags::FiB_Nodes, InKey))
+	{
+		TArray< TSharedPtr< FJsonValue > > NodeList = InJsonValue->AsArray();
+
+		for( TSharedPtr< FJsonValue > NodeValue : NodeList )
+		{
+			ParsedChildData.Add(MakeShareable(new FImaginaryGraphNode(NodeValue->AsObject(), LookupTablePtr)));
+		}
+		return true;
+	}
+	else if (FindInBlueprintsHelpers::IsTextEqualToString(FFindInBlueprintSearchTags::FiB_Properties, InKey))
+	{
+		// Pulls out all properties (local variables) for this graph
+		TArray<TSharedPtr< FJsonValue > > PropertyList = InJsonValue->AsArray();
+		for( TSharedPtr< FJsonValue > PropertyValue : PropertyList )
+		{
+			ParsedChildData.Add(MakeShareable(new FImaginaryProperty(PropertyValue->AsObject(), LookupTablePtr)));
+		}
+		return true;
+	}
+	return false;
+}
+
+//////////////////////////////////////
+// FImaginaryGraphNode
+
+FImaginaryGraphNode::FImaginaryGraphNode(TSharedPtr< FJsonObject > InUnparsedJsonObject, TMap<int32, FText>* InLookupTablePtr)
+	: FImaginaryFiBData(InUnparsedJsonObject, InLookupTablePtr)
+{
+}
+
+FSearchResult FImaginaryGraphNode::CreateSearchResult_Internal(FSearchResult InParent) const
+{
+	return FSearchResult(new FFindInBlueprintsGraphNode(FText::GetEmpty(), InParent));
+}
+
+bool FImaginaryGraphNode::IsCompatibleWithFilter(ESearchQueryFilter InSearchQueryFilter)
+{
+	return InSearchQueryFilter == ESearchQueryFilter::AllFilter || InSearchQueryFilter == ESearchQueryFilter::NodesFilter;
+}
+
+bool FImaginaryGraphNode::CanCallFilter(ESearchQueryFilter InSearchQueryFilter)
+{
+	return InSearchQueryFilter == ESearchQueryFilter::PinsFilter ||
+		FImaginaryFiBData::CanCallFilter(InSearchQueryFilter);
+}
+
+ESearchableValueStatus FImaginaryGraphNode::GetSearchabilityStatus(FString InKey)
+{
+	// This is a non-ideal way to assign searchability vs being a core display item and will be resolved in future versions of the FiB data in the AR
+	if (FindInBlueprintsHelpers::IsTextEqualToString(FFindInBlueprintSearchTags::FiB_Name, InKey)
+		|| FindInBlueprintsHelpers::IsTextEqualToString(FFindInBlueprintSearchTags::FiB_NativeName, InKey)
+		|| FindInBlueprintsHelpers::IsTextEqualToString(FFindInBlueprintSearchTags::FiB_Comment, InKey)
+		)
+	{
+		return ESearchableValueStatus::CoreDisplayItem;
+	}
+	else if (FindInBlueprintsHelpers::IsTextEqualToString(FFindInBlueprintSearchTags::FiB_Glyph, InKey)
+		|| FindInBlueprintsHelpers::IsTextEqualToString(FFindInBlueprintSearchTags::FiB_GlyphColor, InKey)
+		|| FindInBlueprintsHelpers::IsTextEqualToString(FFindInBlueprintSearchTags::FiB_NodeGuid, InKey)
+		)
+	{
+		return ESearchableValueStatus::NotSearchable;
+	}
+	else if (FindInBlueprintsHelpers::IsTextEqualToString(FFindInBlueprintSearchTags::FiB_ClassName, InKey))
+	{
+		return ESearchableValueStatus::ExplicitySearchable;
+	}
+	return ESearchableValueStatus::Searchable;
+}
+
+bool FImaginaryGraphNode::TrySpecialHandleJsonValue(FString InKey, TSharedPtr< FJsonValue > InJsonValue)
+{
+	if (FindInBlueprintsHelpers::IsTextEqualToString(FFindInBlueprintSearchTags::FiB_Pins, InKey))
+	{
+		TArray< TSharedPtr< FJsonValue > > PinsList = InJsonValue->AsArray();
+
+		for( TSharedPtr< FJsonValue > Pin : PinsList )
+		{
+			ParsedChildData.Add(MakeShareable(new FImaginaryPin(Pin->AsObject(), LookupTablePtr, SchemaName)));
+		}
+		return true;
+	}
+	else if (FindInBlueprintsHelpers::IsTextEqualToString(FFindInBlueprintSearchTags::FiB_SchemaName, InKey))
+	{
+		// Previously extracted
+		return true;
+	}
+
+	return false;
+}
+
+void FImaginaryGraphNode::ParseAllChildData()
+{
+	if (UnparsedJsonObject.IsValid())
+	{
+		TSharedPtr< FJsonObject > JsonObject = UnparsedJsonObject;
+		// Very important to get the schema first, other bits of data depend on it
+		TSharedPtr< FJsonValue > SchemaNameValue = JsonObject->GetField< EJson::String >(FFindInBlueprintSearchTags::FiB_SchemaName.ToString());
+		if(SchemaNameValue.IsValid())
+		{
+			SchemaName = FindInBlueprintsHelpers::AsFText(SchemaNameValue, *LookupTablePtr).ToString();
+		}
+
+		FImaginaryFiBData::ParseAllChildData();
+	}
+}
+
+///////////////////////////////////////////
+// FImaginaryProperty
+
+FImaginaryProperty::FImaginaryProperty(TSharedPtr< FJsonObject > InUnparsedJsonObject, TMap<int32, FText>* InLookupTablePtr)
+	: FImaginaryFiBData(InUnparsedJsonObject, InLookupTablePtr)
+{
+}
+
+bool FImaginaryProperty::IsCompatibleWithFilter(ESearchQueryFilter InSearchQueryFilter)
+{
+	return InSearchQueryFilter == ESearchQueryFilter::AllFilter || 
+		InSearchQueryFilter == ESearchQueryFilter::PropertiesFilter || 
+		InSearchQueryFilter == ESearchQueryFilter::VariablesFilter;
+}
+
+FSearchResult FImaginaryProperty::CreateSearchResult_Internal(FSearchResult InParent) const
+{
+	return FSearchResult(new FFindInBlueprintsProperty(FText::GetEmpty(), InParent));
+}
+
+ESearchableValueStatus FImaginaryProperty::GetSearchabilityStatus(FString InKey)
+{
+	// This is a non-ideal way to assign searchability vs being a core display item and will be resolved in future versions of the FiB data in the AR
+	if (FindInBlueprintsHelpers::IsTextEqualToString(FFindInBlueprintSearchTags::FiB_Name, InKey)
+		|| FindInBlueprintsHelpers::IsTextEqualToString(FFindInBlueprintSearchTags::FiB_NativeName, InKey)
+		)
+	{
+		return ESearchableValueStatus::CoreDisplayItem;
+	}
+	else if (FindInBlueprintsHelpers::IsTextEqualToString(FFindInBlueprintSearchTags::FiB_PinCategory, InKey)
+		|| FindInBlueprintsHelpers::IsTextEqualToString(FFindInBlueprintSearchTags::FiB_PinSubCategory, InKey)
+		|| FindInBlueprintsHelpers::IsTextEqualToString(FFindInBlueprintSearchTags::FiB_ObjectClass, InKey)
+		|| FindInBlueprintsHelpers::IsTextEqualToString(FFindInBlueprintSearchTags::FiB_IsArray, InKey)
+		|| FindInBlueprintsHelpers::IsTextEqualToString(FFindInBlueprintSearchTags::FiB_IsReference, InKey)
+		|| FindInBlueprintsHelpers::IsTextEqualToString(FFindInBlueprintSearchTags::FiB_IsSCSComponent, InKey)
+		)
+	{
+		return ESearchableValueStatus::ExplicitySearchableHidden;
+	}
+	return ESearchableValueStatus::Searchable;
+}
+
+//////////////////////////////
+// FImaginaryComponent
+
+FImaginaryComponent::FImaginaryComponent(TSharedPtr< FJsonObject > InUnparsedJsonObject, TMap<int32, FText>* InLookupTablePtr)
+	: FImaginaryProperty(InUnparsedJsonObject, InLookupTablePtr)
+{
+}
+
+bool FImaginaryComponent::IsCompatibleWithFilter(ESearchQueryFilter InSearchQueryFilter)
+{
+	return FImaginaryProperty::IsCompatibleWithFilter(InSearchQueryFilter) || InSearchQueryFilter == ESearchQueryFilter::ComponentsFilter;
+}
+
+//////////////////////////////
+// FImaginaryPin
+
+FImaginaryPin::FImaginaryPin(TSharedPtr< FJsonObject > InUnparsedJsonObject, TMap<int32, FText>* InLookupTablePtr, FString InSchemaName)
+	: FImaginaryFiBData(InUnparsedJsonObject, InLookupTablePtr)
+	, SchemaName(InSchemaName)
+{
+}
+
+bool FImaginaryPin::IsCompatibleWithFilter(ESearchQueryFilter InSearchQueryFilter)
+{
+	return InSearchQueryFilter == ESearchQueryFilter::AllFilter || InSearchQueryFilter == ESearchQueryFilter::PinsFilter;
+}
+
+FSearchResult FImaginaryPin::CreateSearchResult_Internal(FSearchResult InParent) const
+{
+	return FSearchResult(new FFindInBlueprintsPin(FText::GetEmpty(), InParent, SchemaName));
+}
+
+ESearchableValueStatus FImaginaryPin::GetSearchabilityStatus(FString InKey)
+{
+	if (FindInBlueprintsHelpers::IsTextEqualToString(FFindInBlueprintSearchTags::FiB_Name, InKey)
+		|| FindInBlueprintsHelpers::IsTextEqualToString(FFindInBlueprintSearchTags::FiB_NativeName, InKey)
+		)
+	{
+		return ESearchableValueStatus::CoreDisplayItem;
+	}
+	else if (FindInBlueprintsHelpers::IsTextEqualToString(FFindInBlueprintSearchTags::FiB_PinCategory, InKey)
+		|| FindInBlueprintsHelpers::IsTextEqualToString(FFindInBlueprintSearchTags::FiB_PinSubCategory, InKey)
+		|| FindInBlueprintsHelpers::IsTextEqualToString(FFindInBlueprintSearchTags::FiB_ObjectClass, InKey)
+		|| FindInBlueprintsHelpers::IsTextEqualToString(FFindInBlueprintSearchTags::FiB_IsArray, InKey)
+		|| FindInBlueprintsHelpers::IsTextEqualToString(FFindInBlueprintSearchTags::FiB_IsReference, InKey)
+		|| FindInBlueprintsHelpers::IsTextEqualToString(FFindInBlueprintSearchTags::FiB_IsSCSComponent, InKey)
+		)
+	{
+		return ESearchableValueStatus::ExplicitySearchableHidden;
+	}
+	return ESearchableValueStatus::Searchable;
+}
+
+bool FImaginaryPin::TrySpecialHandleJsonValue(FString InKey, TSharedPtr< FJsonValue > InJsonValue)
+{
+	return false;
+}
+
+#undef LOCTEXT_NAMESPACE
\ No newline at end of file
diff --git a/Engine/Source/Editor/Kismet/Public/FindInBlueprintManager.h b/Engine/Source/Editor/Kismet/Public/FindInBlueprintManager.h
--- a/Engine/Source/Editor/Kismet/Public/FindInBlueprintManager.h
+++ b/Engine/Source/Editor/Kismet/Public/FindInBlueprintManager.h
@@ -45,6 +45,12 @@ struct KISMET_API FFindInBlueprintSearchTags
 	static const FText FiB_Description;
 	/** Comment tag */
 	static const FText FiB_Comment;
+	/** Path tag */
+	static const FText FiB_Path;
+	/** Parent Class tag */
+	static const FText FiB_ParentClass;
+	/** Interfaces tag */
+	static const FText FiB_Interfaces;
 
 	/** Pin type tags */
 
@@ -80,6 +86,12 @@ struct FSearchData
 	/** Search data block for the Blueprint */
 	FString Value;
 
+	/** Parent Class */
+	FString ParentClass;
+
+	/** Interfaces implemented by the Blueprint */
+	TArray<FString> Interfaces;
+
 	/** Cached to determine if the Blueprint is seen as no longer valid, allows it to be cleared out next save to disk */
 	bool bMarkedForDeletion;
 
@@ -95,6 +107,8 @@ struct FSearchData
 		: Blueprint(Other.Blueprint)
 		, BlueprintPath(MoveTemp(Other.BlueprintPath))
 		, Value(MoveTemp(Other.Value))
+		, ParentClass(MoveTemp(Other.ParentClass))
+		, Interfaces(MoveTemp(Other.Interfaces))
 		, bMarkedForDeletion(Other.bMarkedForDeletion)
 	{
 	}
@@ -110,6 +124,8 @@ struct FSearchData
 		BlueprintPath = MoveTemp(RHS.BlueprintPath);
 		Value = MoveTemp(RHS.Value);
 		bMarkedForDeletion = RHS.bMarkedForDeletion;
+		ParentClass = MoveTemp(RHS.ParentClass);
+		Interfaces = MoveTemp(RHS.Interfaces);
 		return *this;
 	}
 
@@ -117,6 +133,8 @@ struct FSearchData
 		: Blueprint(Other.Blueprint)
 		, BlueprintPath(Other.BlueprintPath)
 		, Value(Other.Value)
+		, ParentClass(Other.ParentClass)
+		, Interfaces(Other.Interfaces)
 		, bMarkedForDeletion(Other.bMarkedForDeletion)
 	{
 	}
@@ -132,6 +150,8 @@ struct FSearchData
 		BlueprintPath = RHS.BlueprintPath;
 		Value = RHS.Value;
 		bMarkedForDeletion = RHS.bMarkedForDeletion;
+		ParentClass = RHS.ParentClass;
+		Interfaces = RHS.Interfaces;
 		return *this;
 	}
 
diff --git a/Engine/Source/Editor/Kismet/Public/FindInBlueprints.h b/Engine/Source/Editor/Kismet/Public/FindInBlueprints.h
--- a/Engine/Source/Editor/Kismet/Public/FindInBlueprints.h
+++ b/Engine/Source/Editor/Kismet/Public/FindInBlueprints.h
@@ -5,6 +5,37 @@
 typedef TSharedPtr<class FFindInBlueprintsResult> FSearchResult;
 typedef STreeView<FSearchResult>  STreeViewType;
 
+/** Some utility functions to help with Find-in-Blueprint functionality */
+namespace FindInBlueprintsHelpers
+{
+	/** Looks up a JsonValue's FText from the passed lookup table */
+	FText AsFText(TSharedPtr< FJsonValue > InJsonValue, const TMap<int32, FText>& InLookupTable);
+
+	/** Looks up a JsonValue's FText from the passed lookup table */
+	FText AsFText(int32 InValue, const TMap<int32, FText>& InLookupTable);
+
+	bool IsTextEqualToString(FText InText, FString InString);
+
+	/**
+	 * Retrieves the pin type as a string value
+	 *
+	 * @param InPinType		The pin type to look at
+	 *
+	 * @return				The pin type as a string in format [category]'[sub-category object]'
+	 */
+	FString GetPinTypeAsString(const FEdGraphPinType& InPinType);
+
+	/**
+	 * Parses a pin type from passed in key names and values
+	 *
+	 * @param InKey					The key name for what the data should be translated as
+	 * @param InValue				Value to be be translated
+	 * @param InOutPinType			Modifies the PinType based on the passed parameters, building it up over multiple calls
+	 * @return						TRUE when the parsing is successful
+	 */
+	bool ParsePinType(FText InKey, FText InValue, FEdGraphPinType& InOutPinType);
+}
+
 /* Item that matched the search results */
 class FFindInBlueprintsResult : public TSharedFromThis< FFindInBlueprintsResult >
 {
@@ -35,27 +66,6 @@ public:
 	/** gets the blueprint housing all these search results */
 	UBlueprint* GetParentBlueprint() const;
 
-	/**
-	 * Extracts the content of a JsonNode, building the sub-hierarchy for display
-	 *
-	 * @param InLookupTable		Lookup table for converting Json strings into an FText
-	 * @param InJsonNode		Node to extract data from
-	 * @param InTokens			Tokens to match searchable content against
-	 * @return					Returns TRUE if the content matches search results and was extracted
-	 */
-	virtual bool ExtractContent(const TMap<int32, FText>& InLookupTable, TSharedPtr< FJsonObject > InJsonNode, const TArray<FString>& InTokens) { return ExtractContent(InLookupTable, InJsonNode, InTokens, NULL); };
-
-	/**
-	 * Extracts the content of a JsonNode, building the sub-hierarchy for display
-	 *
-	 * @param InLookupTable		Lookup table for converting Json strings into an FText
-	 * @param InJsonNode		Node to extract data from
-	 * @param InTokens			Tokens to match searchable content against
-	 * @param InParentOverride	Parent search result to attach sub-content to
-	 * @return					Returns TRUE if the content matches search results and was extracted
-	 */
-	virtual bool ExtractContent(const TMap<int32, FText>& InLookupTable, TSharedPtr< FJsonObject > InJsonNode, const TArray<FString>& InTokens, TSharedPtr< FFindInBlueprintsResult > InParentOverride);
-
 	/**
 	 * Parses search info for specific data important for displaying the search result in an easy to understand format
 	 *
@@ -63,10 +73,8 @@ public:
 	 * @param	InKey			This is the tag for the data, describing what it is so special handling can occur if needed
 	 * @param	InValue			Compared against search query to see if it passes the filter, sometimes data is rejected because it is deemed unsearchable
 	 * @param	InParent		The parent search result
-	 *
-	 * @return					TRUE if the item passes the search filter or is otherwise accepted
 	 */
-	virtual bool ParseSearchInfo(const TArray<FString> &InTokens, FText InKey, FText InValue, TSharedPtr< FFindInBlueprintsResult > InParent);
+	virtual void ParseSearchInfo(FText InKey, FText InValue) {};
 
 	/**
 	 * Adds extra search info, anything that does not have a predestined place in the search result. Adds a sub-item to the searches and formats its description so the tag displays
@@ -85,19 +93,6 @@ public:
 	/** Returns the display string for the row */
 	FText GetDisplayString() const;
 
-protected:
-
-	/**
-	 * Extracts the details of a JsonValue in the search results, pushing arrays into sub-objects correctly, breaking up objects, and passing on the rest of the types
-	 *
-	 * @param	InLookupTable			Lookup table for converting Json strings into an FText
-	 * @param	InTokens				The search tokens to check results against
-	 * @param	InKey					This is the tag for the data, describing what it is so special handling can occur if needed
-	 * @param	InValue					Has all the data for this object that needs to be deciphered or broken up into base types to be deciphered
-	 * @param	InParentOverride		The parent to attach all children data to, can be NULL.
-	 */
-	bool ExtractJsonValue(const TMap<int32, FText>& InLookupTable, const TArray<FString>& InTokens, FText InKey, TSharedPtr< FJsonValue > InValue, TSharedPtr< FFindInBlueprintsResult > InParentOverride);
-
 public:
 	/*Any children listed under this category */
 	TArray< TSharedPtr<FFindInBlueprintsResult> > Children;
@@ -122,8 +117,7 @@ public:
 	/** FFindInBlueprintsResult Interface */
 	virtual FReply OnClick() override;
 	virtual TSharedRef<SWidget>	CreateIcon() const override;
-	virtual bool ParseSearchInfo(const TArray<FString> &InTokens, FText InKey, FText InValue, TSharedPtr< FFindInBlueprintsResult > InParentOverride = NULL) override;
-	virtual bool ExtractContent(const TMap<int32, FText>& InLookupTable, TSharedPtr< FJsonObject > InJsonNode, const TArray<FString>& InTokens) override;
+	virtual void ParseSearchInfo(FText InKey, FText InValue) override;
 	virtual FText GetCategory() const override;
 	virtual void FinalizeSearchData() override;
 	/** End FFindInBlueprintsResult Interface */
@@ -154,7 +148,7 @@ public:
 
 	/** FFindInBlueprintsResult Interface */
 	virtual TSharedRef<SWidget>	CreateIcon() const override;
-	virtual bool ParseSearchInfo(const TArray<FString> &InTokens, FText InKey, FText InValue, TSharedPtr< FFindInBlueprintsResult > InParentOverride = NULL) override;
+	virtual void ParseSearchInfo(FText InKey, FText InValue) override;
 	virtual FText GetCategory() const override;
 	virtual void FinalizeSearchData() override;
 	/** End FFindInBlueprintsResult Interface */
@@ -165,9 +159,6 @@ protected:
 	/** The pin that this search result refers to */
 	FEdGraphPinType PinType;
 
-	/** The Pin's sub-category name */
-	FString PinSubCategory;
-
 	/** Pin's icon color */
 	FSlateColor IconColor;
 };
@@ -181,7 +172,7 @@ public:
 
 	/** FFindInBlueprintsResult Interface */
 	virtual TSharedRef<SWidget>	CreateIcon() const override;
-	virtual bool ParseSearchInfo(const TArray<FString> &InTokens, FText InKey, FText InValue, TSharedPtr< FFindInBlueprintsResult > InParentOverride = NULL) override;
+	virtual void ParseSearchInfo(FText InKey, FText InValue) override;
 	virtual FText GetCategory() const override;
 	virtual void FinalizeSearchData() override;
 	/** End FFindInBlueprintsResult Interface */
@@ -192,9 +183,6 @@ protected:
 	/** The default value of a property as a string */
 	FString DefaultValue;
 
-	/** The Property's sub-category name, for the pin display type */
-	FString PinSubCategory;
-
 	/** TRUE if the property is an SCS_Component */
 	bool bIsSCSComponent;
 };
@@ -209,9 +197,8 @@ public:
 	/** FFindInBlueprintsResult Interface */
 	virtual FReply OnClick() override;
 	virtual TSharedRef<SWidget>	CreateIcon() const override;
-	virtual bool ParseSearchInfo(const TArray<FString> &InTokens, FText InKey, FText InValue, TSharedPtr< FFindInBlueprintsResult > InParentOverride = NULL) override;
+	virtual void ParseSearchInfo(FText InKey, FText InValue) override;
 	virtual FText GetCategory() const override;
-	virtual bool ExtractContent(const TMap<int32, FText>& InLookupTable, TSharedPtr< FJsonObject > InJsonNode, const TArray<FString>& InTokens) override;
 	/** End FFindInBlueprintsResult Interface */
 protected:
 
@@ -265,11 +252,8 @@ private:
 	/** Begins the search based on the SearchValue */
 	void InitiateSearch();
 	
-	/* Find any results that contain all of the tokens, however, only in find within blueprints */
-	void MatchTokensWithinBlueprint(const TArray<FString>& Tokens);
-
 	/** Launches a thread for streaming more content into the results widget */
-	void LaunchStreamThread(const TArray<FString>& InTokens);
+	void LaunchStreamThread(const FString& InSearchValue);
 
 	/** Returns the percent complete on the search for the progress bar */
 	TOptional<float> GetPercentCompleteSearch() const;
diff --git a/Engine/Source/Editor/Kismet/Public/ImaginaryBlueprintData.h b/Engine/Source/Editor/Kismet/Public/ImaginaryBlueprintData.h
new file mode 100644
--- /dev/null
+++ b/Engine/Source/Editor/Kismet/Public/ImaginaryBlueprintData.h
@@ -0,0 +1,368 @@
+// Copyright 1998-2015 Epic Games, Inc. All Rights Reserved.
+#pragma once
+
+#include "TextFilterExpressionEvaluator.h"
+#include "FindInBlueprints.h"
+
+/** Filters are used by functions for searching to decide whether items can call certain functions or match the requirements of a function */
+enum ESearchQueryFilter
+{
+	BlueprintFilter = 0,
+	GraphsFilter,
+	UberGraphsFilter,
+	FunctionsFilter,
+	MacrosFilter,
+	NodesFilter,
+	PinsFilter,
+	PropertiesFilter,
+	VariablesFilter,
+	ComponentsFilter,
+	AllFilter, // Will search all items, when used inside of another filter it will search all sub-items of that filter
+};
+
+enum ESearchableValueStatus
+{
+	NotSearchable = 0x00000000, // Cannot search this value, it is used for display purposes only
+	Searchable = 0x00000001, // Generically searchable, value will appear as a sub-item and has no sub-data
+	Hidden = 0x00000002, // Item will not be shown
+	Explicit = 0x00000004, // Item must be explicitly requested via the tag
+
+	
+	CoreDisplayItem = Hidden | Searchable, // Core display items are searchable but should not display as sub-items because their data is presented in another fashion.
+	ExplicitySearchable = Explicit | Searchable, // Will only be allowed to be found if searching using a tag
+	ExplicitySearchableHidden = Explicit | Searchable | Hidden, // Will only be allowed to be found if searching using a tag but will not display the tag in the results (because it is a CoreDisplayItem)
+	AllSearchable = CoreDisplayItem | ExplicitySearchable, // Covers all searchability methods
+};
+
+class FSearchableValueInfo
+{
+public:
+	FSearchableValueInfo(FText InDisplayKey, int32 InLookUpTableKey)
+		: SearchableValueStatus(ESearchableValueStatus::Searchable)
+		, DisplayKey(InDisplayKey)
+		, LookupTableKey(InLookUpTableKey)
+	{
+	}
+
+	FSearchableValueInfo(FText InDisplayKey, FText InDisplayText)
+		: SearchableValueStatus(ESearchableValueStatus::Searchable)
+		, DisplayKey(InDisplayKey)
+		, LookupTableKey(-1)
+		, DisplayText(InDisplayText)
+	{
+
+	}
+
+	FSearchableValueInfo(FText InDisplayKey, int32 InLookUpTableKey, ESearchableValueStatus InSearchableValueStatus)
+		: SearchableValueStatus(InSearchableValueStatus)
+		, DisplayKey(InDisplayKey)
+		, LookupTableKey(InLookUpTableKey)
+	{
+
+	}
+
+	FSearchableValueInfo(FText InDisplayKey, FText InDisplayText, ESearchableValueStatus InSearchableValueStatus)
+		: SearchableValueStatus(InSearchableValueStatus)
+		, DisplayKey(InDisplayKey)
+		, LookupTableKey(-1)
+		, DisplayText(InDisplayText)
+	{
+
+	}
+
+	/** Returns TRUE if the data is searchable */
+	bool IsSearchable() const { return (SearchableValueStatus & ESearchableValueStatus::Searchable) != 0; }
+
+	/** Returns TRUE if the item should be treated as a CoreDisplayItem, which is searchable but not displayed */
+	bool IsCoreDisplay() const { return (SearchableValueStatus & ESearchableValueStatus::CoreDisplayItem) == ESearchableValueStatus::CoreDisplayItem; }
+
+	/** Returns TRUE if the item should only be searchable if explicitly searched for using the tag */
+	bool IsExplicitSearchable() { return (SearchableValueStatus & ESearchableValueStatus::ExplicitySearchable) == ESearchableValueStatus::ExplicitySearchable; }
+
+	/** Returns the display text to use for this item */
+	FText GetDisplayText(const TMap<int32, FText>& InLookupTable) const;
+
+	/** Returns the display key for this item */
+	FText GetDisplayKey() const
+	{
+		return DisplayKey;
+	}
+
+protected:
+	/** The searchability status of this item */
+	ESearchableValueStatus SearchableValueStatus;
+
+	/** Key that this item is associated with, used for display purposes */
+	FText DisplayKey;
+
+	/** Key to use to lookup into a table if DisplayText does not override */
+	int32 LookupTableKey;
+
+	/** Text value to use instead of a lookup into the table */
+	FText DisplayText;
+};
+
+/** Struct to contain search results and help compare them for uniqueness. */
+struct FComponentUniqueDisplay
+{
+	FComponentUniqueDisplay( FSearchResult InSearchResult )
+		: SearchResult(InSearchResult)
+	{}
+
+	bool operator==(const FComponentUniqueDisplay& Other);
+
+	/** Search result contained and used for comparing of uniqueness */
+	FSearchResult SearchResult;
+};
+
+class FImaginaryFiBData : public ITextFilterExpressionContext
+{
+public:
+	FImaginaryFiBData();
+	FImaginaryFiBData(TSharedPtr< FJsonObject > InUnparsedJsonObject, TMap<int32, FText>* InLookupTablePtr);
+
+	/** ITextFilterExpressionContext Interface */
+	// We don't actually use these overrides, see FFiBContextHelper for how we call the alternate functions. These will assert if they are accidentally called.
+	virtual bool TestBasicStringExpression(const FTextFilterString& InValue, const ETextFilterTextComparisonMode InTextComparisonMode) const override { ensure(0); return false; };
+	virtual bool TestComplexExpression(const FName& InKey, const FTextFilterString& InValue, const ETextFilterComparisonOperation InComparisonOperation, const ETextFilterTextComparisonMode InTextComparisonMode) const override { ensure(0); return false; };
+	/** End ITextFilterExpressionContext Interface */
+
+	/** Returns TRUE if this item is a category type, which helps to organize child data */
+	virtual bool IsCategory() { return false; }
+
+	/** Returns TRUE if this item is considered a Tag and Value category, where it's contents should be considered no different than the parent owner */
+	virtual bool IsTagAndValueCategory() { return false; }
+
+	/** Checks if the filter is compatible with the current object, returns TRUE by default */
+	virtual bool IsCompatibleWithFilter(ESearchQueryFilter InSearchQueryFilter);
+
+	/** Checks if the filter can call functions for the passed filter, returns FALSE by default if the filter is not the AllFilter */
+	virtual bool CanCallFilter(ESearchQueryFilter InSearchQueryFilter);
+
+	/** Parses all child data, non-recursively, so children will be left in an unparsed Json state */
+	virtual void ParseAllChildData();
+
+	/** Test the given value against the strings extracted from the current item. Will return the matching search components if any (can return TRUE without having any if the search components are hidden) */
+	virtual bool TestBasicStringExpression(const FTextFilterString& InValue, const ETextFilterTextComparisonMode InTextComparisonMode, TMultiMap< const FImaginaryFiBData*, FComponentUniqueDisplay >& InOutMatchingSearchComponents) const;
+
+	/** Perform a complex expression test for the current item. Will return the matching search components if any (can return TRUE without having any if the search components are hidden) */
+	virtual bool TestComplexExpression(const FName& InKey, const FTextFilterString& InValue, const ETextFilterComparisonOperation InComparisonOperation, const ETextFilterTextComparisonMode InTextComparisonMode, TMultiMap< const FImaginaryFiBData*, FComponentUniqueDisplay >& InOutMatchingSearchComponents) const;
+
+	/** Requests internal creation of the search result and properly initializes the visual representation of the result */
+	FSearchResult CreateSearchResult(FSearchResult InParent);
+
+	/** Accessor for the parsed child data for this item */
+	const TArray< TSharedPtr<FImaginaryFiBData> >& GetAllParsedChildData() const
+	{
+		return ParsedChildData;
+	}
+
+	void AddKeyValuePair(FText InKey, FSearchableValueInfo& InValue)
+	{
+		ParsedTagsAndValues.Add(InKey.ToString(), InValue);
+	}
+
+	static FSearchResult CreateSearchTree(FSearchResult InParentSearchResult, TWeakPtr< FImaginaryFiBData > InCurrentPointer, TArray< const FImaginaryFiBData* >& InValidSearchResults, TMultiMap< const FImaginaryFiBData*, FComponentUniqueDisplay >& InMatchingSearchComponents);
+
+protected:
+	/**
+	 * Checks if the Key has any special handling to be done, such as making a Pin out of it
+	 *
+	 * @param InKey			Key that the JsonValue was stored under
+	 * @param InJsonValue	JsonValue to be specially parsed
+	 * @return				TRUE if the JsonValue was specially handled, will not be further handled
+	 */
+	virtual bool TrySpecialHandleJsonValue(FString InKey, TSharedPtr< FJsonValue > InJsonValue) { return false; }
+
+	/** Returns the searchability status of a passed in Key, all Keys are searchable by default */
+	virtual ESearchableValueStatus GetSearchabilityStatus(FString InKey) { return ESearchableValueStatus::Searchable; };
+
+	/** Protected internal function which builds the search result for this item */
+	virtual FSearchResult CreateSearchResult_Internal(FSearchResult InParent) const = 0;
+
+	/** Creates a display string for this item in search results */
+	FText CreateSearchComponentDisplayText(FText InKey, FText InValue) const;
+
+	/** Helper function for parsing Json values into usable properties */
+	void ParseJsonValue(FText InKey, FText InDisplayKey, TSharedPtr< FJsonValue > InJsonValue);
+
+protected:
+	/** The unparsed Json object representing this item. Auto-cleared after parsing */
+	TSharedPtr< FJsonObject > UnparsedJsonObject;
+
+	/** All parsed child data for this item */
+	TArray< TSharedPtr<FImaginaryFiBData> > ParsedChildData;
+
+	/** A mapping of tags to their values and searchability status */
+	TMultiMap< FString, FSearchableValueInfo > ParsedTagsAndValues;
+
+	/** Pointer to the lookup table to decompressed the Json strings back into fully formed FTexts */
+	TMap<int32, FText>* LookupTablePtr;
+};
+
+class FCategorySectionHelper : public FImaginaryFiBData
+{
+public:
+	/** Callback declaration for handling special parsing of the items in the category */
+	DECLARE_DELEGATE_TwoParams(FCategorySectionHelperCallback, TSharedPtr< FJsonObject >, TArray< TSharedPtr<FImaginaryFiBData> >&);
+
+	FCategorySectionHelper(TMap<int32, FText>* InLookupTablePtr, FText InCategoryName, bool bInTagAndValueCategory);
+	FCategorySectionHelper(TSharedPtr< FJsonObject > InUnparsedJsonObject, TMap<int32, FText>* InLookupTablePtr, FText InCategoryName, bool bInTagAndValueCategory);
+	FCategorySectionHelper(TSharedPtr< FJsonObject > InUnparsedJsonObject, TMap<int32, FText>* InLookupTablePtr, FText InCategoryName, bool bInTagAndValueCategory, FCategorySectionHelperCallback InSpecialHandlingCallback);
+
+	/** FImaginaryFiBData Interface */
+	virtual void ParseAllChildData() override;
+	virtual bool IsCategory() override { return true; }
+	virtual bool IsTagAndValueCategory() override { return bIsTagAndValue; }
+		/** End FImaginaryFiBData Interface */
+
+protected:
+	/** FImaginaryFiBData Interface */
+	virtual bool CanCallFilter(ESearchQueryFilter InSearchQueryFilter) override;
+	virtual FSearchResult CreateSearchResult_Internal(FSearchResult InParent) const override;
+	/** End FImaginaryFiBData Interface */
+
+protected:
+	/** Callback to specially handle parsing of the Json Object instead of using generic handling */
+	FCategorySectionHelperCallback SpecialHandlingCallback;
+
+	/** The display text for this item in the search results */
+	FText CategoryName;
+
+	/** TRUE if this category should be considered no different than a normal Tag and Value in it's parent */
+	bool bIsTagAndValue;
+};
+
+/** An "imaginary" representation of a UBlueprint, featuring raw strings or other imaginary objects in the place of more structured substances */
+class FImaginaryBlueprint : public FImaginaryFiBData
+{
+public:
+	FImaginaryBlueprint(FString InBlueprintName, FString InBlueprintPath, FString InBlueprintParentClass, TArray<FString>& InInterfaces, FString InUnparsedStringData);
+
+	/** FImaginaryFiBData Interface */
+	virtual bool IsCompatibleWithFilter(ESearchQueryFilter InSearchQueryFilter) override;
+	virtual bool CanCallFilter(ESearchQueryFilter InSearchQueryFilter) override;
+	/** End FImaginaryFiBData Interface */
+
+protected:
+	/** FImaginaryFiBData Interface */
+	virtual bool TrySpecialHandleJsonValue(FString InKey, TSharedPtr< FJsonValue > InJsonValue) override;
+	virtual FSearchResult CreateSearchResult_Internal(FSearchResult InParent) const override;
+	/** End FImaginaryFiBData Interface */
+
+	/** Helper function to parse an array of Json Object representing graphs */
+	void ParseGraph( TSharedPtr< FJsonValue > InJsonValue, FString InCategoryTitle, EGraphType InGraphType );
+	
+	/** Callback to specially parse an array of Json Objects representing components */
+	void ParseComponents(TSharedPtr< FJsonObject > InJsonObject, TArray< TSharedPtr<FImaginaryFiBData> >& OutParsedChildData);
+
+	/** Parses a raw string of Json to a Json object hierarchy */
+	void ParseToJson();
+
+protected:
+	/** The path for this Blueprint */
+	FString BlueprintPath;
+
+	/** The raw Json string yet to be parsed */
+	FString UnparsedStringData;
+
+	/** Lookup table used as a compression tool for the FTexts stored in the Json object */
+	TMap<int32, FText> LookupTable;
+};
+
+/** An "imaginary" representation of a UEdGraph, featuring raw strings or other imaginary objects in the place of more structured substances */
+class FImaginaryGraph : public FImaginaryFiBData
+{
+public:
+	FImaginaryGraph(TSharedPtr< FJsonObject > InUnparsedJsonObject, TMap<int32, FText>* InLookupTablePtr, EGraphType InGraphType);
+
+	/** FImaginaryFiBData Interface */
+	virtual bool IsCompatibleWithFilter(ESearchQueryFilter InSearchQueryFilter) override;
+	virtual bool CanCallFilter(ESearchQueryFilter InSearchQueryFilter) override;
+	/** End FImaginaryFiBData Interface */
+
+protected:
+	/** FImaginaryFiBData Interface */
+	virtual bool TrySpecialHandleJsonValue(FString InKey, TSharedPtr< FJsonValue > InJsonValue) override;
+	virtual ESearchableValueStatus GetSearchabilityStatus(FString InKey) override;
+	virtual FSearchResult CreateSearchResult_Internal(FSearchResult InParent) const override;
+	/** End FImaginaryFiBData Interface */
+
+protected:
+	/** The graph type being represented */
+	EGraphType GraphType;
+};
+
+/** An "imaginary" representation of a UEdGraphNode, featuring raw strings or other imaginary objects in the place of more structured substances */
+class FImaginaryGraphNode : public FImaginaryFiBData
+{
+public:
+	FImaginaryGraphNode(TSharedPtr< FJsonObject > InUnparsedJsonObject, TMap<int32, FText>* InLookupTablePtr);
+
+	/** FImaginaryFiBData Interface */
+	virtual bool IsCompatibleWithFilter(ESearchQueryFilter InSearchQueryFilter) override;
+	virtual bool CanCallFilter(ESearchQueryFilter InSearchQueryFilter) override;
+	virtual void ParseAllChildData() override;
+	/** End FImaginaryFiBData Interface */
+
+protected:
+	/** FImaginaryFiBData Interface */
+	virtual bool TrySpecialHandleJsonValue(FString InKey, TSharedPtr< FJsonValue > InJsonValue) override;
+	virtual ESearchableValueStatus GetSearchabilityStatus(FString InKey) override;
+	virtual FSearchResult CreateSearchResult_Internal(FSearchResult InParent) const override;
+	/** End FImaginaryFiBData Interface */
+
+protected:
+	/** Schema name that manages this node */
+	FString SchemaName;
+};
+
+/** An "imaginary" representation of a UProperty, featuring raw strings or other imaginary objects in the place of more structured substances */
+class FImaginaryProperty : public FImaginaryFiBData
+{
+public:
+	FImaginaryProperty(TSharedPtr< FJsonObject > InUnparsedJsonObject, TMap<int32, FText>* InLookupTablePtr);
+
+	/** FImaginaryFiBData Interface */
+	virtual bool IsCompatibleWithFilter(ESearchQueryFilter InSearchQueryFilter) override;
+	/** End FImaginaryFiBData Interface */
+	
+protected:
+	/** FImaginaryFiBData Interface */
+	virtual ESearchableValueStatus GetSearchabilityStatus(FString InKey) override;
+	virtual FSearchResult CreateSearchResult_Internal(FSearchResult InParent) const override;
+	/** End FImaginaryFiBData Interface */
+};
+
+/** An "imaginary" representation of a UProperty of an instanced component, featuring raw strings or other imaginary objects in the place of more structured substances */
+class FImaginaryComponent : public FImaginaryProperty
+{
+public:
+	FImaginaryComponent(TSharedPtr< FJsonObject > InUnparsedJsonObject, TMap<int32, FText>* InLookupTablePtr);
+
+	/** FImaginaryFiBData Interface */
+	virtual bool IsCompatibleWithFilter(ESearchQueryFilter InSearchQueryFilter) override;
+	/** End FImaginaryFiBData Interface */
+};
+
+/** An "imaginary" representation of a UEdGraphPin, featuring raw strings or other imaginary objects in the place of more structured substances */
+class FImaginaryPin : public FImaginaryFiBData
+{
+public:
+	FImaginaryPin(TSharedPtr< FJsonObject > InUnparsedJsonObject, TMap<int32, FText>* InLookupTablePtr, FString InSchemaName);
+
+	/** FImaginaryFiBData Interface */
+	virtual bool IsCompatibleWithFilter(ESearchQueryFilter InSearchQueryFilter) override;
+	/** End FImaginaryFiBData Interface */
+
+protected:
+	/** FImaginaryFiBData Interface */
+	virtual bool TrySpecialHandleJsonValue(FString InKey, TSharedPtr< FJsonValue > InJsonValue);
+	virtual ESearchableValueStatus GetSearchabilityStatus(FString InKey) override;
+	virtual FSearchResult CreateSearchResult_Internal(FSearchResult InParent) const override;
+	/** End FImaginaryFiBData Interface */
+
+protected:
+	/** Schema name that manages this pin */
+	FString SchemaName;
+};
\ No newline at end of file
diff --git a/Engine/Source/Runtime/Core/Private/Misc/TextFilterExpressionEvaluator.cpp b/Engine/Source/Runtime/Core/Private/Misc/TextFilterExpressionEvaluator.cpp
--- a/Engine/Source/Runtime/Core/Private/Misc/TextFilterExpressionEvaluator.cpp
+++ b/Engine/Source/Runtime/Core/Private/Misc/TextFilterExpressionEvaluator.cpp
@@ -4,117 +4,6 @@
 #include "TextFilterExpressionEvaluator.h"
 #include "BasicMathExpressionEvaluator.h"
 
-namespace TextFilterExpressionParser
-{
-	class FTextToken
-	{
-	public:
-		enum class EInvertResult : uint8
-		{
-			No,
-			Yes,
-		};
-
-		FTextToken(FTextFilterString InString, ETextFilterTextComparisonMode InTextComparisonMode, const EInvertResult InInvertResult)
-			: String(MoveTemp(InString))
-			, TextComparisonMode(MoveTemp(InTextComparisonMode))
-			, InvertResult(InInvertResult)
-		{
-		}
-
-		FTextToken(const FTextToken& Other)
-			: String(Other.String)
-			, TextComparisonMode(Other.TextComparisonMode)
-			, InvertResult(Other.InvertResult)
-		{
-		}
-
-		FTextToken(FTextToken&& Other)
-			: String(MoveTemp(Other.String))
-			, TextComparisonMode(Other.TextComparisonMode)
-			, InvertResult(Other.InvertResult)
-		{
-		}
-
-		FTextToken& operator=(const FTextToken& Other)
-		{
-			String = Other.String;
-			TextComparisonMode = Other.TextComparisonMode;
-			InvertResult = Other.InvertResult;
-			return *this;
-		}
-
-		FTextToken& operator=(FTextToken&& Other)
-		{
-			String = MoveTemp(Other.String);
-			TextComparisonMode = Other.TextComparisonMode;
-			InvertResult = Other.InvertResult;
-			return *this;
-		}
-
-		const FTextFilterString& GetString() const
-		{
-			return String;
-		}
-
-		bool EvaluateAsBasicStringExpression(const ITextFilterExpressionContext* InContext) const
-		{
-			if (String.IsEmpty())
-			{
-				// An empty string is always considered to be true
-				return true;
-			}
-
-			const bool bMatched = InContext->TestBasicStringExpression(String, TextComparisonMode);
-			return (InvertResult == EInvertResult::No) ? bMatched : !bMatched;
-		}
-
-		bool EvaluateAsComplexExpression(const ITextFilterExpressionContext* InContext, const FTextFilterString& InKey, const ETextFilterComparisonOperation InComparisonOperation) const
-		{
-			const bool bMatched = InContext->TestComplexExpression(InKey.AsName(), String, InComparisonOperation, TextComparisonMode);
-			return (InvertResult == EInvertResult::No) ? bMatched : !bMatched;
-		}
-
-		ETextFilterTextComparisonMode GetTextComparisonMode() const
-		{
-			return TextComparisonMode;
-		}
-
-	private:
-		FTextFilterString String;
-		ETextFilterTextComparisonMode TextComparisonMode;
-		EInvertResult InvertResult;
-	};
-} // namespace TextFilterExpressionParser
-
-#define DEFINE_TEXT_EXPRESSION_OPERATOR_NODE(TYPE, MONIKER_COUNT, ...)	\
-	namespace TextFilterExpressionParser								\
-	{																	\
-		struct TYPE														\
-		{																\
-			static const int32 MonikerCount = MONIKER_COUNT;			\
-			static const TCHAR* Monikers[MonikerCount];					\
-		};																\
-	}																	\
-	DEFINE_EXPRESSION_NODE_TYPE(TextFilterExpressionParser::TYPE, __VA_ARGS__)
-
-DEFINE_TEXT_EXPRESSION_OPERATOR_NODE(FSubExpressionStart,		1,	0x0E7E1BC9, 0xF0B94D5D, 0x80F44D45, 0x85A082AA)
-DEFINE_TEXT_EXPRESSION_OPERATOR_NODE(FSubExpressionEnd,			1,	0x5E10956D, 0x2E17411F, 0xB6469E22, 0xB5E56E6C)
-
-DEFINE_TEXT_EXPRESSION_OPERATOR_NODE(FEqual,					3,	0x32457EFC, 0x4928406F, 0xBD78D943, 0x633797D1)
-DEFINE_TEXT_EXPRESSION_OPERATOR_NODE(FNotEqual,					2,	0x8F445A19, 0xF33443D9, 0x90D6DC85, 0xB0C5D071)
-DEFINE_TEXT_EXPRESSION_OPERATOR_NODE(FLess,						1,	0xB85E222B, 0x47E24E1F, 0xBC5A384D, 0x2FF471E1)
-DEFINE_TEXT_EXPRESSION_OPERATOR_NODE(FLessOrEqual,				2,	0x8C0A46B0, 0x8DAA4E2B, 0xA7FE4A23, 0xEF215918)
-DEFINE_TEXT_EXPRESSION_OPERATOR_NODE(FGreater,					1,	0x6A6247F4, 0xFD78467F, 0xA6AC1244, 0x0A31E0A5)
-DEFINE_TEXT_EXPRESSION_OPERATOR_NODE(FGreaterOrEqual,			2,	0x09D75C5E, 0xA29A4194, 0x8E8E5278, 0xC84365FD)
-
-DEFINE_TEXT_EXPRESSION_OPERATOR_NODE(FOr,						3,	0xF4778B51, 0xF535414D, 0x9C0EB5F2, 0x2F2B0FD4)
-DEFINE_TEXT_EXPRESSION_OPERATOR_NODE(FAnd,						3,	0x7511397A, 0x02D24DC2, 0x86729800, 0xF454C320)
-DEFINE_TEXT_EXPRESSION_OPERATOR_NODE(FNot,						2,	0x03D78990, 0x41D04E26, 0x8E98AD2F, 0x74667868)
-
-DEFINE_EXPRESSION_NODE_TYPE(TextFilterExpressionParser::FTextToken,	0x09E49538, 0x633545E3, 0x84B5644F, 0x1F11628F);
-DEFINE_EXPRESSION_NODE_TYPE(bool,									0xC1CD5DCF, 0x2AB44958, 0xB3FF4F8F, 0xE665D121);
-
 namespace TextFilterExpressionParser
 {
 	/**
@@ -447,8 +336,52 @@ bool FTextFilterExpressionEvaluator::SetFilterText(const FText& InFilterText)
 			{
 				// We inject AND operators between any adjacent text and sub-expression tokens to mimic the old search filter behavior
 				bool bWasPreviousTokenValidToInjectANDOperatorAfter = false;
-				for (FExpressionToken& CurrentToken : TmpLexTokens)
+				for (int32 TokenIdx = 0; TokenIdx < TmpLexTokens.Num(); ++TokenIdx)
 				{
+					FExpressionToken& CurrentToken = TmpLexTokens[TokenIdx];
+
+					if (TokenIdx + 3 < TmpLexTokens.Num())
+					{
+						if (const FTextToken* TextToken = CurrentToken.Node.Cast<FTextToken>())
+						{
+							FStringToken FunctionName(CurrentToken.Context);
+
+							FExpressionToken& ExpressionStartToken = TmpLexTokens[TokenIdx + 1];
+							if (ExpressionStartToken.Node.Cast<FSubExpressionStart>())
+							{
+								// We are inside a function, we need the token after the Start, but the TokenIdx is still on the function name
+								TokenIdx += 1;
+								FExpressionToken* NextToken = &TmpLexTokens[++TokenIdx];
+								FStringToken FunctionParameter(NextToken->Context);
+								int32 NestedCount = 1;
+								for (; TokenIdx < TmpLexTokens.Num(); ++TokenIdx) 
+								{
+									NextToken = &TmpLexTokens[TokenIdx];
+
+									if (NextToken->Node.Cast<FSubExpressionStart>())
+									{
+										++NestedCount;
+									}
+									else if (NextToken->Node.Cast<FSubExpressionEnd>())
+									{
+										--NestedCount;
+										if (NestedCount == 0)
+										{
+											break;
+										}
+									}
+									FunctionParameter.Accumulate(NextToken->Context);
+								}
+								FStringToken EmptyToken;
+								FinalTokens.Add(FExpressionToken(EmptyToken, FSubExpressionStart()));
+								FinalTokens.Add(FExpressionToken(FunctionName, FTextToken(FunctionName.GetString(), ETextFilterTextComparisonMode::Partial, FTextToken::EInvertResult::No)));
+								FinalTokens.Add(FExpressionToken(EmptyToken, FFunction()));
+								FinalTokens.Add(FExpressionToken(FunctionParameter, FTextToken(FunctionParameter.GetString(), ETextFilterTextComparisonMode::Partial, FTextToken::EInvertResult::No)));
+								FinalTokens.Add(FExpressionToken(EmptyToken, FSubExpressionEnd()));
+								continue;
+							}
+						}
+					}
 					// Key->value pairs require some form of comparison operator, so if we've found one of those, consider this to be a complex expression
 					if (!bIsComplexExpression && 
 						(CurrentToken.Node.Cast<FEqual>() || CurrentToken.Node.Cast<FNotEqual>() || CurrentToken.Node.Cast<FLess>() || CurrentToken.Node.Cast<FLessOrEqual>() || CurrentToken.Node.Cast<FGreater>() || CurrentToken.Node.Cast<FGreaterOrEqual>())
@@ -532,7 +465,12 @@ bool FTextFilterExpressionEvaluator::TestTextFilter(const ITextFilterExpressionC
 	return false;
 }
 
-void FTextFilterExpressionEvaluator::ConstructExpressionParser()
+void FTextFilterExpressionEvaluator::AddFunctionTokenCallback(FString InFunctionName, FTokenFunctionHandler InCallback)
+{
+	TokenFunctionHandlers.Add(InFunctionName, InCallback);
+}
+
+void FTextFilterExpressionEvaluator::SetupGrammar()
 {
 	using namespace TextFilterExpressionParser;
 	using TextFilterExpressionParser::FTextToken;
@@ -565,7 +503,14 @@ void FTextFilterExpressionEvaluator::ConstructExpressionParser()
 	Grammar.DefineBinaryOperator<FEqual>(1);
 	Grammar.DefineBinaryOperator<FOr>(2);
 	Grammar.DefineBinaryOperator<FAnd>(2);
+	Grammar.DefineBinaryOperator<FFunction>(1);
 	Grammar.DefinePreUnaryOperator<FNot>();
+}
+
+void FTextFilterExpressionEvaluator::ConstructExpressionParser()
+{
+	using namespace TextFilterExpressionParser;
+	using TextFilterExpressionParser::FTextToken;
 
 	JumpTable.MapBinary<FLessOrEqual>([](const FTextToken& A, const FTextToken& B, const ITextFilterExpressionContext* InContext)		{ return B.EvaluateAsComplexExpression(InContext, A.GetString(), ETextFilterComparisonOperation::LessOrEqual); });
 	JumpTable.MapBinary<FLess>([](const FTextToken& A, const FTextToken& B, const ITextFilterExpressionContext* InContext)				{ return B.EvaluateAsComplexExpression(InContext, A.GetString(), ETextFilterComparisonOperation::Less); });
@@ -586,6 +531,17 @@ void FTextFilterExpressionEvaluator::ConstructExpressionParser()
 
 	JumpTable.MapPreUnary<FNot>([](const FTextToken& V, const ITextFilterExpressionContext* InContext)									{ return !V.EvaluateAsBasicStringExpression(InContext); });
 	JumpTable.MapPreUnary<FNot>([](bool V, const ITextFilterExpressionContext* InContext)												{ return !V; });
+
+	JumpTable.MapBinary<FFunction>(
+		[this](const TextFilterExpressionParser::FTextToken& A, const TextFilterExpressionParser::FTextToken& B, const ITextFilterExpressionContext* InContext)	
+	{
+		bool bResult = false;
+		if (FTokenFunctionHandler* FunctionCallback = TokenFunctionHandlers.Find(A.GetString().AsString()))
+		{
+			bResult = FunctionCallback->Execute(B.GetString());
+		}
+		return bResult;
+	});
 }
 
 bool FTextFilterExpressionEvaluator::EvaluateCompiledExpression(const ExpressionParser::CompileResultType& InCompiledResult, const ITextFilterExpressionContext& InContext, FText* OutErrorText) const
diff --git a/Engine/Source/Runtime/Core/Public/Misc/TextFilterExpressionEvaluator.h b/Engine/Source/Runtime/Core/Public/Misc/TextFilterExpressionEvaluator.h
--- a/Engine/Source/Runtime/Core/Public/Misc/TextFilterExpressionEvaluator.h
+++ b/Engine/Source/Runtime/Core/Public/Misc/TextFilterExpressionEvaluator.h
@@ -26,7 +26,7 @@ class ITextFilterExpressionContext
 public:
 	/** Test the given value against the strings extracted from the current item */
 	virtual bool TestBasicStringExpression(const FTextFilterString& InValue, const ETextFilterTextComparisonMode InTextComparisonMode) const = 0;
-		
+
 	/** Perform a complex expression test for the current item */
 	virtual bool TestComplexExpression(const FName& InKey, const FTextFilterString& InValue, const ETextFilterComparisonOperation InComparisonOperation, const ETextFilterTextComparisonMode InTextComparisonMode) const = 0;
 
@@ -34,6 +34,125 @@ protected:
 	virtual ~ITextFilterExpressionContext() {}
 };
 
+namespace TextFilterExpressionParser
+{
+	class FTextToken
+	{
+	public:
+		enum class EInvertResult : uint8
+		{
+			No,
+			Yes,
+		};
+
+		FTextToken(FTextFilterString InString, ETextFilterTextComparisonMode InTextComparisonMode, const EInvertResult InInvertResult)
+			: String(MoveTemp(InString))
+			, TextComparisonMode(MoveTemp(InTextComparisonMode))
+			, InvertResult(InInvertResult)
+		{
+		}
+
+		FTextToken(const FTextToken& Other)
+			: String(Other.String)
+			, TextComparisonMode(Other.TextComparisonMode)
+			, InvertResult(Other.InvertResult)
+		{
+		}
+
+		FTextToken(FTextToken&& Other)
+			: String(MoveTemp(Other.String))
+			, TextComparisonMode(Other.TextComparisonMode)
+			, InvertResult(Other.InvertResult)
+		{
+		}
+
+		FTextToken& operator=(const FTextToken& Other)
+		{
+			String = Other.String;
+			TextComparisonMode = Other.TextComparisonMode;
+			InvertResult = Other.InvertResult;
+			return *this;
+		}
+
+		FTextToken& operator=(FTextToken&& Other)
+		{
+			String = MoveTemp(Other.String);
+			TextComparisonMode = Other.TextComparisonMode;
+			InvertResult = Other.InvertResult;
+			return *this;
+		}
+
+		const FTextFilterString& GetString() const
+		{
+			return String;
+		}
+
+		bool EvaluateAsBasicStringExpression(const ITextFilterExpressionContext* InContext) const
+		{
+			if (String.IsEmpty())
+			{
+				// An empty string is always considered to be true
+				return true;
+			}
+
+			const bool bMatched = InContext->TestBasicStringExpression(String, TextComparisonMode);
+			return (InvertResult == EInvertResult::No) ? bMatched : !bMatched;
+		}
+
+		bool EvaluateAsComplexExpression(const ITextFilterExpressionContext* InContext, const FTextFilterString& InKey, const ETextFilterComparisonOperation InComparisonOperation) const
+		{
+			const bool bMatched = InContext->TestComplexExpression(InKey.AsName(), String, InComparisonOperation, TextComparisonMode);
+			return (InvertResult == EInvertResult::No) ? bMatched : !bMatched;
+		}
+
+		ETextFilterTextComparisonMode GetTextComparisonMode() const
+		{
+			return TextComparisonMode;
+		}
+
+	private:
+		FTextFilterString String;
+		ETextFilterTextComparisonMode TextComparisonMode;
+		EInvertResult InvertResult;
+	};
+
+	class FFunction
+	{
+
+	};
+} // namespace TextFilterExpressionParser
+
+#define DEFINE_TEXT_EXPRESSION_OPERATOR_NODE(TYPE, MONIKER_COUNT, ...)	\
+	namespace TextFilterExpressionParser								\
+	{																	\
+		struct TYPE														\
+		{																\
+			static const int32 MonikerCount = MONIKER_COUNT;			\
+			static const TCHAR* Monikers[MonikerCount];					\
+		};																\
+	}																	\
+DEFINE_EXPRESSION_NODE_TYPE(TextFilterExpressionParser::TYPE, __VA_ARGS__)
+
+DEFINE_TEXT_EXPRESSION_OPERATOR_NODE(FSubExpressionStart,		1,	0x0E7E1BC9, 0xF0B94D5D, 0x80F44D45, 0x85A082AA)
+DEFINE_TEXT_EXPRESSION_OPERATOR_NODE(FSubExpressionEnd,			1,	0x5E10956D, 0x2E17411F, 0xB6469E22, 0xB5E56E6C)
+
+DEFINE_TEXT_EXPRESSION_OPERATOR_NODE(FEqual,					3,	0x32457EFC, 0x4928406F, 0xBD78D943, 0x633797D1)
+DEFINE_TEXT_EXPRESSION_OPERATOR_NODE(FNotEqual,					2,	0x8F445A19, 0xF33443D9, 0x90D6DC85, 0xB0C5D071)
+DEFINE_TEXT_EXPRESSION_OPERATOR_NODE(FLess,						1,	0xB85E222B, 0x47E24E1F, 0xBC5A384D, 0x2FF471E1)
+DEFINE_TEXT_EXPRESSION_OPERATOR_NODE(FLessOrEqual,				2,	0x8C0A46B0, 0x8DAA4E2B, 0xA7FE4A23, 0xEF215918)
+DEFINE_TEXT_EXPRESSION_OPERATOR_NODE(FGreater,					1,	0x6A6247F4, 0xFD78467F, 0xA6AC1244, 0x0A31E0A5)
+DEFINE_TEXT_EXPRESSION_OPERATOR_NODE(FGreaterOrEqual,			2,	0x09D75C5E, 0xA29A4194, 0x8E8E5278, 0xC84365FD)
+
+DEFINE_TEXT_EXPRESSION_OPERATOR_NODE(FOr,						3,	0xF4778B51, 0xF535414D, 0x9C0EB5F2, 0x2F2B0FD4)
+DEFINE_TEXT_EXPRESSION_OPERATOR_NODE(FAnd,						3,	0x7511397A, 0x02D24DC2, 0x86729800, 0xF454C320)
+DEFINE_TEXT_EXPRESSION_OPERATOR_NODE(FNot,						2,	0x03D78990, 0x41D04E26, 0x8E98AD2F, 0x74667868)
+
+DEFINE_EXPRESSION_NODE_TYPE(TextFilterExpressionParser::FTextToken,	0x09E49538, 0x633545E3, 0x84B5644F, 0x1F11628F);
+DEFINE_EXPRESSION_NODE_TYPE(TextFilterExpressionParser::FFunction,	0x084E6214, 0x032FFA48, 0x9245ABF1, 0x9E248F1A);
+DEFINE_EXPRESSION_NODE_TYPE(bool,									0xC1CD5DCF, 0x2AB44958, 0xB3FF4F8F, 0xE665D121);
+
+DECLARE_DELEGATE_RetVal_OneParam(bool, FTokenFunctionHandler, const FTextFilterString&);
+
 /** Defines an expression evaluator that can be used to perform complex text filter queries */
 class CORE_API FTextFilterExpressionEvaluator
 {
@@ -42,6 +161,7 @@ public:
 	FTextFilterExpressionEvaluator(const ETextFilterExpressionEvaluatorMode InMode);
 	FTextFilterExpressionEvaluator(const FTextFilterExpressionEvaluator& Other);
 	FTextFilterExpressionEvaluator& operator=(const FTextFilterExpressionEvaluator& Other);
+	virtual ~FTextFilterExpressionEvaluator() {}
 
 	/** Get the complexity of the current filter terms */
 	ETextFilterExpressionType GetFilterType() const;
@@ -58,12 +178,17 @@ public:
 	/** Test our compiled filter using the given context */
 	bool TestTextFilter(const ITextFilterExpressionContext& InContext) const;
 
-private:
+	/** Helper function to add callbacks for function tokens */
+	void AddFunctionTokenCallback(FString InFunctionName, FTokenFunctionHandler InCallback);
+protected:
+	/** Sets up grammar used for evaluation */
+	void SetupGrammar();
+
 	/** Common function to construct the expression parser */
-	void ConstructExpressionParser();
+	virtual void ConstructExpressionParser();
 
 	/** Evaluate the given compiled result, and optionally populate OutErrorText with any error information */
-	bool EvaluateCompiledExpression(const ExpressionParser::CompileResultType& InCompiledResult, const ITextFilterExpressionContext& InContext, FText* OutErrorText) const;
+	virtual bool EvaluateCompiledExpression(const ExpressionParser::CompileResultType& InCompiledResult, const ITextFilterExpressionContext& InContext, FText* OutErrorText) const;
 
 	/** Defines whether or not the expression parser can evaluate complex expressions */
 	ETextFilterExpressionEvaluatorMode ExpressionEvaluatorMode;
@@ -80,6 +205,9 @@ private:
 	/** The compiled filter created from the current filter text (compiled from FilterText) */
 	TOptional<ExpressionParser::CompileResultType> CompiledFilter;
 
+	/** Mapping of function names to their callbacks */
+	TMap<FString, FTokenFunctionHandler> TokenFunctionHandlers;
+
 	/** Expression parser */
 	FTokenDefinitions TokenDefinitions;
 	FExpressionGrammar Grammar;
