Back out changelist 2664397

diff --git a/Engine/Source/Editor/Kismet/Private/FiBSearchInstance.cpp b/Engine/Source/Editor/Kismet/Private/FiBSearchInstance.cpp
deleted file mode 100644
diff --git a/Engine/Source/Editor/Kismet/Private/FiBSearchInstance.h b/Engine/Source/Editor/Kismet/Private/FiBSearchInstance.h
deleted file mode 100644
diff --git a/Engine/Source/Editor/Kismet/Private/FindInBlueprintManager.cpp b/Engine/Source/Editor/Kismet/Private/FindInBlueprintManager.cpp
--- a/Engine/Source/Editor/Kismet/Private/FindInBlueprintManager.cpp
+++ b/Engine/Source/Editor/Kismet/Private/FindInBlueprintManager.cpp
@@ -39,9 +39,6 @@ const FText FFindInBlueprintSearchTags::FiB_Tooltip = LOCTEXT("Tooltip", "Toolti
 const FText FFindInBlueprintSearchTags::FiB_DefaultValue = LOCTEXT("DefaultValue", "DefaultValue");
 const FText FFindInBlueprintSearchTags::FiB_Description = LOCTEXT("Description", "Description");
 const FText FFindInBlueprintSearchTags::FiB_Comment = LOCTEXT("Comment", "Comment");
-const FText FFindInBlueprintSearchTags::FiB_Path = LOCTEXT("Path", "Path");
-const FText FFindInBlueprintSearchTags::FiB_ParentClass = LOCTEXT("ParentClass", "ParentClass");
-const FText FFindInBlueprintSearchTags::FiB_Interfaces = LOCTEXT("Interfaces", "Interfaces");
 
 const FText FFindInBlueprintSearchTags::FiB_Pins = LOCTEXT("Pins", "Pins");
 const FText FFindInBlueprintSearchTags::FiB_PinCategory = LOCTEXT("PinCategory", "PinCategory");
@@ -994,32 +991,6 @@ void FFindInBlueprintSearchManager::OnAssetAdded(const FAssetData& InAssetData)
 				FSearchData NewSearchData;
 
 				NewSearchData.BlueprintPath = InAssetData.ObjectPath;
-
-				if (const FString* ParentClass = InAssetData.TagsAndValues.Find(TEXT("ParentClass")))
-				{
-					NewSearchData.ParentClass = *ParentClass;
-				}
-
-				const FString* ImplementedInterfaces = InAssetData.TagsAndValues.Find("ImplementedInterfaces");
-				if(ImplementedInterfaces)
-				{
-					FString FullInterface;
-					FString RemainingString;
-					FString InterfaceName;
-					FString CurrentString = *ImplementedInterfaces;
-					while(CurrentString.Split(TEXT(","), &FullInterface, &RemainingString))
-					{
-						if(FullInterface.Split(TEXT("."), &CurrentString, &InterfaceName, ESearchCase::CaseSensitive, ESearchDir::FromEnd))
-						{
-							if(!CurrentString.StartsWith(TEXT("Graphs=(")))
-							{
-								NewSearchData.Interfaces.Add(InterfaceName);
-							}
-						}
-						CurrentString = RemainingString;
-					}
-				}
-
 				NewSearchData.bMarkedForDeletion = false;
 
 				// Since the asset was not loaded, pull out the searchable data stored in the asset
@@ -1215,6 +1186,7 @@ void FFindInBlueprintSearchManager::AddOrUpdateBlueprintSearchMetadata(UBlueprin
 		FSearchData SearchData;
 		SearchData.Blueprint = InBlueprint;
 		SearchData.BlueprintPath = BlueprintPath;
+
 		Index = AddSearchDataToDatabase(MoveTemp(SearchData));
 	}
 	else
@@ -1224,10 +1196,7 @@ void FFindInBlueprintSearchManager::AddOrUpdateBlueprintSearchMetadata(UBlueprin
 
 	// Build the search data
 	SearchArray[Index].BlueprintPath = BlueprintPath;
-	if (UProperty* ParentClassProp = InBlueprint->GetClass()->FindPropertyByName(GET_MEMBER_NAME_CHECKED(UBlueprint, ParentClass)))
-	{
-		ParentClassProp->ExportTextItem(SearchArray[Index].ParentClass, ParentClassProp->ContainerPtrToValuePtr<uint8>(InBlueprint), nullptr, InBlueprint, 0);
-	}
+
 	// Cannot successfully gather most searchable data if there is no SkeletonGeneratedClass, so don't try, leave it as whatever it was last set to
 	if(InBlueprint->SkeletonGeneratedClass != nullptr)
 	{
diff --git a/Engine/Source/Editor/Kismet/Private/FindInBlueprints.cpp b/Engine/Source/Editor/Kismet/Private/FindInBlueprints.cpp
--- a/Engine/Source/Editor/Kismet/Private/FindInBlueprints.cpp
+++ b/Engine/Source/Editor/Kismet/Private/FindInBlueprints.cpp
@@ -11,84 +11,214 @@
 #include "IDocumentation.h"
 #include "SSearchBox.h"
 #include "GenericCommands.h"
-#include "ImaginaryBlueprintData.h"
-#include "FiBSearchInstance.h"
 
 #define LOCTEXT_NAMESPACE "FindInBlueprints"
 
-
-FText FindInBlueprintsHelpers::AsFText(TSharedPtr< FJsonValue > InJsonValue, const TMap<int32, FText>& InLookupTable)
+/** Some utility functions to help with Find-in-Blueprint functionality */
+namespace FindInBlueprintsHelpers
 {
-	if (const FText* LookupText = InLookupTable.Find(FCString::Atoi(*InJsonValue->AsString())))
+	/** Looks up a JsonValue's FText from the passed lookup table */
+	FText AsFText(TSharedPtr< FJsonValue > InJsonValue, const TMap<int32, FText>& InLookupTable)
 	{
-		return *LookupText;
+		if(const FText* LookupText = InLookupTable.Find(FCString::Atoi(*InJsonValue->AsString())))
+		{
+			return *LookupText;
+		}
+		// Let's never get here.
+		return LOCTEXT("FiBSerializationError", "There was an error in serialization!");
 	}
-	// Let's never get here.
-	return LOCTEXT("FiBSerializationError", "There was an error in serialization!");
-}
 
-FText FindInBlueprintsHelpers::AsFText(int32 InValue, const TMap<int32, FText>& InLookupTable)
-{
-	if (const FText* LookupText = InLookupTable.Find(InValue))
+	/**
+	 * Retrieves the pin type as a string value
+	 *
+	 * @param InPinType		The pin type to look at
+	 *
+	 * @return				The pin type as a string in format [category]'[sub-category object]'
+	 */
+	FString GetPinTypeAsString(const FEdGraphPinType& InPinType)
 	{
-		return *LookupText;
-	}
-	// Let's never get here.
-	return LOCTEXT("FiBSerializationError", "There was an error in serialization!");
-}
-
-bool FindInBlueprintsHelpers::IsTextEqualToString(FText InText, FString InString)
-{
-	return InString == InText.ToString() || InString == InText.BuildSourceString();
-}
+		FString Result = InPinType.PinCategory;
+		if(UObject* SubCategoryObject = InPinType.PinSubCategoryObject.Get()) 
+		{
+			Result += FString(" '") + SubCategoryObject->GetName() + "'";
+		}
 
-FString FindInBlueprintsHelpers::GetPinTypeAsString(const FEdGraphPinType& InPinType)
-{
-	FString Result = InPinType.PinCategory;
-	if(UObject* SubCategoryObject = InPinType.PinSubCategoryObject.Get()) 
-	{
-		Result += FString(" '") + SubCategoryObject->GetName() + "'";
-	}
-	else
-	{
-		Result += FString(" '") + InPinType.PinSubCategory + "'";
+		return Result;
 	}
 
-	return Result;
-}
-
-bool FindInBlueprintsHelpers::ParsePinType(FText InKey, FText InValue, FEdGraphPinType& InOutPinType)
-{
-	bool bParsed = true;
-
-	if(InKey.CompareTo(FFindInBlueprintSearchTags::FiB_PinCategory) == 0)
-	{
-		InOutPinType.PinCategory = InValue.ToString();
-	}
-	else if(InKey.CompareTo(FFindInBlueprintSearchTags::FiB_PinSubCategory) == 0)
+	/**
+	 * Determines if a string matches the search tokens
+	 *
+	 * @param InTokens					Search tokens to compare with
+	 * @param InComparisonString		The string to check for matching tokens in
+	 *
+	 * @return							TRUE if the string matches any of the search tokens
+	 */
+	bool StringMatchesSearchTokens(const TArray<FString>& InTokens, const FString& InComparisonString)
 	{
-		InOutPinType.PinSubCategory = InValue.ToString();
+		const FString SafeSearchString = InComparisonString.Replace( TEXT( " " ), TEXT( "" ) );
+
+		bool bFoundAllTokens = true;
+		//search the entry for each token, it must have all of them to pass
+		for(auto TokIT(InTokens.CreateConstIterator());TokIT;++TokIT)
+		{
+			const FString& Token = *TokIT;
+			if(!SafeSearchString.Contains(Token))
+			{
+				bFoundAllTokens = false;
+				break;
+			}
+		}
+		return bFoundAllTokens;
 	}
-	else if(InKey.CompareTo(FFindInBlueprintSearchTags::FiB_ObjectClass) == 0)
+
+	/**
+	 * Determines if an FText matches the search tokens, comparing both the (localized) display string and the (unlocalized) source string
+	 *
+	 * @param InTokens					Search tokens to compare with
+	 * @param InComparisonText			The FText to check for matching tokens in
+	 *
+	 * @return							TRUE if the FText matches any of the search tokens
+	 */
+	bool TextMatchesSearchTokens(const TArray<FString>& InTokens, const FText& InComparisonText)
 	{
-		InOutPinType.PinSubCategory = InValue.ToString();
+		return StringMatchesSearchTokens(InTokens, InComparisonText.ToString()) || StringMatchesSearchTokens(InTokens, InComparisonText.BuildSourceString());
 	}
-	else if(InKey.CompareTo(FFindInBlueprintSearchTags::FiB_IsArray) == 0)
+
+	/**
+	 * Helper function to extract all graph data from a Json object and into a search result. There is little type safety, it assumes all Json objects found in the InCategoryTitle are graphs
+	 *
+	 * @param InJsonObject					The JsonObject to extract the graphs from
+	 * @param InCategoryTitle				The category the graphs are titled under.
+	 * @param InTokens						Filtering tokens to use
+	 * @param InGraphType					There is nothing in the Json script to identify graphs, they are all grouped by category, so pass in what type they should be considered
+	 * @param OutBlueprintCategory			The search category these graphs will be children to
+	 */
+	void ExtractGraph( const TMap<int32, FText>& InLookupTable, TSharedPtr< FJsonObject > InJsonObject, FString InCategoryTitle, const TArray<FString>& InTokens, EGraphType InGraphType, FSearchResult& OutBlueprintCategory )
 	{
-		InOutPinType.bIsArray = InValue.ToString().ToBool();
+		if(InJsonObject->HasField(InCategoryTitle))
+		{
+			// Pulls out all functions for this Blueprint
+			TArray<TSharedPtr< FJsonValue > > FunctionList = InJsonObject->GetArrayField(InCategoryTitle);
+			for( TSharedPtr< FJsonValue > FunctionValue : FunctionList )
+			{
+				FSearchResult Result(new FFindInBlueprintsGraph(FText::GetEmpty(), OutBlueprintCategory, InGraphType) );
+
+				if( Result->ExtractContent(InLookupTable, FunctionValue->AsObject(), InTokens) )
+				{
+					OutBlueprintCategory->Children.Add(Result);
+				}
+			}
+		}
 	}
-	else if(InKey.CompareTo(FFindInBlueprintSearchTags::FiB_IsReference) == 0)
+
+	/**
+	 * Utility function to begin a search
+	 *
+	 * @param InBuffer					The search buffer, a Json string to be parsed and converted into search result
+	 * @param InTokens					Search tokens to filter the results by
+	 * @param OutBlueprintCategory		The root category for the search results
+	 */
+	void Extract( const FString& InBuffer, const TArray<FString>& InTokens, FSearchResult& OutBlueprintCategory )
 	{
-		InOutPinType.bIsReference = InValue.ToString().ToBool();
+		// Skip empty buffers
+		if(InBuffer.Len() == 0)
+		{
+			return;
+		}
+
+		TSharedPtr< FJsonObject > JsonObject = NULL;
+		TMap<int32, FText> LookupTable;
+		JsonObject = FFindInBlueprintSearchManager::ConvertJsonStringToObject(InBuffer, LookupTable);
+
+		if(JsonObject.IsValid())
+		{
+			if(JsonObject->HasField(FFindInBlueprintSearchTags::FiB_Properties.ToString()))
+			{
+				// Pulls out all properties (variables) for this Blueprint
+				TArray<TSharedPtr< FJsonValue > > PropertyList = JsonObject->GetArrayField(FFindInBlueprintSearchTags::FiB_Properties.ToString());
+				for( TSharedPtr< FJsonValue > PropertyValue : PropertyList )
+				{
+					FSearchResult Result(new FFindInBlueprintsProperty(FText::GetEmpty(), OutBlueprintCategory) );
+
+					if( Result->ExtractContent(LookupTable, PropertyValue->AsObject(), InTokens) )
+					{
+						OutBlueprintCategory->Children.Add(Result);
+					}
+				}
+			}
+
+			ExtractGraph(LookupTable, JsonObject, FFindInBlueprintSearchTags::FiB_Functions.ToString(), InTokens, EGraphType::GT_Function, OutBlueprintCategory);
+			ExtractGraph(LookupTable, JsonObject, FFindInBlueprintSearchTags::FiB_Macros.ToString(), InTokens, EGraphType::GT_Macro, OutBlueprintCategory);
+			ExtractGraph(LookupTable, JsonObject, FFindInBlueprintSearchTags::FiB_UberGraphs.ToString(), InTokens, EGraphType::GT_Ubergraph, OutBlueprintCategory);
+			ExtractGraph(LookupTable, JsonObject, FFindInBlueprintSearchTags::FiB_SubGraphs.ToString(), InTokens, EGraphType::GT_Ubergraph, OutBlueprintCategory);
+
+			if(JsonObject->HasField(FFindInBlueprintSearchTags::FiB_Components.ToString()))
+			{
+				FSearchResult ComponentsResult(new FFindInBlueprintsResult(FFindInBlueprintSearchTags::FiB_Components, OutBlueprintCategory) );
+
+				// Pulls out all properties (variables) for this Blueprint
+				TArray<TSharedPtr< FJsonValue > > ComponentList = JsonObject->GetArrayField(FFindInBlueprintSearchTags::FiB_Components.ToString());
+				for( TSharedPtr< FJsonValue > ComponentValue : ComponentList )
+				{
+					FSearchResult Result(new FFindInBlueprintsProperty(FText::GetEmpty(), ComponentsResult) );
+
+					if( Result->ExtractContent(LookupTable, ComponentValue->AsObject(), InTokens) )
+					{
+						ComponentsResult->Children.Add(Result);
+					}
+				}
+
+				// Only add the component list category if there were components
+				if(ComponentsResult->Children.Num())
+				{
+					OutBlueprintCategory->Children.Add(ComponentsResult);
+				}
+			}
+		}
 	}
-	else
+
+	/**
+	 * Parses a pin type from passed in key names and values
+	 *
+	 * @param InKey					The key name for what the data should be translated as
+	 * @param InValue				Value to be be translated
+	 * @param InOutPinType			Modifies the PinType based on the passed parameters, building it up over multiple calls
+	 * @param InOutPinSubCategory	When the pin's sub-category is found, this value will be changed. This value will never be wiped during other translations. It is not thread safe to determine the sub-object.
+	 * @return						TRUE when the parsing is successful
+	 */
+	bool ParsePinType(FText InKey, FText InValue, FEdGraphPinType& InOutPinType, FString& OutPinSubCategory)
 	{
-		bParsed = false;
-	}
+		bool bParsed = true;
 
-	return bParsed;
-}
+		if(InKey.CompareTo(FFindInBlueprintSearchTags::FiB_PinCategory) == 0)
+		{
+			InOutPinType.PinCategory = InValue.ToString();
+		}
+		else if(InKey.CompareTo(FFindInBlueprintSearchTags::FiB_PinSubCategory) == 0)
+		{
+			InOutPinType.PinSubCategory = InValue.ToString();
+		}
+		else if(InKey.CompareTo(FFindInBlueprintSearchTags::FiB_ObjectClass) == 0)
+		{
+			OutPinSubCategory = InValue.ToString();
+		}
+		else if(InKey.CompareTo(FFindInBlueprintSearchTags::FiB_IsArray) == 0)
+		{
+			InOutPinType.bIsArray = InValue.ToString().ToBool();
+		}
+		else if(InKey.CompareTo(FFindInBlueprintSearchTags::FiB_IsReference) == 0)
+		{
+			InOutPinType.bIsReference = InValue.ToString().ToBool();
+		}
+		else
+		{
+			bParsed = false;
+		}
 
+		return bParsed;
+	}
+}
 
 //////////////////////////////////////////////////////////////////////////
 // FBlueprintSearchResult
@@ -123,6 +253,96 @@ FReply FFindInBlueprintsResult::OnClick()
 	return FReply::Handled();
 }
 
+bool FFindInBlueprintsResult::ParseSearchInfo(const TArray<FString> &InTokens, FText InKey, FText InValue, TSharedPtr< FFindInBlueprintsResult > InParent)
+{
+	// By default, nothing much gets processed, just check to see if it matches the search requirements
+	return FindInBlueprintsHelpers::TextMatchesSearchTokens(InTokens, InValue);
+}
+
+void FFindInBlueprintsResult::AddExtraSearchInfo(FText InKey, FText InValue, TSharedPtr< FFindInBlueprintsResult > InParent)
+{
+	FText ExtraSearchInfoText;
+
+	FFormatNamedArguments Args;
+	Args.Add(TEXT("Key"), InKey);
+	Args.Add(TEXT("Value"), InValue);
+	ExtraSearchInfoText = FText::Format(LOCTEXT("ExtraSearchInfo", "{Key}: {Value}"), Args);
+
+	TSharedPtr< FFindInBlueprintsResult > SearchResult(new FFindInBlueprintsResult(ExtraSearchInfoText, InParent) );
+
+	if (InKey.CompareTo(FFindInBlueprintSearchTags::FiB_NativeName) != 0)
+	{
+		InParent->Children.Add(SearchResult);
+	}
+}
+
+bool FFindInBlueprintsResult::ExtractJsonValue(const TMap<int32, FText>& InLookupTable, const TArray<FString>& InTokens, FText InKey, TSharedPtr< FJsonValue > InValue, TSharedPtr< FFindInBlueprintsResult > InParentOverride)
+{
+	bool bMatchesSearchResults = false;
+
+	TSharedPtr< FFindInBlueprintsResult > ParentResultPtr = InParentOverride.IsValid()? InParentOverride : SharedThis( this );
+
+	if( InValue->Type == EJson::Array)
+	{
+		TSharedPtr< FFindInBlueprintsResult > ArrayHeader = MakeShareable(new FFindInBlueprintsResult(InKey, ParentResultPtr) );
+		bool bArrayHasMatches = false;
+
+		TArray< TSharedPtr< FJsonValue > > Array = InValue->AsArray();
+		for( int ArrayIdx = 0; ArrayIdx < Array.Num(); ++ArrayIdx )
+		{
+			if(Array[ArrayIdx]->Type == EJson::String)
+			{
+				bArrayHasMatches |= ParseSearchInfo(InTokens, FText::FromString(FString::FromInt(ArrayIdx)), FindInBlueprintsHelpers::AsFText(Array[ArrayIdx], InLookupTable), ArrayHeader);
+			}
+			else
+			{
+				bArrayHasMatches |= ExtractJsonValue(InLookupTable, InTokens, FText::FromString(FString::FromInt(ArrayIdx)), Array[ArrayIdx], ArrayHeader);
+			}
+		}
+
+		bMatchesSearchResults |= bArrayHasMatches;
+
+		if(bArrayHasMatches)
+		{
+			ParentResultPtr->Children.Add(ArrayHeader);
+		}
+	}
+	else if(InValue->Type == EJson::Object)
+	{
+		TSharedPtr< FFindInBlueprintsResult > ObjectResultHeader = MakeShareable(new FFindInBlueprintsResult(InKey, ParentResultPtr) );
+
+		bMatchesSearchResults |= ExtractContent(InLookupTable, InValue->AsObject(), InTokens, ObjectResultHeader);
+
+		if(bMatchesSearchResults)
+		{
+			ParentResultPtr->Children.Add(ObjectResultHeader);
+		}
+	}
+	else if(InValue->Type == EJson::String)
+	{
+		bMatchesSearchResults |= ParseSearchInfo(InTokens, InKey, FindInBlueprintsHelpers::AsFText(InValue, InLookupTable), ParentResultPtr);
+	}
+	else
+	{
+		bMatchesSearchResults |= ParseSearchInfo(InTokens, InKey, FindInBlueprintsHelpers::AsFText(InValue, InLookupTable), ParentResultPtr);
+	}
+
+	return bMatchesSearchResults;
+}
+
+bool FFindInBlueprintsResult::ExtractContent(const TMap<int32, FText>& InLookupTable, TSharedPtr< FJsonObject > InJsonNode, const TArray<FString>& InTokens, TSharedPtr< FFindInBlueprintsResult > InParentOverride)
+{
+	bool bMatchesSearchResults = false;
+
+	TSharedPtr< FFindInBlueprintsResult > ParentResultPtr = InParentOverride.IsValid()? InParentOverride : SharedThis( this );
+
+	for( auto MapValues : InJsonNode->Values )
+	{
+		bMatchesSearchResults |= ExtractJsonValue(InLookupTable, InTokens, FText::FromString(MapValues.Key), MapValues.Value, ParentResultPtr);
+	}
+	return bMatchesSearchResults;
+}
+
 FText FFindInBlueprintsResult::GetCategory() const
 {
 	return FText::GetEmpty();
@@ -223,17 +443,23 @@ TSharedRef<SWidget> FFindInBlueprintsGraphNode::CreateIcon() const
 		.ToolTipText( GetCategory() );
 }
 
-void FFindInBlueprintsGraphNode::ParseSearchInfo(FText InKey, FText InValue)
+bool FFindInBlueprintsGraphNode::ParseSearchInfo(const TArray<FString> &InTokens, FText InKey, FText InValue, TSharedPtr< FFindInBlueprintsResult > InParent)
 {
 	if(InKey.CompareTo(FFindInBlueprintSearchTags::FiB_NodeGuid) == 0)
 	{
 		FString NodeGUIDAsString = InValue.ToString();
 		FGuid::Parse(NodeGUIDAsString, NodeGuid);
+
+		// NodeGuid's are not searchable
+		return false;
 	}
 
+	bool bMatchesTokens = FindInBlueprintsHelpers::TextMatchesSearchTokens(InTokens, InValue);
+
 	if(InKey.CompareTo(FFindInBlueprintSearchTags::FiB_ClassName) == 0)
 	{
 		ClassName = InValue.ToString();
+		bMatchesTokens = false;
 	}
 	else if(InKey.CompareTo(FFindInBlueprintSearchTags::FiB_Name) == 0)
 	{
@@ -246,11 +472,67 @@ void FFindInBlueprintsGraphNode::ParseSearchInfo(FText InKey, FText InValue)
 	else if(InKey.CompareTo(FFindInBlueprintSearchTags::FiB_Glyph) == 0)
 	{
 		GlyphBrush = FEditorStyle::GetBrush(*InValue.ToString());
+		bMatchesTokens = false;
 	}
 	else if(InKey.CompareTo(FFindInBlueprintSearchTags::FiB_GlyphColor) == 0)
 	{
 		GlyphColor.InitFromString(InValue.ToString());
+		bMatchesTokens = false;
 	}
+	else
+	{
+		// Any generic keys and value pairs get put in as "extra" search info
+		if(bMatchesTokens)
+		{
+			AddExtraSearchInfo(InKey, InValue, InParent);
+			return true;
+		}
+	}
+	return bMatchesTokens;
+}
+
+bool FFindInBlueprintsGraphNode::ExtractContent(const TMap<int32, FText>& InLookupTable, TSharedPtr< FJsonObject > InJsonNode, const TArray<FString>& InTokens)
+{
+	// Very important to get the schema first, other bits of data depend on it
+	FString SchemaName;
+	TSharedPtr< FJsonValue > SchemaNameValue = InJsonNode->GetField< EJson::String >(FFindInBlueprintSearchTags::FiB_SchemaName.ToString());
+	if(SchemaNameValue.IsValid())
+	{
+		SchemaName = FindInBlueprintsHelpers::AsFText(SchemaNameValue, InLookupTable).ToString();
+	}
+
+	bool bMatchesSearchResults = false;
+
+	for( auto MapValues : InJsonNode->Values )
+	{
+		FText Key = FText::FromString(MapValues.Key);
+
+		if(Key.CompareTo(FFindInBlueprintSearchTags::FiB_Pins) == 0)
+		{
+			TArray< TSharedPtr< FJsonValue > > PinsList = MapValues.Value->AsArray();
+
+			for( TSharedPtr< FJsonValue > Pin : PinsList )
+			{
+				// The pin
+				TSharedPtr< FFindInBlueprintsPin > PinResult(new FFindInBlueprintsPin(FText::GetEmpty(), SharedThis( this ), SchemaName) );
+				if(PinResult->ExtractContent(InLookupTable, Pin->AsObject(), InTokens))
+				{
+					Children.Add(PinResult);
+					bMatchesSearchResults = true;
+				}
+			}
+		}
+		else if(Key.CompareTo(FFindInBlueprintSearchTags::FiB_SchemaName) == 0)
+		{
+			// Previously extracted
+			continue;
+		}
+		else
+		{
+			bMatchesSearchResults |= ExtractJsonValue(InLookupTable, InTokens, FText::FromString(MapValues.Key), MapValues.Value, SharedThis( this ));
+		}
+	}
+	return bMatchesSearchResults;
 }
 
 FText FFindInBlueprintsGraphNode::GetCategory() const
@@ -321,16 +603,32 @@ TSharedRef<SWidget> FFindInBlueprintsPin::CreateIcon() const
 		.ToolTipText(FText::FromString(FindInBlueprintsHelpers::GetPinTypeAsString(PinType)));
 }
 
-void FFindInBlueprintsPin::ParseSearchInfo(FText InKey, FText InValue)
+bool FFindInBlueprintsPin::ParseSearchInfo(const TArray<FString> &InTokens, FText InKey, FText InValue, TSharedPtr< FFindInBlueprintsResult > InParent)
 {
+	bool bMatchesTokens = FindInBlueprintsHelpers::TextMatchesSearchTokens(InTokens, InValue);
+
 	if(InKey.CompareTo(FFindInBlueprintSearchTags::FiB_Name) == 0)
 	{
 		DisplayText = InValue;
+		bMatchesTokens = FindInBlueprintsHelpers::StringMatchesSearchTokens(InTokens, DisplayText.ToString());
+	}
+	else if(FindInBlueprintsHelpers::ParsePinType(InKey, InValue, PinType, PinSubCategory))
+	{
+		// For pins, the only information in the Pin Type we care about is the object class, so return false if the Key is not that
+		if(InKey.CompareTo(FFindInBlueprintSearchTags::FiB_ObjectClass) != 0)
+		{
+			return false;
+		}
 	}
 	else
 	{
-		FindInBlueprintsHelpers::ParsePinType(InKey, InValue, PinType);
+		// Any generic keys and value pairs get put in as "extra" search info
+		if(bMatchesTokens)
+		{
+			AddExtraSearchInfo(InKey, InValue, InParent);
+		}
 	}
+	return bMatchesTokens;
 }
 
 FText FFindInBlueprintsPin::GetCategory() const
@@ -340,18 +638,14 @@ FText FFindInBlueprintsPin::GetCategory() const
 
 void FFindInBlueprintsPin::FinalizeSearchData()
 {
-	if(!PinType.PinSubCategory.IsEmpty())
+	if(!PinSubCategory.IsEmpty())
 	{
-		PinType.PinSubCategoryObject = FindObject<UClass>(ANY_PACKAGE, *PinType.PinSubCategory, true);
+		PinType.PinSubCategoryObject = FindObject<UClass>(ANY_PACKAGE, *PinSubCategory, true);
 		if(!PinType.PinSubCategoryObject.IsValid())
 		{
-			PinType.PinSubCategoryObject = FindObject<UScriptStruct>(UObject::StaticClass(), *PinType.PinSubCategory);
-		}
-
-		if (PinType.PinSubCategoryObject.IsValid())
-		{
-			PinType.PinSubCategory.Empty();
+			PinType.PinSubCategoryObject = FindObject<UScriptStruct>(UObject::StaticClass(), *PinSubCategory);
 		}
+		PinSubCategory.Empty();
 	}
 
 	if(!SchemaName.IsEmpty())
@@ -389,20 +683,38 @@ TSharedRef<SWidget> FFindInBlueprintsProperty::CreateIcon() const
 		.ToolTipText( FText::FromString(FindInBlueprintsHelpers::GetPinTypeAsString(PinType)) );
 }
 
-void FFindInBlueprintsProperty::ParseSearchInfo(FText InKey, FText InValue)
+bool FFindInBlueprintsProperty::ParseSearchInfo(const TArray<FString> &InTokens, FText InKey, FText InValue, TSharedPtr< FFindInBlueprintsResult > InParent)
 {
+	bool bMatchesTokens = FindInBlueprintsHelpers::TextMatchesSearchTokens(InTokens, InValue);
+
 	if(InKey.CompareTo(FFindInBlueprintSearchTags::FiB_Name) == 0)
 	{
 		DisplayText = InValue;
 	}
+	else if(FindInBlueprintsHelpers::ParsePinType(InKey, InValue, PinType, PinSubCategory))
+	{
+		// For properties, we care about all the information in the pin type, except whether its a reference or an array
+		if(InKey.CompareTo(FFindInBlueprintSearchTags::FiB_IsArray) == 0 || InKey.CompareTo(FFindInBlueprintSearchTags::FiB_IsReference) == 0)
+		{
+			return false;
+		}
+	}
 	else if(InKey.CompareTo(FFindInBlueprintSearchTags::FiB_IsSCSComponent) == 0)
 	{
 		bIsSCSComponent = true;
+
+		// This value is not searchable
+		bMatchesTokens = false;
 	}
 	else
 	{
-		FindInBlueprintsHelpers::ParsePinType(InKey, InValue, PinType);
+		// Any generic keys and value pairs get put in as "extra" search info
+		if(bMatchesTokens)
+		{
+			AddExtraSearchInfo(InKey, InValue, InParent);
+		}
 	}
+	return bMatchesTokens;
 }
 
 FText FFindInBlueprintsProperty::GetCategory() const
@@ -416,18 +728,14 @@ FText FFindInBlueprintsProperty::GetCategory() const
 
 void FFindInBlueprintsProperty::FinalizeSearchData()
 {
-	if(!PinType.PinSubCategory.IsEmpty())
+	if(!PinSubCategory.IsEmpty())
 	{
-		PinType.PinSubCategoryObject = FindObject<UClass>(ANY_PACKAGE, *PinType.PinSubCategory, true);
+		PinType.PinSubCategoryObject = FindObject<UClass>(ANY_PACKAGE, *PinSubCategory, true);
 		if(!PinType.PinSubCategoryObject.IsValid())
 		{
-			PinType.PinSubCategoryObject = FindObject<UScriptStruct>(UObject::StaticClass(), *PinType.PinSubCategory);
-		}
-
-		if (PinType.PinSubCategoryObject.IsValid())
-		{
-			PinType.PinSubCategory.Empty();
+			PinType.PinSubCategoryObject = FindObject<UScriptStruct>(UObject::StaticClass(), *PinSubCategory);
 		}
+		PinSubCategory.Empty();
 	}
 }
 
@@ -484,12 +792,74 @@ TSharedRef<SWidget> FFindInBlueprintsGraph::CreateIcon() const
 		.ToolTipText( GetCategory() );
 }
 
-void FFindInBlueprintsGraph::ParseSearchInfo(FText InKey, FText InValue)
+bool FFindInBlueprintsGraph::ExtractContent(const TMap<int32, FText>& InLookupTable, TSharedPtr< FJsonObject > InJsonNode, const TArray<FString>& InTokens)
 {
+	bool bMatchesSearchResults = false;
+
+	for( auto MapValues : InJsonNode->Values )
+	{
+		FText Key = FText::FromString(MapValues.Key);
+
+		if(Key.CompareTo(FFindInBlueprintSearchTags::FiB_Nodes) == 0)
+		{
+			TArray< TSharedPtr< FJsonValue > > NodeList = MapValues.Value->AsArray();
+
+			for( TSharedPtr< FJsonValue > NodeValue : NodeList )
+			{
+				TSharedPtr< FFindInBlueprintsGraphNode > NodeResult(new FFindInBlueprintsGraphNode(FText::GetEmpty(), SharedThis( this )) );
+
+				if( NodeResult->ExtractContent(InLookupTable, NodeValue->AsObject(), InTokens) )
+				{
+					Children.Add(NodeResult);
+					bMatchesSearchResults = true;
+				}
+			}
+		}
+		else if(Key.CompareTo(FFindInBlueprintSearchTags::FiB_SchemaName) == 0)
+		{
+			// Previously extracted
+			continue;
+		}
+		else if(Key.CompareTo(FFindInBlueprintSearchTags::FiB_Properties) == 0)
+		{
+			// Pulls out all properties (variables) for this Blueprint
+			TArray<TSharedPtr< FJsonValue > > PropertyList = MapValues.Value->AsArray();
+			for( TSharedPtr< FJsonValue > PropertyValue : PropertyList )
+			{
+				FSearchResult Result(new FFindInBlueprintsProperty(FText::GetEmpty(), SharedThis( this )) );
+
+				if( Result->ExtractContent(InLookupTable, PropertyValue->AsObject(), InTokens) )
+				{
+					Children.Add(Result);
+					bMatchesSearchResults = true;
+				}
+			}
+		}
+		else
+		{
+			bMatchesSearchResults |= ExtractJsonValue( InLookupTable, InTokens, FText::FromString(MapValues.Key), MapValues.Value, SharedThis( this ) );
+		}
+	}
+	return bMatchesSearchResults;
+}
+
+bool FFindInBlueprintsGraph::ParseSearchInfo(const TArray<FString> &InTokens, FText InKey, FText InValue, TSharedPtr< FFindInBlueprintsResult > InParent)
+{
+	bool bMatchesTokens = FindInBlueprintsHelpers::TextMatchesSearchTokens(InTokens, InValue);
+
 	if(InKey.CompareTo(FFindInBlueprintSearchTags::FiB_Name) == 0)
 	{
 		DisplayText = InValue;
 	}
+	else
+	{
+		if(bMatchesTokens)
+		{
+			AddExtraSearchInfo(InKey, InValue, InParent);
+		}
+	}
+
+	return bMatchesTokens;
 }
 
 FText FFindInBlueprintsGraph::GetCategory() const
@@ -515,8 +885,8 @@ class FStreamSearch : public FRunnable
 {
 public:
 	/** Constructor */
-	FStreamSearch(const FString& InSearchValue )
-		: SearchValue(InSearchValue)
+	FStreamSearch(const TArray<FString>& InSearchTokens )
+		: SearchTokens(InSearchTokens)
 		, bThreadCompleted(false)
 		, StopTaskCounter(0)
 	{
@@ -530,44 +900,7 @@ public:
 		return true;
 	}
 
-	virtual uint32 Run() override
-	{
-		FFindInBlueprintSearchManager::Get().BeginSearchQuery(this);
-
-		TFunction<void(const FSearchResult&)> OnResultReady = [this](const FSearchResult& Result) {
-			FScopeLock ScopeLock(&SearchCriticalSection);
-			ItemsFound.Add(Result);
-		};
-
-		// Searching comes to an end if it is requested using the StopTaskCounter or continuing the search query yields no results
-		FSearchData QueryResult;
-		while (FFindInBlueprintSearchManager::Get().ContinueSearchQuery(this, QueryResult))
-		{
-			// This work cannot be spun off to a use the TaskGraph system for spreading the work over multiple threads/cores as the stack limit is too small for the work being done.
-			if (QueryResult.Value.Len() > 0)
-			{
-				TSharedPtr< FImaginaryBlueprint> ImaginaryBlueprint(new FImaginaryBlueprint(FPaths::GetBaseFilename(QueryResult.BlueprintPath.ToString()), QueryResult.BlueprintPath.ToString(), QueryResult.ParentClass, QueryResult.Interfaces, QueryResult.Value));
-				TSharedPtr< FFiBSearchInstance > SearchInstance(new FFiBSearchInstance);
-				FSearchResult SearchResult = SearchInstance->StartSearchQuery(*SearchValue, ImaginaryBlueprint);
-
-				// If there are children, add the item to the search results
-				if(SearchResult.IsValid() && SearchResult->Children.Num() != 0)
-				{
-					OnResultReady(SearchResult);
-				}
-			}
-
-			if (StopTaskCounter.GetValue())
-			{
-				// Ensure that the FiB Manager knows that we are done searching
-				FFindInBlueprintSearchManager::Get().EnsureSearchQueryEnds(this);
-			}
-		}
-
-		bThreadCompleted = true;
-
-		return 0;
-	}
+	virtual uint32 Run() override;
 
 	virtual void Stop() override
 	{
@@ -580,6 +913,8 @@ public:
 	}
 	/** End FRunnable Interface */
 
+	TArray<FSearchData> MakeExtractionBatch() const;
+
 	/** Brings the thread to a safe stop before continuing. */
 	void EnsureCompletion()
 	{
@@ -601,7 +936,7 @@ public:
 	}
 
 	/**
-	 * Appends the items filtered through the search filter to the passed array
+	 * Appends the items filtered through the search tokens to the passed array
 	 *
 	 * @param OutItemsFound		All the items found since last queried
 	 */
@@ -625,8 +960,8 @@ public:
 	/** A list of items found, cleared whenever the main thread pulls them to display to screen */
 	TArray<FSearchResult> ItemsFound;
 
-	/** The search value to filter results by */
-	FString SearchValue;
+	/** The search tokens to filter results by */
+	TArray<FString> SearchTokens;
 
 	/** Prevents searching while other threads are pulling search results */
 	FCriticalSection SearchCriticalSection;
@@ -1040,6 +1375,34 @@ ECheckBoxState SFindInBlueprints::OnGetFindModeChecked() const
 
 void SFindInBlueprints::InitiateSearch()
 {
+	TArray<FString> Tokens;
+	if(SearchValue.Contains("\"") && SearchValue.ParseIntoArray(Tokens, TEXT("\""), true)>0)
+	{
+		for( auto &TokenIt : Tokens )
+		{
+			// we have the token, we don't need the quotes anymore, they'll just confused the comparison later on
+			TokenIt = TokenIt.TrimQuotes();
+			// We remove the spaces as all later comparison strings will also be de-spaced
+			TokenIt = TokenIt.Replace( TEXT( " " ), TEXT( "" ) );
+		}
+
+		// due to being able to handle multiple quoted blocks like ("Make Epic" "Game Now") we can end up with
+		// and empty string between (" ") blocks so this simply removes them
+		struct FRemoveMatchingStrings
+		{ 
+			bool operator()(const FString& RemovalCandidate) const
+			{
+				return RemovalCandidate.IsEmpty();
+			}
+		};
+		Tokens.RemoveAll( FRemoveMatchingStrings() );
+	}
+	else
+	{
+		// unquoted search equivalent to a match-any-of search
+		SearchValue.ParseIntoArray(Tokens, TEXT(" "), true);
+	}
+
 	if(ItemsFound.Num())
 	{
 		// Reset the scroll to the top
@@ -1048,7 +1411,7 @@ void SFindInBlueprints::InitiateSearch()
 
 	ItemsFound.Empty();
 
-	if (SearchValue.Len() > 0)
+	if (Tokens.Num() > 0)
 	{
 		OnRemoveCacheBar();
 
@@ -1057,59 +1420,55 @@ void SFindInBlueprints::InitiateSearch()
 		HighlightText = FText::FromString( SearchValue );
 		if (bIsInFindWithinBlueprintMode)
 		{
-			if(StreamSearch.IsValid() && !StreamSearch->IsComplete())
-			{
-				StreamSearch->Stop();
-				StreamSearch->EnsureCompletion();
-				StreamSearch.Reset();
-			}
-
-			UBlueprint* Blueprint = BlueprintEditorPtr.Pin()->GetBlueprintObj();
-			FString ParentClass;
-			if (UProperty* ParentClassProp = Blueprint->GetClass()->FindPropertyByName(GET_MEMBER_NAME_CHECKED(UBlueprint, ParentClass)))
-			{
-				ParentClassProp->ExportTextItem(ParentClass, ParentClassProp->ContainerPtrToValuePtr<uint8>(Blueprint), nullptr, Blueprint, 0);
-			}
-
-			TArray<FString> Interfaces;
+			MatchTokensWithinBlueprint(Tokens);
+		}
+		else
+		{
+			LaunchStreamThread(Tokens);
+		}
+	}
+}
 
-			for (FBPInterfaceDescription& InterfaceDesc : Blueprint->ImplementedInterfaces)
-			{
-				Interfaces.Add(InterfaceDesc.Interface->GetPathName());
-			}
-			TSharedPtr< FImaginaryBlueprint> ImaginaryBlueprint(new FImaginaryBlueprint(Blueprint->GetName(), Blueprint->GetPathName(), ParentClass, Interfaces, FFindInBlueprintSearchManager::Get().QuerySingleBlueprint(Blueprint)));
-			TSharedPtr< FFiBSearchInstance > SearchInstance(new FFiBSearchInstance);
-			FSearchResult SearchResult = SearchInstance->StartSearchQuery(SearchValue, ImaginaryBlueprint);
+void SFindInBlueprints::MatchTokensWithinBlueprint(const TArray<FString>& Tokens)
+{
+	if(StreamSearch.IsValid() && !StreamSearch->IsComplete())
+	{
+		StreamSearch->Stop();
+		StreamSearch->EnsureCompletion();
+		StreamSearch->GetFilteredItems(ItemsFound);
+		ItemsFound.Empty();
+		StreamSearch.Reset();
+	}
 
-			if (SearchResult.IsValid())
-			{
-				ItemsFound = SearchResult->Children;
-			}
+	// No multi-threading for single Blueprint search, just query for the data block and process it.
+	UBlueprint* Blueprint = BlueprintEditorPtr.Pin()->GetBlueprintObj();
 
-			if(ItemsFound.Num() == 0)
-			{
-				// Insert a fake result to inform user if none found
-				ItemsFound.Add(FSearchResult(new FFindInBlueprintsResult(LOCTEXT("BlueprintSearchNoResults", "No Results found"))));
-				HighlightText = FText::GetEmpty();
-			}
-			else
-			{
-				for(auto Item : ItemsFound)
-				{
-					Item->ExpandAllChildren(TreeView);
-				}
-			}
+	RootSearchResult = FSearchResult(new FFindInBlueprintsResult(FText::FromString(Blueprint->GetPathName())));
+	
+	FindInBlueprintsHelpers::Extract(FFindInBlueprintSearchManager::Get().QuerySingleBlueprint(Blueprint), Tokens, RootSearchResult);
 
-			TreeView->RequestTreeRefresh();
-		}
-		else
+	// We do not use the root, the user already knows what Blueprint they are searching within
+	ItemsFound = RootSearchResult->Children;
+	
+	if(ItemsFound.Num() == 0)
+	{
+		// Insert a fake result to inform user if none found
+		ItemsFound.Add(FSearchResult(new FFindInBlueprintsResult(LOCTEXT("BlueprintSearchNoResults", "No Results found"))));
+		HighlightText = FText::GetEmpty();
+		TreeView->RequestTreeRefresh();
+	}
+	else
+	{
+		for(auto Item : ItemsFound)
 		{
-			LaunchStreamThread(SearchValue);
+			Item->ExpandAllChildren(TreeView);
 		}
 	}
+
+	TreeView->RequestTreeRefresh();
 }
 
-void SFindInBlueprints::LaunchStreamThread(const FString& InSearchValue)
+void SFindInBlueprints::LaunchStreamThread(const TArray<FString>& InTokens)
 {
 	if(StreamSearch.IsValid() && !StreamSearch->IsComplete())
 	{
@@ -1122,7 +1481,7 @@ void SFindInBlueprints::LaunchStreamThread(const FString& InSearchValue)
 		RegisterActiveTimer( 0.f, FWidgetActiveTimerDelegate::CreateSP( this, &SFindInBlueprints::UpdateSearchResults ) );
 	}
 
-	StreamSearch = MakeShareable(new FStreamSearch(InSearchValue));
+	StreamSearch = MakeShareable(new FStreamSearch(InTokens));
 }
 
 TSharedRef<ITableRow> SFindInBlueprints::OnGenerateRow( FSearchResult InItem, const TSharedRef<STableViewBase>& OwnerTable )
@@ -1170,7 +1529,7 @@ TSharedRef<ITableRow> SFindInBlueprints::OnGenerateRow( FSearchResult InItem, co
 			[
 				SNew(SHorizontalBox)
 				+SHorizontalBox::Slot()
-				.VAlign(VAlign_Center)
+				.VAlign(VAlign_Top)
 				.AutoWidth()
 				[
 					InItem->CreateIcon()
@@ -1412,4 +1771,140 @@ void SFindInBlueprints::OnCopyAction()
 	FPlatformMisc::ClipboardCopy( *SelectedText );
 }
 
+////////////////////////////////////
+// FExtractionTask
+
+/**
+ * Task to search a batch of blueprints
+ *
+ * Each small batch is non-abandonable, but if the search is cancelled, no more batches are queued
+ */
+class FExtractionTask : public FNonAbandonableTask
+{
+public:
+	/** Async task function */
+	void DoWork();
+
+	/** Batch of blueprints and associated data to search */
+	TArray<FSearchData> QueryResultBatch;
+
+	/** Pointer to array of tokens to search for */
+	const TArray<FString>* Tokens;
+
+	/** Function to call once the search result is ready */
+	TFunction<void(const FSearchResult&)> OnResultReady;
+
+	FORCEINLINE TStatId GetStatId() const
+	{
+		RETURN_QUICK_DECLARE_CYCLE_STAT( FExtractionTask, STATGROUP_ThreadPoolAsyncTasks );
+	}
+};
+
+TArray<FSearchData> FStreamSearch::MakeExtractionBatch() const
+{
+	// Batch sizes between 10 and 30 seem fastest
+	static const int BatchSize = 15;
+
+	TArray<FSearchData> Batch;
+	do
+	{
+		if (StopTaskCounter.GetValue())
+			return TArray<FSearchData>();
+
+		FSearchData QueryResult;
+		if (!FFindInBlueprintSearchManager::Get().ContinueSearchQuery(this, QueryResult))
+			return Batch;
+
+		Batch.Push(QueryResult);
+	}
+	while (Batch.Num() <= BatchSize);
+
+	return Batch;
+}
+
+uint32 FStreamSearch::Run()
+{
+	FFindInBlueprintSearchManager::Get().BeginSearchQuery(this);
+
+	// Seems to go fastest at or slightly above one thread per physical core
+	TArray<FAsyncTask<FExtractionTask>> ExtractionTasks;
+	ExtractionTasks.SetNum(FPlatformMisc::NumberOfCores());
+
+	TFunction<void(const FSearchResult&)> OnResultReady = [this](const FSearchResult& Result) {
+		FScopeLock ScopeLock(&SearchCriticalSection);
+		ItemsFound.Add(Result);
+	};
+
+	// Searching comes to an end if it is requested using the StopTaskCounter or continuing the search query yields no results
+	TArray<FSearchData> QueryResultBatch;
+	for (;;)
+	{
+		if (QueryResultBatch.Num() == 0)
+		{
+			QueryResultBatch = MakeExtractionBatch();
+		}
+		else
+		{
+			// Spinning until a task becomes available
+			FPlatformProcess::Sleep(0);
+		}
+
+		if (QueryResultBatch.Num() == 0)
+		{
+			// All search work has been dispatched
+			break;
+		}
+
+		for (auto& Task: ExtractionTasks)
+		{
+			if (Task.IsWorkDone())
+			{
+				Task.EnsureCompletion();
+
+				auto& worker = Task.GetTask();
+				worker.QueryResultBatch = MoveTemp(QueryResultBatch);
+				worker.Tokens = &SearchTokens;
+				worker.OnResultReady = OnResultReady;
+
+				Task.StartBackgroundTask();
+
+				// Ensure batch array is empty (probably already is due to move above)
+				QueryResultBatch.Reset();
+				break;
+			}
+		}
+	}
+
+	// Wait for all work to finish
+	for (auto& Task: ExtractionTasks)
+	{
+		Task.EnsureCompletion();
+	}
+
+	if (StopTaskCounter.GetValue())
+	{
+		// Ensure that the FiB Manager knows that we are done searching
+		FFindInBlueprintSearchManager::Get().EnsureSearchQueryEnds(this);
+	}
+
+	bThreadCompleted = true;
+
+	return 0;
+}
+
+void FExtractionTask::DoWork()
+{
+	for (const auto& QueryResult: QueryResultBatch)
+	{
+		FSearchResult BlueprintCategory = FSearchResult(new FFindInBlueprintsResult(FText::FromName(QueryResult.BlueprintPath)));
+		FindInBlueprintsHelpers::Extract(QueryResult.Value, *Tokens, BlueprintCategory);
+
+		// If there are children, add the item to the search results
+		if(BlueprintCategory->Children.Num() != 0)
+		{
+			OnResultReady(BlueprintCategory);
+		}
+	}
+}
+
 #undef LOCTEXT_NAMESPACE
diff --git a/Engine/Source/Editor/Kismet/Private/ImaginaryBlueprintData.cpp b/Engine/Source/Editor/Kismet/Private/ImaginaryBlueprintData.cpp
deleted file mode 100644
diff --git a/Engine/Source/Editor/Kismet/Public/FindInBlueprintManager.h b/Engine/Source/Editor/Kismet/Public/FindInBlueprintManager.h
--- a/Engine/Source/Editor/Kismet/Public/FindInBlueprintManager.h
+++ b/Engine/Source/Editor/Kismet/Public/FindInBlueprintManager.h
@@ -45,12 +45,6 @@ struct KISMET_API FFindInBlueprintSearchTags
 	static const FText FiB_Description;
 	/** Comment tag */
 	static const FText FiB_Comment;
-	/** Path tag */
-	static const FText FiB_Path;
-	/** Parent Class tag */
-	static const FText FiB_ParentClass;
-	/** Interfaces tag */
-	static const FText FiB_Interfaces;
 
 	/** Pin type tags */
 
@@ -86,12 +80,6 @@ struct FSearchData
 	/** Search data block for the Blueprint */
 	FString Value;
 
-	/** Parent Class */
-	FString ParentClass;
-
-	/** Interfaces implemented by the Blueprint */
-	TArray<FString> Interfaces;
-
 	/** Cached to determine if the Blueprint is seen as no longer valid, allows it to be cleared out next save to disk */
 	bool bMarkedForDeletion;
 
@@ -107,8 +95,6 @@ struct FSearchData
 		: Blueprint(Other.Blueprint)
 		, BlueprintPath(MoveTemp(Other.BlueprintPath))
 		, Value(MoveTemp(Other.Value))
-		, ParentClass(MoveTemp(Other.ParentClass))
-		, Interfaces(MoveTemp(Other.Interfaces))
 		, bMarkedForDeletion(Other.bMarkedForDeletion)
 	{
 	}
@@ -124,8 +110,6 @@ struct FSearchData
 		BlueprintPath = MoveTemp(RHS.BlueprintPath);
 		Value = MoveTemp(RHS.Value);
 		bMarkedForDeletion = RHS.bMarkedForDeletion;
-		ParentClass = MoveTemp(RHS.ParentClass);
-		Interfaces = MoveTemp(RHS.Interfaces);
 		return *this;
 	}
 
@@ -133,8 +117,6 @@ struct FSearchData
 		: Blueprint(Other.Blueprint)
 		, BlueprintPath(Other.BlueprintPath)
 		, Value(Other.Value)
-		, ParentClass(Other.ParentClass)
-		, Interfaces(Other.Interfaces)
 		, bMarkedForDeletion(Other.bMarkedForDeletion)
 	{
 	}
@@ -150,8 +132,6 @@ struct FSearchData
 		BlueprintPath = RHS.BlueprintPath;
 		Value = RHS.Value;
 		bMarkedForDeletion = RHS.bMarkedForDeletion;
-		ParentClass = RHS.ParentClass;
-		Interfaces = RHS.Interfaces;
 		return *this;
 	}
 
diff --git a/Engine/Source/Editor/Kismet/Public/FindInBlueprints.h b/Engine/Source/Editor/Kismet/Public/FindInBlueprints.h
--- a/Engine/Source/Editor/Kismet/Public/FindInBlueprints.h
+++ b/Engine/Source/Editor/Kismet/Public/FindInBlueprints.h
@@ -5,37 +5,6 @@
 typedef TSharedPtr<class FFindInBlueprintsResult> FSearchResult;
 typedef STreeView<FSearchResult>  STreeViewType;
 
-/** Some utility functions to help with Find-in-Blueprint functionality */
-namespace FindInBlueprintsHelpers
-{
-	/** Looks up a JsonValue's FText from the passed lookup table */
-	FText AsFText(TSharedPtr< FJsonValue > InJsonValue, const TMap<int32, FText>& InLookupTable);
-
-	/** Looks up a JsonValue's FText from the passed lookup table */
-	FText AsFText(int32 InValue, const TMap<int32, FText>& InLookupTable);
-
-	bool IsTextEqualToString(FText InText, FString InString);
-
-	/**
-	 * Retrieves the pin type as a string value
-	 *
-	 * @param InPinType		The pin type to look at
-	 *
-	 * @return				The pin type as a string in format [category]'[sub-category object]'
-	 */
-	FString GetPinTypeAsString(const FEdGraphPinType& InPinType);
-
-	/**
-	 * Parses a pin type from passed in key names and values
-	 *
-	 * @param InKey					The key name for what the data should be translated as
-	 * @param InValue				Value to be be translated
-	 * @param InOutPinType			Modifies the PinType based on the passed parameters, building it up over multiple calls
-	 * @return						TRUE when the parsing is successful
-	 */
-	bool ParsePinType(FText InKey, FText InValue, FEdGraphPinType& InOutPinType);
-}
-
 /* Item that matched the search results */
 class FFindInBlueprintsResult : public TSharedFromThis< FFindInBlueprintsResult >
 {
@@ -66,6 +35,27 @@ public:
 	/** gets the blueprint housing all these search results */
 	UBlueprint* GetParentBlueprint() const;
 
+	/**
+	 * Extracts the content of a JsonNode, building the sub-hierarchy for display
+	 *
+	 * @param InLookupTable		Lookup table for converting Json strings into an FText
+	 * @param InJsonNode		Node to extract data from
+	 * @param InTokens			Tokens to match searchable content against
+	 * @return					Returns TRUE if the content matches search results and was extracted
+	 */
+	virtual bool ExtractContent(const TMap<int32, FText>& InLookupTable, TSharedPtr< FJsonObject > InJsonNode, const TArray<FString>& InTokens) { return ExtractContent(InLookupTable, InJsonNode, InTokens, NULL); };
+
+	/**
+	 * Extracts the content of a JsonNode, building the sub-hierarchy for display
+	 *
+	 * @param InLookupTable		Lookup table for converting Json strings into an FText
+	 * @param InJsonNode		Node to extract data from
+	 * @param InTokens			Tokens to match searchable content against
+	 * @param InParentOverride	Parent search result to attach sub-content to
+	 * @return					Returns TRUE if the content matches search results and was extracted
+	 */
+	virtual bool ExtractContent(const TMap<int32, FText>& InLookupTable, TSharedPtr< FJsonObject > InJsonNode, const TArray<FString>& InTokens, TSharedPtr< FFindInBlueprintsResult > InParentOverride);
+
 	/**
 	 * Parses search info for specific data important for displaying the search result in an easy to understand format
 	 *
@@ -73,8 +63,10 @@ public:
 	 * @param	InKey			This is the tag for the data, describing what it is so special handling can occur if needed
 	 * @param	InValue			Compared against search query to see if it passes the filter, sometimes data is rejected because it is deemed unsearchable
 	 * @param	InParent		The parent search result
+	 *
+	 * @return					TRUE if the item passes the search filter or is otherwise accepted
 	 */
-	virtual void ParseSearchInfo(FText InKey, FText InValue) {};
+	virtual bool ParseSearchInfo(const TArray<FString> &InTokens, FText InKey, FText InValue, TSharedPtr< FFindInBlueprintsResult > InParent);
 
 	/**
 	 * Adds extra search info, anything that does not have a predestined place in the search result. Adds a sub-item to the searches and formats its description so the tag displays
@@ -93,6 +85,19 @@ public:
 	/** Returns the display string for the row */
 	FText GetDisplayString() const;
 
+protected:
+
+	/**
+	 * Extracts the details of a JsonValue in the search results, pushing arrays into sub-objects correctly, breaking up objects, and passing on the rest of the types
+	 *
+	 * @param	InLookupTable			Lookup table for converting Json strings into an FText
+	 * @param	InTokens				The search tokens to check results against
+	 * @param	InKey					This is the tag for the data, describing what it is so special handling can occur if needed
+	 * @param	InValue					Has all the data for this object that needs to be deciphered or broken up into base types to be deciphered
+	 * @param	InParentOverride		The parent to attach all children data to, can be NULL.
+	 */
+	bool ExtractJsonValue(const TMap<int32, FText>& InLookupTable, const TArray<FString>& InTokens, FText InKey, TSharedPtr< FJsonValue > InValue, TSharedPtr< FFindInBlueprintsResult > InParentOverride);
+
 public:
 	/*Any children listed under this category */
 	TArray< TSharedPtr<FFindInBlueprintsResult> > Children;
@@ -117,7 +122,8 @@ public:
 	/** FFindInBlueprintsResult Interface */
 	virtual FReply OnClick() override;
 	virtual TSharedRef<SWidget>	CreateIcon() const override;
-	virtual void ParseSearchInfo(FText InKey, FText InValue) override;
+	virtual bool ParseSearchInfo(const TArray<FString> &InTokens, FText InKey, FText InValue, TSharedPtr< FFindInBlueprintsResult > InParentOverride = NULL) override;
+	virtual bool ExtractContent(const TMap<int32, FText>& InLookupTable, TSharedPtr< FJsonObject > InJsonNode, const TArray<FString>& InTokens) override;
 	virtual FText GetCategory() const override;
 	virtual void FinalizeSearchData() override;
 	/** End FFindInBlueprintsResult Interface */
@@ -148,7 +154,7 @@ public:
 
 	/** FFindInBlueprintsResult Interface */
 	virtual TSharedRef<SWidget>	CreateIcon() const override;
-	virtual void ParseSearchInfo(FText InKey, FText InValue) override;
+	virtual bool ParseSearchInfo(const TArray<FString> &InTokens, FText InKey, FText InValue, TSharedPtr< FFindInBlueprintsResult > InParentOverride = NULL) override;
 	virtual FText GetCategory() const override;
 	virtual void FinalizeSearchData() override;
 	/** End FFindInBlueprintsResult Interface */
@@ -159,6 +165,9 @@ protected:
 	/** The pin that this search result refers to */
 	FEdGraphPinType PinType;
 
+	/** The Pin's sub-category name */
+	FString PinSubCategory;
+
 	/** Pin's icon color */
 	FSlateColor IconColor;
 };
@@ -172,7 +181,7 @@ public:
 
 	/** FFindInBlueprintsResult Interface */
 	virtual TSharedRef<SWidget>	CreateIcon() const override;
-	virtual void ParseSearchInfo(FText InKey, FText InValue) override;
+	virtual bool ParseSearchInfo(const TArray<FString> &InTokens, FText InKey, FText InValue, TSharedPtr< FFindInBlueprintsResult > InParentOverride = NULL) override;
 	virtual FText GetCategory() const override;
 	virtual void FinalizeSearchData() override;
 	/** End FFindInBlueprintsResult Interface */
@@ -183,6 +192,9 @@ protected:
 	/** The default value of a property as a string */
 	FString DefaultValue;
 
+	/** The Property's sub-category name, for the pin display type */
+	FString PinSubCategory;
+
 	/** TRUE if the property is an SCS_Component */
 	bool bIsSCSComponent;
 };
@@ -197,8 +209,9 @@ public:
 	/** FFindInBlueprintsResult Interface */
 	virtual FReply OnClick() override;
 	virtual TSharedRef<SWidget>	CreateIcon() const override;
-	virtual void ParseSearchInfo(FText InKey, FText InValue) override;
+	virtual bool ParseSearchInfo(const TArray<FString> &InTokens, FText InKey, FText InValue, TSharedPtr< FFindInBlueprintsResult > InParentOverride = NULL) override;
 	virtual FText GetCategory() const override;
+	virtual bool ExtractContent(const TMap<int32, FText>& InLookupTable, TSharedPtr< FJsonObject > InJsonNode, const TArray<FString>& InTokens) override;
 	/** End FFindInBlueprintsResult Interface */
 protected:
 
@@ -252,8 +265,11 @@ private:
 	/** Begins the search based on the SearchValue */
 	void InitiateSearch();
 	
+	/* Find any results that contain all of the tokens, however, only in find within blueprints */
+	void MatchTokensWithinBlueprint(const TArray<FString>& Tokens);
+
 	/** Launches a thread for streaming more content into the results widget */
-	void LaunchStreamThread(const FString& InSearchValue);
+	void LaunchStreamThread(const TArray<FString>& InTokens);
 
 	/** Returns the percent complete on the search for the progress bar */
 	TOptional<float> GetPercentCompleteSearch() const;
diff --git a/Engine/Source/Editor/Kismet/Public/ImaginaryBlueprintData.h b/Engine/Source/Editor/Kismet/Public/ImaginaryBlueprintData.h
deleted file mode 100644
diff --git a/Engine/Source/Runtime/Core/Private/Misc/TextFilterExpressionEvaluator.cpp b/Engine/Source/Runtime/Core/Private/Misc/TextFilterExpressionEvaluator.cpp
--- a/Engine/Source/Runtime/Core/Private/Misc/TextFilterExpressionEvaluator.cpp
+++ b/Engine/Source/Runtime/Core/Private/Misc/TextFilterExpressionEvaluator.cpp
@@ -4,6 +4,117 @@
 #include "TextFilterExpressionEvaluator.h"
 #include "BasicMathExpressionEvaluator.h"
 
+namespace TextFilterExpressionParser
+{
+	class FTextToken
+	{
+	public:
+		enum class EInvertResult : uint8
+		{
+			No,
+			Yes,
+		};
+
+		FTextToken(FTextFilterString InString, ETextFilterTextComparisonMode InTextComparisonMode, const EInvertResult InInvertResult)
+			: String(MoveTemp(InString))
+			, TextComparisonMode(MoveTemp(InTextComparisonMode))
+			, InvertResult(InInvertResult)
+		{
+		}
+
+		FTextToken(const FTextToken& Other)
+			: String(Other.String)
+			, TextComparisonMode(Other.TextComparisonMode)
+			, InvertResult(Other.InvertResult)
+		{
+		}
+
+		FTextToken(FTextToken&& Other)
+			: String(MoveTemp(Other.String))
+			, TextComparisonMode(Other.TextComparisonMode)
+			, InvertResult(Other.InvertResult)
+		{
+		}
+
+		FTextToken& operator=(const FTextToken& Other)
+		{
+			String = Other.String;
+			TextComparisonMode = Other.TextComparisonMode;
+			InvertResult = Other.InvertResult;
+			return *this;
+		}
+
+		FTextToken& operator=(FTextToken&& Other)
+		{
+			String = MoveTemp(Other.String);
+			TextComparisonMode = Other.TextComparisonMode;
+			InvertResult = Other.InvertResult;
+			return *this;
+		}
+
+		const FTextFilterString& GetString() const
+		{
+			return String;
+		}
+
+		bool EvaluateAsBasicStringExpression(const ITextFilterExpressionContext* InContext) const
+		{
+			if (String.IsEmpty())
+			{
+				// An empty string is always considered to be true
+				return true;
+			}
+
+			const bool bMatched = InContext->TestBasicStringExpression(String, TextComparisonMode);
+			return (InvertResult == EInvertResult::No) ? bMatched : !bMatched;
+		}
+
+		bool EvaluateAsComplexExpression(const ITextFilterExpressionContext* InContext, const FTextFilterString& InKey, const ETextFilterComparisonOperation InComparisonOperation) const
+		{
+			const bool bMatched = InContext->TestComplexExpression(InKey.AsName(), String, InComparisonOperation, TextComparisonMode);
+			return (InvertResult == EInvertResult::No) ? bMatched : !bMatched;
+		}
+
+		ETextFilterTextComparisonMode GetTextComparisonMode() const
+		{
+			return TextComparisonMode;
+		}
+
+	private:
+		FTextFilterString String;
+		ETextFilterTextComparisonMode TextComparisonMode;
+		EInvertResult InvertResult;
+	};
+} // namespace TextFilterExpressionParser
+
+#define DEFINE_TEXT_EXPRESSION_OPERATOR_NODE(TYPE, MONIKER_COUNT, ...)	\
+	namespace TextFilterExpressionParser								\
+	{																	\
+		struct TYPE														\
+		{																\
+			static const int32 MonikerCount = MONIKER_COUNT;			\
+			static const TCHAR* Monikers[MonikerCount];					\
+		};																\
+	}																	\
+	DEFINE_EXPRESSION_NODE_TYPE(TextFilterExpressionParser::TYPE, __VA_ARGS__)
+
+DEFINE_TEXT_EXPRESSION_OPERATOR_NODE(FSubExpressionStart,		1,	0x0E7E1BC9, 0xF0B94D5D, 0x80F44D45, 0x85A082AA)
+DEFINE_TEXT_EXPRESSION_OPERATOR_NODE(FSubExpressionEnd,			1,	0x5E10956D, 0x2E17411F, 0xB6469E22, 0xB5E56E6C)
+
+DEFINE_TEXT_EXPRESSION_OPERATOR_NODE(FEqual,					3,	0x32457EFC, 0x4928406F, 0xBD78D943, 0x633797D1)
+DEFINE_TEXT_EXPRESSION_OPERATOR_NODE(FNotEqual,					2,	0x8F445A19, 0xF33443D9, 0x90D6DC85, 0xB0C5D071)
+DEFINE_TEXT_EXPRESSION_OPERATOR_NODE(FLess,						1,	0xB85E222B, 0x47E24E1F, 0xBC5A384D, 0x2FF471E1)
+DEFINE_TEXT_EXPRESSION_OPERATOR_NODE(FLessOrEqual,				2,	0x8C0A46B0, 0x8DAA4E2B, 0xA7FE4A23, 0xEF215918)
+DEFINE_TEXT_EXPRESSION_OPERATOR_NODE(FGreater,					1,	0x6A6247F4, 0xFD78467F, 0xA6AC1244, 0x0A31E0A5)
+DEFINE_TEXT_EXPRESSION_OPERATOR_NODE(FGreaterOrEqual,			2,	0x09D75C5E, 0xA29A4194, 0x8E8E5278, 0xC84365FD)
+
+DEFINE_TEXT_EXPRESSION_OPERATOR_NODE(FOr,						3,	0xF4778B51, 0xF535414D, 0x9C0EB5F2, 0x2F2B0FD4)
+DEFINE_TEXT_EXPRESSION_OPERATOR_NODE(FAnd,						3,	0x7511397A, 0x02D24DC2, 0x86729800, 0xF454C320)
+DEFINE_TEXT_EXPRESSION_OPERATOR_NODE(FNot,						2,	0x03D78990, 0x41D04E26, 0x8E98AD2F, 0x74667868)
+
+DEFINE_EXPRESSION_NODE_TYPE(TextFilterExpressionParser::FTextToken,	0x09E49538, 0x633545E3, 0x84B5644F, 0x1F11628F);
+DEFINE_EXPRESSION_NODE_TYPE(bool,									0xC1CD5DCF, 0x2AB44958, 0xB3FF4F8F, 0xE665D121);
+
 namespace TextFilterExpressionParser
 {
 	/**
@@ -336,52 +447,8 @@ bool FTextFilterExpressionEvaluator::SetFilterText(const FText& InFilterText)
 			{
 				// We inject AND operators between any adjacent text and sub-expression tokens to mimic the old search filter behavior
 				bool bWasPreviousTokenValidToInjectANDOperatorAfter = false;
-				for (int32 TokenIdx = 0; TokenIdx < TmpLexTokens.Num(); ++TokenIdx)
+				for (FExpressionToken& CurrentToken : TmpLexTokens)
 				{
-					FExpressionToken& CurrentToken = TmpLexTokens[TokenIdx];
-
-					if (TokenIdx + 3 < TmpLexTokens.Num())
-					{
-						if (const FTextToken* TextToken = CurrentToken.Node.Cast<FTextToken>())
-						{
-							FStringToken FunctionName(CurrentToken.Context);
-
-							FExpressionToken& ExpressionStartToken = TmpLexTokens[TokenIdx + 1];
-							if (ExpressionStartToken.Node.Cast<FSubExpressionStart>())
-							{
-								// We are inside a function, we need the token after the Start, but the TokenIdx is still on the function name
-								TokenIdx += 1;
-								FExpressionToken* NextToken = &TmpLexTokens[++TokenIdx];
-								FStringToken FunctionParameter(NextToken->Context);
-								int32 NestedCount = 1;
-								for (; TokenIdx < TmpLexTokens.Num(); ++TokenIdx) 
-								{
-									NextToken = &TmpLexTokens[TokenIdx];
-
-									if (NextToken->Node.Cast<FSubExpressionStart>())
-									{
-										++NestedCount;
-									}
-									else if (NextToken->Node.Cast<FSubExpressionEnd>())
-									{
-										--NestedCount;
-										if (NestedCount == 0)
-										{
-											break;
-										}
-									}
-									FunctionParameter.Accumulate(NextToken->Context);
-								}
-								FStringToken EmptyToken;
-								FinalTokens.Add(FExpressionToken(EmptyToken, FSubExpressionStart()));
-								FinalTokens.Add(FExpressionToken(FunctionName, FTextToken(FunctionName.GetString(), ETextFilterTextComparisonMode::Partial, FTextToken::EInvertResult::No)));
-								FinalTokens.Add(FExpressionToken(EmptyToken, FFunction()));
-								FinalTokens.Add(FExpressionToken(FunctionParameter, FTextToken(FunctionParameter.GetString(), ETextFilterTextComparisonMode::Partial, FTextToken::EInvertResult::No)));
-								FinalTokens.Add(FExpressionToken(EmptyToken, FSubExpressionEnd()));
-								continue;
-							}
-						}
-					}
 					// Key->value pairs require some form of comparison operator, so if we've found one of those, consider this to be a complex expression
 					if (!bIsComplexExpression && 
 						(CurrentToken.Node.Cast<FEqual>() || CurrentToken.Node.Cast<FNotEqual>() || CurrentToken.Node.Cast<FLess>() || CurrentToken.Node.Cast<FLessOrEqual>() || CurrentToken.Node.Cast<FGreater>() || CurrentToken.Node.Cast<FGreaterOrEqual>())
@@ -465,12 +532,7 @@ bool FTextFilterExpressionEvaluator::TestTextFilter(const ITextFilterExpressionC
 	return false;
 }
 
-void FTextFilterExpressionEvaluator::AddFunctionTokenCallback(FString InFunctionName, FTokenFunctionHandler InCallback)
-{
-	TokenFunctionHandlers.Add(InFunctionName, InCallback);
-}
-
-void FTextFilterExpressionEvaluator::SetupGrammar()
+void FTextFilterExpressionEvaluator::ConstructExpressionParser()
 {
 	using namespace TextFilterExpressionParser;
 	using TextFilterExpressionParser::FTextToken;
@@ -503,14 +565,7 @@ void FTextFilterExpressionEvaluator::SetupGrammar()
 	Grammar.DefineBinaryOperator<FEqual>(1);
 	Grammar.DefineBinaryOperator<FOr>(2);
 	Grammar.DefineBinaryOperator<FAnd>(2);
-	Grammar.DefineBinaryOperator<FFunction>(1);
 	Grammar.DefinePreUnaryOperator<FNot>();
-}
-
-void FTextFilterExpressionEvaluator::ConstructExpressionParser()
-{
-	using namespace TextFilterExpressionParser;
-	using TextFilterExpressionParser::FTextToken;
 
 	JumpTable.MapBinary<FLessOrEqual>([](const FTextToken& A, const FTextToken& B, const ITextFilterExpressionContext* InContext)		{ return B.EvaluateAsComplexExpression(InContext, A.GetString(), ETextFilterComparisonOperation::LessOrEqual); });
 	JumpTable.MapBinary<FLess>([](const FTextToken& A, const FTextToken& B, const ITextFilterExpressionContext* InContext)				{ return B.EvaluateAsComplexExpression(InContext, A.GetString(), ETextFilterComparisonOperation::Less); });
@@ -531,17 +586,6 @@ void FTextFilterExpressionEvaluator::ConstructExpressionParser()
 
 	JumpTable.MapPreUnary<FNot>([](const FTextToken& V, const ITextFilterExpressionContext* InContext)									{ return !V.EvaluateAsBasicStringExpression(InContext); });
 	JumpTable.MapPreUnary<FNot>([](bool V, const ITextFilterExpressionContext* InContext)												{ return !V; });
-
-	JumpTable.MapBinary<FFunction>(
-		[this](const TextFilterExpressionParser::FTextToken& A, const TextFilterExpressionParser::FTextToken& B, const ITextFilterExpressionContext* InContext)	
-	{
-		bool bResult = false;
-		if (FTokenFunctionHandler* FunctionCallback = TokenFunctionHandlers.Find(A.GetString().AsString()))
-		{
-			bResult = FunctionCallback->Execute(B.GetString());
-		}
-		return bResult;
-	});
 }
 
 bool FTextFilterExpressionEvaluator::EvaluateCompiledExpression(const ExpressionParser::CompileResultType& InCompiledResult, const ITextFilterExpressionContext& InContext, FText* OutErrorText) const
diff --git a/Engine/Source/Runtime/Core/Public/Misc/TextFilterExpressionEvaluator.h b/Engine/Source/Runtime/Core/Public/Misc/TextFilterExpressionEvaluator.h
--- a/Engine/Source/Runtime/Core/Public/Misc/TextFilterExpressionEvaluator.h
+++ b/Engine/Source/Runtime/Core/Public/Misc/TextFilterExpressionEvaluator.h
@@ -26,7 +26,7 @@ class ITextFilterExpressionContext
 public:
 	/** Test the given value against the strings extracted from the current item */
 	virtual bool TestBasicStringExpression(const FTextFilterString& InValue, const ETextFilterTextComparisonMode InTextComparisonMode) const = 0;
-
+		
 	/** Perform a complex expression test for the current item */
 	virtual bool TestComplexExpression(const FName& InKey, const FTextFilterString& InValue, const ETextFilterComparisonOperation InComparisonOperation, const ETextFilterTextComparisonMode InTextComparisonMode) const = 0;
 
@@ -34,125 +34,6 @@ protected:
 	virtual ~ITextFilterExpressionContext() {}
 };
 
-namespace TextFilterExpressionParser
-{
-	class FTextToken
-	{
-	public:
-		enum class EInvertResult : uint8
-		{
-			No,
-			Yes,
-		};
-
-		FTextToken(FTextFilterString InString, ETextFilterTextComparisonMode InTextComparisonMode, const EInvertResult InInvertResult)
-			: String(MoveTemp(InString))
-			, TextComparisonMode(MoveTemp(InTextComparisonMode))
-			, InvertResult(InInvertResult)
-		{
-		}
-
-		FTextToken(const FTextToken& Other)
-			: String(Other.String)
-			, TextComparisonMode(Other.TextComparisonMode)
-			, InvertResult(Other.InvertResult)
-		{
-		}
-
-		FTextToken(FTextToken&& Other)
-			: String(MoveTemp(Other.String))
-			, TextComparisonMode(Other.TextComparisonMode)
-			, InvertResult(Other.InvertResult)
-		{
-		}
-
-		FTextToken& operator=(const FTextToken& Other)
-		{
-			String = Other.String;
-			TextComparisonMode = Other.TextComparisonMode;
-			InvertResult = Other.InvertResult;
-			return *this;
-		}
-
-		FTextToken& operator=(FTextToken&& Other)
-		{
-			String = MoveTemp(Other.String);
-			TextComparisonMode = Other.TextComparisonMode;
-			InvertResult = Other.InvertResult;
-			return *this;
-		}
-
-		const FTextFilterString& GetString() const
-		{
-			return String;
-		}
-
-		bool EvaluateAsBasicStringExpression(const ITextFilterExpressionContext* InContext) const
-		{
-			if (String.IsEmpty())
-			{
-				// An empty string is always considered to be true
-				return true;
-			}
-
-			const bool bMatched = InContext->TestBasicStringExpression(String, TextComparisonMode);
-			return (InvertResult == EInvertResult::No) ? bMatched : !bMatched;
-		}
-
-		bool EvaluateAsComplexExpression(const ITextFilterExpressionContext* InContext, const FTextFilterString& InKey, const ETextFilterComparisonOperation InComparisonOperation) const
-		{
-			const bool bMatched = InContext->TestComplexExpression(InKey.AsName(), String, InComparisonOperation, TextComparisonMode);
-			return (InvertResult == EInvertResult::No) ? bMatched : !bMatched;
-		}
-
-		ETextFilterTextComparisonMode GetTextComparisonMode() const
-		{
-			return TextComparisonMode;
-		}
-
-	private:
-		FTextFilterString String;
-		ETextFilterTextComparisonMode TextComparisonMode;
-		EInvertResult InvertResult;
-	};
-
-	class FFunction
-	{
-
-	};
-} // namespace TextFilterExpressionParser
-
-#define DEFINE_TEXT_EXPRESSION_OPERATOR_NODE(TYPE, MONIKER_COUNT, ...)	\
-	namespace TextFilterExpressionParser								\
-	{																	\
-		struct TYPE														\
-		{																\
-			static const int32 MonikerCount = MONIKER_COUNT;			\
-			static const TCHAR* Monikers[MonikerCount];					\
-		};																\
-	}																	\
-DEFINE_EXPRESSION_NODE_TYPE(TextFilterExpressionParser::TYPE, __VA_ARGS__)
-
-DEFINE_TEXT_EXPRESSION_OPERATOR_NODE(FSubExpressionStart,		1,	0x0E7E1BC9, 0xF0B94D5D, 0x80F44D45, 0x85A082AA)
-DEFINE_TEXT_EXPRESSION_OPERATOR_NODE(FSubExpressionEnd,			1,	0x5E10956D, 0x2E17411F, 0xB6469E22, 0xB5E56E6C)
-
-DEFINE_TEXT_EXPRESSION_OPERATOR_NODE(FEqual,					3,	0x32457EFC, 0x4928406F, 0xBD78D943, 0x633797D1)
-DEFINE_TEXT_EXPRESSION_OPERATOR_NODE(FNotEqual,					2,	0x8F445A19, 0xF33443D9, 0x90D6DC85, 0xB0C5D071)
-DEFINE_TEXT_EXPRESSION_OPERATOR_NODE(FLess,						1,	0xB85E222B, 0x47E24E1F, 0xBC5A384D, 0x2FF471E1)
-DEFINE_TEXT_EXPRESSION_OPERATOR_NODE(FLessOrEqual,				2,	0x8C0A46B0, 0x8DAA4E2B, 0xA7FE4A23, 0xEF215918)
-DEFINE_TEXT_EXPRESSION_OPERATOR_NODE(FGreater,					1,	0x6A6247F4, 0xFD78467F, 0xA6AC1244, 0x0A31E0A5)
-DEFINE_TEXT_EXPRESSION_OPERATOR_NODE(FGreaterOrEqual,			2,	0x09D75C5E, 0xA29A4194, 0x8E8E5278, 0xC84365FD)
-
-DEFINE_TEXT_EXPRESSION_OPERATOR_NODE(FOr,						3,	0xF4778B51, 0xF535414D, 0x9C0EB5F2, 0x2F2B0FD4)
-DEFINE_TEXT_EXPRESSION_OPERATOR_NODE(FAnd,						3,	0x7511397A, 0x02D24DC2, 0x86729800, 0xF454C320)
-DEFINE_TEXT_EXPRESSION_OPERATOR_NODE(FNot,						2,	0x03D78990, 0x41D04E26, 0x8E98AD2F, 0x74667868)
-
-DEFINE_EXPRESSION_NODE_TYPE(TextFilterExpressionParser::FTextToken,	0x09E49538, 0x633545E3, 0x84B5644F, 0x1F11628F);
-DEFINE_EXPRESSION_NODE_TYPE(TextFilterExpressionParser::FFunction,	0x084E6214, 0x032FFA48, 0x9245ABF1, 0x9E248F1A);
-DEFINE_EXPRESSION_NODE_TYPE(bool,									0xC1CD5DCF, 0x2AB44958, 0xB3FF4F8F, 0xE665D121);
-
-DECLARE_DELEGATE_RetVal_OneParam(bool, FTokenFunctionHandler, const FTextFilterString&);
-
 /** Defines an expression evaluator that can be used to perform complex text filter queries */
 class CORE_API FTextFilterExpressionEvaluator
 {
@@ -161,7 +42,6 @@ public:
 	FTextFilterExpressionEvaluator(const ETextFilterExpressionEvaluatorMode InMode);
 	FTextFilterExpressionEvaluator(const FTextFilterExpressionEvaluator& Other);
 	FTextFilterExpressionEvaluator& operator=(const FTextFilterExpressionEvaluator& Other);
-	virtual ~FTextFilterExpressionEvaluator() {}
 
 	/** Get the complexity of the current filter terms */
 	ETextFilterExpressionType GetFilterType() const;
@@ -178,17 +58,12 @@ public:
 	/** Test our compiled filter using the given context */
 	bool TestTextFilter(const ITextFilterExpressionContext& InContext) const;
 
-	/** Helper function to add callbacks for function tokens */
-	void AddFunctionTokenCallback(FString InFunctionName, FTokenFunctionHandler InCallback);
-protected:
-	/** Sets up grammar used for evaluation */
-	void SetupGrammar();
-
+private:
 	/** Common function to construct the expression parser */
-	virtual void ConstructExpressionParser();
+	void ConstructExpressionParser();
 
 	/** Evaluate the given compiled result, and optionally populate OutErrorText with any error information */
-	virtual bool EvaluateCompiledExpression(const ExpressionParser::CompileResultType& InCompiledResult, const ITextFilterExpressionContext& InContext, FText* OutErrorText) const;
+	bool EvaluateCompiledExpression(const ExpressionParser::CompileResultType& InCompiledResult, const ITextFilterExpressionContext& InContext, FText* OutErrorText) const;
 
 	/** Defines whether or not the expression parser can evaluate complex expressions */
 	ETextFilterExpressionEvaluatorMode ExpressionEvaluatorMode;
@@ -205,9 +80,6 @@ protected:
 	/** The compiled filter created from the current filter text (compiled from FilterText) */
 	TOptional<ExpressionParser::CompileResultType> CompiledFilter;
 
-	/** Mapping of function names to their callbacks */
-	TMap<FString, FTokenFunctionHandler> TokenFunctionHandlers;
-
 	/** Expression parser */
 	FTokenDefinitions TokenDefinitions;
 	FExpressionGrammar Grammar;
