CIS fix

diff --git a/Engine/Source/Editor/Kismet/Private/SBlueprintDiff.cpp b/Engine/Source/Editor/Kismet/Private/SBlueprintDiff.cpp
--- a/Engine/Source/Editor/Kismet/Private/SBlueprintDiff.cpp
+++ b/Engine/Source/Editor/Kismet/Private/SBlueprintDiff.cpp
@@ -213,18 +213,18 @@ FCDODiffControl::FCDODiffControl(
 			if (IterOld != OldProperties.Num())
 			{
 				bool bFoundDifference = false;
-				if (const FSingleObjectDiffEntry* Differing = FindDiffering(DifferingProperties, OldProperties[IterOld]))
+				if (const FSingleObjectDiffEntry* OldDiffering = FindDiffering(DifferingProperties, OldProperties[IterOld]))
 				{
 					bFoundDifference = true;
-					OrderedProperties.Push(Differing);
+					OrderedProperties.Push(OldDiffering);
 					++IterOld;
 				}
 				else if (IterNew != NewProperties.Num())
 				{
-					if (const FSingleObjectDiffEntry* Differing = FindDiffering(DifferingProperties, NewProperties[IterNew]))
+					if (const FSingleObjectDiffEntry* NewDiffering = FindDiffering(DifferingProperties, NewProperties[IterNew]))
 					{
 						bFoundDifference = true;
-						OrderedProperties.Push(Differing);
+						OrderedProperties.Push(NewDiffering);
 						++IterNew;
 					}
 				}
