2549150 was a little overzealous, default values of inherited variables (blueprint or native) should still be editable.

diff --git a/Engine/Source/Editor/Kismet/Private/BlueprintDetailsCustomization.cpp b/Engine/Source/Editor/Kismet/Private/BlueprintDetailsCustomization.cpp
--- a/Engine/Source/Editor/Kismet/Private/BlueprintDetailsCustomization.cpp
+++ b/Engine/Source/Editor/Kismet/Private/BlueprintDetailsCustomization.cpp
@@ -715,16 +715,12 @@ void FBlueprintVarActionDetails::CustomizeDetails( IDetailLayoutBuilder& DetailL
 
 				IDetailPropertyRow* Row = DefaultValueCategory.AddExternalProperty(StructData, VariableProperty->GetFName());
 			}
-			else if (GetPropertyOwnerBlueprint())
+			else
 			{
 				// Things are in order, show the property and allow it to be edited
 				TArray<UObject*> ObjectList;
-				ObjectList.Add(GetPropertyOwnerBlueprint()->GeneratedClass->GetDefaultObject());
+				ObjectList.Add(Blueprint->GeneratedClass->GetDefaultObject());
 				IDetailPropertyRow* Row = DefaultValueCategory.AddExternalProperty(ObjectList, VariableProperty->GetFName());
-				if (Row != nullptr)
-				{
-					Row->IsEnabled(IsVariableInBlueprint());
-				}
 			}
 		}
 
