Fixes for disabling variable defaults in Blueprints when selecting nodes unless the variable is a part of the current BP or inherited.

diff --git a/Engine/Source/Editor/Kismet/Private/BlueprintDetailsCustomization.cpp b/Engine/Source/Editor/Kismet/Private/BlueprintDetailsCustomization.cpp
--- a/Engine/Source/Editor/Kismet/Private/BlueprintDetailsCustomization.cpp
+++ b/Engine/Source/Editor/Kismet/Private/BlueprintDetailsCustomization.cpp
@@ -715,12 +715,16 @@ void FBlueprintVarActionDetails::CustomizeDetails( IDetailLayoutBuilder& DetailL
 
 				IDetailPropertyRow* Row = DefaultValueCategory.AddExternalProperty(StructData, VariableProperty->GetFName());
 			}
-			else
+			else if (GetPropertyOwnerBlueprint())
 			{
 				// Things are in order, show the property and allow it to be edited
 				TArray<UObject*> ObjectList;
-				ObjectList.Add(Blueprint->GeneratedClass->GetDefaultObject());
+				ObjectList.Add(GetPropertyOwnerBlueprint()->GeneratedClass->GetDefaultObject());
 				IDetailPropertyRow* Row = DefaultValueCategory.AddExternalProperty(ObjectList, VariableProperty->GetFName());
+				if (Row != nullptr)
+				{
+					Row->IsEnabled(IsVariableInheritedByBlueprint());
+				}
 			}
 		}
 
diff --git a/Engine/Source/Editor/Kismet/Private/BlueprintDetailsCustomization.h b/Engine/Source/Editor/Kismet/Private/BlueprintDetailsCustomization.h
--- a/Engine/Source/Editor/Kismet/Private/BlueprintDetailsCustomization.h
+++ b/Engine/Source/Editor/Kismet/Private/BlueprintDetailsCustomization.h
@@ -190,6 +190,9 @@ private:
 
 	/** Returns TRUE if the Variable is in the current Blueprint */
 	bool IsVariableInBlueprint() const { return GetPropertyOwnerBlueprint() == GetBlueprintObj(); }
+
+	/** Returns TRUE if the Variable is inherited by the current Blueprint */
+	bool IsVariableInheritedByBlueprint() const { return GetBlueprintObj()->SkeletonGeneratedClass->IsChildOf(GetPropertyOwnerBlueprint()->SkeletonGeneratedClass); }
 private:
 	/** Pointer back to my parent tab */
 	TWeakPtr<SMyBlueprint> MyBlueprint;
