Fixes issue with appearing to be able to select actors in the viewport using Level Blueprints loaded through the Find-in-Blueprint system.

Double clicking on an actor reference node now validates if the world is a part of the GEditor's WorldContexts.

diff --git a/Engine/Source/Editor/Kismet/Private/BlueprintEditor.cpp b/Engine/Source/Editor/Kismet/Private/BlueprintEditor.cpp
--- a/Engine/Source/Editor/Kismet/Private/BlueprintEditor.cpp
+++ b/Engine/Source/Editor/Kismet/Private/BlueprintEditor.cpp
@@ -6224,7 +6224,17 @@ void FBlueprintEditor::OnNodeDoubleClicked(class UEdGraphNode* Node)
 		{
 			if (AActor* Actor = K2Node->GetReferencedLevelActor())
 			{
-				JumpToHyperlink(Actor);
+				// Check if the world is active in the editor. It's possible to open level BPs without formally opening
+				// the levels through Find-in-Blueprints
+				const TIndirectArray<FWorldContext>& WorldContextList = GEditor->GetWorldContexts();
+				for (const FWorldContext& WorldContext : WorldContextList)
+				{
+					if (WorldContext.World() == Actor->GetWorld())
+					{
+						JumpToHyperlink(Actor);
+						break;
+					}
+				}
 				bJumpToLevelActor = true;
 			}
 		}
