Fixes issue causing Find-in-all-Blueprints to not be able to load up Blueprints due to missing path.

diff --git a/Engine/Source/Editor/Kismet/Private/ImaginaryBlueprintData.cpp b/Engine/Source/Editor/Kismet/Private/ImaginaryBlueprintData.cpp
--- a/Engine/Source/Editor/Kismet/Private/ImaginaryBlueprintData.cpp
+++ b/Engine/Source/Editor/Kismet/Private/ImaginaryBlueprintData.cpp
@@ -312,7 +312,7 @@ FImaginaryBlueprint::FImaginaryBlueprint(FString InBlueprintName, FString InBlue
 
 FSearchResult FImaginaryBlueprint::CreateSearchResult_Internal(FSearchResult InParent) const
 {
-	return FSearchResult(new FFindInBlueprintsResult(ParsedTagsAndValues.Find(FFindInBlueprintSearchTags::FiB_Name.ToString())->GetDisplayText(LookupTable)));
+	return FSearchResult(new FFindInBlueprintsResult(ParsedTagsAndValues.Find(FFindInBlueprintSearchTags::FiB_Path.ToString())->GetDisplayText(LookupTable)));
 }
 
 bool FImaginaryBlueprint::IsCompatibleWithFilter(ESearchQueryFilter InSearchQueryFilter)
