Adds logging for tracking of issue.

diff --git a/Engine/Source/Editor/Kismet/Private/BlueprintDetailsCustomization.cpp b/Engine/Source/Editor/Kismet/Private/BlueprintDetailsCustomization.cpp
--- a/Engine/Source/Editor/Kismet/Private/BlueprintDetailsCustomization.cpp
+++ b/Engine/Source/Editor/Kismet/Private/BlueprintDetailsCustomization.cpp
@@ -193,6 +193,9 @@ void FBlueprintVarActionDetails::CustomizeDetails( IDetailLayoutBuilder& DetailL
 		return;
 	}
 
+	// Help in tracking UE-20390
+	UE_LOG(LogBlueprintUserMessages, Log, TEXT("CustomizeDetails for Blueprint '%s', cached variable property's owner is '%s'"), *GetBlueprintObj()->GetName(), *CachedVariableProperty->GetOwnerClass()->GetPathName());
+
 	CachedVariableName = GetVariableName();
 
 	TWeakPtr<FBlueprintEditor> BlueprintEditor = MyBlueprint.Pin()->GetBlueprintEditor();
@@ -1834,6 +1837,12 @@ void FBlueprintVarActionDetails::OnPostEditorRefresh()
 {
 	CachedVariableProperty = SelectionAsProperty();
 	CachedVariableName = GetVariableName();
+
+	// Help in tracking UE-20390
+	if (CachedVariableProperty.IsValid())
+	{
+		UE_LOG(LogBlueprintUserMessages, Log, TEXT("PostEditorRefresh for Blueprint '%s', cached variable property's owner is '%s'"), *GetBlueprintObj()->GetName(), *CachedVariableProperty->GetOwnerClass()->GetPathName());
+	}
 }
 
 EVisibility FBlueprintVarActionDetails::GetTransientVisibility() const
