UserDefinedStructs and Local Variables can have non-basic type arrays of size 1 without resetting to being size 0.

ArrayProperty will no longer assume "()" is an empty array. This was a bug.

GameplayTag pins will no longer store "()" for empty sized arrays.

UserDefinedStructureEditor now uses a normal DetailsView instead of a StructureDetails view, leveraging the same code that supports local variables in Blueprints.

diff --git a/Engine/Source/Editor/GameplayTagsEditor/Private/SGameplayTagGraphPin.cpp b/Engine/Source/Editor/GameplayTagsEditor/Private/SGameplayTagGraphPin.cpp
--- a/Engine/Source/Editor/GameplayTagsEditor/Private/SGameplayTagGraphPin.cpp
+++ b/Engine/Source/Editor/GameplayTagsEditor/Private/SGameplayTagGraphPin.cpp
@@ -141,18 +141,19 @@ void SGameplayTagGraphPin::RefreshTagList()
 	}
 
 	// Set Pin Data
-	FString TagString = TEXT("(");
+	FString TagString;
 	if (!TagName.IsEmpty())
 	{
+		TagString = TEXT("(");
 		TagString += TEXT("TagName=\"");
 		TagString += TagName;
 		TagString += TEXT("\"");
+		TagString += TEXT(")");
 	}
-	TagString += TEXT(")");
 	FString CurrentDefaultValue = GraphPinObj->GetDefaultAsString();
 	if (CurrentDefaultValue.IsEmpty())
 	{
-		CurrentDefaultValue = FString(TEXT("()"));
+		CurrentDefaultValue = FString(TEXT(""));
 	}
 	if (!CurrentDefaultValue.Equals(TagString))
 	{
diff --git a/Engine/Source/Editor/Kismet/Private/UserDefinedStructureEditor.cpp b/Engine/Source/Editor/Kismet/Private/UserDefinedStructureEditor.cpp
--- a/Engine/Source/Editor/Kismet/Private/UserDefinedStructureEditor.cpp
+++ b/Engine/Source/Editor/Kismet/Private/UserDefinedStructureEditor.cpp
@@ -15,69 +15,123 @@
 
 #define LOCTEXT_NAMESPACE "StructureEditor"
 
-class FStructureDefaultValueView : public FStructureEditorUtils::INotifyOnStructChanged, public TSharedFromThis<FStructureDefaultValueView>
+///////////////////////////////////////////////////////////////////////////////////////
+// FDefaultValueDetails
+
+class FDefaultValueDetails : public IDetailCustomization
 {
 public:
-	FStructureDefaultValueView(UUserDefinedStruct* EditedStruct) 
-		: UserDefinedStruct(EditedStruct)
+	/** Makes a new instance of this detail layout class for a specific detail view requesting it */
+	static TSharedRef<IDetailCustomization> MakeInstance(TSharedPtr<FStructOnScope> InStructData)
 	{
+		return MakeShareable(new FDefaultValueDetails(InStructData));
 	}
 
-	void Initialize()
+	FDefaultValueDetails(TSharedPtr<FStructOnScope> InStructData)
+		: StructData(InStructData)
+	{}
+
+	~FDefaultValueDetails()
 	{
-		StructData = MakeShareable(new FStructOnScope(GetUserDefinedStruct()));
-		FStructureEditorUtils::Fill_MakeStructureDefaultValue(GetUserDefinedStruct(), StructData->GetStructMemory());
+	}
 
-		FPropertyEditorModule& PropertyModule = FModuleManager::LoadModuleChecked<FPropertyEditorModule>("PropertyEditor");
-		FDetailsViewArgs ViewArgs;
-		ViewArgs.bAllowSearch = false;
-		ViewArgs.bHideSelectionTip = false;
-		ViewArgs.bShowActorLabel = false;
+	/** IDetailCustomization interface */
+	virtual void CustomizeDetails(class IDetailLayoutBuilder& DetailLayout) override;
 
-		FStructureDetailsViewArgs StructureViewArgs;
-		StructureViewArgs.bShowObjects = false;
-		StructureViewArgs.bShowAssets = true;
-		StructureViewArgs.bShowClasses = true;
-		StructureViewArgs.bShowInterfaces = false;
+	/** Callback when finished changing properties to export the default value from the property to where strings are stored */
+	void OnFinishedChangingProperties(const FPropertyChangedEvent& PropertyChangedEvent);
 
-		StructureDetailsView = PropertyModule.CreateStructureDetailView(ViewArgs, StructureViewArgs, StructData, LOCTEXT("DefaultValues", "Default Values"));
-		StructureDetailsView->GetOnFinishedChangingPropertiesDelegate().AddSP(this, &FStructureDefaultValueView::OnFinishedChangingProperties);
-	}
+private:
+	TWeakObjectPtr<UUserDefinedStruct> UserDefinedStruct;
+	TSharedPtr<class FUserDefinedStructureLayout> Layout;
+	TSharedPtr<FStructOnScope> StructData;
+	IDetailLayoutBuilder* DetailLayoutPtr;
+};
 
-	virtual ~FStructureDefaultValueView()
-	{
-	}
+void FDefaultValueDetails::CustomizeDetails(class IDetailLayoutBuilder& DetailLayout) 
+{
+	DetailLayoutPtr = &DetailLayout;
+	const TArray<TWeakObjectPtr<UObject>> Objects = DetailLayout.GetDetailsView().GetSelectedObjects();
+	check(Objects.Num() > 0);
 
-	void OnFinishedChangingProperties(const FPropertyChangedEvent& PropertyChangedEvent)
+	if (Objects.Num() == 1)
 	{
-		check(PropertyChangedEvent.MemberProperty && PropertyChangedEvent.MemberProperty->GetOwnerStruct());
+		UserDefinedStruct = CastChecked<UUserDefinedStruct>(Objects[0].Get());
+
+		IDetailsView& DetailsView = (IDetailsView&)(DetailLayout.GetDetailsView());
+		DetailsView.OnFinishedChangingProperties().AddSP(this, &FDefaultValueDetails::OnFinishedChangingProperties);
 
-		const bool bPropertyIsActual = PropertyChangedEvent.MemberProperty->GetOwnerStruct() == GetUserDefinedStruct();
-		const UProperty* DirectProperty = bPropertyIsActual ? PropertyChangedEvent.MemberProperty : nullptr;
-		while (DirectProperty && !Cast<const UUserDefinedStruct>(DirectProperty->GetOuter()))
+		IDetailCategoryBuilder& StructureCategory = DetailLayout.EditCategory("DefaultValues", LOCTEXT("DefaultValues", "Default Values"));
+
+		for (TFieldIterator<UProperty> PropertyIter(UserDefinedStruct.Get()); PropertyIter; ++PropertyIter)
 		{
-			DirectProperty = Cast<const UProperty>(DirectProperty->GetOuter());
+			StructureCategory.AddExternalProperty(StructData, (*PropertyIter)->GetFName());
 		}
-		ensure(nullptr != DirectProperty);
+	}
+}
+
+void FDefaultValueDetails::OnFinishedChangingProperties(const FPropertyChangedEvent& PropertyChangedEvent)
+{
+	check(PropertyChangedEvent.MemberProperty && PropertyChangedEvent.MemberProperty->GetOwnerStruct());
+
+	const bool bPropertyIsActual = PropertyChangedEvent.MemberProperty->GetOwnerStruct() == UserDefinedStruct.Get();
+	const UProperty* DirectProperty = bPropertyIsActual ? PropertyChangedEvent.MemberProperty : nullptr;
+	while (DirectProperty && !Cast<const UUserDefinedStruct>(DirectProperty->GetOuter()))
+	{
+		DirectProperty = Cast<const UProperty>(DirectProperty->GetOuter());
+	}
+	ensure(nullptr != DirectProperty);
 
-		if (DirectProperty)
+	if (DirectProperty)
+	{
+		FString DefaultValueString;
+		bool bDefaultValueSet = false;
 		{
-			FString DefaultValueString;
-			bool bDefaultValueSet = false;
+			if (StructData.IsValid() && StructData->IsValid())
 			{
-				if (StructData.IsValid() && StructData->IsValid())
-				{
-					bDefaultValueSet = FBlueprintEditorUtils::PropertyValueToString(DirectProperty, StructData->GetStructMemory(), DefaultValueString);
-				}
+				bDefaultValueSet = FBlueprintEditorUtils::PropertyValueToString(DirectProperty, StructData->GetStructMemory(), DefaultValueString);
 			}
+		}
 
-			const FGuid VarGuid = FStructureEditorUtils::GetGuidForProperty(DirectProperty);
-			if (bDefaultValueSet && VarGuid.IsValid())
-			{
-				FStructureEditorUtils::ChangeVariableDefaultValue(GetUserDefinedStruct(), VarGuid, DefaultValueString);
-			}
+		const FGuid VarGuid = FStructureEditorUtils::GetGuidForProperty(DirectProperty);
+		if (bDefaultValueSet && VarGuid.IsValid())
+		{
+			FStructureEditorUtils::ChangeVariableDefaultValue(UserDefinedStruct.Get(), VarGuid, DefaultValueString);
 		}
 	}
+}
+
+/////////////////////////////////
+// FStructureDefaultValueView
+
+class FStructureDefaultValueView : public FStructureEditorUtils::INotifyOnStructChanged, public TSharedFromThis<FStructureDefaultValueView>
+{
+public:
+	FStructureDefaultValueView(UUserDefinedStruct* EditedStruct) 
+		: UserDefinedStruct(EditedStruct)
+	{
+	}
+
+	void Initialize()
+	{
+		StructData = MakeShareable(new FStructOnScope(UserDefinedStruct.Get()));
+		FStructureEditorUtils::Fill_MakeStructureDefaultValue(UserDefinedStruct.Get(), StructData->GetStructMemory());
+
+		FPropertyEditorModule& PropertyModule = FModuleManager::LoadModuleChecked<FPropertyEditorModule>("PropertyEditor");
+		FDetailsViewArgs ViewArgs;
+		ViewArgs.bAllowSearch = false;
+		ViewArgs.bHideSelectionTip = false;
+		ViewArgs.bShowActorLabel = false;
+
+		DetailsView = PropertyModule.CreateDetailView(ViewArgs);
+		FOnGetDetailCustomizationInstance LayoutStructDetails = FOnGetDetailCustomizationInstance::CreateStatic(&FDefaultValueDetails::MakeInstance, StructData);
+		DetailsView->RegisterInstancedCustomPropertyLayout(UUserDefinedStruct::StaticClass(), LayoutStructDetails);
+		DetailsView->SetObject(UserDefinedStruct.Get());
+	}
+
+	virtual ~FStructureDefaultValueView()
+	{
+	}
 
 	UUserDefinedStruct* GetUserDefinedStruct()
 	{
@@ -86,41 +140,31 @@ public:
 
 	TSharedPtr<class SWidget> GetWidget()
 	{
-		return StructureDetailsView.IsValid() ? StructureDetailsView->GetWidget() : NULL;
+		return DetailsView;
 	}
 
 	virtual void PreChange(const class UUserDefinedStruct* Struct, FStructureEditorUtils::EStructureEditorChangeInfo Info) override
 	{
-		if (Struct && (GetUserDefinedStruct() == Struct))
-		{
-			if (StructureDetailsView.IsValid())
-			{
-				StructureDetailsView->SetStructureData(NULL);
-			}
-			if (StructData.IsValid())
-			{
-				StructData->Destroy();
-				StructData.Reset();
-			}
-		}
+		StructData->Destroy();
+		DetailsView->SetObject(nullptr);
+		DetailsView->OnFinishedChangingProperties().Clear();
 	}
 
 	virtual void PostChange(const class UUserDefinedStruct* Struct, FStructureEditorUtils::EStructureEditorChangeInfo Info) override
 	{
-		if (Struct && (GetUserDefinedStruct() == Struct))
-		{
-			StructData = MakeShareable(new FStructOnScope(Struct));
-			FStructureEditorUtils::Fill_MakeStructureDefaultValue(Struct, StructData->GetStructMemory());
-
-			if (StructureDetailsView.IsValid())
-			{
-				StructureDetailsView->SetStructureData(StructData);
-			}
-		}
+		StructData->Initialize(UserDefinedStruct.Get());
+		FStructureEditorUtils::Fill_MakeStructureDefaultValue(UserDefinedStruct.Get(), StructData->GetStructMemory());
+		DetailsView->SetObject(UserDefinedStruct.Get());
 	}
 private:
+
+	/** Struct on scope data that is being viewed in the details panel */
 	TSharedPtr<FStructOnScope> StructData;
-	TSharedPtr<class IStructureDetailsView> StructureDetailsView;
+
+	/** Details view being used for viewing the struct */
+	TSharedPtr<class IDetailsView> DetailsView;
+
+	/** User defined struct that is being represented */
 	const TWeakObjectPtr<UUserDefinedStruct> UserDefinedStruct;
 };
 
diff --git a/Engine/Source/Runtime/CoreUObject/Private/UObject/PropertyArray.cpp b/Engine/Source/Runtime/CoreUObject/Private/UObject/PropertyArray.cpp
--- a/Engine/Source/Runtime/CoreUObject/Private/UObject/PropertyArray.cpp
+++ b/Engine/Source/Runtime/CoreUObject/Private/UObject/PropertyArray.cpp
@@ -216,11 +216,7 @@ const TCHAR* UArrayProperty::ImportText_Internal( const TCHAR* Buffer, void* Dat
 
 	int32 Index = 0;
 
-	const bool bEmptyArray = *Buffer == TCHAR(')');
-	if (!bEmptyArray)
-	{
-		ArrayHelper.ExpandForIndex(0);
-	}
+	ArrayHelper.ExpandForIndex(0);
 	while ((Buffer != NULL) && (*Buffer != TCHAR(')')))
 	{
 		SkipWhitespace(Buffer);
diff --git a/Engine/Source/Runtime/CoreUObject/Public/UObject/Class.h b/Engine/Source/Runtime/CoreUObject/Public/UObject/Class.h
--- a/Engine/Source/Runtime/CoreUObject/Public/UObject/Class.h
+++ b/Engine/Source/Runtime/CoreUObject/Public/UObject/Class.h
@@ -1161,6 +1161,13 @@ public:
 		Destroy();
 	}
 
+	/** Re-initializes the scope with a specified UStruct */
+	void Initialize(TWeakObjectPtr<const UStruct> InScriptStruct)
+	{
+		ScriptStruct = InScriptStruct;
+		Initialize();
+	}
+
 private:
 	FStructOnScope(const FStructOnScope&);
 	FStructOnScope& operator=(const FStructOnScope&);
diff --git a/Engine/Source/Runtime/Engine/Classes/EdGraph/EdGraphPin.h b/Engine/Source/Runtime/Engine/Classes/EdGraph/EdGraphPin.h
--- a/Engine/Source/Runtime/Engine/Classes/EdGraph/EdGraphPin.h
+++ b/Engine/Source/Runtime/Engine/Classes/EdGraph/EdGraphPin.h
@@ -257,6 +257,7 @@ class UEdGraphPin : public UObject
 public:
 	// UObject interface
 	virtual bool IsSafeForRootSet() const override { return false; }
+	virtual void Serialize(FArchive& Ar) override;
 	// End UObject interface
 
 	/** Create a link. Note, this does not check that schema allows it, and will not break any existing connections */
diff --git a/Engine/Source/Runtime/Engine/Private/EdGraph/EdGraphPin.cpp b/Engine/Source/Runtime/Engine/Private/EdGraph/EdGraphPin.cpp
--- a/Engine/Source/Runtime/Engine/Private/EdGraph/EdGraphPin.cpp
+++ b/Engine/Source/Runtime/Engine/Private/EdGraph/EdGraphPin.cpp
@@ -302,6 +302,25 @@ const FString UEdGraphPin::GetLinkInfoString( const FString& InFunctionName, con
 #endif
 }
 
+void UEdGraphPin::Serialize(FArchive& Ar)
+{
+	Super::Serialize(Ar);
+
+	if (Ar.IsLoading() && Ar.IsPersistent())
+	{
+		// Pins of type FGameplayTag were storing "()" for empty arrays and then importing that into ArrayProperty and expecting an empty array.
+		// That it was working was a bug and has been fixed, so let's fixup pins. A pin that wants an array size of 1 will always fill the parenthesis
+		// so there is no worry about breaking those cases.
+		if (PinType.PinSubCategoryObject.IsValid() && PinType.PinSubCategoryObject->GetPathName() == TEXT("/Script/GameplayTags.GameplayTag"))
+		{
+			if (DefaultValue == TEXT("()"))
+			{
+				DefaultValue.Empty();
+			}
+		}
+	}
+}
+
 FString FGraphDisplayInfo::GetNotesAsString() const
 {
 	FString Result;
