Fix crash from calling FindFunctionByName on null SkeletonGeneratedClass

diff --git a/Engine/Source/Editor/BlueprintGraph/Private/EdGraphSchema_K2.cpp b/Engine/Source/Editor/BlueprintGraph/Private/EdGraphSchema_K2.cpp
--- a/Engine/Source/Editor/BlueprintGraph/Private/EdGraphSchema_K2.cpp
+++ b/Engine/Source/Editor/BlueprintGraph/Private/EdGraphSchema_K2.cpp
@@ -4413,7 +4413,7 @@ void UEdGraphSchema_K2::GetGraphDisplayInformation(const UEdGraph& Graph, /*out*
 
 	UFunction* Function = NULL;
 	UBlueprint* Blueprint = FBlueprintEditorUtils::FindBlueprintForGraph(&Graph);
-	if (Blueprint)
+	if (Blueprint && Blueprint->SkeletonGeneratedClass)
 	{
 		Function = Blueprint->SkeletonGeneratedClass->FindFunctionByName(Graph.GetFName());
 	}
diff --git a/Engine/Source/Editor/Kismet/Private/SMyBlueprint.cpp b/Engine/Source/Editor/Kismet/Private/SMyBlueprint.cpp
--- a/Engine/Source/Editor/Kismet/Private/SMyBlueprint.cpp
+++ b/Engine/Source/Editor/Kismet/Private/SMyBlueprint.cpp
@@ -1159,9 +1159,12 @@ void SMyBlueprint::CollectAllActions(FGraphActionListBuilderBase& OutAllActions)
 		Graph->GetSchema()->GetGraphDisplayInformation(*Graph, DisplayInfo);
 
 		FText FunctionCategory;
-		if(UFunction* Function = BlueprintObj->SkeletonGeneratedClass->FindFunctionByName(Graph->GetFName()))
+		if (BlueprintObj->SkeletonGeneratedClass)
 		{
-			FunctionCategory = Function->GetMetaDataText(FBlueprintMetadata::MD_FunctionCategory, TEXT("UObjectCategory"), Function->GetFullGroupName(false));
+			if(UFunction* Function = BlueprintObj->SkeletonGeneratedClass->FindFunctionByName(Graph->GetFName()))
+			{
+				FunctionCategory = Function->GetMetaDataText(FBlueprintMetadata::MD_FunctionCategory, TEXT("UObjectCategory"), Function->GetFullGroupName(false));
+			}
 		}
 
 		//@TODO: Should be a bit more generic (or the AnimGraph shouldn't be stored as a FunctionGraph...)
