Fixed bug with Find-in-Blueprint where double clicking on an item would not go to the item when searching locally.

diff --git a/Engine/Source/Editor/Kismet/Private/FindInBlueprints.cpp b/Engine/Source/Editor/Kismet/Private/FindInBlueprints.cpp
--- a/Engine/Source/Editor/Kismet/Private/FindInBlueprints.cpp
+++ b/Engine/Source/Editor/Kismet/Private/FindInBlueprints.cpp
@@ -1079,7 +1079,7 @@ void SFindInBlueprints::InitiateSearch()
 			}
 			TSharedPtr< FImaginaryBlueprint> ImaginaryBlueprint(new FImaginaryBlueprint(Blueprint->GetName(), Blueprint->GetPathName(), ParentClass, Interfaces, FFindInBlueprintSearchManager::Get().QuerySingleBlueprint(Blueprint)));
 			TSharedPtr< FFiBSearchInstance > SearchInstance(new FFiBSearchInstance);
-			FSearchResult SearchResult = SearchInstance->StartSearchQuery(SearchValue, ImaginaryBlueprint);
+			FSearchResult SearchResult = RootSearchResult = SearchInstance->StartSearchQuery(SearchValue, ImaginaryBlueprint);
 
 			if (SearchResult.IsValid())
 			{
