When adding new variables to user defined structures, the variable type will be the same as the last set type. This functions the same as variables/pins in Blueprints.

diff --git a/Engine/Source/Editor/Kismet/Private/UserDefinedStructureEditor.cpp b/Engine/Source/Editor/Kismet/Private/UserDefinedStructureEditor.cpp
--- a/Engine/Source/Editor/Kismet/Private/UserDefinedStructureEditor.cpp
+++ b/Engine/Source/Editor/Kismet/Private/UserDefinedStructureEditor.cpp
@@ -368,7 +368,10 @@ TSharedRef<SDockTab> FUserDefinedStructureEditor::SpawnStructureTab(const FSpawn
 class FUserDefinedStructureLayout : public IDetailCustomNodeBuilder, public TSharedFromThis<FUserDefinedStructureLayout>
 {
 public:
-	FUserDefinedStructureLayout(TWeakPtr<class FUserDefinedStructureDetails> InStructureDetails) : StructureDetails(InStructureDetails) {}
+	FUserDefinedStructureLayout(TWeakPtr<class FUserDefinedStructureDetails> InStructureDetails)
+		: StructureDetails(InStructureDetails)
+		, InitialPinType(GetDefault<UEdGraphSchema_K2>()->PC_Boolean, TEXT(""), NULL, false, false)
+	{}
 
 	void OnChanged()
 	{
@@ -380,9 +383,7 @@ public:
 		auto StructureDetailsSP = StructureDetails.Pin();
 		if(StructureDetailsSP.IsValid())
 		{
-			const UEdGraphSchema_K2* K2Schema = GetDefault<UEdGraphSchema_K2>();
-			const  FEdGraphPinType InitialType(K2Schema->PC_Boolean, TEXT(""), NULL, false, false);
-			FStructureEditorUtils::AddVariable(StructureDetailsSP->GetUserDefinedStruct(), InitialType);
+			FStructureEditorUtils::AddVariable(StructureDetailsSP->GetUserDefinedStruct(), InitialPinType);
 		}
 
 		return FReply::Handled();
@@ -452,6 +453,12 @@ public:
 		}
 	}
 
+	/** Callback when a pin type is selected to cache the value so new variables in the struct will be set to the cached type */
+	void OnPinTypeSelected(const FEdGraphPinType& InPinType)
+	{
+		InitialPinType = InPinType;
+	}
+
 	/** IDetailCustomNodeBuilder Interface*/
 	virtual void SetOnRebuildChildren( FSimpleDelegate InOnRegenerateChildren ) override 
 	{
@@ -480,6 +487,9 @@ public:
 private:
 	TWeakPtr<class FUserDefinedStructureDetails> StructureDetails;
 	FSimpleDelegate OnRegenerateChildren;
+
+	/** Cached value of the last pin type the user selected, used as the initial value for new struct members */
+	FEdGraphPinType InitialPinType;
 };
 
 enum EMemberFieldPosition
@@ -545,6 +555,10 @@ public:
 		if(StructureDetailsSP.IsValid())
 		{
 			FStructureEditorUtils::ChangeVariableType(StructureDetailsSP->GetUserDefinedStruct(), FieldGuid, PinType);
+			if (StructureLayout.IsValid())
+			{
+				StructureLayout.Pin()->OnPinTypeSelected(PinType);
+			}
 		}
 	}
 
