Fix for crash when deleting graphs or delegates while a "Create Event" node's graph is in the transaction buffer.

Added some safety checks when iterating over all these nodes to check for IsPendingKill on both the node and the node's graph.

diff --git a/Engine/Source/Editor/Kismet/Private/SMyBlueprint.cpp b/Engine/Source/Editor/Kismet/Private/SMyBlueprint.cpp
--- a/Engine/Source/Editor/Kismet/Private/SMyBlueprint.cpp
+++ b/Engine/Source/Editor/Kismet/Private/SMyBlueprint.cpp
@@ -2183,7 +2183,10 @@ void SMyBlueprint::OnDeleteGraph(UEdGraph* InGraph, EEdGraphSchemaAction_K2Graph
 		{
 			if (It->GetGraph() != InGraph)
 			{
-				It->HandleAnyChange();
+				if (!It->IsPendingKill() && It->GetGraph() && !It->GetGraph()->IsPendingKill())
+				{
+					It->HandleAnyChange();
+				}
 			}
 		}
 
@@ -2230,7 +2233,10 @@ void SMyBlueprint::OnDeleteDelegate(FEdGraphSchemaAction_K2Delegate* InDelegateA
 
 		for (TObjectIterator<UK2Node_CreateDelegate> It(RF_ClassDefaultObject | RF_PendingKill); It; ++It)
 		{
-			It->HandleAnyChange();
+			if (!It->IsPendingKill() && It->GetGraph() && !It->GetGraph()->IsPendingKill())
+			{
+				It->HandleAnyChange();
+			}
 		}
 	}
 }
