Added new blueprint alignment algorithm: straighten connections

This allows straightening of connections between all selected nodes in a graph, as well as specific connections between pins (by right clicking on the pin).
Existing alignment routines now align to the clicked-on node, where applicable, and use an average where this is not available.
Added icons for all blueprint alignment commands.

diff --git a/Engine/Source/Editor/BlueprintGraph/Classes/EdGraphSchema_K2.h b/Engine/Source/Editor/BlueprintGraph/Classes/EdGraphSchema_K2.h
--- a/Engine/Source/Editor/BlueprintGraph/Classes/EdGraphSchema_K2.h
+++ b/Engine/Source/Editor/BlueprintGraph/Classes/EdGraphSchema_K2.h
@@ -912,6 +912,12 @@ public:
 	/** Get menu for jumping to specific pin links */
 	void GetJumpToConnectionSubMenuActions(class FMenuBuilder& MenuBuilder, class UEdGraphPin* InGraphPin);
 
+	/** Get menu for straightening links to specific nodes*/
+	void GetStraightenConnectionToSubMenuActions( class FMenuBuilder& MenuBuilder, UEdGraphPin* InGraphPin ) const;
+
+	/** Get the destination pin for a straighten operation */
+	static UEdGraphPin* GetAndResetStraightenDestinationPin();
+
 	/** Create menu for variable get/set nodes which refer to a variable which does not exist. */
 	void GetNonExistentVariableMenu(const UEdGraphNode* InGraphNode, UBlueprint* OwnerBlueprint, FMenuBuilder* MenuBuilder) const;
 
diff --git a/Engine/Source/Editor/BlueprintGraph/Private/EdGraphSchema_K2.cpp b/Engine/Source/Editor/BlueprintGraph/Private/EdGraphSchema_K2.cpp
--- a/Engine/Source/Editor/BlueprintGraph/Private/EdGraphSchema_K2.cpp
+++ b/Engine/Source/Editor/BlueprintGraph/Private/EdGraphSchema_K2.cpp
@@ -1304,11 +1304,28 @@ void UEdGraphSchema_K2::GetContextMenuActions(const UEdGraph* CurrentGraph, cons
 							LOCTEXT("JumpToConnection", "Jump to Connection..."),
 							LOCTEXT("JumpToSpecificConnection", "Jump to specific connection..."),
 							FNewMenuDelegate::CreateUObject( (UEdGraphSchema_K2*const)this, &UEdGraphSchema_K2::GetJumpToConnectionSubMenuActions, const_cast<UEdGraphPin*>(InGraphPin)));
+
+						MenuBuilder->AddSubMenu(
+							LOCTEXT("StraightenConnection", "Straighten Connection To..."),
+							LOCTEXT("StraightenConnection_Tip", "Straighten a specific connection"),
+							FNewMenuDelegate::CreateUObject( this, &UEdGraphSchema_K2::GetStraightenConnectionToSubMenuActions, const_cast<UEdGraphPin*>(InGraphPin)));
 					}
 					else
 					{
 						((UEdGraphSchema_K2*const)this)->GetBreakLinkToSubMenuActions(*MenuBuilder, const_cast<UEdGraphPin*>(InGraphPin));
 						((UEdGraphSchema_K2*const)this)->GetJumpToConnectionSubMenuActions(*MenuBuilder, const_cast<UEdGraphPin*>(InGraphPin));
+						
+						UEdGraphPin* Pin = InGraphPin->LinkedTo[0];
+						FText PinName = Pin->GetDisplayName();
+						FText NodeName = Pin->GetOwningNode()->GetNodeTitle(ENodeTitleType::ListView);
+
+						MenuBuilder->AddMenuEntry(
+							FGraphEditorCommands::Get().StraightenConnections,
+							NAME_None,
+							FText::Format(LOCTEXT("StraightenDescription_SinglePin", "Straighten Connection to {0} ({1})"), NodeName, PinName),
+							FText::Format(LOCTEXT("StraightenDescription_SinglePin_Node_Tip", "Straighten the connection between this pin, and {0} ({1})"), NodeName, PinName),
+							FSlateIcon(NAME_None, NAME_None, NAME_None)
+						);
 					}
 				}
 	
@@ -1498,6 +1515,7 @@ void UEdGraphSchema_K2::GetContextMenuActions(const UEdGraph* CurrentGraph, cons
 						InMenuBuilder.AddMenuEntry( FGraphEditorCommands::Get().AlignNodesLeft );
 						InMenuBuilder.AddMenuEntry( FGraphEditorCommands::Get().AlignNodesCenter );
 						InMenuBuilder.AddMenuEntry( FGraphEditorCommands::Get().AlignNodesRight );
+						InMenuBuilder.AddMenuEntry( FGraphEditorCommands::Get().StraightenConnections );
 						InMenuBuilder.EndSection();
 
 						InMenuBuilder.BeginSection("EdGraphSchemaDistribution", LOCTEXT("DistributionHeader", "Distribution"));
@@ -1887,6 +1905,63 @@ void UEdGraphSchema_K2::GetJumpToConnectionSubMenuActions( class FMenuBuilder& M
 	}
 }
 
+// todo: this is a long way off ideal, but we can't pass context from our menu items onto the graph panel implementation
+// It'd be better to be able to pass context through to menu/ui commands
+namespace { UEdGraphPin* StraightenDestinationPin = nullptr; }
+UEdGraphPin* UEdGraphSchema_K2::GetAndResetStraightenDestinationPin()
+{
+	UEdGraphPin* Temp = StraightenDestinationPin;
+	StraightenDestinationPin = nullptr;
+	return Temp;
+}
+
+void UEdGraphSchema_K2::GetStraightenConnectionToSubMenuActions( class FMenuBuilder& MenuBuilder, UEdGraphPin* InGraphPin ) const
+{
+	auto MenuCommandList = MenuBuilder.GetTopCommandList();
+	if(!ensure(MenuCommandList.IsValid()))
+	{
+		return;
+	}
+
+	// Make sure we have a unique name for every entry in the list
+	TMap<FString, uint32> LinkTitleCount;
+
+	TMap<UEdGraphNode*, TArray<UEdGraphPin*>> NodeToPins;
+
+	for (UEdGraphPin* Pin : InGraphPin->LinkedTo)
+	{
+		UEdGraphNode* Node = Pin->GetOwningNode();
+		if (Node)
+		{
+			NodeToPins.FindOrAdd(Node).Add(Pin);
+		}
+	}
+
+	MenuBuilder.AddMenuEntry( FGraphEditorCommands::Get().StraightenConnections,
+		NAME_None, LOCTEXT("StraightenAllConnections", "All Connected Pins"),
+		TAttribute<FText>(), FSlateIcon(NAME_None, NAME_None, NAME_None) );
+
+	for (auto& Pair : NodeToPins)
+	{
+		FText NodeName = Pair.Key->GetNodeTitle(ENodeTitleType::ListView);
+		for (UEdGraphPin* Pin : Pair.Value)
+		{
+			FText PinName = Pin->GetDisplayName();
+			MenuBuilder.AddMenuEntry(
+				FText::Format(LOCTEXT("StraightenDescription_Node", "{0} ({1})"), NodeName, Pin->GetDisplayName()),
+				FText::Format(LOCTEXT("StraightenDescription_Node_Tip", "Straighten the connection between this pin, and {0} ({1})"), NodeName, Pin->GetDisplayName()),
+				FSlateIcon(),
+				FExecuteAction::CreateLambda([=]{
+				if (const FUIAction* UIAction = MenuCommandList->GetActionForCommand(FGraphEditorCommands::Get().StraightenConnections))
+				{
+					StraightenDestinationPin = Pin;
+					UIAction->ExecuteAction.Execute();
+				}
+			}));
+		}
+	}
+}
+
 const FPinConnectionResponse UEdGraphSchema_K2::DetermineConnectionResponseOfCompatibleTypedPins(const UEdGraphPin* PinA, const UEdGraphPin* PinB, const UEdGraphPin* InputPin, const UEdGraphPin* OutputPin) const
 {
 	// Now check to see if there are already connections and this is an 'exclusive' connection
diff --git a/Engine/Source/Editor/EditorStyle/Private/SlateEditorStyle.cpp b/Engine/Source/Editor/EditorStyle/Private/SlateEditorStyle.cpp
--- a/Engine/Source/Editor/EditorStyle/Private/SlateEditorStyle.cpp
+++ b/Engine/Source/Editor/EditorStyle/Private/SlateEditorStyle.cpp
@@ -4035,6 +4035,18 @@ void FSlateEditorStyle::FStyle::SetupGraphEditorStyles()
 
 		Set( "GraphEditor.GoToDocumentation", new IMAGE_BRUSH( "Common/icon_Help_Hover_16x", Icon16x16 ) );
 
+		Set( "GraphEditor.AlignNodesTop", new IMAGE_BRUSH( "Icons/GraphEditor/icon_AlignNodesTop_20px", Icon20x20 ) );
+		Set( "GraphEditor.AlignNodesMiddle", new IMAGE_BRUSH( "Icons/GraphEditor/icon_AlignNodesMiddle_20px", Icon20x20 ) );
+		Set( "GraphEditor.AlignNodesBottom", new IMAGE_BRUSH( "Icons/GraphEditor/icon_AlignNodesBottom_20px", Icon20x20 ) );
+		Set( "GraphEditor.AlignNodesLeft", new IMAGE_BRUSH( "Icons/GraphEditor/icon_AlignNodesLeft_20px", Icon20x20 ) );
+		Set( "GraphEditor.AlignNodesCenter", new IMAGE_BRUSH( "Icons/GraphEditor/icon_AlignNodesCenter_20px", Icon20x20 ) );
+		Set( "GraphEditor.AlignNodesRight", new IMAGE_BRUSH( "Icons/GraphEditor/icon_AlignNodesRight_20px", Icon20x20 ) );
+
+		Set( "GraphEditor.StraightenConnections", new IMAGE_BRUSH( "Icons/GraphEditor/icon_StraightenConnections_20px", Icon20x20 ) );
+
+		Set( "GraphEditor.DistributeNodesHorizontally", new IMAGE_BRUSH( "Icons/GraphEditor/icon_DistributeNodesHorizontally_20px", Icon20x20 ) );
+		Set( "GraphEditor.DistributeNodesVertically", new IMAGE_BRUSH( "Icons/GraphEditor/icon_DistributeNodesVertically_20px", Icon20x20 ) );
+
 		// Graph editor widgets
 		{
 			// EditableTextBox
diff --git a/Engine/Source/Editor/GraphEditor/Private/GraphEditorActions.cpp b/Engine/Source/Editor/GraphEditor/Private/GraphEditorActions.cpp
--- a/Engine/Source/Editor/GraphEditor/Private/GraphEditorActions.cpp
+++ b/Engine/Source/Editor/GraphEditor/Private/GraphEditorActions.cpp
@@ -49,6 +49,8 @@ void FGraphEditorCommandsImpl::RegisterCommands()
 	UI_COMMAND( AlignNodesCenter, "Align Center", "Aligns the horizontal centers of the selected nodes", EUserInterfaceActionType::Button, FInputChord() )
 	UI_COMMAND( AlignNodesRight, "Align Right", "Aligns the right edges of the selected nodes", EUserInterfaceActionType::Button, FInputChord() )
 
+	UI_COMMAND( StraightenConnections, "Straighten Connection(s)", "Straightens connections between the selected nodes.", EUserInterfaceActionType::Button, FInputChord() )
+
 	UI_COMMAND( DistributeNodesHorizontally, "Distribute Horizontally", "Evenly distributes the selected nodes horizontally", EUserInterfaceActionType::Button, FInputChord() )
 	UI_COMMAND( DistributeNodesVertically, "Distribute Vertically", "Evenly distributes the selected nodes vertically", EUserInterfaceActionType::Button, FInputChord() )
 	
diff --git a/Engine/Source/Editor/GraphEditor/Private/SGraphEditorImpl.cpp b/Engine/Source/Editor/GraphEditor/Private/SGraphEditorImpl.cpp
--- a/Engine/Source/Editor/GraphEditor/Private/SGraphEditorImpl.cpp
+++ b/Engine/Source/Editor/GraphEditor/Private/SGraphEditorImpl.cpp
@@ -150,6 +150,11 @@ UEdGraphPin* SGraphEditorImpl::GetGraphPinForMenu()
 	return GraphPinForMenu;
 }
 
+UEdGraphNode* SGraphEditorImpl::GetGraphNodeForMenu()
+{
+	return GraphNodeForMenu;
+}
+
 void SGraphEditorImpl::ZoomToFit(bool bOnlySelection)
 {
 	GraphPanel->ZoomToFit(bOnlySelection);
@@ -173,6 +178,16 @@ bool SGraphEditorImpl::GetBoundsForNode( const UEdGraphNode* InNode, class FSlat
 	return false;
 }
 
+void SGraphEditorImpl::StraightenConnections()
+{
+	GraphPanel->StraightenConnections();
+}
+
+void SGraphEditorImpl::StraightenConnections(UEdGraphPin* SourcePin, UEdGraphPin* PinToAlign)
+{
+	GraphPanel->StraightenConnections(SourcePin, PinToAlign);
+}
+
 void SGraphEditorImpl::Construct( const FArguments& InArgs )
 {
 	Commands = MakeShareable( new FUICommandList() );
@@ -436,7 +451,18 @@ FActionMenuContent SGraphEditorImpl::GraphEd_OnGetContextMenuFor(const FGraphCon
 			
 		// Cache the pin this menu is being brought up for
 		GraphPinForMenu = SpawnInfo.GraphPin;
-
+		GraphNodeForMenu = SpawnInfo.GraphNode;
+		
+		// Ensure we reset the GraphNodeForMenu when the menu closes
+		RegisterActiveTimer(0.1f, FWidgetActiveTimerDelegate::CreateLambda([this](double, float){
+			if (!FSlateApplication::Get().AnyMenusVisible())
+			{
+				GraphNodeForMenu = nullptr;
+				return EActiveTimerReturnType::Stop;
+			}
+			return EActiveTimerReturnType::Continue;
+		}));
+		
 		if ((SpawnInfo.GraphPin != NULL) || (SpawnInfo.GraphNode != NULL))
 		{
 			// Get all menu extenders for this context menu from the graph editor module
diff --git a/Engine/Source/Editor/GraphEditor/Private/SGraphEditorImpl.h b/Engine/Source/Editor/GraphEditor/Private/SGraphEditorImpl.h
--- a/Engine/Source/Editor/GraphEditor/Private/SGraphEditorImpl.h
+++ b/Engine/Source/Editor/GraphEditor/Private/SGraphEditorImpl.h
@@ -51,6 +51,7 @@ private:
 	TSharedPtr<SWidget>	TitleBar;
 
 	UEdGraphPin* GraphPinForMenu;
+	UEdGraphNode* GraphNodeForMenu;
 
 	/** Info on the appearance */
 	TAttribute<FGraphAppearanceInfo> Appearance;
@@ -129,6 +130,7 @@ public:
 	virtual void JumpToNode( const UEdGraphNode* JumpToMe, bool bRequestRename = false, bool bSelectNode = true ) override;
 	virtual void JumpToPin( const UEdGraphPin* JumpToMe ) override;
 	virtual UEdGraphPin* GetGraphPinForMenu() override;
+	virtual UEdGraphNode* GetGraphNodeForMenu() override;
 	virtual void ZoomToFit(bool bOnlySelection) override;
 	virtual bool GetBoundsForSelectedNodes( class FSlateRect& Rect, float Padding) override;
 	virtual bool GetBoundsForNode( const UEdGraphNode* InNode, class FSlateRect& Rect, float Padding) const override;
@@ -140,6 +142,8 @@ public:
 	virtual void UnlockFromGraphEditor(TWeakPtr<SGraphEditor> Other) override;
 	virtual void AddNotification ( FNotificationInfo& Info, bool bSuccess ) override;
 	virtual void SetPinVisibility(SGraphEditor::EPinVisibility Visibility) override;
+	virtual void StraightenConnections() override;
+	virtual void StraightenConnections(UEdGraphPin* SourcePin, UEdGraphPin* PinToAlign) override;
 	// End of SGraphEditor interface
 protected:
 	//
diff --git a/Engine/Source/Editor/GraphEditor/Private/SGraphPanel.cpp b/Engine/Source/Editor/GraphEditor/Private/SGraphPanel.cpp
--- a/Engine/Source/Editor/GraphEditor/Private/SGraphPanel.cpp
+++ b/Engine/Source/Editor/GraphEditor/Private/SGraphPanel.cpp
@@ -1091,6 +1091,268 @@ bool SGraphPanel::GetBoundsForNode(const UObject* InNode, FVector2D& MinCorner,
 	return SNodePanel::GetBoundsForNode(InNode, MinCorner, MaxCorner, Padding);
 }
 
+class FConnectionAligner
+{
+public:
+	void DefineConnection(UEdGraphNode* SourceNode, const TSharedPtr<SGraphPin>& SourcePin, UEdGraphNode* DestinationNode, const TSharedPtr<SGraphPin>& DestinationPin)
+	{
+		auto& Dependencies = Connections.FindOrAdd(SourceNode);
+		if (SourcePin->GetPinObj()->Direction == EEdGraphPinDirection::EGPD_Output)
+		{
+			Dependencies.Outputs.FindOrAdd(DestinationNode).Add(FPinPair{ SourcePin, DestinationPin });
+		}
+		else
+		{
+			Dependencies.Inputs.FindOrAdd(DestinationNode).Add(FPinPair{ SourcePin, DestinationPin });
+		}
+	}
+
+	/** Align all the connections */
+	void Process()
+	{
+		struct FRankedNode
+		{
+			UEdGraphNode* Node;
+			uint32 Rank;
+		};
+
+		TArray<FRankedNode> RankedNodes;
+		RankedNodes.Reserve(Connections.Num());
+
+		TMap<UEdGraphNode*, uint32> LongestChainCache;
+		LongestChainCache.Reserve(Connections.Num());
+
+		for (auto& Pair : Connections)
+		{
+			RankedNodes.Add(FRankedNode{ Pair.Key, CalculateNodeRank(Pair.Key, LongestChainCache) });
+		}
+
+		// Sort the nodes based on dependencies - highest is processed first
+		RankedNodes.Sort([](const FRankedNode& A, const FRankedNode& B){
+			return A.Rank > B.Rank;
+		});
+
+		TSet<UEdGraphNode*> VistedNodes;
+		for (FRankedNode& RankedNode : RankedNodes)
+		{
+			StraightenConnectionsForNode(RankedNode.Node, VistedNodes, EEdGraphPinDirection::EGPD_Output);
+			if (VistedNodes.Num() == RankedNodes.Num())
+			{
+				return;
+			}
+
+			StraightenConnectionsForNode(RankedNode.Node, VistedNodes, EEdGraphPinDirection::EGPD_Input);
+			if (VistedNodes.Num() == RankedNodes.Num())
+			{
+				return;
+			}
+		}
+	}
+
+private:
+
+	void StraightenConnectionsForNode(UEdGraphNode* Node, TSet<UEdGraphNode*>& VisitedNodes, EEdGraphPinDirection Direction)
+	{
+		FDependencyInfo* Info = Connections.Find(Node);
+		if (!Info)
+		{
+			return;
+		}
+
+		for (auto& NodeToPins : Info->GetDirection(Direction))
+		{
+			if (NodeToPins.Value.Num() == 0 || VisitedNodes.Contains(NodeToPins.Key))
+			{
+				continue;
+			}
+
+			// Align the averages of all the pins
+			float AlignmentDelta = 0.f;
+			for (const FPinPair& Pins : NodeToPins.Value)
+			{
+				AlignmentDelta += (Node->NodePosY + Pins.SrcPin->GetNodeOffset().Y) - (NodeToPins.Key->NodePosY + Pins.DstPin->GetNodeOffset().Y);
+			}
+
+			NodeToPins.Key->Modify();
+			NodeToPins.Key->NodePosY += AlignmentDelta / NodeToPins.Value.Num();
+
+			VisitedNodes.Add(Node);
+			VisitedNodes.Add(NodeToPins.Key);
+			
+			StraightenConnectionsForNode(NodeToPins.Key, VisitedNodes, Direction);
+		}
+	}
+
+	/** Find the longest chain of single-connection nodes connected to the specified node */
+	uint32 FindLongestUniqueChain(UEdGraphNode* Node, TMap<UEdGraphNode*, uint32>& LongestChainCache, EEdGraphPinDirection Direction)
+	{
+		if (uint32* Length = LongestChainCache.Find(Node))
+		{
+			// Already set, or circular dependency - ignore
+			return *Length;
+		}
+
+		// Prevent reentrancy
+		LongestChainCache.Add(Node, 0);
+
+		uint32 ThisLength = 0;
+		
+		if (FDependencyInfo* Dependencies = Connections.Find(Node))
+		{
+			auto& ConnectedNodes = Dependencies->GetDirection(Direction);
+
+			// We only follow unique (1-1) connections
+			if (ConnectedNodes.Num() == 1)
+			{
+				for (auto& NodeToPins : ConnectedNodes)
+				{
+					ThisLength = FindLongestUniqueChain(NodeToPins.Key, LongestChainCache, Direction) + 1;
+				}
+			}
+		}
+
+		LongestChainCache[Node] = ThisLength;
+		return ThisLength;
+	};
+
+	/** Calculate the depth of dependencies for the specified node */
+	uint32 CalculateNodeRank(UEdGraphNode* Node, TMap<UEdGraphNode*, uint32>& LongestChainCache)
+	{
+		uint32 Rank = 0;
+		if (FDependencyInfo* PinMap = Connections.Find(Node))
+		{
+			for (auto& NodeToPins : PinMap->Outputs)
+			{
+				Rank += FindLongestUniqueChain(NodeToPins.Key, LongestChainCache, EEdGraphPinDirection::EGPD_Output) + 1;
+			}
+			for (auto& NodeToPins : PinMap->Inputs)
+			{
+				Rank += FindLongestUniqueChain(NodeToPins.Key, LongestChainCache, EEdGraphPinDirection::EGPD_Input) + 1;
+			}
+		}
+		return Rank;
+	}
+
+private:
+
+	/** A pair of pins */
+	struct FPinPair
+	{
+		TSharedPtr<SGraphPin> SrcPin, DstPin;
+	};
+
+	/** Map of nodes and pins that are connected to the owning pin */
+	struct FDependencyInfo
+	{
+		TMap<UEdGraphNode*, TArray<FPinPair>> Outputs;
+		TMap<UEdGraphNode*, TArray<FPinPair>> Inputs;
+		uint32 Rank;
+
+		TMap<UEdGraphNode*, TArray<FPinPair>>& GetDirection(EEdGraphPinDirection Direction)
+		{
+			return Direction == EEdGraphPinDirection::EGPD_Output ? Outputs : Inputs;
+		}
+	};
+	typedef TMap<UEdGraphNode*, FDependencyInfo> FConnections;
+
+	FConnections Connections;
+};
+
+void SGraphPanel::StraightenConnections()
+{
+	FConnectionAligner Aligner;
+	for (auto* It : SelectionManager.SelectedNodes)
+	{
+		UEdGraphNode* SourceNode = Cast<UEdGraphNode>(It);
+		if (!SourceNode)
+		{
+			continue;
+		}
+
+		TSharedRef<SNode>* ThisNodePtr = NodeToWidgetLookup.Find(SourceNode);
+		if (!ThisNodePtr)
+		{
+			continue;
+		}
+
+		for (UEdGraphPin* SourcePin : SourceNode->Pins)
+		{
+			for (UEdGraphPin* LinkedTo : SourcePin->LinkedTo)
+			{
+				UEdGraphNode* DestNode = LinkedTo ? LinkedTo->GetOwningNode() : nullptr;
+				if (DestNode && SelectionManager.SelectedNodes.Contains(DestNode))
+				{
+					TSharedRef<SNode>* DestGraphNodePtr = NodeToWidgetLookup.Find(DestNode);
+					if (!DestGraphNodePtr)
+					{
+						continue;
+					}
+
+					TSharedPtr<SGraphPin> PinWidget = StaticCastSharedRef<SGraphNode>(*ThisNodePtr)->FindWidgetForPin(SourcePin);
+					TSharedPtr<SGraphPin> LinkedPinWidget = StaticCastSharedRef<SGraphNode>(*DestGraphNodePtr)->FindWidgetForPin(LinkedTo);
+					
+					if (PinWidget.IsValid() && LinkedPinWidget.IsValid())
+					{
+						Aligner.DefineConnection(SourceNode, PinWidget, DestNode, LinkedPinWidget);
+					}
+				}
+			}
+		}
+	}
+
+	Aligner.Process();
+}
+
+void SGraphPanel::StraightenConnections(UEdGraphPin* SourcePin, UEdGraphPin* PinToAlign)
+{
+	UEdGraphNode* OwningNode = SourcePin->GetOwningNode();
+
+	TSharedRef<SNode>* OwningNodeWidgetPtr = NodeToWidgetLookup.Find(OwningNode);
+	if (!OwningNodeWidgetPtr)
+	{
+		return;
+	}
+
+	TSharedRef<SGraphNode> SourceGraphNode = StaticCastSharedRef<SGraphNode>(*OwningNodeWidgetPtr);
+
+	FConnectionAligner Aligner;
+
+	auto AddConnectedPin = [&](UEdGraphPin* ConnectedPin){
+		UEdGraphNode* ConnectedNode = ConnectedPin ? ConnectedPin->GetOwningNode() : nullptr;
+		if (!ConnectedNode)
+		{
+			return;
+		}
+
+		TSharedRef<SNode>* DestGraphNodePtr = NodeToWidgetLookup.Find(ConnectedNode);
+		if (!DestGraphNodePtr)
+		{
+			return;
+		}
+
+		TSharedPtr<SGraphPin> PinWidget = SourceGraphNode->FindWidgetForPin(SourcePin);
+		TSharedPtr<SGraphPin> LinkedPinWidget = StaticCastSharedRef<SGraphNode>(*DestGraphNodePtr)->FindWidgetForPin(ConnectedPin);
+			
+		if (PinWidget.IsValid() && LinkedPinWidget.IsValid())
+		{
+			Aligner.DefineConnection(OwningNode, PinWidget, ConnectedNode, LinkedPinWidget);
+		}
+	};
+
+	if (PinToAlign)
+	{
+		// If we're only aligning a specific pin, do that
+		AddConnectedPin(PinToAlign);
+	}
+	// Else add all the connected pins
+	else for (UEdGraphPin* ConnectedPin : SourcePin->LinkedTo)
+	{
+		AddConnectedPin(ConnectedPin);
+	}
+	
+	Aligner.Process();
+}
+
 const TSharedRef<SGraphNode> SGraphPanel::GetChild(int32 ChildIndex)
 {
 	return StaticCastSharedRef<SGraphNode>(Children[ChildIndex]);
diff --git a/Engine/Source/Editor/GraphEditor/Private/SGraphPanel.h b/Engine/Source/Editor/GraphEditor/Private/SGraphPanel.h
--- a/Engine/Source/Editor/GraphEditor/Private/SGraphPanel.h
+++ b/Engine/Source/Editor/GraphEditor/Private/SGraphPanel.h
@@ -151,6 +151,12 @@ public:
 	/** Attempt to retrieve the bounds for the specified node */
 	bool GetBoundsForNode(const UObject* InNode, FVector2D& MinCorner, FVector2D& MaxCorner, float Padding = 0.0f) const;
 
+	/** Straighten all connections between the selected nodes */
+	void StraightenConnections();
+	
+	/** Straighten any connections attached to the specified pin, optionally limiting to the specified pin to align */
+	void StraightenConnections(UEdGraphPin* SourcePin, UEdGraphPin* PinToAlign = nullptr);
+
 protected:
 
 	void NotifyGraphChanged ( const struct FEdGraphEditAction& InAction);
diff --git a/Engine/Source/Editor/GraphEditor/Public/GraphEditorActions.h b/Engine/Source/Editor/GraphEditor/Public/GraphEditorActions.h
--- a/Engine/Source/Editor/GraphEditor/Public/GraphEditorActions.h
+++ b/Engine/Source/Editor/GraphEditor/Public/GraphEditorActions.h
@@ -69,6 +69,8 @@ public:
 	TSharedPtr< FUICommandInfo > AlignNodesCenter;
 	TSharedPtr< FUICommandInfo > AlignNodesRight;
 
+	TSharedPtr< FUICommandInfo > StraightenConnections;
+
 	TSharedPtr< FUICommandInfo > DistributeNodesHorizontally;
 	TSharedPtr< FUICommandInfo > DistributeNodesVertically;
 	
diff --git a/Engine/Source/Editor/Kismet/Private/BlueprintEditor.cpp b/Engine/Source/Editor/Kismet/Private/BlueprintEditor.cpp
--- a/Engine/Source/Editor/Kismet/Private/BlueprintEditor.cpp
+++ b/Engine/Source/Editor/Kismet/Private/BlueprintEditor.cpp
@@ -1032,6 +1032,10 @@ TSharedRef<SGraphEditor> FBlueprintEditor::CreateGraphEditorWidget(TSharedRef<FT
 				FExecuteAction::CreateSP( this, &FBlueprintEditor::OnAlignRight )
 				);
 
+			GraphEditorCommands->MapAction( FGraphEditorCommands::Get().StraightenConnections,
+				FExecuteAction::CreateSP( this, &FBlueprintEditor::OnStraightenConnections )
+				);
+
 			// Distribution Commands
 			GraphEditorCommands->MapAction( FGraphEditorCommands::Get().DistributeNodesHorizontally,
 				FExecuteAction::CreateSP( this, &FBlueprintEditor::OnDistributeNodesH )
@@ -4928,36 +4932,10 @@ struct FAlignmentData
 	}
 };
 
-/** Align the specified array of node data so all the target properties line up */
-void AlignNodes(const TArray<FAlignmentData>& InData)
+enum class EAlignType : uint8
 {
-	// Calculate the average target property...
-	float MeanPropertyValue = 0.f;
-	for (auto& Entry : InData)
-	{
-		MeanPropertyValue += Entry.GetTarget();
-	}
-	MeanPropertyValue /= InData.Num();
-
-	// ...and find the closest *actual* property to this average
-	float ThisTarget = InData[0].GetTarget();
-	float ClosestTarget = ThisTarget;
-	for (int32 Index = 1; Index < InData.Num(); ++Index)
-	{
-		ThisTarget = InData[Index].GetTarget();
-		if (FMath::Abs(ThisTarget - MeanPropertyValue) < FMath::Abs(ClosestTarget - MeanPropertyValue))
-		{
-			ClosestTarget = ThisTarget;
-		}
-	}
-
-	// Now set all the properties to the target
-	for (auto& Entry : InData)
-	{
-		Entry.Node->Modify();
-		Entry.TargetProperty = ClosestTarget - Entry.TargetOffset;
-	}
-}
+	Minimum, Middle, Maximum
+};
 
 FVector2D GetNodeSize(const SGraphEditor& GraphEditor, const UEdGraphNode* Node)
 {
@@ -4970,135 +4948,204 @@ FVector2D GetNodeSize(const SGraphEditor& GraphEditor, const UEdGraphNode* Node)
 	return FVector2D(Node->NodeWidth, Node->NodeHeight);
 }
 
-void FBlueprintEditor::OnAlignTop()
+/** Helper class for aligning nodes */
+struct FAlignmentHelper
 {
-	TArray<FAlignmentData> AlignData;
-	for (UObject* It : GetSelectedNodes())
+	/** Construct from a graph editor, an orientation, and an alignment type */
+	FAlignmentHelper(TSharedRef<SGraphEditor> InGraphEditor, EOrientation InOrientation, EAlignType InAlignType)
+		: GraphEditor(MoveTemp(InGraphEditor))
 	{
-		if (UEdGraphNode* Node = Cast<UEdGraphNode>(It))
+		// We align to the node that was clicked on, if available (not when invoked from a key shortcut)
+		CardinalNode = GraphEditor->GetGraphNodeForMenu();
+
+		Orientation = InOrientation;
+		AlignType = InAlignType;
+
+		// Collect all the alignment data for all the selected nodes
+		for (UObject* It : GraphEditor->GetSelectedNodes())
 		{
-			AlignData.Add(FAlignmentData(Node, Node->NodePosY, 0.f));
+			if (UEdGraphNode* Node = Cast<UEdGraphNode>(It))
+			{
+				AlignmentData.Add(GetAlignmentDataForNode(Node));
+			}
 		}
+
+		// Sort the data based on target - important for future algorithsm
+		AlignmentData.Sort([](const FAlignmentData& A, const FAlignmentData& B){
+			return A.GetTarget() < B.GetTarget();
+		});
 	}
 
-	if (AlignData.Num())
+	/** Align all the nodes */
+	void Align()
 	{
-		const FScopedTransaction Transaction(FGraphEditorCommands::Get().AlignNodesTop->GetLabel());
-		AlignNodes(AlignData);
+		if (AlignmentData.Num() > 1)
+		{
+			float Target = DetermineAlignmentTarget();
+
+			for (FAlignmentData& Entry : AlignmentData)
+			{
+				Entry.Node->Modify();
+				Entry.TargetProperty = Target - Entry.TargetOffset;
+			}
+		}
 	}
-}
 
-void FBlueprintEditor::OnAlignMiddle()
-{
-	TSharedPtr<SGraphEditor> FocusedGraphEd = FocusedGraphEdPtr.Pin();
-	if (!FocusedGraphEd.IsValid())
+private:
+
+	/** Collect alignment data for a given node, based on our settings */
+	FAlignmentData GetAlignmentDataForNode(UEdGraphNode* Node)
 	{
-		return;
+		float PropertyOffset = 0.f;
+
+		const float NodeSize = Orientation == Orient_Horizontal ? GetNodeSize(*GraphEditor, Node).X : GetNodeSize(*GraphEditor, Node).Y;
+		switch(AlignType)
+		{
+			case EAlignType::Minimum:	PropertyOffset = 0.f;			break;
+			case EAlignType::Middle:	PropertyOffset = NodeSize*.5f;	break;
+			case EAlignType::Maximum:	PropertyOffset = NodeSize;		break;
+		}
+		int32* Property = Orientation == Orient_Horizontal ? &Node->NodePosX : &Node->NodePosY;
+		return FAlignmentData(Node, *Property, PropertyOffset);
 	}
 
-	TArray<FAlignmentData> AlignData;
-	for (UObject* It : GetSelectedNodes())
+	/** Determine the horizontal/vertical position that all nodes should align to */
+	float DetermineAlignmentTarget()
 	{
-		if (UEdGraphNode* Node = Cast<UEdGraphNode>(It))
+		if (CardinalNode)
+		{
+			return GetAlignmentDataForNode(CardinalNode).GetTarget();
+		}
+
+		if (AlignType == EAlignType::Minimum)
+		{
+			float Target = TNumericLimits<float>::Max();
+			for (auto& Entry : AlignmentData)
+			{
+				Target = FMath::Min(Target, Entry.GetTarget());
+			}
+			return Target;
+		}
+		else if (AlignType == EAlignType::Maximum)
+		{
+			float Target = TNumericLimits<float>::Lowest();
+			for (auto& Entry : AlignmentData)
+			{
+				Target = FMath::Max(Target, Entry.GetTarget());
+			}
+			return Target;
+		}
+		else
 		{
-			AlignData.Add(FAlignmentData(Node, Node->NodePosY, GetNodeSize(*FocusedGraphEd, Node).Y*.5f));
+			// Use the mean
+			float SumTotal = 0.f;
+			for (auto& Entry : AlignmentData)
+			{
+				SumTotal += Entry.GetTarget();
+			}
+			return SumTotal / AlignmentData.Num();
 		}
 	}
 
-	if (AlignData.Num())
+	/** The graph editor */
+	TSharedRef<SGraphEditor> GraphEditor;
+	/** Whether we are aligning horizontally/vertically */
+	EOrientation Orientation;
+	/** Whether we are aligning to the minimum/middle/maximum bounds */
+	EAlignType AlignType;
+	/** The cardinal node that all other nodes should align to (possibly null) */
+	UEdGraphNode* CardinalNode;
+	/** Generated alignment data */
+	TArray<FAlignmentData> AlignmentData;
+};
+
+void FBlueprintEditor::OnAlignTop()
+{
+	TSharedPtr<SGraphEditor> FocusedGraphEd = FocusedGraphEdPtr.Pin();
+	if (FocusedGraphEd.IsValid())
 	{
-		const FScopedTransaction Transaction(FGraphEditorCommands::Get().AlignNodesMiddle->GetLabel());
-		AlignNodes(AlignData);
+		const FScopedTransaction Transaction(FGraphEditorCommands::Get().AlignNodesTop->GetLabel());
+
+		FAlignmentHelper Helper(FocusedGraphEd.ToSharedRef(), Orient_Vertical, EAlignType::Minimum);
+		Helper.Align();
 	}
 }
 
-void FBlueprintEditor::OnAlignBottom()
+void FBlueprintEditor::OnAlignMiddle()
 {
 	TSharedPtr<SGraphEditor> FocusedGraphEd = FocusedGraphEdPtr.Pin();
-	if (!FocusedGraphEd.IsValid())
+	if (FocusedGraphEd.IsValid())
 	{
-		return;
-	}
+		const FScopedTransaction Transaction(FGraphEditorCommands::Get().AlignNodesMiddle->GetLabel());
 
-	TArray<FAlignmentData> AlignData;
-	for (UObject* It : GetSelectedNodes())
-	{
-		if (UEdGraphNode* Node = Cast<UEdGraphNode>(It))
-		{
-			AlignData.Add(FAlignmentData(Node, Node->NodePosY, GetNodeSize(*FocusedGraphEd, Node).Y));
-		}
+		FAlignmentHelper Helper(FocusedGraphEd.ToSharedRef(), Orient_Vertical, EAlignType::Middle);
+		Helper.Align();
 	}
+}
 
-	if (AlignData.Num())
+void FBlueprintEditor::OnAlignBottom()
+{
+	TSharedPtr<SGraphEditor> FocusedGraphEd = FocusedGraphEdPtr.Pin();
+	if (FocusedGraphEd.IsValid())
 	{
 		const FScopedTransaction Transaction(FGraphEditorCommands::Get().AlignNodesBottom->GetLabel());
-		AlignNodes(AlignData);
+
+		FAlignmentHelper Helper(FocusedGraphEd.ToSharedRef(), Orient_Vertical, EAlignType::Maximum);
+		Helper.Align();
 	}
 }
 
 void FBlueprintEditor::OnAlignLeft()
 {
-	TArray<FAlignmentData> AlignData;
-	for (UObject* It : GetSelectedNodes())
-	{
-		if (UEdGraphNode* Node = Cast<UEdGraphNode>(It))
-		{
-			AlignData.Add(FAlignmentData(Node, Node->NodePosX, 0.f));
-		}
-	}
-
-	if (AlignData.Num())
+	TSharedPtr<SGraphEditor> FocusedGraphEd = FocusedGraphEdPtr.Pin();
+	if (FocusedGraphEd.IsValid())
 	{
 		const FScopedTransaction Transaction(FGraphEditorCommands::Get().AlignNodesLeft->GetLabel());
-		AlignNodes(AlignData);
+
+		FAlignmentHelper Helper(FocusedGraphEd.ToSharedRef(), Orient_Horizontal, EAlignType::Minimum);
+		Helper.Align();
 	}
 }
 
 void FBlueprintEditor::OnAlignCenter()
 {
 	TSharedPtr<SGraphEditor> FocusedGraphEd = FocusedGraphEdPtr.Pin();
-	if (!FocusedGraphEd.IsValid())
-	{
-		return;
-	}
-	
-	TArray<FAlignmentData> AlignData;
-	for (UObject* It : GetSelectedNodes())
-	{
-		if (UEdGraphNode* Node = Cast<UEdGraphNode>(It))
-		{
-			AlignData.Add(FAlignmentData(Node, Node->NodePosX, GetNodeSize(*FocusedGraphEd, Node).X*.5f));
-		}
-	}
-
-	if (AlignData.Num())
+	if (FocusedGraphEd.IsValid())
 	{
 		const FScopedTransaction Transaction(FGraphEditorCommands::Get().AlignNodesCenter->GetLabel());
-		AlignNodes(AlignData);
+
+		FAlignmentHelper Helper(FocusedGraphEd.ToSharedRef(), Orient_Horizontal, EAlignType::Middle);
+		Helper.Align();
 	}
 }
 
 void FBlueprintEditor::OnAlignRight()
 {
 	TSharedPtr<SGraphEditor> FocusedGraphEd = FocusedGraphEdPtr.Pin();
-	if (!FocusedGraphEd.IsValid())
+	if (FocusedGraphEd.IsValid())
 	{
-		return;
+		const FScopedTransaction Transaction(FGraphEditorCommands::Get().AlignNodesRight->GetLabel());
+
+		FAlignmentHelper Helper(FocusedGraphEd.ToSharedRef(), Orient_Horizontal, EAlignType::Maximum);
+		Helper.Align();
 	}
-	
-	TArray<FAlignmentData> AlignData;
-	for (UObject* It : GetSelectedNodes())
+}
+
+void FBlueprintEditor::OnStraightenConnections()
+{
+	TSharedPtr<SGraphEditor> FocusedGraphEd = FocusedGraphEdPtr.Pin();
+	if (FocusedGraphEd.IsValid())
 	{
-		if (UEdGraphNode* Node = Cast<UEdGraphNode>(It))
+		const FScopedTransaction Transaction(FGraphEditorCommands::Get().StraightenConnections->GetLabel());
+
+		if (UEdGraphPin* Pin = FocusedGraphEd->GetGraphPinForMenu())
 		{
-			AlignData.Add(FAlignmentData(Node, Node->NodePosX, GetNodeSize(*FocusedGraphEd, Node).X));
+			FocusedGraphEd->StraightenConnections(Pin, UEdGraphSchema_K2::GetAndResetStraightenDestinationPin());
+		}
+		else
+		{
+			FocusedGraphEd->StraightenConnections();
 		}
-	}
-
-	if (AlignData.Num())
-	{
-		const FScopedTransaction Transaction(FGraphEditorCommands::Get().AlignNodesRight->GetLabel());
-		AlignNodes(AlignData);
 	}
 }
 
diff --git a/Engine/Source/Editor/Kismet/Public/BlueprintEditor.h b/Engine/Source/Editor/Kismet/Public/BlueprintEditor.h
--- a/Engine/Source/Editor/Kismet/Public/BlueprintEditor.h
+++ b/Engine/Source/Editor/Kismet/Public/BlueprintEditor.h
@@ -720,6 +720,8 @@ protected:
 	void OnAlignCenter();
 	void OnAlignRight();
 	
+	void OnStraightenConnections();
+
 	void OnDistributeNodesH();
 	void OnDistributeNodesV();
 
diff --git a/Engine/Source/Editor/UnrealEd/Public/GraphEditor.h b/Engine/Source/Editor/UnrealEd/Public/GraphEditor.h
--- a/Engine/Source/Editor/UnrealEd/Public/GraphEditor.h
+++ b/Engine/Source/Editor/UnrealEd/Public/GraphEditor.h
@@ -302,6 +302,18 @@ public:
 		}
 	}
 
+	virtual class UEdGraphNode* GetGraphNodeForMenu()
+	{
+		if ( Implementation.IsValid() )
+		{
+			return Implementation->GetGraphNodeForMenu();
+		}
+		else
+		{
+			return NULL;
+		}
+	}
+
 	// Zooms out to fit either all nodes or only the selected ones
 	virtual void ZoomToFit(bool bOnlySelection)
 	{
@@ -331,6 +343,22 @@ public:
 		return false;
 	}
 
+	virtual void StraightenConnections()
+	{
+		if (Implementation.IsValid())
+		{
+			return Implementation->StraightenConnections();
+		}
+	}
+
+	virtual void StraightenConnections(UEdGraphPin* SourcePin, UEdGraphPin* PinToAlign = nullptr)
+	{
+		if (Implementation.IsValid())
+		{
+			return Implementation->StraightenConnections(SourcePin, PinToAlign);
+		}
+	}
+
 	// Invoked to let this widget know that the GraphEditor module has been reloaded
 	UNREALED_API void OnModuleReloaded();
 
