Cleaning out programmer-only comments that were harvested by Doxygen.

diff --git a/Engine/Source/Editor/BlueprintGraph/Classes/K2Node_ExecutionSequence.h b/Engine/Source/Editor/BlueprintGraph/Classes/K2Node_ExecutionSequence.h
--- a/Engine/Source/Editor/BlueprintGraph/Classes/K2Node_ExecutionSequence.h
+++ b/Engine/Source/Editor/BlueprintGraph/Classes/K2Node_ExecutionSequence.h
@@ -10,23 +10,23 @@ class UK2Node_ExecutionSequence : public UK2Node
 {
 	GENERATED_UCLASS_BODY()
 
-	// Begin UEdGraphNode interface
+	//~ Begin UEdGraphNode Interface
 	virtual void AllocateDefaultPins() override;
 	virtual FText GetTooltipText() const override;
 	virtual FLinearColor GetNodeTitleColor() const override;
 	virtual FText GetNodeTitle(ENodeTitleType::Type TitleType) const override;
 	virtual FName GetPaletteIcon(FLinearColor& OutColor) const override{ return TEXT("GraphEditor.Sequence_16x"); }
-	// End UEdGraphNode interface
+	//~ End UEdGraphNode Interface
 
-	// Begin UK2Node interface
+	//~ Begin UK2Node Interface
 	void ReallocatePinsDuringReconstruction(TArray<UEdGraphPin*>& OldPins) override;
 	virtual class FNodeHandlingFunctor* CreateNodeHandler(class FKismetCompilerContext& CompilerContext) const override;
 	virtual void GetMenuActions(FBlueprintActionDatabaseRegistrar& ActionRegistrar) const override;
 	virtual FText GetMenuCategory() const override;
 	virtual bool CanEverRemoveExecutionPin() const override { return true; }
-	// End UK2Node interface
+	//~ End UK2Node Interface
 
-	// K2Node_ExecutionSequence interface
+	//~ Begin K2Node_ExecutionSequence Interface
 
 	/** Gets a unique pin name, the next in the sequence */
 	FString GetUniquePinName();
diff --git a/Engine/Source/Editor/BlueprintGraph/Classes/K2Node_FormatText.h b/Engine/Source/Editor/BlueprintGraph/Classes/K2Node_FormatText.h
--- a/Engine/Source/Editor/BlueprintGraph/Classes/K2Node_FormatText.h
+++ b/Engine/Source/Editor/BlueprintGraph/Classes/K2Node_FormatText.h
@@ -8,12 +8,12 @@ class UK2Node_FormatText : public UK2Node
 {
 	GENERATED_UCLASS_BODY()
 
-	// UObject interface
+	//~ Begin UObject Interface
 	virtual void PostLoad() override;
 	virtual void PostEditChangeProperty(struct FPropertyChangedEvent& PropertyChangedEvent) override;
-	// End of UObject interface
+	//~ End UObject Interface
 
-	// Begin UEdGraphNode interface.
+	//~ Begin UEdGraphNode Interface.
 	virtual void AllocateDefaultPins() override;
 	virtual FText GetNodeTitle(ENodeTitleType::Type TitleType) const override;
 	virtual bool ShouldShowNodeProperties() const override { return true; }
@@ -21,16 +21,16 @@ class UK2Node_FormatText : public UK2Node
 	virtual void PinDefaultValueChanged(UEdGraphPin* Pin) override;
 	virtual FText GetTooltipText() const override;
 	virtual FText GetPinDisplayName(const UEdGraphPin* Pin) const override;
-	// End UEdGraphNode interface.
+	//~ End UEdGraphNode Interface.
 
-	// Begin UK2Node interface.
+	//~ Begin UK2Node Interface.
 	virtual bool IsNodePure() const override { return true; }
 	virtual bool NodeCausesStructuralBlueprintChange() const override { return true; }
 	virtual void ExpandNode(class FKismetCompilerContext& CompilerContext, UEdGraph* SourceGraph) override;
 	virtual ERedirectType DoPinsMatchForReconstruction(const UEdGraphPin* NewPin, int32 NewPinIndex, const UEdGraphPin* OldPin, int32 OldPinIndex) const override;
 	virtual void GetMenuActions(FBlueprintActionDatabaseRegistrar& ActionRegistrar) const override;
 	virtual FText GetMenuCategory() const override;
-	// End UK2Node interface.
+	//~ End UK2Node Interface.
 
 public:
 	/** Adds a new pin to the node */
diff --git a/Engine/Source/Editor/BlueprintGraph/Classes/K2Node_FunctionEntry.h b/Engine/Source/Editor/BlueprintGraph/Classes/K2Node_FunctionEntry.h
--- a/Engine/Source/Editor/BlueprintGraph/Classes/K2Node_FunctionEntry.h
+++ b/Engine/Source/Editor/BlueprintGraph/Classes/K2Node_FunctionEntry.h
@@ -25,34 +25,34 @@ class UK2Node_FunctionEntry : public UK2Node_FunctionTerminator
 	UPROPERTY()
 	bool bEnforceConstCorrectness;
 
-	// Begin UObject interface
+	//~ Begin UObject Interface
 	virtual void Serialize(FArchive& Ar) override;
-	// End of UObject interface
+	//~ End UObject Interface
 
-	// Begin UEdGraphNode interface
+	//~ Begin UEdGraphNode Interface
 	virtual void AllocateDefaultPins() override;
 	virtual FText GetNodeTitle(ENodeTitleType::Type TitleType) const override;
 	virtual bool CanUserDeleteNode() const override { return false; }
 	virtual bool IsDeprecated() const override;
 	virtual FString GetDeprecationMessage() const override;
 	virtual FText GetTooltipText() const override;
-	// End UEdGraphNode interface
+	//~ End UEdGraphNode Interface
 
-	// Begin UK2Node interface
+	//~ Begin UK2Node Interface
 	virtual bool DrawNodeAsEntry() const override { return true; }
 	virtual class FNodeHandlingFunctor* CreateNodeHandler(class FKismetCompilerContext& CompilerContext) const override;
 	virtual void GetRedirectPinNames(const UEdGraphPin& Pin, TArray<FString>& RedirectPinNames) const override;
 	virtual void ExpandNode(class FKismetCompilerContext& CompilerContext, UEdGraph* SourceGraph) override;
-	// End UK2Node interface
+	//~ End UK2Node Interface
 
-	// Begin UK2Node_EditablePinBase interface
+	//~ Begin UK2Node_EditablePinBase Interface
 	virtual bool CanUseRefParams() const override { return true; }
-	// End UK2Node_EditablePinBase interface
+	//~ End UK2Node_EditablePinBase Interface
 
-	// Begin K2Node_FunctionTerminator interface
+	//~ Begin K2Node_FunctionTerminator Interface
 	virtual bool CanCreateUserDefinedPin(const FEdGraphPinType& InPinType, EEdGraphPinDirection InDesiredDirection, FText& OutErrorMessage) override;
 	virtual UEdGraphPin* CreatePinFromUserDefinition(const TSharedPtr<FUserPinInfo> NewPinInfo) override;
-	// End K2Node_FunctionTerminator interface
+	//~ End K2Node_FunctionTerminator Interface
 
 	// Removes an output pin from the node
 	BLUEPRINTGRAPH_API void RemoveOutputPin(UEdGraphPin* PinToRemove);
diff --git a/Engine/Source/Editor/BlueprintGraph/Classes/K2Node_FunctionResult.h b/Engine/Source/Editor/BlueprintGraph/Classes/K2Node_FunctionResult.h
--- a/Engine/Source/Editor/BlueprintGraph/Classes/K2Node_FunctionResult.h
+++ b/Engine/Source/Editor/BlueprintGraph/Classes/K2Node_FunctionResult.h
@@ -10,14 +10,14 @@ class UK2Node_FunctionResult : public UK2Node_FunctionTerminator
 {
 	GENERATED_UCLASS_BODY()
 
-	// Begin UEdGraphNode interface
+	//~ Begin UEdGraphNode Interface
 	virtual void AllocateDefaultPins() override;
 	virtual FText GetNodeTitle(ENodeTitleType::Type TitleType) const override;
 	virtual FText GetTooltipText() const override;
 	virtual void ValidateNodeDuringCompilation(class FCompilerResultsLog& MessageLog) const override;
-	// End UEdGraphNode interface
+	//~ End UEdGraphNode Interface
 
-	// Begin UK2Node interface
+	//~ Begin UK2Node Interface
 	virtual bool DrawNodeAsExit() const override { return true; }
 	virtual bool ShouldShowNodeProperties() const override { return true; }
 	virtual class FNodeHandlingFunctor* CreateNodeHandler(class FKismetCompilerContext& CompilerContext) const override;
@@ -27,16 +27,16 @@ class UK2Node_FunctionResult : public UK2Node_FunctionTerminator
 	virtual void PostPlacedNewNode();
 	virtual bool CanDuplicateNode() const { return true; }
 	virtual void PostPasteNode() override;
-	// End UK2Node interface
+	//~ End UK2Node Interface
 
-	// Begin UK2Node_EditablePinBase interface
+	//~ Begin UK2Node_EditablePinBase Interface
 	virtual bool CanCreateUserDefinedPin(const FEdGraphPinType& InPinType, EEdGraphPinDirection InDesiredDirection, FText& OutErrorMessage) override;
 	virtual UEdGraphPin* CreatePinFromUserDefinition(const TSharedPtr<FUserPinInfo> NewPinInfo) override;
-	// End UK2Node_EditablePinBase interface
+	//~ End UK2Node_EditablePinBase Interface
 
-	// Begin K2Node_FunctionTerminator interface
+	//~ Begin K2Node_FunctionTerminator Interface
 	virtual void PromoteFromInterfaceOverride(bool bIsPrimaryTerminator = true);
-	// End K2Node_FunctionTerminator interface
+	//~ End K2Node_FunctionTerminator Interface
 
 	BLUEPRINTGRAPH_API TArray<UK2Node_FunctionResult*> GetAllResultNodes() const;
 
diff --git a/Engine/Source/Editor/BlueprintGraph/Classes/K2Node_FunctionTerminator.h b/Engine/Source/Editor/BlueprintGraph/Classes/K2Node_FunctionTerminator.h
--- a/Engine/Source/Editor/BlueprintGraph/Classes/K2Node_FunctionTerminator.h
+++ b/Engine/Source/Editor/BlueprintGraph/Classes/K2Node_FunctionTerminator.h
@@ -22,20 +22,20 @@ class UK2Node_FunctionTerminator : public UK2Node_EditablePinBase
 	FName SignatureName;
 
 
-	// Begin UEdGraphNode interface
+	//~ Begin UEdGraphNode Interface
 	virtual bool CanDuplicateNode() const override { return false; }
 	virtual FLinearColor GetNodeTitleColor() const override;
 	virtual FString CreateUniquePinName(FString SourcePinName) const override;
-	// End UEdGraphNode interface
+	//~ End UEdGraphNode Interface
 
-	// Begin UK2Node interface
+	//~ Begin UK2Node Interface
 	virtual bool NodeCausesStructuralBlueprintChange() const override { return true; }
 	virtual bool HasExternalDependencies(TArray<class UStruct*>* OptionalOutput) const override;
-	// End UK2Node interface
+	//~ End UK2Node Interface
 
-	// Begin UK2Node_EditablePinBase interface
+	//~ Begin UK2Node_EditablePinBase Interface
 	virtual bool CanCreateUserDefinedPin(const FEdGraphPinType& InPinType, EEdGraphPinDirection InDesiredDirection, FText& OutErrorMessage) override;
-	// End UK2Node_EditablePinBase interface
+	//~ End UK2Node_EditablePinBase Interface
 
 	/** Promotes the node from being a part of an interface override to a full function that allows for parameter and result pin additions */
 	virtual void PromoteFromInterfaceOverride(bool bIsPrimaryTerminator = true);
diff --git a/Engine/Source/Editor/BlueprintGraph/Classes/K2Node_GenericCreateObject.h b/Engine/Source/Editor/BlueprintGraph/Classes/K2Node_GenericCreateObject.h
--- a/Engine/Source/Editor/BlueprintGraph/Classes/K2Node_GenericCreateObject.h
+++ b/Engine/Source/Editor/BlueprintGraph/Classes/K2Node_GenericCreateObject.h
@@ -11,12 +11,12 @@ class BLUEPRINTGRAPH_API UK2Node_GenericCreateObject : public UK2Node_ConstructO
 {
 	GENERATED_BODY()
 
-	// Begin UEdGraphNode interface.
+	//~ Begin UEdGraphNode Interface.
 	virtual void ExpandNode(class FKismetCompilerContext& CompilerContext, UEdGraph* SourceGraph) override;
 	virtual void ValidateNodeDuringCompilation(class FCompilerResultsLog& MessageLog) const override;
 	virtual void EarlyValidation(class FCompilerResultsLog& MessageLog) const override;
 	virtual bool IsCompatibleWithGraph(const UEdGraph* TargetGraph) const override;
-	// End UEdGraphNode interface.
+	//~ End UEdGraphNode Interface.
 
 	virtual bool UseWorldContext() const override { return false; }
 	virtual bool UseOuter() const override { return true; }
diff --git a/Engine/Source/Editor/BlueprintGraph/Classes/K2Node_GetClassDefaults.h b/Engine/Source/Editor/BlueprintGraph/Classes/K2Node_GetClassDefaults.h
--- a/Engine/Source/Editor/BlueprintGraph/Classes/K2Node_GetClassDefaults.h
+++ b/Engine/Source/Editor/BlueprintGraph/Classes/K2Node_GetClassDefaults.h
@@ -9,19 +9,19 @@ class UK2Node_GetClassDefaults : public UK2Node
 {
 	GENERATED_UCLASS_BODY()
 
-	// Begin UObject interface
+	//~ Begin UObject Interface
 	virtual void PostEditChangeProperty(struct FPropertyChangedEvent& PropertyChangedEvent) override;
-	// End UObject interface
+	//~ End UObject Interface
 
-	// Begin UEdGraphNode interface
+	//~ Begin UEdGraphNode Interface
 	virtual FText GetNodeTitle(ENodeTitleType::Type TitleType) const override;
 	virtual void AllocateDefaultPins() override;
 	virtual void PostPlacedNewNode() override;
 	virtual void PinDefaultValueChanged(UEdGraphPin* Pin) override;
 	virtual void PinConnectionListChanged(UEdGraphPin* Pin) override;
-	// End UEdGraphNode interface
+	//~ End UEdGraphNode Interface
 
-	// Begin UK2Node interface
+	//~ Begin UK2Node Interface
 	virtual bool IsNodePure() const override { return true; }
 	virtual bool ShouldShowNodeProperties() const override { return true; }
 	virtual void ReallocatePinsDuringReconstruction(TArray<UEdGraphPin*>& OldPins) override;
@@ -30,7 +30,7 @@ class UK2Node_GetClassDefaults : public UK2Node
 	virtual void ExpandNode(class FKismetCompilerContext& CompilerContext, UEdGraph* SourceGraph) override;
 	virtual void GetMenuActions(FBlueprintActionDatabaseRegistrar& ActionRegistrar) const override;
 	virtual FText GetMenuCategory() const override;
-	// End UK2Node interface
+	//~ End UK2Node Interface
 
 public:
 	/** Finds and returns the class input pin from the current set of pins. */
diff --git a/Engine/Source/Editor/BlueprintGraph/Classes/K2Node_GetEnumeratorNameAsString.h b/Engine/Source/Editor/BlueprintGraph/Classes/K2Node_GetEnumeratorNameAsString.h
--- a/Engine/Source/Editor/BlueprintGraph/Classes/K2Node_GetEnumeratorNameAsString.h
+++ b/Engine/Source/Editor/BlueprintGraph/Classes/K2Node_GetEnumeratorNameAsString.h
@@ -9,17 +9,17 @@ class UK2Node_GetEnumeratorNameAsString : public UK2Node_GetEnumeratorName
 {
 	GENERATED_UCLASS_BODY()
 
-	// Begin UEdGraphNode interface
+	//~ Begin UEdGraphNode Interface
 	virtual void AllocateDefaultPins() override;
 	virtual FText GetTooltipText() const override;
 	virtual FText GetNodeTitle(ENodeTitleType::Type TitleType) const override;
 	virtual FName GetPaletteIcon(FLinearColor& OutColor) const override{ return TEXT("GraphEditor.Enum_16x"); }
-	// End UEdGraphNode interface
+	//~ End UEdGraphNode Interface
 
-	// UK2Node interface
+	//~ Begin UK2Node Interface
 	virtual void GetMenuActions(FBlueprintActionDatabaseRegistrar& ActionRegistrar) const override;
 	virtual FText GetMenuCategory() const override;
-	// End UK2Node interface
+	//~ End UK2Node Interface
 
 	virtual FName GetFunctionName() const override;
 };
diff --git a/Engine/Source/Editor/BlueprintGraph/Classes/K2Node_GetInputAxisKeyValue.h b/Engine/Source/Editor/BlueprintGraph/Classes/K2Node_GetInputAxisKeyValue.h
--- a/Engine/Source/Editor/BlueprintGraph/Classes/K2Node_GetInputAxisKeyValue.h
+++ b/Engine/Source/Editor/BlueprintGraph/Classes/K2Node_GetInputAxisKeyValue.h
@@ -21,15 +21,15 @@ class UK2Node_GetInputAxisKeyValue : public UK2Node_CallFunction
 	UPROPERTY(EditAnywhere, Category="Input")
 	uint32 bExecuteWhenPaused:1;
 
-	// Begin EdGraphNode interface
+	//~ Begin EdGraphNode Interface
 	virtual void AllocateDefaultPins() override;
 	virtual FText GetNodeTitle(ENodeTitleType::Type TitleType) const override;
 	virtual FText GetTooltipText() const override;
 	virtual FName GetPaletteIcon(FLinearColor& OutColor) const override;
 	virtual bool IsCompatibleWithGraph(UEdGraph const* Graph) const override;
-	// End EdGraphNode interface
+	//~ End EdGraphNode Interface
 
-	// Begin UK2Node interface
+	//~ Begin UK2Node Interface
 	virtual void ValidateNodeDuringCompilation(class FCompilerResultsLog& MessageLog) const override;
 	virtual bool ShouldShowNodeProperties() const override { return true; }
 	virtual UClass* GetDynamicBindingClass() const override;
@@ -37,7 +37,7 @@ class UK2Node_GetInputAxisKeyValue : public UK2Node_CallFunction
 	virtual void GetMenuActions(FBlueprintActionDatabaseRegistrar& ActionRegistrar) const override;
 	virtual FText GetMenuCategory() const override;
 	virtual FBlueprintNodeSignature GetSignature() const override;
-	// End UK2Node interface
+	//~ End UK2Node Interface
 	
 	void Initialize(const FKey AxisKey);
 
diff --git a/Engine/Source/Editor/BlueprintGraph/Classes/K2Node_GetInputAxisValue.h b/Engine/Source/Editor/BlueprintGraph/Classes/K2Node_GetInputAxisValue.h
--- a/Engine/Source/Editor/BlueprintGraph/Classes/K2Node_GetInputAxisValue.h
+++ b/Engine/Source/Editor/BlueprintGraph/Classes/K2Node_GetInputAxisValue.h
@@ -21,14 +21,14 @@ class UK2Node_GetInputAxisValue : public UK2Node_CallFunction
 	UPROPERTY(EditAnywhere, Category="Input")
 	uint32 bExecuteWhenPaused:1;
 
-	// Begin EdGraphNode interface
+	//~ Begin EdGraphNode Interface
 	virtual void AllocateDefaultPins() override;
 	virtual FText GetNodeTitle(ENodeTitleType::Type TitleType) const override;
 	virtual FText GetTooltipText() const override;
 	virtual bool IsCompatibleWithGraph(UEdGraph const* Graph) const override;
-	// End EdGraphNode interface
+	//~ End EdGraphNode Interface
 
-	// Begin UK2Node interface
+	//~ Begin UK2Node Interface
 	virtual void ValidateNodeDuringCompilation(class FCompilerResultsLog& MessageLog) const override;
 	virtual bool ShouldShowNodeProperties() const override { return true; }
 	virtual UClass* GetDynamicBindingClass() const override;
@@ -36,7 +36,7 @@ class UK2Node_GetInputAxisValue : public UK2Node_CallFunction
 	virtual void GetMenuActions(FBlueprintActionDatabaseRegistrar& ActionRegistrar) const override;
 	virtual FText GetMenuCategory() const override;
 	virtual FBlueprintNodeSignature GetSignature() const override;
-	// End UK2Node interface
+	//~ End UK2Node Interface
 	
 	void Initialize(const FName AxisName);
 
diff --git a/Engine/Source/Editor/BlueprintGraph/Classes/K2Node_IfThenElse.h b/Engine/Source/Editor/BlueprintGraph/Classes/K2Node_IfThenElse.h
--- a/Engine/Source/Editor/BlueprintGraph/Classes/K2Node_IfThenElse.h
+++ b/Engine/Source/Editor/BlueprintGraph/Classes/K2Node_IfThenElse.h
@@ -10,19 +10,19 @@ class UK2Node_IfThenElse : public UK2Node
 {
 	GENERATED_UCLASS_BODY()
 
-	// Begin UEdGraphNode interface
+	//~ Begin UEdGraphNode Interface
 	virtual void AllocateDefaultPins() override;
 	virtual FLinearColor GetNodeTitleColor() const override;
 	virtual FText GetTooltipText() const override;
 	virtual FText GetNodeTitle(ENodeTitleType::Type TitleType) const override;
 	virtual FName GetPaletteIcon(FLinearColor& OutColor) const override{ return TEXT("GraphEditor.Branch_16x"); }
-	// End UEdGraphNode interface
+	//~ End UEdGraphNode Interface
 
-	// Begin K2Node interface.
+	//~ Begin K2Node Interface.
 	virtual class FNodeHandlingFunctor* CreateNodeHandler(class FKismetCompilerContext& CompilerContext) const override;
 	virtual void GetMenuActions(FBlueprintActionDatabaseRegistrar& ActionRegistrar) const override;
 	virtual FText GetMenuCategory() const override;
-	// End K2Node interface.
+	//~ End K2Node Interface.
 
 	/** Get the then output pin */
 	BLUEPRINTGRAPH_API UEdGraphPin* GetThenPin() const;
diff --git a/Engine/Source/Editor/BlueprintGraph/Classes/K2Node_InputAction.h b/Engine/Source/Editor/BlueprintGraph/Classes/K2Node_InputAction.h
--- a/Engine/Source/Editor/BlueprintGraph/Classes/K2Node_InputAction.h
+++ b/Engine/Source/Editor/BlueprintGraph/Classes/K2Node_InputAction.h
@@ -26,27 +26,27 @@ class UK2Node_InputAction : public UK2Node
 	UPROPERTY(EditAnywhere, Category="Input")
 	uint32 bOverrideParentBinding:1;
 
-	// Begin UObject interface
+	//~ Begin UObject Interface
 	virtual void PostLoad() override;
-	// End UObject interface
+	//~ End UObject Interface
 
-	// Begin UEdGraphNode interface.
+	//~ Begin UEdGraphNode Interface.
 	virtual void AllocateDefaultPins() override;
 	virtual FLinearColor GetNodeTitleColor() const override;
 	virtual FText GetNodeTitle(ENodeTitleType::Type TitleType) const override;
 	virtual FText GetTooltipText() const override;
 	virtual FName GetPaletteIcon(FLinearColor& OutColor) const override{ return TEXT("GraphEditor.Event_16x"); }
 	virtual bool IsCompatibleWithGraph(UEdGraph const* Graph) const override;
-	// End UEdGraphNode interface.
+	//~ End UEdGraphNode Interface.
 
-	// Begin UK2Node interface
+	//~ Begin UK2Node Interface
 	virtual void ValidateNodeDuringCompilation(class FCompilerResultsLog& MessageLog) const override;
 	virtual bool ShouldShowNodeProperties() const override { return true; }
 	virtual void ExpandNode(FKismetCompilerContext& CompilerContext, UEdGraph* SourceGraph) override;
 	virtual void GetMenuActions(FBlueprintActionDatabaseRegistrar& ActionRegistrar) const override;
 	virtual FText GetMenuCategory() const override;
 	virtual FBlueprintNodeSignature GetSignature() const override;
-	// End UK2Node interface
+	//~ End UK2Node Interface
 
 	/** Get the 'pressed' input pin */
 	BLUEPRINTGRAPH_API UEdGraphPin* GetPressedPin() const;
diff --git a/Engine/Source/Editor/BlueprintGraph/Classes/K2Node_InputActionEvent.h b/Engine/Source/Editor/BlueprintGraph/Classes/K2Node_InputActionEvent.h
--- a/Engine/Source/Editor/BlueprintGraph/Classes/K2Node_InputActionEvent.h
+++ b/Engine/Source/Editor/BlueprintGraph/Classes/K2Node_InputActionEvent.h
@@ -25,8 +25,8 @@ class UK2Node_InputActionEvent : public UK2Node_Event
 	UPROPERTY()
 	uint32 bOverrideParentBinding:1;
 
-	// Begin UK2Node interface
+	//~ Begin UK2Node Interface
 	virtual UClass* GetDynamicBindingClass() const override;
 	virtual void RegisterDynamicBinding(UDynamicBlueprintBinding* BindingObject) const override;
-	// End UK2Node interface
+	//~ End UK2Node Interface
 };
\ No newline at end of file
diff --git a/Engine/Source/Editor/BlueprintGraph/Classes/K2Node_InputAxisEvent.h b/Engine/Source/Editor/BlueprintGraph/Classes/K2Node_InputAxisEvent.h
--- a/Engine/Source/Editor/BlueprintGraph/Classes/K2Node_InputAxisEvent.h
+++ b/Engine/Source/Editor/BlueprintGraph/Classes/K2Node_InputAxisEvent.h
@@ -25,18 +25,18 @@ class UK2Node_InputAxisEvent : public UK2Node_Event
 	UPROPERTY(EditAnywhere, Category="Input")
 	uint32 bOverrideParentBinding:1;
 
-	// Begin UObject interface
+	//~ Begin UObject Interface
 	virtual void PostLoad() override;
 	virtual void Serialize(FArchive& Ar) override;
-	// End of UObject interface
+	//~ End UObject Interface
 
-	// Begin EdGraphNode interface
+	//~ Begin EdGraphNode Interface
 	virtual FText GetNodeTitle(ENodeTitleType::Type TitleType) const override;
 	virtual FText GetTooltipText() const override;
 	virtual bool IsCompatibleWithGraph(const UEdGraph* TargetGraph) const override;
-	// End EdGraphNode interface
+	//~ End EdGraphNode Interface
 
-	// Begin UK2Node interface
+	//~ Begin UK2Node Interface
 	virtual void ValidateNodeDuringCompilation(class FCompilerResultsLog& MessageLog) const override;
 	virtual bool ShouldShowNodeProperties() const override { return true; }
 	virtual UClass* GetDynamicBindingClass() const override;
@@ -44,7 +44,7 @@ class UK2Node_InputAxisEvent : public UK2Node_Event
 	virtual void GetMenuActions(FBlueprintActionDatabaseRegistrar& ActionRegistrar) const override;
 	virtual FText GetMenuCategory() const override;
 	virtual FBlueprintNodeSignature GetSignature() const override;
-	// End UK2Node interface
+	//~ End UK2Node Interface
 
 	void Initialize(const FName AxisName);
 
diff --git a/Engine/Source/Editor/BlueprintGraph/Classes/K2Node_InputAxisKeyEvent.h b/Engine/Source/Editor/BlueprintGraph/Classes/K2Node_InputAxisKeyEvent.h
--- a/Engine/Source/Editor/BlueprintGraph/Classes/K2Node_InputAxisKeyEvent.h
+++ b/Engine/Source/Editor/BlueprintGraph/Classes/K2Node_InputAxisKeyEvent.h
@@ -24,18 +24,18 @@ class UK2Node_InputAxisKeyEvent : public UK2Node_Event
 	UPROPERTY(EditAnywhere, Category = "Input")
 	uint32 bOverrideParentBinding : 1;
 
-	// UObject interface
+	//~ Begin UObject Interface
 	virtual void Serialize(FArchive& Ar) override;
-	// End of UObject interface
+	//~ End UObject Interface
 
-	// Begin EdGraphNode interface
+	//~ Begin EdGraphNode Interface
 	virtual FText GetNodeTitle(ENodeTitleType::Type TitleType) const override;
 	virtual FText GetTooltipText() const override;
 	virtual FName GetPaletteIcon(FLinearColor& OutColor) const override;
 	virtual bool IsCompatibleWithGraph(const UEdGraph* TargetGraph) const override;
-	// End EdGraphNode interface
+	//~ End EdGraphNode Interface
 
-	// Begin UK2Node interface
+	//~ Begin UK2Node Interface
 	virtual void ValidateNodeDuringCompilation(class FCompilerResultsLog& MessageLog) const override;
 	virtual bool ShouldShowNodeProperties() const override{ return true; }
 	virtual UClass* GetDynamicBindingClass() const override;
@@ -43,7 +43,7 @@ class UK2Node_InputAxisKeyEvent : public UK2Node_Event
 	virtual void GetMenuActions(FBlueprintActionDatabaseRegistrar& ActionRegistrar) const override;
 	virtual FText GetMenuCategory() const override;
 	virtual FBlueprintNodeSignature GetSignature() const override;
-	// End UK2Node interface
+	//~ End UK2Node Interface
 
 	void Initialize(const FKey AxisKey);
 
diff --git a/Engine/Source/Editor/BlueprintGraph/Classes/K2Node_InputKey.h b/Engine/Source/Editor/BlueprintGraph/Classes/K2Node_InputKey.h
--- a/Engine/Source/Editor/BlueprintGraph/Classes/K2Node_InputKey.h
+++ b/Engine/Source/Editor/BlueprintGraph/Classes/K2Node_InputKey.h
@@ -43,28 +43,28 @@ class UK2Node_InputKey : public UK2Node
 	UPROPERTY(EditAnywhere, Category="Modifier")
 	uint32 bCommand:1;
 
-	// Begin UObject interface
+	//~ Begin UObject Interface
 	virtual void PostLoad() override;
 	virtual void PostEditChangeProperty(struct FPropertyChangedEvent& PropertyChangedEvent) override;
-	// End UObject interface
+	//~ End UObject Interface
 
-	// Begin UK2Node interface.
+	//~ Begin UK2Node Interface.
 	virtual bool ShouldShowNodeProperties() const override { return true; }
 	virtual void ValidateNodeDuringCompilation(class FCompilerResultsLog& MessageLog) const override;
 	virtual void ExpandNode(FKismetCompilerContext& CompilerContext, UEdGraph* SourceGraph) override;
 	virtual void GetMenuActions(FBlueprintActionDatabaseRegistrar& ActionRegistrar) const override;
 	virtual FText GetMenuCategory() const override;
 	virtual FBlueprintNodeSignature GetSignature() const override;
-	// End UK2Node interface
+	//~ End UK2Node Interface
 
-	// Begin UEdGraphNode interface.
+	//~ Begin UEdGraphNode Interface.
 	virtual void AllocateDefaultPins() override;
 	virtual FLinearColor GetNodeTitleColor() const override;
 	virtual FText GetNodeTitle(ENodeTitleType::Type TitleType) const override;
 	virtual FText GetTooltipText() const override;
 	virtual FName GetPaletteIcon(FLinearColor& OutColor) const override;
 	virtual bool IsCompatibleWithGraph(UEdGraph const* Graph) const override;
-	// End UEdGraphNode interface.
+	//~ End UEdGraphNode Interface.
 
 	BLUEPRINTGRAPH_API FText GetModifierText() const;
 
diff --git a/Engine/Source/Editor/BlueprintGraph/Classes/K2Node_InputKeyEvent.h b/Engine/Source/Editor/BlueprintGraph/Classes/K2Node_InputKeyEvent.h
--- a/Engine/Source/Editor/BlueprintGraph/Classes/K2Node_InputKeyEvent.h
+++ b/Engine/Source/Editor/BlueprintGraph/Classes/K2Node_InputKeyEvent.h
@@ -26,8 +26,8 @@ class UK2Node_InputKeyEvent : public UK2Node_Event
 	UPROPERTY()
 	uint32 bOverrideParentBinding:1;
 
-	// Begin UK2Node interface
+	//~ Begin UK2Node Interface
 	virtual UClass* GetDynamicBindingClass() const override;
 	virtual void RegisterDynamicBinding(UDynamicBlueprintBinding* BindingObject) const override;
-	// End UK2Node interface
+	//~ End UK2Node Interface
 };
\ No newline at end of file
diff --git a/Engine/Source/Editor/BlueprintGraph/Classes/K2Node_InputTouch.h b/Engine/Source/Editor/BlueprintGraph/Classes/K2Node_InputTouch.h
--- a/Engine/Source/Editor/BlueprintGraph/Classes/K2Node_InputTouch.h
+++ b/Engine/Source/Editor/BlueprintGraph/Classes/K2Node_InputTouch.h
@@ -22,16 +22,16 @@ class UK2Node_InputTouch : public UK2Node
 	UPROPERTY(EditAnywhere, Category="Input")
 	uint32 bOverrideParentBinding:1;
 
-	// Begin UObject interface
+	//~ Begin UObject Interface
 	virtual void PostLoad() override;
-	// End UObject interface
+	//~ End UObject Interface
 
-	// Begin UK2Node interface.
+	//~ Begin UK2Node Interface.
 	virtual bool ShouldShowNodeProperties() const override { return true; }
 	virtual void ExpandNode(FKismetCompilerContext& CompilerContext, UEdGraph* SourceGraph) override;
-	// End UK2Node interface
+	//~ End UK2Node Interface
 
-	// Begin UEdGraphNode interface.
+	//~ Begin UEdGraphNode Interface.
 	virtual void AllocateDefaultPins() override;
 	virtual FLinearColor GetNodeTitleColor() const override;
 	virtual FText GetNodeTitle(ENodeTitleType::Type TitleType) const override;
@@ -40,7 +40,7 @@ class UK2Node_InputTouch : public UK2Node
 	virtual void GetMenuActions(FBlueprintActionDatabaseRegistrar& ActionRegistrar) const override;
 	virtual FText GetMenuCategory() const override;
 	virtual bool IsCompatibleWithGraph(const UEdGraph* TargetGraph) const override;
-	// End UEdGraphNode interface.
+	//~ End UEdGraphNode Interface.
 
 	BLUEPRINTGRAPH_API static UEnum* GetTouchIndexEnum();
 
diff --git a/Engine/Source/Editor/BlueprintGraph/Classes/K2Node_InputVectorAxisEvent.h b/Engine/Source/Editor/BlueprintGraph/Classes/K2Node_InputVectorAxisEvent.h
--- a/Engine/Source/Editor/BlueprintGraph/Classes/K2Node_InputVectorAxisEvent.h
+++ b/Engine/Source/Editor/BlueprintGraph/Classes/K2Node_InputVectorAxisEvent.h
@@ -9,14 +9,14 @@ class UK2Node_InputVectorAxisEvent : public UK2Node_InputAxisKeyEvent
 {
 	GENERATED_UCLASS_BODY()
 
-	// UObject interface
+	//~ Begin UObject Interface
 	virtual void Serialize(FArchive& Ar) override;
-	// End of UObject interface
+	//~ End UObject Interface
 
-	// Begin UK2Node interface
+	//~ Begin UK2Node Interface
 	virtual void ValidateNodeDuringCompilation(class FCompilerResultsLog& MessageLog) const override;
 	virtual UClass* GetDynamicBindingClass() const override;
 	virtual void RegisterDynamicBinding(UDynamicBlueprintBinding* BindingObject) const override;
 	virtual void GetMenuActions(FBlueprintActionDatabaseRegistrar& ActionRegistrar) const override;
-	// End UK2Node interface
+	//~ End UK2Node Interface
 };
\ No newline at end of file
diff --git a/Engine/Source/Editor/BlueprintGraph/Classes/K2Node_LocalVariable.h b/Engine/Source/Editor/BlueprintGraph/Classes/K2Node_LocalVariable.h
--- a/Engine/Source/Editor/BlueprintGraph/Classes/K2Node_LocalVariable.h
+++ b/Engine/Source/Editor/BlueprintGraph/Classes/K2Node_LocalVariable.h
@@ -17,7 +17,7 @@ class UDEPRECATED_K2Node_LocalVariable : public UK2Node_TemporaryVariable
 	UPROPERTY()
 	FText VariableTooltip;
 
-	// Begin UEdGraphNode interface.
+	//~ Begin UEdGraphNode Interface.
 	virtual FText GetTooltipText() const override;
 	virtual FText GetNodeTitle(ENodeTitleType::Type TitleType) const override;
 	virtual void OnRenameNode(const FString& NewName) override;
@@ -26,13 +26,13 @@ class UDEPRECATED_K2Node_LocalVariable : public UK2Node_TemporaryVariable
 	virtual void PostPasteNode() override;
 	virtual bool IsCompatibleWithGraph(const UEdGraph* TargetGraph) const override;
 	virtual void ReconstructNode() override;
-	// End UEdGraphNode interface.
+	//~ End UEdGraphNode Interface.
 
-	// Begin UK2Node interface.
+	//~ Begin UK2Node Interface.
 	virtual bool NodeCausesStructuralBlueprintChange() const override { return true; }
 	virtual bool ShouldShowNodeProperties() const override { return true; }
 	virtual void GetMenuActions(FBlueprintActionDatabaseRegistrar& ActionRegistrar) const override {}
-	// End UK2Node interface.
+	//~ End UK2Node Interface.
 
 	/**
 	 * Assigns the local variable a type
diff --git a/Engine/Source/Editor/BlueprintGraph/Classes/K2Node_MacroInstance.h b/Engine/Source/Editor/BlueprintGraph/Classes/K2Node_MacroInstance.h
--- a/Engine/Source/Editor/BlueprintGraph/Classes/K2Node_MacroInstance.h
+++ b/Engine/Source/Editor/BlueprintGraph/Classes/K2Node_MacroInstance.h
@@ -29,11 +29,11 @@ public:
 	/** Whether we need to reconstruct the node after the pins have changed */
 	bool bReconstructNode;
 
-	// Begin UObject interface
+	//~ Begin UObject Interface
 	virtual void Serialize(FArchive& Ar) override;
-	// End UObject interface
+	//~ End UObject Interface
 
-	// Begin UEdGraphNode interface
+	//~ Begin UEdGraphNode Interface
 	virtual void AllocateDefaultPins() override;
 	virtual FText GetTooltipText() const override;
 	virtual FText GetKeywords() const override;
@@ -48,9 +48,9 @@ public:
 	virtual FString GetDocumentationExcerptName() const override;
 	virtual FName GetPaletteIcon(FLinearColor& OutColor) const override;
 	virtual bool CanPasteHere(const UEdGraph* TargetGraph) const override;
-	// End UEdGraphNode interface
+	//~ End UEdGraphNode Interface
 
-	// Begin UK2Node interface
+	//~ Begin UK2Node Interface
 	virtual bool DrawNodeAsExit() const override { return false; }
 	virtual bool DrawNodeAsEntry() const override { return false; }
 	virtual void NotifyPinConnectionListChanged(UEdGraphPin* Pin) override;
@@ -64,11 +64,11 @@ public:
 	virtual bool IsActionFilteredOut(class FBlueprintActionFilter const& Filter) override;
 	virtual FText GetCompactNodeTitle() const override;
 	virtual bool ShouldDrawCompact() const override;
-	// End UK2Node interface
+	//~ End UK2Node Interface
 
-	// Begin UK2Node_EditablePinBase interface
+	//~ Begin UK2Node_EditablePinBase Interface
 	virtual bool CanCreateUserDefinedPin(const FEdGraphPinType& InPinType, EEdGraphPinDirection InDesiredDirection, FText& OutErrorMessage) override { return false; }
-	// End UK2Node_EditablePinBase interface
+	//~ End UK2Node_EditablePinBase Interface
 
 	void SetMacroGraph(UEdGraph* Graph) { MacroGraphReference.SetGraph(Graph); }
 	UEdGraph* GetMacroGraph() const { return MacroGraphReference.GetGraph(); }
diff --git a/Engine/Source/Editor/BlueprintGraph/Classes/K2Node_MakeStruct.h b/Engine/Source/Editor/BlueprintGraph/Classes/K2Node_MakeStruct.h
--- a/Engine/Source/Editor/BlueprintGraph/Classes/K2Node_MakeStruct.h
+++ b/Engine/Source/Editor/BlueprintGraph/Classes/K2Node_MakeStruct.h
@@ -27,16 +27,16 @@ class UK2Node_MakeStruct : public UK2Node_StructMemberSet
 	*/
 	BLUEPRINTGRAPH_API static bool CanBeMade(const UScriptStruct* Struct, bool bIncludeEditAnywhere = true, bool bMustHaveValidProperties = false);
 	
-	// Begin UEdGraphNode interface
+	//~ Begin UEdGraphNode Interface
 	virtual void AllocateDefaultPins() override;
 	virtual FText GetNodeTitle(ENodeTitleType::Type TitleType) const override;
 	virtual FLinearColor GetNodeTitleColor() const override;
 	virtual FText GetTooltipText() const override;
 	virtual void ValidateNodeDuringCompilation(class FCompilerResultsLog& MessageLog) const override;
 	virtual FName GetPaletteIcon(FLinearColor& OutColor) const override{ return TEXT("GraphEditor.MakeStruct_16x"); }
-	// End  UEdGraphNode interface
+	//~ End  UEdGraphNode Interface
 
-	// Begin K2Node interface
+	//~ Begin K2Node Interface
 	virtual bool NodeCausesStructuralBlueprintChange() const override { return false; }
 	virtual bool IsNodePure() const override { return true; }
 	virtual bool DrawNodeAsVariable() const override { return false; }
@@ -44,7 +44,7 @@ class UK2Node_MakeStruct : public UK2Node_StructMemberSet
 	virtual ERedirectType DoPinsMatchForReconstruction(const UEdGraphPin* NewPin, int32 NewPinIndex, const UEdGraphPin* OldPin, int32 OldPinIndex)  const override;
 	virtual void GetMenuActions(FBlueprintActionDatabaseRegistrar& ActionRegistrar) const override;
 	virtual FText GetMenuCategory() const override;
-	// End K2Node interface
+	//~ End K2Node Interface
 
 protected:
 	struct FMakeStructPinManager : public FStructOperationOptionalPinManager
diff --git a/Engine/Source/Editor/BlueprintGraph/Classes/K2Node_MathExpression.h b/Engine/Source/Editor/BlueprintGraph/Classes/K2Node_MathExpression.h
--- a/Engine/Source/Editor/BlueprintGraph/Classes/K2Node_MathExpression.h
+++ b/Engine/Source/Editor/BlueprintGraph/Classes/K2Node_MathExpression.h
@@ -28,12 +28,12 @@ public:
 	bool bMadeAfterRotChange;
 
 public:
-	// UObject interface
+	//~ Begin UObject Interface
 	virtual void Serialize(FArchive& Ar) override;
 	virtual void PostEditChangeProperty(struct FPropertyChangedEvent& PropertyChangedEvent) override;
-	// End of UObject interface
+	//~ End UObject Interface
 
-	// UEdGraphNode interface
+	//~ Begin UEdGraphNode Interface
 	virtual TSharedPtr<class INameValidatorInterface> MakeNameValidator() const override;
 	virtual void OnRenameNode(const FString& NewName) override;
 	virtual FText GetNodeTitle(ENodeTitleType::Type TitleType) const override;
@@ -41,18 +41,18 @@ public:
 	virtual void ReconstructNode() override;
 	virtual void FindDiffs(class UEdGraphNode* OtherNode, struct FDiffResults& Results )  override;
 	virtual bool ShouldMergeChildGraphs() const override { return ShouldExpandInsteadCompile(); }
-	// End of UEdGraphNode interface
+	//~ End UEdGraphNode Interface
 
-	// UK2Node interface
+	//~ Begin UK2Node Interface
 	virtual void ValidateNodeDuringCompilation(class FCompilerResultsLog& MessageLog) const override;
 	virtual void GetMenuActions(FBlueprintActionDatabaseRegistrar& ActionRegistrar) const override;
 	virtual class FNodeHandlingFunctor* CreateNodeHandler(class FKismetCompilerContext& CompilerContext) const override;
 	virtual bool IsNodePure() const { return !ShouldExpandInsteadCompile(); }
-	// End of UK2Node interface
+	//~ End UK2Node Interface
 
-	// Begin UK2Node_EditablePinBase interface
+	//~ Begin UK2Node_EditablePinBase Interface
 	virtual bool CanCreateUserDefinedPin(const FEdGraphPinType& InPinType, EEdGraphPinDirection InDesiredDirection, FText& OutErrorMessage) override { return false; }
-	// End UK2Node_EditablePinBase interface
+	//~ End UK2Node_EditablePinBase Interface
 
 private:
 	/* Returns true, when the node can/should not be optimized.*/
diff --git a/Engine/Source/Editor/BlueprintGraph/Classes/K2Node_MatineeController.h b/Engine/Source/Editor/BlueprintGraph/Classes/K2Node_MatineeController.h
--- a/Engine/Source/Editor/BlueprintGraph/Classes/K2Node_MatineeController.h
+++ b/Engine/Source/Editor/BlueprintGraph/Classes/K2Node_MatineeController.h
@@ -14,24 +14,24 @@ class UK2Node_MatineeController : public UK2Node
 	UPROPERTY(EditAnywhere, Category=K2Node_MatineeController)
 	class AMatineeActor* MatineeActor;
 
-	// Begin UEdGraphNode interface
+	//~ Begin UEdGraphNode Interface
 	virtual void AllocateDefaultPins() override;
 	virtual FLinearColor GetNodeTitleColor() const override;
 	virtual FText GetNodeTitle(ENodeTitleType::Type TitleType) const override;
 	virtual bool CanDuplicateNode() const override { return false; }
 	virtual bool NodeCausesStructuralBlueprintChange() const override { return true; }
-	// End UEdGraphNode interface
+	//~ End UEdGraphNode Interface
 
-	// Begin UK2Node interface
+	//~ Begin UK2Node Interface
 	virtual bool IsNodeSafeToIgnore() const override { return true; }
 	virtual AActor* GetReferencedLevelActor() const override;
 	virtual void ExpandNode(class FKismetCompilerContext& CompilerContext, UEdGraph* SourceGraph) override;
 	virtual void GetMenuActions(FBlueprintActionDatabaseRegistrar& ActionRegistrar) const override;
-	// End UK2Node interface
+	//~ End UK2Node Interface
 
-	// Begin UObject interface
+	//~ Begin UObject Interface
 	virtual void BeginDestroy() override;
-	// End UObject interface
+	//~ End UObject Interface
 
 	/** Gets the "finished playing matinee sequence" pin */
 	UEdGraphPin* GetFinishedPin() const;
diff --git a/Engine/Source/Editor/BlueprintGraph/Classes/K2Node_Message.h b/Engine/Source/Editor/BlueprintGraph/Classes/K2Node_Message.h
--- a/Engine/Source/Editor/BlueprintGraph/Classes/K2Node_Message.h
+++ b/Engine/Source/Editor/BlueprintGraph/Classes/K2Node_Message.h
@@ -12,20 +12,20 @@ class UK2Node_Message : public UK2Node_CallFunction
 	GENERATED_UCLASS_BODY()
 
 
-	// Begin UEdGraphNode interface.
+	//~ Begin UEdGraphNode Interface.
 	virtual FText GetNodeTitle(ENodeTitleType::Type TitleType) const override;
 	virtual void AllocateDefaultPins() override;
-	// End UEdGraphNode interface.
+	//~ End UEdGraphNode Interface.
 
-	// Begin K2Node interface.
+	//~ Begin K2Node Interface.
 	virtual bool IsNodeSafeToIgnore() const override { return true; }
 	virtual class FNodeHandlingFunctor* CreateNodeHandler(class FKismetCompilerContext& CompilerContext) const override;
 	virtual void ExpandNode(class FKismetCompilerContext& CompilerContext, UEdGraph* SourceGraph) override;
 	virtual FName GetCornerIcon() const override;
 	virtual bool IsNodePure() const override { return false; }
-	// End K2Node interface.
+	//~ End K2Node Interface.
 
-	// Begin K2Node_CallFunction Interface.
+	//~ Begin K2Node_CallFunction Interface.
 	virtual UEdGraphPin* CreateSelfPin(const UFunction* Function) override;
 
 protected:
@@ -41,7 +41,7 @@ protected:
 	void ExpandLevelStreamingHandlers(class FKismetCompilerContext& InCompilerContext, UEdGraph* InSourceGraph, class UEdGraphPin* InStartingExecPin, class UEdGraphPin* InMessageSelfPin, class UK2Node_DynamicCast* InCastToInterfaceNode);
 
 	virtual void EnsureFunctionIsInBlueprint() override;
-	// End K2Node_CallFunction Interface.
+	//~ End K2Node_CallFunction Interface.
 
 	/** Constructing FText strings can be costly, so we cache the node's title */
 	FNodeTitleTextTable CachedNodeTitles;
diff --git a/Engine/Source/Editor/BlueprintGraph/Classes/K2Node_Self.h b/Engine/Source/Editor/BlueprintGraph/Classes/K2Node_Self.h
--- a/Engine/Source/Editor/BlueprintGraph/Classes/K2Node_Self.h
+++ b/Engine/Source/Editor/BlueprintGraph/Classes/K2Node_Self.h
@@ -10,20 +10,20 @@ class UK2Node_Self : public UK2Node
 {
 	GENERATED_UCLASS_BODY()
 
-	// Begin UEdGraphNode interface
+	//~ Begin UEdGraphNode Interface
 	virtual FText GetNodeTitle(ENodeTitleType::Type TitleType) const override;
 	virtual FText GetTooltipText() const override;
 	virtual FText GetKeywords() const override;
 	virtual void AllocateDefaultPins() override;
 	virtual void ValidateNodeDuringCompilation(class FCompilerResultsLog& MessageLog) const override;
-	// End UEdGraphNode interface
+	//~ End UEdGraphNode Interface
 
-	// Begin K2Node interface
+	//~ Begin K2Node Interface
 	virtual bool IsNodePure() const override { return true; }
 	virtual bool DrawNodeAsVariable() const override { return true; }
 	virtual class FNodeHandlingFunctor* CreateNodeHandler(class FKismetCompilerContext& CompilerContext) const override;
 	virtual void GetMenuActions(FBlueprintActionDatabaseRegistrar& ActionRegistrar) const override;
 	virtual FText GetMenuCategory() const override;
-	// End K2Node interface
+	//~ End K2Node Interface
 };
 
diff --git a/Engine/Source/Editor/BlueprintGraph/Classes/K2Node_SetFieldsInStruct.h b/Engine/Source/Editor/BlueprintGraph/Classes/K2Node_SetFieldsInStruct.h
--- a/Engine/Source/Editor/BlueprintGraph/Classes/K2Node_SetFieldsInStruct.h
+++ b/Engine/Source/Editor/BlueprintGraph/Classes/K2Node_SetFieldsInStruct.h
@@ -16,11 +16,11 @@ class UK2Node_SetFieldsInStruct : public UK2Node_MakeStruct
 	UPROPERTY()
 	bool bMadeAfterOverridePinRemoval;
 
-	// UObject interface
+	//~ Begin UObject Interface
 	virtual void Serialize(FArchive& Ar) override;
-	// End of UObject interface
+	//~ End UObject Interface
 
-	// Begin UEdGraphNode interface
+	//~ Begin UEdGraphNode Interface
 	virtual void AllocateDefaultPins() override;
 	virtual FText GetNodeTitle(ENodeTitleType::Type TitleType) const override;
 	virtual FText GetTooltipText() const override;
@@ -28,12 +28,12 @@ class UK2Node_SetFieldsInStruct : public UK2Node_MakeStruct
 	virtual void ValidateNodeDuringCompilation(class FCompilerResultsLog& MessageLog) const override;
 	virtual void ExpandNode(class FKismetCompilerContext& CompilerContext, UEdGraph* SourceGraph);
 	virtual void PostPlacedNewNode() override;
-	// End  UEdGraphNode interface
+	//~ End  UEdGraphNode Interface
 
-	// Begin K2Node interface
+	//~ Begin K2Node Interface
 	virtual bool IsNodePure() const override { return false; }
 	virtual class FNodeHandlingFunctor* CreateNodeHandler(class FKismetCompilerContext& CompilerContext) const override;
-	// End K2Node interface
+	//~ End K2Node Interface
 
 	BLUEPRINTGRAPH_API static bool ShowCustomPinActions(const UEdGraphPin* InGraphPin, bool bIgnorePinsNum);
 
diff --git a/Engine/Source/Editor/BlueprintGraph/Classes/K2Node_SetVariableOnPersistentFrame.h b/Engine/Source/Editor/BlueprintGraph/Classes/K2Node_SetVariableOnPersistentFrame.h
--- a/Engine/Source/Editor/BlueprintGraph/Classes/K2Node_SetVariableOnPersistentFrame.h
+++ b/Engine/Source/Editor/BlueprintGraph/Classes/K2Node_SetVariableOnPersistentFrame.h
@@ -14,13 +14,13 @@ class UK2Node_SetVariableOnPersistentFrame : public UK2Node
 	GENERATED_BODY()
 
 public:
-	// Begin UEdGraphNode interface
+	//~ Begin UEdGraphNode Interface
 	virtual void AllocateDefaultPins() override;
-	// End UEdGraphNode interface
+	//~ End UEdGraphNode Interface
 
-	// Begin K2Node interface
+	//~ Begin K2Node Interface
 	virtual bool IsNodePure() const override { return false; }
 	virtual class FNodeHandlingFunctor* CreateNodeHandler(class FKismetCompilerContext& CompilerContext) const override;
-	// End K2Node interface
+	//~ End K2Node Interface
 };
 
diff --git a/Engine/Source/Editor/BlueprintGraph/Classes/K2Node_StructOperation.h b/Engine/Source/Editor/BlueprintGraph/Classes/K2Node_StructOperation.h
--- a/Engine/Source/Editor/BlueprintGraph/Classes/K2Node_StructOperation.h
+++ b/Engine/Source/Editor/BlueprintGraph/Classes/K2Node_StructOperation.h
@@ -14,22 +14,22 @@ class UK2Node_StructOperation : public UK2Node_Variable
 	UPROPERTY()
 	UScriptStruct* StructType;
 
-	// Begin UEdGraphNode interface
+	//~ Begin UEdGraphNode Interface
 	virtual FString GetPinMetaData(FString InPinName, FName InKey) override;
-	// End UEdGraphNode interface
+	//~ End UEdGraphNode Interface
 
-	// UK2Node interface
+	//~ Begin UK2Node Interface
 	//virtual bool DrawNodeAsVariable() const override { return true; }
 	virtual bool ShouldShowNodeProperties() const override { return true; }
 	virtual void ValidateNodeDuringCompilation(class FCompilerResultsLog& MessageLog) const override {}
 	virtual bool HasExternalDependencies(TArray<class UStruct*>* OptionalOutput = NULL) const override;
-	// End of UK2Node interface
+	//~ End UK2Node Interface
 
 protected:
 	// Updater for subclasses that allow hiding pins
 	struct FStructOperationOptionalPinManager : public FOptionalPinManager
 	{
-		// FOptionalPinsUpdater interface
+		//~ Begin FOptionalPinsUpdater Interface
 		virtual void GetRecordDefaults(UProperty* TestProperty, FOptionalPinFromProperty& Record) const override
 		{
 			Record.bCanToggleVisibility = true;
diff --git a/Engine/Source/Editor/BlueprintGraph/Classes/K2Node_Tunnel.h b/Engine/Source/Editor/BlueprintGraph/Classes/K2Node_Tunnel.h
--- a/Engine/Source/Editor/BlueprintGraph/Classes/K2Node_Tunnel.h
+++ b/Engine/Source/Editor/BlueprintGraph/Classes/K2Node_Tunnel.h
@@ -33,7 +33,7 @@ class UK2Node_Tunnel : public UK2Node_EditablePinBase
 	UPROPERTY()
 	struct FKismetUserDeclaredFunctionMetadata MetaData;
 
-	// Begin UEdGraphNode interface.
+	//~ Begin UEdGraphNode Interface.
 	virtual void DestroyNode() override;
 	virtual void PostPasteNode() override;
 	virtual FText GetTooltipText() const override;
@@ -42,22 +42,22 @@ class UK2Node_Tunnel : public UK2Node_EditablePinBase
 	BLUEPRINTGRAPH_API virtual bool CanDuplicateNode() const override;
 	virtual UObject* GetJumpTargetForDoubleClick() const override;
 	virtual FString CreateUniquePinName(FString SourcePinName) const override;
-	// End UEdGraphNode interface
+	//~ End UEdGraphNode Interface
 
-	// Begin UK2Node interface.
+	//~ Begin UK2Node Interface.
 	BLUEPRINTGRAPH_API virtual bool IsNodeSafeToIgnore() const override;
 	virtual bool DrawNodeAsEntry() const override;
 	virtual bool DrawNodeAsExit() const override;
 	virtual bool NodeCausesStructuralBlueprintChange() const override { return true; }
-	// End UK2Node interface
+	//~ End UK2Node Interface
 
-	// Begin UK2Node_EditablePinBase interface.
+	//~ Begin UK2Node_EditablePinBase Interface.
 	BLUEPRINTGRAPH_API virtual UEdGraphPin* CreatePinFromUserDefinition(const TSharedPtr<FUserPinInfo> NewPinInfo) override;
 	BLUEPRINTGRAPH_API virtual bool CanModifyExecutionWires() override;
 	BLUEPRINTGRAPH_API virtual ERenamePinResult RenameUserDefinedPin(const FString& OldName, const FString& NewName, bool bTest = false) override;
 	virtual bool CanUseRefParams() const override { return true; }
 	virtual bool CanCreateUserDefinedPin(const FEdGraphPinType& InPinType, EEdGraphPinDirection InDesiredDirection, FText& OutErrorMessage) override;
-	// End UK2Node_EditablePinBase interface
+	//~ End UK2Node_EditablePinBase Interface
 public:
 	// The input pins of this tunnel go to the output pins of InputSinkNode (can be NULL).
 	BLUEPRINTGRAPH_API virtual UK2Node_Tunnel* GetInputSink() const;
diff --git a/Engine/Source/Editor/BlueprintGraph/Classes/K2Node_Variable.h b/Engine/Source/Editor/BlueprintGraph/Classes/K2Node_Variable.h
--- a/Engine/Source/Editor/BlueprintGraph/Classes/K2Node_Variable.h
+++ b/Engine/Source/Editor/BlueprintGraph/Classes/K2Node_Variable.h
@@ -51,11 +51,11 @@ protected:
 	bool RemapRestrictedLinkReference(FName OldVariableName, FName NewVariableName, const UClass* MatchInVariableClass, const UClass* RemapIfLinkedToClass, bool bLogWarning);
 
 public:
-	// Begin UObject interface
+	//~ Begin UObject Interface
 	virtual void Serialize(FArchive& Ar) override;
-	// End UObject interface
+	//~ End UObject Interface
 
-	// Begin UEdGraphNode interface
+	//~ Begin UEdGraphNode Interface
 	virtual FLinearColor GetNodeTitleColor() const override;
 	virtual void ReconstructNode() override;
 	virtual FString GetDocumentationLink() const override;
@@ -65,16 +65,16 @@ public:
 	virtual void AutowireNewNode(UEdGraphPin* FromPin) override;
 	virtual bool CanPasteHere(const UEdGraph* TargetGraph) const override;
 	virtual void PostPasteNode() override;
-	// End UEdGraphNode interface
+	//~ End UEdGraphNode Interface
 
-	// Begin K2Node interface
+	//~ Begin K2Node Interface
 	virtual bool HasExternalDependencies(TArray<class UStruct*>* OptionalOutput) const override;
 	virtual bool DrawNodeAsVariable() const override { return true; }
 	virtual ERedirectType DoPinsMatchForReconstruction(const UEdGraphPin* NewPin, int32 NewPinIndex, const UEdGraphPin* OldPin, int32 OldPinIndex)  const override;
 	virtual void ValidateNodeDuringCompilation(class FCompilerResultsLog& MessageLog) const override;
 	virtual FText GetToolTipHeading() const override;
 	virtual void GetNodeAttributes(TArray<TKeyValuePair<FString, FString>>& OutNodeAttributes) const override;
-	// End K2Node interface
+	//~ End K2Node Interface
 
 	/** Set up this variable node from the supplied UProperty */
 	void SetFromProperty(const UProperty* Property, bool bSelfContext);
diff --git a/Engine/Source/Editor/BlueprintGraph/Classes/K2Node_VariableGet.h b/Engine/Source/Editor/BlueprintGraph/Classes/K2Node_VariableGet.h
--- a/Engine/Source/Editor/BlueprintGraph/Classes/K2Node_VariableGet.h
+++ b/Engine/Source/Editor/BlueprintGraph/Classes/K2Node_VariableGet.h
@@ -11,24 +11,24 @@ class BLUEPRINTGRAPH_API UK2Node_VariableGet : public UK2Node_Variable
 {
 	GENERATED_UCLASS_BODY()
 
-	// Begin UObject interface
+	//~ Begin UObject Interface
 	virtual void Serialize(FArchive& Ar) override;
-	// End of UObject interface
+	//~ End UObject Interface
 
-	// Begin UEdGraphNode interface
+	//~ Begin UEdGraphNode Interface
 	virtual void AllocateDefaultPins() override;
 	virtual FText GetTooltipText() const override;
 	virtual FText GetNodeTitle(ENodeTitleType::Type TitleType) const override;
 	virtual void GetContextMenuActions(const FGraphNodeContextMenuBuilder& Context) const override;
-	// End UEdGraphNode interface
+	//~ End UEdGraphNode Interface
 
-	// Begin K2Node interface
+	//~ Begin K2Node Interface
 	virtual bool IsNodePure() const override { return bIsPureGet; }
 	virtual void ReallocatePinsDuringReconstruction(TArray<UEdGraphPin*>& OldPins) override;
 	virtual bool ShouldShowNodeProperties() const override { return true; }
 	virtual class FNodeHandlingFunctor* CreateNodeHandler(class FKismetCompilerContext& CompilerContext) const override;
 	virtual void ExpandNode(class FKismetCompilerContext& CompilerContext, UEdGraph* SourceGraph) override;
-	// End K2Node interface
+	//~ End K2Node Interface
 
 	static FText GetPropertyTooltip(UProperty const* VariableProperty);
 	static FText GetBlueprintVarTooltip(FBPVariableDescription const& VarDesc);
diff --git a/Engine/Source/Editor/Cascade/Classes/CascadeParticleSystemComponent.h b/Engine/Source/Editor/Cascade/Classes/CascadeParticleSystemComponent.h
--- a/Engine/Source/Editor/Cascade/Classes/CascadeParticleSystemComponent.h
+++ b/Engine/Source/Editor/Cascade/Classes/CascadeParticleSystemComponent.h
@@ -17,9 +17,9 @@ class UCascadeParticleSystemComponent : public UParticleSystemComponent
 public:	
 	class FCascadeEdPreviewViewportClient*	CascadePreviewViewportPtr;
 
-	// Begin UActorComponent Interface
+	//~ Begin UActorComponent Interface
 	virtual void TickComponent(float DeltaTime, enum ELevelTick TickType, FActorComponentTickFunction *ThisTickFunction) override;
-	// End UActorComponent Interface
+	//~ End UActorComponent Interface
 
 	/**
 	 * Updates time dependent state for this component, called from within Cascade.
diff --git a/Engine/Source/Editor/Cascade/Private/Cascade.h b/Engine/Source/Editor/Cascade/Private/Cascade.h
--- a/Engine/Source/Editor/Cascade/Private/Cascade.h
+++ b/Engine/Source/Editor/Cascade/Private/Cascade.h
@@ -193,7 +193,7 @@ public:
 	static bool ConvertAllModulesToSeeded(UParticleSystem* ParticleSystem);
 
 private:
-	// Begin FEditorUndoClient Interface
+	//~ Begin FEditorUndoClient Interface
 	virtual void PostUndo(bool bSuccess) override;
 	virtual void PostRedo(bool bSuccess) override { PostUndo(bSuccess); }
 	// End of FEditorUndoClient
diff --git a/Engine/Source/Editor/ComponentVisualizers/Private/AudioComponentVisualizer.h b/Engine/Source/Editor/ComponentVisualizers/Private/AudioComponentVisualizer.h
--- a/Engine/Source/Editor/ComponentVisualizers/Private/AudioComponentVisualizer.h
+++ b/Engine/Source/Editor/ComponentVisualizers/Private/AudioComponentVisualizer.h
@@ -7,7 +7,7 @@
 class FAudioComponentVisualizer : public FComponentVisualizer
 {
 public:
-	// Begin FComponentVisualizer interface
+	//~ Begin FComponentVisualizer Interface
 	virtual void DrawVisualization(const UActorComponent* Component, const FSceneView* View, FPrimitiveDrawInterface* PDI) override;
-	// End FComponentVisualizer interface
+	//~ End FComponentVisualizer Interface
 };
diff --git a/Engine/Source/Editor/ComponentVisualizers/Private/ConstraintComponentVisualizer.h b/Engine/Source/Editor/ComponentVisualizers/Private/ConstraintComponentVisualizer.h
--- a/Engine/Source/Editor/ComponentVisualizers/Private/ConstraintComponentVisualizer.h
+++ b/Engine/Source/Editor/ComponentVisualizers/Private/ConstraintComponentVisualizer.h
@@ -7,7 +7,7 @@
 class FConstraintComponentVisualizer : public FComponentVisualizer
 {
 public:
-	// Begin FComponentVisualizer interface
+	//~ Begin FComponentVisualizer Interface
 	virtual void DrawVisualization(const UActorComponent* Component, const FSceneView* View, FPrimitiveDrawInterface* PDI) override;
-	// End FComponentVisualizer interface
+	//~ End FComponentVisualizer Interface
 };
diff --git a/Engine/Source/Editor/ConfigEditor/Private/ConfigEditorModule.h b/Engine/Source/Editor/ConfigEditor/Private/ConfigEditorModule.h
--- a/Engine/Source/Editor/ConfigEditor/Private/ConfigEditorModule.h
+++ b/Engine/Source/Editor/ConfigEditor/Private/ConfigEditorModule.h
@@ -12,18 +12,18 @@
 class FConfigEditorModule : public IConfigEditorModule
 {
 public:
-	// Begin IModuleInterface interface
+	//~ Begin IModuleInterface Interface
 	virtual void StartupModule() override;
 	virtual void ShutdownModule() override;
-	// End IModuleInterface interface
+	//~ End IModuleInterface Interface
 
 
 public:
-	// Begin IConfigEditorModule interface
+	//~ Begin IConfigEditorModule Interface
 	virtual void CreateHierarchyEditor(UProperty* EditProperty) override;
 	virtual void AddExternalPropertyValueWidgetAndConfigPairing(const FString& ConfigFile, const TSharedPtr<SWidget> ValueWidget) override;
 	virtual TSharedRef<SWidget> GetValueWidgetForConfigProperty(const FString& ConfigFile) override;
-	// End IConfigEditorModule interface	
+	//~ End IConfigEditorModule Interface	
 	
 
 private:
diff --git a/Engine/Source/Editor/ConfigEditor/Private/SConfigEditor.h b/Engine/Source/Editor/ConfigEditor/Private/SConfigEditor.h
--- a/Engine/Source/Editor/ConfigEditor/Private/SConfigEditor.h
+++ b/Engine/Source/Editor/ConfigEditor/Private/SConfigEditor.h
@@ -16,9 +16,9 @@ public:
 	SLATE_USER_ARGS(SConfigEditor) {}
 	SLATE_END_ARGS()
 
-	// Begin SCompoundWidget|SWidget interface
+	//~ Begin SCompoundWidget|SWidget Interface
 	virtual void Construct(const FArguments& InArgs, TWeakObjectPtr<UProperty> InEditProperty);
-	// End SCompoundWidget|SWidget interface
+	//~ End SCompoundWidget|SWidget Interface
 
 private:
 	/**
diff --git a/Engine/Source/Editor/ConfigEditor/Private/STargetPlatformSelector.h b/Engine/Source/Editor/ConfigEditor/Private/STargetPlatformSelector.h
--- a/Engine/Source/Editor/ConfigEditor/Private/STargetPlatformSelector.h
+++ b/Engine/Source/Editor/ConfigEditor/Private/STargetPlatformSelector.h
@@ -14,9 +14,9 @@ public:
 		SLATE_EVENT(FSimpleDelegate, OnTargetPlatformChanged)
 	SLATE_END_ARGS()
 
-	// Begin SCompoundWidget|SWidget interface
+	//~ Begin SCompoundWidget|SWidget Interface
 	void Construct(const FArguments& InArgs);
-	// End SCompoundWidget|SWidget interface
+	//~ End SCompoundWidget|SWidget Interface
 
 	/**
 	 * Access to the selected target platform in this widget.
diff --git a/Engine/Source/Editor/DeviceProfileEditor/Public/DeviceProfileEditorModule.h b/Engine/Source/Editor/DeviceProfileEditor/Public/DeviceProfileEditorModule.h
--- a/Engine/Source/Editor/DeviceProfileEditor/Public/DeviceProfileEditorModule.h
+++ b/Engine/Source/Editor/DeviceProfileEditor/Public/DeviceProfileEditorModule.h
@@ -9,11 +9,11 @@ class FDeviceProfileEditorModule : public IModuleInterface
 
 public:
 
-	// Begin IModuleInterface interface
+	//~ Begin IModuleInterface Interface
 	virtual void StartupModule() override;
 
 	virtual void ShutdownModule() override;
-	// End IModuleInterface interface
+	//~ End IModuleInterface Interface
 
 public:
 	/**
diff --git a/Engine/Source/Editor/EnvironmentQueryEditor/Classes/EdGraphSchema_EnvironmentQuery.h b/Engine/Source/Editor/EnvironmentQueryEditor/Classes/EdGraphSchema_EnvironmentQuery.h
--- a/Engine/Source/Editor/EnvironmentQueryEditor/Classes/EdGraphSchema_EnvironmentQuery.h
+++ b/Engine/Source/Editor/EnvironmentQueryEditor/Classes/EdGraphSchema_EnvironmentQuery.h
@@ -10,15 +10,15 @@ class UEdGraphSchema_EnvironmentQuery : public UAIGraphSchema
 {
 	GENERATED_UCLASS_BODY()
 
-	// Begin EdGraphSchema interface
+	//~ Begin EdGraphSchema Interface
 	virtual void CreateDefaultNodesForGraph(UEdGraph& Graph) const override;
 	virtual void GetGraphContextActions(FGraphContextMenuBuilder& ContextMenuBuilder) const override;
 	virtual const FPinConnectionResponse CanCreateConnection(const UEdGraphPin* A, const UEdGraphPin* B) const override;
 	virtual const FPinConnectionResponse CanMergeNodes(const UEdGraphNode* A, const UEdGraphNode* B) const override;
 	virtual int32 GetNodeSelectionCount(const UEdGraph* Graph) const override;
-	// End EdGraphSchema interface
+	//~ End EdGraphSchema Interface
 
-	// Begin UAIGraphSchema interface
+	//~ Begin UAIGraphSchema Interface
 	virtual void GetSubNodeClasses(int32 SubNodeFlags, TArray<FGraphNodeClassData>& ClassData, UClass*& GraphNodeClass) const override;
-	// End UAIGraphSchema interface
+	//~ End UAIGraphSchema Interface
 };
diff --git a/Engine/Source/Editor/EnvironmentQueryEditor/Private/DetailCustomizations/EnvDirectionCustomization.h b/Engine/Source/Editor/EnvironmentQueryEditor/Private/DetailCustomizations/EnvDirectionCustomization.h
--- a/Engine/Source/Editor/EnvironmentQueryEditor/Private/DetailCustomizations/EnvDirectionCustomization.h
+++ b/Engine/Source/Editor/EnvironmentQueryEditor/Private/DetailCustomizations/EnvDirectionCustomization.h
@@ -9,10 +9,10 @@ class FEnvDirectionCustomization : public IPropertyTypeCustomization
 {
 public:
 
-	// Begin IPropertyTypeCustomization interface
+	//~ Begin IPropertyTypeCustomization Interface
 	virtual void CustomizeHeader( TSharedRef<class IPropertyHandle> StructPropertyHandle, class FDetailWidgetRow& HeaderRow, IPropertyTypeCustomizationUtils& StructCustomizationUtils ) override;
 	virtual void CustomizeChildren( TSharedRef<class IPropertyHandle> StructPropertyHandle, class IDetailChildrenBuilder& StructBuilder, IPropertyTypeCustomizationUtils& StructCustomizationUtils ) override;
-	// End IPropertyTypeCustomization interface
+	//~ End IPropertyTypeCustomization Interface
 
 	static TSharedRef<IPropertyTypeCustomization> MakeInstance( );
 
diff --git a/Engine/Source/Editor/EnvironmentQueryEditor/Private/DetailCustomizations/EnvTraceDataCustomization.h b/Engine/Source/Editor/EnvironmentQueryEditor/Private/DetailCustomizations/EnvTraceDataCustomization.h
--- a/Engine/Source/Editor/EnvironmentQueryEditor/Private/DetailCustomizations/EnvTraceDataCustomization.h
+++ b/Engine/Source/Editor/EnvironmentQueryEditor/Private/DetailCustomizations/EnvTraceDataCustomization.h
@@ -11,10 +11,10 @@ public:
 	/** Makes a new instance of this detail layout class for a specific detail view requesting it */
 	static TSharedRef<IPropertyTypeCustomization> MakeInstance();
 
-	// Begin IPropertyTypeCustomization interface
+	//~ Begin IPropertyTypeCustomization Interface
 	virtual void CustomizeHeader( TSharedRef<class IPropertyHandle> StructPropertyHandle, class FDetailWidgetRow& HeaderRow, IPropertyTypeCustomizationUtils& StructCustomizationUtils ) override;
 	virtual void CustomizeChildren( TSharedRef<class IPropertyHandle> StructPropertyHandle, class IDetailChildrenBuilder& StructBuilder, IPropertyTypeCustomizationUtils& StructCustomizationUtils ) override;
-	// End IPropertyTypeCustomization interface
+	//~ End IPropertyTypeCustomization Interface
 
 protected:
 
diff --git a/Engine/Source/Editor/EnvironmentQueryEditor/Private/EnvironmentQueryEditor.h b/Engine/Source/Editor/EnvironmentQueryEditor/Private/EnvironmentQueryEditor.h
--- a/Engine/Source/Editor/EnvironmentQueryEditor/Private/EnvironmentQueryEditor.h
+++ b/Engine/Source/Editor/EnvironmentQueryEditor/Private/EnvironmentQueryEditor.h
@@ -16,16 +16,16 @@ public:
 
 	void InitEnvironmentQueryEditor( const EToolkitMode::Type Mode, const TSharedPtr< class IToolkitHost >& InitToolkitHost, UEnvQuery* Script );
 
-	// Begin IToolkit interface
+	//~ Begin IToolkit Interface
 	virtual FName GetToolkitFName() const override;
 	virtual FText GetBaseToolkitName() const override;
 	virtual FString GetWorldCentricTabPrefix() const override;
 	virtual FLinearColor GetWorldCentricTabColorScale() const override;
-	// End IToolkit interface
+	//~ End IToolkit Interface
 
-	// Begin IEnvironmentQueryEditor interface
+	//~ Begin IEnvironmentQueryEditor Interface
 	virtual uint32 GetSelectedNodesCount() const override { return SelectedNodesCount; }
-	// End IEnvironmentQueryEditor interface
+	//~ End IEnvironmentQueryEditor Interface
 
 	void OnFinishedChangingProperties(const FPropertyChangedEvent& PropertyChangedEvent);
 
diff --git a/Engine/Source/Editor/GameplayAbilitiesEditor/Classes/GameplayAbilitiesBlueprintFactory.h b/Engine/Source/Editor/GameplayAbilitiesEditor/Classes/GameplayAbilitiesBlueprintFactory.h
--- a/Engine/Source/Editor/GameplayAbilitiesEditor/Classes/GameplayAbilitiesBlueprintFactory.h
+++ b/Engine/Source/Editor/GameplayAbilitiesEditor/Classes/GameplayAbilitiesBlueprintFactory.h
@@ -17,9 +17,9 @@ class UGameplayAbilitiesBlueprintFactory : public UFactory
 	UPROPERTY(EditAnywhere, Category=GameplayAbilitiesBlueprintFactory)
 	TSubclassOf<class UGameplayAbility> ParentClass;
 
-	// Begin UFactory Interface
+	//~ Begin UFactory Interface
 	virtual bool ConfigureProperties() override;
 	virtual UObject* FactoryCreateNew(UClass* Class, UObject* InParent, FName Name, EObjectFlags Flags, UObject* Context, FFeedbackContext* Warn, FName CallingContext) override;
 	virtual UObject* FactoryCreateNew(UClass* Class, UObject* InParent, FName Name, EObjectFlags Flags, UObject* Context, FFeedbackContext* Warn) override;
-	// Begin UFactory Interface	
+	//~ Begin UFactory Interface	
 };
diff --git a/Engine/Source/Editor/GameplayAbilitiesEditor/Classes/K2Node_GameplayCueEvent.h b/Engine/Source/Editor/GameplayAbilitiesEditor/Classes/K2Node_GameplayCueEvent.h
--- a/Engine/Source/Editor/GameplayAbilitiesEditor/Classes/K2Node_GameplayCueEvent.h
+++ b/Engine/Source/Editor/GameplayAbilitiesEditor/Classes/K2Node_GameplayCueEvent.h
@@ -9,17 +9,17 @@ class UK2Node_GameplayCueEvent : public UK2Node_Event
 {
 	GENERATED_UCLASS_BODY()
 
-	// UObject interface
+	//~ Begin UObject Interface
 	virtual void Serialize(FArchive& Ar) override;
-	// End of UObject interface
+	//~ End UObject Interface
 
-	// Begin UEdGraphNode interface
+	//~ Begin UEdGraphNode Interface
 	virtual FText GetTooltipText() const override;
 	virtual FText GetNodeTitle(ENodeTitleType::Type TitleType) const override;
 	virtual bool IsCompatibleWithGraph(const UEdGraph* TargetGraph) const override;
-	// End UEdGraphNode interface
+	//~ End UEdGraphNode Interface
 	
-	// Begin UK2Node interface
+	//~ Begin UK2Node Interface
 	virtual void GetMenuActions(FBlueprintActionDatabaseRegistrar& ActionRegistrar) const override;
-	// End UK2Node interface
+	//~ End UK2Node Interface
 };
\ No newline at end of file
diff --git a/Engine/Source/Editor/GameplayAbilitiesEditor/Private/SGameplayAttributeGraphPin.h b/Engine/Source/Editor/GameplayAbilitiesEditor/Private/SGameplayAttributeGraphPin.h
--- a/Engine/Source/Editor/GameplayAbilitiesEditor/Private/SGameplayAttributeGraphPin.h
+++ b/Engine/Source/Editor/GameplayAbilitiesEditor/Private/SGameplayAttributeGraphPin.h
@@ -13,9 +13,9 @@ public:
 
 	void Construct(const FArguments& InArgs, UEdGraphPin* InGraphPinObj);
 
-	// Begin SGraphPin interface
+	//~ Begin SGraphPin Interface
 	virtual TSharedRef<SWidget>	GetDefaultValueWidget() override;
-	// End SGraphPin interface
+	//~ End SGraphPin Interface
 
 	void OnAttributeChanged(UProperty* SelectedAttribute);
 
diff --git a/Engine/Source/Editor/GameplayTagsEditor/Private/GameplayTagContainerCustomization.h b/Engine/Source/Editor/GameplayTagsEditor/Private/GameplayTagContainerCustomization.h
--- a/Engine/Source/Editor/GameplayTagsEditor/Private/GameplayTagContainerCustomization.h
+++ b/Engine/Source/Editor/GameplayTagsEditor/Private/GameplayTagContainerCustomization.h
@@ -23,10 +23,10 @@ public:
 	/** Overridden to do nothing */
 	virtual void CustomizeChildren(TSharedRef<IPropertyHandle> InStructPropertyHandle, class IDetailChildrenBuilder& ChildBuilder, IPropertyTypeCustomizationUtils& StructCustomizationUtils) override {}
 
-	// Begin FEditorUndoClient Interface
+	//~ Begin FEditorUndoClient Interface
 	virtual void PostUndo( bool bSuccess ) override;
 	virtual void PostRedo( bool bSuccess ) override;	
-	// End FEditorUndoClient Interface
+	//~ End FEditorUndoClient Interface
 
 private:
 	/** Called when the edit button is clicked; Launches the gameplay tag editor */
diff --git a/Engine/Source/Editor/GameplayTagsEditor/Private/GameplayTagCustomization.h b/Engine/Source/Editor/GameplayTagsEditor/Private/GameplayTagCustomization.h
--- a/Engine/Source/Editor/GameplayTagsEditor/Private/GameplayTagCustomization.h
+++ b/Engine/Source/Editor/GameplayTagsEditor/Private/GameplayTagCustomization.h
@@ -23,10 +23,10 @@ public:
 	/** Overridden to do nothing */
 	virtual void CustomizeChildren(TSharedRef<IPropertyHandle> InStructPropertyHandle, class IDetailChildrenBuilder& ChildBuilder, IPropertyTypeCustomizationUtils& StructCustomizationUtils) override {}
 
-	// Begin FEditorUndoClient Interface
+	//~ Begin FEditorUndoClient Interface
 	virtual void PostUndo(bool bSuccess) override;
 	virtual void PostRedo(bool bSuccess) override;
-	// End FEditorUndoClient Interface
+	//~ End FEditorUndoClient Interface
 
 private:
 
diff --git a/Engine/Source/Editor/GameplayTagsEditor/Private/SGameplayTagContainerGraphPin.h b/Engine/Source/Editor/GameplayTagsEditor/Private/SGameplayTagContainerGraphPin.h
--- a/Engine/Source/Editor/GameplayTagsEditor/Private/SGameplayTagContainerGraphPin.h
+++ b/Engine/Source/Editor/GameplayTagsEditor/Private/SGameplayTagContainerGraphPin.h
@@ -14,9 +14,9 @@ public:
 
 	void Construct(const FArguments& InArgs, UEdGraphPin* InGraphPinObj);
 
-	// Begin SGraphPin interface
+	//~ Begin SGraphPin Interface
 	virtual TSharedRef<SWidget>	GetDefaultValueWidget() override;
-	// End SGraphPin interface
+	//~ End SGraphPin Interface
 
 private:
 
diff --git a/Engine/Source/Editor/GeometryMode/Classes/GeomModifier_Clip.h b/Engine/Source/Editor/GeometryMode/Classes/GeomModifier_Clip.h
--- a/Engine/Source/Editor/GeometryMode/Classes/GeomModifier_Clip.h
+++ b/Engine/Source/Editor/GeometryMode/Classes/GeomModifier_Clip.h
@@ -28,7 +28,7 @@ class UGeomModifier_Clip : public UGeomModifier_Edit
 	FVector SnappedMouseWorldSpacePos;
 
 
-	// Begin UGeomModifier Interface
+	//~ Begin UGeomModifier Interface
 	virtual bool Supports() override;
 	virtual bool InputKey(class FEditorViewportClient* ViewportClient, FViewport* Viewport, FKey Key, EInputEvent Event) override;
 	virtual void Render(const FSceneView* View,FViewport* Viewport,FPrimitiveDrawInterface* PDI) override;
@@ -37,7 +37,7 @@ class UGeomModifier_Clip : public UGeomModifier_Edit
 	virtual void WasActivated() override;
 protected:
 	virtual bool OnApply() override;
-	// End UGeomModifier Interface
+	//~ End UGeomModifier Interface
 
 private:
 	// @todo document
diff --git a/Engine/Source/Editor/GeometryMode/Classes/GeomModifier_Create.h b/Engine/Source/Editor/GeometryMode/Classes/GeomModifier_Create.h
--- a/Engine/Source/Editor/GeometryMode/Classes/GeomModifier_Create.h
+++ b/Engine/Source/Editor/GeometryMode/Classes/GeomModifier_Create.h
@@ -14,11 +14,11 @@ class UGeomModifier_Create : public UGeomModifier_Edit
 	GENERATED_UCLASS_BODY()
 
 
-	// Begin UGeomModifier Interface
+	//~ Begin UGeomModifier Interface
 	virtual bool Supports() override;
 protected:
 	virtual bool OnApply() override;
-	// End UGeomModifier Interface
+	//~ End UGeomModifier Interface
 };
 
 
diff --git a/Engine/Source/Editor/GeometryMode/Classes/GeomModifier_Edit.h b/Engine/Source/Editor/GeometryMode/Classes/GeomModifier_Edit.h
--- a/Engine/Source/Editor/GeometryMode/Classes/GeomModifier_Edit.h
+++ b/Engine/Source/Editor/GeometryMode/Classes/GeomModifier_Edit.h
@@ -6,9 +6,9 @@
  */
 
 
-//=============================================================================
+//~=============================================================================
 // GeomModifier_Edit: Maniupalating selected objects with the widget
-//=============================================================================
+//~=============================================================================
 
 #pragma once
 #include "GeomModifier_Edit.generated.h"
@@ -20,9 +20,9 @@ class UGeomModifier_Edit : public UGeomModifier
 
 
 
-	// Begin UGeomModifier Interface
+	//~ Begin UGeomModifier Interface
 	virtual bool InputDelta(class FEditorViewportClient* InViewportClient,FViewport* InViewport,FVector& InDrag,FRotator& InRot,FVector& InScale) override;	
-	// End UGeomModifier Interface
+	//~ End UGeomModifier Interface
 };
 
 
diff --git a/Engine/Source/Editor/GraphEditor/Private/KismetNodes/SGraphNodeK2Base.h b/Engine/Source/Editor/GraphEditor/Private/KismetNodes/SGraphNodeK2Base.h
--- a/Engine/Source/Editor/GraphEditor/Private/KismetNodes/SGraphNodeK2Base.h
+++ b/Engine/Source/Editor/GraphEditor/Private/KismetNodes/SGraphNodeK2Base.h
@@ -6,10 +6,10 @@ class GRAPHEDITOR_API SGraphNodeK2Base : public SGraphNode
 {
 public:
 
-	// SGraphNode interface
+	//~ Begin SGraphNode Interface
 	virtual void UpdateGraphNode() override;
 
-	// SNodePanel::SNode interface
+	//~ Begin SNodePanel::SNode Interface
 	virtual bool RequiresSecondPassLayout() const override;
 	virtual void GetOverlayBrushes(bool bSelected, const FVector2D WidgetSize, TArray<FOverlayBrushInfo>& Brushes) const override;
 	virtual void GetNodeInfoPopups(FNodeInfoContext* Context, TArray<FGraphInformationPopupInfo>& Popups) const override;
@@ -17,9 +17,9 @@ public:
 	void PerformSecondPassLayout(const TMap< UObject*, TSharedRef<SNode> >& NodeToWidgetLookup) const override;
 
 protected :
-	// Begin SGraphNode interface
+	//~ Begin SGraphNode Interface
 	virtual TSharedPtr<SToolTip> GetComplexTooltip() override;
-	// End SGraphNode interface
+	//~ End SGraphNode Interface
 
 	/** Set up node in 'standard' mode */
 	void UpdateStandardNode();
diff --git a/Engine/Source/Editor/GraphEditor/Private/KismetPins/SGraphPinClass.h b/Engine/Source/Editor/GraphEditor/Private/KismetPins/SGraphPinClass.h
--- a/Engine/Source/Editor/GraphEditor/Private/KismetPins/SGraphPinClass.h
+++ b/Engine/Source/Editor/GraphEditor/Private/KismetPins/SGraphPinClass.h
@@ -19,12 +19,12 @@ protected:
 	// Called when a new class was picked via the asset picker
 	void OnPickedNewClass(UClass* ChosenClass);
 
-	// Begin SGraphPinObject interface
+	//~ Begin SGraphPinObject Interface
 	virtual FReply OnClickUse() override;
 	virtual bool AllowSelfPinWidget() const override { return false; }
 	virtual TSharedRef<SWidget> GenerateAssetPicker() override;
 	virtual FText GetDefaultComboText() const override;
 	virtual FOnClicked GetOnUseButtonDelegate() override;
-	// End SGraphPinObject interface
+	//~ End SGraphPinObject Interface
 
 };
diff --git a/Engine/Source/Editor/GraphEditor/Private/KismetPins/SGraphPinInteger.h b/Engine/Source/Editor/GraphEditor/Private/KismetPins/SGraphPinInteger.h
--- a/Engine/Source/Editor/GraphEditor/Private/KismetPins/SGraphPinInteger.h
+++ b/Engine/Source/Editor/GraphEditor/Private/KismetPins/SGraphPinInteger.h
@@ -11,7 +11,7 @@ public:
 	void Construct(const FArguments& InArgs, UEdGraphPin* InGraphPinObj);
 
 protected:
-	// Begin SGraphPinString interface
+	//~ Begin SGraphPinString Interface
 	virtual void SetTypeInValue(const FText& NewTypeInValue, ETextCommit::Type CommitInfo) override;
-	// End of SGraphPinString interface
+	//~ End SGraphPinString Interface
 };
diff --git a/Engine/Source/Editor/GraphEditor/Public/SCommentBubble.h b/Engine/Source/Editor/GraphEditor/Public/SCommentBubble.h
--- a/Engine/Source/Editor/GraphEditor/Public/SCommentBubble.h
+++ b/Engine/Source/Editor/GraphEditor/Public/SCommentBubble.h
@@ -66,10 +66,10 @@ class GRAPHEDITOR_API SCommentBubble : public SCompoundWidget
 	void Construct( const FArguments& InArgs );
 
 
-	// Begin SWidget interface
+	//~ Begin SWidget Interface
 	virtual FCursorReply OnCursorQuery( const FGeometry& MyGeometry, const FPointerEvent& CursorEvent ) const override;
 	virtual void Tick( const FGeometry& AllottedGeometry, const double InCurrentTime, const float InDeltaTime ) override;
-	// End SWidget Interface
+	//~ End SWidget Interface
 
 	/** Returns the offset from the SNode center slot */
 	FVector2D GetOffset() const;
diff --git a/Engine/Source/Editor/GraphEditor/Public/SGraphNodeComment.h b/Engine/Source/Editor/GraphEditor/Public/SGraphNodeComment.h
--- a/Engine/Source/Editor/GraphEditor/Public/SGraphNodeComment.h
+++ b/Engine/Source/Editor/GraphEditor/Public/SGraphNodeComment.h
@@ -8,31 +8,31 @@ public:
 	SLATE_BEGIN_ARGS(SGraphNodeComment){}
 	SLATE_END_ARGS()
 
-	// Begin SWidget Interface
+	//~ Begin SWidget Interface
 	virtual FReply OnMouseButtonDoubleClick( const FGeometry& InMyGeometry, const FPointerEvent& InMouseEvent ) override;
 	virtual FReply OnMouseButtonUp( const FGeometry& MyGeometry, const FPointerEvent& MouseEvent ) override;
 	virtual void Tick(const FGeometry& AllottedGeometry, const double InCurrentTime, const float InDeltaTime) override;
 	virtual FReply OnDrop( const FGeometry& MyGeometry, const FDragDropEvent& DragDropEvent ) override;
 	virtual void OnDragEnter( const FGeometry& MyGeometry, const FDragDropEvent& DragDropEvent ) override;
-	// End SWidget Interface
+	//~ End SWidget Interface
 
-	// Begin SNodePanel::SNode interface
+	//~ Begin SNodePanel::SNode Interface
 	virtual const FSlateBrush* GetShadowBrush(bool bSelected) const override;
 	virtual void GetOverlayBrushes(bool bSelected, const FVector2D WidgetSize, TArray<FOverlayBrushInfo>& Brushes) const override;
 	virtual bool ShouldAllowCulling() const override { return false; }
 	virtual int32 GetSortDepth() const override;
 	virtual void EndUserInteraction() const override;
 	virtual FString GetNodeComment() const override;
-	// End SNodePanel::SNode interface
+	//~ End SNodePanel::SNode Interface
 
-	// Begin SPanel Interface
+	//~ Begin SPanel Interface
 	virtual FVector2D ComputeDesiredSize(float) const override;
-	// End SPanel interface
+	//~ End SPanel Interface
 
-	// Begin SGraphNode Interface
+	//~ Begin SGraphNode Interface
 	virtual bool IsNameReadOnly() const override;
 	virtual FSlateColor GetCommentColor() const override { return GetCommentBodyColor(); }
-	// End SGraphNode Interface
+	//~ End SGraphNode Interface
 
 	void Construct( const FArguments& InArgs, UEdGraphNode* InNode );
 
@@ -46,7 +46,7 @@ public:
 	virtual FSlateRect GetTitleRect() const override;
 
 protected:
-	// SGraphNode Interface
+	//~ Begin SGraphNode Interface
 	virtual void UpdateGraphNode() override;
 	virtual void PopulateMetaTag(class FGraphNodeMetaData* TagMeta) const override;
 
@@ -60,11 +60,11 @@ protected:
 	/** called when user is moving the comment node */
 	virtual void MoveTo(const FVector2D& NewPosition, FNodeSet& NodeFilter ) override;
 
-	// SGraphNodeResizable Interface
+	//~ Begin SGraphNodeResizable Interface
 	virtual float GetTitleBarHeight() const override;
 	virtual FSlateRect GetHitTestingBorder() const override;
 	virtual FVector2D GetNodeMaximumSize() const override;
-	// SGraphNodeResizable Interface
+	//~ Begin SGraphNodeResizable Interface
 
 private:
 
diff --git a/Engine/Source/Editor/GraphEditor/Public/SGraphNodeDocumentation.h b/Engine/Source/Editor/GraphEditor/Public/SGraphNodeDocumentation.h
--- a/Engine/Source/Editor/GraphEditor/Public/SGraphNodeDocumentation.h
+++ b/Engine/Source/Editor/GraphEditor/Public/SGraphNodeDocumentation.h
@@ -10,29 +10,29 @@ public:
 	SLATE_BEGIN_ARGS(SGraphNodeDocumentation){}
 	SLATE_END_ARGS()
 
-	// Begin SPanel Interface
+	//~ Begin SPanel Interface
 	virtual FVector2D ComputeDesiredSize(float) const override;
-	// End SPanel interface
+	//~ End SPanel Interface
 
-	// Begin SWidget Interface
+	//~ Begin SWidget Interface
 	void Construct( const FArguments& InArgs, UEdGraphNode* InNode );
 	virtual FReply OnMouseWheel( const FGeometry& MyGeometry, const FPointerEvent& MouseEvent ) override;
 	virtual void Tick( const FGeometry& AllottedGeometry, const double InCurrentTime, const float InDeltaTime ) override;
-	// End SWidget Interface
+	//~ End SWidget Interface
 
-	// Begin SGraphNodeResizable Interface
+	//~ Begin SGraphNodeResizable Interface
 	virtual FVector2D GetNodeMinimumSize() const override;
 	virtual FVector2D GetNodeMaximumSize() const override;
 	virtual float GetTitleBarHeight() const override;
 	virtual FSlateRect GetHitTestingBorder() const override;
-	// End SGraphNodeResizable Interface
+	//~ End SGraphNodeResizable Interface
 
 protected:
 
-	// Begin SGraphNode Interface
+	//~ Begin SGraphNode Interface
 	virtual void UpdateGraphNode() override;
 	virtual bool IsNameReadOnly () const override { return false; }
-	// End SGraphNode Interface
+	//~ End SGraphNode Interface
 
 	/** Retrives the current documentation title based on the chosen excerpt */
 	FText GetDocumentationTitle() const;
diff --git a/Engine/Source/Editor/GraphEditor/Public/SGraphNodeResizable.h b/Engine/Source/Editor/GraphEditor/Public/SGraphNodeResizable.h
--- a/Engine/Source/Editor/GraphEditor/Public/SGraphNodeResizable.h
+++ b/Engine/Source/Editor/GraphEditor/Public/SGraphNodeResizable.h
@@ -24,14 +24,14 @@ public:
 		CRWZ_TitleBar			= 10,
 	};
 
-	// Begin SWidget Interface
+	//~ Begin SWidget Interface
 	virtual FReply OnMouseMove( const FGeometry& MyGeometry, const FPointerEvent& MouseEvent ) override;
 	virtual void OnMouseEnter( const FGeometry& MyGeometry, const FPointerEvent& MouseEvent ) override;
 	virtual void OnMouseLeave( const FPointerEvent& MouseEvent ) override;
 	virtual FReply OnMouseButtonDown( const FGeometry& MyGeometry, const FPointerEvent& MouseEvent ) override;
 	virtual FReply OnMouseButtonUp( const FGeometry& MyGeometry, const FPointerEvent& MouseEvent ) override;
 	virtual FCursorReply OnCursorQuery( const FGeometry& MyGeometry, const FPointerEvent& CursorEvent ) const override;
-	// End SWidget Interface
+	//~ End SWidget Interface
 	
 protected:
 
diff --git a/Engine/Source/Editor/HierarchicalLODOutliner/Private/HLODOutliner.h b/Engine/Source/Editor/HierarchicalLODOutliner/Private/HLODOutliner.h
--- a/Engine/Source/Editor/HierarchicalLODOutliner/Private/HLODOutliner.h
+++ b/Engine/Source/Editor/HierarchicalLODOutliner/Private/HLODOutliner.h
@@ -81,16 +81,16 @@ namespace HLODOutliner
 		/** Initializes and creates the settings view */
 		void CreateSettingsView();
 
-		// Begin SCompoundWidget Interface
+		//~ Begin SCompoundWidget Interface
 		virtual void Tick(const FGeometry& AllottedGeometry, const double InCurrentTime, const float InDeltaTime) override;
 		virtual void OnMouseEnter(const FGeometry& MyGeometry, const FPointerEvent& MouseEvent) override;
 		virtual void OnMouseLeave(const FPointerEvent& MouseEvent) override;
 		virtual FReply OnKeyDown(const FGeometry& MyGeometry, const FKeyEvent& InKeyEvent) override;
 		virtual FReply OnDrop(const FGeometry& MyGeometry, const FDragDropEvent& DragDropEvent) override;
 		virtual FReply OnDragOver(const FGeometry& MyGeometry, const FDragDropEvent& DragDropEvent)  override;
-		// End of SCompoundWidget Interface
+		//~ End SCompoundWidget Interface
 
-		// Begin FEditorUndoClient Interface
+		//~ Begin FEditorUndoClient Interface
 		virtual void PostUndo(bool bSuccess) override;
 		virtual void PostRedo(bool bSuccess) override { PostUndo(bSuccess); }
 		// End of FEditorUndoClient	
diff --git a/Engine/Source/Editor/HierarchicalLODOutliner/Private/HLODTreeWidgetItem.h b/Engine/Source/Editor/HierarchicalLODOutliner/Private/HLODTreeWidgetItem.h
--- a/Engine/Source/Editor/HierarchicalLODOutliner/Private/HLODTreeWidgetItem.h
+++ b/Engine/Source/Editor/HierarchicalLODOutliner/Private/HLODTreeWidgetItem.h
@@ -48,12 +48,12 @@ namespace HLODOutliner
 		/** Returns the tint of the node */
 		FSlateColor GetTint() const;
 		
-		// Begin STableRow interface.
+		//~ Begin STableRow Interface.
 		virtual void OnDragEnter(FGeometry const& MyGeometry, FDragDropEvent const& DragDropEvent) override;
 		virtual void OnDragLeave(FDragDropEvent const& DragDropEvent) override;
 		virtual FReply OnDragOver(const FGeometry& MyGeometry, const FDragDropEvent& DragDropEvent) override;
 		virtual FReply OnDrop(const FGeometry& MyGeometry, const FDragDropEvent& DragDropEvent) override;
-		// End STableRow interface.
+		//~ End STableRow Interface.
 
 	private:
 		/** The info about the widget that we are visualizing. */
diff --git a/Engine/Source/Editor/HierarchicalLODOutliner/Private/LODActorItem.h b/Engine/Source/Editor/HierarchicalLODOutliner/Private/LODActorItem.h
--- a/Engine/Source/Editor/HierarchicalLODOutliner/Private/LODActorItem.h
+++ b/Engine/Source/Editor/HierarchicalLODOutliner/Private/LODActorItem.h
@@ -42,13 +42,13 @@ namespace HLODOutliner
 
 		FLODActorItem(ALODActor* InLODActor);
 
-		// Begin ITreeItem interface.
+		//~ Begin ITreeItem Interface.
 		virtual bool CanInteract() const override;
 		virtual void GenerateContextMenu(FMenuBuilder& MenuBuilder, SHLODOutliner& Outliner) override;
 		virtual FString GetDisplayString() const override;
 		virtual FSlateColor GetTint() const override;
 		virtual FTreeItemID GetID() override;
-		// End ITreeItem interface.
+		//~ End ITreeItem Interface.
 
 		/** Returns the number of triangles for the represented LODActor in FText form */
 		FText GetNumTrianglesAsText() const;
diff --git a/Engine/Source/Editor/HierarchicalLODOutliner/Private/LODLevelItem.h b/Engine/Source/Editor/HierarchicalLODOutliner/Private/LODLevelItem.h
--- a/Engine/Source/Editor/HierarchicalLODOutliner/Private/LODLevelItem.h
+++ b/Engine/Source/Editor/HierarchicalLODOutliner/Private/LODLevelItem.h
@@ -31,12 +31,12 @@ namespace HLODOutliner
 	{
 		FLODLevelItem(const uint32 InLODIndex);
 
-		// Begin ITreeItem interface.
+		//~ Begin ITreeItem Interface.
 		virtual bool CanInteract() const override;
 		virtual void GenerateContextMenu(FMenuBuilder& MenuBuilder, SHLODOutliner& Outliner) override;
 		virtual FString GetDisplayString() const override;
 		virtual FTreeItemID GetID() override;
-		// Begin ITreeItem interface.
+		//~ Begin ITreeItem Interface.
 
 		/** Populate the specified drag/drop payload with any relevant information for this type */
 		virtual void PopulateDragDropPayload(FDragDropPayload& Payload) const override;
diff --git a/Engine/Source/Editor/HierarchicalLODOutliner/Private/StaticMeshActorItem.h b/Engine/Source/Editor/HierarchicalLODOutliner/Private/StaticMeshActorItem.h
--- a/Engine/Source/Editor/HierarchicalLODOutliner/Private/StaticMeshActorItem.h
+++ b/Engine/Source/Editor/HierarchicalLODOutliner/Private/StaticMeshActorItem.h
@@ -35,12 +35,12 @@ namespace HLODOutliner
 
 		FStaticMeshActorItem(AActor* InStaticMeshActor);
 
-		// Begin ITreeItem interface.
+		//~ Begin ITreeItem Interface.
 		virtual bool CanInteract() const override;
 		virtual void GenerateContextMenu(FMenuBuilder& MenuBuilder, SHLODOutliner& Outliner) override;
 		virtual FString GetDisplayString() const override;
 		virtual FTreeItemID GetID() override;
-		// End ITreeItem interface.
+		//~ End ITreeItem Interface.
 
 		/** Populate the specified drag/drop payload with any relevant information for this type */
 		virtual void PopulateDragDropPayload(FDragDropPayload& Payload) const override;
diff --git a/Engine/Source/Editor/IntroTutorials/Private/EditorTutorialImportFactory.h b/Engine/Source/Editor/IntroTutorials/Private/EditorTutorialImportFactory.h
--- a/Engine/Source/Editor/IntroTutorials/Private/EditorTutorialImportFactory.h
+++ b/Engine/Source/Editor/IntroTutorials/Private/EditorTutorialImportFactory.h
@@ -9,17 +9,17 @@ class UEditorTutorialImportFactory : public UFactory, public FReimportHandler
 {
 	GENERATED_UCLASS_BODY()
 
-	// UFactory Interface
+	//~ Begin UFactory Interface
 	virtual bool FactoryCanImport(const FString& Filename) override;
 	virtual UObject* FactoryCreateBinary(UClass* InClass, UObject* InParent, FName InName, EObjectFlags Flags, UObject* Context, const TCHAR* Type, const uint8*& Buffer, const uint8* BufferEnd, FFeedbackContext* Warn) override;
-	// End of UFactory Interface
+	//~ End UFactory Interface
 
-	// Begin FReimportHandler interface
+	//~ Begin FReimportHandler Interface
 	virtual bool CanReimport(UObject* Obj, TArray<FString>& OutFilenames) override;
 	virtual void SetReimportPaths(UObject* Obj, const TArray<FString>& NewReimportPaths) override;
 	virtual EReimportResult::Type Reimport(UObject* Obj) override;
 	virtual int32 GetPriority() const override;
-	// End FReimportHandler interface
+	//~ End FReimportHandler Interface
 
 private:
 	/** Import the tutorial from file */
diff --git a/Engine/Source/Editor/Kismet/Classes/BlueprintPaletteFavorites.h b/Engine/Source/Editor/Kismet/Classes/BlueprintPaletteFavorites.h
--- a/Engine/Source/Editor/Kismet/Classes/BlueprintPaletteFavorites.h
+++ b/Engine/Source/Editor/Kismet/Classes/BlueprintPaletteFavorites.h
@@ -91,10 +91,10 @@ class KISMET_API UBlueprintPaletteFavorites : public UObject
 	GENERATED_UCLASS_BODY()
 
 public:
-	// Begin UObject interface
+	//~ Begin UObject Interface
 	virtual void PostInitProperties() override;
 	virtual void PostEditChangeProperty(struct FPropertyChangedEvent& PropertyChangedEvent) override;
-	// End UObject interface
+	//~ End UObject Interface
 
 	/**
 	 * Not all palette actions can be turned into a favorite. This method is
diff --git a/Engine/Source/Editor/Kismet/Public/BlueprintEditor.h b/Engine/Source/Editor/Kismet/Public/BlueprintEditor.h
--- a/Engine/Source/Editor/Kismet/Public/BlueprintEditor.h
+++ b/Engine/Source/Editor/Kismet/Public/BlueprintEditor.h
@@ -114,9 +114,9 @@ class KISMET_API FBlueprintEditor : public IBlueprintEditor, public FGCObject, p
 	};
 
 public:
-	// IToolkit interface
+	//~ Begin IToolkit Interface
 	virtual void RegisterTabSpawners(const TSharedRef<class FTabManager>& TabManager) override;
-	// End of IToolkit interface
+	//~ End IToolkit Interface
 
 	DECLARE_MULTICAST_DELEGATE_OneParam(FOnModeSet, FName);
 	FOnModeSet& OnModeSet() { return OnModeSetData; }
@@ -134,13 +134,13 @@ public:
 	void InitBlueprintEditor(const EToolkitMode::Type Mode, const TSharedPtr< class IToolkitHost >& InitToolkitHost, const TArray<class UBlueprint*>& InBlueprints, bool bShouldOpenInDefaultsMode);
 
 public:
-	// FAssetEditorToolkit interface
+	//~ Begin FAssetEditorToolkit Interface
 	virtual bool OnRequestClose() override;
 	virtual void ToolkitBroughtToFront() override;
 	virtual void SaveAsset_Execute() override;
 	// End of FAssetEditorToolkit 
 
-	// IToolkit interface
+	//~ Begin IToolkit Interface
 	virtual FName GetToolkitContextFName() const override;
 	virtual FName GetToolkitFName() const override;
 	virtual FText GetBaseToolkitName() const override;
@@ -149,13 +149,13 @@ public:
 	virtual FString GetWorldCentricTabPrefix() const override;
 	virtual FLinearColor GetWorldCentricTabColorScale() const override;
 	virtual bool IsBlueprintEditor() const override;
-	// End of IToolkit interface
+	//~ End IToolkit Interface
 
-	// FGCObject interface
+	//~ Begin FGCObject Interface
 	virtual void AddReferencedObjects( FReferenceCollector& Collector ) override;
-	// End of FGCObject interface
+	//~ End FGCObject Interface
 
-	// IBlueprintEditor interface
+	//~ Begin IBlueprintEditor Interface
 	virtual void RefreshEditors(ERefreshBlueprintEditorReason::Type Reason = ERefreshBlueprintEditorReason::UnknownReason) override;
 	virtual void JumpToHyperlink(const UObject* ObjectReference, bool bRequestRename = false) override;
 	virtual void SummonSearchUI(bool bSetFindWithinBlueprint, FString NewSearchTerms = FString(), bool bSelectFirstResult = false) override;
@@ -164,13 +164,13 @@ public:
 	virtual int32 GetNumberOfSelectedNodes() const override;
 	virtual void AnalyticsTrackNodeEvent( UBlueprint* Blueprint, UEdGraphNode *GraphNode, bool bNodeDelete = false ) const override;
 	void AnalyticsTrackCompileEvent( UBlueprint* Blueprint, int32 NumErrors, int32 NumWarnings ) const;
-	// End of IBlueprintEditor interface
+	//~ End IBlueprintEditor Interface
 
-	// FTickableEditorObject interface
+	//~ Begin FTickableEditorObject Interface
 	virtual void Tick(float DeltaTime) override;
 	virtual bool IsTickable() const override { return true; }
 	virtual TStatId GetStatId() const override;
-	// End of FTickableEditorObject interface
+	//~ End FTickableEditorObject Interface
 
 public:
 	FBlueprintEditor();
@@ -865,10 +865,10 @@ protected:
 	/** Called when Find In Blueprints menu is opened is clicked */
 	void FindInBlueprints_OnClicked();
 
-	// FNotifyHook interface
+	//~ Begin FNotifyHook Interface
 	virtual void NotifyPreChange( UProperty* PropertyAboutToChange ) override;
 	virtual void NotifyPostChange( const FPropertyChangedEvent& PropertyChangedEvent, UProperty* PropertyThatChanged) override;
-	// End of FNotifyHook interface
+	//~ End FNotifyHook Interface
 
 	/** Callback when properties have finished being handled */
 	void OnFinishedChangingProperties(const FPropertyChangedEvent& PropertyChangedEvent);
@@ -886,7 +886,7 @@ protected:
 	/**Load macro & function blueprint libraries from asset registry*/
 	void LoadLibrariesFromAssetRegistry();
 
-	// Begin FEditorUndoClient Interface
+	//~ Begin FEditorUndoClient Interface
 	virtual void	PostUndo(bool bSuccess) override;
 	virtual void	PostRedo(bool bSuccess) override;
 	// End of FEditorUndoClient
diff --git a/Engine/Source/Editor/LandscapeEditor/Public/LandscapeEditorObject.h b/Engine/Source/Editor/LandscapeEditor/Public/LandscapeEditorObject.h
--- a/Engine/Source/Editor/LandscapeEditor/Public/LandscapeEditorObject.h
+++ b/Engine/Source/Editor/LandscapeEditor/Public/LandscapeEditorObject.h
@@ -547,9 +547,9 @@ class ULandscapeEditorObject : public UObject
 	TEnumAsByte<ELandscapeLayerPaintingRestriction::Type> PaintingRestriction;
 
 #if WITH_EDITOR
-	// Begin UObject Interface
+	//~ Begin UObject Interface
 	void PostEditChangeProperty(FPropertyChangedEvent& PropertyChangedEvent) override;
-	// End UObject Interface
+	//~ End UObject Interface
 #endif // WITH_EDITOR
 
 	void Load();
diff --git a/Engine/Source/Editor/MaterialEditor/Private/MaterialInstanceEditor.h b/Engine/Source/Editor/MaterialEditor/Private/MaterialInstanceEditor.h
--- a/Engine/Source/Editor/MaterialEditor/Private/MaterialInstanceEditor.h
+++ b/Engine/Source/Editor/MaterialEditor/Private/MaterialInstanceEditor.h
@@ -130,7 +130,7 @@ private:
 	/** Builds the toolbar widget for the material editor */
 	void ExtendToolbar();
 
-	// IMaterialEditor interface
+	//~ Begin IMaterialEditor Interface
 	virtual bool ApproveSetPreviewAsset(UObject* InAsset) override;
 
 	/**	Spawns the preview tab */
@@ -148,7 +148,7 @@ private:
 	/**	Refresh the viewport and property window */
 	void Refresh();
 
-	// Begin FEditorUndoClient Interface
+	//~ Begin FEditorUndoClient Interface
 	virtual void PostUndo( bool bSuccess ) override;
 	virtual void PostRedo( bool bSuccess ) override;
 	// End of FEditorUndoClient
diff --git a/Engine/Source/Editor/NiagaraEditor/Private/NiagaraEditor.h b/Engine/Source/Editor/NiagaraEditor/Private/NiagaraEditor.h
--- a/Engine/Source/Editor/NiagaraEditor/Private/NiagaraEditor.h
+++ b/Engine/Source/Editor/NiagaraEditor/Private/NiagaraEditor.h
@@ -21,14 +21,14 @@ public:
 	/** Destructor */
 	virtual ~FNiagaraEditor();
 
-	// Begin IToolkit interface
+	//~ Begin IToolkit Interface
 	virtual FName GetToolkitFName() const override;
 	virtual FText GetBaseToolkitName() const override;
 	virtual FString GetWorldCentricTabPrefix() const override;
 	virtual FLinearColor GetWorldCentricTabColorScale() const override;
-	// End IToolkit interface
+	//~ End IToolkit Interface
 
-	// Begin FEditorUndoClient Interface
+	//~ Begin FEditorUndoClient Interface
 	virtual void PostUndo(bool bSuccess) override;
 	virtual void PostRedo(bool bSuccess) override;
 	// End of FEditorUndoClient
diff --git a/Engine/Source/Editor/NiagaraEditor/Private/NiagaraEffectEditor.h b/Engine/Source/Editor/NiagaraEditor/Private/NiagaraEffectEditor.h
--- a/Engine/Source/Editor/NiagaraEditor/Private/NiagaraEffectEditor.h
+++ b/Engine/Source/Editor/NiagaraEditor/Private/NiagaraEffectEditor.h
@@ -26,12 +26,12 @@ public:
 	/** Destructor */
 	virtual ~FNiagaraEffectEditor();
 
-	// Begin IToolkit interface
+	//~ Begin IToolkit Interface
 	virtual FName GetToolkitFName() const override;
 	virtual FText GetBaseToolkitName() const override;
 	virtual FString GetWorldCentricTabPrefix() const override;
 	virtual FLinearColor GetWorldCentricTabColorScale() const override;
-	// End IToolkit interface
+	//~ End IToolkit Interface
 
 	// Delegates for the individual widgets
 	FReply OnEmitterSelected(TSharedPtr<FNiagaraSimulation> SelectedItem, ESelectInfo::Type SelType);
diff --git a/Engine/Source/Editor/OnlineBlueprintSupport/Classes/K2Node_LeaderboardFlush.h b/Engine/Source/Editor/OnlineBlueprintSupport/Classes/K2Node_LeaderboardFlush.h
--- a/Engine/Source/Editor/OnlineBlueprintSupport/Classes/K2Node_LeaderboardFlush.h
+++ b/Engine/Source/Editor/OnlineBlueprintSupport/Classes/K2Node_LeaderboardFlush.h
@@ -9,12 +9,12 @@ class ONLINEBLUEPRINTSUPPORT_API UK2Node_LeaderboardFlush : public UK2Node_BaseA
 {
 	GENERATED_UCLASS_BODY()
 
-	// UEdGraphNode interface
+	//~ Begin UEdGraphNode Interface
 	virtual FText GetTooltipText() const override;
 	virtual FText GetNodeTitle(ENodeTitleType::Type TitleType) const override;
-	// End of UEdGraphNode interface
+	//~ End UEdGraphNode Interface
 
-	// Begin UK2Node interface
+	//~ Begin UK2Node Interface
 	virtual FText GetMenuCategory() const override;
-	// End UK2Node interface
+	//~ End UK2Node Interface
 };
diff --git a/Engine/Source/Editor/PackagesDialog/Private/SPackagesDialog.h b/Engine/Source/Editor/PackagesDialog/Private/SPackagesDialog.h
--- a/Engine/Source/Editor/PackagesDialog/Private/SPackagesDialog.h
+++ b/Engine/Source/Editor/PackagesDialog/Private/SPackagesDialog.h
@@ -400,9 +400,9 @@ public:
 
 private:
 
-	// Begin SWidget Interface
+	//~ Begin SWidget Interface
 	virtual FReply OnKeyDown( const FGeometry& MyGeometry, const FKeyEvent& InKeyEvent ) override;
-	// End SWidget Interface
+	//~ End SWidget Interface
 
 	/**
 	 * Called when the checkbox items have changed state
diff --git a/Engine/Source/Editor/Persona/Private/Persona.h b/Engine/Source/Editor/Persona/Private/Persona.h
--- a/Engine/Source/Editor/Persona/Private/Persona.h
+++ b/Engine/Source/Editor/Persona/Private/Persona.h
@@ -178,7 +178,7 @@ public:
 	void SetSelectedBlendProfile(UBlendProfile* InBlendProfile);
 
 public:
-	// IToolkit interface
+	//~ Begin IToolkit Interface
 	virtual FName GetToolkitContextFName() const override;
 	virtual FName GetToolkitFName() const override;
 	virtual FText GetBaseToolkitName() const override;
@@ -187,7 +187,7 @@ public:
 	virtual FString GetWorldCentricTabPrefix() const override;
 	virtual FLinearColor GetWorldCentricTabColorScale() const override;
 	virtual void SaveAsset_Execute() override;
-	// End of IToolkit interface
+	//~ End IToolkit Interface
 
 	/** Saves all animation assets related to a skeleton */
 	void SaveAnimationAssets_Execute();
@@ -202,10 +202,10 @@ public:
 	/** Returns a pointer to the Blueprint object we are currently editing, as long as we are editing exactly one */
 	virtual UBlueprint* GetBlueprintObj() const override;
 
-	// FTickableEditorObject interface
+	//~ Begin FTickableEditorObject Interface
 	virtual void Tick(float DeltaTime) override;
 	virtual TStatId GetStatId() const override;
-	// End of FTickableEditorObject interface
+	//~ End FTickableEditorObject Interface
 
 	/** Returns the image brush to use for each modes dirty marker */
 	const FSlateBrush* GetDirtyImageForMode(FName Mode) const;
@@ -226,7 +226,7 @@ protected:
 	bool CanPreviewAsset() const;
 	FText GetPreviewAssetTooltip() const;
 
-	// FBlueprintEditor interface
+	//~ Begin FBlueprintEditor Interface
 	//virtual void CreateDefaultToolbar() override;
 	virtual void CreateDefaultCommands() override;
 	virtual void OnSelectBone() override;
@@ -249,18 +249,18 @@ protected:
 	virtual bool IsEditable(UEdGraph* InGraph) const override;
 	virtual FText GetGraphDecorationString(UEdGraph* InGraph) const override;
 	virtual void OnBlueprintChangedImpl(UBlueprint* InBlueprint, bool bIsJustBeingCompiled = false) override;
-	// End of FBlueprintEditor interface
+	//~ End FBlueprintEditor Interface
 
-	// IAssetEditorInstance interface
+	//~ Begin IAssetEditorInstance Interface
 	virtual void FocusWindow(UObject* ObjectToFocusOn = NULL) override;
-	// End of IAssetEditorInstance interface
+	//~ End IAssetEditorInstance Interface
 
-	// Begin FEditorUndoClient Interface
+	//~ Begin FEditorUndoClient Interface
 	virtual void PostUndo(bool bSuccess) override;
 	virtual void PostRedo(bool bSuccess) override;
 	// End of FEditorUndoClient
 
-	// Begin FAssetEditorToolkit Interface
+	//~ Begin FAssetEditorToolkit Interface
 	virtual void FindInContentBrowser_Execute() override;
 	// End of FAssetEditorToolkit
 
