Fixed issues that could lead to a crash with modifying some properties in sub-struct variables inside user defined structs.

diff --git a/Engine/Source/Editor/Kismet/Private/UserDefinedStructureEditor.cpp b/Engine/Source/Editor/Kismet/Private/UserDefinedStructureEditor.cpp
--- a/Engine/Source/Editor/Kismet/Private/UserDefinedStructureEditor.cpp
+++ b/Engine/Source/Editor/Kismet/Private/UserDefinedStructureEditor.cpp
@@ -22,13 +22,14 @@ class FDefaultValueDetails : public IDetailCustomization
 {
 public:
 	/** Makes a new instance of this detail layout class for a specific detail view requesting it */
-	static TSharedRef<IDetailCustomization> MakeInstance(TSharedPtr<FStructOnScope> InStructData)
+	static TSharedRef<IDetailCustomization> MakeInstance(TWeakPtr<class FStructureDefaultValueView> InDefaultValueView, TSharedPtr<FStructOnScope> InStructData)
 	{
-		return MakeShareable(new FDefaultValueDetails(InStructData));
+		return MakeShareable(new FDefaultValueDetails(InDefaultValueView, InStructData));
 	}
 
-	FDefaultValueDetails(TSharedPtr<FStructOnScope> InStructData)
+	FDefaultValueDetails(TWeakPtr<class FStructureDefaultValueView> InDefaultValueView, TSharedPtr<FStructOnScope> InStructData)
 		: StructData(InStructData)
+		, DefaultValueView(InDefaultValueView)
 	{}
 
 	~FDefaultValueDetails()
@@ -43,7 +44,7 @@ public:
 
 private:
 	TWeakObjectPtr<UUserDefinedStruct> UserDefinedStruct;
-	TSharedPtr<class FUserDefinedStructureLayout> Layout;
+	TWeakPtr<class FStructureDefaultValueView> DefaultValueView;
 	TSharedPtr<FStructOnScope> StructData;
 	IDetailLayoutBuilder* DetailLayoutPtr;
 };
@@ -70,45 +71,15 @@ void FDefaultValueDetails::CustomizeDetails(class IDetailLayoutBuilder& DetailLa
 	}
 }
 
-void FDefaultValueDetails::OnFinishedChangingProperties(const FPropertyChangedEvent& PropertyChangedEvent)
-{
-	check(PropertyChangedEvent.MemberProperty && PropertyChangedEvent.MemberProperty->GetOwnerStruct());
-
-	const bool bPropertyIsActual = PropertyChangedEvent.MemberProperty->GetOwnerStruct() == UserDefinedStruct.Get();
-	const UProperty* DirectProperty = bPropertyIsActual ? PropertyChangedEvent.MemberProperty : nullptr;
-	while (DirectProperty && !Cast<const UUserDefinedStruct>(DirectProperty->GetOuter()))
-	{
-		DirectProperty = Cast<const UProperty>(DirectProperty->GetOuter());
-	}
-	ensure(nullptr != DirectProperty);
-
-	if (DirectProperty)
-	{
-		FString DefaultValueString;
-		bool bDefaultValueSet = false;
-		{
-			if (StructData.IsValid() && StructData->IsValid())
-			{
-				bDefaultValueSet = FBlueprintEditorUtils::PropertyValueToString(DirectProperty, StructData->GetStructMemory(), DefaultValueString);
-			}
-		}
-
-		const FGuid VarGuid = FStructureEditorUtils::GetGuidForProperty(DirectProperty);
-		if (bDefaultValueSet && VarGuid.IsValid())
-		{
-			FStructureEditorUtils::ChangeVariableDefaultValue(UserDefinedStruct.Get(), VarGuid, DefaultValueString);
-		}
-	}
-}
-
 /////////////////////////////////
 // FStructureDefaultValueView
 
-class FStructureDefaultValueView : public FStructureEditorUtils::INotifyOnStructChanged, public TSharedFromThis<FStructureDefaultValueView>
+class FStructureDefaultValueView : public FStructureEditorUtils::INotifyOnStructChanged, public TSharedFromThis<FStructureDefaultValueView>, public FNotifyHook
 {
 public:
 	FStructureDefaultValueView(UUserDefinedStruct* EditedStruct) 
 		: UserDefinedStruct(EditedStruct)
+		, PropertyChangeRecursionGuard(0)
 	{
 	}
 
@@ -122,9 +93,11 @@ public:
 		ViewArgs.bAllowSearch = false;
 		ViewArgs.bHideSelectionTip = false;
 		ViewArgs.bShowActorLabel = false;
+		ViewArgs.NotifyHook = this;
 
 		DetailsView = PropertyModule.CreateDetailView(ViewArgs);
-		FOnGetDetailCustomizationInstance LayoutStructDetails = FOnGetDetailCustomizationInstance::CreateStatic(&FDefaultValueDetails::MakeInstance, StructData);
+		TWeakPtr< FStructureDefaultValueView > WeakThis = SharedThis(this);
+		FOnGetDetailCustomizationInstance LayoutStructDetails = FOnGetDetailCustomizationInstance::CreateStatic(&FDefaultValueDetails::MakeInstance, WeakThis, StructData);
 		DetailsView->RegisterInstancedCustomPropertyLayout(UUserDefinedStruct::StaticClass(), LayoutStructDetails);
 		DetailsView->SetObject(UserDefinedStruct.Get());
 	}
@@ -156,6 +129,21 @@ public:
 		FStructureEditorUtils::Fill_MakeStructureDefaultValue(UserDefinedStruct.Get(), StructData->GetStructMemory());
 		DetailsView->SetObject(UserDefinedStruct.Get());
 	}
+
+	// FNotifyHook interface
+	virtual void NotifyPreChange( UProperty* PropertyAboutToChange ) override
+	{
+		++PropertyChangeRecursionGuard;
+	}
+
+	virtual void NotifyPostChange( const FPropertyChangedEvent& PropertyChangedEvent, UProperty* PropertyThatChanged) override
+	{
+		--PropertyChangeRecursionGuard;
+	}
+	// End of FNotifyHook interface
+
+	/** Returns TRUE when property changes are complete, according to recursion counts */
+	bool IsPropertyChangeComplete() { return PropertyChangeRecursionGuard == 0; }
 private:
 
 	/** Struct on scope data that is being viewed in the details panel */
@@ -166,8 +154,50 @@ private:
 
 	/** User defined struct that is being represented */
 	const TWeakObjectPtr<UUserDefinedStruct> UserDefinedStruct;
+
+	/** Manages recursion in property changing, to ensure we only compile the structure when all properties are done changing */
+	int32 PropertyChangeRecursionGuard;
 };
 
+///////////////////////////////////////////////////////////////////////////////////////
+// FDefaultValueDetails
+
+void FDefaultValueDetails::OnFinishedChangingProperties(const FPropertyChangedEvent& PropertyChangedEvent)
+{
+	if (DefaultValueView.Pin()->IsPropertyChangeComplete())
+	{
+		UStruct* OwnerStruct = PropertyChangedEvent.MemberProperty->GetOwnerStruct();
+
+		check(PropertyChangedEvent.MemberProperty && OwnerStruct);
+
+		ensure(OwnerStruct == UserDefinedStruct.Get());
+		const UProperty* DirectProperty = PropertyChangedEvent.MemberProperty;
+		while (DirectProperty && !Cast<const UUserDefinedStruct>(DirectProperty->GetOuter()))
+		{
+			DirectProperty = Cast<const UProperty>(DirectProperty->GetOuter());
+		}
+		ensure(nullptr != DirectProperty);
+
+		if (DirectProperty)
+		{
+			FString DefaultValueString;
+			bool bDefaultValueSet = false;
+			{
+				if (StructData.IsValid() && StructData->IsValid())
+				{
+					bDefaultValueSet = FBlueprintEditorUtils::PropertyValueToString(DirectProperty, StructData->GetStructMemory(), DefaultValueString);
+				}
+			}
+
+			const FGuid VarGuid = FStructureEditorUtils::GetGuidForProperty(DirectProperty);
+			if (bDefaultValueSet && VarGuid.IsValid())
+			{
+				FStructureEditorUtils::ChangeVariableDefaultValue(UserDefinedStruct.Get(), VarGuid, DefaultValueString);
+			}
+		}
+	}
+}
+
 ///////////////////////////////////////////////////////////////////////////////////////
 // FUserDefinedStructureDetails
 
