Fix initialization order

diff --git a/Engine/Source/Editor/Kismet/Private/UserDefinedStructureEditor.cpp b/Engine/Source/Editor/Kismet/Private/UserDefinedStructureEditor.cpp
--- a/Engine/Source/Editor/Kismet/Private/UserDefinedStructureEditor.cpp
+++ b/Engine/Source/Editor/Kismet/Private/UserDefinedStructureEditor.cpp
@@ -28,8 +28,8 @@ public:
 	}
 
 	FDefaultValueDetails(TWeakPtr<class FStructureDefaultValueView> InDefaultValueView, TSharedPtr<FStructOnScope> InStructData)
-		: StructData(InStructData)
-		, DefaultValueView(InDefaultValueView)
+		: DefaultValueView(InDefaultValueView)
+		, StructData(InStructData)
 	{}
 
 	~FDefaultValueDetails()
