Improvements to Find-in-Blueprints for variable nodes.

Properties in nodes can be marked as Blueprint searchable. Content will be gathered by FiB system for use in unloaded Blueprints on resave.

"Find References" on variables in the MyBlueprint list or on variable nodes is now more accurate.

Removed FEdGraphSchemaAction_K2AddCallOnVariable and FEdGraphSchemaAction_K2AddDocumentation, both unused in the editor and are a part of the old Blueprint menu system.

FMemberReference::SetSelfMember and FMemberReference::SetExternalMember support accepting a FGuid for the reference.

Variable nodes will assign their variable's Guid to their VariableReference
Added UK2Node::GetFindReferenceSearchString to find the search query for a node to find references to it, currently returns the node title for all nodes but variables (where it tries to do an advanced search to limit to only actual references to the same variable).

Can mark UProperties with metadata that informs the Find-in-Blueprint system to gather the property for searching. Currently FVariableReference in K2Node_Variable is the only thing gathered and only nodes are yet supported by this feature. Will expand to all objects contained within UBlueprint in the future.

When using Find-in-Blueprints, can use category names as functions to filter into child tags and values.

All Find-in-Blueprint data is now stored in an Asset Registry tag titled "FiBData" instead of "FiB", this is to support the new version feature.

diff --git a/Engine/Source/Editor/BlueprintGraph/Classes/EdGraphSchema_K2.h b/Engine/Source/Editor/BlueprintGraph/Classes/EdGraphSchema_K2.h
--- a/Engine/Source/Editor/BlueprintGraph/Classes/EdGraphSchema_K2.h
+++ b/Engine/Source/Editor/BlueprintGraph/Classes/EdGraphSchema_K2.h
@@ -419,7 +419,6 @@ public:
 	virtual UEdGraphNode* CreateSubstituteNode(UEdGraphNode* Node, const UEdGraph* Graph, FObjectInstancingGraph* InstanceGraph) const override;
 	virtual int32 GetNodeSelectionCount(const UEdGraph* Graph) const override;
 	virtual TSharedPtr<FEdGraphSchemaAction> GetCreateCommentAction() const override;
-	virtual TSharedPtr<FEdGraphSchemaAction> GetCreateDocumentNodeAction() const override;
 	virtual bool FadeNodeWhenDraggingOffPin(const UEdGraphNode* Node, const UEdGraphPin* Pin) const override;
 	virtual void BackwardCompatibilityNodeConversion(UEdGraph* Graph, bool bOnlySafeChanges) const override;
 	virtual bool ShouldAlwaysPurgeOnModification() const override { return false; }
diff --git a/Engine/Source/Editor/BlueprintGraph/Classes/EdGraphSchema_K2_Actions.h b/Engine/Source/Editor/BlueprintGraph/Classes/EdGraphSchema_K2_Actions.h
--- a/Engine/Source/Editor/BlueprintGraph/Classes/EdGraphSchema_K2_Actions.h
+++ b/Engine/Source/Editor/BlueprintGraph/Classes/EdGraphSchema_K2_Actions.h
@@ -307,33 +307,6 @@ struct FEdGraphSchemaAction_K2AddCallOnActor : public FEdGraphSchemaAction_K2New
 	// End of FEdGraphSchemaAction interface	
 };
 
-/*******************************************************************************
-* FEdGraphSchemaAction_K2AddCallOnVariable
-*******************************************************************************/
-
-/** Action to add a 'call function on variable' pair of nodes to the graph */
-USTRUCT()
-struct FEdGraphSchemaAction_K2AddCallOnVariable : public FEdGraphSchemaAction_K2NewNode
-{
-	GENERATED_USTRUCT_BODY()
-
-	/** Property name we want to call the function on */
-	UPROPERTY()
-	FName VariableName;
-
-	FEdGraphSchemaAction_K2AddCallOnVariable()
-		: FEdGraphSchemaAction_K2NewNode()
-	{}
-
-	FEdGraphSchemaAction_K2AddCallOnVariable(const FText& InNodeCategory, const FText& InMenuDesc, const FString& InToolTip, const int32 InGrouping)
-		: FEdGraphSchemaAction_K2NewNode(InNodeCategory, InMenuDesc, InToolTip, InGrouping)
-	{}
-
-	// FEdGraphSchemaAction interface
-	virtual UEdGraphNode* PerformAction(class UEdGraph* ParentGraph, UEdGraphPin* FromPin, const FVector2D Location, bool bSelectNewNode = true) override;
-	// End of FEdGraphSchemaAction interface	
-};
-
 /*******************************************************************************
 * FEdGraphSchemaAction_K2AddComment
 *******************************************************************************/
@@ -363,35 +336,6 @@ struct FEdGraphSchemaAction_K2AddComment : public FEdGraphSchemaAction
 	// End of FEdGraphSchemaAction interface
 };
 
-/*******************************************************************************
-* FEdGraphSchemaAction_K2AddDocumentation
-*******************************************************************************/
-
-/** Action to add a 'documentation' node to the graph */
-USTRUCT()
-struct FEdGraphSchemaAction_K2AddDocumentation : public FEdGraphSchemaAction
-{
-	GENERATED_USTRUCT_BODY()
-
-	// Simple type info
-	static FName StaticGetTypeId() {static FName Type("FEdGraphSchemaAction_K2AddDocumentation"); return Type;}
-	virtual FName GetTypeId() const override { return StaticGetTypeId(); } 
-
-	FEdGraphSchemaAction_K2AddDocumentation()
-		: FEdGraphSchemaAction(FText(), NSLOCTEXT("K2AddDocumentation", "AddDocumentation", "Add Documentation ..."), NSLOCTEXT("K2AddDocumentation", "AddDocumentation_Tooltip", "Creates a Documentation Node.").ToString(), 0)
-	{
-	}
-
-	FEdGraphSchemaAction_K2AddDocumentation(const FString& InDescription, const FString& InToolTip)
-		: FEdGraphSchemaAction(FText(), FText::FromString(InDescription), InToolTip, 0)
-	{
-	}
-	
-	// FEdGraphSchemaAction interface
-	BLUEPRINTGRAPH_API virtual UEdGraphNode* PerformAction(class UEdGraph* ParentGraph, UEdGraphPin* FromPin, const FVector2D Location, bool bSelectNewNode = true) override;
-	// End of FEdGraphSchemaAction interface
-};
-
 /*******************************************************************************
 * FEdGraphSchemaAction_K2TargetNode
 *******************************************************************************/
diff --git a/Engine/Source/Editor/BlueprintGraph/Classes/K2Node_Variable.h b/Engine/Source/Editor/BlueprintGraph/Classes/K2Node_Variable.h
--- a/Engine/Source/Editor/BlueprintGraph/Classes/K2Node_Variable.h
+++ b/Engine/Source/Editor/BlueprintGraph/Classes/K2Node_Variable.h
@@ -25,7 +25,7 @@ class BLUEPRINTGRAPH_API UK2Node_Variable : public UK2Node
 	GENERATED_UCLASS_BODY()
 
 	/** Reference to variable we want to set/get */
-	UPROPERTY()
+	UPROPERTY(meta=(BlueprintSearchable="true", BlueprintSearchableHiddenExplicit="true"))
 	FMemberReference	VariableReference;
 
 	UPROPERTY()
@@ -57,6 +57,7 @@ public:
 
 	//~ Begin UEdGraphNode Interface
 	virtual FLinearColor GetNodeTitleColor() const override;
+	virtual FString GetFindReferenceSearchString() const override;
 	virtual void ReconstructNode() override;
 	virtual FString GetDocumentationLink() const override;
 	virtual FString GetDocumentationExcerptName() const override;
diff --git a/Engine/Source/Editor/BlueprintGraph/Private/EdGraphSchema_K2.cpp b/Engine/Source/Editor/BlueprintGraph/Private/EdGraphSchema_K2.cpp
--- a/Engine/Source/Editor/BlueprintGraph/Private/EdGraphSchema_K2.cpp
+++ b/Engine/Source/Editor/BlueprintGraph/Private/EdGraphSchema_K2.cpp
@@ -1616,7 +1616,8 @@ void UEdGraphSchema_K2::OnCreateNonExistentVariable( UK2Node_Variable* Variable,
 
 		if (FBlueprintEditorUtils::AddMemberVariable(OwnerBlueprint,Variable->GetVarName(), Pin->PinType))
 		{
-			Variable->VariableReference.SetSelfMember( Variable->GetVarName() );
+			FGuid Guid = FBlueprintEditorUtils::FindMemberVariableGuidByName(OwnerBlueprint, Variable->GetVarName());
+			Variable->VariableReference.SetSelfMember( Variable->GetVarName(), Guid );
 		}
 	}	
 }
@@ -1659,7 +1660,9 @@ void UEdGraphSchema_K2::OnReplaceVariableForVariableNode( UK2Node_Variable* Vari
 
 		if (bIsSelfMember)
 		{
-			Variable->VariableReference.SetSelfMember( FName(*VariableName) );
+			FName VarName = FName(*VariableName);
+			FGuid Guid = FBlueprintEditorUtils::FindMemberVariableGuidByName(OwnerBlueprint, VarName);
+			Variable->VariableReference.SetSelfMember( VarName, Guid );
 		}
 		else
 		{
@@ -5456,11 +5459,6 @@ TSharedPtr<FEdGraphSchemaAction> UEdGraphSchema_K2::GetCreateCommentAction() con
 	return TSharedPtr<FEdGraphSchemaAction>(static_cast<FEdGraphSchemaAction*>(new FEdGraphSchemaAction_K2AddComment));
 }
 
-TSharedPtr<FEdGraphSchemaAction> UEdGraphSchema_K2::GetCreateDocumentNodeAction() const
-{
-	return TSharedPtr<FEdGraphSchemaAction>(static_cast<FEdGraphSchemaAction*>(new FEdGraphSchemaAction_K2AddDocumentation));
-}
-
 bool UEdGraphSchema_K2::CanDuplicateGraph(UEdGraph* InSourceGraph) const
 {
 	EGraphType GraphType = GetGraphType(InSourceGraph);
@@ -6008,11 +6006,18 @@ void UEdGraphSchema_K2::ConfigureVarNode(UK2Node_Variable* InVarNode, FName InVa
 	// See if this is a 'self context' (ie. blueprint class is owner (or child of owner) of dropped var class)
 	if ((InVariableSource == NULL) || InTargetBlueprint->SkeletonGeneratedClass->IsChildOf(InVariableSource))
 	{
-		InVarNode->VariableReference.SetSelfMember(InVariableName);
+		FGuid Guid = FBlueprintEditorUtils::FindMemberVariableGuidByName(InTargetBlueprint, InVariableName);
+		InVarNode->VariableReference.SetSelfMember(InVariableName, Guid);
 	}
 	else if (InVariableSource->IsA(UClass::StaticClass()))
 	{
-		InVarNode->VariableReference.SetExternalMember(InVariableName, CastChecked<UClass>(InVariableSource));
+		FGuid Guid;
+		if (UBlueprint* VariableOwnerBP = Cast<UBlueprint>(Cast<UClass>(InVariableSource)->ClassGeneratedBy))
+		{
+			Guid = FBlueprintEditorUtils::FindMemberVariableGuidByName(VariableOwnerBP, InVariableName);
+		}
+
+		InVarNode->VariableReference.SetExternalMember(InVariableName, CastChecked<UClass>(InVariableSource), Guid);
 	}
 	else
 	{
diff --git a/Engine/Source/Editor/BlueprintGraph/Private/EdGraphSchema_K2_Actions.cpp b/Engine/Source/Editor/BlueprintGraph/Private/EdGraphSchema_K2_Actions.cpp
--- a/Engine/Source/Editor/BlueprintGraph/Private/EdGraphSchema_K2_Actions.cpp
+++ b/Engine/Source/Editor/BlueprintGraph/Private/EdGraphSchema_K2_Actions.cpp
@@ -499,57 +499,6 @@ void FEdGraphSchemaAction_K2AddCallOnActor::AddReferencedObjects( FReferenceColl
 	}
 }
 
-/////////////////////////////////////////////////////
-// FEdGraphSchemaAction_K2AddCallOnVariable
-
-UEdGraphNode* FEdGraphSchemaAction_K2AddCallOnVariable::PerformAction(class UEdGraph* ParentGraph, UEdGraphPin* FromPin, const FVector2D Location, bool bSelectNewNode/* = true*/)
-{
-	const UEdGraphSchema_K2* K2Schema = GetDefault<UEdGraphSchema_K2>();
-
-	// Snap the node placement location to the grid, ensures calculations later match up better
-	FVector2D LocalLocation;
-	LocalLocation.X = FMath::GridSnap( Location.X, SNAP_GRID );
-	LocalLocation.Y = FMath::GridSnap( Location.Y, SNAP_GRID );
-
-	// First use the base functionality to spawn the 'call function' node
-	FVector2D TempLocation = LocalLocation;
-	UEdGraphNode* CallNode = FEdGraphSchemaAction_K2NewNode::PerformAction(ParentGraph, FromPin, TempLocation, bSelectNewNode);
-
-	// this is the guesstimate of the function node's height, snapped to grid units
-	const float FunctionNodeHeight = FMath::GridSnap( TempLocation.Y - LocalLocation.Y, SNAP_GRID );
-	// this is roughly the middle of the function node height
-	const float FunctionNodeMidY = LocalLocation.Y + FunctionNodeHeight * 0.5f;
-	// this is the offset up from the mid point at which we start placing nodes 
-	const float StartYOffset = -NodeLiteralHeight * 0.5f;
-	// The Y location we start placing nodes from
-	const float ReferencedNodesPlacementYLocation = FunctionNodeMidY + StartYOffset;
-
-	// Now we need to create the variable literal to wire up
-	if(VariableName != NAME_None)
-	{
-		UK2Node_VariableGet* GetVarNode =  NewObject<UK2Node_VariableGet>(ParentGraph);
-		ParentGraph->AddNode(GetVarNode, false, bSelectNewNode);
-		GetVarNode->SetFlags(RF_Transactional);
-
-		GetVarNode->VariableReference.SetSelfMember(VariableName);
-		GetVarNode->AllocateDefaultPins();
-		GetVarNode->NodePosX = LocalLocation.X - FunctionNodeLiteralReferencesXOffset;
-		GetVarNode->NodePosY = ReferencedNodesPlacementYLocation;
-
-		GetVarNode->SnapToGrid(SNAP_GRID);
-
-		// Connect the literal out to the self of the call
-		UEdGraphPin* LiteralOutput = GetVarNode->GetValuePin();
-		UEdGraphPin* CallSelfInput = CallNode->FindPin(K2Schema->PN_Self);
-		if(LiteralOutput != NULL && CallSelfInput != NULL)
-		{
-			LiteralOutput->MakeLinkTo(CallSelfInput);
-		}
-	}
-
-	return CallNode;
-}
-
 /////////////////////////////////////////////////////
 // FEdGraphSchemaAction_K2AddComment
 
@@ -582,35 +531,6 @@ UEdGraphNode* FEdGraphSchemaAction_K2AddComment::PerformAction(class UEdGraph* P
 	return NewNode;
 }
 
-/////////////////////////////////////////////////////
-// FEdGraphSchemaAction_K2AddDocumentation
-
-UEdGraphNode* FEdGraphSchemaAction_K2AddDocumentation::PerformAction(class UEdGraph* ParentGraph, UEdGraphPin* FromPin, const FVector2D Location, bool bSelectNewNode/* = true*/)
-{
-	// Add menu item for creating document nodes
-	UEdGraphNode_Documentation* DocumentTemplate = NewObject<UEdGraphNode_Documentation>();
-
-	UBlueprint* Blueprint = FBlueprintEditorUtils::FindBlueprintForGraph(ParentGraph);
-
-	FVector2D SpawnLocation = Location;
-
-	FSlateRect Bounds;
-	if ((Blueprint != NULL) && FKismetEditorUtilities::GetBoundsForSelectedNodes(Blueprint, Bounds, 50.0f))
-	{
-		DocumentTemplate->SetBounds(Bounds);
-		SpawnLocation.X = DocumentTemplate->NodePosX;
-		SpawnLocation.Y = DocumentTemplate->NodePosY;
-	}
-
-	UEdGraphNode* NewNode = FEdGraphSchemaAction_NewNode::SpawnNodeFromTemplate<UEdGraphNode_Documentation>(ParentGraph, DocumentTemplate, SpawnLocation, bSelectNewNode);
-
-	// Mark Blueprint as structurally modified since
-	// UK2Node_Documentation::NodeCausesStructuralBlueprintChange used to return true
-	FBlueprintEditorUtils::MarkBlueprintAsStructurallyModified(Blueprint);
-
-	return NewNode;
-}
-
 /////////////////////////////////////////////////////
 // FEdGraphSchemaAction_K2TargetNode
 
diff --git a/Engine/Source/Editor/BlueprintGraph/Private/K2Node_Variable.cpp b/Engine/Source/Editor/BlueprintGraph/Private/K2Node_Variable.cpp
--- a/Engine/Source/Editor/BlueprintGraph/Private/K2Node_Variable.cpp
+++ b/Engine/Source/Editor/BlueprintGraph/Private/K2Node_Variable.cpp
@@ -27,14 +27,19 @@ void UK2Node_Variable::Serialize(FArchive& Ar)
 			VariableReference.SetDirect(VariableName_DEPRECATED, FGuid(), VariableSourceClass_DEPRECATED, bSelfContext_DEPRECATED);
 		}
 
-		if(Ar.UE4Ver() < VER_UE4_K2NODE_REFERENCEGUIDS)
+		if(Ar.UE4Ver() < VER_UE4_K2NODE_VAR_REFERENCEGUIDS)
 		{
 			FGuid VarGuid;
 			
-			if (UBlueprint::GetGuidFromClassByFieldName<UProperty>(GetBlueprint()->GeneratedClass, VariableReference.GetMemberName(), VarGuid))
+			// Do not let this code run for local variables
+			if (!VariableReference.IsLocalScope())
 			{
 				const bool bSelf = VariableReference.IsSelfContext();
-				VariableReference.SetDirect(VariableReference.GetMemberName(), VarGuid, (bSelf ? NULL : VariableReference.GetMemberParentClass((UClass*)NULL)), bSelf);
+				UClass* MemberParentClass = VariableReference.GetMemberParentClass(nullptr);
+				if (UBlueprint::GetGuidFromClassByFieldName<UProperty>(bSelf? GetBlueprint()->GeneratedClass : MemberParentClass, VariableReference.GetMemberName(), VarGuid))
+				{
+					VariableReference.SetDirect(VariableReference.GetMemberName(), VarGuid, bSelf ? nullptr : MemberParentClass, bSelf);
+				}
 			}
 		}
 	}
@@ -194,6 +199,28 @@ FLinearColor UK2Node_Variable::GetNodeTitleColor() const
 	return FLinearColor::White;
 }
 
+FString UK2Node_Variable::GetFindReferenceSearchString() const
+{
+	FString ResultSearchString;
+	if (VariableReference.IsLocalScope())
+	{
+		return FString::Printf(TEXT("Nodes(VariableReference(MemberName=+%s && MemberScope=+%s))"), *VariableReference.GetMemberName().ToString(), *VariableReference.GetMemberScopeName());
+	}
+	else
+	{
+		FGuid Guid = VariableReference.GetMemberGuid();
+		if (Guid.IsValid())
+		{
+			ResultSearchString = FString::Printf(TEXT("Nodes(VariableReference(MemberName=+%s && MemberGuid(A=%i && B=%i && C=%i && D=%i)))"), *VariableReference.GetMemberName().ToString(), Guid.A, Guid.B, Guid.C, Guid.D);
+		}
+		else
+		{
+			ResultSearchString = FString::Printf(TEXT("Nodes(VariableReference(MemberName=+%s))"), *VariableReference.GetMemberName().ToString());
+		}
+	}
+	return ResultSearchString;
+}
+
 UK2Node::ERedirectType UK2Node_Variable::DoPinsMatchForReconstruction( const UEdGraphPin* NewPin, int32 NewPinIndex, const UEdGraphPin* OldPin, int32 OldPinIndex ) const 
 {
 	const UEdGraphSchema_K2* K2Schema = GetDefault<UEdGraphSchema_K2>();
diff --git a/Engine/Source/Editor/Kismet/Private/BlueprintActionMenuUtils.cpp b/Engine/Source/Editor/Kismet/Private/BlueprintActionMenuUtils.cpp
--- a/Engine/Source/Editor/Kismet/Private/BlueprintActionMenuUtils.cpp
+++ b/Engine/Source/Editor/Kismet/Private/BlueprintActionMenuUtils.cpp
@@ -747,7 +747,6 @@ const UK2Node* FBlueprintActionMenuUtils::ExtractNodeTemplateFromAction(TSharedP
 			ActionId == FEdGraphSchemaAction_K2AddTimeline::StaticGetTypeId() ||
 			ActionId == FEdGraphSchemaAction_K2AddCustomEvent::StaticGetTypeId() ||
 			ActionId == FEdGraphSchemaAction_K2AddCallOnActor::StaticGetTypeId() ||
-			ActionId == FEdGraphSchemaAction_K2AddCallOnVariable::StaticGetTypeId() ||
 			ActionId == FEdGraphSchemaAction_K2TargetNode::StaticGetTypeId() ||
 			ActionId == FEdGraphSchemaAction_K2PasteHere::StaticGetTypeId() ||
 			ActionId == FEdGraphSchemaAction_K2Event::StaticGetTypeId() || 
diff --git a/Engine/Source/Editor/Kismet/Private/BlueprintEditor.cpp b/Engine/Source/Editor/Kismet/Private/BlueprintEditor.cpp
--- a/Engine/Source/Editor/Kismet/Private/BlueprintEditor.cpp
+++ b/Engine/Source/Editor/Kismet/Private/BlueprintEditor.cpp
@@ -7796,22 +7796,21 @@ void FBlueprintEditor::OnFindReferences()
 				if (UK2Node_CustomEvent* CustomEvent = Cast<UK2Node_CustomEvent>(SelectedNode))
 				{
 					SearchTerm =  CustomEvent->CustomFunctionName.ToString();
-				}
-				else if (UK2Node_Variable* VariableNode = Cast<UK2Node_Variable>(SelectedNode))
-				{
-					SearchTerm = FName::NameToDisplayString(VariableNode->GetVarName().ToString(), VariableNode->GetPropertyForVariable()->IsA(UBoolProperty::StaticClass()) );
+					if (!SearchTerm.IsEmpty())
+					{
+						SummonSearchUI(true, FString::Printf(TEXT("\"%s\""), *SearchTerm));
+					}
 				}
 				else
 				{
-					SearchTerm = SelectedNode->GetNodeTitle(ENodeTitleType::ListView).ToString();
+					SearchTerm = SelectedNode->GetFindReferenceSearchString();
+					if (!SearchTerm.IsEmpty())
+					{
+						SummonSearchUI(true, SearchTerm);
+					}
 				}
 			}
 		}
-
-		if (!SearchTerm.IsEmpty())
-		{
-			SummonSearchUI(true, FString::Printf(TEXT("\"%s\""), *SearchTerm));
-		}
 	}
 }
 
diff --git a/Engine/Source/Editor/Kismet/Private/FiBSearchInstance.cpp b/Engine/Source/Editor/Kismet/Private/FiBSearchInstance.cpp
--- a/Engine/Source/Editor/Kismet/Private/FiBSearchInstance.cpp
+++ b/Engine/Source/Editor/Kismet/Private/FiBSearchInstance.cpp
@@ -167,6 +167,7 @@ bool FFiBSearchInstance::DoSearchQuery(const FString& InSearchString, bool bInCo
 	ExpressionEvaluator->AddFunctionTokenCallback( TEXT("Components"), FTokenFunctionHandler::CreateRaw(this, &FFiBSearchInstance::OnFilterFunction, ESearchQueryFilter::ComponentsFilter));
 	ExpressionEvaluator->AddFunctionTokenCallback( TEXT("Nodes"), FTokenFunctionHandler::CreateRaw(this, &FFiBSearchInstance::OnFilterFunction, ESearchQueryFilter::NodesFilter));
 	ExpressionEvaluator->AddFunctionTokenCallback( TEXT("Pins"), FTokenFunctionHandler::CreateRaw(this, &FFiBSearchInstance::OnFilterFunction, ESearchQueryFilter::PinsFilter));
+	ExpressionEvaluator->SetDefaultFunctionHandler(FTokenDefaultFunctionHandler::CreateRaw(this, &FFiBSearchInstance::OnFilterDefaultFunction));
 	ExpressionEvaluator->SetFilterText(FText::FromString(InSearchString));
 
 	for (int SearchableIdx = 0; SearchableIdx < PendingSearchables.Num(); ++SearchableIdx)
@@ -213,6 +214,25 @@ void FFiBSearchInstance::BuildFunctionTargets(TSharedPtr<FImaginaryFiBData> InRo
 	}
 }
 
+void FFiBSearchInstance::BuildFunctionTargetsByName(TSharedPtr<FImaginaryFiBData> InRootData, FString InTagName, TArray< TWeakPtr< FImaginaryFiBData > >& OutTargetPendingSearchables)
+{
+	for (TSharedPtr<FImaginaryFiBData> ChildData : InRootData->GetAllParsedChildData())
+	{
+		if (ChildData->IsCategory())
+		{
+			FCategorySectionHelper* CategoryData = static_cast<FCategorySectionHelper*>(ChildData.Get());
+			if (CategoryData->GetCategoryFunctionName().Equals(InTagName, ESearchCase::IgnoreCase))
+			{
+				OutTargetPendingSearchables.Add(ChildData);
+			}
+			else
+			{
+				BuildFunctionTargetsByName(ChildData, InTagName, OutTargetPendingSearchables);
+			}
+		}
+	}
+}
+
 bool FFiBSearchInstance::OnFilterFunction(const FTextFilterString& A, ESearchQueryFilter InSearchQueryFilter)
 {
 	if (CurrentSearchable.IsValid())
@@ -255,6 +275,40 @@ bool FFiBSearchInstance::OnFilterFunction(const FTextFilterString& A, ESearchQue
 	return false;
 }
 
+bool FFiBSearchInstance::OnFilterDefaultFunction(const FTextFilterString& InFunctionName, const FTextFilterString& InFunctionParams)
+{
+	if (CurrentSearchable.IsValid())
+	{
+		TSharedPtr< FFiBSearchInstance > SubSearchInstance(new FFiBSearchInstance);
+		bool bSearchSuccess = false;
+
+		TSharedPtr< FImaginaryFiBData > CurrentSearchablePinned = CurrentSearchable.Pin();
+		CurrentSearchablePinned->ParseAllChildData();
+		BuildFunctionTargetsByName(CurrentSearchablePinned, InFunctionName.AsString(), SubSearchInstance->PendingSearchables);
+
+		// Proceed to doing a sub-search
+		if (SubSearchInstance->PendingSearchables.Num() > 0)
+		{
+			bSearchSuccess = SubSearchInstance->DoSearchQuery(InFunctionParams.AsString(), true);
+			if (bSearchSuccess)
+			{
+				for (auto& MatchesItem : SubSearchInstance->MatchesSearchQuery)
+				{
+					LastFunctionResultMatchesSearchQuery.AddUnique(MatchesItem);
+				}
+
+				for (auto& MatchesItem : SubSearchInstance->MatchingSearchComponents)
+				{
+					LastFunctionMatchingSearchComponents.AddUnique(MatchesItem.Key, MatchesItem.Value);
+				}
+			}
+		}
+		return bSearchSuccess;
+	}
+
+	return false;
+}
+
 ////////////////////////////////////////
 // FFindInBlueprintExpressionEvaluator
 
@@ -702,6 +756,10 @@ void FFindInBlueprintExpressionEvaluator::ConstructExpressionParser()
 		{
 			bResult = FunctionCallback->Execute(B.GetString());
 		}
+		else
+		{
+			bResult = DefaultFunctionHandler.Execute(A.GetString(), B.GetString());
+		}
 		FFiBToken Result(bResult, SearchInstance->LastFunctionResultMatchesSearchQuery);
 		Result.MatchingSearchComponents = SearchInstance->LastFunctionMatchingSearchComponents;
 		return Result;
diff --git a/Engine/Source/Editor/Kismet/Private/FiBSearchInstance.h b/Engine/Source/Editor/Kismet/Private/FiBSearchInstance.h
--- a/Engine/Source/Editor/Kismet/Private/FiBSearchInstance.h
+++ b/Engine/Source/Editor/Kismet/Private/FiBSearchInstance.h
@@ -6,6 +6,8 @@
 
 class FFiBSearchInstance;
 
+DECLARE_DELEGATE_RetVal_TwoParams(bool, FTokenDefaultFunctionHandler, const FTextFilterString&, const FTextFilterString&);
+
 /** Evaluates the expression the user submitted to be searched for */
 class FFindInBlueprintExpressionEvaluator : public FTextFilterExpressionEvaluator
 {
@@ -18,6 +20,12 @@ public:
 		ConstructExpressionParser();
 	}
 
+	/** Sets the default function handler, which supports generic functions which are categories or other object based on FImaginaryFiBData */
+	void SetDefaultFunctionHandler(FTokenDefaultFunctionHandler InFunctionHandler)
+	{
+		DefaultFunctionHandler = InFunctionHandler;
+	}
+
 protected:
 	/** FTextFilterExpressionEvaluator Interface */
 	virtual void ConstructExpressionParser() override;
@@ -36,6 +44,9 @@ protected:
 protected:
 	/** Referenced SearchInstance that is powering this search */
 	FFiBSearchInstance* SearchInstance;
+
+	/** Fallback for all functions, Find-in-Blueprints filters into any sub-data using functions */
+	FTokenDefaultFunctionHandler DefaultFunctionHandler;
 };
 
 /** Used to manage searches through imaginary Blueprints */
@@ -69,6 +80,15 @@ public:
 	 */
 	void BuildFunctionTargets(TSharedPtr<FImaginaryFiBData> InRootData, ESearchQueryFilter InSearchQueryFilter, TArray< TWeakPtr< FImaginaryFiBData > >& OutTargetPendingSearchables);
 
+	/**
+	 * Builds a list of imaginary items, using their names, that can be targeted by a function
+	 *
+	 * @param InRootData						Item to find the child imaginary data that can have be sub-searched for the function
+	 * @param InTagName							The name of objects to find
+	 * @param OutTargetPendingSearchables		List of pending imaginary items that will be sub-searched, this gets filled out by the function
+	 */
+	void BuildFunctionTargetsByName(TSharedPtr<FImaginaryFiBData> InRootData, FString InTagName, TArray< TWeakPtr< FImaginaryFiBData > >& OutTargetPendingSearchables);
+
 	/**
 	 * Callback when a function is called by the evaluator
 	 *
@@ -78,6 +98,15 @@ public:
 	 */
 	bool OnFilterFunction(const FTextFilterString& A, ESearchQueryFilter InSearchQueryFilter);
 
+	/**
+	 * Callback when a default/generic function is called by the evaluator
+	 *
+	 * @param InFunctionName		The function name to call, will query for sub-objects with this name to do a sub-search on
+	 * @param InFunctionParams		The parameter string to sub-search using
+	 * @return						TRUE if the function was successful at finding valid results
+	 */
+	bool OnFilterDefaultFunction(const FTextFilterString& InFunctionName, const FTextFilterString& InFunctionParams);
+public:
 	/** Current item being searched in the Imaginary Blueprint */
 	TWeakPtr< FImaginaryFiBData > CurrentSearchable;
 
diff --git a/Engine/Source/Editor/Kismet/Private/FindInBlueprintManager.cpp b/Engine/Source/Editor/Kismet/Private/FindInBlueprintManager.cpp
--- a/Engine/Source/Editor/Kismet/Private/FindInBlueprintManager.cpp
+++ b/Engine/Source/Editor/Kismet/Private/FindInBlueprintManager.cpp
@@ -52,6 +52,108 @@ const FText FFindInBlueprintSearchTags::FiB_IsReference = LOCTEXT("IsReference",
 const FText FFindInBlueprintSearchTags::FiB_Glyph = LOCTEXT("Glyph", "Glyph");
 const FText FFindInBlueprintSearchTags::FiB_GlyphColor = LOCTEXT("GlyphColor", "GlyphColor");
 
+const FText FFindInBlueprintSearchTags::FiBMetaDataTag = LOCTEXT("FiBMetaDataTag", "!!FiBMD");
+
+const FString FFiBMD::FiBSearchableMD = TEXT("BlueprintSearchable");
+const FString FFiBMD::FiBSearchableShallowMD = TEXT("BlueprintSearchableShallow");
+const FString FFiBMD::FiBSearchableExplicitMD = TEXT("BlueprintSearchableExplicit");
+const FString FFiBMD::FiBSearchableHiddenExplicitMD = TEXT("BlueprintSearchableHiddenExplicit");
+
+////////////////////////////////////
+// FStreamSearch
+FStreamSearch::FStreamSearch(const FString& InSearchValue )
+	: SearchValue(InSearchValue)
+	, bThreadCompleted(false)
+	, StopTaskCounter(0)
+{
+	// Add on a Guid to the thread name to ensure the thread is uniquely named.
+	Thread = FRunnableThread::Create( this, *FString::Printf(TEXT("FStreamSearch%s"), *FGuid::NewGuid().ToString()), 0, TPri_BelowNormal );
+}
+
+bool FStreamSearch::Init()
+{
+	return true;
+}
+
+uint32 FStreamSearch::Run()
+{
+	FFindInBlueprintSearchManager::Get().BeginSearchQuery(this);
+
+	TFunction<void(const FSearchResult&)> OnResultReady = [this](const FSearchResult& Result) {
+		FScopeLock ScopeLock(&SearchCriticalSection);
+		ItemsFound.Add(Result);
+	};
+
+	// Searching comes to an end if it is requested using the StopTaskCounter or continuing the search query yields no results
+	FSearchData QueryResult;
+	while (FFindInBlueprintSearchManager::Get().ContinueSearchQuery(this, QueryResult))
+	{
+		// This work cannot be spun off to a use the TaskGraph system for spreading the work over multiple threads/cores as the stack limit is too small for the work being done.
+		if (QueryResult.Value.Len() > 0)
+		{
+			TSharedPtr< FImaginaryBlueprint> ImaginaryBlueprint(new FImaginaryBlueprint(FPaths::GetBaseFilename(QueryResult.BlueprintPath.ToString()), QueryResult.BlueprintPath.ToString(), QueryResult.ParentClass, QueryResult.Interfaces, QueryResult.Value));
+			TSharedPtr< FFiBSearchInstance > SearchInstance(new FFiBSearchInstance);
+			FSearchResult SearchResult = SearchInstance->StartSearchQuery(*SearchValue, ImaginaryBlueprint);
+
+			// If there are children, add the item to the search results
+			if(SearchResult.IsValid() && SearchResult->Children.Num() != 0)
+			{
+				OnResultReady(SearchResult);
+			}
+		}
+
+		if (StopTaskCounter.GetValue())
+		{
+			// Ensure that the FiB Manager knows that we are done searching
+			FFindInBlueprintSearchManager::Get().EnsureSearchQueryEnds(this);
+		}
+	}
+
+	bThreadCompleted = true;
+
+	return 0;
+}
+
+void FStreamSearch::Stop()
+{
+	StopTaskCounter.Increment();
+}
+
+void FStreamSearch::Exit()
+{
+
+}
+
+void FStreamSearch::EnsureCompletion()
+{
+	{
+		FScopeLock CritSectionLock(&SearchCriticalSection);
+		ItemsFound.Empty();
+	}
+
+	Stop();
+	Thread->WaitForCompletion();
+	delete Thread;
+	Thread = NULL;
+}
+
+bool FStreamSearch::IsComplete() const
+{
+	return bThreadCompleted;
+}
+
+void FStreamSearch::GetFilteredItems(TArray<FSearchResult>& OutItemsFound)
+{
+	FScopeLock ScopeLock(&SearchCriticalSection);
+	OutItemsFound.Append(ItemsFound);
+	ItemsFound.Empty();
+}
+
+float FStreamSearch::GetPercentComplete() const
+{
+	return FFindInBlueprintSearchManager::Get().GetPercentComplete(this);
+}
+
 /** Temporarily forces all nodes and pins to use non-friendly names, forces all schema to have nodes clear their cached values so they will re-cache, and then reverts at the end */
 struct FTemporarilyUseFriendlyNodeTitles
 {
@@ -540,6 +642,136 @@ namespace BlueprintSearchMetaDataHelpers
 		InWriter->WriteObjectEnd();
 	}
 
+	/** Helper enum to gather searchable UProperties */
+	enum EGatherSearchableType
+	{
+		SEARCHABLE_AS_DESIRED = 0,
+		SEARCHABLE_FULL,
+		SEARCHABLE_SHALLOW,
+	};
+
+	/**
+	 * Gathers all searchable properties in a UObject and writes them out to Json
+	 *
+	 * @param InWriter				Json writer
+	 * @param InValue				Value of the Object to serialize
+	 * @param InStruct				Struct or class that represent the UObject's layout
+	 * @param InSearchableType		Informs the system how it should examine the properties to determine if they are searchable. All sub-properties of searchable properties are automatically gathered unless marked as not being searchable
+	 */
+	void GatherSearchableProperties(TSharedRef< SearchMetaDataWriter>& InWriter, const void* InValue, UStruct* InStruct, EGatherSearchableType InSearchableType = SEARCHABLE_AS_DESIRED);
+
+	/**
+	 * Examines a searchable property and digs in deeper if it is a UObject, UStruct, or an array, or serializes it straight out to Json
+	 *
+	 * @param InWriter				Json writer
+	 * @param InProperty			Property to examine
+	 * @param InValue				Value to find the property in the UStruct
+	 * @param InStruct				Struct or class that represent the UObject's layout
+	 */
+	void GatherSearchablesFromProperty(TSharedRef< SearchMetaDataWriter>& InWriter, UProperty* InProperty, const void* InValue, UStruct* InStruct)
+	{
+		if (UArrayProperty* ArrayProperty = Cast<UArrayProperty>(InProperty))
+		{
+			FScriptArrayHelper Helper(ArrayProperty, InValue);
+			InWriter->WriteArrayStart(FText::FromString(InProperty->GetName()));
+			for (int32 i=0, n=Helper.Num(); i<n; ++i)
+			{
+				GatherSearchablesFromProperty(InWriter, ArrayProperty->Inner, Helper.GetRawPtr(i), InStruct);
+			}
+			InWriter->WriteArrayEnd();
+		}
+		else if (UStructProperty* StructProperty = Cast<UStructProperty>(InProperty))
+		{
+			if (!InProperty->HasMetaData(*FFiBMD::FiBSearchableMD) || InProperty->GetBoolMetaData(*FFiBMD::FiBSearchableMD))
+			{
+				InWriter->WriteObjectStart(FText::FromString(InProperty->GetName()));
+				GatherSearchableProperties(InWriter, InValue, StructProperty->Struct, SEARCHABLE_FULL);
+				InWriter->WriteObjectEnd();
+			}
+		}
+		else if (UObjectProperty* ObjectProperty = Cast<UObjectProperty>(InProperty))
+		{
+			UObject* SubObject = ObjectProperty->GetObjectPropertyValue(InValue);
+			if (SubObject)
+			{
+				// Objects default to shallow unless they are marked as searchable
+				EGatherSearchableType searchType = SEARCHABLE_SHALLOW;
+
+				// Check if there is any Searchable metadata
+				if (InProperty->HasMetaData(*FFiBMD::FiBSearchableMD))
+				{
+					// Check if that metadata informs us that the property should not be searchable
+					bool bSearchable = InProperty->GetBoolMetaData(*FFiBMD::FiBSearchableMD);
+					if (bSearchable)
+					{
+						InWriter->WriteObjectStart(FText::FromString(InProperty->GetName()));
+						GatherSearchableProperties(InWriter, SubObject, SubObject->GetClass(), SEARCHABLE_FULL);
+						InWriter->WriteObjectEnd();
+					}
+				}
+				else
+				{
+					// Shallow conversion of property to string
+					auto JsonValue = FJsonObjectConverter::UPropertyToJsonValue(InProperty, InValue, 0, 0);
+					FJsonSerializer::Serialize(JsonValue, InProperty->GetName(), StaticCastSharedRef<SearchMetaDataWriterParentClass>(InWriter));
+				}
+			}
+		}
+		else
+		{
+			auto JsonValue = FJsonObjectConverter::UPropertyToJsonValue(InProperty, InValue, 0, 0);
+			FJsonSerializer::Serialize(JsonValue, InProperty->GetName(), StaticCastSharedRef<SearchMetaDataWriterParentClass>(InWriter));
+		}
+	}
+
+	void GatherSearchableProperties(TSharedRef< SearchMetaDataWriter>& InWriter, const void* InValue, UStruct* InStruct, EGatherSearchableType InSearchableType)
+	{
+		if (InValue)
+		{
+			for (TFieldIterator<UProperty> PropIt(InStruct); PropIt; ++PropIt)
+			{
+				UProperty* Property = *PropIt;
+				bool bIsSearchableMD = Property->GetBoolMetaData(*FFiBMD::FiBSearchableMD);
+				bool bIsShallowSearchableMD = Property->GetBoolMetaData(*FFiBMD::FiBSearchableShallowMD);
+				// It only is truly marked as not searchable if it has the metadata set to false, if the metadata is missing then we assume the searchable type that is passed in unless SEARCHABLE_AS_DESIRED
+				bool bIsMarkedNotSearchableMD = Property->HasMetaData(*FFiBMD::FiBSearchableMD) && !bIsSearchableMD;
+
+				if ( (InSearchableType != SEARCHABLE_AS_DESIRED && !bIsMarkedNotSearchableMD) 
+					|| bIsShallowSearchableMD || bIsSearchableMD)
+				{
+					const void* Value = Property->ContainerPtrToValuePtr<uint8>(InValue);
+
+					// Need to store the metadata on the property in a sub-object
+					InWriter->WriteObjectStart(FFindInBlueprintSearchTags::FiBMetaDataTag);
+					{
+						if (Property->GetBoolMetaData(*FFiBMD::FiBSearchableHiddenExplicitMD))
+						{
+							InWriter->WriteValue(FText::FromString(FFiBMD::FiBSearchableHiddenExplicitMD), true);
+						}
+						else if (Property->GetBoolMetaData(*FFiBMD::FiBSearchableExplicitMD))
+						{
+							InWriter->WriteValue(FText::FromString(FFiBMD::FiBSearchableExplicitMD), true);
+						}
+					}
+					InWriter->WriteObjectEnd();
+
+					if (Property->ArrayDim == 1)
+					{
+						GatherSearchablesFromProperty(InWriter, Property, Value, InStruct);
+					}
+					else
+					{
+						TArray< TSharedPtr<FJsonValue> > Array;
+						for (int Index = 0; Index != Property->ArrayDim; ++Index)
+						{
+							GatherSearchablesFromProperty(InWriter, Property, (char*)Value + Index * Property->ElementSize, InStruct);
+						}
+					}
+				}
+			}
+		}
+	}
+
 	/**
 	 * Gathers all nodes from a specified graph and serializes their searchable data to Json
 	 *
@@ -591,6 +823,9 @@ namespace BlueprintSearchMetaDataHelpers
 					}
 					InWriter->WriteArrayEnd();
 
+					// Only support this for nodes for now, will gather all searchable properties
+					GatherSearchableProperties(InWriter, Node, Node->GetClass());
+
 					InWriter->WriteObjectEnd();
 				}
 				
@@ -991,40 +1226,11 @@ void FFindInBlueprintSearchManager::OnAssetAdded(const FAssetData& InAssetData)
 			}
 			else if(const FString* FiBSearchData = InAssetData.TagsAndValues.Find("FiB"))
 			{
-				FSearchData NewSearchData;
-
-				NewSearchData.BlueprintPath = InAssetData.ObjectPath;
-
-				if (const FString* ParentClass = InAssetData.TagsAndValues.Find(TEXT("ParentClass")))
-				{
-					NewSearchData.ParentClass = *ParentClass;
-				}
-
-				const FString* ImplementedInterfaces = InAssetData.TagsAndValues.Find("ImplementedInterfaces");
-				if(ImplementedInterfaces)
-				{
-					FString FullInterface;
-					FString RemainingString;
-					FString InterfaceName;
-					FString CurrentString = *ImplementedInterfaces;
-					while(CurrentString.Split(TEXT(","), &FullInterface, &RemainingString))
-					{
-						if(FullInterface.Split(TEXT("."), &CurrentString, &InterfaceName, ESearchCase::CaseSensitive, ESearchDir::FromEnd))
-						{
-							if(!CurrentString.StartsWith(TEXT("Graphs=(")))
-							{
-								NewSearchData.Interfaces.Add(InterfaceName);
-							}
-						}
-						CurrentString = RemainingString;
-					}
-				}
-
-				NewSearchData.bMarkedForDeletion = false;
-
-				// Since the asset was not loaded, pull out the searchable data stored in the asset
-				NewSearchData.Value = *FiBSearchData;
-				AddSearchDataToDatabase(MoveTemp(NewSearchData));
+				ExtractUnloadedFiBData(InAssetData, *FiBSearchData, false);
+			}
+			else if(const FString* FiBVersionedSearchData = InAssetData.TagsAndValues.Find("FiBData"))
+			{
+				ExtractUnloadedFiBData(InAssetData, *FiBVersionedSearchData, true);
 			}
 			else
 			{
@@ -1035,6 +1241,51 @@ void FFindInBlueprintSearchManager::OnAssetAdded(const FAssetData& InAssetData)
 	}
 }
 
+void FFindInBlueprintSearchManager::ExtractUnloadedFiBData(const FAssetData& InAssetData, const FString& InFiBData, bool bIsVersioned)
+{
+	FSearchData NewSearchData;
+
+	NewSearchData.BlueprintPath = InAssetData.ObjectPath;
+
+	if (const FString* ParentClass = InAssetData.TagsAndValues.Find(TEXT("ParentClass")))
+	{
+		NewSearchData.ParentClass = *ParentClass;
+	}
+
+	const FString* ImplementedInterfaces = InAssetData.TagsAndValues.Find("ImplementedInterfaces");
+	if(ImplementedInterfaces)
+	{
+		FString FullInterface;
+		FString RemainingString;
+		FString InterfaceName;
+		FString CurrentString = *ImplementedInterfaces;
+		while(CurrentString.Split(TEXT(","), &FullInterface, &RemainingString))
+		{
+			if(FullInterface.Split(TEXT("."), &CurrentString, &InterfaceName, ESearchCase::CaseSensitive, ESearchDir::FromEnd))
+			{
+				if(!CurrentString.StartsWith(TEXT("Graphs=(")))
+				{
+					NewSearchData.Interfaces.Add(InterfaceName);
+				}
+			}
+			CurrentString = RemainingString;
+		}
+	}
+
+	NewSearchData.bMarkedForDeletion = false;
+	NewSearchData.Value = *InFiBData;
+
+	// Deserialize the version if available
+	if (bIsVersioned)
+	{
+		FBufferReader ReaderStream((void*)*NewSearchData.Value, NewSearchData.Value.Len() * sizeof(TCHAR), false);
+		NewSearchData.Version = FiBSerializationHelpers::Deserialize<int32>(ReaderStream);
+	}
+
+	// Since the asset was not loaded, pull out the searchable data stored in the asset
+	AddSearchDataToDatabase(MoveTemp(NewSearchData));
+}
+
 int32 FFindInBlueprintSearchManager::AddSearchDataToDatabase(FSearchData InSearchData)
 {
 	FName BlueprintPath = InSearchData.BlueprintPath; // Copy before we move the data into the array
@@ -1184,7 +1435,8 @@ FString FFindInBlueprintSearchManager::GatherBlueprintSearchMetadata(const UBlue
 	Writer->WriteObjectEnd();
 	Writer->Close();
 
-	SearchMetaData = Writer->GetSerializedLookupTable() + SearchMetaData;
+	int32 Version = EFiBVersion::FIB_VER_ALL;
+	SearchMetaData = FiBSerializationHelpers::Serialize(Version, false) + Writer->GetSerializedLookupTable() + SearchMetaData;
 	return SearchMetaData;
 }
 
@@ -1633,8 +1885,9 @@ bool FFindInBlueprintSearchManager::IsCacheInProgress() const
 TSharedPtr< FJsonObject > FFindInBlueprintSearchManager::ConvertJsonStringToObject(FString InJsonString, TMap<int32, FText>& OutFTextLookupTable)
 {
 	/** The searchable data is more complicated than a Json string, the Json being the main searchable body that is parsed. Below is a diagram of the full data:
-	 *  | int32 "Size" | TMap "Lookup Table" | Json String |
+	 *  | int32 "Version" | int32 "Size" | TMap "Lookup Table" | Json String |
 	 *
+	 * Version: Version of the FiB data, which may impact searching
 	 * Size: The size of the TMap in bytes
 	 * Lookup Table: The Json's identifiers and string values are in Hex strings and stored in a TMap, the Json stores these values as ints and uses them as the Key into the TMap
 	 * Json String: The Json string to be deserialized in full
@@ -1646,6 +1899,8 @@ TSharedPtr< FJsonObject > FFindInBlueprintSearchManager::ConvertJsonStringToObje
 	int32 SizeOfData;
 	FBufferReader ReaderStream((void*)*InJsonString, InJsonString.Len() * sizeof(TCHAR), false);
 
+	int32 Version = FiBSerializationHelpers::Deserialize<int32>(ReaderStream);
+
  	// Read, as a byte string, the number of characters composing the Lookup Table for the Json.
 	SizeOfData = FiBSerializationHelpers::Deserialize<int32>(ReaderStream);
 
diff --git a/Engine/Source/Editor/Kismet/Private/FindInBlueprints.cpp b/Engine/Source/Editor/Kismet/Private/FindInBlueprints.cpp
--- a/Engine/Source/Editor/Kismet/Private/FindInBlueprints.cpp
+++ b/Engine/Source/Editor/Kismet/Private/FindInBlueprints.cpp
@@ -505,139 +505,6 @@ FText FFindInBlueprintsGraph::GetCategory() const
 	return LOCTEXT("GraphCategory", "Graph");
 }
 
-////////////////////////////////////
-// FStreamSearch
-
-/**
- * Async task for searching Blueprints
- */
-class FStreamSearch : public FRunnable
-{
-public:
-	/** Constructor */
-	FStreamSearch(const FString& InSearchValue )
-		: SearchValue(InSearchValue)
-		, bThreadCompleted(false)
-		, StopTaskCounter(0)
-	{
-		// Add on a Guid to the thread name to ensure the thread is uniquely named.
-		Thread = FRunnableThread::Create( this, *FString::Printf(TEXT("FStreamSearch%s"), *FGuid::NewGuid().ToString()), 0, TPri_BelowNormal );
-	}
-
-	/** Begin FRunnable Interface */
-	virtual bool Init() override
-	{
-		return true;
-	}
-
-	virtual uint32 Run() override
-	{
-		FFindInBlueprintSearchManager::Get().BeginSearchQuery(this);
-
-		TFunction<void(const FSearchResult&)> OnResultReady = [this](const FSearchResult& Result) {
-			FScopeLock ScopeLock(&SearchCriticalSection);
-			ItemsFound.Add(Result);
-		};
-
-		// Searching comes to an end if it is requested using the StopTaskCounter or continuing the search query yields no results
-		FSearchData QueryResult;
-		while (FFindInBlueprintSearchManager::Get().ContinueSearchQuery(this, QueryResult))
-		{
-			// This work cannot be spun off to a use the TaskGraph system for spreading the work over multiple threads/cores as the stack limit is too small for the work being done.
-			if (QueryResult.Value.Len() > 0)
-			{
-				TSharedPtr< FImaginaryBlueprint> ImaginaryBlueprint(new FImaginaryBlueprint(FPaths::GetBaseFilename(QueryResult.BlueprintPath.ToString()), QueryResult.BlueprintPath.ToString(), QueryResult.ParentClass, QueryResult.Interfaces, QueryResult.Value));
-				TSharedPtr< FFiBSearchInstance > SearchInstance(new FFiBSearchInstance);
-				FSearchResult SearchResult = SearchInstance->StartSearchQuery(*SearchValue, ImaginaryBlueprint);
-
-				// If there are children, add the item to the search results
-				if(SearchResult.IsValid() && SearchResult->Children.Num() != 0)
-				{
-					OnResultReady(SearchResult);
-				}
-			}
-
-			if (StopTaskCounter.GetValue())
-			{
-				// Ensure that the FiB Manager knows that we are done searching
-				FFindInBlueprintSearchManager::Get().EnsureSearchQueryEnds(this);
-			}
-		}
-
-		bThreadCompleted = true;
-
-		return 0;
-	}
-
-	virtual void Stop() override
-	{
-		StopTaskCounter.Increment();
-	}
-
-	virtual void Exit() override
-	{
-
-	}
-	/** End FRunnable Interface */
-
-	/** Brings the thread to a safe stop before continuing. */
-	void EnsureCompletion()
-	{
-		{
-			FScopeLock CritSectionLock(&SearchCriticalSection);
-			ItemsFound.Empty();
-		}
-
-		Stop();
-		Thread->WaitForCompletion();
-		delete Thread;
-		Thread = NULL;
-	}
-
-	/** Returns TRUE if the thread is done with it's work. */
-	bool IsComplete() const
-	{
-		return bThreadCompleted;
-	}
-
-	/**
-	 * Appends the items filtered through the search filter to the passed array
-	 *
-	 * @param OutItemsFound		All the items found since last queried
-	 */
-	void GetFilteredItems(TArray<FSearchResult>& OutItemsFound)
-	{
-		FScopeLock ScopeLock(&SearchCriticalSection);
-		OutItemsFound.Append(ItemsFound);
-		ItemsFound.Empty();
-	}
-
-	/** Helper function to query the percent complete this search is */
-	float GetPercentComplete() const
-	{
-		return FFindInBlueprintSearchManager::Get().GetPercentComplete(this);
-	}
-
-public:
-	/** Thread to run the cleanup FRunnable on */
-	FRunnableThread* Thread;
-
-	/** A list of items found, cleared whenever the main thread pulls them to display to screen */
-	TArray<FSearchResult> ItemsFound;
-
-	/** The search value to filter results by */
-	FString SearchValue;
-
-	/** Prevents searching while other threads are pulling search results */
-	FCriticalSection SearchCriticalSection;
-
-	// Whether the thread has finished running
-	bool bThreadCompleted;
-
-	/** > 0 if we've been asked to abort work in progress at the next opportunity */
-	FThreadSafeCounter StopTaskCounter;
-};
-
 //////////////////////////////////////////////////////////////////////////
 // SBlueprintSearch
 
diff --git a/Engine/Source/Editor/Kismet/Private/ImaginaryBlueprintData.cpp b/Engine/Source/Editor/Kismet/Private/ImaginaryBlueprintData.cpp
--- a/Engine/Source/Editor/Kismet/Private/ImaginaryBlueprintData.cpp
+++ b/Engine/Source/Editor/Kismet/Private/ImaginaryBlueprintData.cpp
@@ -76,7 +76,7 @@ FSearchResult FImaginaryFiBData::CreateSearchTree(FSearchResult InParentSearchRe
 	}
 
 	// If the children did not match the search results but this item does, then we will want to return true
-	if (!bValidSearchResults && (InValidSearchResults.Find(InCurrentPointer.Pin().Get()) != INDEX_NONE || InMatchingSearchComponents.Find(InCurrentPointer.Pin().Get())))
+	if (!bValidSearchResults && !InCurrentPointer.Pin()->IsCategory() && (InValidSearchResults.Find(InCurrentPointer.Pin().Get()) != INDEX_NONE || InMatchingSearchComponents.Find(InCurrentPointer.Pin().Get())))
 	{
 		bValidSearchResults = true;
 	}
@@ -98,29 +98,53 @@ FSearchResult FImaginaryFiBData::CreateSearchTree(FSearchResult InParentSearchRe
 	return nullptr;
 }
 
-bool FImaginaryFiBData::IsCompatibleWithFilter(ESearchQueryFilter InSearchQueryFilter)
+bool FImaginaryFiBData::IsCompatibleWithFilter(ESearchQueryFilter InSearchQueryFilter) const
 {
 	return true;
 }
 
-bool FImaginaryFiBData::CanCallFilter(ESearchQueryFilter InSearchQueryFilter)
+bool FImaginaryFiBData::CanCallFilter(ESearchQueryFilter InSearchQueryFilter) const
 {
 	// Always compatible with the AllFilter
 	return InSearchQueryFilter == ESearchQueryFilter::AllFilter;
 }
 
-void FImaginaryFiBData::ParseAllChildData()
+void FImaginaryFiBData::ParseAllChildData(ESearchableValueStatus InSearchabilityOverride/* = ESearchableValueStatus::Searchable*/)
 {
 	if (UnparsedJsonObject.IsValid())
 	{
+		if (InSearchabilityOverride == ESearchableValueStatus::Searchable)
+		{
+			const TSharedPtr< FJsonObject >* MetaDataField;
+			if (UnparsedJsonObject->TryGetObjectField(FFindInBlueprintSearchTags::FiBMetaDataTag.ToString(), MetaDataField))
+			{
+				TSharedPtr<FFiBMetaData> MetaDataFiBInfo = MakeShareable(new FFiBMetaData(*MetaDataField, LookupTablePtr));
+				MetaDataFiBInfo->ParseAllChildData();
+
+				if (MetaDataFiBInfo->IsHidden() && MetaDataFiBInfo->IsExplicit())
+				{
+					InSearchabilityOverride = ESearchableValueStatus::ExplicitySearchableHidden;
+				}
+				else if (MetaDataFiBInfo->IsExplicit())
+				{
+					InSearchabilityOverride = ESearchableValueStatus::ExplicitySearchable;
+				}
+			}
+		}
+
 		for( auto MapValues : UnparsedJsonObject->Values )
 		{
 			TSharedPtr< FJsonValue > JsonValue = MapValues.Value;
 
+			if (MapValues.Key == FFindInBlueprintSearchTags::FiBMetaDataTag.ToString())
+			{
+				// Do not let this be processed again
+				continue;
+			}
 			if (!TrySpecialHandleJsonValue(MapValues.Key, JsonValue))
 			{
 				FText Key = FText::FromString(MapValues.Key);
-				ParseJsonValue(Key, Key, JsonValue);
+				ParseJsonValue(Key, Key, JsonValue, false, InSearchabilityOverride);
 			}
 		}
 	}
@@ -128,15 +152,16 @@ void FImaginaryFiBData::ParseAllChildData()
 	UnparsedJsonObject.Reset();
 }
 
-void FImaginaryFiBData::ParseJsonValue(FText InKey, FText InDisplayKey, TSharedPtr< FJsonValue > InJsonValue)
+void FImaginaryFiBData::ParseJsonValue(FText InKey, FText InDisplayKey, TSharedPtr< FJsonValue > InJsonValue, bool bIsInArray/*=false*/, ESearchableValueStatus InSearchabilityOverride/* = ESearchableValueStatus::Searchable*/)
 {
+	ESearchableValueStatus SearchabilityStatus = (InSearchabilityOverride == ESearchableValueStatus::Searchable)? GetSearchabilityStatus(InKey.ToString()) : InSearchabilityOverride;
 	if( InJsonValue->Type == EJson::String)
 	{
-		ParsedTagsAndValues.Add(InKey.ToString(), FSearchableValueInfo(InDisplayKey, FCString::Atoi(*InJsonValue->AsString()), GetSearchabilityStatus(InKey.ToString())));
+		ParsedTagsAndValues.Add(InKey.ToString(), FSearchableValueInfo(InDisplayKey, FCString::Atoi(*InJsonValue->AsString()), SearchabilityStatus));
 	}
 	else if (InJsonValue->Type == EJson::Boolean)
 	{
-		ParsedTagsAndValues.Add(InKey.ToString(), FSearchableValueInfo(InDisplayKey, FText::FromString(InJsonValue->AsString()), GetSearchabilityStatus(InKey.ToString())));
+		ParsedTagsAndValues.Add(InKey.ToString(), FSearchableValueInfo(InDisplayKey, FText::FromString(InJsonValue->AsString()), SearchabilityStatus));
 	}
 	else if( InJsonValue->Type == EJson::Array)
 	{
@@ -146,19 +171,19 @@ void FImaginaryFiBData::ParseJsonValue(FText InKey, FText InDisplayKey, TSharedP
 		for( int32 ArrayIdx = 0; ArrayIdx < ArrayList.Num(); ++ArrayIdx)
 		{
 			TSharedPtr< FJsonValue > ArrayValue = ArrayList[ArrayIdx];
-			ArrayCategory->ParseJsonValue(InKey, FText::FromString(FString::FromInt(ArrayIdx)), ArrayValue);
+			ArrayCategory->ParseJsonValue(InKey, FText::FromString(FString::FromInt(ArrayIdx)), ArrayValue,/*bIsInArray=*/true, SearchabilityStatus);
 		}		
 	}
 	else if (InJsonValue->Type == EJson::Object)
 	{
-		TSharedPtr< FCategorySectionHelper > SubObjectCategory = MakeShareable(new FCategorySectionHelper(InJsonValue->AsObject(), LookupTablePtr, InDisplayKey, false));
-		SubObjectCategory->ParseAllChildData();
+		TSharedPtr< FCategorySectionHelper > SubObjectCategory = MakeShareable(new FCategorySectionHelper(InJsonValue->AsObject(), LookupTablePtr, InDisplayKey, bIsInArray));
+		SubObjectCategory->ParseAllChildData(SearchabilityStatus);
 		ParsedChildData.Add(SubObjectCategory);
 	}
 	else
 	{
 		// For everything else, there's this. Numbers come here and will be treated as strings
-		ParsedTagsAndValues.Add(InKey.ToString(), FSearchableValueInfo(InDisplayKey, FText::FromString(InJsonValue->AsString()), GetSearchabilityStatus(InKey.ToString())));
+		ParsedTagsAndValues.Add(InKey.ToString(), FSearchableValueInfo(InDisplayKey, FText::FromString(InJsonValue->AsString()), SearchabilityStatus));
 	}
 }
 
@@ -238,6 +263,32 @@ bool FImaginaryFiBData::TestComplexExpression(const FName& InKey, const FTextFil
 	return bMatchesSearchQuery;
 }
 
+///////////////////////////
+// FFiBMetaData
+
+FFiBMetaData::FFiBMetaData(TSharedPtr< FJsonObject > InUnparsedJsonObject, TMap<int32, FText>* InLookupTablePtr)
+	: FImaginaryFiBData(InUnparsedJsonObject, InLookupTablePtr)
+{
+}
+
+bool FFiBMetaData::TrySpecialHandleJsonValue(FString InKey, TSharedPtr< FJsonValue > InJsonValue)
+{
+	bool bResult = false;
+	if(InKey == FFiBMD::FiBSearchableExplicitMD)
+	{
+		bIsExplicit = true;
+		bResult = true;
+	}
+	else if(InKey == FFiBMD::FiBSearchableHiddenExplicitMD)
+	{
+		bIsExplicit = true;
+		bIsHidden = true;
+		bResult = true;
+	}
+	ensure(bResult);
+	return bResult;
+}
+
 ///////////////////////////
 // FCategorySectionHelper
 
@@ -264,7 +315,7 @@ FCategorySectionHelper::FCategorySectionHelper(TSharedPtr< FJsonObject > InUnpar
 
 }
 
-bool FCategorySectionHelper::CanCallFilter(ESearchQueryFilter InSearchQueryFilter)
+bool FCategorySectionHelper::CanCallFilter(ESearchQueryFilter InSearchQueryFilter) const
 {
 	return true;
 }
@@ -274,7 +325,7 @@ FSearchResult FCategorySectionHelper::CreateSearchResult_Internal(FSearchResult
 	return FSearchResult(new FFindInBlueprintsResult(CategoryName, InParent));
 }
 
-void FCategorySectionHelper::ParseAllChildData()
+void FCategorySectionHelper::ParseAllChildData(ESearchableValueStatus InSearchabilityOverride/* = ESearchableValueStatus::Searchable*/)
 {
 	if (UnparsedJsonObject.IsValid() && SpecialHandlingCallback.IsBound())
 	{
@@ -283,7 +334,7 @@ void FCategorySectionHelper::ParseAllChildData()
 	}
 	else
 	{
-		FImaginaryFiBData::ParseAllChildData();
+		FImaginaryFiBData::ParseAllChildData(InSearchabilityOverride);
 	}
 }
 
@@ -315,12 +366,12 @@ FSearchResult FImaginaryBlueprint::CreateSearchResult_Internal(FSearchResult InP
 	return FSearchResult(new FFindInBlueprintsResult(ParsedTagsAndValues.Find(FFindInBlueprintSearchTags::FiB_Path.ToString())->GetDisplayText(LookupTable)));
 }
 
-bool FImaginaryBlueprint::IsCompatibleWithFilter(ESearchQueryFilter InSearchQueryFilter)
+bool FImaginaryBlueprint::IsCompatibleWithFilter(ESearchQueryFilter InSearchQueryFilter) const
 {
 	return InSearchQueryFilter == ESearchQueryFilter::AllFilter || InSearchQueryFilter == ESearchQueryFilter::BlueprintFilter;
 }
 
-bool FImaginaryBlueprint::CanCallFilter(ESearchQueryFilter InSearchQueryFilter)
+bool FImaginaryBlueprint::CanCallFilter(ESearchQueryFilter InSearchQueryFilter) const
 {
 	return InSearchQueryFilter == ESearchQueryFilter::NodesFilter ||
 		InSearchQueryFilter == ESearchQueryFilter::PinsFilter ||
@@ -422,7 +473,7 @@ FSearchResult FImaginaryGraph::CreateSearchResult_Internal(FSearchResult InParen
 	return FSearchResult(new FFindInBlueprintsGraph(FText::GetEmpty(), InParent, GraphType));
 }
 
-bool FImaginaryGraph::IsCompatibleWithFilter(ESearchQueryFilter InSearchQueryFilter)
+bool FImaginaryGraph::IsCompatibleWithFilter(ESearchQueryFilter InSearchQueryFilter) const
 {
 	return InSearchQueryFilter == ESearchQueryFilter::AllFilter || 
 		InSearchQueryFilter == ESearchQueryFilter::GraphsFilter ||
@@ -431,7 +482,7 @@ bool FImaginaryGraph::IsCompatibleWithFilter(ESearchQueryFilter InSearchQueryFil
 		(GraphType == GT_Macro && InSearchQueryFilter == ESearchQueryFilter::MacrosFilter);
 }
 
-bool FImaginaryGraph::CanCallFilter(ESearchQueryFilter InSearchQueryFilter)
+bool FImaginaryGraph::CanCallFilter(ESearchQueryFilter InSearchQueryFilter) const
 {
 	return InSearchQueryFilter == ESearchQueryFilter::PinsFilter ||
 		InSearchQueryFilter == ESearchQueryFilter::NodesFilter ||
@@ -491,12 +542,12 @@ FSearchResult FImaginaryGraphNode::CreateSearchResult_Internal(FSearchResult InP
 	return FSearchResult(new FFindInBlueprintsGraphNode(FText::GetEmpty(), InParent));
 }
 
-bool FImaginaryGraphNode::IsCompatibleWithFilter(ESearchQueryFilter InSearchQueryFilter)
+bool FImaginaryGraphNode::IsCompatibleWithFilter(ESearchQueryFilter InSearchQueryFilter) const
 {
 	return InSearchQueryFilter == ESearchQueryFilter::AllFilter || InSearchQueryFilter == ESearchQueryFilter::NodesFilter;
 }
 
-bool FImaginaryGraphNode::CanCallFilter(ESearchQueryFilter InSearchQueryFilter)
+bool FImaginaryGraphNode::CanCallFilter(ESearchQueryFilter InSearchQueryFilter) const
 {
 	return InSearchQueryFilter == ESearchQueryFilter::PinsFilter ||
 		FImaginaryFiBData::CanCallFilter(InSearchQueryFilter);
@@ -547,7 +598,7 @@ bool FImaginaryGraphNode::TrySpecialHandleJsonValue(FString InKey, TSharedPtr< F
 	return false;
 }
 
-void FImaginaryGraphNode::ParseAllChildData()
+void FImaginaryGraphNode::ParseAllChildData(ESearchableValueStatus InSearchabilityOverride/* = ESearchableValueStatus::Searchable*/)
 {
 	if (UnparsedJsonObject.IsValid())
 	{
@@ -559,7 +610,7 @@ void FImaginaryGraphNode::ParseAllChildData()
 			SchemaName = FindInBlueprintsHelpers::AsFText(SchemaNameValue, *LookupTablePtr).ToString();
 		}
 
-		FImaginaryFiBData::ParseAllChildData();
+		FImaginaryFiBData::ParseAllChildData(InSearchabilityOverride);
 	}
 }
 
@@ -571,7 +622,7 @@ FImaginaryProperty::FImaginaryProperty(TSharedPtr< FJsonObject > InUnparsedJsonO
 {
 }
 
-bool FImaginaryProperty::IsCompatibleWithFilter(ESearchQueryFilter InSearchQueryFilter)
+bool FImaginaryProperty::IsCompatibleWithFilter(ESearchQueryFilter InSearchQueryFilter) const
 {
 	return InSearchQueryFilter == ESearchQueryFilter::AllFilter || 
 		InSearchQueryFilter == ESearchQueryFilter::PropertiesFilter || 
@@ -613,7 +664,7 @@ FImaginaryComponent::FImaginaryComponent(TSharedPtr< FJsonObject > InUnparsedJso
 {
 }
 
-bool FImaginaryComponent::IsCompatibleWithFilter(ESearchQueryFilter InSearchQueryFilter)
+bool FImaginaryComponent::IsCompatibleWithFilter(ESearchQueryFilter InSearchQueryFilter) const
 {
 	return FImaginaryProperty::IsCompatibleWithFilter(InSearchQueryFilter) || InSearchQueryFilter == ESearchQueryFilter::ComponentsFilter;
 }
@@ -627,7 +678,7 @@ FImaginaryPin::FImaginaryPin(TSharedPtr< FJsonObject > InUnparsedJsonObject, TMa
 {
 }
 
-bool FImaginaryPin::IsCompatibleWithFilter(ESearchQueryFilter InSearchQueryFilter)
+bool FImaginaryPin::IsCompatibleWithFilter(ESearchQueryFilter InSearchQueryFilter) const
 {
 	return InSearchQueryFilter == ESearchQueryFilter::AllFilter || InSearchQueryFilter == ESearchQueryFilter::PinsFilter;
 }
diff --git a/Engine/Source/Editor/Kismet/Private/SMyBlueprint.cpp b/Engine/Source/Editor/Kismet/Private/SMyBlueprint.cpp
--- a/Engine/Source/Editor/Kismet/Private/SMyBlueprint.cpp
+++ b/Engine/Source/Editor/Kismet/Private/SMyBlueprint.cpp
@@ -2097,6 +2097,7 @@ void SMyBlueprint::ImplementFunction(FEdGraphSchemaAction_K2Graph* GraphAction)
 
 void SMyBlueprint::OnFindReference()
 {
+	bool bUseQuotes = true;
 	FString SearchTerm;
 	if (FEdGraphSchemaAction_K2Graph* GraphAction = SelectionAsGraph())
 	{
@@ -2104,11 +2105,24 @@ void SMyBlueprint::OnFindReference()
 	}
 	else if (FEdGraphSchemaAction_K2Var* VarAction = SelectionAsVar())
 	{
-		SearchTerm = VarAction->GetFriendlyVariableName();
+		FGuid Guid = FBlueprintEditorUtils::FindMemberVariableGuidByName(Blueprint, VarAction->GetVariableName());
+		if (Guid.IsValid())
+		{
+			FString VariableName = VarAction->GetVariableName().ToString();
+			SearchTerm = FString::Printf(TEXT("Nodes(VariableReference(MemberName=+%s && MemberGuid(A=%i && B=%i && C=%i && D=%i)))"), *VariableName, Guid.A, Guid.B, Guid.C, Guid.D);
+			bUseQuotes = false;
+		}
+		else
+		{
+			FString VariableName = VarAction->GetVariableName().ToString();
+			SearchTerm = FString::Printf(TEXT("Nodes(VariableReference(MemberName=+%s))"), *VariableName);
+			bUseQuotes = false;
+		}
 	}
 	else if (FEdGraphSchemaAction_K2LocalVar* LocalVarAction = SelectionAsLocalVar())
 	{
-		SearchTerm = LocalVarAction->GetFriendlyVariableName();
+		SearchTerm = FString::Printf(TEXT("Nodes(VariableReference(MemberName=+%s && MemberScope=+%s))"), *LocalVarAction->GetVariableName().ToString(), *LocalVarAction->GetVariableScope()->GetName());
+		bUseQuotes = false;
 	}
 	else if (FEdGraphSchemaAction_K2Delegate* DelegateAction = SelectionAsDelegate())
 	{
@@ -2135,7 +2149,11 @@ void SMyBlueprint::OnFindReference()
 
 	if(!SearchTerm.IsEmpty())
 	{
-		BlueprintEditorPtr.Pin()->SummonSearchUI(true, FString::Printf(TEXT("\"%s\""), *SearchTerm));
+		if (bUseQuotes)
+		{
+			SearchTerm = FString::Printf(TEXT("\"%s\""), *SearchTerm);
+		}
+		BlueprintEditorPtr.Pin()->SummonSearchUI(true, SearchTerm);
 	}
 }
 
diff --git a/Engine/Source/Editor/Kismet/Public/FindInBlueprintManager.h b/Engine/Source/Editor/Kismet/Public/FindInBlueprintManager.h
--- a/Engine/Source/Editor/Kismet/Public/FindInBlueprintManager.h
+++ b/Engine/Source/Editor/Kismet/Public/FindInBlueprintManager.h
@@ -71,6 +71,9 @@ struct KISMET_API FFindInBlueprintSearchTags
 	static const FText FiB_Glyph;
 	/** Glyph icon color tag */
 	static const FText FiB_GlyphColor;
+
+	// Identifier for metadata storage, completely unsearchable tag
+	static const FText FiBMetaDataTag;
 	/** End const values for Find-in-Blueprint */
 };
 
@@ -92,12 +95,16 @@ struct FSearchData
 	/** Interfaces implemented by the Blueprint */
 	TArray<FString> Interfaces;
 
+	/** Version of the data */
+	int32 Version;
+
 	/** Cached to determine if the Blueprint is seen as no longer valid, allows it to be cleared out next save to disk */
 	bool bMarkedForDeletion;
 
 	FSearchData()
 		: Blueprint(nullptr)
 		, bMarkedForDeletion(false)
+		, Version(0)
 	{
 	}
 
@@ -110,6 +117,7 @@ struct FSearchData
 		, ParentClass(MoveTemp(Other.ParentClass))
 		, Interfaces(MoveTemp(Other.Interfaces))
 		, bMarkedForDeletion(Other.bMarkedForDeletion)
+		, Version(0)
 	{
 	}
 
@@ -126,6 +134,7 @@ struct FSearchData
 		bMarkedForDeletion = RHS.bMarkedForDeletion;
 		ParentClass = MoveTemp(RHS.ParentClass);
 		Interfaces = MoveTemp(RHS.Interfaces);
+		Version = RHS.Version;
 		return *this;
 	}
 
@@ -136,6 +145,7 @@ struct FSearchData
 		, ParentClass(Other.ParentClass)
 		, Interfaces(Other.Interfaces)
 		, bMarkedForDeletion(Other.bMarkedForDeletion)
+		, Version(Other.Version)
 	{
 	}
 
@@ -152,6 +162,7 @@ struct FSearchData
 		bMarkedForDeletion = RHS.bMarkedForDeletion;
 		ParentClass = RHS.ParentClass;
 		Interfaces = RHS.Interfaces;
+		Version = RHS.Version;
 		return *this;
 	}
 
@@ -173,6 +184,85 @@ struct FSearchTagDataPair
 	FText Value;
 };
 
+enum EFiBVersion
+{
+	FIB_VER_BASE = 0, // All Blueprints prior to versioning will automatically be assumed to be at 0 if they have FiB data collected
+	FIB_VER_VARIABLE_REFERENCE, // Variable references (FMemberReference) is collected in FiB
+
+	// -----<new versions can be added before this line>-------------------------------------------------
+	FIB_VER_PLUS_ONE,
+	FIB_VER_ALL = FIB_VER_PLUS_ONE - 1 // Always the last version, we want Blueprints to be at latest
+};
+
+struct KISMET_API FFiBMD
+{
+	static const FString FiBSearchableMD;
+	static const FString FiBSearchableShallowMD;
+	static const FString FiBSearchableExplicitMD;
+	static const FString FiBSearchableHiddenExplicitMD;
+};
+
+////////////////////////////////////
+// FStreamSearch
+
+/**
+ * Async task for searching Blueprints
+ */
+class FStreamSearch : public FRunnable
+{
+public:
+	/** Constructor */
+	FStreamSearch(const FString& InSearchValue );
+
+	/** Begin FRunnable Interface */
+	virtual bool Init() override;
+
+	virtual uint32 Run() override;
+
+	virtual void Stop() override;
+
+	virtual void Exit() override;
+	/** End FRunnable Interface */
+
+	/** Brings the thread to a safe stop before continuing. */
+	void EnsureCompletion();
+
+	/** Returns TRUE if the thread is done with it's work. */
+	bool IsComplete() const;
+
+	/**
+	 * Appends the items filtered through the search filter to the passed array
+	 *
+	 * @param OutItemsFound		All the items found since last queried
+	 */
+	void GetFilteredItems(TArray<TSharedPtr<class FFindInBlueprintsResult>>& OutItemsFound);
+
+	/** Helper function to query the percent complete this search is */
+	float GetPercentComplete() const;
+
+public:
+	/** Thread to run the cleanup FRunnable on */
+	FRunnableThread* Thread;
+
+	/** A list of items found, cleared whenever the main thread pulls them to display to screen */
+	TArray<TSharedPtr<class FFindInBlueprintsResult>> ItemsFound;
+
+	/** The search value to filter results by */
+	FString SearchValue;
+
+	/** Prevents searching while other threads are pulling search results */
+	FCriticalSection SearchCriticalSection;
+
+	// Whether the thread has finished running
+	bool bThreadCompleted;
+
+	/** > 0 if we've been asked to abort work in progress at the next opportunity */
+	FThreadSafeCounter StopTaskCounter;
+};
+
+////////////////////////////////////
+// FFindInBlueprintSearchManager
+
 /** Singleton manager for handling all Blueprint searches, helps to manage the going progress of Blueprints, and is thread-safe. */
 class KISMET_API FFindInBlueprintSearchManager
 {
@@ -330,6 +420,9 @@ private:
 	/** Removes a Blueprint from being managed by the FiB system by passing in the UBlueprint's path */
 	void RemoveBlueprintByPath(FName InPath);
 
+	/** Begins the process of extracting unloaded FiB data */
+	void ExtractUnloadedFiBData(const FAssetData& InAssetData, const FString& InFiBData, bool bIsVersioned);
+
 protected:
 	/** Maps the Blueprint paths to their index in the SearchArray */
 	TMap<FName, int32> SearchMap;
diff --git a/Engine/Source/Editor/Kismet/Public/ImaginaryBlueprintData.h b/Engine/Source/Editor/Kismet/Public/ImaginaryBlueprintData.h
--- a/Engine/Source/Editor/Kismet/Public/ImaginaryBlueprintData.h
+++ b/Engine/Source/Editor/Kismet/Public/ImaginaryBlueprintData.h
@@ -134,13 +134,13 @@ public:
 	virtual bool IsTagAndValueCategory() { return false; }
 
 	/** Checks if the filter is compatible with the current object, returns TRUE by default */
-	virtual bool IsCompatibleWithFilter(ESearchQueryFilter InSearchQueryFilter);
+	virtual bool IsCompatibleWithFilter(ESearchQueryFilter InSearchQueryFilter) const;
 
 	/** Checks if the filter can call functions for the passed filter, returns FALSE by default if the filter is not the AllFilter */
-	virtual bool CanCallFilter(ESearchQueryFilter InSearchQueryFilter);
+	virtual bool CanCallFilter(ESearchQueryFilter InSearchQueryFilter) const;
 
 	/** Parses all child data, non-recursively, so children will be left in an unparsed Json state */
-	virtual void ParseAllChildData();
+	virtual void ParseAllChildData(ESearchableValueStatus InSearchabilityOverride = ESearchableValueStatus::Searchable);
 
 	/** Test the given value against the strings extracted from the current item. Will return the matching search components if any (can return TRUE without having any if the search components are hidden) */
 	virtual bool TestBasicStringExpression(const FTextFilterString& InValue, const ETextFilterTextComparisonMode InTextComparisonMode, TMultiMap< const FImaginaryFiBData*, FComponentUniqueDisplay >& InOutMatchingSearchComponents) const;
@@ -184,7 +184,7 @@ protected:
 	FText CreateSearchComponentDisplayText(FText InKey, FText InValue) const;
 
 	/** Helper function for parsing Json values into usable properties */
-	void ParseJsonValue(FText InKey, FText InDisplayKey, TSharedPtr< FJsonValue > InJsonValue);
+	void ParseJsonValue(FText InKey, FText InDisplayKey, TSharedPtr< FJsonValue > InJsonValue, bool bIsInArray=false, ESearchableValueStatus InSearchabilityOverride = ESearchableValueStatus::Searchable);
 
 protected:
 	/** The unparsed Json object representing this item. Auto-cleared after parsing */
@@ -200,6 +200,40 @@ protected:
 	TMap<int32, FText>* LookupTablePtr;
 };
 
+class FFiBMetaData : public FImaginaryFiBData
+{
+public:
+	FFiBMetaData(TSharedPtr< FJsonObject > InUnparsedJsonObject, TMap<int32, FText>* InLookupTablePtr);
+	
+	/** FImaginaryFiBData Interface */
+	virtual bool TrySpecialHandleJsonValue(FString InKey, TSharedPtr< FJsonValue > InJsonValue) override;
+	virtual FSearchResult CreateSearchResult_Internal(FSearchResult InParent) const override
+	{
+		return nullptr;
+	}
+	/** End FImaginaryFiBData Interface */
+
+	/** Returns TRUE if the metadata is informing that the UProperty and children should be hidden */
+	bool IsHidden() const
+	{
+		// While handled separately, when hidden it should always be explicit
+		ensure(bIsHidden);
+		return bIsHidden;
+	}
+
+	/** Returns TRUE if the metadata is informing that the UProperty and children should be explicit */
+	bool IsExplicit() const
+	{
+		return bIsExplicit;
+	}
+protected:
+	/** TRUE if the UProperty this metadata represents is hidden */
+	bool bIsHidden;
+
+	/** TRUE if the UProperty this metadata represents is explicit, should always be true if bIsHidden is true */
+	bool bIsExplicit;
+};
+
 class FCategorySectionHelper : public FImaginaryFiBData
 {
 public:
@@ -211,14 +245,20 @@ public:
 	FCategorySectionHelper(TSharedPtr< FJsonObject > InUnparsedJsonObject, TMap<int32, FText>* InLookupTablePtr, FText InCategoryName, bool bInTagAndValueCategory, FCategorySectionHelperCallback InSpecialHandlingCallback);
 
 	/** FImaginaryFiBData Interface */
-	virtual void ParseAllChildData() override;
+	virtual void ParseAllChildData(ESearchableValueStatus InSearchabilityOverride/* = ESearchableValueStatus::ESearchableValueStatus::Searchable*/) override;
 	virtual bool IsCategory() override { return true; }
 	virtual bool IsTagAndValueCategory() override { return bIsTagAndValue; }
-		/** End FImaginaryFiBData Interface */
+	/** End FImaginaryFiBData Interface */
+
+	/** Returns the category name prepared for checking as a function */
+	FString GetCategoryFunctionName() const
+	{
+		return CategoryName.BuildSourceString();
+	}
 
 protected:
 	/** FImaginaryFiBData Interface */
-	virtual bool CanCallFilter(ESearchQueryFilter InSearchQueryFilter) override;
+	virtual bool CanCallFilter(ESearchQueryFilter InSearchQueryFilter) const override;
 	virtual FSearchResult CreateSearchResult_Internal(FSearchResult InParent) const override;
 	/** End FImaginaryFiBData Interface */
 
@@ -240,8 +280,8 @@ public:
 	FImaginaryBlueprint(FString InBlueprintName, FString InBlueprintPath, FString InBlueprintParentClass, TArray<FString>& InInterfaces, FString InUnparsedStringData);
 
 	/** FImaginaryFiBData Interface */
-	virtual bool IsCompatibleWithFilter(ESearchQueryFilter InSearchQueryFilter) override;
-	virtual bool CanCallFilter(ESearchQueryFilter InSearchQueryFilter) override;
+	virtual bool IsCompatibleWithFilter(ESearchQueryFilter InSearchQueryFilter) const override;
+	virtual bool CanCallFilter(ESearchQueryFilter InSearchQueryFilter) const override;
 	/** End FImaginaryFiBData Interface */
 
 protected:
@@ -277,8 +317,8 @@ public:
 	FImaginaryGraph(TSharedPtr< FJsonObject > InUnparsedJsonObject, TMap<int32, FText>* InLookupTablePtr, EGraphType InGraphType);
 
 	/** FImaginaryFiBData Interface */
-	virtual bool IsCompatibleWithFilter(ESearchQueryFilter InSearchQueryFilter) override;
-	virtual bool CanCallFilter(ESearchQueryFilter InSearchQueryFilter) override;
+	virtual bool IsCompatibleWithFilter(ESearchQueryFilter InSearchQueryFilter) const override;
+	virtual bool CanCallFilter(ESearchQueryFilter InSearchQueryFilter) const override;
 	/** End FImaginaryFiBData Interface */
 
 protected:
@@ -300,9 +340,9 @@ public:
 	FImaginaryGraphNode(TSharedPtr< FJsonObject > InUnparsedJsonObject, TMap<int32, FText>* InLookupTablePtr);
 
 	/** FImaginaryFiBData Interface */
-	virtual bool IsCompatibleWithFilter(ESearchQueryFilter InSearchQueryFilter) override;
-	virtual bool CanCallFilter(ESearchQueryFilter InSearchQueryFilter) override;
-	virtual void ParseAllChildData() override;
+	virtual bool IsCompatibleWithFilter(ESearchQueryFilter InSearchQueryFilter) const override;
+	virtual bool CanCallFilter(ESearchQueryFilter InSearchQueryFilter) const override;
+	virtual void ParseAllChildData(ESearchableValueStatus InSearchabilityOverride/* = ESearchableValueStatus::Searchable*/) override;
 	/** End FImaginaryFiBData Interface */
 
 protected:
@@ -324,7 +364,7 @@ public:
 	FImaginaryProperty(TSharedPtr< FJsonObject > InUnparsedJsonObject, TMap<int32, FText>* InLookupTablePtr);
 
 	/** FImaginaryFiBData Interface */
-	virtual bool IsCompatibleWithFilter(ESearchQueryFilter InSearchQueryFilter) override;
+	virtual bool IsCompatibleWithFilter(ESearchQueryFilter InSearchQueryFilter) const override;
 	/** End FImaginaryFiBData Interface */
 	
 protected:
@@ -341,7 +381,7 @@ public:
 	FImaginaryComponent(TSharedPtr< FJsonObject > InUnparsedJsonObject, TMap<int32, FText>* InLookupTablePtr);
 
 	/** FImaginaryFiBData Interface */
-	virtual bool IsCompatibleWithFilter(ESearchQueryFilter InSearchQueryFilter) override;
+	virtual bool IsCompatibleWithFilter(ESearchQueryFilter InSearchQueryFilter) const override;
 	/** End FImaginaryFiBData Interface */
 };
 
@@ -352,7 +392,7 @@ public:
 	FImaginaryPin(TSharedPtr< FJsonObject > InUnparsedJsonObject, TMap<int32, FText>* InLookupTablePtr, FString InSchemaName);
 
 	/** FImaginaryFiBData Interface */
-	virtual bool IsCompatibleWithFilter(ESearchQueryFilter InSearchQueryFilter) override;
+	virtual bool IsCompatibleWithFilter(ESearchQueryFilter InSearchQueryFilter) const override;
 	/** End FImaginaryFiBData Interface */
 
 protected:
diff --git a/Engine/Source/Editor/UnrealEd/Private/Kismet2/BlueprintEditorUtils.cpp b/Engine/Source/Editor/UnrealEd/Private/Kismet2/BlueprintEditorUtils.cpp
--- a/Engine/Source/Editor/UnrealEd/Private/Kismet2/BlueprintEditorUtils.cpp
+++ b/Engine/Source/Editor/UnrealEd/Private/Kismet2/BlueprintEditorUtils.cpp
@@ -4536,7 +4536,7 @@ void FBlueprintEditorUtils::ChangeMemberVariableType(UBlueprint* Blueprint, cons
 
 							const bool bSetFindWithinBlueprint = false;
 							const bool bSelectFirstResult = false;
-							BlueprintEditor->SummonSearchUI(bSetFindWithinBlueprint, VariableName.ToString(), bSelectFirstResult);
+							BlueprintEditor->SummonSearchUI(bSetFindWithinBlueprint, AllVariableNodes[0]->GetFindReferenceSearchString(), bSelectFirstResult);
 						}
 					}
 					else
@@ -5014,7 +5014,7 @@ void FBlueprintEditorUtils::ChangeLocalVariableType(UBlueprint* InBlueprint, con
 
 					const bool bSetFindWithinBlueprint = true;
 					const bool bSelectFirstResult = false;
-					BlueprintEditor->SummonSearchUI(bSetFindWithinBlueprint, InVariableName.ToString(), bSelectFirstResult);
+					BlueprintEditor->SummonSearchUI(bSetFindWithinBlueprint, VariableNodes[0]->GetFindReferenceSearchString(), bSelectFirstResult);
 				}
 			}
 		}
diff --git a/Engine/Source/Runtime/Core/Public/UObject/ObjectVersion.h b/Engine/Source/Runtime/Core/Public/UObject/ObjectVersion.h
--- a/Engine/Source/Runtime/Core/Public/UObject/ObjectVersion.h
+++ b/Engine/Source/Runtime/Core/Public/UObject/ObjectVersion.h
@@ -570,6 +570,8 @@ enum EUnrealEngineObjectUE4Version
 	VER_UE4_DIALOGUE_WAVE_NAMESPACE_AND_CONTEXT_CHANGES,
 	// Renamed MakeRot MakeRotator and rearranged parameters.
 	VER_UE4_MAKE_ROT_RENAME_AND_REORDER,
+	// K2Node_Variable will properly have the VariableReference Guid set if available
+	VER_UE4_K2NODE_VAR_REFERENCEGUIDS,
 	// -----<new versions can be added before this line>-------------------------------------------------
 	// - this needs to be the last line (see note below)
 	VER_UE4_AUTOMATIC_VERSION_PLUS_ONE,
diff --git a/Engine/Source/Runtime/Engine/Classes/EdGraph/EdGraphNode.h b/Engine/Source/Runtime/Engine/Classes/EdGraph/EdGraphNode.h
--- a/Engine/Source/Runtime/Engine/Classes/EdGraph/EdGraphNode.h
+++ b/Engine/Source/Runtime/Engine/Classes/EdGraph/EdGraphNode.h
@@ -337,6 +337,9 @@ public:
 	/** Gets the name of this node, shown in title bar */
 	virtual FText GetNodeTitle(ENodeTitleType::Type TitleType) const;
 
+	/** Gets the search string to find references to this node */
+	virtual FString GetFindReferenceSearchString() const;
+
 	/** 
 	 * Gets the draw color of a node's title bar
 	 */
diff --git a/Engine/Source/Runtime/Engine/Classes/EdGraph/EdGraphSchema.h b/Engine/Source/Runtime/Engine/Classes/EdGraph/EdGraphSchema.h
--- a/Engine/Source/Runtime/Engine/Classes/EdGraph/EdGraphSchema.h
+++ b/Engine/Source/Runtime/Engine/Classes/EdGraph/EdGraphSchema.h
@@ -763,9 +763,6 @@ class ENGINE_API UEdGraphSchema : public UObject
 	/** Returns schema action to create comment from implemention */
 	virtual TSharedPtr<FEdGraphSchemaAction> GetCreateCommentAction() const { return NULL; }
 
-	/** Returns schema action to create documention node from implemention */
-	virtual TSharedPtr<FEdGraphSchemaAction> GetCreateDocumentNodeAction() const { return NULL; }
-
 	/**
 	 * Handle a graph being removed by the user (potentially removing associated bound nodes, etc...)
 	 */
diff --git a/Engine/Source/Runtime/Engine/Classes/Engine/MemberReference.h b/Engine/Source/Runtime/Engine/Classes/Engine/MemberReference.h
--- a/Engine/Source/Runtime/Engine/Classes/Engine/MemberReference.h
+++ b/Engine/Source/Runtime/Engine/Classes/Engine/MemberReference.h
@@ -164,6 +164,7 @@ public:
 
 	/** Set to a non-'self' member, so must include reference to class owning the member. */
 	ENGINE_API void SetExternalMember(FName InMemberName, TSubclassOf<class UObject> InMemberParentClass);
+	ENGINE_API void SetExternalMember(FName InMemberName, TSubclassOf<class UObject> InMemberParentClass, FGuid& InMemberGuid);
 
 	/** Set to reference a global field (intended for things like natively defined delegate signatures) */
 	ENGINE_API void SetGlobalField(FName InFieldName, UPackage* InParentPackage);
@@ -173,6 +174,7 @@ public:
 
 	/** Set up this reference to a 'self' member name */
 	ENGINE_API void SetSelfMember(FName InMemberName);
+	ENGINE_API void SetSelfMember(FName InMemberName, FGuid& InMemberGuid);
 
 	/** Set up this reference to a 'self' member name, scoped to a struct */
 	ENGINE_API void SetLocalMember(FName InMemberName, UStruct* InScope, const FGuid InMemberGuid);
diff --git a/Engine/Source/Runtime/Engine/Private/Blueprint.cpp b/Engine/Source/Runtime/Engine/Private/Blueprint.cpp
--- a/Engine/Source/Runtime/Engine/Private/Blueprint.cpp
+++ b/Engine/Source/Runtime/Engine/Private/Blueprint.cpp
@@ -747,7 +747,7 @@ void UBlueprint::GetAssetRegistryTags(TArray<FAssetRegistryTag>& OutTags) const
 			FBlueprintEditorUtils::IsDataOnlyBlueprint(this) ? TEXT("True") : TEXT("False"),
 			FAssetRegistryTag::TT_Alphabetical ) );
 
-		OutTags.Add( FAssetRegistryTag("FiB", FFindInBlueprintSearchManager::Get().QuerySingleBlueprint((UBlueprint*)this, false), FAssetRegistryTag::TT_Hidden) );
+		OutTags.Add( FAssetRegistryTag("FiBData", FFindInBlueprintSearchManager::Get().QuerySingleBlueprint((UBlueprint*)this, false), FAssetRegistryTag::TT_Hidden) );
 	}
 }
 
diff --git a/Engine/Source/Runtime/Engine/Private/EdGraph/EdGraphNode.cpp b/Engine/Source/Runtime/Engine/Private/EdGraph/EdGraphNode.cpp
--- a/Engine/Source/Runtime/Engine/Private/EdGraph/EdGraphNode.cpp
+++ b/Engine/Source/Runtime/Engine/Private/EdGraph/EdGraphNode.cpp
@@ -359,6 +359,11 @@ FText UEdGraphNode::GetNodeTitle(ENodeTitleType::Type TitleType) const
 	return FText::FromString(GetClass()->GetName());
 }
 
+FString UEdGraphNode::GetFindReferenceSearchString() const
+{
+	return GetNodeTitle(ENodeTitleType::ListView).ToString();
+}
+
 UObject* UEdGraphNode::GetJumpTargetForDoubleClick() const
 {
 	return NULL;
diff --git a/Engine/Source/Runtime/Engine/Private/MemberReference.cpp b/Engine/Source/Runtime/Engine/Private/MemberReference.cpp
--- a/Engine/Source/Runtime/Engine/Private/MemberReference.cpp
+++ b/Engine/Source/Runtime/Engine/Private/MemberReference.cpp
@@ -24,6 +24,12 @@ void FMemberReference::SetExternalMember(FName InMemberName, TSubclassOf<class U
 	bWasDeprecated = false;
 }
 
+void FMemberReference::SetExternalMember(FName InMemberName, TSubclassOf<class UObject> InMemberParentClass, FGuid& InMemberGuid)
+{
+	SetExternalMember(InMemberName, InMemberParentClass);
+	MemberGuid = InMemberGuid;
+}
+
 void FMemberReference::SetGlobalField(FName InFieldName, UPackage* InParentPackage)
 {
 	MemberName = InFieldName;
@@ -47,6 +53,12 @@ void FMemberReference::SetSelfMember(FName InMemberName)
 	bWasDeprecated = false;
 }
 
+void FMemberReference::SetSelfMember(FName InMemberName, FGuid& InMemberGuid)
+{
+	SetSelfMember(InMemberName);
+	MemberGuid = InMemberGuid;
+}
+
 void FMemberReference::SetDirect(const FName InMemberName, const FGuid InMemberGuid, TSubclassOf<class UObject> InMemberParentClass, bool bIsConsideredSelfContext)
 {
 	MemberName = InMemberName;
