Fixed non-unity

diff --git a/Engine/Source/Editor/Kismet/Private/FindInBlueprintManager.cpp b/Engine/Source/Editor/Kismet/Private/FindInBlueprintManager.cpp
--- a/Engine/Source/Editor/Kismet/Private/FindInBlueprintManager.cpp
+++ b/Engine/Source/Editor/Kismet/Private/FindInBlueprintManager.cpp
@@ -6,7 +6,9 @@
 #include "FindInBlueprints.h"
 #include "Editor/UnrealEd/Public/Kismet2/BlueprintEditorUtils.h"
 #include "AssetRegistryModule.h"
+#include "FiBSearchInstance.h"
 #include "HotReloadInterface.h"
+#include "ImaginaryBlueprintData.h"
 
 #include "JsonUtilities.h"
 #include "SNotificationList.h"
