Added experimental feature for "Find and Replace References" which currently works for variables (and is not advisable to use a native variable as the source).

Added FBlueprintEditorUtils::GetIconFromPin

Removed SPinTypeSelector::GetIconFromPin

Simplified UK2Node_Variable::GetFindReferenceSearchString to leverage FMemberReference::GetSearchString

Optimizations to parsing of Find-in-Blueprint data so it only occurs once (per refresh of data)

diff --git a/Engine/Config/BaseEditorPerProjectUserSettings.ini b/Engine/Config/BaseEditorPerProjectUserSettings.ini
--- a/Engine/Config/BaseEditorPerProjectUserSettings.ini
+++ b/Engine/Config/BaseEditorPerProjectUserSettings.ini
@@ -133,6 +133,7 @@ bToolbarCustomization=False
 bBehaviorTreeEditor=False
 bEnableWidgetVisualDiff=False
 bEnableAnimVisualDiff=False
+bEnableFindAndReplaceReferences =False
 
 [/Script/UnrealEd.EditorLoadingSavingSettings]
 ; True if we should automatically load a default level at start up
diff --git a/Engine/Source/Editor/BlueprintGraph/Private/K2Node_Variable.cpp b/Engine/Source/Editor/BlueprintGraph/Private/K2Node_Variable.cpp
--- a/Engine/Source/Editor/BlueprintGraph/Private/K2Node_Variable.cpp
+++ b/Engine/Source/Editor/BlueprintGraph/Private/K2Node_Variable.cpp
@@ -204,19 +204,12 @@ FString UK2Node_Variable::GetFindReferenceSearchString() const
 	FString ResultSearchString;
 	if (VariableReference.IsLocalScope())
 	{
-		return FString::Printf(TEXT("Nodes(VariableReference(MemberName=+%s && MemberScope=+%s))"), *VariableReference.GetMemberName().ToString(), *VariableReference.GetMemberScopeName());
+		ResultSearchString = VariableReference.GetReferenceSearchString(nullptr);
 	}
 	else
 	{
-		FGuid Guid = VariableReference.GetMemberGuid();
-		if (Guid.IsValid())
-		{
-			ResultSearchString = FString::Printf(TEXT("Nodes(VariableReference(MemberName=+%s && MemberGuid(A=%i && B=%i && C=%i && D=%i)))"), *VariableReference.GetMemberName().ToString(), Guid.A, Guid.B, Guid.C, Guid.D);
-		}
-		else
-		{
-			ResultSearchString = FString::Printf(TEXT("Nodes(VariableReference(MemberName=+%s))"), *VariableReference.GetMemberName().ToString());
-		}
+		UProperty* VariableProperty = VariableReference.ResolveMember<UProperty>(GetBlueprintClassFromNode());
+		ResultSearchString = VariableReference.GetReferenceSearchString(VariableProperty->GetOwnerClass());
 	}
 	return ResultSearchString;
 }
diff --git a/Engine/Source/Editor/GraphEditor/Private/GraphEditorActions.cpp b/Engine/Source/Editor/GraphEditor/Private/GraphEditorActions.cpp
--- a/Engine/Source/Editor/GraphEditor/Private/GraphEditorActions.cpp
+++ b/Engine/Source/Editor/GraphEditor/Private/GraphEditorActions.cpp
@@ -62,6 +62,7 @@ void FGraphEditorCommandsImpl::RegisterCommands()
 	UI_COMMAND( SelectReferenceInLevel, "Find Actor in Level", "Select the actor referenced by this node in the level", EUserInterfaceActionType::Button, FInputChord() )
 	UI_COMMAND( AssignReferencedActor, "Assign selected Actor", "Assign the selected actor to be this node's referenced object", EUserInterfaceActionType::Button, FInputChord() )
 	UI_COMMAND( FindReferences, "Find References", "Find references of this item", EUserInterfaceActionType::Button, FInputChord() )
+	UI_COMMAND( FindAndReplaceReferences, "Find and Replace References", "Brings up a window to help find and replace all instances of this item", EUserInterfaceActionType::Button, FInputChord() )
 	
 	UI_COMMAND( GotoNativeFunctionDefinition, "Goto Code Definition", "Goto the native code definition of this function", EUserInterfaceActionType::Button, FInputChord() )
 	UI_COMMAND( GotoNativeVariableDefinition, "Goto Code Definition", "Goto the native code definition of this variable", EUserInterfaceActionType::Button, FInputChord() )
diff --git a/Engine/Source/Editor/GraphEditor/Public/GraphEditorActions.h b/Engine/Source/Editor/GraphEditor/Public/GraphEditorActions.h
--- a/Engine/Source/Editor/GraphEditor/Public/GraphEditorActions.h
+++ b/Engine/Source/Editor/GraphEditor/Public/GraphEditorActions.h
@@ -86,6 +86,7 @@ public:
 
 	// Find references
 	TSharedPtr< FUICommandInfo > FindReferences;
+	TSharedPtr< FUICommandInfo > FindAndReplaceReferences;
 
 	// Goto native code actions
 	TSharedPtr< FUICommandInfo > GotoNativeFunctionDefinition;
diff --git a/Engine/Source/Editor/Kismet/Private/BlueprintEditor.cpp b/Engine/Source/Editor/Kismet/Private/BlueprintEditor.cpp
--- a/Engine/Source/Editor/Kismet/Private/BlueprintEditor.cpp
+++ b/Engine/Source/Editor/Kismet/Private/BlueprintEditor.cpp
@@ -53,6 +53,7 @@
 #include "SBlueprintPalette.h"
 #include "SBlueprintActionMenu.h"
 #include "SMyBlueprint.h"
+#include "SReplaceNodeReferences.h"
 #include "FindInBlueprints.h"
 // End of core kismet tabs
 
@@ -661,6 +662,11 @@ void FBlueprintEditor::RefreshEditors(ERefreshBlueprintEditorReason::Type Reason
 		MyBlueprintWidget->Refresh();
 	}
 
+	if (ReplaceReferencesWidget.IsValid())
+	{
+		ReplaceReferencesWidget->Refresh();
+	}
+
 	if (SCSEditor.IsValid())
 	{
 		SCSEditor->UpdateTree();
@@ -727,6 +733,11 @@ void FBlueprintEditor::SummonSearchUI(bool bSetFindWithinBlueprint, FString NewS
 	FindResults->FocusForUse(bSetFindWithinBlueprint, NewSearchTerms, bSelectFirstResult);
 }
 
+void FBlueprintEditor::SummonFindAndReplaceUI()
+{
+	TabManager->InvokeTab(FBlueprintEditorTabs::ReplaceNodeReferencesID);
+}
+
 void FBlueprintEditor::EnableSCSPreview(bool bEnable)
 {
 	if(SCSViewport.IsValid())
@@ -2181,6 +2192,7 @@ void FBlueprintEditor::CreateDefaultTabContents(const TArray<UBlueprint*>& InBlu
 	if (IsEditingSingleBlueprint())
 	{
 		this->MyBlueprintWidget = SNew(SMyBlueprint, SharedThis(this));
+		this->ReplaceReferencesWidget = SNew(SReplaceNodeReferences, SharedThis(this));
 	}
 	
 	FMessageLogModule& MessageLogModule = FModuleManager::LoadModuleChecked<FMessageLogModule>("MessageLog");
diff --git a/Engine/Source/Editor/Kismet/Private/BlueprintEditorModes.cpp b/Engine/Source/Editor/Kismet/Private/BlueprintEditorModes.cpp
--- a/Engine/Source/Editor/Kismet/Private/BlueprintEditorModes.cpp
+++ b/Engine/Source/Editor/Kismet/Private/BlueprintEditorModes.cpp
@@ -105,6 +105,10 @@ FBlueprintEditorApplicationMode::FBlueprintEditorApplicationMode(TSharedPtr<clas
 	BlueprintEditorTabFactories.RegisterFactory(MakeShareable(new FDebugInfoSummoner(InBlueprintEditor)));
 	BlueprintEditorTabFactories.RegisterFactory(MakeShareable(new FPaletteSummoner(InBlueprintEditor)));
 	BlueprintEditorTabFactories.RegisterFactory(MakeShareable(new FMyBlueprintSummoner(InBlueprintEditor)));
+	if (GetDefault<UEditorExperimentalSettings>()->bEnableFindAndReplaceReferences)
+	{
+		BlueprintEditorTabFactories.RegisterFactory(MakeShareable(new FReplaceNodeReferencesSummoner(InBlueprintEditor)));
+	}
 	BlueprintEditorTabFactories.RegisterFactory(MakeShareable(new FCompilerResultsSummoner(InBlueprintEditor)));
 	BlueprintEditorTabFactories.RegisterFactory(MakeShareable(new FFindResultsSummoner(InBlueprintEditor)));
 	
@@ -352,6 +356,10 @@ FBlueprintInterfaceApplicationMode::FBlueprintInterfaceApplicationMode(TSharedPt
 	// Create the tab factories
 	BlueprintInterfaceTabFactories.RegisterFactory(MakeShareable(new FDebugInfoSummoner(InBlueprintEditor)));
 	BlueprintInterfaceTabFactories.RegisterFactory(MakeShareable(new FMyBlueprintSummoner(InBlueprintEditor)));
+	if (GetDefault<UEditorExperimentalSettings>()->bEnableFindAndReplaceReferences)
+	{
+		BlueprintInterfaceTabFactories.RegisterFactory(MakeShareable(new FReplaceNodeReferencesSummoner(InBlueprintEditor)));
+	}
 	BlueprintInterfaceTabFactories.RegisterFactory(MakeShareable(new FCompilerResultsSummoner(InBlueprintEditor)));
 	BlueprintInterfaceTabFactories.RegisterFactory(MakeShareable(new FFindResultsSummoner(InBlueprintEditor)));
 	BlueprintInterfaceTabFactories.RegisterFactory(MakeShareable(new FSelectionDetailsSummoner(InBlueprintEditor)));
@@ -452,6 +460,10 @@ FBlueprintMacroApplicationMode::FBlueprintMacroApplicationMode(TSharedPtr<class
 	// Create the tab factories
 	BlueprintMacroTabFactories.RegisterFactory(MakeShareable(new FDebugInfoSummoner(InBlueprintEditor)));
 	BlueprintMacroTabFactories.RegisterFactory(MakeShareable(new FMyBlueprintSummoner(InBlueprintEditor)));
+	if (GetDefault<UEditorExperimentalSettings>()->bEnableFindAndReplaceReferences)
+	{
+		BlueprintMacroTabFactories.RegisterFactory(MakeShareable(new FReplaceNodeReferencesSummoner(InBlueprintEditor)));
+	}
 	BlueprintMacroTabFactories.RegisterFactory(MakeShareable(new FPaletteSummoner(InBlueprintEditor)));
 	BlueprintMacroTabFactories.RegisterFactory(MakeShareable(new FFindResultsSummoner(InBlueprintEditor)));
 	BlueprintMacroTabFactories.RegisterFactory(MakeShareable(new FSelectionDetailsSummoner(InBlueprintEditor)));
@@ -554,6 +566,10 @@ FBlueprintEditorUnifiedMode::FBlueprintEditorUnifiedMode(TSharedPtr<class FBluep
 	BlueprintEditorTabFactories.RegisterFactory(MakeShareable(new FDebugInfoSummoner(InBlueprintEditor)));
 	BlueprintEditorTabFactories.RegisterFactory(MakeShareable(new FPaletteSummoner(InBlueprintEditor)));
 	BlueprintEditorTabFactories.RegisterFactory(MakeShareable(new FMyBlueprintSummoner(InBlueprintEditor)));
+	if (GetDefault<UEditorExperimentalSettings>()->bEnableFindAndReplaceReferences)
+	{
+		BlueprintEditorTabFactories.RegisterFactory(MakeShareable(new FReplaceNodeReferencesSummoner(InBlueprintEditor)));
+	}
 	BlueprintEditorTabFactories.RegisterFactory(MakeShareable(new FCompilerResultsSummoner(InBlueprintEditor)));
 	BlueprintEditorTabFactories.RegisterFactory(MakeShareable(new FFindResultsSummoner(InBlueprintEditor)));
 	
diff --git a/Engine/Source/Editor/Kismet/Private/BlueprintEditorTabFactories.cpp b/Engine/Source/Editor/Kismet/Private/BlueprintEditorTabFactories.cpp
--- a/Engine/Source/Editor/Kismet/Private/BlueprintEditorTabFactories.cpp
+++ b/Engine/Source/Editor/Kismet/Private/BlueprintEditorTabFactories.cpp
@@ -15,6 +15,7 @@
 #include "FindInBlueprints.h"
 #include "Kismet2/BlueprintEditorUtils.h"
 #include "SMyBlueprint.h"
+#include "SReplaceNodeReferences.h"
 #include "Engine/Blueprint.h"
 #include "Engine/TimelineTemplate.h"
 #include "EditorClassUtils.h"
@@ -329,6 +330,25 @@ TSharedRef<SWidget> FMyBlueprintSummoner::CreateTabBody(const FWorkflowTabSpawnI
 	return BlueprintEditorPtr->GetMyBlueprintWidget().ToSharedRef();
 }
 
+FReplaceNodeReferencesSummoner::FReplaceNodeReferencesSummoner(TSharedPtr<class FAssetEditorToolkit> InHostingApp)
+	: FWorkflowTabFactory(FBlueprintEditorTabs::ReplaceNodeReferencesID, InHostingApp)
+{
+	TabLabel = LOCTEXT("ReplaceNodeReferences", "Replace References");
+	TabIcon = FSlateIcon(FEditorStyle::GetStyleSetName(), "ClassIcon.BlueprintCore");
+
+	bIsSingleton = true;
+
+	ViewMenuDescription = LOCTEXT("ReplaceNodeReferences", "Replace References");
+	ViewMenuTooltip = LOCTEXT("ReplaceNodeReferences_Tooltip", "Show Replace References");
+}
+
+TSharedRef<SWidget> FReplaceNodeReferencesSummoner::CreateTabBody(const FWorkflowTabSpawnInfo& Info) const
+{
+	TSharedPtr<FBlueprintEditor> BlueprintEditorPtr = StaticCastSharedPtr<FBlueprintEditor>(HostingApp.Pin());
+
+	return BlueprintEditorPtr->GetReplaceReferencesWidget().ToSharedRef();
+}
+
 FCompilerResultsSummoner::FCompilerResultsSummoner(TSharedPtr<class FAssetEditorToolkit> InHostingApp)
 	: FWorkflowTabFactory(FBlueprintEditorTabs::CompilerResultsID, InHostingApp)
 {
diff --git a/Engine/Source/Editor/Kismet/Private/BlueprintEditorTabFactories.h b/Engine/Source/Editor/Kismet/Private/BlueprintEditorTabFactories.h
--- a/Engine/Source/Editor/Kismet/Private/BlueprintEditorTabFactories.h
+++ b/Engine/Source/Editor/Kismet/Private/BlueprintEditorTabFactories.h
@@ -239,6 +239,22 @@ public:
 	}
 };
 
+/////////////////////////////////////////////////////
+// FReplaceNodeReferencesSummoner
+
+struct FReplaceNodeReferencesSummoner : public FWorkflowTabFactory
+{
+public:
+	FReplaceNodeReferencesSummoner(TSharedPtr<class FAssetEditorToolkit> InHostingApp);
+
+	virtual TSharedRef<SWidget> CreateTabBody(const FWorkflowTabSpawnInfo& Info) const override;
+
+	virtual FText GetTabToolTipText(const FWorkflowTabSpawnInfo& Info) const override
+	{
+		return LOCTEXT("ReplaceNodeReferencesTooltip", "Allows the user to replace node references in all Blueprint graphs with another node reference.");
+	}
+};
+
 //////////////////////////////////////////////////////////////////////////
 // FCompilerResultsSummoner
 
diff --git a/Engine/Source/Editor/Kismet/Private/BlueprintEditorTabs.cpp b/Engine/Source/Editor/Kismet/Private/BlueprintEditorTabs.cpp
--- a/Engine/Source/Editor/Kismet/Private/BlueprintEditorTabs.cpp
+++ b/Engine/Source/Editor/Kismet/Private/BlueprintEditorTabs.cpp
@@ -12,6 +12,7 @@ const FName FBlueprintEditorTabs::FindResultsID( TEXT( "FindResults" ) );
 const FName FBlueprintEditorTabs::ConstructionScriptEditorID( TEXT( "ConstructionScriptEditor" ) );
 const FName FBlueprintEditorTabs::SCSViewportID( TEXT( "SCSViewport" ) );
 const FName FBlueprintEditorTabs::MyBlueprintID( TEXT( "MyBlueprint" ) );
+const FName FBlueprintEditorTabs::ReplaceNodeReferencesID( TEXT( "ReplaceNodeReferences" ) );
 
 const FName FBlueprintEditorTabs::GraphEditorID( TEXT( "GraphEditor" ) );
 const FName FBlueprintEditorTabs::TimelineEditorID( TEXT( "TimelineEditor" ) );
diff --git a/Engine/Source/Editor/Kismet/Private/FiBSearchInstance.cpp b/Engine/Source/Editor/Kismet/Private/FiBSearchInstance.cpp
--- a/Engine/Source/Editor/Kismet/Private/FiBSearchInstance.cpp
+++ b/Engine/Source/Editor/Kismet/Private/FiBSearchInstance.cpp
@@ -147,6 +147,17 @@ FSearchResult FFiBSearchInstance::StartSearchQuery(const FString& InSearchString
 	PendingSearchables.Add(InImaginaryBlueprintRoot);
 	DoSearchQuery(InSearchString);
 
+	return GetSearchResults(InImaginaryBlueprintRoot);
+}
+
+void FFiBSearchInstance::MakeSearchQuery(const FString& InSearchString, TSharedPtr<FImaginaryBlueprint> InImaginaryBlueprintRoot)
+{
+	PendingSearchables.Add(InImaginaryBlueprintRoot);
+	DoSearchQuery(InSearchString);
+}
+
+FSearchResult FFiBSearchInstance::GetSearchResults(TSharedPtr<FImaginaryBlueprint> InImaginaryBlueprintRoot)
+{
 	FSearchResult CachedSearchResult = FImaginaryFiBData::CreateSearchTree(nullptr, InImaginaryBlueprintRoot, MatchesSearchQuery, MatchingSearchComponents);
 	return MatchesSearchQuery.Num() > 0? CachedSearchResult : nullptr;
 }
@@ -198,6 +209,18 @@ bool FFiBSearchInstance::DoSearchQuery(const FString& InSearchString, bool bInCo
 	return MatchesSearchQuery.Num() > 0;
 }
 
+void FFiBSearchInstance::CreateFilteredResultsListFromTree(ESearchQueryFilter InSearchQueryFilter, TArray< TSharedPtr<FImaginaryFiBData> >& InOutValidSearchResults)
+{
+	for (const FImaginaryFiBData* ImaginaryDataPtr : MatchesSearchQuery)
+	{
+		if (!ImaginaryDataPtr->IsCategory() && ImaginaryDataPtr->IsCompatibleWithFilter(InSearchQueryFilter))
+		{
+			FImaginaryFiBData* NonCastImaginaryDataPtr = const_cast<FImaginaryFiBData*>(ImaginaryDataPtr);
+			InOutValidSearchResults.Add(NonCastImaginaryDataPtr->AsShared());
+		}
+	}
+}
+
 void FFiBSearchInstance::BuildFunctionTargets(TSharedPtr<FImaginaryFiBData> InRootData, ESearchQueryFilter InSearchQueryFilter, TArray< TWeakPtr< FImaginaryFiBData > >& OutTargetPendingSearchables)
 {
 	for (TSharedPtr<FImaginaryFiBData> ChildData : InRootData->GetAllParsedChildData())
diff --git a/Engine/Source/Editor/Kismet/Private/FiBSearchInstance.h b/Engine/Source/Editor/Kismet/Private/FiBSearchInstance.h
--- a/Engine/Source/Editor/Kismet/Private/FiBSearchInstance.h
+++ b/Engine/Source/Editor/Kismet/Private/FiBSearchInstance.h
@@ -61,6 +61,14 @@ public:
 	 * @return								Search result shared pointer, can be used for display in the search results window
 	 */
 	FSearchResult StartSearchQuery(const FString& InSearchString, TSharedPtr<FImaginaryBlueprint> InImaginaryBlueprintRoot);
+
+	/**
+	 * Starts a search query given a string and an imaginary Blueprint
+	 *
+	 * @param InSearchString				The string to search using
+	 * @param InImaginaryBlueprintRoot		The imaginary Blueprint to search through
+	 */
+	void MakeSearchQuery(const FString& InSearchString, TSharedPtr<FImaginaryBlueprint> InImaginaryBlueprintRoot);
 	
 	/**
 	 * Runs a search query on any pending imaginary data
@@ -106,6 +114,12 @@ public:
 	 * @return						TRUE if the function was successful at finding valid results
 	 */
 	bool OnFilterDefaultFunction(const FTextFilterString& InFunctionName, const FTextFilterString& InFunctionParams);
+
+	/** Builds a list of search results in their imaginary data form, filtered by an object type */
+	void CreateFilteredResultsListFromTree(ESearchQueryFilter InSearchQueryFilter, TArray< TSharedPtr<FImaginaryFiBData> >& InOutValidSearchResults);
+
+	/** Helper function to return search results given an imaginary Blueprint root */
+	FSearchResult GetSearchResults(TSharedPtr<FImaginaryBlueprint> InImaginaryBlueprintRoot);
 public:
 	/** Current item being searched in the Imaginary Blueprint */
 	TWeakPtr< FImaginaryFiBData > CurrentSearchable;
diff --git a/Engine/Source/Editor/Kismet/Private/FindInBlueprintManager.cpp b/Engine/Source/Editor/Kismet/Private/FindInBlueprintManager.cpp
--- a/Engine/Source/Editor/Kismet/Private/FindInBlueprintManager.cpp
+++ b/Engine/Source/Editor/Kismet/Private/FindInBlueprintManager.cpp
@@ -63,10 +63,25 @@ const FString FFiBMD::FiBSearchableHiddenExplicitMD = TEXT("BlueprintSearchableH
 
 ////////////////////////////////////
 // FStreamSearch
-FStreamSearch::FStreamSearch(const FString& InSearchValue )
+FStreamSearch::FStreamSearch(const FString& InSearchValue)
 	: SearchValue(InSearchValue)
 	, bThreadCompleted(false)
 	, StopTaskCounter(0)
+	, MinimiumVersionRequirement(EFiBVersion::FIB_VER_LATEST)
+	, BlueprintCountBelowVersion(0)
+	, ImaginaryDataFilter(ESearchQueryFilter::AllFilter)
+{
+	// Add on a Guid to the thread name to ensure the thread is uniquely named.
+	Thread = FRunnableThread::Create( this, *FString::Printf(TEXT("FStreamSearch%s"), *FGuid::NewGuid().ToString()), 0, TPri_BelowNormal );
+}
+
+FStreamSearch::FStreamSearch(const FString& InSearchValue, ESearchQueryFilter InImaginaryDataFilter, EFiBVersion InMinimiumVersionRequirement)
+	: SearchValue(InSearchValue)
+	, bThreadCompleted(false)
+	, StopTaskCounter(0)
+	, MinimiumVersionRequirement(InMinimiumVersionRequirement)
+	, BlueprintCountBelowVersion(0)
+	, ImaginaryDataFilter(InImaginaryDataFilter)
 {
 	// Add on a Guid to the thread name to ensure the thread is uniquely named.
 	Thread = FRunnableThread::Create( this, *FString::Printf(TEXT("FStreamSearch%s"), *FGuid::NewGuid().ToString()), 0, TPri_BelowNormal );
@@ -90,12 +105,26 @@ uint32 FStreamSearch::Run()
 	FSearchData QueryResult;
 	while (FFindInBlueprintSearchManager::Get().ContinueSearchQuery(this, QueryResult))
 	{
-		// This work cannot be spun off to a use the TaskGraph system for spreading the work over multiple threads/cores as the stack limit is too small for the work being done.
-		if (QueryResult.Value.Len() > 0)
+		if (QueryResult.ImaginaryBlueprint.IsValid())
 		{
-			TSharedPtr< FImaginaryBlueprint> ImaginaryBlueprint(new FImaginaryBlueprint(FPaths::GetBaseFilename(QueryResult.BlueprintPath.ToString()), QueryResult.BlueprintPath.ToString(), QueryResult.ParentClass, QueryResult.Interfaces, QueryResult.Value));
+			// If the Blueprint is below the version, add it to a list. The search will still proceed on this Blueprint
+			if (QueryResult.Version < MinimiumVersionRequirement)
+			{
+				++BlueprintCountBelowVersion;
+			}
+
 			TSharedPtr< FFiBSearchInstance > SearchInstance(new FFiBSearchInstance);
-			FSearchResult SearchResult = SearchInstance->StartSearchQuery(*SearchValue, ImaginaryBlueprint);
+			FSearchResult SearchResult;
+			if (ImaginaryDataFilter != ESearchQueryFilter::AllFilter)
+			{
+				SearchInstance->MakeSearchQuery(*SearchValue, QueryResult.ImaginaryBlueprint);
+				SearchInstance->CreateFilteredResultsListFromTree(ImaginaryDataFilter, FilteredImaginaryResults);
+				SearchResult = SearchInstance->GetSearchResults(QueryResult.ImaginaryBlueprint);
+			}
+			else
+			{
+				SearchResult = SearchInstance->StartSearchQuery(*SearchValue, QueryResult.ImaginaryBlueprint);
+			}
 
 			// If there are children, add the item to the search results
 			if(SearchResult.IsValid() && SearchResult->Children.Num() != 0)
@@ -156,6 +185,11 @@ float FStreamSearch::GetPercentComplete() const
 	return FFindInBlueprintSearchManager::Get().GetPercentComplete(this);
 }
 
+void FStreamSearch::GetFilteredImaginaryResults(TArray<TSharedPtr<class FImaginaryFiBData>>& OutFilteredImaginaryResults)
+{
+	OutFilteredImaginaryResults = MoveTemp(FilteredImaginaryResults);
+}
+
 /** Temporarily forces all nodes and pins to use non-friendly names, forces all schema to have nodes clear their cached values so they will re-cache, and then reverts at the end */
 struct FTemporarilyUseFriendlyNodeTitles
 {
@@ -907,13 +941,14 @@ class FCacheAllBlueprintsTickableObject
 
 public:
 
-	FCacheAllBlueprintsTickableObject(TArray<FName>& InUncachedBlueprints, bool bInCheckOutAndSave)
+	FCacheAllBlueprintsTickableObject(TArray<FName> InUncachedBlueprints, bool bInCheckOutAndSave, FSimpleDelegate InOnFinished = FSimpleDelegate())
 		: TickCacheIndex(0)
 		, UncachedBlueprints(InUncachedBlueprints)
 		, bIsStarted(false)
 		, bIsCancelled(false)
 		, bRecursionGuard(false)
 		, bCheckOutAndSave(bInCheckOutAndSave)
+		, OnFinished(InOnFinished)
 	{
 		// Start the Blueprint indexing 'progress' notification
 		FNotificationInfo Info( LOCTEXT("BlueprintIndexMessage", "Indexing Blueprints...") );
@@ -1006,9 +1041,6 @@ public:
 			FAssetRegistryModule* AssetRegistryModule = &FModuleManager::LoadModuleChecked<FAssetRegistryModule>(TEXT("AssetRegistry"));
 			FAssetData AssetData = AssetRegistryModule->Get().GetAssetByObjectPath(UncachedBlueprints[TickCacheIndex]);
 
-			// Cache whether the Blueprint is already loaded
-			bool bIsLoaded = AssetData.IsAssetLoaded();
-
 			bool bIsWorldAsset = AssetData.GetClass() == UWorld::StaticClass();
 
 			// Construct a full package filename with path so we can query the read only status and save to disk
@@ -1081,6 +1113,10 @@ public:
 				ProgressNotification.Pin()->SetText(LOCTEXT("BlueprintIndexComplete", "Finished indexing Blueprints!"));
 
 				FFindInBlueprintSearchManager::Get().FinishedCachingBlueprints(TickCacheIndex, FailedToCacheList);
+
+				// We have actually finished, use the OnFinished callback.
+				OnFinished.ExecuteIfBound();
+
 				return EActiveTimerReturnType::Stop;
 			}
 			else
@@ -1118,6 +1154,9 @@ private:
 
 	/** If TRUE, Blueprints will be checked out and resaved after being loaded */
 	bool bCheckOutAndSave;
+
+	/** Callback for when caching is finished */
+	FSimpleDelegate OnFinished;
 };
 
 FFindInBlueprintSearchManager& FFindInBlueprintSearchManager::Get()
@@ -1437,7 +1476,7 @@ FString FFindInBlueprintSearchManager::GatherBlueprintSearchMetadata(const UBlue
 	Writer->WriteObjectEnd();
 	Writer->Close();
 
-	int32 Version = EFiBVersion::FIB_VER_ALL;
+	int32 Version = EFiBVersion::FIB_VER_LATEST;
 	SearchMetaData = FiBSerializationHelpers::Serialize(Version, false) + Writer->GetSerializedLookupTable() + SearchMetaData;
 	return SearchMetaData;
 }
@@ -1532,6 +1571,13 @@ bool FFindInBlueprintSearchManager::ContinueSearchQuery(const FStreamSearch* InS
 			}
 			else
 			{
+				// If there is FiB data, parse it into an ImaginaryBlueprint
+				if (SearchArray[SearchIdx].Value.Len() > 0)
+				{
+					SearchArray[SearchIdx].ImaginaryBlueprint = MakeShareable(new FImaginaryBlueprint(FPaths::GetBaseFilename(SearchArray[SearchIdx].BlueprintPath.ToString()), SearchArray[SearchIdx].BlueprintPath.ToString(), SearchArray[SearchIdx].ParentClass, SearchArray[SearchIdx].Interfaces, SearchArray[SearchIdx].Value));
+					SearchArray[SearchIdx].Value.Empty();
+				}
+
  				OutSearchData = SearchArray[SearchIdx++];
 				return true;
 			}
@@ -1757,12 +1803,12 @@ FString FFindInBlueprintSearchManager::ConvertFTextToHexString(FText InValue)
 	return BytesToHex(SerializedData.GetData(), SerializedData.Num());
 }
 
-void FFindInBlueprintSearchManager::OnCacheAllUncachedBlueprints(bool bInSourceControlActive, bool bCheckoutAndSave)
+void FFindInBlueprintSearchManager::OnCacheAllUncachedBlueprints(bool bInSourceControlActive, bool bInCheckoutAndSave)
 {
 	// Multiple threads can be adding to this at the same time
 	FScopeLock ScopeLock(&SafeModifyCacheCriticalSection);
 
-	if(bInSourceControlActive && bCheckoutAndSave)
+	if(bInSourceControlActive && bInCheckoutAndSave)
 	{
 		TArray<FString> UncachedBlueprintStrings;
 		UncachedBlueprintStrings.Reserve(UncachedBlueprints.Num());
@@ -1777,15 +1823,47 @@ void FFindInBlueprintSearchManager::OnCacheAllUncachedBlueprints(bool bInSourceC
 	CachingObject->Start();
 }
 
-void FFindInBlueprintSearchManager::CacheAllUncachedBlueprints(TWeakPtr< SFindInBlueprints > InSourceWidget, FWidgetActiveTimerDelegate& OutActiveTimerDelegate)
+void FFindInBlueprintSearchManager::CacheAllUncachedBlueprints(TWeakPtr< SFindInBlueprints > InSourceWidget, FWidgetActiveTimerDelegate& InOutActiveTimerDelegate, FSimpleDelegate InOnFinished/* = FSimpleDelegate()*/, EFiBVersion InMinimiumVersionRequirement/* = EFiBVersion::FIB_VER_LATEST*/)
 {
 	// Do not start another caching process if one is in progress
 	if(!IsCacheInProgress())
 	{
+		TArray<FName> BlueprintsToUpdate;
+		// Add any out-of-date Blueprints to the list
+		for (FSearchData SearchData : SearchArray)
+		{
+			if ((SearchData.Value.Len() != 0 || SearchData.ImaginaryBlueprint.IsValid()) && SearchData.Version < InMinimiumVersionRequirement)
+			{
+				BlueprintsToUpdate.Add(SearchData.BlueprintPath);
+			}
+		}
+
 		FText DialogTitle = LOCTEXT("ConfirmIndexAll_Title", "Indexing All");
 		FFormatNamedArguments Args;
-		Args.Add(TEXT("PackageCount"), UncachedBlueprints.Num());
-		const EAppReturnType::Type ReturnValue = FMessageDialog::Open(EAppMsgType::YesNoCancel, FText::Format(LOCTEXT("CacheAllConfirmationMessage", "This process can take a long time and the editor may become unresponsive; there are {PackageCount} Blueprints to load.\n\nWould you like to checkout, load, and save all Blueprints to make this indexing permanent? Otherwise, all Blueprints will still be loaded but you will be required to re-index the next time you start the editor!"), Args), &DialogTitle);
+		Args.Add(TEXT("PackageCount"), UncachedBlueprints.Num() + BlueprintsToUpdate.Num());
+
+		FText DialogDisplayText;
+		
+		if (UncachedBlueprints.Num() && BlueprintsToUpdate.Num())
+		{
+			Args.Add(TEXT("PackageCount"), UncachedBlueprints.Num() + BlueprintsToUpdate.Num());
+			Args.Add(TEXT("UnindexedCount"), UncachedBlueprints.Num());
+			Args.Add(TEXT("OutOfDateCount"), BlueprintsToUpdate.Num());
+			DialogDisplayText = FText::Format(LOCTEXT("CacheAllConfirmationMessage", "This process can take a long time and the editor may become unresponsive; there are {PackageCount} ({UnindexedCount} Unindexed/{OutOfDateCount} Out-of-Date) Blueprints to load. \
+																					\n\nWould you like to checkout, load, and save all Blueprints to make this indexing permanent? Otherwise, all Blueprints will still be loaded but you will be required to re-index the next time you start the editor!"), Args);
+		}
+		else if (UncachedBlueprints.Num() && BlueprintsToUpdate.Num() == 0)
+		{
+			DialogDisplayText = FText::Format(LOCTEXT("CacheAllConfirmationMessage", "This process can take a long time and the editor may become unresponsive; there are {PackageCount} unindexed Blueprints to load. \
+																					 \n\nWould you like to checkout, load, and save all Blueprints to make this indexing permanent? Otherwise, all Blueprints will still be loaded but you will be required to re-index the next time you start the editor!"), Args);
+		}
+		else if (UncachedBlueprints.Num() == 0 && BlueprintsToUpdate.Num())
+		{
+			DialogDisplayText = FText::Format(LOCTEXT("CacheAllConfirmationMessage", "This process can take a long time and the editor may become unresponsive; there are {PackageCount} out-of-date Blueprints to load. \
+																					 \n\nWould you like to checkout, load, and save all Blueprints to make this indexing permanent? Otherwise, all Blueprints will still be loaded but you will be required to re-index the next time you start the editor!"), Args);
+		}
+
+		const EAppReturnType::Type ReturnValue = FMessageDialog::Open(EAppMsgType::YesNoCancel, DialogDisplayText, &DialogTitle);
 
 		// If Yes is chosen, checkout and save all Blueprints, if No is chosen, only load all Blueprints
 		if (ReturnValue != EAppReturnType::Cancel)
@@ -1793,10 +1871,14 @@ void FFindInBlueprintSearchManager::CacheAllUncachedBlueprints(TWeakPtr< SFindIn
 			// Add all failed to cache Blueprints to the UncachedBlueprints list and resubmit for caching
 			UncachedBlueprints.Append(FailedToCachePaths);
 			FailedToCachePaths.Empty();
+			
+			TArray<FName> TempUncachedBlueprints;
+			TempUncachedBlueprints.Append(UncachedBlueprints);
+			TempUncachedBlueprints.Append(BlueprintsToUpdate);
 
 			const bool bCheckoutAndSave = ReturnValue == EAppReturnType::Yes;
-			CachingObject = new FCacheAllBlueprintsTickableObject(UncachedBlueprints, bCheckoutAndSave);
-			OutActiveTimerDelegate.BindRaw(CachingObject, &FCacheAllBlueprintsTickableObject::Tick);
+			CachingObject = new FCacheAllBlueprintsTickableObject(MoveTemp(TempUncachedBlueprints), bCheckoutAndSave, InOnFinished);
+			InOutActiveTimerDelegate.BindRaw(CachingObject, &FCacheAllBlueprintsTickableObject::Tick);
 
 			const bool bIsSourceControlEnabled = ISourceControlModule::Get().IsEnabled();
 			if(!bIsSourceControlEnabled && bCheckoutAndSave)
diff --git a/Engine/Source/Editor/Kismet/Private/FindInBlueprints.cpp b/Engine/Source/Editor/Kismet/Private/FindInBlueprints.cpp
--- a/Engine/Source/Editor/Kismet/Private/FindInBlueprints.cpp
+++ b/Engine/Source/Editor/Kismet/Private/FindInBlueprints.cpp
@@ -123,6 +123,11 @@ FReply FFindInBlueprintsResult::OnClick()
 	return FReply::Handled();
 }
 
+UObject* FFindInBlueprintsResult::GetObject(UBlueprint* InBlueprint) const
+{
+	return GetParentBlueprint();
+}
+
 FText FFindInBlueprintsResult::GetCategory() const
 {
 	return FText::GetEmpty();
@@ -288,6 +293,11 @@ void FFindInBlueprintsGraphNode::FinalizeSearchData()
 	}
 }
 
+UObject* FFindInBlueprintsGraphNode::GetObject(UBlueprint* InBlueprint) const
+{
+	return FBlueprintEditorUtils::GetNodeByGUID(InBlueprint, NodeGuid);
+}
+
 //////////////////////////////////////////////////////////
 // FFindInBlueprintsPin
 
@@ -510,9 +520,14 @@ FText FFindInBlueprintsGraph::GetCategory() const
 
 void SFindInBlueprints::Construct( const FArguments& InArgs, TSharedPtr<FBlueprintEditor> InBlueprintEditor)
 {
+	OutOfDateWithLastSearchBPCount = 0;
+	LastSearchedFiBVersion = EFiBVersion::FIB_VER_LATEST;
 	BlueprintEditorPtr = InBlueprintEditor;
 
-	RegisterCommands();
+	if (InArgs._bIsSearchWindow)
+	{
+		RegisterCommands();
+	}
 
 	bIsInFindWithinBlueprintMode = true;
 	
@@ -530,6 +545,7 @@ void SFindInBlueprints::Construct( const FArguments& InArgs, TSharedPtr<FBluepri
 					.HintText(LOCTEXT("BlueprintSearchHint", "Enter function or event name to find references..."))
 					.OnTextChanged(this, &SFindInBlueprints::OnSearchTextChanged)
 					.OnTextCommitted(this, &SFindInBlueprints::OnSearchTextCommitted)
+					.Visibility(InArgs._bHideSearchBar? EVisibility::Collapsed : EVisibility::Visible)
 				]
 				+SHorizontalBox::Slot()
 				.Padding(2.f, 0.f)
@@ -538,6 +554,7 @@ void SFindInBlueprints::Construct( const FArguments& InArgs, TSharedPtr<FBluepri
 					SNew(SCheckBox)
 					.OnCheckStateChanged(this, &SFindInBlueprints::OnFindModeChanged)
 					.IsChecked(this, &SFindInBlueprints::OnGetFindModeChecked)
+					.Visibility(InArgs._bHideSearchBar? EVisibility::Collapsed : EVisibility::Visible)
 					[
 						SNew(STextBlock)
 						.Text(LOCTEXT("BlueprintSearchModeChange", "Find In Current Blueprint Only"))
@@ -830,7 +847,18 @@ EActiveTimerReturnType SFindInBlueprints::UpdateSearchResults( double InCurrentT
 			ConditionallyAddCacheBar();
 
 			StreamSearch->EnsureCompletion();
+
+			TArray<TSharedPtr<class FImaginaryFiBData>> ImaginaryResults;
+			if (OnSearchComplete.IsBound())
+			{
+				// Pull out the filtered imaginary results if there is a callback to pass them to
+				StreamSearch->GetFilteredImaginaryResults(ImaginaryResults);
+			}
+			OutOfDateWithLastSearchBPCount = StreamSearch->GetOutOfDateCount();
+
 			StreamSearch.Reset();
+
+			OnSearchComplete.ExecuteIfBound(ImaginaryResults);
 		}
 	}
 
@@ -863,7 +891,7 @@ void SFindInBlueprints::FocusForUse(bool bSetFindWithinBlueprint, FString NewSea
 	if (!NewSearchTerms.IsEmpty())
 	{
 		SearchTextField->SetText(FText::FromString(NewSearchTerms));
-		InitiateSearch();
+		MakeSearchQuery(SearchValue, bIsInFindWithinBlueprintMode);
 
 		// Select the first result
 		if (bSelectFirstResult && ItemsFound.Num())
@@ -882,31 +910,11 @@ void SFindInBlueprints::FocusForUse(bool bSetFindWithinBlueprint, FString NewSea
 	}
 }
 
-void SFindInBlueprints::OnSearchTextChanged( const FText& Text)
-{
-	SearchValue = Text.ToString();
-}
-
-void SFindInBlueprints::OnSearchTextCommitted( const FText& Text, ETextCommit::Type CommitType )
-{
-	if (CommitType == ETextCommit::OnEnter)
-	{
-		InitiateSearch();
-	}
-}
-
-void SFindInBlueprints::OnFindModeChanged(ECheckBoxState CheckState)
+void SFindInBlueprints::MakeSearchQuery(FString InSearchString, bool bInIsFindWithinBlueprint, enum ESearchQueryFilter InSearchFilterForImaginaryDataReturn/* = ESearchQueryFilter::AllFilter*/, EFiBVersion InMinimiumVersionRequirement/* = EFiBVersion::FIB_VER_LATEST*/, FOnSearchComplete InOnSearchComplete/* = FOnSearchComplete()*/)
 {
-	bIsInFindWithinBlueprintMode = CheckState == ECheckBoxState::Checked;
-}
+	SearchTextField->SetText(FText::FromString(InSearchString));
+	LastSearchedFiBVersion = InMinimiumVersionRequirement;
 
-ECheckBoxState SFindInBlueprints::OnGetFindModeChecked() const
-{
-	return bIsInFindWithinBlueprintMode ? ECheckBoxState::Checked : ECheckBoxState::Unchecked;
-}
-
-void SFindInBlueprints::InitiateSearch()
-{
 	if(ItemsFound.Num())
 	{
 		// Reset the scroll to the top
@@ -915,19 +923,20 @@ void SFindInBlueprints::InitiateSearch()
 
 	ItemsFound.Empty();
 
-	if (SearchValue.Len() > 0)
+	if (InSearchString.Len() > 0)
 	{
 		OnRemoveCacheBar();
 
 		TreeView->RequestTreeRefresh();
+		HighlightText = FText::FromString( InSearchString );
 
-		HighlightText = FText::FromString( SearchValue );
-		if (bIsInFindWithinBlueprintMode)
+		if (bInIsFindWithinBlueprint)
 		{
 			if(StreamSearch.IsValid() && !StreamSearch->IsComplete())
 			{
 				StreamSearch->Stop();
 				StreamSearch->EnsureCompletion();
+				OutOfDateWithLastSearchBPCount = StreamSearch->GetOutOfDateCount();
 				StreamSearch.Reset();
 			}
 
@@ -971,11 +980,34 @@ void SFindInBlueprints::InitiateSearch()
 		}
 		else
 		{
-			LaunchStreamThread(SearchValue);
+			LaunchStreamThread(InSearchString, InSearchFilterForImaginaryDataReturn, InMinimiumVersionRequirement, InOnSearchComplete);
 		}
 	}
 }
 
+void SFindInBlueprints::OnSearchTextChanged( const FText& Text)
+{
+	SearchValue = Text.ToString();
+}
+
+void SFindInBlueprints::OnSearchTextCommitted( const FText& Text, ETextCommit::Type CommitType )
+{
+	if (CommitType == ETextCommit::OnEnter)
+	{
+		MakeSearchQuery(SearchValue, bIsInFindWithinBlueprintMode);
+	}
+}
+
+void SFindInBlueprints::OnFindModeChanged(ECheckBoxState CheckState)
+{
+	bIsInFindWithinBlueprintMode = CheckState == ECheckBoxState::Checked;
+}
+
+ECheckBoxState SFindInBlueprints::OnGetFindModeChecked() const
+{
+	return bIsInFindWithinBlueprintMode ? ECheckBoxState::Checked : ECheckBoxState::Unchecked;
+}
+
 void SFindInBlueprints::LaunchStreamThread(const FString& InSearchValue)
 {
 	if(StreamSearch.IsValid() && !StreamSearch->IsComplete())
@@ -990,6 +1022,24 @@ void SFindInBlueprints::LaunchStreamThread(const FString& InSearchValue)
 	}
 
 	StreamSearch = MakeShareable(new FStreamSearch(InSearchValue));
+	OnSearchComplete = FOnSearchComplete();
+}
+
+void SFindInBlueprints::LaunchStreamThread(const FString& InSearchValue, enum ESearchQueryFilter InSearchFilterForRawDataReturn, EFiBVersion InMinimiumVersionRequirement, FOnSearchComplete InOnSearchComplete)
+{
+	if(StreamSearch.IsValid() && !StreamSearch->IsComplete())
+	{
+		StreamSearch->Stop();
+		StreamSearch->EnsureCompletion();
+	}
+	else
+	{
+		// If the stream search wasn't already running, register the active timer
+		RegisterActiveTimer( 0.f, FWidgetActiveTimerDelegate::CreateSP( this, &SFindInBlueprints::UpdateSearchResults ) );
+	}
+
+	StreamSearch = MakeShareable(new FStreamSearch(InSearchValue, InSearchFilterForRawDataReturn, InMinimiumVersionRequirement));
+	OnSearchComplete = InOnSearchComplete;
 }
 
 TSharedRef<ITableRow> SFindInBlueprints::OnGenerateRow( FSearchResult InItem, const TSharedRef<STableViewBase>& OwnerTable )
@@ -1094,13 +1144,23 @@ EVisibility SFindInBlueprints::GetSearchbarVisiblity() const
 	return StreamSearch.IsValid()? EVisibility::Visible : EVisibility::Collapsed;
 }
 
+void SFindInBlueprints::CacheAllBlueprints(FSimpleDelegate InOnFinished/* = FSimpleDelegate()*/, EFiBVersion InMinimiumVersionRequirement/* = EFiBVersion::FIB_VER_LATEST*/)
+{
+	OnCacheAllBlueprints(InOnFinished, InMinimiumVersionRequirement);
+}
+
 FReply SFindInBlueprints::OnCacheAllBlueprints()
+{
+	return OnCacheAllBlueprints(FSimpleDelegate(), LastSearchedFiBVersion);
+}
+
+FReply SFindInBlueprints::OnCacheAllBlueprints(FSimpleDelegate InOnFinished/* = FSimpleDelegate()*/, EFiBVersion InMinimiumVersionRequirement/* = EFiBVersion::FIB_VER_LATEST*/)
 {
 	if(!FFindInBlueprintSearchManager::Get().IsCacheInProgress())
 	{
 		// Request from the SearchManager a delegate to use for ticking the cache system.
 		FWidgetActiveTimerDelegate WidgetActiveTimer;
-		FFindInBlueprintSearchManager::Get().CacheAllUncachedBlueprints(SharedThis(this), WidgetActiveTimer);
+		FFindInBlueprintSearchManager::Get().CacheAllUncachedBlueprints(SharedThis(this), WidgetActiveTimer, InOnFinished, InMinimiumVersionRequirement);
 		RegisterActiveTimer(0.f, WidgetActiveTimer);
 	}
 
@@ -1178,7 +1238,7 @@ FText SFindInBlueprints::GetUncachedBlueprintWarningText() const
 
 	// The number of unindexed Blueprints is the total of those that failed to cache and those that haven't been attempted yet.
 	FFormatNamedArguments Args;
-	Args.Add(TEXT("Count"), FindInBlueprintManager.GetNumberUncachedBlueprints() + FailedToCacheCount);
+	Args.Add(TEXT("Count"), FindInBlueprintManager.GetNumberUncachedBlueprints() + FailedToCacheCount + OutOfDateWithLastSearchBPCount);
 
 	FText ReturnDisplayText;
 	if(IsCacheInProgress())
diff --git a/Engine/Source/Editor/Kismet/Private/ImaginaryBlueprintData.cpp b/Engine/Source/Editor/Kismet/Private/ImaginaryBlueprintData.cpp
--- a/Engine/Source/Editor/Kismet/Private/ImaginaryBlueprintData.cpp
+++ b/Engine/Source/Editor/Kismet/Private/ImaginaryBlueprintData.cpp
@@ -33,17 +33,22 @@ bool FComponentUniqueDisplay::operator==(const FComponentUniqueDisplay& Other)
 ///////////////////////
 // FImaginaryFiBData
 
-FImaginaryFiBData::FImaginaryFiBData() : LookupTablePtr(nullptr)
+FCriticalSection FImaginaryFiBData::ParseChildDataCriticalSection;
+
+FImaginaryFiBData::FImaginaryFiBData(TWeakPtr<FImaginaryFiBData> InOuter) 
+	: LookupTablePtr(nullptr)
+	, Outer(InOuter)
 {
 }
 
-FImaginaryFiBData::FImaginaryFiBData(TSharedPtr< FJsonObject > InUnparsedJsonObject, TMap<int32, FText>* InLookupTablePtr)
+FImaginaryFiBData::FImaginaryFiBData(TWeakPtr<FImaginaryFiBData> InOuter, TSharedPtr< FJsonObject > InUnparsedJsonObject, TMap<int32, FText>* InLookupTablePtr)
 	: UnparsedJsonObject(InUnparsedJsonObject)
 	, LookupTablePtr(InLookupTablePtr)
+	, Outer(InOuter)
 {
 }
 
-FSearchResult FImaginaryFiBData::CreateSearchResult(FSearchResult InParent)
+FSearchResult FImaginaryFiBData::CreateSearchResult(FSearchResult InParent) const
 {
 	FSearchResult ReturnSearchResult = CreateSearchResult_Internal(InParent);
 
@@ -109,8 +114,10 @@ bool FImaginaryFiBData::CanCallFilter(ESearchQueryFilter InSearchQueryFilter) co
 	return InSearchQueryFilter == ESearchQueryFilter::AllFilter;
 }
 
-void FImaginaryFiBData::ParseAllChildData(ESearchableValueStatus InSearchabilityOverride/* = ESearchableValueStatus::Searchable*/)
+void FImaginaryFiBData::ParseAllChildData_Internal(ESearchableValueStatus InSearchabilityOverride/* = ESearchableValueStatus::Searchable*/)
 {
+	FScopeLock ScopeLock(&FImaginaryFiBData::ParseChildDataCriticalSection);
+
 	if (UnparsedJsonObject.IsValid())
 	{
 		if (InSearchabilityOverride == ESearchableValueStatus::Searchable)
@@ -118,8 +125,8 @@ void FImaginaryFiBData::ParseAllChildData(ESearchableValueStatus InSearchability
 			const TSharedPtr< FJsonObject >* MetaDataField;
 			if (UnparsedJsonObject->TryGetObjectField(FFindInBlueprintSearchTags::FiBMetaDataTag.ToString(), MetaDataField))
 			{
-				TSharedPtr<FFiBMetaData> MetaDataFiBInfo = MakeShareable(new FFiBMetaData(*MetaDataField, LookupTablePtr));
-				MetaDataFiBInfo->ParseAllChildData();
+				TSharedPtr<FFiBMetaData> MetaDataFiBInfo = MakeShareable(new FFiBMetaData(AsShared(), *MetaDataField, LookupTablePtr));
+				MetaDataFiBInfo->ParseAllChildData_Internal();
 
 				if (MetaDataFiBInfo->IsHidden() && MetaDataFiBInfo->IsExplicit())
 				{
@@ -152,6 +159,12 @@ void FImaginaryFiBData::ParseAllChildData(ESearchableValueStatus InSearchability
 	UnparsedJsonObject.Reset();
 }
 
+void FImaginaryFiBData::ParseAllChildData(ESearchableValueStatus InSearchabilityOverride/* = ESearchableValueStatus::Searchable*/)
+{
+	FScopeLock ScopeLock(&FImaginaryFiBData::ParseChildDataCriticalSection);
+	ParseAllChildData_Internal(InSearchabilityOverride);
+}
+
 void FImaginaryFiBData::ParseJsonValue(FText InKey, FText InDisplayKey, TSharedPtr< FJsonValue > InJsonValue, bool bIsInArray/*=false*/, ESearchableValueStatus InSearchabilityOverride/* = ESearchableValueStatus::Searchable*/)
 {
 	ESearchableValueStatus SearchabilityStatus = (InSearchabilityOverride == ESearchableValueStatus::Searchable)? GetSearchabilityStatus(InKey.ToString()) : InSearchabilityOverride;
@@ -165,7 +178,7 @@ void FImaginaryFiBData::ParseJsonValue(FText InKey, FText InDisplayKey, TSharedP
 	}
 	else if( InJsonValue->Type == EJson::Array)
 	{
-		TSharedPtr< FCategorySectionHelper > ArrayCategory = MakeShareable(new FCategorySectionHelper(LookupTablePtr, InKey, true));
+		TSharedPtr< FCategorySectionHelper > ArrayCategory = MakeShareable(new FCategorySectionHelper(AsShared(), LookupTablePtr, InKey, true));
 		ParsedChildData.Add(ArrayCategory);
 		TArray<TSharedPtr< FJsonValue > > ArrayList = InJsonValue->AsArray();
 		for( int32 ArrayIdx = 0; ArrayIdx < ArrayList.Num(); ++ArrayIdx)
@@ -176,8 +189,8 @@ void FImaginaryFiBData::ParseJsonValue(FText InKey, FText InDisplayKey, TSharedP
 	}
 	else if (InJsonValue->Type == EJson::Object)
 	{
-		TSharedPtr< FCategorySectionHelper > SubObjectCategory = MakeShareable(new FCategorySectionHelper(InJsonValue->AsObject(), LookupTablePtr, InDisplayKey, bIsInArray));
-		SubObjectCategory->ParseAllChildData(SearchabilityStatus);
+		TSharedPtr< FCategorySectionHelper > SubObjectCategory = MakeShareable(new FCategorySectionHelper(AsShared(), InJsonValue->AsObject(), LookupTablePtr, InDisplayKey, bIsInArray));
+		SubObjectCategory->ParseAllChildData_Internal(SearchabilityStatus);
 		ParsedChildData.Add(SubObjectCategory);
 	}
 	else
@@ -263,11 +276,16 @@ bool FImaginaryFiBData::TestComplexExpression(const FName& InKey, const FTextFil
 	return bMatchesSearchQuery;
 }
 
+UObject* FImaginaryFiBData::GetObject(UBlueprint* InBlueprint) const
+{
+	return CreateSearchResult(nullptr)->GetObject(InBlueprint);
+}
+
 ///////////////////////////
 // FFiBMetaData
 
-FFiBMetaData::FFiBMetaData(TSharedPtr< FJsonObject > InUnparsedJsonObject, TMap<int32, FText>* InLookupTablePtr)
-	: FImaginaryFiBData(InUnparsedJsonObject, InLookupTablePtr)
+FFiBMetaData::FFiBMetaData(TWeakPtr<FImaginaryFiBData> InOuter, TSharedPtr< FJsonObject > InUnparsedJsonObject, TMap<int32, FText>* InLookupTablePtr)
+	: FImaginaryFiBData(InOuter, InUnparsedJsonObject, InLookupTablePtr)
 {
 }
 
@@ -292,22 +310,22 @@ bool FFiBMetaData::TrySpecialHandleJsonValue(FString InKey, TSharedPtr< FJsonVal
 ///////////////////////////
 // FCategorySectionHelper
 
-FCategorySectionHelper::FCategorySectionHelper(TMap<int32, FText>* InLookupTablePtr, FText InCategoryName, bool bInTagAndValueCategory)
-	: FImaginaryFiBData(nullptr, InLookupTablePtr)
+FCategorySectionHelper::FCategorySectionHelper(TWeakPtr<FImaginaryFiBData> InOuter, TMap<int32, FText>* InLookupTablePtr, FText InCategoryName, bool bInTagAndValueCategory)
+	: FImaginaryFiBData(InOuter, nullptr, InLookupTablePtr)
 	, CategoryName(InCategoryName)
 	, bIsTagAndValue(bInTagAndValueCategory)
 {
 }
 
-FCategorySectionHelper::FCategorySectionHelper(TSharedPtr< FJsonObject > InUnparsedJsonObject, TMap<int32, FText>* InLookupTablePtr, FText InCategoryName, bool bInTagAndValueCategory)
-	: FImaginaryFiBData(InUnparsedJsonObject, InLookupTablePtr)
+FCategorySectionHelper::FCategorySectionHelper(TWeakPtr<FImaginaryFiBData> InOuter, TSharedPtr< FJsonObject > InUnparsedJsonObject, TMap<int32, FText>* InLookupTablePtr, FText InCategoryName, bool bInTagAndValueCategory)
+	: FImaginaryFiBData(InOuter, InUnparsedJsonObject, InLookupTablePtr)
 	, CategoryName(InCategoryName)
 	, bIsTagAndValue(bInTagAndValueCategory)
 {
 }
 
-FCategorySectionHelper::FCategorySectionHelper(TSharedPtr< FJsonObject > InUnparsedJsonObject, TMap<int32, FText>* InLookupTablePtr, FText InCategoryName, bool bInTagAndValueCategory, FCategorySectionHelperCallback InSpecialHandlingCallback)
-	: FImaginaryFiBData(InUnparsedJsonObject, InLookupTablePtr)
+FCategorySectionHelper::FCategorySectionHelper(TWeakPtr<FImaginaryFiBData> InOuter, TSharedPtr< FJsonObject > InUnparsedJsonObject, TMap<int32, FText>* InLookupTablePtr, FText InCategoryName, bool bInTagAndValueCategory, FCategorySectionHelperCallback InSpecialHandlingCallback)
+	: FImaginaryFiBData(InOuter, InUnparsedJsonObject, InLookupTablePtr)
 	, SpecialHandlingCallback(InSpecialHandlingCallback)
 	, CategoryName(InCategoryName)
 	, bIsTagAndValue(bInTagAndValueCategory)
@@ -325,7 +343,7 @@ FSearchResult FCategorySectionHelper::CreateSearchResult_Internal(FSearchResult
 	return FSearchResult(new FFindInBlueprintsResult(CategoryName, InParent));
 }
 
-void FCategorySectionHelper::ParseAllChildData(ESearchableValueStatus InSearchabilityOverride/* = ESearchableValueStatus::Searchable*/)
+void FCategorySectionHelper::ParseAllChildData_Internal(ESearchableValueStatus InSearchabilityOverride/* = ESearchableValueStatus::Searchable*/)
 {
 	if (UnparsedJsonObject.IsValid() && SpecialHandlingCallback.IsBound())
 	{
@@ -334,7 +352,7 @@ void FCategorySectionHelper::ParseAllChildData(ESearchableValueStatus InSearchab
 	}
 	else
 	{
-		FImaginaryFiBData::ParseAllChildData(InSearchabilityOverride);
+		FImaginaryFiBData::ParseAllChildData_Internal(InSearchabilityOverride);
 	}
 }
 
@@ -342,7 +360,8 @@ void FCategorySectionHelper::ParseAllChildData(ESearchableValueStatus InSearchab
 // FImaginaryBlueprint
 
 FImaginaryBlueprint::FImaginaryBlueprint(FString InBlueprintName, FString InBlueprintPath, FString InBlueprintParentClass, TArray<FString>& InInterfaces, FString InUnparsedStringData)
-	: BlueprintPath(InBlueprintPath)
+	: FImaginaryFiBData(nullptr)
+	, BlueprintPath(InBlueprintPath)
 	, UnparsedStringData(InUnparsedStringData)
 {
 	ParseToJson();
@@ -366,6 +385,11 @@ FSearchResult FImaginaryBlueprint::CreateSearchResult_Internal(FSearchResult InP
 	return FSearchResult(new FFindInBlueprintsResult(ParsedTagsAndValues.Find(FFindInBlueprintSearchTags::FiB_Path.ToString())->GetDisplayText(LookupTable)));
 }
 
+UBlueprint* FImaginaryBlueprint::GetBlueprint() const
+{
+	return Cast<UBlueprint>(GetObject(nullptr));
+}
+
 bool FImaginaryBlueprint::IsCompatibleWithFilter(ESearchQueryFilter InSearchQueryFilter) const
 {
 	return InSearchQueryFilter == ESearchQueryFilter::AllFilter || InSearchQueryFilter == ESearchQueryFilter::BlueprintFilter;
@@ -400,7 +424,7 @@ bool FImaginaryBlueprint::TrySpecialHandleJsonValue(FString InKey, TSharedPtr< F
 		TArray<TSharedPtr< FJsonValue > > PropertyList = InJsonValue->AsArray();
 		for( TSharedPtr< FJsonValue > PropertyValue : PropertyList )
 		{
-			ParsedChildData.Add(MakeShareable(new FImaginaryProperty(PropertyValue->AsObject(), &LookupTable)));
+			ParsedChildData.Add(MakeShareable(new FImaginaryProperty(AsShared(), PropertyValue->AsObject(), &LookupTable)));
 		}
 		bResult = true;
 	}
@@ -429,7 +453,7 @@ bool FImaginaryBlueprint::TrySpecialHandleJsonValue(FString InKey, TSharedPtr< F
 		TArray<TSharedPtr< FJsonValue > > ComponentList = InJsonValue->AsArray();
 		TSharedPtr< FJsonObject > ComponentsWrapperObject(new FJsonObject);
 		ComponentsWrapperObject->Values.Add(FFindInBlueprintSearchTags::FiB_Components.ToString(), InJsonValue);
-		ParsedChildData.Add(MakeShareable(new FCategorySectionHelper(ComponentsWrapperObject, &LookupTable, FFindInBlueprintSearchTags::FiB_Components, false, FCategorySectionHelper::FCategorySectionHelperCallback::CreateRaw(this, &FImaginaryBlueprint::ParseComponents))));
+		ParsedChildData.Add(MakeShareable(new FCategorySectionHelper(AsShared(), ComponentsWrapperObject, &LookupTable, FFindInBlueprintSearchTags::FiB_Components, false, FCategorySectionHelper::FCategorySectionHelperCallback::CreateRaw(this, &FImaginaryBlueprint::ParseComponents))));
 		bResult = true;
 	}
 
@@ -445,7 +469,7 @@ void FImaginaryBlueprint::ParseGraph( TSharedPtr< FJsonValue > InJsonValue, FStr
 	TArray<TSharedPtr< FJsonValue > > GraphList = InJsonValue->AsArray();
 	for( TSharedPtr< FJsonValue > GraphValue : GraphList )
 	{
-		ParsedChildData.Add(MakeShareable(new FImaginaryGraph(GraphValue->AsObject(), &LookupTable, InGraphType)));
+		ParsedChildData.Add(MakeShareable(new FImaginaryGraph(AsShared(), GraphValue->AsObject(), &LookupTable, InGraphType)));
 	}
 }
 
@@ -455,15 +479,15 @@ void FImaginaryBlueprint::ParseComponents(TSharedPtr< FJsonObject > InJsonObject
 	TArray<TSharedPtr< FJsonValue > > ComponentList = InJsonObject->GetArrayField(FFindInBlueprintSearchTags::FiB_Components.ToString());
 	for( TSharedPtr< FJsonValue > ComponentValue : ComponentList )
 	{
-		OutParsedChildData.Add(MakeShareable(new FImaginaryComponent(ComponentValue->AsObject(), &LookupTable)));
+		OutParsedChildData.Add(MakeShareable(new FImaginaryComponent(AsShared(), ComponentValue->AsObject(), &LookupTable)));
 	}
 }
 
 //////////////////////////
 // FImaginaryGraph
 
-FImaginaryGraph::FImaginaryGraph(TSharedPtr< FJsonObject > InUnparsedJsonObject, TMap<int32, FText>* InLookupTablePtr, EGraphType InGraphType)
-	: FImaginaryFiBData(InUnparsedJsonObject, InLookupTablePtr)
+FImaginaryGraph::FImaginaryGraph(TWeakPtr<FImaginaryFiBData> InOuter, TSharedPtr< FJsonObject > InUnparsedJsonObject, TMap<int32, FText>* InLookupTablePtr, EGraphType InGraphType)
+	: FImaginaryFiBData(InOuter, InUnparsedJsonObject, InLookupTablePtr)
 	, GraphType(InGraphType)
 {
 }
@@ -512,7 +536,7 @@ bool FImaginaryGraph::TrySpecialHandleJsonValue(FString InKey, TSharedPtr< FJson
 
 		for( TSharedPtr< FJsonValue > NodeValue : NodeList )
 		{
-			ParsedChildData.Add(MakeShareable(new FImaginaryGraphNode(NodeValue->AsObject(), LookupTablePtr)));
+			ParsedChildData.Add(MakeShareable(new FImaginaryGraphNode(AsShared(), NodeValue->AsObject(), LookupTablePtr)));
 		}
 		return true;
 	}
@@ -522,7 +546,7 @@ bool FImaginaryGraph::TrySpecialHandleJsonValue(FString InKey, TSharedPtr< FJson
 		TArray<TSharedPtr< FJsonValue > > PropertyList = InJsonValue->AsArray();
 		for( TSharedPtr< FJsonValue > PropertyValue : PropertyList )
 		{
-			ParsedChildData.Add(MakeShareable(new FImaginaryProperty(PropertyValue->AsObject(), LookupTablePtr)));
+			ParsedChildData.Add(MakeShareable(new FImaginaryProperty(AsShared(), PropertyValue->AsObject(), LookupTablePtr)));
 		}
 		return true;
 	}
@@ -532,8 +556,8 @@ bool FImaginaryGraph::TrySpecialHandleJsonValue(FString InKey, TSharedPtr< FJson
 //////////////////////////////////////
 // FImaginaryGraphNode
 
-FImaginaryGraphNode::FImaginaryGraphNode(TSharedPtr< FJsonObject > InUnparsedJsonObject, TMap<int32, FText>* InLookupTablePtr)
-	: FImaginaryFiBData(InUnparsedJsonObject, InLookupTablePtr)
+FImaginaryGraphNode::FImaginaryGraphNode(TWeakPtr<FImaginaryFiBData> InOuter, TSharedPtr< FJsonObject > InUnparsedJsonObject, TMap<int32, FText>* InLookupTablePtr)
+	: FImaginaryFiBData(InOuter, InUnparsedJsonObject, InLookupTablePtr)
 {
 }
 
@@ -585,7 +609,7 @@ bool FImaginaryGraphNode::TrySpecialHandleJsonValue(FString InKey, TSharedPtr< F
 
 		for( TSharedPtr< FJsonValue > Pin : PinsList )
 		{
-			ParsedChildData.Add(MakeShareable(new FImaginaryPin(Pin->AsObject(), LookupTablePtr, SchemaName)));
+			ParsedChildData.Add(MakeShareable(new FImaginaryPin(AsShared(), Pin->AsObject(), LookupTablePtr, SchemaName)));
 		}
 		return true;
 	}
@@ -598,7 +622,7 @@ bool FImaginaryGraphNode::TrySpecialHandleJsonValue(FString InKey, TSharedPtr< F
 	return false;
 }
 
-void FImaginaryGraphNode::ParseAllChildData(ESearchableValueStatus InSearchabilityOverride/* = ESearchableValueStatus::Searchable*/)
+void FImaginaryGraphNode::ParseAllChildData_Internal(ESearchableValueStatus InSearchabilityOverride/* = ESearchableValueStatus::Searchable*/)
 {
 	if (UnparsedJsonObject.IsValid())
 	{
@@ -610,15 +634,15 @@ void FImaginaryGraphNode::ParseAllChildData(ESearchableValueStatus InSearchabili
 			SchemaName = FindInBlueprintsHelpers::AsFText(SchemaNameValue, *LookupTablePtr).ToString();
 		}
 
-		FImaginaryFiBData::ParseAllChildData(InSearchabilityOverride);
+		FImaginaryFiBData::ParseAllChildData_Internal(InSearchabilityOverride);
 	}
 }
 
 ///////////////////////////////////////////
 // FImaginaryProperty
 
-FImaginaryProperty::FImaginaryProperty(TSharedPtr< FJsonObject > InUnparsedJsonObject, TMap<int32, FText>* InLookupTablePtr)
-	: FImaginaryFiBData(InUnparsedJsonObject, InLookupTablePtr)
+FImaginaryProperty::FImaginaryProperty(TWeakPtr<FImaginaryFiBData> InOuter, TSharedPtr< FJsonObject > InUnparsedJsonObject, TMap<int32, FText>* InLookupTablePtr)
+	: FImaginaryFiBData(InOuter, InUnparsedJsonObject, InLookupTablePtr)
 {
 }
 
@@ -659,8 +683,8 @@ ESearchableValueStatus FImaginaryProperty::GetSearchabilityStatus(FString InKey)
 //////////////////////////////
 // FImaginaryComponent
 
-FImaginaryComponent::FImaginaryComponent(TSharedPtr< FJsonObject > InUnparsedJsonObject, TMap<int32, FText>* InLookupTablePtr)
-	: FImaginaryProperty(InUnparsedJsonObject, InLookupTablePtr)
+FImaginaryComponent::FImaginaryComponent(TWeakPtr<FImaginaryFiBData> InOuter, TSharedPtr< FJsonObject > InUnparsedJsonObject, TMap<int32, FText>* InLookupTablePtr)
+	: FImaginaryProperty(InOuter, InUnparsedJsonObject, InLookupTablePtr)
 {
 }
 
@@ -672,8 +696,8 @@ bool FImaginaryComponent::IsCompatibleWithFilter(ESearchQueryFilter InSearchQuer
 //////////////////////////////
 // FImaginaryPin
 
-FImaginaryPin::FImaginaryPin(TSharedPtr< FJsonObject > InUnparsedJsonObject, TMap<int32, FText>* InLookupTablePtr, FString InSchemaName)
-	: FImaginaryFiBData(InUnparsedJsonObject, InLookupTablePtr)
+FImaginaryPin::FImaginaryPin(TWeakPtr<FImaginaryFiBData> InOuter, TSharedPtr< FJsonObject > InUnparsedJsonObject, TMap<int32, FText>* InLookupTablePtr, FString InSchemaName)
+	: FImaginaryFiBData(InOuter, InUnparsedJsonObject, InLookupTablePtr)
 	, SchemaName(InSchemaName)
 {
 }
diff --git a/Engine/Source/Editor/Kismet/Private/SMyBlueprint.cpp b/Engine/Source/Editor/Kismet/Private/SMyBlueprint.cpp
--- a/Engine/Source/Editor/Kismet/Private/SMyBlueprint.cpp
+++ b/Engine/Source/Editor/Kismet/Private/SMyBlueprint.cpp
@@ -43,6 +43,7 @@
 #include "Engine/TimelineTemplate.h"
 
 #include "BlueprintEditorSettings.h"
+#include "SReplaceNodeReferences.h"
 
 #define LOCTEXT_NAMESPACE "MyBlueprint"
 
@@ -172,7 +173,14 @@ void SMyBlueprint::Construct(const FArguments& InArgs, TWeakPtr<FBlueprintEditor
 		ToolKitCommandList->MapAction( FGraphEditorCommands::Get().FindReferences,
 			FExecuteAction::CreateSP(this, &SMyBlueprint::OnFindReference),
 			FCanExecuteAction(),
+			FIsActionChecked(),
 			FIsActionButtonVisible::CreateSP(this, &SMyBlueprint::CanFindReference) );
+
+		ToolKitCommandList->MapAction( FGraphEditorCommands::Get().FindAndReplaceReferences,
+			FExecuteAction::CreateSP(this, &SMyBlueprint::OnFindAndReplaceReference),
+			FCanExecuteAction(),
+			FIsActionChecked(),
+			FIsActionButtonVisible::CreateSP(this, &SMyBlueprint::CanFindAndReplaceReference) );
 	
 		ToolKitCommandList->MapAction( FMyBlueprintCommands::Get().DeleteEntry,
 			FExecuteAction::CreateSP(this, &SMyBlueprint::OnDeleteEntry),
@@ -1557,11 +1565,11 @@ void SMyBlueprint::OnActionSelected( const TArray< TSharedPtr<FEdGraphSchemaActi
 		CurrentBlueprint = BlueprintEditor->GetBlueprintObj();
 		CurrentInspector = BlueprintEditor->GetInspector();
 
-		OnActionSelectedHelper(InAction, Blueprint, CurrentInspector.ToSharedRef());
+		OnActionSelectedHelper(InAction, BlueprintEditorPtr, Blueprint, CurrentInspector.ToSharedRef());
 	}
 }
 
-void SMyBlueprint::OnActionSelectedHelper(TSharedPtr<FEdGraphSchemaAction> InAction, UBlueprint* Blueprint, TSharedRef<SKismetInspector> Inspector)
+void SMyBlueprint::OnActionSelectedHelper(TSharedPtr<FEdGraphSchemaAction> InAction, TWeakPtr< FBlueprintEditor > InBlueprintEditor, UBlueprint* Blueprint, TSharedRef<SKismetInspector> Inspector)
 {
 	if (InAction.IsValid())
 	{
@@ -1592,6 +1600,7 @@ void SMyBlueprint::OnActionSelectedHelper(TSharedPtr<FEdGraphSchemaAction> InAct
 			Options.bForceRefresh = true;
 
 			Inspector->ShowDetailsForSingleObject(VarAction->GetProperty(), Options);
+			InBlueprintEditor.Pin()->GetReplaceReferencesWidget()->SetSourceVariable(VarAction->GetProperty());
 		}
 		else if (InAction->GetTypeId() == FEdGraphSchemaAction_K2LocalVar::StaticGetTypeId())
 		{
@@ -1772,7 +1781,7 @@ FEdGraphSchemaAction_K2InputAction* SMyBlueprint::SelectionAsInputAction() const
 
 bool SMyBlueprint::SelectionIsCategory() const
 {
-	return !GraphActionMenu->GetSelectedCategoryName().IsEmpty();
+	return !SelectionHasContextMenu();
 }
 
 bool SMyBlueprint::SelectionHasContextMenu() const
@@ -1849,6 +1858,7 @@ TSharedPtr<SWidget> SMyBlueprint::OnContextMenuOpening()
 			MenuBuilder.AddMenuEntry(FGenericCommands::Get().Rename, NAME_None, LOCTEXT("Rename", "Rename"), LOCTEXT("Rename_Tooltip", "Renames this function or variable from blueprint.") );
 			MenuBuilder.AddMenuEntry(FMyBlueprintCommands::Get().ImplementFunction);
 			MenuBuilder.AddMenuEntry(FGraphEditorCommands::Get().FindReferences);
+			MenuBuilder.AddMenuEntry(FGraphEditorCommands::Get().FindAndReplaceReferences);
 			MenuBuilder.AddMenuEntry(FMyBlueprintCommands::Get().GotoNativeVarDefinition);
 			MenuBuilder.AddMenuEntry(FGenericCommands::Get().Duplicate);
 			MenuBuilder.AddMenuEntry(FMyBlueprintCommands::Get().DeleteEntry);
@@ -2168,6 +2178,21 @@ bool SMyBlueprint::CanFindReference() const
 	return true;
 }
 
+void SMyBlueprint::OnFindAndReplaceReference()
+{
+	BlueprintEditorPtr.Pin()->SummonFindAndReplaceUI();
+}
+
+bool SMyBlueprint::CanFindAndReplaceReference() const
+{
+	if (SelectionAsVar() && GetDefault<UEditorExperimentalSettings>()->bEnableFindAndReplaceReferences)
+	{
+		return true;
+	}
+
+	return false;
+}
+
 void SMyBlueprint::OnDeleteGraph(UEdGraph* InGraph, EEdGraphSchemaAction_K2Graph::Type InGraphType)
 {
 	if (InGraph)
diff --git a/Engine/Source/Editor/Kismet/Private/SMyBlueprint.h b/Engine/Source/Editor/Kismet/Private/SMyBlueprint.h
--- a/Engine/Source/Editor/Kismet/Private/SMyBlueprint.h
+++ b/Engine/Source/Editor/Kismet/Private/SMyBlueprint.h
@@ -151,7 +151,7 @@ private:
 	FReply OnActionDragged(const TArray< TSharedPtr<FEdGraphSchemaAction> >& InActions, const FPointerEvent& MouseEvent);
 	FReply OnCategoryDragged(const FText& InCategory, const FPointerEvent& MouseEvent);
 	void OnActionSelected(const TArray< TSharedPtr<FEdGraphSchemaAction> >& InActions);
-	static void OnActionSelectedHelper(TSharedPtr<FEdGraphSchemaAction> InAction, UBlueprint* Blueprint, TSharedRef<SKismetInspector> Inspector);
+	static void OnActionSelectedHelper(TSharedPtr<FEdGraphSchemaAction> InAction, TWeakPtr< FBlueprintEditor > InBlueprintEditor, UBlueprint* Blueprint, TSharedRef<SKismetInspector> Inspector);
 	void OnGlobalActionSelected(const TArray< TSharedPtr<FEdGraphSchemaAction> >& InActions, ESelectInfo::Type InSelectionType);
 	void OnActionDoubleClicked(const TArray< TSharedPtr<FEdGraphSchemaAction> >& InActions);
 	void ExecuteAction(TSharedPtr<FEdGraphSchemaAction> InAction);
@@ -194,6 +194,8 @@ private:
 	bool CanImplementFunction() const;
 	void OnFindReference();
 	bool CanFindReference() const;
+	void OnFindAndReplaceReference();
+	bool CanFindAndReplaceReference() const;
 	void OnDeleteEntry();
 	bool CanDeleteEntry() const;
 	FReply OnAddNewLocalVariable();
diff --git a/Engine/Source/Editor/Kismet/Private/SReplaceNodeReferences.cpp b/Engine/Source/Editor/Kismet/Private/SReplaceNodeReferences.cpp
new file mode 100644
--- /dev/null
+++ b/Engine/Source/Editor/Kismet/Private/SReplaceNodeReferences.cpp
@@ -0,0 +1,552 @@
+#include "BlueprintEditorPrivatePCH.h"
+
+#include "FindInBlueprints.h"
+#include "SReplaceNodeReferences.h"
+#include "BlueprintEditor.h"
+#include "Editor/UnrealEd/Public/Kismet2/BlueprintEditorUtils.h"
+#include "ImaginaryBlueprintData.h"
+#include "ObjectEditorUtils.h"
+#include "EditorCategoryUtils.h"
+#include "ScopedTransaction.h"
+
+#define LOCTEXT_NAMESPACE "SNodeVariableReferences"
+
+class FTargetCategoryReplaceReferences : public FTargetReplaceReferences
+{
+public:
+	FTargetCategoryReplaceReferences(FText InCategoryTitle)
+		: CategoryTitle(InCategoryTitle)
+	{}
+
+	// FTargetReplaceReferences interface
+	virtual TSharedRef<SWidget> CreateWidget() const override
+	{
+		return 
+			SNew(STextBlock)
+			.Text(CategoryTitle);
+	}
+
+	virtual bool GetMemberReference(FMemberReference& OutVariableReference) const override
+	{
+		return false;
+	}
+
+	virtual FText GetDisplayTitle() const override
+	{
+		return CategoryTitle;
+	}
+
+	virtual bool IsCategory() const override { return true; }
+	// End of FTargetReplaceReferences interface
+
+public:
+	/** Category title to display for this item */
+	FText CategoryTitle;
+};
+
+class FTargetVariableReplaceReferences : public FTargetReplaceReferences
+{
+public:
+	// FTargetReplaceReferences interface
+	virtual TSharedRef<SWidget> CreateWidget() const override
+	{
+		return 
+			SNew(SHorizontalBox)
+			+SHorizontalBox::Slot()
+			.AutoWidth()
+			.VAlign(VAlign_Center)
+			.Padding(2.0f, 0.0f)
+			[
+				SNew(SImage)
+				.Image(GetIcon())
+				.ColorAndOpacity(GetIconColor())
+			]
+
+			+SHorizontalBox::Slot()
+			.VAlign(VAlign_Center)
+			[
+				SNew(STextBlock)
+					.Text(FText::FromString(VariableReference.GetMemberName().ToString()))
+			];
+	}
+
+	virtual bool GetMemberReference(FMemberReference& OutVariableReference) const override
+	{
+		OutVariableReference = VariableReference;
+		return true;
+	}
+
+	virtual FText GetDisplayTitle() const override
+	{
+		return FText::FromName(VariableReference.GetMemberName());
+	}
+
+	virtual const struct FSlateBrush* GetIcon() const override
+	{
+		return FBlueprintEditorUtils::GetIconFromPin(PinType);
+	}
+
+	virtual FSlateColor GetIconColor() const override
+	{
+		UEdGraphSchema_K2 const* K2Schema = GetDefault<UEdGraphSchema_K2>();
+		return K2Schema->GetPinTypeColor(PinType);
+	}
+	// End of FTargetReplaceReferences interface
+
+public:
+	/** Variable reference for this item */
+	FMemberReference VariableReference;
+
+	/** Pin type representing the UProperty of this item */
+	FEdGraphPinType PinType;
+};
+
+void SReplaceNodeReferences::Construct(const FArguments& InArgs, TSharedPtr<class FBlueprintEditor> InBlueprintEditor)
+{
+	BlueprintEditor = InBlueprintEditor;
+	Refresh();
+
+	ChildSlot
+		[
+			SNew(SVerticalBox)
+			+SVerticalBox::Slot()
+				.AutoHeight()
+				.Padding(3.0f, 5.0f)
+			[
+				SNew(SHorizontalBox)
+				+SHorizontalBox::Slot()
+				.AutoWidth()
+				.VAlign(VAlign_Center)
+				[
+					SNew(STextBlock)
+					.Text(LOCTEXT("FindWhat", "Find what:"))
+				]
+
+				+SHorizontalBox::Slot()
+				.VAlign(VAlign_Center)
+				[
+					SNew(STextBlock)
+					.Visibility(this, &SReplaceNodeReferences::GetPickSourceReferenceVisibility)
+					.Text(LOCTEXT("PickSourceVariable", "Pick a source variable from the My Blueprints list!"))
+				]
+
+				+SHorizontalBox::Slot()
+				.AutoWidth()
+				.VAlign(VAlign_Center)
+				.Padding(2.0f, 0.0f)
+				[
+					SNew(SImage)
+					.Visibility(this, &SReplaceNodeReferences::GetSourceReferenceVisibility)
+					.Image(this, &SReplaceNodeReferences::GetSourceReferenceIcon)
+					.ColorAndOpacity(this, &SReplaceNodeReferences::GetSourceReferenceIconColor)
+				]
+
+				+SHorizontalBox::Slot()
+				.VAlign(VAlign_Center)
+				[
+					SNew(STextBlock)
+						.Visibility(this, &SReplaceNodeReferences::GetSourceReferenceVisibility)
+						.Text(this, &SReplaceNodeReferences::GetSourceDisplayText)
+				]
+			]
+
+			+SVerticalBox::Slot()
+				.AutoHeight()
+				[
+					SNew(STextBlock)
+					.Text(LOCTEXT("ReplaceWith", "Replace with:"))
+				]
+
+			+SVerticalBox::Slot()
+				.AutoHeight()
+				.HAlign(HAlign_Left)
+				[
+					SNew(SBox)
+					.MinDesiredWidth(150.0f)
+					[
+						SAssignNew( TargetReferencesComboBox, SComboButton )
+						.OnGetMenuContent(this, &SReplaceNodeReferences::GetMenuContent)
+						.ContentPadding(0)
+						.ToolTipText(this, &SReplaceNodeReferences::GetTargetDisplayText)
+						.HasDownArrow(true)
+						.ButtonContent()
+						[
+							SNew(SHorizontalBox)
+							+SHorizontalBox::Slot()
+							.AutoWidth()
+							.VAlign(VAlign_Center)
+							.Padding(2.0f, 0.0f)
+							[
+								SNew(SImage)
+								.Image(this, &SReplaceNodeReferences::GetTargetIcon)
+								.ColorAndOpacity(this, &SReplaceNodeReferences::GetTargetIconColor)
+							]
+
+							+SHorizontalBox::Slot()
+								.VAlign(VAlign_Center)
+								[
+									SNew(STextBlock)
+									.Text(this, &SReplaceNodeReferences::GetTargetDisplayText)
+								]
+						]
+					]
+				]
+
+			+SVerticalBox::Slot()
+				[
+					SNew(SBox)
+						.MinDesiredHeight(150.0f)
+						[
+							SAssignNew(FindInBlueprints, SFindInBlueprints, InBlueprintEditor)
+								.bIsSearchWindow(false)
+								.bHideSearchBar(true)
+						]
+
+				]
+
+			+SVerticalBox::Slot()
+				.AutoHeight()
+				[
+					SNew(SHorizontalBox)
+					+SHorizontalBox::Slot()
+					.Padding(2.0f)
+					.HAlign(HAlign_Left)
+					.AutoWidth()
+					[
+						SNew(SButton)
+						.Text(FText::FromString(TEXT("Find All")))
+						.OnClicked(this, &SReplaceNodeReferences::OnFindAll)
+					]
+
+					+SHorizontalBox::Slot()
+					.Padding(2.0f)
+					.HAlign(HAlign_Left)
+					.AutoWidth()
+					[
+						SNew(SButton)
+						.Text(FText::FromString(TEXT("Find and Replace All")))
+						.OnClicked(this, &SReplaceNodeReferences::OnFindAndReplaceAll)
+					]
+					
+				]
+		];
+}
+
+void SReplaceNodeReferences::Refresh()
+{
+	SetSourceVariable(nullptr);
+
+	BlueprintVariableList.Empty();
+	TargetClass = BlueprintEditor.Pin()->GetBlueprintObj()->SkeletonGeneratedClass;
+	GatherAllAvailableBlueprintVariables(TargetClass);
+}
+
+void SReplaceNodeReferences::SetSourceVariable(UProperty* InProperty)
+{
+	if (InProperty)
+	{
+		UEdGraphSchema_K2 const* K2Schema = GetDefault<UEdGraphSchema_K2>();
+		K2Schema->ConvertPropertyToPinType(InProperty, SourcePinType);
+
+		SourceProperty = InProperty;
+
+		BlueprintVariableList.Empty();
+		GatherAllAvailableBlueprintVariables(TargetClass);
+
+		if (AvailableTargetReferencesTreeView.IsValid())
+		{
+			AvailableTargetReferencesTreeView->RequestTreeRefresh();
+		}
+	}
+	else
+	{
+		SourceProperty = nullptr;
+	}
+}
+
+SReplaceNodeReferences::~SReplaceNodeReferences()
+{
+
+}
+
+TSharedRef<SWidget>	SReplaceNodeReferences::GetMenuContent()
+{
+	return SAssignNew(AvailableTargetReferencesTreeView, SReplaceReferencesTreeViewType)
+		.ItemHeight(24)
+		.TreeItemsSource( &BlueprintVariableList )
+		.OnSelectionChanged(this, &SReplaceNodeReferences::OnSelectionChanged)
+		.OnGenerateRow( this, &SReplaceNodeReferences::OnGenerateRow )
+		.OnGetChildren( this, &SReplaceNodeReferences::OnGetChildren );
+}
+
+void SReplaceNodeReferences::GatherAllAvailableBlueprintVariables(UClass* InTargetClass)
+{
+	if (InTargetClass == nullptr)
+	{
+		return;
+	}
+	GatherAllAvailableBlueprintVariables(InTargetClass->GetSuperClass());
+
+	TMap<FString, TSharedPtr< FTargetCategoryReplaceReferences > > CategoryMap;
+
+	UObject* PathObject = InTargetClass->ClassGeneratedBy? InTargetClass->ClassGeneratedBy : InTargetClass;
+	TSharedPtr< FTargetCategoryReplaceReferences > BlueprintCategory = MakeShareable(new FTargetCategoryReplaceReferences(FText::FromString(PathObject->GetPathName())));
+	for (TFieldIterator<UProperty> PropertyIt(InTargetClass, EFieldIteratorFlags::ExcludeSuper); PropertyIt; ++PropertyIt)
+	{
+		UProperty* Property = *PropertyIt;
+
+		if (Property == SourceProperty)
+		{
+			continue;
+		}
+		FName PropName = Property->GetFName();
+
+		// Don't show delegate properties, there is special handling for these
+		const bool bMulticastDelegateProp = Property->IsA(UMulticastDelegateProperty::StaticClass());
+		const bool bDelegateProp = (Property->IsA(UDelegateProperty::StaticClass()) || bMulticastDelegateProp);
+		const bool bShouldShowAsVar = (!Property->HasAnyPropertyFlags(CPF_Parm) && Property->HasAllPropertyFlags(CPF_BlueprintVisible)) && !bDelegateProp;
+		const bool bShouldShowAsDelegate = !Property->HasAnyPropertyFlags(CPF_Parm) && bMulticastDelegateProp 
+			&& Property->HasAnyPropertyFlags(CPF_BlueprintAssignable | CPF_BlueprintCallable);
+		UObjectPropertyBase* Obj = Cast<UObjectPropertyBase>(Property);
+		if(!bShouldShowAsVar && !bShouldShowAsDelegate)
+		{
+			continue;
+		}
+
+		const FText PropertyTooltip = Property->GetToolTipText();
+		const FName PropertyName = Property->GetFName();
+		const FText PropertyDesc = FText::FromName(PropertyName);
+
+		FText CategoryName = FObjectEditorUtils::GetCategoryText(Property);
+		FText PropertyCategory = FObjectEditorUtils::GetCategoryText(Property);
+		const FString UserCategoryName = FEditorCategoryUtils::GetCategoryDisplayString( PropertyCategory.ToString() );
+
+		UEdGraphSchema_K2 const* K2Schema = GetDefault<UEdGraphSchema_K2>();
+		if (CategoryName.EqualTo(FText::FromString(PathObject->GetName())) || CategoryName.EqualTo(K2Schema->VR_DefaultCategory))
+		{
+			CategoryName = FText::GetEmpty();		// default, so place in 'non' category
+			PropertyCategory = FText::GetEmpty();
+		}
+
+		if (bShouldShowAsVar)
+		{
+			const bool bComponentProperty = Obj && Obj->PropertyClass ? Obj->PropertyClass->IsChildOf<UActorComponent>() : false;
+
+			// By default components go into the variable section under the component category unless a custom category is specified.
+			if ( bComponentProperty && CategoryName.IsEmpty() )
+			{
+				PropertyCategory = LOCTEXT("Components", "Components");
+			}
+
+			TSharedPtr< FTargetCategoryReplaceReferences > CategoryReference = BlueprintCategory;
+			if (!CategoryName.IsEmpty())
+			{
+				TSharedPtr< FTargetCategoryReplaceReferences >* CategoryReferencePtr = CategoryMap.Find(PropertyCategory.ToString());
+				if (CategoryReferencePtr)
+				{
+					CategoryReference = *CategoryReferencePtr;
+				}
+				else
+				{
+					CategoryReference = MakeShareable(new FTargetCategoryReplaceReferences(PropertyCategory));
+					CategoryMap.Add(PropertyCategory.ToString(), CategoryReference);
+				}
+			}
+
+			TSharedPtr< FTargetVariableReplaceReferences > VariableItem = MakeShareable(new FTargetVariableReplaceReferences);
+			VariableItem->VariableReference.SetFromField<UProperty>(Property, true);
+
+			FEdGraphPinType Type;
+			K2Schema->ConvertPropertyToPinType(Property, VariableItem->PinType);
+
+			if (VariableItem->PinType == SourcePinType)
+			{
+				CategoryReference->Children.Add(VariableItem);
+
+				// If this is the first Child, add the category to the main Blueprint category
+				if (CategoryReference->Children.Num() == 1)
+				{
+					BlueprintCategory->Children.Add(CategoryReference);
+				}
+			}
+		}
+	}
+
+	if (BlueprintCategory->Children.Num())
+	{
+		BlueprintVariableList.Add(BlueprintCategory);
+		// Sort markers
+		struct FCompareCategoryTitles
+		{
+			FORCEINLINE bool operator()(const TSharedPtr<FTargetReplaceReferences> A, const TSharedPtr<FTargetReplaceReferences> B) const
+			{
+				if (A->Children.Num() > 0 && B->Children.Num() == 0)
+				{
+					return true;
+				}
+				else if (A->Children.Num() == 0 && B->Children.Num() > 0)
+				{
+					return false;
+				}
+				return A->GetDisplayTitle().CompareTo(B->GetDisplayTitle()) < 0;
+			}
+		};
+		BlueprintCategory->Children.Sort(FCompareCategoryTitles());
+	}
+}
+
+TSharedRef<ITableRow> SReplaceNodeReferences::OnGenerateRow(FTreeViewItem InItem, const TSharedRef<STableViewBase>& OwnerTable)
+{
+	return SNew( STableRow< TSharedPtr<FFindInBlueprintsResult> >, OwnerTable )
+		[
+			InItem->CreateWidget()
+		];
+}
+
+void SReplaceNodeReferences::OnGetChildren( FTreeViewItem InItem, TArray< FTreeViewItem >& OutChildren )
+{
+	OutChildren += InItem->Children;
+}
+
+FReply SReplaceNodeReferences::OnFindAll()
+{
+	OnSubmitSearchQuery(false);
+	return FReply::Handled();
+}
+
+FReply SReplaceNodeReferences::OnFindAndReplaceAll()
+{
+	if (SelectedTargetReferenceItem.IsValid())
+	{
+		FindInBlueprints->CacheAllBlueprints(FSimpleDelegate::CreateSP(this, &SReplaceNodeReferences::OnSubmitSearchQuery, true), EFiBVersion::FIB_VER_VARIABLE_REFERENCE);
+	}
+	return FReply::Handled();
+}
+
+void SReplaceNodeReferences::OnSubmitSearchQuery(bool bFindAndReplace)
+{
+	FString SearchTerm;
+
+	FMemberReference SourceVariableReference;
+	SourceVariableReference.SetFromField<UProperty>(SourceProperty, true);
+	SearchTerm = SourceVariableReference.GetReferenceSearchString(SourceProperty->GetOwnerClass());
+
+	FOnSearchComplete OnSearchComplete;
+	if (bFindAndReplace)
+	{
+		OnSearchComplete = FOnSearchComplete::CreateSP(this, &SReplaceNodeReferences::FindAllReplacementsComplete);
+	}
+
+	FindInBlueprints->MakeSearchQuery(SearchTerm, false, ESearchQueryFilter::NodesFilter, EFiBVersion::FIB_VER_VARIABLE_REFERENCE, OnSearchComplete);
+}
+
+void SReplaceNodeReferences::FindAllReplacementsComplete(TArray<TSharedPtr<class FImaginaryFiBData>>& InRawDataList)
+{
+	if (SelectedTargetReferenceItem.IsValid())
+	{
+		FMemberReference VariableReference;
+		if (SelectedTargetReferenceItem->GetMemberReference(VariableReference))
+		{
+			FText TransactionTitle = FText::Format(LOCTEXT("FindReplaceAllTransaction", "{0} replaced with {1}"), FText::FromString(SourceProperty->GetName()), FText::FromName(VariableReference.GetMemberName()));
+			const FScopedTransaction Transaction( TransactionTitle );
+			BlueprintEditor.Pin()->GetBlueprintObj()->Modify();
+
+			TArray< UBlueprint* > BlueprintsModified;
+			for (TSharedPtr<class FImaginaryFiBData> ImaginaryData : InRawDataList)
+			{
+				BlueprintsModified.AddUnique(ImaginaryData->GetBlueprint());
+				UObject* Node = ImaginaryData->GetObject(ImaginaryData->GetBlueprint());
+				Node = Node;
+				UK2Node_Variable* VariableNode = Cast<UK2Node_Variable>(Node);
+				if (ensure(VariableNode))
+				{
+					VariableNode->Modify();
+					if (VariableNode->VariableReference.IsLocalScope() || VariableNode->VariableReference.IsSelfContext())
+					{
+						VariableNode->VariableReference = VariableReference;
+					}
+					else
+					{
+						UBlueprint* Blueprint = BlueprintEditor.Pin()->GetBlueprintObj();
+						VariableNode->VariableReference.SetFromField<UProperty>(VariableReference.ResolveMember<UProperty>(Blueprint), Blueprint->GeneratedClass);
+					}
+					VariableNode->ReconstructNode();
+				}
+			}
+
+			for (UBlueprint* Blueprint : BlueprintsModified)
+			{
+				FBlueprintEditorUtils::MarkBlueprintAsStructurallyModified(Blueprint);
+				FFindInBlueprintSearchManager::Get().AddOrUpdateBlueprintSearchMetadata(Blueprint);
+			}
+			FindInBlueprints->MakeSearchQuery(VariableReference.GetReferenceSearchString(SourceProperty->GetOwnerClass()), false, ESearchQueryFilter::NodesFilter, EFiBVersion::FIB_VER_VARIABLE_REFERENCE, FOnSearchComplete());
+		}
+	}
+}
+
+void SReplaceNodeReferences::OnSelectionChanged(FTreeViewItem Selection, ESelectInfo::Type SelectInfo)
+{
+	// When the user is navigating, do not act upon the selection change
+	if(SelectInfo == ESelectInfo::OnNavigation || (Selection.IsValid() && Selection->IsCategory()))
+	{
+		return;
+	}
+	
+	SelectedTargetReferenceItem = Selection;
+	TargetReferencesComboBox->SetIsOpen(false);
+}
+
+FText SReplaceNodeReferences::GetSourceDisplayText() const
+{
+	if (SourceProperty == nullptr)
+	{
+		return FText::FromString(TEXT("Hello World!"));
+	}
+	return FText::FromString(SourceProperty->GetName());
+}
+
+const FSlateBrush* SReplaceNodeReferences::GetSourceReferenceIcon() const
+{
+	return FBlueprintEditorUtils::GetIconFromPin(SourcePinType);
+}
+
+FSlateColor SReplaceNodeReferences::GetSourceReferenceIconColor() const
+{
+	UEdGraphSchema_K2 const* K2Schema = GetDefault<UEdGraphSchema_K2>();
+	return K2Schema->GetPinTypeColor(SourcePinType);
+}
+
+FText SReplaceNodeReferences::GetTargetDisplayText() const
+{
+	FText ReturnText = LOCTEXT("UnselectedTargetReference", "Please select a target reference!");
+
+	if (SelectedTargetReferenceItem.IsValid())
+	{
+		ReturnText = SelectedTargetReferenceItem->GetDisplayTitle();
+	}
+	return ReturnText;
+}
+
+const FSlateBrush* SReplaceNodeReferences::GetTargetIcon() const
+{
+	const FSlateBrush* ReturnBrush = nullptr;
+
+	if (SelectedTargetReferenceItem.IsValid())
+	{
+		ReturnBrush = SelectedTargetReferenceItem->GetIcon();
+	}
+	return ReturnBrush;
+}
+
+FSlateColor SReplaceNodeReferences::GetTargetIconColor() const
+{
+	FSlateColor ReturnColor = FLinearColor::White;
+
+	if (SelectedTargetReferenceItem.IsValid())
+	{
+		ReturnColor = SelectedTargetReferenceItem->GetIconColor();
+	}
+	return ReturnColor;
+}
+
+#undef LOCTEXT_NAMESPACE
\ No newline at end of file
diff --git a/Engine/Source/Editor/Kismet/Private/SReplaceNodeReferences.h b/Engine/Source/Editor/Kismet/Private/SReplaceNodeReferences.h
new file mode 100644
--- /dev/null
+++ b/Engine/Source/Editor/Kismet/Private/SReplaceNodeReferences.h
@@ -0,0 +1,148 @@
+#pragma once
+
+class FTargetReplaceReferences
+{
+public:
+	/** Returns a generated widget to represent this target item reference */
+	virtual TSharedRef<SWidget> CreateWidget() const = 0;
+
+	/**
+	 * Retrieves the MemberReference represented by this item, if any
+	 *
+	 * @param OutVariableReference		The output MemberReference represented by this item
+	 * @return							TRUE if successful at having a MemberReference
+	 */
+	virtual bool GetMemberReference(FMemberReference& OutVariableReference) const = 0;
+
+	/** Returns the display title for this item */
+	virtual FText GetDisplayTitle() const = 0;
+
+	/** TRUE if this item is a category and nothing else */
+	virtual bool IsCategory() const { return false; }
+
+	/** Returns the Icon representing this reference */
+	virtual const struct FSlateBrush* GetIcon() const { return nullptr; }
+
+	/** Returns the Icon Color of this reference */
+	virtual FSlateColor GetIconColor() const { return FLinearColor::White; }
+
+public:
+	/** Children members to sub-list in the tree */
+	TArray< TSharedPtr<FTargetReplaceReferences> > Children;
+};
+
+class SReplaceNodeReferences : public SCompoundWidget
+{
+protected:
+	typedef TSharedPtr< FTargetReplaceReferences > FTreeViewItem;
+	typedef STreeView<FTreeViewItem>  SReplaceReferencesTreeViewType;
+public:
+	SLATE_BEGIN_ARGS( SReplaceNodeReferences ) {}
+	SLATE_END_ARGS()
+
+		SReplaceNodeReferences()
+		: SourceProperty(nullptr)
+	{
+
+	}
+
+	void Construct(const FArguments& InArgs, TSharedPtr<class FBlueprintEditor> InBlueprintEditor);
+	~SReplaceNodeReferences();
+
+	/** Forces a refresh on this widget when things in the Blueprint Editor have changed */
+	void Refresh();
+
+	/** Sets a source variable reference to replace */
+	void SetSourceVariable(UProperty* InProperty);
+
+protected:
+
+	/** Callback for determining if the SourceReference is visible */
+	EVisibility GetSourceReferenceVisibility() const
+	{
+		return SourceProperty == nullptr? EVisibility::Collapsed : EVisibility::Visible;
+	}
+
+	/** Callback for determining if the message to inform the user a SourceReference is needed is visible */
+	EVisibility GetPickSourceReferenceVisibility() const
+	{
+		return SourceProperty == nullptr? EVisibility::Visible : EVisibility::Collapsed;
+	}
+
+	/* Called when a new row is being generated */
+	TSharedRef<ITableRow> OnGenerateRow(FTreeViewItem InItem, const TSharedRef<STableViewBase>& OwnerTable);
+
+	/* Get the children of a row */
+	void OnGetChildren( FTreeViewItem InItem, TArray< FTreeViewItem >& OutChildren );
+
+	/* Returns the menu content for the Target Reference drop down section of the combo button */
+	TSharedRef<SWidget>	GetMenuContent();
+
+	/** Callback when selection in the combo button has changed */
+	void OnSelectionChanged(FTreeViewItem Selection, ESelectInfo::Type SelectInfo);
+
+	/** 
+	 * Submits a search query and potentially does a mass replace on results
+	 *
+	 * @param bInFindAndReplace		TRUE if it's a Find-and-Replace action, FALSE if it should do a search of items that will be replaced
+	 */
+	void OnSubmitSearchQuery(bool bFindAndReplace);
+
+	/** Callback for "Find All" button */
+	FReply OnFindAll();
+
+	/** Callback for "Find and Replace All" button */
+	FReply OnFindAndReplaceAll();
+
+	/** Callback when the search for "Find and Replace All" is complete so that the replacements can begin */
+	void FindAllReplacementsComplete(TArray<TSharedPtr<class FImaginaryFiBData>>& InRawDataList);
+
+	/** Recursively gathers all available Blueprint Variable references to replace with */
+	void GatherAllAvailableBlueprintVariables(UClass* InTargetClass);
+
+	/** Returns the display text for the target reference */
+	FText GetTargetDisplayText() const;
+
+	/** Returns the icon for the target reference */
+	const struct FSlateBrush* GetTargetIcon() const;
+
+	/** Returns the icon color for the target reference */
+	FSlateColor GetTargetIconColor() const;
+
+	/** Returns the display text for the source reference */
+	FText GetSourceDisplayText() const;
+
+	/** Returns the icon for the source reference */
+	const FSlateBrush* GetSourceReferenceIcon() const;
+
+	/** Returns the icon color for the source reference */
+	FSlateColor GetSourceReferenceIconColor() const;
+
+protected:
+	/** Combo box for selecting the target reference */
+	TSharedPtr< SComboButton > TargetReferencesComboBox;
+
+	/** Tree view for display available target references */
+	TSharedPtr< SReplaceReferencesTreeViewType > AvailableTargetReferencesTreeView;
+
+	/** List of items used for the root of the tree */
+	TArray<FTreeViewItem> BlueprintVariableList;
+
+	/** Target SKEL_ class that is being referenced by this window */
+	UClass* TargetClass;
+
+	/** Blueprint editor that owns this window */
+	TWeakPtr<FBlueprintEditor> BlueprintEditor;
+
+	/** Cached SourcePinType for the property the user wants to replace */
+	FEdGraphPinType SourcePinType;
+
+	/** Cached SourceProperty that the user wants to replace */
+	UProperty* SourceProperty;
+
+	/** Find-in-Blueprints window used for making search queries and presenting results to the user */
+	TSharedPtr< SFindInBlueprints > FindInBlueprints;
+
+	/** Currently selected target reference */
+	FTreeViewItem SelectedTargetReferenceItem;
+};
\ No newline at end of file
diff --git a/Engine/Source/Editor/Kismet/Public/BlueprintEditor.h b/Engine/Source/Editor/Kismet/Public/BlueprintEditor.h
--- a/Engine/Source/Editor/Kismet/Public/BlueprintEditor.h
+++ b/Engine/Source/Editor/Kismet/Public/BlueprintEditor.h
@@ -159,6 +159,7 @@ public:
 	virtual void RefreshEditors(ERefreshBlueprintEditorReason::Type Reason = ERefreshBlueprintEditorReason::UnknownReason) override;
 	virtual void JumpToHyperlink(const UObject* ObjectReference, bool bRequestRename = false) override;
 	virtual void SummonSearchUI(bool bSetFindWithinBlueprint, FString NewSearchTerms = FString(), bool bSelectFirstResult = false) override;
+	virtual void SummonFindAndReplaceUI() override;
 	virtual TArray<TSharedPtr<class FSCSEditorTreeNode> > GetSelectedSCSEditorTreeNodes() const override;
 	virtual TSharedPtr<class FSCSEditorTreeNode> FindAndSelectSCSEditorTreeNode(const UActorComponent* InComponent, bool IsCntrlDown) override;
 	virtual int32 GetNumberOfSelectedNodes() const override;
@@ -204,6 +205,7 @@ public:
 	TSharedPtr<class SSCSEditor> GetSCSEditor() const {return SCSEditor;}
 	TSharedPtr<class SSCSEditorViewport> GetSCSViewport() const {return SCSViewport;}
 	TSharedPtr<class SMyBlueprint> GetMyBlueprintWidget() const {return MyBlueprintWidget;}
+	TSharedPtr<class SReplaceNodeReferences> GetReplaceReferencesWidget() const {return ReplaceReferencesWidget;}
 
 	/**
 	 * Provides access to the preview actor.
@@ -1027,6 +1029,9 @@ protected:
 	/** All of this blueprints' functions and variables */
 	TSharedPtr<class SMyBlueprint> MyBlueprintWidget;
 	
+	/** Widget for replacing node references */
+	TSharedPtr<class SReplaceNodeReferences> ReplaceReferencesWidget;
+
 	/** Compiler results log, with the log listing that it reflects */
 	TSharedPtr<class SWidget> CompilerResults;
 	TSharedPtr<class IMessageLogListing> CompilerResultsListing;
diff --git a/Engine/Source/Editor/Kismet/Public/BlueprintEditorModule.h b/Engine/Source/Editor/Kismet/Public/BlueprintEditorModule.h
--- a/Engine/Source/Editor/Kismet/Public/BlueprintEditorModule.h
+++ b/Engine/Source/Editor/Kismet/Public/BlueprintEditorModule.h
@@ -46,6 +46,9 @@ public:
 	/** Invokes the search UI and sets the mode and search terms optionally */
 	virtual void SummonSearchUI(bool bSetFindWithinBlueprint, FString NewSearchTerms = FString(), bool bSelectFirstResult = false) = 0;
 
+	/** Invokes the Find and Replace UI */
+	virtual void SummonFindAndReplaceUI() = 0;
+
 	virtual void RefreshEditors(ERefreshBlueprintEditorReason::Type Reason = ERefreshBlueprintEditorReason::UnknownReason) = 0;
 
 	virtual bool CanPasteNodes() const= 0;
diff --git a/Engine/Source/Editor/Kismet/Public/BlueprintEditorTabs.h b/Engine/Source/Editor/Kismet/Public/BlueprintEditorTabs.h
--- a/Engine/Source/Editor/Kismet/Public/BlueprintEditorTabs.h
+++ b/Engine/Source/Editor/Kismet/Public/BlueprintEditorTabs.h
@@ -14,6 +14,7 @@ struct KISMET_API FBlueprintEditorTabs
 	static const FName ConstructionScriptEditorID;
 	static const FName SCSViewportID;
 	static const FName MyBlueprintID;
+	static const FName ReplaceNodeReferencesID;
 	static const FName UserDefinedStructureID;
 
 	// Document tab identifiers
diff --git a/Engine/Source/Editor/Kismet/Public/FindInBlueprintManager.h b/Engine/Source/Editor/Kismet/Public/FindInBlueprintManager.h
--- a/Engine/Source/Editor/Kismet/Public/FindInBlueprintManager.h
+++ b/Engine/Source/Editor/Kismet/Public/FindInBlueprintManager.h
@@ -98,6 +98,9 @@ struct FSearchData
 	/** Cached to determine if the Blueprint is seen as no longer valid, allows it to be cleared out next save to disk */
 	bool bMarkedForDeletion;
 
+	/** Cached ImaginaryBlueprint data for the searchable content, prevents having to re-parse every search */
+	TSharedPtr< class FImaginaryBlueprint > ImaginaryBlueprint;
+
 	/** Version of the data */
 	int32 Version;
 
@@ -117,7 +120,8 @@ struct FSearchData
 		, ParentClass(MoveTemp(Other.ParentClass))
 		, Interfaces(MoveTemp(Other.Interfaces))
 		, bMarkedForDeletion(Other.bMarkedForDeletion)
-		, Version(0)
+		, Version(Other.Version)
+		, ImaginaryBlueprint(Other.ImaginaryBlueprint)
 	{
 	}
 
@@ -135,6 +139,7 @@ struct FSearchData
 		ParentClass = MoveTemp(RHS.ParentClass);
 		Interfaces = MoveTemp(RHS.Interfaces);
 		Version = RHS.Version;
+		ImaginaryBlueprint = RHS.ImaginaryBlueprint;
 		return *this;
 	}
 
@@ -146,6 +151,7 @@ struct FSearchData
 		, Interfaces(Other.Interfaces)
 		, bMarkedForDeletion(Other.bMarkedForDeletion)
 		, Version(Other.Version)
+		, ImaginaryBlueprint(Other.ImaginaryBlueprint)
 	{
 	}
 
@@ -163,6 +169,7 @@ struct FSearchData
 		ParentClass = RHS.ParentClass;
 		Interfaces = RHS.Interfaces;
 		Version = RHS.Version;
+		ImaginaryBlueprint = RHS.ImaginaryBlueprint;
 		return *this;
 	}
 
@@ -172,6 +179,22 @@ struct FSearchData
 	}
 };
 
+/** Filters are used by functions for searching to decide whether items can call certain functions or match the requirements of a function */
+enum ESearchQueryFilter
+{
+	BlueprintFilter = 0,
+	GraphsFilter,
+	UberGraphsFilter,
+	FunctionsFilter,
+	MacrosFilter,
+	NodesFilter,
+	PinsFilter,
+	PropertiesFilter,
+	VariablesFilter,
+	ComponentsFilter,
+	AllFilter, // Will search all items, when used inside of another filter it will search all sub-items of that filter
+};
+
 /** Used for external gather functions to add Key/Value pairs to be placed into Json */
 struct FSearchTagDataPair
 {
@@ -191,7 +214,7 @@ enum EFiBVersion
 
 	// -----<new versions can be added before this line>-------------------------------------------------
 	FIB_VER_PLUS_ONE,
-	FIB_VER_ALL = FIB_VER_PLUS_ONE - 1 // Always the last version, we want Blueprints to be at latest
+	FIB_VER_LATEST = FIB_VER_PLUS_ONE - 1 // Always the last version, we want Blueprints to be at latest
 };
 
 struct KISMET_API FFiBMD
@@ -212,7 +235,8 @@ class FStreamSearch : public FRunnable
 {
 public:
 	/** Constructor */
-	FStreamSearch(const FString& InSearchValue );
+	FStreamSearch(const FString& InSearchValue);
+	FStreamSearch(const FString& InSearchValue, enum ESearchQueryFilter InImaginaryDataFilter, EFiBVersion InMinimiumVersionRequirement);
 
 	/** Begin FRunnable Interface */
 	virtual bool Init() override;
@@ -240,6 +264,15 @@ public:
 	/** Helper function to query the percent complete this search is */
 	float GetPercentComplete() const;
 
+	/** Returns the Out-of-Date Blueprint count */
+	int32 GetOutOfDateCount() const
+	{
+		return BlueprintCountBelowVersion;
+	}
+
+	/** Returns the FilteredImaginaryResults from the search query, these results have been filtered by the ImaginaryDataFilter. */
+	void GetFilteredImaginaryResults(TArray<TSharedPtr<class FImaginaryFiBData>>& OutFilteredImaginaryResults);
+
 public:
 	/** Thread to run the cleanup FRunnable on */
 	FRunnableThread* Thread;
@@ -258,6 +291,18 @@ public:
 
 	/** > 0 if we've been asked to abort work in progress at the next opportunity */
 	FThreadSafeCounter StopTaskCounter;
+
+	/** When searching, any Blueprint below this version will be considered out-of-date */
+	EFiBVersion MinimiumVersionRequirement;
+
+	/** A going count of all Blueprints below the MinimiumVersionRequirement */
+	int32 BlueprintCountBelowVersion;
+
+	/** Filtered (ImaginaryDataFilter) list of imaginary data results that met the search requirements */
+	TArray<TSharedPtr<class FImaginaryFiBData>> FilteredImaginaryResults;
+
+	/** Filter to limit the FilteredImaginaryResults to */
+	enum ESearchQueryFilter ImaginaryDataFilter;
 };
 
 ////////////////////////////////////
@@ -334,10 +379,20 @@ public:
 	/**
 	 * Starts caching all uncached Blueprints at a rate of 1 per tick
 	 *
-	 * @param InSourceWidget		The source FindInBlueprints widget, this widget will be informed when caching is complete
+	 * @param InSourceWidget				The source FindInBlueprints widget, this widget will be informed when caching is complete
+	 * @param InOutActiveTimerDelegate		Binds an object that ticks every time the Find-in-Blueprints widget ticks to cache all Blueprints
+	 * @param InOnFinished					Callback when caching is finished
+	 * @param InMinimiumVersionRequirement	Minimum version requirement for caching, any Blueprints below this version will be re-indexed
+	 */
+	void CacheAllUncachedBlueprints(TWeakPtr< class SFindInBlueprints > InSourceWidget, FWidgetActiveTimerDelegate& InOutActiveTimerDelegate, FSimpleDelegate InOnFinished = FSimpleDelegate(), EFiBVersion InMinimiumVersionRequirement = EFiBVersion::FIB_VER_LATEST);
+	
+	/**
+	 * Starts the actual caching process
+	 *
+	 * @param bInSourceControlActive		TRUE if source control is active
+	 * @param bInCheckoutAndSave			TRUE if the system should checkout and save all assets that need to be reindexed
 	 */
-	void CacheAllUncachedBlueprints(TWeakPtr< class SFindInBlueprints > InSourceWidgetm, FWidgetActiveTimerDelegate& OutActiveTimerDelegate);
-	void OnCacheAllUncachedBlueprints(bool bInSourceControlActive, bool bCheckoutAndSave);
+	void OnCacheAllUncachedBlueprints(bool bInSourceControlActive, bool bInCheckoutAndSave);
 
 	/** Stops the caching process where it currently is at, the rest can be continued later */
 	void CancelCacheAll(SFindInBlueprints* InFindInBlueprintWidget);
diff --git a/Engine/Source/Editor/Kismet/Public/FindInBlueprints.h b/Engine/Source/Editor/Kismet/Public/FindInBlueprints.h
--- a/Engine/Source/Editor/Kismet/Public/FindInBlueprints.h
+++ b/Engine/Source/Editor/Kismet/Public/FindInBlueprints.h
@@ -2,9 +2,13 @@
 
 #pragma once
 
+#include "FindInBlueprintManager.h"
+
 typedef TSharedPtr<class FFindInBlueprintsResult> FSearchResult;
 typedef STreeView<FSearchResult>  STreeViewType;
 
+DECLARE_DELEGATE_OneParam(FOnSearchComplete, TArray<TSharedPtr<class FImaginaryFiBData>>&);
+
 /** Some utility functions to help with Find-in-Blueprint functionality */
 namespace FindInBlueprintsHelpers
 {
@@ -76,6 +80,9 @@ public:
 	 */
 	virtual void ParseSearchInfo(FText InKey, FText InValue) {};
 
+	/** Returns the Object represented by this search information give the Blueprint it can be found in */
+	virtual UObject* GetObject(UBlueprint* InBlueprint) const;
+
 	/**
 	 * Adds extra search info, anything that does not have a predestined place in the search result. Adds a sub-item to the searches and formats its description so the tag displays
 	 *
@@ -120,6 +127,7 @@ public:
 	virtual void ParseSearchInfo(FText InKey, FText InValue) override;
 	virtual FText GetCategory() const override;
 	virtual void FinalizeSearchData() override;
+	virtual UObject* GetObject(UBlueprint* InBlueprint) const override;
 	/** End FFindInBlueprintsResult Interface */
 
 private:
@@ -210,7 +218,12 @@ protected:
 class SFindInBlueprints: public SCompoundWidget
 {
 public:
-	SLATE_BEGIN_ARGS( SFindInBlueprints ){}
+	SLATE_BEGIN_ARGS( SFindInBlueprints )
+		: _bIsSearchWindow(true)
+		, _bHideSearchBar(false)
+	{}
+		SLATE_ARGUMENT(bool, bIsSearchWindow)
+		SLATE_ARGUMENT(bool, bHideSearchBar)
 	SLATE_END_ARGS()
 
 	void Construct(const FArguments& InArgs, TSharedPtr<class FBlueprintEditor> InBlueprintEditor);
@@ -219,8 +232,28 @@ public:
 	/** Focuses this widget's search box, and changes the mode as well, and optionally the search terms */
 	void FocusForUse(bool bSetFindWithinBlueprint, FString NewSearchTerms = FString(), bool bSelectFirstResult = false);
 
+	/**
+	 * Submits a search query
+	 *
+	 * @param InSearchString						String to search using
+	 * @param bInIsFindWithinBlueprint				TRUE if searching within the current Blueprint only
+	 * @param InSearchFilterForImaginaryDataReturn	If requesting a callback on search complete for the filtered imaginary data, this is the filter that the raw data will be passed through. By default nothing is collected
+	 * @param InMinimiumVersionRequirement			The minimum search requirement Blueprints must be to be searched or they will be reported as out-of-date.
+	 * @param InOnSearchComplete					Callback when the search is complete, passing the filtered imaginary data (if any).
+	 */
+	void MakeSearchQuery(FString InSearchString, bool bInIsFindWithinBlueprint, enum ESearchQueryFilter InSearchFilterForImaginaryDataReturn = ESearchQueryFilter::AllFilter, EFiBVersion InMinimiumVersionRequirement = EFiBVersion::FIB_VER_LATEST, FOnSearchComplete InOnSearchComplete = FOnSearchComplete());
+
 	/** Called when caching Blueprints is complete, if this widget initiated the indexing */
 	void OnCacheComplete();
+
+	/**
+	 * Asynchronously caches all Blueprints below a specified version.
+	 *
+	 * @param InOnFinished						Callback when the cache process is complete.
+	 * @param InMinimiumVersionRequirement		Version that Blueprints must be below to be loaded and indexed, by default all out-of-date Blueprints
+	 */
+	void CacheAllBlueprints(FSimpleDelegate InOnFinished = FSimpleDelegate(), EFiBVersion InMinimiumVersionRequirement = EFiBVersion::FIB_VER_LATEST);
+
 private:
 	/** Processes results of the ongoing async stream search */
 	EActiveTimerReturnType UpdateSearchResults( double InCurrentTime, float InDeltaTime );
@@ -248,12 +281,10 @@ private:
 
 	/* Called when a new row is being generated */
 	TSharedRef<ITableRow> OnGenerateRow(FSearchResult InItem, const TSharedRef<STableViewBase>& OwnerTable);
-
-	/** Begins the search based on the SearchValue */
-	void InitiateSearch();
 	
 	/** Launches a thread for streaming more content into the results widget */
 	void LaunchStreamThread(const FString& InSearchValue);
+	void LaunchStreamThread(const FString& InSearchValue, enum ESearchQueryFilter InSearchFilterForRawDataReturn, EFiBVersion InMinimiumVersionRequirement, FOnSearchComplete InOnSearchComplete);
 
 	/** Returns the percent complete on the search for the progress bar */
 	TOptional<float> GetPercentCompleteSearch() const;
@@ -275,6 +306,7 @@ private:
 
 	/** Callback to cache all uncached Blueprints */
 	FReply OnCacheAllBlueprints();
+	FReply OnCacheAllBlueprints(FSimpleDelegate InOnFinished, EFiBVersion InMinimiumVersionRequirement = EFiBVersion::FIB_VER_LATEST);
 
 	/** Callback to cancel the caching process */
 	FReply OnCancelCacheAll();
@@ -351,4 +383,13 @@ private:
 
 	/** Weak pointer to the cache bar slot, so it can be removed */
 	TWeakPtr< SWidget > CacheBarSlot;
+
+	/** Callback when search is complete */
+	FOnSearchComplete OnSearchComplete;
+
+	/** Cached count of out of date Blueprints from last search. */
+	int32 OutOfDateWithLastSearchBPCount;
+
+	/** Cached version that was last searched */
+	EFiBVersion LastSearchedFiBVersion;
 };
\ No newline at end of file
diff --git a/Engine/Source/Editor/Kismet/Public/ImaginaryBlueprintData.h b/Engine/Source/Editor/Kismet/Public/ImaginaryBlueprintData.h
--- a/Engine/Source/Editor/Kismet/Public/ImaginaryBlueprintData.h
+++ b/Engine/Source/Editor/Kismet/Public/ImaginaryBlueprintData.h
@@ -4,22 +4,6 @@
 #include "TextFilterExpressionEvaluator.h"
 #include "FindInBlueprints.h"
 
-/** Filters are used by functions for searching to decide whether items can call certain functions or match the requirements of a function */
-enum ESearchQueryFilter
-{
-	BlueprintFilter = 0,
-	GraphsFilter,
-	UberGraphsFilter,
-	FunctionsFilter,
-	MacrosFilter,
-	NodesFilter,
-	PinsFilter,
-	PropertiesFilter,
-	VariablesFilter,
-	ComponentsFilter,
-	AllFilter, // Will search all items, when used inside of another filter it will search all sub-items of that filter
-};
-
 enum ESearchableValueStatus
 {
 	NotSearchable = 0x00000000, // Cannot search this value, it is used for display purposes only
@@ -115,11 +99,11 @@ struct FComponentUniqueDisplay
 	FSearchResult SearchResult;
 };
 
-class FImaginaryFiBData : public ITextFilterExpressionContext
+class FImaginaryFiBData : public ITextFilterExpressionContext, public TSharedFromThis<FImaginaryFiBData>
 {
 public:
-	FImaginaryFiBData();
-	FImaginaryFiBData(TSharedPtr< FJsonObject > InUnparsedJsonObject, TMap<int32, FText>* InLookupTablePtr);
+	FImaginaryFiBData(TWeakPtr<FImaginaryFiBData> InOuter);
+	FImaginaryFiBData(TWeakPtr<FImaginaryFiBData> InOuter, TSharedPtr< FJsonObject > InUnparsedJsonObject, TMap<int32, FText>* InLookupTablePtr);
 
 	/** ITextFilterExpressionContext Interface */
 	// We don't actually use these overrides, see FFiBContextHelper for how we call the alternate functions. These will assert if they are accidentally called.
@@ -128,7 +112,7 @@ public:
 	/** End ITextFilterExpressionContext Interface */
 
 	/** Returns TRUE if this item is a category type, which helps to organize child data */
-	virtual bool IsCategory() { return false; }
+	virtual bool IsCategory() const { return false; }
 
 	/** Returns TRUE if this item is considered a Tag and Value category, where it's contents should be considered no different than the parent owner */
 	virtual bool IsTagAndValueCategory() { return false; }
@@ -139,8 +123,8 @@ public:
 	/** Checks if the filter can call functions for the passed filter, returns FALSE by default if the filter is not the AllFilter */
 	virtual bool CanCallFilter(ESearchQueryFilter InSearchQueryFilter) const;
 
-	/** Parses all child data, non-recursively, so children will be left in an unparsed Json state */
-	virtual void ParseAllChildData(ESearchableValueStatus InSearchabilityOverride = ESearchableValueStatus::Searchable);
+	/** Parses, in a thread-safe manner, all child data, non-recursively, so children will be left in an unparsed Json state */
+	void ParseAllChildData(ESearchableValueStatus InSearchabilityOverride = ESearchableValueStatus::Searchable);
 
 	/** Test the given value against the strings extracted from the current item. Will return the matching search components if any (can return TRUE without having any if the search components are hidden) */
 	virtual bool TestBasicStringExpression(const FTextFilterString& InValue, const ETextFilterTextComparisonMode InTextComparisonMode, TMultiMap< const FImaginaryFiBData*, FComponentUniqueDisplay >& InOutMatchingSearchComponents) const;
@@ -148,8 +132,21 @@ public:
 	/** Perform a complex expression test for the current item. Will return the matching search components if any (can return TRUE without having any if the search components are hidden) */
 	virtual bool TestComplexExpression(const FName& InKey, const FTextFilterString& InValue, const ETextFilterComparisonOperation InComparisonOperation, const ETextFilterTextComparisonMode InTextComparisonMode, TMultiMap< const FImaginaryFiBData*, FComponentUniqueDisplay >& InOutMatchingSearchComponents) const;
 
+	/** Returns the UObject represented by this Imaginary data give the UBlueprint owner. */
+	virtual UObject* GetObject(UBlueprint* InBlueprint) const;
+
+	/** This will return and force load the UBlueprint that owns this object data. */
+	virtual UBlueprint* GetBlueprint() const
+	{
+		if (Outer.IsValid())
+		{
+			return Outer.Pin()->GetBlueprint();
+		}
+		return nullptr;
+	}
+
 	/** Requests internal creation of the search result and properly initializes the visual representation of the result */
-	FSearchResult CreateSearchResult(FSearchResult InParent);
+	FSearchResult CreateSearchResult(FSearchResult InParent) const;
 
 	/** Accessor for the parsed child data for this item */
 	const TArray< TSharedPtr<FImaginaryFiBData> >& GetAllParsedChildData() const
@@ -157,11 +154,19 @@ public:
 		return ParsedChildData;
 	}
 
+	/** Adds a KeyValue pair to the ParsedTagAndValues map */
 	void AddKeyValuePair(FText InKey, FSearchableValueInfo& InValue)
 	{
 		ParsedTagsAndValues.Add(InKey.ToString(), InValue);
 	}
 
+	/** Returns the Outer of this Imaginary data that directly owns it */
+	TWeakPtr< FImaginaryFiBData > GetOuter() const
+	{
+		return Outer;
+	}
+
+	/** Builds a SearchTree ready to be displayed in the Find-in-Blueprints window */
 	static FSearchResult CreateSearchTree(FSearchResult InParentSearchResult, TWeakPtr< FImaginaryFiBData > InCurrentPointer, TArray< const FImaginaryFiBData* >& InValidSearchResults, TMultiMap< const FImaginaryFiBData*, FComponentUniqueDisplay >& InMatchingSearchComponents);
 
 protected:
@@ -186,6 +191,9 @@ protected:
 	/** Helper function for parsing Json values into usable properties */
 	void ParseJsonValue(FText InKey, FText InDisplayKey, TSharedPtr< FJsonValue > InJsonValue, bool bIsInArray=false, ESearchableValueStatus InSearchabilityOverride = ESearchableValueStatus::Searchable);
 
+	/** Internal version of the ParseAllChildData function, handles the bulk of the work */
+	virtual void ParseAllChildData_Internal(ESearchableValueStatus InSearchabilityOverride = ESearchableValueStatus::Searchable);
+
 protected:
 	/** The unparsed Json object representing this item. Auto-cleared after parsing */
 	TSharedPtr< FJsonObject > UnparsedJsonObject;
@@ -198,12 +206,18 @@ protected:
 
 	/** Pointer to the lookup table to decompressed the Json strings back into fully formed FTexts */
 	TMap<int32, FText>* LookupTablePtr;
+
+	/** Outer of this object that owns it, used for climbing up the hierarchy */
+	TWeakPtr<FImaginaryFiBData> Outer;
+
+	/** Allows for thread-safe parsing of the imaginary data. Only a single Imaginary data can be parsed at a time */
+	static FCriticalSection ParseChildDataCriticalSection;
 };
 
 class FFiBMetaData : public FImaginaryFiBData
 {
 public:
-	FFiBMetaData(TSharedPtr< FJsonObject > InUnparsedJsonObject, TMap<int32, FText>* InLookupTablePtr);
+	FFiBMetaData(TWeakPtr<FImaginaryFiBData> InOuter, TSharedPtr< FJsonObject > InUnparsedJsonObject, TMap<int32, FText>* InLookupTablePtr);
 	
 	/** FImaginaryFiBData Interface */
 	virtual bool TrySpecialHandleJsonValue(FString InKey, TSharedPtr< FJsonValue > InJsonValue) override;
@@ -240,13 +254,13 @@ public:
 	/** Callback declaration for handling special parsing of the items in the category */
 	DECLARE_DELEGATE_TwoParams(FCategorySectionHelperCallback, TSharedPtr< FJsonObject >, TArray< TSharedPtr<FImaginaryFiBData> >&);
 
-	FCategorySectionHelper(TMap<int32, FText>* InLookupTablePtr, FText InCategoryName, bool bInTagAndValueCategory);
-	FCategorySectionHelper(TSharedPtr< FJsonObject > InUnparsedJsonObject, TMap<int32, FText>* InLookupTablePtr, FText InCategoryName, bool bInTagAndValueCategory);
-	FCategorySectionHelper(TSharedPtr< FJsonObject > InUnparsedJsonObject, TMap<int32, FText>* InLookupTablePtr, FText InCategoryName, bool bInTagAndValueCategory, FCategorySectionHelperCallback InSpecialHandlingCallback);
+	FCategorySectionHelper(TWeakPtr<FImaginaryFiBData> InOuter, TMap<int32, FText>* InLookupTablePtr, FText InCategoryName, bool bInTagAndValueCategory);
+	FCategorySectionHelper(TWeakPtr<FImaginaryFiBData> InOuter, TSharedPtr< FJsonObject > InUnparsedJsonObject, TMap<int32, FText>* InLookupTablePtr, FText InCategoryName, bool bInTagAndValueCategory);
+	FCategorySectionHelper(TWeakPtr<FImaginaryFiBData> InOuter, TSharedPtr< FJsonObject > InUnparsedJsonObject, TMap<int32, FText>* InLookupTablePtr, FText InCategoryName, bool bInTagAndValueCategory, FCategorySectionHelperCallback InSpecialHandlingCallback);
 
 	/** FImaginaryFiBData Interface */
-	virtual void ParseAllChildData(ESearchableValueStatus InSearchabilityOverride/* = ESearchableValueStatus::ESearchableValueStatus::Searchable*/) override;
-	virtual bool IsCategory() override { return true; }
+	virtual void ParseAllChildData_Internal(ESearchableValueStatus InSearchabilityOverride/* = ESearchableValueStatus::ESearchableValueStatus::Searchable*/) override;
+	virtual bool IsCategory() const override { return true; }
 	virtual bool IsTagAndValueCategory() override { return bIsTagAndValue; }
 	/** End FImaginaryFiBData Interface */
 
@@ -282,6 +296,7 @@ public:
 	/** FImaginaryFiBData Interface */
 	virtual bool IsCompatibleWithFilter(ESearchQueryFilter InSearchQueryFilter) const override;
 	virtual bool CanCallFilter(ESearchQueryFilter InSearchQueryFilter) const override;
+	virtual UBlueprint* GetBlueprint() const override;
 	/** End FImaginaryFiBData Interface */
 
 protected:
@@ -314,7 +329,7 @@ protected:
 class FImaginaryGraph : public FImaginaryFiBData
 {
 public:
-	FImaginaryGraph(TSharedPtr< FJsonObject > InUnparsedJsonObject, TMap<int32, FText>* InLookupTablePtr, EGraphType InGraphType);
+	FImaginaryGraph(TWeakPtr<FImaginaryFiBData> InOuter, TSharedPtr< FJsonObject > InUnparsedJsonObject, TMap<int32, FText>* InLookupTablePtr, EGraphType InGraphType);
 
 	/** FImaginaryFiBData Interface */
 	virtual bool IsCompatibleWithFilter(ESearchQueryFilter InSearchQueryFilter) const override;
@@ -337,12 +352,12 @@ protected:
 class FImaginaryGraphNode : public FImaginaryFiBData
 {
 public:
-	FImaginaryGraphNode(TSharedPtr< FJsonObject > InUnparsedJsonObject, TMap<int32, FText>* InLookupTablePtr);
+	FImaginaryGraphNode(TWeakPtr<FImaginaryFiBData> InOuter, TSharedPtr< FJsonObject > InUnparsedJsonObject, TMap<int32, FText>* InLookupTablePtr);
 
 	/** FImaginaryFiBData Interface */
 	virtual bool IsCompatibleWithFilter(ESearchQueryFilter InSearchQueryFilter) const override;
 	virtual bool CanCallFilter(ESearchQueryFilter InSearchQueryFilter) const override;
-	virtual void ParseAllChildData(ESearchableValueStatus InSearchabilityOverride/* = ESearchableValueStatus::Searchable*/) override;
+	virtual void ParseAllChildData_Internal(ESearchableValueStatus InSearchabilityOverride/* = ESearchableValueStatus::Searchable*/) override;
 	/** End FImaginaryFiBData Interface */
 
 protected:
@@ -361,7 +376,7 @@ protected:
 class FImaginaryProperty : public FImaginaryFiBData
 {
 public:
-	FImaginaryProperty(TSharedPtr< FJsonObject > InUnparsedJsonObject, TMap<int32, FText>* InLookupTablePtr);
+	FImaginaryProperty(TWeakPtr<FImaginaryFiBData> InOuter, TSharedPtr< FJsonObject > InUnparsedJsonObject, TMap<int32, FText>* InLookupTablePtr);
 
 	/** FImaginaryFiBData Interface */
 	virtual bool IsCompatibleWithFilter(ESearchQueryFilter InSearchQueryFilter) const override;
@@ -378,7 +393,7 @@ protected:
 class FImaginaryComponent : public FImaginaryProperty
 {
 public:
-	FImaginaryComponent(TSharedPtr< FJsonObject > InUnparsedJsonObject, TMap<int32, FText>* InLookupTablePtr);
+	FImaginaryComponent(TWeakPtr<FImaginaryFiBData> InOuter, TSharedPtr< FJsonObject > InUnparsedJsonObject, TMap<int32, FText>* InLookupTablePtr);
 
 	/** FImaginaryFiBData Interface */
 	virtual bool IsCompatibleWithFilter(ESearchQueryFilter InSearchQueryFilter) const override;
@@ -389,7 +404,7 @@ public:
 class FImaginaryPin : public FImaginaryFiBData
 {
 public:
-	FImaginaryPin(TSharedPtr< FJsonObject > InUnparsedJsonObject, TMap<int32, FText>* InLookupTablePtr, FString InSchemaName);
+	FImaginaryPin(TWeakPtr<FImaginaryFiBData> InOuter, TSharedPtr< FJsonObject > InUnparsedJsonObject, TMap<int32, FText>* InLookupTablePtr, FString InSchemaName);
 
 	/** FImaginaryFiBData Interface */
 	virtual bool IsCompatibleWithFilter(ESearchQueryFilter InSearchQueryFilter) const override;
diff --git a/Engine/Source/Editor/KismetWidgets/Private/SPinTypeSelector.cpp b/Engine/Source/Editor/KismetWidgets/Private/SPinTypeSelector.cpp
--- a/Engine/Source/Editor/KismetWidgets/Private/SPinTypeSelector.cpp
+++ b/Engine/Source/Editor/KismetWidgets/Private/SPinTypeSelector.cpp
@@ -7,6 +7,7 @@
 #include "Editor/UnrealEd/Public/SListViewSelectorDropdownMenu.h"
 #include "SSearchBox.h"
 #include "SSubMenuHandler.h"
+#include "BlueprintEditorUtils.h"
 
 #define LOCTEXT_NAMESPACE "PinTypeSelector"
 
@@ -207,7 +208,7 @@ FText SPinTypeSelector::GetTypeDescription() const
 
 const FSlateBrush* SPinTypeSelector::GetTypeIconImage() const
 {
-	return GetIconFromPin( TargetPinType.Get() );
+	return FBlueprintEditorUtils::GetIconFromPin( TargetPinType.Get() );
 }
 
 FSlateColor SPinTypeSelector::GetTypeIconColor() const
@@ -242,7 +243,7 @@ TSharedRef<ITableRow> SPinTypeSelector::GenerateTypeTreeRow(FPinTypeTreeItem InI
 	const FEdGraphPinType& PinType = InItem->GetPinType(false);
 
 	// Determine the best icon the to represents this item
-	const FSlateBrush* IconBrush = GetIconFromPin(PinType);
+	const FSlateBrush* IconBrush = FBlueprintEditorUtils::GetIconFromPin(PinType);
 
 	// Use tooltip if supplied, otherwise just repeat description
 	const FText OrgTooltip = InItem->GetToolTip();
@@ -381,7 +382,7 @@ TSharedRef< SWidget > SPinTypeSelector::GetAllowedObjectTypes(FPinTypeTreeItem I
 
 	// Do not force the pin type here, that causes a load of the Blueprint (if unloaded)
 	FEdGraphPinType PinType = InItem->GetPinType(false);
-	const FSlateBrush* IconBrush = GetIconFromPin(PinType);
+	const FSlateBrush* IconBrush = FBlueprintEditorUtils::GetIconFromPin(PinType);
 
 	FFormatNamedArguments Args;
 
@@ -709,25 +710,6 @@ bool SPinTypeSelector::GetChildrenMatchingSearch(const FText& InSearchText, cons
 	return bReturnVal;
 }
 
-const FSlateBrush* SPinTypeSelector::GetIconFromPin( const FEdGraphPinType& PinType ) const
-{
-	const FSlateBrush* IconBrush = FEditorStyle::GetBrush(TEXT("Kismet.VariableList.TypeIcon"));
-	const UObject* PinSubObject = PinType.PinSubCategoryObject.Get();
-	if( (PinType.bIsArray || (IsArrayChecked() == ECheckBoxState::Checked)) && PinType.PinCategory != Schema->PC_Exec )
-	{
-		IconBrush = FEditorStyle::GetBrush(TEXT("Kismet.VariableList.ArrayTypeIcon"));
-	}
-	else if( PinSubObject )
-	{
-		UClass* VarClass = FindObject<UClass>(ANY_PACKAGE, *PinSubObject->GetName());
-		if( VarClass )
-		{
-			IconBrush = FClassIconFinder::FindIconForClass( VarClass );
-		}
-	}
-	return IconBrush;
-}
-
 FText SPinTypeSelector::GetToolTipForComboBoxType() const
 {
 	if(IsEnabled())
diff --git a/Engine/Source/Editor/KismetWidgets/Public/SPinTypeSelector.h b/Engine/Source/Editor/KismetWidgets/Public/SPinTypeSelector.h
--- a/Engine/Source/Editor/KismetWidgets/Public/SPinTypeSelector.h
+++ b/Engine/Source/Editor/KismetWidgets/Public/SPinTypeSelector.h
@@ -137,14 +137,6 @@ protected:
 	/** Helper to generate the filtered list of types, based on the search string matching */
 	bool GetChildrenMatchingSearch(const FText& SearchText, const TArray<FPinTypeTreeItem>& UnfilteredList, TArray<FPinTypeTreeItem>& OutFilteredList);
 
-	/** 
-	 * Determine the best icon to represent the given pin.
-	 *
-	 * @param PinType		The pin get the icon for.
-	 * @param returns a brush that best represents the icon (or Kismet.VariableList.TypeIcon if none is available )
-	 */
-	const FSlateBrush* GetIconFromPin( const FEdGraphPinType& PinType ) const;
-
 	/** Callback to get the tooltip text for the pin type combo box */
 	FText GetToolTipForComboBoxType() const;
 
diff --git a/Engine/Source/Editor/UnrealEd/Classes/Settings/EditorExperimentalSettings.h b/Engine/Source/Editor/UnrealEd/Classes/Settings/EditorExperimentalSettings.h
--- a/Engine/Source/Editor/UnrealEd/Classes/Settings/EditorExperimentalSettings.h
+++ b/Engine/Source/Editor/UnrealEd/Classes/Settings/EditorExperimentalSettings.h
@@ -77,6 +77,10 @@ public:
 	UPROPERTY(EditAnywhere, config, Category = Blueprints, meta = (DisplayName = "Use the Diff Tool for Animation Blueprints"))
 	bool bEnableAnimVisualDiff;
 
+	/** Enables "Find and Replace All" tool in the MyBlueprint window for variables */
+	UPROPERTY(EditAnywhere, config, Category = Blueprints, meta = (DisplayName = "Find and Replace All References Tool"))
+	bool bEnableFindAndReplaceReferences;
+
 	/** Should arrows indicating data/execution flow be drawn halfway along wires? */
 	UPROPERTY(/*EditAnywhere - deprecated (moved into UBlueprintEditorSettings), */config/*, Category=Blueprints, meta=(DisplayName="Draw midpoint arrows in Blueprints")*/)
 	bool bDrawMidpointArrowsInBlueprints;
diff --git a/Engine/Source/Editor/UnrealEd/Private/Kismet2/BlueprintEditorUtils.cpp b/Engine/Source/Editor/UnrealEd/Private/Kismet2/BlueprintEditorUtils.cpp
--- a/Engine/Source/Editor/UnrealEd/Private/Kismet2/BlueprintEditorUtils.cpp
+++ b/Engine/Source/Editor/UnrealEd/Private/Kismet2/BlueprintEditorUtils.cpp
@@ -54,6 +54,7 @@
 #include "EditorCategoryUtils.h"
 #include "EngineUtils.h"
 #include "Engine/LevelScriptActor.h"
+#include "ClassIconFinder.h"
 #define LOCTEXT_NAMESPACE "Blueprint"
 
 DEFINE_LOG_CATEGORY(LogBlueprintDebug);
@@ -8145,6 +8146,27 @@ void FBlueprintEditorUtils::PostSetupObjectPinType(UBlueprint* InBlueprint, FBPV
 	}
 }
 
+const FSlateBrush* FBlueprintEditorUtils::GetIconFromPin( const FEdGraphPinType& PinType )
+{
+	const UEdGraphSchema_K2* K2Schema = GetDefault<UEdGraphSchema_K2>();
+	
+	const FSlateBrush* IconBrush = FEditorStyle::GetBrush(TEXT("Kismet.VariableList.TypeIcon"));
+	const UObject* PinSubObject = PinType.PinSubCategoryObject.Get();
+	if( PinType.bIsArray && PinType.PinCategory != K2Schema->PC_Exec )
+	{
+		IconBrush = FEditorStyle::GetBrush(TEXT("Kismet.VariableList.ArrayTypeIcon"));
+	}
+	else if( PinSubObject )
+	{
+		UClass* VarClass = FindObject<UClass>(ANY_PACKAGE, *PinSubObject->GetName());
+		if( VarClass )
+		{
+			IconBrush = FClassIconFinder::FindIconForClass( VarClass );
+		}
+	}
+	return IconBrush;
+}
+
 FText FBlueprintEditorUtils::GetFriendlyClassDisplayName(const UClass* Class)
 {
 	if (Class != nullptr)
diff --git a/Engine/Source/Editor/UnrealEd/Public/Kismet2/BlueprintEditorUtils.h b/Engine/Source/Editor/UnrealEd/Public/Kismet2/BlueprintEditorUtils.h
--- a/Engine/Source/Editor/UnrealEd/Public/Kismet2/BlueprintEditorUtils.h
+++ b/Engine/Source/Editor/UnrealEd/Public/Kismet2/BlueprintEditorUtils.h
@@ -1238,6 +1238,14 @@ public:
 	 */
 	static class UK2Node_FunctionResult* FindOrCreateFunctionResultNode(class UK2Node_EditablePinBase* InFunctionEntryNode);
 
+	/** 
+	 * Determine the best icon to represent the given pin.
+	 *
+	 * @param PinType		The pin get the icon for.
+	 * @param returns a brush that best represents the icon (or Kismet.VariableList.TypeIcon if none is available )
+	 */
+	static const struct FSlateBrush* GetIconFromPin( const FEdGraphPinType& PinType );
+
 protected:
 	// Removes all NULL graph references from the SubGraphs array and recurses thru the non-NULL ones
 	static void CleanNullGraphReferencesRecursive(UEdGraph* Graph);
diff --git a/Engine/Source/Runtime/Engine/Classes/Engine/MemberReference.h b/Engine/Source/Runtime/Engine/Classes/Engine/MemberReference.h
--- a/Engine/Source/Runtime/Engine/Classes/Engine/MemberReference.h
+++ b/Engine/Source/Runtime/Engine/Classes/Engine/MemberReference.h
@@ -228,6 +228,17 @@ public:
 	{
 		return !MemberScope.IsEmpty();
 	}
+
+#if WITH_EDITOR
+	/**
+	 * Returns a search string to submit to Find-in-Blueprints to find references to this reference
+	 *
+	 * @param InFieldOwner		The owner of the field, cannot be resolved internally
+	 * @return					Search string to find this reference in other Blueprints
+	 */
+	ENGINE_API FString GetReferenceSearchString(UClass* InFieldOwner) const;
+#endif
+
 private:
 #if WITH_EDITOR
 	/**
diff --git a/Engine/Source/Runtime/Engine/Private/MemberReference.cpp b/Engine/Source/Runtime/Engine/Private/MemberReference.cpp
--- a/Engine/Source/Runtime/Engine/Private/MemberReference.cpp
+++ b/Engine/Source/Runtime/Engine/Private/MemberReference.cpp
@@ -116,6 +116,43 @@ void FMemberReference::InvalidateScope()
 }
 
 #if WITH_EDITOR
+
+FString FMemberReference::GetReferenceSearchString(UClass* InFieldOwner) const
+{
+	if (!IsLocalScope())
+	{
+		if (InFieldOwner)
+		{
+			if (MemberGuid.IsValid())
+			{
+				return FString::Printf(TEXT("Nodes(VariableReference(MemberName=+%s && MemberGuid(A=%i && B=%i && C=%i && D=%i) ))"), *MemberName.ToString(), MemberGuid.A, MemberGuid.B, MemberGuid.C, MemberGuid.D);
+			}
+			else
+			{
+				FString ExportMemberParentName;
+				ExportMemberParentName = InFieldOwner->GetClass()->GetName();
+				ExportMemberParentName.AppendChar('\'');
+				ExportMemberParentName += InFieldOwner->GetAuthoritativeClass()->GetPathName();
+				ExportMemberParentName.AppendChar('\'');
+
+				return FString::Printf(TEXT("Nodes(VariableReference(MemberName=+%s && (MemberParent=\"%s\" || bSelfContext=true) ))"), *MemberName.ToString(), *ExportMemberParentName);
+			}
+		}
+		else if (MemberGuid.IsValid())
+		{
+			return FString::Printf(TEXT("Nodes(VariableReference(MemberName=+%s && MemberGuid(A=%i && B=%i && C=%i && D=%i)))"), *MemberName.ToString(), MemberGuid.A, MemberGuid.B, MemberGuid.C, MemberGuid.D);
+		}
+		else
+		{
+			return FString::Printf(TEXT("Nodes(VariableReference(MemberName=+%s))"), *MemberName.ToString());
+		}
+	}
+	else
+	{
+		return FString::Printf(TEXT("Nodes(VariableReference(MemberName=+%s && MemberScope=+%s))"), *MemberName.ToString(), *GetMemberScopeName());
+	}
+}
+
 FName FMemberReference::RefreshLocalVariableName(UClass* InSelfScope) const
 {
 	TArray<UBlueprint*> Blueprints;
