Fix build issues from CL# 2713446

diff --git a/Engine/Source/Editor/Kismet/Private/SReplaceNodeReferences.cpp b/Engine/Source/Editor/Kismet/Private/SReplaceNodeReferences.cpp
--- a/Engine/Source/Editor/Kismet/Private/SReplaceNodeReferences.cpp
+++ b/Engine/Source/Editor/Kismet/Private/SReplaceNodeReferences.cpp
@@ -456,7 +456,6 @@ void SReplaceNodeReferences::FindAllReplacementsComplete(TArray<TSharedPtr<class
 			{
 				BlueprintsModified.AddUnique(ImaginaryData->GetBlueprint());
 				UObject* Node = ImaginaryData->GetObject(ImaginaryData->GetBlueprint());
-				Node = Node;
 				UK2Node_Variable* VariableNode = Cast<UK2Node_Variable>(Node);
 				if (ensure(VariableNode))
 				{
diff --git a/Engine/Source/Editor/Kismet/Private/SReplaceNodeReferences.h b/Engine/Source/Editor/Kismet/Private/SReplaceNodeReferences.h
--- a/Engine/Source/Editor/Kismet/Private/SReplaceNodeReferences.h
+++ b/Engine/Source/Editor/Kismet/Private/SReplaceNodeReferences.h
@@ -3,6 +3,9 @@
 class FTargetReplaceReferences
 {
 public:
+	virtual ~FTargetReplaceReferences()
+	{}
+
 	/** Returns a generated widget to represent this target item reference */
 	virtual TSharedRef<SWidget> CreateWidget() const = 0;
 
