Using shortcut keys to add local variables will no longer crash the editor when done in a graph that does not support local variables.

Added FBlueprintEditor::CanAddNewLocalVariable.

--------------
Merging

//UE4/Dev-Blueprints/Engine/Source/Editor/Kismet/...

to //UE4/Engine/Source/Editor/Kismet/...

diff --git a/Engine/Source/Editor/Kismet/Private/BlueprintEditor.cpp b/Engine/Source/Editor/Kismet/Private/BlueprintEditor.cpp
--- a/Engine/Source/Editor/Kismet/Private/BlueprintEditor.cpp
+++ b/Engine/Source/Editor/Kismet/Private/BlueprintEditor.cpp
@@ -2454,7 +2454,7 @@ void FBlueprintEditor::CreateDefaultCommands()
 	
 	ToolkitCommands->MapAction( FBlueprintEditorCommands::Get().AddNewLocalVariable,
 		FExecuteAction::CreateSP(this, &FBlueprintEditor::OnAddNewLocalVariable),
-		FCanExecuteAction::CreateSP(this, &FBlueprintEditor::InEditingMode),
+		FCanExecuteAction::CreateSP(this, &FBlueprintEditor::CanAddNewLocalVariable),
 		FIsActionChecked(),
 		FIsActionButtonVisible::CreateSP(this, &FBlueprintEditor::NewDocument_IsVisibleForType, CGT_NewLocalVariable));
 
@@ -7032,6 +7032,17 @@ void FBlueprintEditor::OnAddNewVariable()
 	}
 }
 
+bool FBlueprintEditor::CanAddNewLocalVariable() const
+{
+	if (InEditingMode())
+	{
+		UEdGraph* TargetGraph = FBlueprintEditorUtils::GetTopLevelGraph(FocusedGraphEdPtr.Pin()->GetCurrentGraph());
+
+		return TargetGraph->GetSchema()->GetGraphType(TargetGraph) == GT_Function;
+	}
+	return false;
+}
+
 void FBlueprintEditor::OnAddNewLocalVariable()
 {
 	// Find the top level graph to place the local variables into
diff --git a/Engine/Source/Editor/Kismet/Public/BlueprintEditor.h b/Engine/Source/Editor/Kismet/Public/BlueprintEditor.h
--- a/Engine/Source/Editor/Kismet/Public/BlueprintEditor.h
+++ b/Engine/Source/Editor/Kismet/Public/BlueprintEditor.h
@@ -416,6 +416,9 @@ public:
 	void OnAddNewVariable();
 	FReply OnAddNewVariable_OnClick() {OnAddNewVariable(); return FReply::Handled();}
 	
+	/** Checks if adding a local variable is allowed in the focused graph */
+	bool CanAddNewLocalVariable() const;
+
 	/** Adds a new local variable to the focused function graph */
 	void OnAddNewLocalVariable();
 
