Copying //UE4/Dev-Blueprints to Dev-Main (//UE4/Dev-Main) @ 2781164

diff --git a/Engine/Source/Editor/EditorStyle/Private/SlateEditorStyle.cpp b/Engine/Source/Editor/EditorStyle/Private/SlateEditorStyle.cpp
--- a/Engine/Source/Editor/EditorStyle/Private/SlateEditorStyle.cpp
+++ b/Engine/Source/Editor/EditorStyle/Private/SlateEditorStyle.cpp
@@ -1052,6 +1052,11 @@ void FSlateEditorStyle::FStyle::SetupGeneralStyles()
 		Set("DebugTools.TabIcon", new IMAGE_BRUSH( "Icons/icon_tab_DebugTools_16x", Icon16x16 ) );
 	}
 
+	// Performance Analysis Tools Window
+	{
+		Set("PerfTools.TabIcon", new IMAGE_BRUSH( "Icons/icon_tab_PerfTools_16x", Icon16x16 ) );
+	}
+
 	// Modules Window
 	{
 		Set("Modules.TabIcon", new IMAGE_BRUSH( "Icons/icon_tab_Modules_16px", Icon16x16 ) );
@@ -4628,6 +4633,25 @@ void FSlateEditorStyle::FStyle::SetupLevelEditorStyle()
 	}
 #endif // WITH_EDITOR || IS_PROGRAM
 
+	// Blueprint Profiler
+#if WITH_EDITOR
+	{
+		// Profiler stat window
+		Set( "BlueprintProfiler.BPIcon_Normal", new IMAGE_BRUSH("Icons/AssetIcons/Blueprint_16x", Icon16x16, FLinearColor( 0.46f, 0.54f, 0.81f )));
+		Set( "BlueprintProfiler.BPIcon_Macro", new IMAGE_BRUSH("Icons/AssetIcons/BlueprintMacroLibrary_16x", Icon16x16, FLinearColor( 0.46f, 0.54f, 0.81f )));
+		Set( "BlueprintProfiler.BPIcon_Interface", new IMAGE_BRUSH("Icons/AssetIcons/BlueprintInterface_16x", Icon16x16, FLinearColor( 0.46f, 0.54f, 0.81f )));
+		Set( "BlueprintProfiler.BPIcon_FunctionLibrary", new IMAGE_BRUSH("Icons/AssetIcons/BlueprintFunctionLibrary_16x", Icon16x16, FLinearColor( 0.46f, 0.54f, 0.81f )));
+		Set( "BlueprintProfiler.Actor", new IMAGE_BRUSH("Icons/ActorIcons/Sphere_16px", Icon16x16));
+		Set( "BlueprintProfiler.BPEventNode", new IMAGE_BRUSH("Icons/icon_Blueprint_Event_16x", Icon16x16, FLinearColor( 0.91f, 0.16f, 0.16f )));
+		Set( "BlueprintProfiler.BPFunctionNode", new IMAGE_BRUSH("Icons/icon_Blueprint_NewFunction_16x", Icon16x16, FLinearColor( 0.46f, 0.54f, 0.81f )));
+		Set( "BlueprintProfiler.BPMacroNode", new IMAGE_BRUSH("Icons/icon_Blueprint_Macro_16x", Icon16x16, FLinearColor( 0.46f, 0.54f, 0.81f )));
+		Set( "BlueprintProfiler.BPBranchNode", new IMAGE_BRUSH("Icons/icon_Blueprint_Sequence_16x", Icon16x16, FLinearColor( 0.46f, 0.54f, 0.81f )));
+		Set( "BlueprintProfiler.BPPinConnected", new IMAGE_BRUSH( "Icons/BPProfiler_ExecPin_Connected", Icon12x16 ) );
+		Set( "BlueprintProfiler.BPPinDisconnected", new IMAGE_BRUSH( "Icons/BPProfiler_ExecPin_Disconnected", Icon12x16 ) );
+		Set( "BlueprintProfiler.BPNode", new IMAGE_BRUSH("Icons/Icon_Profiler_BPNode_16x", Icon16x16));
+	}
+#endif // WITH_EDITOR
+
 	// Mobility Icons
 	{
 		Set("Mobility.Movable", new IMAGE_BRUSH("/Icons/Mobility/Movable_16x", Icon16x16));
@@ -5037,6 +5061,7 @@ void FSlateEditorStyle::FStyle::SetupPersonaStyle()
 		Set( "BlueprintEditor.ShowFloor", new IMAGE_BRUSH( "Icons/icon_Show_Floor_40px", Icon16x16));
 		Set( "BlueprintEditor.ShowGrid", new IMAGE_BRUSH( "Icons/icon_ShowGrid_16x", Icon16x16 ) );
 		Set( "BlueprintEditor.EnableSimulation", new IMAGE_BRUSH( "Icons/icon_Enable_Simulation_40px", Icon40x40));
+		Set( "BlueprintEditor.EnableProfiling", new IMAGE_BRUSH("/Icons/icon_Enable_Profiling_40x", Icon40x40) );
 		Set( "BlueprintEditor.EnableSimulation.Small", new IMAGE_BRUSH( "Icons/icon_Enable_Simulation_40px", Icon20x20));
 		Set( "SCS.NativeComponent", new IMAGE_BRUSH( "Icons/NativeSCSComponent", Icon20x20 ));
 		Set( "SCS.Component", new IMAGE_BRUSH( "Icons/SCSComponent", Icon20x20 ));
diff --git a/Engine/Source/Editor/GameProjectGeneration/Private/GameProjectUtils.cpp b/Engine/Source/Editor/GameProjectGeneration/Private/GameProjectUtils.cpp
--- a/Engine/Source/Editor/GameProjectGeneration/Private/GameProjectUtils.cpp
+++ b/Engine/Source/Editor/GameProjectGeneration/Private/GameProjectUtils.cpp
@@ -2793,6 +2793,35 @@ bool GameProjectUtils::GenerateGameModuleHeaderFile(const FString& NewBuildFileN
 	return WriteOutputFile(NewBuildFileName, FinalOutput, OutFailReason);
 }
 
+bool GameProjectUtils::GeneratePluginModuleCPPFile(const FString& CPPFileName, const FString& ModuleName, const FString& StartupSourceCode, FText& OutFailReason)
+{
+	FString Template;
+	if (!ReadTemplateFile(TEXT("PluginModule.cpp.template"), Template, OutFailReason))
+	{
+		return false;
+	}
+
+	FString FinalOutput = Template.Replace(TEXT("%COPYRIGHT_LINE%"), *MakeCopyrightLine(), ESearchCase::CaseSensitive);
+	FinalOutput = FinalOutput.Replace(TEXT("%MODULE_NAME%"), *ModuleName, ESearchCase::CaseSensitive);
+	FinalOutput = FinalOutput.Replace(TEXT("%MODULE_STARTUP_CODE%"), *StartupSourceCode, ESearchCase::CaseSensitive);
+
+	return WriteOutputFile(CPPFileName, FinalOutput, OutFailReason);
+}
+
+bool GameProjectUtils::GeneratePluginModuleHeaderFile(const FString& HeaderFileName, const TArray<FString>& PublicHeaderIncludes, FText& OutFailReason)
+{
+	FString Template;
+	if (!ReadTemplateFile(TEXT("PluginModule.h.template"), Template, OutFailReason))
+	{
+		return false;
+	}
+
+	FString FinalOutput = Template.Replace(TEXT("%COPYRIGHT_LINE%"), *MakeCopyrightLine(), ESearchCase::CaseSensitive);
+	FinalOutput = FinalOutput.Replace(TEXT("%PUBLIC_HEADER_INCLUDES%"), *MakeIncludeList(PublicHeaderIncludes), ESearchCase::CaseSensitive);
+
+	return WriteOutputFile(HeaderFileName, FinalOutput, OutFailReason);
+}
+
 void GameProjectUtils::OnUpdateProjectConfirm()
 {
 	UpdateProject();
diff --git a/Engine/Source/Editor/GameProjectGeneration/Public/GameProjectUtils.h b/Engine/Source/Editor/GameProjectGeneration/Public/GameProjectUtils.h
--- a/Engine/Source/Editor/GameProjectGeneration/Public/GameProjectUtils.h
+++ b/Engine/Source/Editor/GameProjectGeneration/Public/GameProjectUtils.h
@@ -220,6 +220,12 @@ public:
 	/** Generates a Build.cs file for a game module */
 	static bool GenerateGameModuleBuildFile(const FString& NewBuildFileName, const FString& ModuleName, const TArray<FString>& PublicDependencyModuleNames, const TArray<FString>& PrivateDependencyModuleNames, FText& OutFailReason);
 
+	/** Generates a module .cpp file, intended for plugin use */
+	static bool GeneratePluginModuleCPPFile(const FString& CPPFileName, const FString& ModuleName, const FString& StartupSourceCode, FText& OutFailReason);
+
+	/** Generates a module .h file, intended for plugin use */
+	static bool GeneratePluginModuleHeaderFile(const FString& HeaderFileName, const TArray<FString>& PublicHeaderIncludes, FText& OutFailReason);
+
 	/** Returns true if the currently loaded project has code files */
 	static bool ProjectHasCodeFiles();
 
diff --git a/Engine/Source/Editor/GameplayTagsEditor/Private/SGameplayTagContainerGraphPin.cpp b/Engine/Source/Editor/GameplayTagsEditor/Private/SGameplayTagContainerGraphPin.cpp
--- a/Engine/Source/Editor/GameplayTagsEditor/Private/SGameplayTagContainerGraphPin.cpp
+++ b/Engine/Source/Editor/GameplayTagsEditor/Private/SGameplayTagContainerGraphPin.cpp
@@ -158,7 +158,7 @@ void SGameplayTagContainerGraphPin::RefreshTagList()
 	FString CurrentDefaultValue = GraphPinObj->GetDefaultAsString();
 	if (CurrentDefaultValue.IsEmpty())
 	{
-		CurrentDefaultValue = FString(TEXT("(GameplayTags=())"));
+		CurrentDefaultValue = FString(TEXT("(GameplayTags=)"));
 	}
 	if (!CurrentDefaultValue.Equals(TagContainerString))
 	{
diff --git a/Engine/Source/Editor/GraphEditor/Private/SCommentBubble.cpp b/Engine/Source/Editor/GraphEditor/Private/SCommentBubble.cpp
--- a/Engine/Source/Editor/GraphEditor/Private/SCommentBubble.cpp
+++ b/Engine/Source/Editor/GraphEditor/Private/SCommentBubble.cpp
@@ -309,11 +309,7 @@ void SCommentBubble::UpdateBubble()
 
 FVector2D SCommentBubble::GetOffset() const
 {
-	if( GraphNode->bCommentBubbleVisible || OpacityValue > 0.f )
-	{
-		return FVector2D( 0.f, -GetDesiredSize().Y );
-	}
-	return FVector2D::ZeroVector;
+	return FVector2D( 0.f, -GetDesiredSize().Y );
 }
 
 float SCommentBubble::GetArrowCenterOffset() const
diff --git a/Engine/Source/Editor/GraphEditor/Private/SGraphEditorImpl.cpp b/Engine/Source/Editor/GraphEditor/Private/SGraphEditorImpl.cpp
--- a/Engine/Source/Editor/GraphEditor/Private/SGraphEditorImpl.cpp
+++ b/Engine/Source/Editor/GraphEditor/Private/SGraphEditorImpl.cpp
@@ -730,7 +730,7 @@ EVisibility SGraphEditorImpl::InstructionTextVisibility() const
 {
 	if (!GetInstructionText().IsEmpty() && (GetInstructionTextFade() > 0.0f))
 	{
-		return EVisibility::Visible;
+		return EVisibility::HitTestInvisible;
 	}
 	return EVisibility::Hidden;
 }
diff --git a/Engine/Source/Editor/Kismet/Private/BlueprintDetailsCustomization.cpp b/Engine/Source/Editor/Kismet/Private/BlueprintDetailsCustomization.cpp
--- a/Engine/Source/Editor/Kismet/Private/BlueprintDetailsCustomization.cpp
+++ b/Engine/Source/Editor/Kismet/Private/BlueprintDetailsCustomization.cpp
@@ -193,9 +193,6 @@ void FBlueprintVarActionDetails::CustomizeDetails( IDetailLayoutBuilder& DetailL
 		return;
 	}
 
-	// Help in tracking UE-20390
-	UE_LOG(LogBlueprintUserMessages, Log, TEXT("CustomizeDetails for Blueprint '%s', cached variable property's owner is '%s'"), *GetBlueprintObj()->GetName(), *CachedVariableProperty->GetOwnerClass()->GetPathName());
-
 	CachedVariableName = GetVariableName();
 
 	TWeakPtr<FBlueprintEditor> BlueprintEditor = MyBlueprint.Pin()->GetBlueprintEditor();
@@ -264,7 +261,7 @@ void FBlueprintVarActionDetails::CustomizeDetails( IDetailLayoutBuilder& DetailL
 		.ToolTip(VarTypeTooltip)
 	];
 
-	TSharedPtr<SToolTip> EditableTooltip = IDocumentation::Get()->CreateToolTip(LOCTEXT("VarEditableTooltip", "Whether this variable is publically editable on instances of this Blueprint."), NULL, DocLink, TEXT("Editable"));
+	TSharedPtr<SToolTip> EditableTooltip = IDocumentation::Get()->CreateToolTip(LOCTEXT("VarEditableTooltip", "Whether this variable is publicly editable on instances of this Blueprint."), NULL, DocLink, TEXT("Editable"));
 
 	Category.AddCustomRow( LOCTEXT("IsVariableEditableLabel", "Editable") )
 	.Visibility(TAttribute<EVisibility>(this, &FBlueprintVarActionDetails::ShowEditableCheckboxVisibilty))
@@ -1835,12 +1832,6 @@ void FBlueprintVarActionDetails::OnPostEditorRefresh()
 {
 	CachedVariableProperty = SelectionAsProperty();
 	CachedVariableName = GetVariableName();
-
-	// Help in tracking UE-20390
-	if (CachedVariableProperty.IsValid())
-	{
-		UE_LOG(LogBlueprintUserMessages, Log, TEXT("PostEditorRefresh for Blueprint '%s', cached variable property's owner is '%s'"), *GetBlueprintObj()->GetName(), *CachedVariableProperty->GetOwnerClass()->GetPathName());
-	}
 }
 
 EVisibility FBlueprintVarActionDetails::GetTransientVisibility() const
@@ -2898,7 +2889,7 @@ void FBlueprintGraphActionDetails::SetNetFlags( TWeakObjectPtr<UK2Node_EditableP
 
 			if (UK2Node_FunctionEntry* TypedEntryNode = Cast<UK2Node_FunctionEntry>(FunctionEntryNode.Get()))
 			{
-				int32 ExtraFlags = TypedEntryNode->GetFunctionFlags();
+				int32 ExtraFlags = TypedEntryNode->GetExtraFlags();
 				ExtraFlags &= ~FlagsToClear;
 				ExtraFlags |= FlagsToSet;
 				TypedEntryNode->SetExtraFlags(ExtraFlags);
@@ -3837,7 +3828,7 @@ void FBlueprintGraphActionDetails::OnAccessSpecifierSelected( TSharedPtr<FAccess
 		const uint32 ClearAccessSpecifierMask = ~FUNC_AccessSpecifiers;
 		if(UK2Node_FunctionEntry* EntryNode = Cast<UK2Node_FunctionEntry>(FunctionEntryNode))
 		{
-			int32 ExtraFlags = EntryNode->GetFunctionFlags();
+			int32 ExtraFlags = EntryNode->GetExtraFlags();
 			ExtraFlags &= ClearAccessSpecifierMask;
 			ExtraFlags |= SpecifierName->SpecifierFlag;
 			EntryNode->SetExtraFlags(ExtraFlags);
@@ -3931,7 +3922,7 @@ void FBlueprintGraphActionDetails::OnIsReliableReplicationFunctionModified(const
 		{
 			if (UK2Node_FunctionEntry* TypedEntryNode = Cast<UK2Node_FunctionEntry>(FunctionEntryNode))
 			{
-				TypedEntryNode->SetExtraFlags(TypedEntryNode->GetFunctionFlags() | FUNC_NetReliable);
+				TypedEntryNode->AddExtraFlags(FUNC_NetReliable);
 			}
 			if (UK2Node_CustomEvent * CustomEventNode = Cast<UK2Node_CustomEvent>(FunctionEntryNode))
 			{
@@ -3942,7 +3933,7 @@ void FBlueprintGraphActionDetails::OnIsReliableReplicationFunctionModified(const
 		{
 			if (UK2Node_FunctionEntry* TypedEntryNode = Cast<UK2Node_FunctionEntry>(FunctionEntryNode))
 			{
-				TypedEntryNode->SetExtraFlags(TypedEntryNode->GetFunctionFlags() & ~FUNC_NetReliable);
+				TypedEntryNode->ClearExtraFlags(FUNC_NetReliable);
 			}
 			if (UK2Node_CustomEvent * CustomEventNode = Cast<UK2Node_CustomEvent>(FunctionEntryNode))
 			{
@@ -4001,7 +3992,7 @@ void FBlueprintGraphActionDetails::OnIsPureFunctionModified( const ECheckBoxStat
 
 		//set flags on function entry node also
 		Function->FunctionFlags ^= FUNC_BlueprintPure;
-		EntryNode->SetExtraFlags(Function->FunctionFlags);
+		EntryNode->SetExtraFlags(EntryNode->GetExtraFlags() ^ FUNC_BlueprintPure);
 		OnParamsChanged(FunctionEntryNode);
 	}
 }
@@ -4045,7 +4036,7 @@ void FBlueprintGraphActionDetails::OnIsConstFunctionModified( const ECheckBoxSta
 
 		//set flags on function entry node also
 		Function->FunctionFlags ^= FUNC_Const;
-		EntryNode->SetExtraFlags(Function->FunctionFlags);
+		EntryNode->SetExtraFlags(EntryNode->GetExtraFlags() ^ FUNC_Const);
 		OnParamsChanged(FunctionEntryNode);
 	}
 }
@@ -4491,9 +4482,9 @@ void FBlueprintGlobalOptionsDetails::CustomizeDetails(IDetailLayoutBuilder& Deta
 		for (TFieldIterator<UProperty> PropertyIt(Blueprint->GetClass(), EFieldIteratorFlags::IncludeSuper); PropertyIt; ++PropertyIt)
 		{
 			UProperty* Property = *PropertyIt;
-			FText Category = FObjectEditorUtils::GetCategoryText(Property);
+			FString Category = Property->GetMetaData(TEXT("Category"));
 
-			if ( Category.ToString() != TEXT("BlueprintOptions") && Category.ToString() != TEXT("ClassOptions") )
+			if (Category != TEXT("BlueprintOptions") && Category != TEXT("ClassOptions") )
 			{
 				DetailLayout.HideProperty(DetailLayout.GetProperty(Property->GetFName()));
 			}
diff --git a/Engine/Source/Editor/Kismet/Private/BlueprintEditor.cpp b/Engine/Source/Editor/Kismet/Private/BlueprintEditor.cpp
--- a/Engine/Source/Editor/Kismet/Private/BlueprintEditor.cpp
+++ b/Engine/Source/Editor/Kismet/Private/BlueprintEditor.cpp
@@ -59,6 +59,7 @@
 
 // Debugging
 #include "Debugging/SKismetDebuggingView.h"
+#include "Profiler/SBlueprintProfilerView.h"
 #include "Debugging/KismetDebugCommands.h"
 // End of debugging
 
@@ -507,6 +508,7 @@ bool FBlueprintEditor::IsInAScriptingMode() const
 
 bool FBlueprintEditor::OnRequestClose()
 {
+	TabManager->InvokeTab(FBlueprintEditorTabs::BlueprintProfilerID)->RequestCloseTab();
 	bEditorMarkedAsClosed = true;
 	return FWorkflowCentricApplication::OnRequestClose();
 }
@@ -657,14 +659,14 @@ void FBlueprintEditor::RefreshEditors(ERefreshBlueprintEditorReason::Type Reason
 		FocusInspectorOnGraphSelection(GetSelectedNodes(), /*bForceRefresh=*/ true);
 	}
 
-	if (MyBlueprintWidget.IsValid())
+	if (ReplaceReferencesWidget.IsValid())
 	{
-		MyBlueprintWidget->Refresh();
+		ReplaceReferencesWidget->Refresh();
 	}
 
-	if (ReplaceReferencesWidget.IsValid())
+	if (MyBlueprintWidget.IsValid())
 	{
-		ReplaceReferencesWidget->Refresh();
+		MyBlueprintWidget->Refresh();
 	}
 
 	if (SCSEditor.IsValid())
@@ -1671,6 +1673,17 @@ void FBlueprintEditor::InitBlueprintEditor(
 	// Post-layout initialization
 	PostLayoutBlueprintEditorInitialization();
 
+	// Ensure the profiler UI respects the current state if it had previously been enabled.
+	const bool bShowBlueprintProfilerTab = IsProfilerActive() && GetDefault<UEditorExperimentalSettings>()->bUnifiedBlueprintEditor;
+	if (bShowBlueprintProfilerTab)
+	{
+		TabManager->InvokeTab(FBlueprintEditorTabs::BlueprintProfilerID);
+	}
+	else
+	{
+		TabManager->InvokeTab(FBlueprintEditorTabs::BlueprintProfilerID)->RequestCloseTab();
+	}
+
 	// Find and set any instances of this blueprint type if any exists and we are not already editing one
 	FBlueprintEditorUtils::FindAndSetDebuggableBlueprintInstances();
 
@@ -2180,6 +2193,12 @@ void FBlueprintEditor::CreateDefaultTabContents(const TArray<UBlueprint*>& InBlu
 			. BlueprintToWatch(InBlueprint)
 			. IsEnabled(!bIsInterface && !bIsMacro);
 
+		if (GetDefault<UEditorExperimentalSettings>()->bBlueprintPerformanceAnalysisTools)
+		{
+			this->BlueprintProfiler =
+				SNew(SBlueprintProfilerView)
+				.ProfileViewType(PVT_LiveView);
+		}
 		this->Palette = 
 			SNew(SBlueprintPalette, SharedThis(this))
 				.IsEnabled(this, &FBlueprintEditor::IsFocusedGraphEditable);
@@ -2443,6 +2462,15 @@ void FBlueprintEditor::CreateDefaultCommands()
 		FCanExecuteAction::CreateSP(this, &FBlueprintEditor::HasAnyWatches)
 		);
 
+	// Blueprint Profiler Commands
+	if (GetDefault<UEditorExperimentalSettings>()->bBlueprintPerformanceAnalysisTools)
+	{
+		ToolkitCommands->MapAction(FFullBlueprintEditorCommands::Get().ToggleProfiler,
+			FExecuteAction::CreateSP(this, &FBlueprintEditor::ToggleProfiler),
+			FCanExecuteAction(),
+			FIsActionChecked::CreateSP(this, &FBlueprintEditor::IsProfilerActive));
+	}
+
 	// New document actions
 	ToolkitCommands->MapAction( FBlueprintEditorCommands::Get().AddNewVariable,
 		FExecuteAction::CreateSP(this, &FBlueprintEditor::OnAddNewVariable),
@@ -2526,6 +2554,34 @@ void FBlueprintEditor::CreateDefaultCommands()
 		FCanExecuteAction::CreateSP(this, &FBlueprintEditor::CanGenerateNativeCode));
 }
 
+bool FBlueprintEditor::IsProfilerActive() const
+{
+	if (GetDefault<UEditorExperimentalSettings>()->bBlueprintPerformanceAnalysisTools)
+	{
+		IBlueprintProfilerInterface& ProfilerModule = FModuleManager::LoadModuleChecked<IBlueprintProfilerInterface>("BlueprintProfiler");
+		return ProfilerModule.IsProfilerEnabled();
+	}
+	return false;
+}
+
+void FBlueprintEditor::ToggleProfiler()
+{
+	if (GetDefault<UEditorExperimentalSettings>()->bBlueprintPerformanceAnalysisTools)
+	{
+		IBlueprintProfilerInterface& ProfilerModule = FModuleManager::LoadModuleChecked<IBlueprintProfilerInterface>("BlueprintProfiler");
+		ProfilerModule.ToggleProfilingCapture();
+	
+		if (ProfilerModule.IsProfilerEnabled())
+		{
+			TabManager->InvokeTab(FBlueprintEditorTabs::BlueprintProfilerID);
+		}
+		else
+		{
+			TabManager->InvokeTab(FBlueprintEditorTabs::BlueprintProfilerID)->RequestCloseTab();
+		}
+	}
+}
+
 void FBlueprintEditor::OpenNativeCodeGenerationTool()
 {
 	auto Blueprint = GetBlueprintObj();
diff --git a/Engine/Source/Editor/Kismet/Private/BlueprintEditorModes.cpp b/Engine/Source/Editor/Kismet/Private/BlueprintEditorModes.cpp
--- a/Engine/Source/Editor/Kismet/Private/BlueprintEditorModes.cpp
+++ b/Engine/Source/Editor/Kismet/Private/BlueprintEditorModes.cpp
@@ -73,6 +73,7 @@ TSharedPtr<FTabManager::FLayout> GetDefaltEditorLayout(TSharedPtr<class FBluepri
 					->SetSizeCoefficient( 0.20f )
 					->AddTab( FBlueprintEditorTabs::CompilerResultsID, ETabState::ClosedTab )
 					->AddTab( FBlueprintEditorTabs::FindResultsID, ETabState::ClosedTab )
+					->AddTab( FBlueprintEditorTabs::BlueprintProfilerID, ETabState::ClosedTab )
 				)
 			)
 			->Split
@@ -111,6 +112,10 @@ FBlueprintEditorApplicationMode::FBlueprintEditorApplicationMode(TSharedPtr<clas
 	}
 	BlueprintEditorTabFactories.RegisterFactory(MakeShareable(new FCompilerResultsSummoner(InBlueprintEditor)));
 	BlueprintEditorTabFactories.RegisterFactory(MakeShareable(new FFindResultsSummoner(InBlueprintEditor)));
+	if (GetDefault<UEditorExperimentalSettings>()->bBlueprintPerformanceAnalysisTools)
+	{
+		BlueprintEditorTabFactories.RegisterFactory(MakeShareable(new FBlueprintProfilerSummoner(InBlueprintEditor)));
+	}
 	
 	if( bRegisterViewport )
 	{
@@ -131,6 +136,10 @@ FBlueprintEditorApplicationMode::FBlueprintEditorApplicationMode(TSharedPtr<clas
 	InBlueprintEditor->GetToolbarBuilder()->AddScriptingToolbar(ToolbarExtender);
 	InBlueprintEditor->GetToolbarBuilder()->AddBlueprintGlobalOptionsToolbar(ToolbarExtender);
 	InBlueprintEditor->GetToolbarBuilder()->AddDebuggingToolbar(ToolbarExtender);
+	if (GetDefault<UEditorExperimentalSettings>()->bBlueprintPerformanceAnalysisTools)
+	{
+		InBlueprintEditor->GetToolbarBuilder()->AddProfilerToolbar(ToolbarExtender);
+	}
 }
 
 void FBlueprintEditorApplicationMode::RegisterTabFactories(TSharedPtr<FTabManager> InTabManager)
@@ -572,6 +581,10 @@ FBlueprintEditorUnifiedMode::FBlueprintEditorUnifiedMode(TSharedPtr<class FBluep
 	}
 	BlueprintEditorTabFactories.RegisterFactory(MakeShareable(new FCompilerResultsSummoner(InBlueprintEditor)));
 	BlueprintEditorTabFactories.RegisterFactory(MakeShareable(new FFindResultsSummoner(InBlueprintEditor)));
+	if (GetDefault<UEditorExperimentalSettings>()->bBlueprintPerformanceAnalysisTools)
+	{
+		BlueprintEditorTabFactories.RegisterFactory(MakeShareable(new FBlueprintProfilerSummoner(InBlueprintEditor)));
+	}
 	
 	if( bRegisterViewport )
 	{
@@ -631,6 +644,7 @@ FBlueprintEditorUnifiedMode::FBlueprintEditorUnifiedMode(TSharedPtr<class FBluep
 						->SetSizeCoefficient( 0.20f )
 						->AddTab( FBlueprintEditorTabs::CompilerResultsID, ETabState::ClosedTab )
 						->AddTab( FBlueprintEditorTabs::FindResultsID, ETabState::ClosedTab )
+						->AddTab( FBlueprintEditorTabs::BlueprintProfilerID, ETabState::ClosedTab )
 					)
 				)
 				->Split
@@ -691,6 +705,7 @@ FBlueprintEditorUnifiedMode::FBlueprintEditorUnifiedMode(TSharedPtr<class FBluep
 						->SetSizeCoefficient( 0.20f )
 						->AddTab( FBlueprintEditorTabs::CompilerResultsID, ETabState::ClosedTab )
 						->AddTab( FBlueprintEditorTabs::FindResultsID, ETabState::ClosedTab )
+						->AddTab( FBlueprintEditorTabs::BlueprintProfilerID, ETabState::ClosedTab )
 					)
 				)
 				->Split
@@ -715,6 +730,10 @@ FBlueprintEditorUnifiedMode::FBlueprintEditorUnifiedMode(TSharedPtr<class FBluep
 	InBlueprintEditor->GetToolbarBuilder()->AddCompileToolbar(ToolbarExtender);
 	InBlueprintEditor->GetToolbarBuilder()->AddScriptingToolbar(ToolbarExtender);
 	InBlueprintEditor->GetToolbarBuilder()->AddBlueprintGlobalOptionsToolbar(ToolbarExtender);
+	if (GetDefault<UEditorExperimentalSettings>()->bBlueprintPerformanceAnalysisTools)
+	{
+		InBlueprintEditor->GetToolbarBuilder()->AddProfilerToolbar(ToolbarExtender);
+	}
 	
 	if ( bRegisterViewport )
 	{
diff --git a/Engine/Source/Editor/Kismet/Private/BlueprintEditorTabFactories.cpp b/Engine/Source/Editor/Kismet/Private/BlueprintEditorTabFactories.cpp
--- a/Engine/Source/Editor/Kismet/Private/BlueprintEditorTabFactories.cpp
+++ b/Engine/Source/Editor/Kismet/Private/BlueprintEditorTabFactories.cpp
@@ -8,6 +8,7 @@
 #include "BlueprintEditor.h"
 #include "STimelineEditor.h"
 #include "Debugging/SKismetDebuggingView.h"
+#include "Profiler/SBlueprintProfilerView.h"
 #include "SKismetInspector.h"
 #include "SSCSEditor.h"
 #include "SSCSEditorViewport.h"
@@ -135,6 +136,24 @@ TSharedRef<SWidget> FDebugInfoSummoner::CreateTabBody(const FWorkflowTabSpawnInf
 	return BlueprintEditorPtr->GetDebuggingView();
 }
 
+FBlueprintProfilerSummoner::FBlueprintProfilerSummoner(TSharedPtr<class FAssetEditorToolkit> InHostingApp) : FWorkflowTabFactory(FBlueprintEditorTabs::BlueprintProfilerID, InHostingApp)
+{
+	TabLabel = LOCTEXT("BlueprintProfilerViewTitle", "BlueprintProfiler");
+	TabIcon = FSlateIcon(FEditorStyle::GetStyleSetName(), "PerfTools.TabIcon");
+
+	EnableTabPadding();
+	bIsSingleton = true;
+
+	ViewMenuDescription = LOCTEXT("BlueprintProfilerView", "Profiler");
+	ViewMenuTooltip = LOCTEXT("BlueprintProfilerView_ToolTip", "Shows the current profiler data");
+}
+
+TSharedRef<SWidget> FBlueprintProfilerSummoner::CreateTabBody(const FWorkflowTabSpawnInfo& Info) const
+{
+	TSharedPtr<FBlueprintEditor> BlueprintEditorPtr = StaticCastSharedPtr<FBlueprintEditor>(HostingApp.Pin());
+	return BlueprintEditorPtr->GetBlueprintProfilerView();
+}
+
 FDefaultsEditorSummoner::FDefaultsEditorSummoner(TSharedPtr<class FAssetEditorToolkit> InHostingApp)
 	: FWorkflowTabFactory(FBlueprintEditorTabs::DefaultEditorID, InHostingApp)
 {
diff --git a/Engine/Source/Editor/Kismet/Private/BlueprintEditorTabFactories.h b/Engine/Source/Editor/Kismet/Private/BlueprintEditorTabFactories.h
--- a/Engine/Source/Editor/Kismet/Private/BlueprintEditorTabFactories.h
+++ b/Engine/Source/Editor/Kismet/Private/BlueprintEditorTabFactories.h
@@ -58,6 +58,17 @@ public:
 	virtual TSharedRef<SWidget> CreateTabBody(const FWorkflowTabSpawnInfo& Info) const override;
 };
 
+/////////////////////////////////////////////////////
+// FBlueprintProfilerSummoner
+
+struct FBlueprintProfilerSummoner : public FWorkflowTabFactory
+{
+public:
+	FBlueprintProfilerSummoner(TSharedPtr<class FAssetEditorToolkit> InHostingApp);
+
+	virtual TSharedRef<SWidget> CreateTabBody(const FWorkflowTabSpawnInfo& Info) const override;
+};
+
 /////////////////////////////////////////////////////
 // FDefaultsEditorSummoner
 
diff --git a/Engine/Source/Editor/Kismet/Private/BlueprintEditorTabs.cpp b/Engine/Source/Editor/Kismet/Private/BlueprintEditorTabs.cpp
--- a/Engine/Source/Editor/Kismet/Private/BlueprintEditorTabs.cpp
+++ b/Engine/Source/Editor/Kismet/Private/BlueprintEditorTabs.cpp
@@ -12,6 +12,7 @@ const FName FBlueprintEditorTabs::FindResultsID( TEXT( "FindResults" ) );
 const FName FBlueprintEditorTabs::ConstructionScriptEditorID( TEXT( "ConstructionScriptEditor" ) );
 const FName FBlueprintEditorTabs::SCSViewportID( TEXT( "SCSViewport" ) );
 const FName FBlueprintEditorTabs::MyBlueprintID( TEXT( "MyBlueprint" ) );
+const FName FBlueprintEditorTabs::BlueprintProfilerID( TEXT( "BlueprintProfiler" ) );
 const FName FBlueprintEditorTabs::ReplaceNodeReferencesID( TEXT( "ReplaceNodeReferences" ) );
 
 const FName FBlueprintEditorTabs::GraphEditorID( TEXT( "GraphEditor" ) );
diff --git a/Engine/Source/Editor/Kismet/Private/FiBSearchInstance.cpp b/Engine/Source/Editor/Kismet/Private/FiBSearchInstance.cpp
--- a/Engine/Source/Editor/Kismet/Private/FiBSearchInstance.cpp
+++ b/Engine/Source/Editor/Kismet/Private/FiBSearchInstance.cpp
@@ -248,7 +248,7 @@ void FFiBSearchInstance::BuildFunctionTargetsByName(TSharedPtr<FImaginaryFiBData
 			{
 				OutTargetPendingSearchables.Add(ChildData);
 			}
-			else
+			else if (CategoryData->IsTagAndValueCategory())
 			{
 				BuildFunctionTargetsByName(ChildData, InTagName, OutTargetPendingSearchables);
 			}
diff --git a/Engine/Source/Editor/Kismet/Private/FindInBlueprintManager.cpp b/Engine/Source/Editor/Kismet/Private/FindInBlueprintManager.cpp
--- a/Engine/Source/Editor/Kismet/Private/FindInBlueprintManager.cpp
+++ b/Engine/Source/Editor/Kismet/Private/FindInBlueprintManager.cpp
@@ -720,9 +720,7 @@ namespace BlueprintSearchMetaDataHelpers
 		{
 			if (!InProperty->HasMetaData(*FFiBMD::FiBSearchableMD) || InProperty->GetBoolMetaData(*FFiBMD::FiBSearchableMD))
 			{
-				InWriter->WriteObjectStart(FText::FromString(InProperty->GetName()));
 				GatherSearchableProperties(InWriter, InValue, StructProperty->Struct, SEARCHABLE_FULL);
-				InWriter->WriteObjectEnd();
 			}
 		}
 		else if (UObjectProperty* ObjectProperty = Cast<UObjectProperty>(InProperty))
@@ -740,9 +738,7 @@ namespace BlueprintSearchMetaDataHelpers
 					bool bSearchable = InProperty->GetBoolMetaData(*FFiBMD::FiBSearchableMD);
 					if (bSearchable)
 					{
-						InWriter->WriteObjectStart(FText::FromString(InProperty->GetName()));
 						GatherSearchableProperties(InWriter, SubObject, SubObject->GetClass(), SEARCHABLE_FULL);
-						InWriter->WriteObjectEnd();
 					}
 				}
 				else
@@ -778,31 +774,35 @@ namespace BlueprintSearchMetaDataHelpers
 					const void* Value = Property->ContainerPtrToValuePtr<uint8>(InValue);
 
 					// Need to store the metadata on the property in a sub-object
-					InWriter->WriteObjectStart(FFindInBlueprintSearchTags::FiBMetaDataTag);
+					InWriter->WriteObjectStart(FText::FromString(Property->GetName()));
 					{
-						if (Property->GetBoolMetaData(*FFiBMD::FiBSearchableHiddenExplicitMD))
+						InWriter->WriteObjectStart(FFindInBlueprintSearchTags::FiBMetaDataTag);
 						{
-							InWriter->WriteValue(FText::FromString(FFiBMD::FiBSearchableHiddenExplicitMD), true);
+							if (Property->GetBoolMetaData(*FFiBMD::FiBSearchableHiddenExplicitMD))
+							{
+								InWriter->WriteValue(FText::FromString(FFiBMD::FiBSearchableHiddenExplicitMD), true);
+							}
+							else if (Property->GetBoolMetaData(*FFiBMD::FiBSearchableExplicitMD))
+							{
+								InWriter->WriteValue(FText::FromString(FFiBMD::FiBSearchableExplicitMD), true);
+							}
 						}
-						else if (Property->GetBoolMetaData(*FFiBMD::FiBSearchableExplicitMD))
+						InWriter->WriteObjectEnd();
+
+						if (Property->ArrayDim == 1)
 						{
-							InWriter->WriteValue(FText::FromString(FFiBMD::FiBSearchableExplicitMD), true);
+							GatherSearchablesFromProperty(InWriter, Property, Value, InStruct);
 						}
-					}
-					InWriter->WriteObjectEnd();
-
-					if (Property->ArrayDim == 1)
-					{
-						GatherSearchablesFromProperty(InWriter, Property, Value, InStruct);
-					}
-					else
-					{
-						TArray< TSharedPtr<FJsonValue> > Array;
-						for (int Index = 0; Index != Property->ArrayDim; ++Index)
+						else
 						{
-							GatherSearchablesFromProperty(InWriter, Property, (char*)Value + Index * Property->ElementSize, InStruct);
+							TArray< TSharedPtr<FJsonValue> > Array;
+							for (int Index = 0; Index != Property->ArrayDim; ++Index)
+							{
+								GatherSearchablesFromProperty(InWriter, Property, (char*)Value + Index * Property->ElementSize, InStruct);
+							}
 						}
 					}
+					InWriter->WriteObjectEnd();
 				}
 			}
 		}
@@ -941,9 +941,9 @@ class FCacheAllBlueprintsTickableObject
 
 public:
 
-	FCacheAllBlueprintsTickableObject(TArray<FName> InUncachedBlueprints, bool bInCheckOutAndSave, FSimpleDelegate InOnFinished = FSimpleDelegate())
+	FCacheAllBlueprintsTickableObject(TSet<FName> InUncachedBlueprints, bool bInCheckOutAndSave, FSimpleDelegate InOnFinished = FSimpleDelegate())
 		: TickCacheIndex(0)
-		, UncachedBlueprints(InUncachedBlueprints)
+		, UncachedBlueprints(InUncachedBlueprints.Array())
 		, bIsStarted(false)
 		, bIsCancelled(false)
 		, bRecursionGuard(false)
@@ -991,6 +991,24 @@ public:
 		return (float)TickCacheIndex / (float)UncachedBlueprints.Num();
 	}
 
+	/** Returns the number of uncached Blueprints */
+	int32 GetUncachedBlueprintCount()
+	{
+		return UncachedBlueprints.Num();
+	}
+
+	/** Returns the entire list of uncached Blueprints that this object will attempt to cache */
+	const TArray<FName>& GetUncachedBlueprintList() const
+	{
+		return UncachedBlueprints;
+	}
+
+	/** True if there is a callback when done caching, this will prevent a re-query from occuring */
+	bool HasPostCacheWork() const
+	{
+		return OnFinished.IsBound();
+	}
+
 	/** Cancels caching and destroys this object */
 	void OnCancelCaching(bool bIsImmediate)
 	{
@@ -1064,34 +1082,47 @@ public:
 				UObject* Asset = AssetData.GetAsset();
 				if(Asset && bCheckOutAndSave)
 				{
-					// Assume the package was correctly checked out from SCC
-					bool bOutPackageLocallyWritable = true;
-
-					UPackage* Package = AssetData.GetPackage();
-
-					ISourceControlProvider& SourceControlProvider = ISourceControlModule::Get().GetProvider();
-					// Trusting the SCC status in the package file cache to minimize network activity during save.
-					const FSourceControlStatePtr SourceControlState = SourceControlProvider.GetState(Package, EStateCacheUsage::Use);
-					// If the package is in the depot, and not recognized as editable by source control, and not read-only, then we know the user has made the package locally writable!
-					const bool bSCCCanEdit = !SourceControlState.IsValid() || SourceControlState->CanCheckIn() || SourceControlState->IsIgnored() || SourceControlState->IsUnknown();
-					const bool bSCCIsCheckedOut = SourceControlState.IsValid() && SourceControlState->IsCheckedOut();
-					const bool bInDepot = SourceControlState.IsValid() && SourceControlState->IsSourceControlled();
-					if ( !bSCCCanEdit && bInDepot && !bIsAssetReadOnlyOnDisk && SourceControlProvider.UsesLocalReadOnlyState() && !bSCCIsCheckedOut )
+					if (UBlueprint* BlueprintAsset = Cast<UBlueprint>(Asset))
 					{
-						bOutPackageLocallyWritable = false;
+						if (BlueprintAsset->SkeletonGeneratedClass == nullptr)
+						{
+							// There is no skeleton class, something was wrong with the Blueprint during compile on load. This asset will be marked as failing to cache.
+							bFailedToCache = false;
+						}
 					}
 
-					// Save the package if the file is writable
-					if(bOutPackageLocallyWritable)
+					// Still good to attempt to save
+					if (bFailedToCache)
 					{
-						UWorld* WorldAsset = Cast<UWorld>(Asset);
+						// Assume the package was correctly checked out from SCC
+						bool bOutPackageLocallyWritable = true;
+
+						UPackage* Package = AssetData.GetPackage();
+
+						ISourceControlProvider& SourceControlProvider = ISourceControlModule::Get().GetProvider();
+						// Trusting the SCC status in the package file cache to minimize network activity during save.
+						const FSourceControlStatePtr SourceControlState = SourceControlProvider.GetState(Package, EStateCacheUsage::Use);
+						// If the package is in the depot, and not recognized as editable by source control, and not read-only, then we know the user has made the package locally writable!
+						const bool bSCCCanEdit = !SourceControlState.IsValid() || SourceControlState->CanCheckIn() || SourceControlState->IsIgnored() || SourceControlState->IsUnknown();
+						const bool bSCCIsCheckedOut = SourceControlState.IsValid() && SourceControlState->IsCheckedOut();
+						const bool bInDepot = SourceControlState.IsValid() && SourceControlState->IsSourceControlled();
+						if (!bSCCCanEdit && bInDepot && !bIsAssetReadOnlyOnDisk && SourceControlProvider.UsesLocalReadOnlyState() && !bSCCIsCheckedOut)
+						{
+							bOutPackageLocallyWritable = false;
+						}
 
-						// Save the package
-						EObjectFlags ObjectFlags = (WorldAsset == nullptr)? RF_Standalone : RF_NoFlags;
-						
-						if (GEditor->SavePackage(Package, WorldAsset, ObjectFlags, *FinalPackageFilename))
+						// Save the package if the file is writable
+						if (bOutPackageLocallyWritable)
 						{
-							bFailedToCache = false;
+							UWorld* WorldAsset = Cast<UWorld>(Asset);
+
+							// Save the package
+							EObjectFlags ObjectFlags = (WorldAsset == nullptr) ? RF_Standalone : RF_NoFlags;
+
+							if (GEditor->SavePackage(Package, WorldAsset, ObjectFlags, *FinalPackageFilename, GError, nullptr, false, true, SAVE_NoError))
+							{
+								bFailedToCache = false;
+							}
 						}
 					}
 				}
@@ -1112,11 +1143,11 @@ public:
 
 				ProgressNotification.Pin()->SetText(LOCTEXT("BlueprintIndexComplete", "Finished indexing Blueprints!"));
 
-				FFindInBlueprintSearchManager::Get().FinishedCachingBlueprints(TickCacheIndex, FailedToCacheList);
-
 				// We have actually finished, use the OnFinished callback.
 				OnFinished.ExecuteIfBound();
 
+				FFindInBlueprintSearchManager::Get().FinishedCachingBlueprints(TickCacheIndex, FailedToCacheList);
+
 				return EActiveTimerReturnType::Stop;
 			}
 			else
@@ -1140,8 +1171,8 @@ private:
 	/** Notification that appears and details progress */
 	TWeakPtr<SNotificationItem> ProgressNotification;
 
-	/** List of Blueprints that failed to be saved */
-	TArray<FName> FailedToCacheList;
+	/** Set of Blueprints that failed to be saved */
+	TSet<FName> FailedToCacheList;
 
 	/** TRUE if the caching process is started */
 	bool bIsStarted;
@@ -1171,16 +1202,21 @@ FFindInBlueprintSearchManager& FFindInBlueprintSearchManager::Get()
 }
 
 FFindInBlueprintSearchManager::FFindInBlueprintSearchManager()
-	: bIsPausing(false)
+	: bEnableGatheringData(true)
+	, bIsPausing(false)
+	, AssetRegistryModule(nullptr)
 	, CachingObject(nullptr)
 {
 }
 
 FFindInBlueprintSearchManager::~FFindInBlueprintSearchManager()
 {
-	AssetRegistryModule->Get().OnAssetAdded().RemoveAll(this);
-	AssetRegistryModule->Get().OnAssetRemoved().RemoveAll(this);
-	AssetRegistryModule->Get().OnAssetRenamed().RemoveAll(this);
+	if (AssetRegistryModule)
+	{
+		AssetRegistryModule->Get().OnAssetAdded().RemoveAll(this);
+		AssetRegistryModule->Get().OnAssetRemoved().RemoveAll(this);
+		AssetRegistryModule->Get().OnAssetRenamed().RemoveAll(this);
+	}
 	FCoreUObjectDelegates::PreGarbageCollect.RemoveAll(this);
 	FCoreUObjectDelegates::PostGarbageCollect.RemoveAll(this);
 	FCoreUObjectDelegates::OnAssetLoaded.RemoveAll(this);
@@ -1194,10 +1230,25 @@ FFindInBlueprintSearchManager::~FFindInBlueprintSearchManager()
 
 void FFindInBlueprintSearchManager::Initialize()
 {
-	AssetRegistryModule = &FModuleManager::LoadModuleChecked<FAssetRegistryModule>(TEXT("AssetRegistry"));
-	AssetRegistryModule->Get().OnAssetAdded().AddRaw(this, &FFindInBlueprintSearchManager::OnAssetAdded);
-	AssetRegistryModule->Get().OnAssetRemoved().AddRaw(this, &FFindInBlueprintSearchManager::OnAssetRemoved);
-	AssetRegistryModule->Get().OnAssetRenamed().AddRaw(this, &FFindInBlueprintSearchManager::OnAssetRenamed);
+	// Must ensure we do not attempt to load the AssetRegistry Module while saving a package, however, if it is loaded already we can safely obtain it
+	if (!GIsSavingPackage || (GIsSavingPackage && FModuleManager::Get().IsModuleLoaded(TEXT("AssetRegistry"))))
+	{
+		AssetRegistryModule = &FModuleManager::LoadModuleChecked<FAssetRegistryModule>(TEXT("AssetRegistry"));
+		AssetRegistryModule->Get().OnAssetAdded().AddRaw(this, &FFindInBlueprintSearchManager::OnAssetAdded);
+		AssetRegistryModule->Get().OnAssetRemoved().AddRaw(this, &FFindInBlueprintSearchManager::OnAssetRemoved);
+		AssetRegistryModule->Get().OnAssetRenamed().AddRaw(this, &FFindInBlueprintSearchManager::OnAssetRenamed);
+	}
+	else
+	{
+		// Log a warning to inform the Asset Registry could not be initialized when FiB initialized due to saving package
+		// The Asset Registry should be initialized before Find-in-Blueprints, or FiB should be explicitly initialized during a safe time
+		// This message will not appear in commandlets because most commandlets do not care. If a search query is made, further warnings will be produced even in commandlets.
+		if (!IsRunningCommandlet())
+		{
+			UE_LOG(LogBlueprint, Warning, TEXT("Find-in-Blueprints could not pre-cache all unloaded Blueprints due to the Asset Registry module being unable to initialize because a package is currently being saved. Pre-cache will not be reattempted!"));
+		}
+	}
+
 	FCoreUObjectDelegates::PreGarbageCollect.AddRaw(this, &FFindInBlueprintSearchManager::PauseFindInBlueprintSearch);
 	FCoreUObjectDelegates::PostGarbageCollect.AddRaw(this, &FFindInBlueprintSearchManager::UnpauseFindInBlueprintSearch);
 	FCoreUObjectDelegates::OnAssetLoaded.AddRaw(this, &FFindInBlueprintSearchManager::OnAssetLoaded);
@@ -1206,7 +1257,7 @@ void FFindInBlueprintSearchManager::Initialize()
 	IHotReloadInterface& HotReloadSupport = FModuleManager::LoadModuleChecked<IHotReloadInterface>("HotReload");
 	HotReloadSupport.OnHotReload().AddRaw(this, &FFindInBlueprintSearchManager::OnHotReload);
 
-	if(!GIsSavingPackage)
+	if(!GIsSavingPackage && AssetRegistryModule)
 	{
 		// Do an immediate load of the cache to catch any Blueprints that were discovered by the asset registry before we initialized.
 		BuildCache();
@@ -1411,6 +1462,7 @@ void FFindInBlueprintSearchManager::OnAssetLoaded(UObject* InAsset)
 			ensureMsgf(!SearchArray[*IndexPtr].Blueprint.IsValid() || SearchArray[*IndexPtr].BlueprintPath == BlueprintPath, TEXT("Blueprint in database has path %s and is being stomped by %s"), *(SearchArray[*IndexPtr].BlueprintPath.ToString()), *BlueprintPath.ToString());
 			SearchArray[*IndexPtr].Blueprint = BlueprintAsset;
 		}
+		UncachedBlueprints.Remove(BlueprintPath);
 	}
 }
 
@@ -1493,6 +1545,11 @@ FString FFindInBlueprintSearchManager::GatherBlueprintSearchMetadata(const UBlue
 
 void FFindInBlueprintSearchManager::AddOrUpdateBlueprintSearchMetadata(UBlueprint* InBlueprint, bool bInForceReCache/* = false*/)
 {
+	if (!bEnableGatheringData)
+	{
+		return;
+	}
+
 	check(InBlueprint);
 
 	// Allow only one thread modify the search data at a time
@@ -1541,6 +1598,11 @@ void FFindInBlueprintSearchManager::AddOrUpdateBlueprintSearchMetadata(UBlueprin
 
 void FFindInBlueprintSearchManager::BeginSearchQuery(const FStreamSearch* InSearchOriginator)
 {
+	if (AssetRegistryModule == nullptr)
+	{
+		UE_LOG(LogBlueprint, Warning, TEXT("Find-in-Blueprints was not fully initialized, possibly due to problems being initialized while saving a package. Please explicitly initialize earlier!"));
+	}
+
 	// Cannot begin a search thread while saving
 	FScopeLock ScopeLock(&PauseThreadsCriticalSection);
 	FScopeLock ScopeLock2(&SafeQueryModifyCriticalSection);
@@ -1584,7 +1646,7 @@ bool FFindInBlueprintSearchManager::ContinueSearchQuery(const FStreamSearch* InS
 				// If there is FiB data, parse it into an ImaginaryBlueprint
 				if (SearchArray[SearchIdx].Value.Len() > 0)
 				{
-					SearchArray[SearchIdx].ImaginaryBlueprint = MakeShareable(new FImaginaryBlueprint(FPaths::GetBaseFilename(SearchArray[SearchIdx].BlueprintPath.ToString()), SearchArray[SearchIdx].BlueprintPath.ToString(), SearchArray[SearchIdx].ParentClass, SearchArray[SearchIdx].Interfaces, SearchArray[SearchIdx].Value));
+					SearchArray[SearchIdx].ImaginaryBlueprint = MakeShareable(new FImaginaryBlueprint(FPaths::GetBaseFilename(SearchArray[SearchIdx].BlueprintPath.ToString()), SearchArray[SearchIdx].BlueprintPath.ToString(), SearchArray[SearchIdx].ParentClass, SearchArray[SearchIdx].Interfaces, SearchArray[SearchIdx].Value, SearchArray[SearchIdx].Version != 0));
 					SearchArray[SearchIdx].Value.Empty();
 				}
 
@@ -1722,7 +1784,7 @@ void FFindInBlueprintSearchManager::CleanCache()
 		else
 		{
 			// Level Blueprints are destroyed when you open a new level, we need to re-add it as an unloaded asset so long as they were not marked for deletion
-			if(!SearchArray[SearchValuePair.Value].bMarkedForDeletion)
+			if(!SearchArray[SearchValuePair.Value].bMarkedForDeletion && FModuleManager::Get().IsModuleLoaded(TEXT("AssetRegistry")))
 			{
 				SearchArray[SearchValuePair.Value].Blueprint = nullptr;
 
@@ -1821,8 +1883,10 @@ void FFindInBlueprintSearchManager::OnCacheAllUncachedBlueprints(bool bInSourceC
 	if(bInSourceControlActive && bInCheckoutAndSave)
 	{
 		TArray<FString> UncachedBlueprintStrings;
-		UncachedBlueprintStrings.Reserve(UncachedBlueprints.Num());
-		for (const FName& UncachedBlueprint : UncachedBlueprints)
+		const TArray<FName>& TotalUncachedBlueprints = CachingObject->GetUncachedBlueprintList();
+		
+		UncachedBlueprintStrings.Reserve(TotalUncachedBlueprints.Num());
+		for (const FName& UncachedBlueprint : TotalUncachedBlueprints)
 		{
 			UncachedBlueprintStrings.Add(UncachedBlueprint.ToString());
 		}
@@ -1878,11 +1942,9 @@ void FFindInBlueprintSearchManager::CacheAllUncachedBlueprints(TWeakPtr< SFindIn
 		// If Yes is chosen, checkout and save all Blueprints, if No is chosen, only load all Blueprints
 		if (ReturnValue != EAppReturnType::Cancel)
 		{
-			// Add all failed to cache Blueprints to the UncachedBlueprints list and resubmit for caching
-			UncachedBlueprints.Append(FailedToCachePaths);
 			FailedToCachePaths.Empty();
 			
-			TArray<FName> TempUncachedBlueprints;
+			TSet<FName> TempUncachedBlueprints;
 			TempUncachedBlueprints.Append(UncachedBlueprints);
 			TempUncachedBlueprints.Append(BlueprintsToUpdate);
 
@@ -1937,7 +1999,7 @@ FName FFindInBlueprintSearchManager::GetCurrentCacheBlueprintName() const
 	return CachingBPName;
 }
 
-float  FFindInBlueprintSearchManager::GetCacheProgress() const
+float FFindInBlueprintSearchManager::GetCacheProgress() const
 {
 	float ReturnCacheValue = 1.0f;
 
@@ -1949,18 +2011,26 @@ float  FFindInBlueprintSearchManager::GetCacheProgress() const
 	return ReturnCacheValue;
 }
 
-void FFindInBlueprintSearchManager::FinishedCachingBlueprints(int32 InNumberCached, TArray<FName>& InFailedToCacheList)
+int32 FFindInBlueprintSearchManager::GetNumberUncachedBlueprints() const
+{
+	int32 ReturnCount = UncachedBlueprints.Num();
+	if (CachingObject)
+	{
+		ReturnCount = CachingObject->GetUncachedBlueprintCount();
+	}
+	return ReturnCount;
+}
+
+void FFindInBlueprintSearchManager::FinishedCachingBlueprints(int32 InNumberCached, TSet<FName>& InFailedToCacheList)
 {
 	// Multiple threads could be adding to this at the same time
 	FScopeLock ScopeLock(&SafeModifyCacheCriticalSection);
 
+	// Update the list of cache failures
 	FailedToCachePaths = InFailedToCacheList;
 
-	// Clean out the list by the count that were successfully cached, this allows us to continue the caching later.
-	UncachedBlueprints.RemoveAt(0, InNumberCached);
-
 	// Signal the callback, so the source FindInBlueprint can resubmit their search queries
-	if(SourceCachingWidget.IsValid())
+	if(SourceCachingWidget.IsValid() && !CachingObject->HasPostCacheWork())
 	{
 		SourceCachingWidget.Pin()->OnCacheComplete();
 	}
@@ -1976,7 +2046,7 @@ bool FFindInBlueprintSearchManager::IsCacheInProgress() const
 	return CachingObject != nullptr;
 }
 
-TSharedPtr< FJsonObject > FFindInBlueprintSearchManager::ConvertJsonStringToObject(FString InJsonString, TMap<int32, FText>& OutFTextLookupTable)
+TSharedPtr< FJsonObject > FFindInBlueprintSearchManager::ConvertJsonStringToObject(bool bInIsVersioned, FString InJsonString, TMap<int32, FText>& OutFTextLookupTable)
 {
 	/** The searchable data is more complicated than a Json string, the Json being the main searchable body that is parsed. Below is a diagram of the full data:
 	 *  | int32 "Version" | int32 "Size" | TMap "Lookup Table" | Json String |
@@ -1993,7 +2063,11 @@ TSharedPtr< FJsonObject > FFindInBlueprintSearchManager::ConvertJsonStringToObje
 	int32 SizeOfData;
 	FBufferReader ReaderStream((void*)*InJsonString, InJsonString.Len() * sizeof(TCHAR), false);
 
-	int32 Version = FiBSerializationHelpers::Deserialize<int32>(ReaderStream);
+	int32 Version = 0;
+	if (bInIsVersioned)
+	{
+		FiBSerializationHelpers::Deserialize<int32>(ReaderStream);
+	}
 
  	// Read, as a byte string, the number of characters composing the Lookup Table for the Json.
 	SizeOfData = FiBSerializationHelpers::Deserialize<int32>(ReaderStream);
diff --git a/Engine/Source/Editor/Kismet/Private/FindInBlueprints.cpp b/Engine/Source/Editor/Kismet/Private/FindInBlueprints.cpp
--- a/Engine/Source/Editor/Kismet/Private/FindInBlueprints.cpp
+++ b/Engine/Source/Editor/Kismet/Private/FindInBlueprints.cpp
@@ -618,7 +618,7 @@ void SFindInBlueprints::ConditionallyAddCacheBar()
 		{
 			// Create a single string of all the Blueprint paths that failed to cache, on separate lines
 			FString PackageList;
-			TArray<FName> FailedToCacheList = FFindInBlueprintSearchManager::Get().GetFailedToCachePathList();
+			TSet<FName> FailedToCacheList = FFindInBlueprintSearchManager::Get().GetFailedToCachePathList();
 			for (FName Package : FailedToCacheList)
 			{
 				PackageList += Package.ToString() + TEXT("\n");
@@ -1238,7 +1238,7 @@ FText SFindInBlueprints::GetUncachedBlueprintWarningText() const
 
 	// The number of unindexed Blueprints is the total of those that failed to cache and those that haven't been attempted yet.
 	FFormatNamedArguments Args;
-	Args.Add(TEXT("Count"), FindInBlueprintManager.GetNumberUncachedBlueprints() + FailedToCacheCount + OutOfDateWithLastSearchBPCount);
+	Args.Add(TEXT("Count"), FindInBlueprintManager.GetNumberUncachedBlueprints() + OutOfDateWithLastSearchBPCount);
 
 	FText ReturnDisplayText;
 	if(IsCacheInProgress())
diff --git a/Engine/Source/Editor/Kismet/Private/ImaginaryBlueprintData.cpp b/Engine/Source/Editor/Kismet/Private/ImaginaryBlueprintData.cpp
--- a/Engine/Source/Editor/Kismet/Private/ImaginaryBlueprintData.cpp
+++ b/Engine/Source/Editor/Kismet/Private/ImaginaryBlueprintData.cpp
@@ -120,7 +120,7 @@ void FImaginaryFiBData::ParseAllChildData_Internal(ESearchableValueStatus InSear
 
 	if (UnparsedJsonObject.IsValid())
 	{
-		if (InSearchabilityOverride == ESearchableValueStatus::Searchable)
+		if (InSearchabilityOverride & ESearchableValueStatus::Searchable)
 		{
 			const TSharedPtr< FJsonObject >* MetaDataField;
 			if (UnparsedJsonObject->TryGetObjectField(FFindInBlueprintSearchTags::FiBMetaDataTag.ToString(), MetaDataField))
@@ -286,6 +286,8 @@ UObject* FImaginaryFiBData::GetObject(UBlueprint* InBlueprint) const
 
 FFiBMetaData::FFiBMetaData(TWeakPtr<FImaginaryFiBData> InOuter, TSharedPtr< FJsonObject > InUnparsedJsonObject, TMap<int32, FText>* InLookupTablePtr)
 	: FImaginaryFiBData(InOuter, InUnparsedJsonObject, InLookupTablePtr)
+	, bIsHidden(false)
+	, bIsExplicit(false)
 {
 }
 
@@ -352,6 +354,26 @@ void FCategorySectionHelper::ParseAllChildData_Internal(ESearchableValueStatus I
 	}
 	else
 	{
+		bool bHasMetaData = false;
+		bool bHasOneOtherItem = false;
+		if (UnparsedJsonObject.IsValid() && UnparsedJsonObject->Values.Num() == 2)
+		{
+			for( auto MapValues : UnparsedJsonObject->Values )
+			{
+				if (MapValues.Key == FFindInBlueprintSearchTags::FiBMetaDataTag.ToString())
+				{
+					bHasMetaData = true;
+				}
+				else
+				{
+					bHasOneOtherItem = true;
+				}
+			}
+
+			// If we have metadata and only one other item, we should be treated like a tag and value category
+			bIsTagAndValue |= (bHasOneOtherItem && bHasMetaData);
+		}
+
 		FImaginaryFiBData::ParseAllChildData_Internal(InSearchabilityOverride);
 	}
 }
@@ -359,12 +381,12 @@ void FCategorySectionHelper::ParseAllChildData_Internal(ESearchableValueStatus I
 //////////////////////////////////////////
 // FImaginaryBlueprint
 
-FImaginaryBlueprint::FImaginaryBlueprint(FString InBlueprintName, FString InBlueprintPath, FString InBlueprintParentClass, TArray<FString>& InInterfaces, FString InUnparsedStringData)
+FImaginaryBlueprint::FImaginaryBlueprint(FString InBlueprintName, FString InBlueprintPath, FString InBlueprintParentClass, TArray<FString>& InInterfaces, FString InUnparsedStringData, bool bInIsVersioned/*=true*/)
 	: FImaginaryFiBData(nullptr)
 	, BlueprintPath(InBlueprintPath)
 	, UnparsedStringData(InUnparsedStringData)
 {
-	ParseToJson();
+	ParseToJson(bInIsVersioned);
 	ParsedTagsAndValues.Add(FFindInBlueprintSearchTags::FiB_Name.ToString(), FSearchableValueInfo(FFindInBlueprintSearchTags::FiB_Name, FText::FromString(InBlueprintName), ESearchableValueStatus::ExplicitySearchable));
 	ParsedTagsAndValues.Add(FFindInBlueprintSearchTags::FiB_Path.ToString(), FSearchableValueInfo(FFindInBlueprintSearchTags::FiB_Path, FText::FromString(InBlueprintPath), ESearchableValueStatus::ExplicitySearchable));
 	ParsedTagsAndValues.Add(FFindInBlueprintSearchTags::FiB_ParentClass.ToString(), FSearchableValueInfo(FFindInBlueprintSearchTags::FiB_ParentClass, FText::FromString(InBlueprintParentClass), ESearchableValueStatus::ExplicitySearchable));
@@ -409,9 +431,9 @@ bool FImaginaryBlueprint::CanCallFilter(ESearchQueryFilter InSearchQueryFilter)
 		FImaginaryFiBData::CanCallFilter(InSearchQueryFilter);
 }
 
-void FImaginaryBlueprint::ParseToJson()
+void FImaginaryBlueprint::ParseToJson(bool bInIsVersioned)
 {
-	UnparsedJsonObject = FFindInBlueprintSearchManager::ConvertJsonStringToObject(UnparsedStringData, LookupTable);
+	UnparsedJsonObject = FFindInBlueprintSearchManager::ConvertJsonStringToObject(bInIsVersioned, UnparsedStringData, LookupTable);
 }
 
 bool FImaginaryBlueprint::TrySpecialHandleJsonValue(FString InKey, TSharedPtr< FJsonValue > InJsonValue)
diff --git a/Engine/Source/Editor/Kismet/Private/Profiler/BPProfilerStatisticWidgets.cpp b/Engine/Source/Editor/Kismet/Private/Profiler/BPProfilerStatisticWidgets.cpp
new file mode 100644
--- /dev/null
+++ b/Engine/Source/Editor/Kismet/Private/Profiler/BPProfilerStatisticWidgets.cpp
@@ -0,0 +1,728 @@
+// Copyright 1998-2015 Epic Games, Inc. All Rights Reserved.
+
+#include "BlueprintEditorPrivatePCH.h"
+#include "SBlueprintProfilerView.h"
+#include "BPProfilerStatisticWidgets.h"
+#include "SHyperlink.h"
+
+#define LOCTEXT_NAMESPACE "BlueprintProfilerViewTypesUI"
+
+//////////////////////////////////////////////////////////////////////////
+// BlueprintProfilerStatText
+
+namespace BlueprintProfilerStatText
+{
+	const FName ColumnId_Name( "Name" );
+	const FName ColumnId_InclusiveTime( "InclusiveTime" );
+	const FName ColumnId_ExclusiveTime( "ExclusiveTime" );
+	const FName ColumnId_MaxTime( "MaxTime" );
+	const FName ColumnId_MinTime( "MinTime" );
+	const FName ColumnId_Samples( "Samples" );
+	const FName ColumnId_TotalTime( "TotalTime" );
+
+	const FText ColumnText_Name( NSLOCTEXT("BlueprintProfilerViewUI", "Name", "Name") );
+	const FText ColumnText_InclusiveTime( NSLOCTEXT("BlueprintProfilerViewUI", "InclusiveTime", "Inclusive Time (ms)") );
+	const FText ColumnText_ExclusiveTime( NSLOCTEXT("BlueprintProfilerViewUI", "ExclusiveTime", "Exclusive Time (ms)") );
+	const FText ColumnText_MaxTime( NSLOCTEXT("BlueprintProfilerViewUI", "MaxTime", "Max Time (ms)") );
+	const FText ColumnText_MinTime( NSLOCTEXT("BlueprintProfilerViewUI", "MinTime", "Min Time (ms)") );
+	const FText ColumnText_Samples( NSLOCTEXT("BlueprintProfilerViewUI", "Samples", "Samples") );
+	const FText ColumnText_TotalTime( NSLOCTEXT("BlueprintProfilerViewUI", "TotalTime", "Total Time (s)") );
+};
+
+//////////////////////////////////////////////////////////////////////////
+// SProfilerStatRow
+
+TSharedRef<SWidget> SProfilerStatRow::GenerateWidgetForColumn( const FName& ColumnName )
+{
+	if (ColumnName == BlueprintProfilerStatText::ColumnId_Name)
+	{
+		return SNew(SHorizontalBox)
+			+SHorizontalBox::Slot()
+			.AutoWidth()
+			.HAlign(HAlign_Left)
+			.VAlign(VAlign_Center)
+			[
+				SNew(SExpanderArrow, SharedThis(this))
+			]
+			+SHorizontalBox::Slot()
+			.FillWidth(1.0f)
+			.HAlign(HAlign_Left)
+			.VAlign(VAlign_Center)
+			[
+				ItemToEdit->GenerateNameWidget()
+			];
+	}
+	else
+	{
+		return SNew(STextBlock)
+			.Text(ItemToEdit->GetTextAttribute(ColumnName));
+	}
+}
+
+void SProfilerStatRow::Construct(const FArguments& InArgs, TSharedRef<STableViewBase> OwnerTableView, FBPProfilerStatPtr InItemToEdit)
+{
+	check(InItemToEdit.IsValid());
+	ItemToEdit = InItemToEdit;
+	SMultiColumnTableRow<FBPProfilerStatPtr>::Construct( FSuperRowType::FArguments(), OwnerTableView );	
+}
+
+const FName SProfilerStatRow::GetStatName(const EBlueprintProfilerStat::Type StatId)
+{
+	switch(StatId)
+	{
+		case EBlueprintProfilerStat::Name:			return BlueprintProfilerStatText::ColumnId_Name;
+		case EBlueprintProfilerStat::TotalTime:		return BlueprintProfilerStatText::ColumnId_TotalTime;
+		case EBlueprintProfilerStat::InclusiveTime:	return BlueprintProfilerStatText::ColumnId_InclusiveTime;
+		case EBlueprintProfilerStat::ExclusiveTime:	return BlueprintProfilerStatText::ColumnId_ExclusiveTime;
+		case EBlueprintProfilerStat::MaxTime:		return BlueprintProfilerStatText::ColumnId_MaxTime;
+		case EBlueprintProfilerStat::MinTime:		return BlueprintProfilerStatText::ColumnId_MinTime;
+		case EBlueprintProfilerStat::Samples:		return BlueprintProfilerStatText::ColumnId_Samples;
+		default:									return NAME_None;
+	}
+}
+
+const FText SProfilerStatRow::GetStatText(const EBlueprintProfilerStat::Type StatId)
+{
+	switch(StatId)
+	{
+		case EBlueprintProfilerStat::Name:			return BlueprintProfilerStatText::ColumnText_Name;
+		case EBlueprintProfilerStat::TotalTime:		return BlueprintProfilerStatText::ColumnText_TotalTime;
+		case EBlueprintProfilerStat::InclusiveTime:	return BlueprintProfilerStatText::ColumnText_InclusiveTime;
+		case EBlueprintProfilerStat::ExclusiveTime:	return BlueprintProfilerStatText::ColumnText_ExclusiveTime;
+		case EBlueprintProfilerStat::MaxTime:		return BlueprintProfilerStatText::ColumnText_MaxTime;
+		case EBlueprintProfilerStat::MinTime:		return BlueprintProfilerStatText::ColumnText_MinTime;
+		case EBlueprintProfilerStat::Samples:		return BlueprintProfilerStatText::ColumnText_Samples;
+		default:									return FText::GetEmpty();
+	}
+}
+
+//////////////////////////////////////////////////////////////////////////
+// FBPPerformanceData
+
+int32 FBPPerformanceData::StatSampleSize = 20;
+
+void FBPPerformanceData::AddEventTiming(const double NodeTiming)
+{
+	const double TimingInMs = NodeTiming * 1000;
+	Samples.WriteNewElementInitialized() = TimingInMs;
+	TotalTime += NodeTiming;
+	NumSamples++;
+	MinTime = FMath::Min<double>(MinTime, TimingInMs);
+	MaxTime = FMath::Max<double>(MaxTime, TimingInMs);
+	bDirty = true;
+}
+
+void FBPPerformanceData::Update()
+{
+	if (bDirty)
+	{
+		CachedExclusiveTime = 0.0;
+		for (int32 Idx = 0; Idx < Samples.Num(); ++Idx)
+		{
+			CachedExclusiveTime += Samples(Idx);
+		}
+		CachedExclusiveTime /= GetSafeSampleCount();
+		bDirty = false;
+	}
+}
+
+void FBPPerformanceData::Reset()
+{
+	MaxTime = -MAX_dbl;
+	MinTime = MAX_dbl;
+	CachedExclusiveTime = 0.0;
+	TotalTime = 0.0;
+	NumSamples = 0;
+	bDirty = false;
+}
+
+//////////////////////////////////////////////////////////////////////////
+// FBPProfilerStat
+
+TAttribute<FText> FBPProfilerStat::GetTextAttribute(FName ColumnName) const
+{
+	if (ColumnName == BlueprintProfilerStatText::ColumnId_TotalTime)
+	{
+		return TAttribute<FText>(this, &FBPProfilerStat::GetTotalTimeText);
+	}
+	else if (ColumnName == BlueprintProfilerStatText::ColumnId_InclusiveTime)
+	{
+		return TAttribute<FText>(this, &FBPProfilerStat::GetInclusiveTimeText);
+	}
+	else if (ColumnName == BlueprintProfilerStatText::ColumnId_ExclusiveTime)
+	{
+		return TAttribute<FText>(this, &FBPProfilerStat::GetExclusiveTimeText);
+	}
+	else if (ColumnName == BlueprintProfilerStatText::ColumnId_MaxTime)
+	{
+		return TAttribute<FText>(this, &FBPProfilerStat::GetMaxTimeText);
+	}
+	else if (ColumnName == BlueprintProfilerStatText::ColumnId_MinTime)
+	{
+		return TAttribute<FText>(this, &FBPProfilerStat::GetMinTimeText);
+	}
+	else if (ColumnName == BlueprintProfilerStatText::ColumnId_Samples)
+	{
+		return TAttribute<FText>(this, &FBPProfilerStat::GetSamplesText);
+	}
+	return TAttribute<FText>();
+}
+
+FText FBPProfilerStat::GetTotalTimeText() const
+{
+	// Covert back up to seconds for total time.
+	return FText::AsNumber(GetTotalTime(), &SBlueprintProfilerView::GetNumberFormat());
+}
+
+FText FBPProfilerStat::GetInclusiveTimeText() const
+{
+	return FText::AsNumber(GetInclusiveTime(), &SBlueprintProfilerView::GetNumberFormat());
+}
+
+FText FBPProfilerStat::GetExclusiveTimeText() const
+{
+	return FText::AsNumber(GetExclusiveTime(), &SBlueprintProfilerView::GetNumberFormat());
+}
+
+FText FBPProfilerStat::GetMaxTimeText() const
+{
+	return FText::AsNumber(GetMaxTime(), &SBlueprintProfilerView::GetNumberFormat());
+}
+
+FText FBPProfilerStat::GetMinTimeText() const
+{
+	return FText::AsNumber(GetMinTime(), &SBlueprintProfilerView::GetNumberFormat());
+}
+
+FText FBPProfilerStat::GetSamplesText() const
+{
+	return FText::AsNumber(GetSampleCount());
+}
+
+double FBPProfilerStat::GetMaxTime() const
+{
+	double Result = 0.0;
+	if (CachedChildren.Num())
+	{
+		Result += MaxInclusiveTime;
+	}
+	else if (PerformanceStats.IsDataValid())
+	{
+		Result += PerformanceStats.GetMaxTime();
+	}
+	return Result;
+}
+
+double FBPProfilerStat::GetMinTime() const
+{
+	double Result = 0.0;
+	if (CachedChildren.Num())
+	{
+		Result += MinInclusiveTime;
+	}
+	else if (PerformanceStats.IsDataValid())
+	{
+		Result += PerformanceStats.GetMinTime();
+	}
+	return Result;
+}
+
+float FBPProfilerStat::GetTotalTime() const
+{
+	float Result = 0.f;
+	if (CachedChildren.Num())
+	{
+		Result += InclusiveTotalTime;
+	}
+	else if (PerformanceStats.IsDataValid())
+	{
+		Result += PerformanceStats.GetTotalTime();
+	}
+	return Result;
+}
+
+float FBPProfilerStat::GetInclusiveTotalTime() const
+{
+	float Result = 0.f;
+	if (CachedChildren.Num())
+	{
+		Result += InclusiveTotalTime;
+	}
+	else if (PerformanceStats.IsDataValid())
+	{
+		Result += PerformanceStats.GetTotalTime();
+	}
+	return Result;
+}
+
+float FBPProfilerStat::GetExclusiveTime() const
+{
+	return PerformanceStats.IsDataValid() ? PerformanceStats.GetExclusiveTime() : 0.f;
+}
+
+float FBPProfilerStat::GetInclusiveTime() const
+{
+	float Result = 0.f;
+	Result += InclusiveTime;
+	return Result;
+}
+
+float FBPProfilerStat::GetSampleCount() const
+{
+	float Result = 0.f;
+	if (PerformanceStats.IsDataValid())
+	{
+		Result += PerformanceStats.GetSampleCount();
+	}
+	else
+	{
+		Result += InclusiveSamples;
+	}
+	return Result;
+}
+
+float FBPProfilerStat::GetSafeSampleCount() const
+{
+	float Result = 0.f;
+	if (PerformanceStats.IsDataValid())
+	{
+		Result += PerformanceStats.GetSampleCount();
+	}
+	return Result != 0.f ? Result : 1.f;
+}
+
+FBPProfilerStatPtr FBPProfilerStat::FindEntryPoint(FScriptExecEvent& EventEntry)
+{
+	FBPProfilerStatPtr Result;
+	if (EventEntry == ObjectContext)
+	{
+		Result = AsShared();
+	}
+	else
+	{
+		for (auto Iter : CachedChildren)
+		{
+			Result = Iter.Value->FindEntryPoint(EventEntry);
+			if (Result.IsValid())
+			{
+				break;
+			}
+		}
+	}
+	return Result;
+}
+
+FBPProfilerStatPtr FBPProfilerStat::FindOrAddChildByContext(FScriptExecEvent& Context)
+{
+	FBPProfilerStatPtr& Result = CachedChildren.FindOrAdd(Context.GetObjectPtr());
+	if (!Result.IsValid())
+	{
+		if (Context.IsBranch())
+		{
+			TSharedPtr<FBPProfilerBranchStat> Branch = MakeShareable(new FBPProfilerBranchStat(AsShared(), Context));
+			Branch->CreateChildrenForPins();
+			Result = StaticCastSharedPtr<FBPProfilerStat>(Branch);
+		}
+		else
+		{
+			Result = MakeShareable(new FBPProfilerStat(AsShared(), Context));
+		}
+	}
+	return Result;
+}
+
+void FBPProfilerStat::SubmitData(class FScriptExecEvent& EventData)
+{
+	check (EventData == ObjectContext);
+	if (EventData.CanSubmitData())
+	{
+		PerformanceStats.AddEventTiming(EventData.GetTime());
+	}
+}
+
+TSharedRef<SWidget> FBPProfilerStat::GenerateNameWidget()
+{
+	return SNew(SHorizontalBox)
+		+SHorizontalBox::Slot()
+		.AutoWidth()
+		[
+			SNew(SImage)
+			.Image(GetIcon())
+			.ColorAndOpacity(GetIconColor())
+		]
+		+SHorizontalBox::Slot()
+		.AutoWidth()
+		.Padding(FMargin(5,0))
+		[
+			SNew(SHyperlink)
+			.Text(StatName)
+			.Style(FEditorStyle::Get(), "HoverOnlyHyperlink")
+		//	.ToolTipText(GetToolTipText())
+			.OnNavigate(this, &FBPProfilerStat::OnNavigateToObject)
+		];
+}
+
+void FBPProfilerStat::GatherChildren(TArray<FBPProfilerStatPtr>& OutChildren)
+{
+	for (auto Child : CachedChildren)
+	{
+		OutChildren.Add(Child.Value);
+	}
+}
+
+void FBPProfilerStat::UpdateProfilingStats()
+{
+	if (PerformanceStats.IsDataDirty() || !PerformanceStats.IsDataValid())
+	{
+		// Update local performance stats
+		PerformanceStats.Update();
+		const bool bIsLocalDataValid = PerformanceStats.IsDataValid();
+
+		switch(ObjectContext.GetType())
+		{
+			case EScriptInstrumentation::Class:
+			case EScriptInstrumentation::Instance:
+			{
+				// Reset stats
+				InclusiveTime		= 0.0;
+				InclusiveTotalTime	= 0.0;
+				InclusiveSamples	= 0.0;
+				// Re-Acquire data from children.
+				for (auto Child : CachedChildren)
+				{
+					FBPProfilerStatPtr& ChildPtr = Child.Value;
+					ChildPtr->UpdateProfilingStats();
+					InclusiveTime += ChildPtr->GetInclusiveTime();
+					InclusiveTotalTime += ChildPtr->GetInclusiveTotalTime();
+					InclusiveSamples += !bIsLocalDataValid ? ChildPtr->GetInclusiveSampleCount() : 0;
+				}
+				MaxInclusiveTime = FMath::Max<double>(MaxInclusiveTime, InclusiveTime);
+				MinInclusiveTime = FMath::Min<double>(MinInclusiveTime, InclusiveTime);
+				break;
+			}
+			case EScriptInstrumentation::Event:
+			{
+				// Reset stats
+				InclusiveTime		= 0.0;
+				InclusiveTotalTime	= 0.0;
+				InclusiveSamples	= bIsLocalDataValid ? PerformanceStats.GetSampleCount() : 0.0;
+				// Re-Acquire data from children.
+				for (auto Child : CachedChildren)
+				{
+					FBPProfilerStatPtr& ChildPtr = Child.Value;
+					ChildPtr->UpdateProfilingStats();
+					InclusiveTime += ChildPtr->GetInclusiveTime();
+					InclusiveTotalTime += ChildPtr->GetInclusiveTotalTime();
+					InclusiveSamples += !bIsLocalDataValid ? ChildPtr->GetInclusiveSampleCount() : 0;
+				}
+				MaxInclusiveTime = FMath::Max<double>(MaxInclusiveTime, InclusiveTime);
+				MinInclusiveTime = FMath::Min<double>(MinInclusiveTime, InclusiveTime);
+				InclusiveSamples	= bIsLocalDataValid ? PerformanceStats.GetSampleCount() : 0.0;
+
+				break;
+			}
+			default:
+			{
+				// Reset stats
+				InclusiveTime		= bIsLocalDataValid ? PerformanceStats.GetExclusiveTime() : 0.0;
+				InclusiveTotalTime	= bIsLocalDataValid ? PerformanceStats.GetTotalTime() : 0.0;
+				InclusiveSamples	= bIsLocalDataValid ? PerformanceStats.GetSampleCount() : 0.0;
+				// Re-Acquire data from children.
+				for (auto Child : CachedChildren)
+				{
+					FBPProfilerStatPtr& ChildPtr = Child.Value;
+					ChildPtr->UpdateProfilingStats();
+					InclusiveTime += ChildPtr->GetInclusiveTime();
+					InclusiveTotalTime += ChildPtr->GetInclusiveTotalTime();
+					InclusiveSamples += !bIsLocalDataValid ? ChildPtr->GetInclusiveSampleCount() : 0;
+				}
+				MaxInclusiveTime = FMath::Max<double>(MaxInclusiveTime, InclusiveTime);
+				MinInclusiveTime = FMath::Min<double>(MinInclusiveTime, InclusiveTime);
+				break;
+			}
+		}
+	}
+}
+
+void FBPProfilerStat::UpdateStatName()
+{
+	switch(ObjectContext.GetType())
+	{
+		case EScriptInstrumentation::Instance:
+		{
+			TWeakObjectPtr<const UObject> ObjectPtr = ObjectContext.GetObjectPtr();
+			if (ObjectPtr.IsValid())
+			{
+				if (const AActor* Actor = Cast<AActor>(ObjectPtr.Get()))
+				{
+					StatName = FText::FromString(Actor->GetActorLabel());
+				}
+				else
+				{
+					StatName = FText::FromString(ObjectPtr.Get()->GetName());
+				}
+			}
+			break;
+		}
+		case EScriptInstrumentation::NodePin:
+		{
+			TWeakObjectPtr<const UEdGraphPin> Pin = ObjectContext.GetTypedObjectPtr<UEdGraphPin>();
+			if (Pin.IsValid())
+			{
+				StatName = Pin.Get()->GetDisplayName();
+			}
+			break;
+		}
+		case EScriptInstrumentation::NodeEntry:
+		case EScriptInstrumentation::Event:
+		case EScriptInstrumentation::Function:
+		case EScriptInstrumentation::Branch:
+		case EScriptInstrumentation::Macro:
+		{
+			TWeakObjectPtr<const UEdGraphNode> Node = ObjectContext.GetTypedObjectPtr<UEdGraphNode>();
+			if (Node.IsValid())
+			{
+				StatName = Node.Get()->GetNodeTitle(ENodeTitleType::ListView);
+			}
+			else
+			{
+				TWeakObjectPtr<const UObject> ObjectPtr = ObjectContext.GetObjectPtr();
+				if (ObjectPtr.IsValid())
+				{
+					StatName = FText::FromString(ObjectPtr.Get()->GetName());
+				}
+			}
+			break;
+		}
+		default:
+		{
+			TWeakObjectPtr<const UObject> ObjectPtr = ObjectContext.GetObjectPtr();
+			if (ObjectPtr.IsValid())
+			{
+				StatName = FText::FromString(ObjectPtr.Get()->GetName());
+			}
+			break;
+		}
+	}
+}
+
+const FSlateBrush* FBPProfilerStat::GetIcon()
+{
+	switch(ObjectContext.GetType())
+	{
+		case EScriptInstrumentation::Class:		return FEditorStyle::GetBrush(TEXT("BlueprintProfiler.BPIcon_Normal"));
+		case EScriptInstrumentation::Instance:	return FEditorStyle::GetBrush(TEXT("BlueprintProfiler.Actor"));
+		case EScriptInstrumentation::NodePin:	return CachedChildren.Num() > 0 ? FEditorStyle::GetBrush(TEXT("BlueprintProfiler.BPPinConnected")) : FEditorStyle::GetBrush(TEXT("BlueprintProfiler.BPPinDisconnected"));
+		default:
+		{
+			if (const UEdGraphNode* GraphNode = ObjectContext.GetTypedObjectPtr<UEdGraphNode>().Get())
+			{
+				if (GraphNode->ShowPaletteIconOnNode())
+				{
+					return FEditorStyle::GetBrush(GraphNode->GetPaletteIcon(IconColor));
+				}
+			}
+			return FEditorStyle::GetBrush(TEXT("BlueprintProfiler.BPNode"));
+		}
+	}
+}
+
+FLinearColor FBPProfilerStat::GetIconColor() const
+{
+	switch(ObjectContext.GetType())
+	{
+		case EScriptInstrumentation::Class:		return FLinearColor( 0.46f, 0.54f, 0.81f );
+		case EScriptInstrumentation::Event:		return FLinearColor( 0.91f, 0.16f, 0.16f );
+		case EScriptInstrumentation::Function:	return FLinearColor( 0.46f, 0.54f, 0.95f );
+		case EScriptInstrumentation::Macro:		return FLinearColor( 0.36f, 0.34f, 0.71f );
+		case EScriptInstrumentation::Branch:	return FLinearColor( 1.f, 1.f, 1.f, 0.8f );
+		default:								return FLinearColor( 1.f, 1.f, 1.f, 0.8f );
+	}
+}
+
+const FText FBPProfilerStat::GetToolTip()
+{
+	switch(ObjectContext.GetType())
+	{
+		case EScriptInstrumentation::Class:		return LOCTEXT("NavigateToBlueprintLocationHyperlink_ToolTip", "Navigate to the Blueprint");
+		case EScriptInstrumentation::Instance:	return LOCTEXT("NavigateToInstanceLocationHyperlink_ToolTip", "Navigate to the Object Instance");
+		case EScriptInstrumentation::Event:		return LOCTEXT("NavigateToEventLocationHyperlink_ToolTip", "Navigate to the Event Node");
+		case EScriptInstrumentation::Function:	return LOCTEXT("NavigateToFunctionLocationHyperlink_ToolTip", "Navigate to the Function Callsite");
+		case EScriptInstrumentation::Macro:		return LOCTEXT("NavigateToMacroLocationHyperlink_ToolTip", "Navigate to the Macro Callsite");
+		default:								return LOCTEXT("NavigateToUnknownLocationHyperlink_ToolTip", "Navigate to the Object");
+	}
+}
+
+void FBPProfilerStat::OnNavigateToObject()
+{
+	switch(ObjectContext.GetType())
+	{
+		case EScriptInstrumentation::Class:
+		{
+			break;
+		}
+		case EScriptInstrumentation::Instance:
+		{
+			break;
+		}
+		case EScriptInstrumentation::Event:
+		case EScriptInstrumentation::Function:
+		case EScriptInstrumentation::Branch:
+		case EScriptInstrumentation::Macro:
+		default:
+		{
+			TWeakObjectPtr<const UObject> ObjectPtr = ObjectContext.GetObjectPtr();
+			if (ObjectPtr.IsValid())
+			{
+				FKismetEditorUtilities::BringKismetToFocusAttentionOnObject(ObjectPtr.Get());
+			}
+			break;
+		}
+	}
+}
+
+bool FBPProfilerBranchStat::IsBranchSequential()
+{
+	TWeakObjectPtr<const UObject> ObjectPtr = ObjectContext.GetObjectPtr();
+	bool bIsSequential = false;
+	if (const UObject* Object = ObjectPtr.Get())
+	{
+		if (Object->IsA<UK2Node_ExecutionSequence>())
+		{
+			bIsSequential = true;
+		}
+	}
+	return bIsSequential;
+}
+
+void FBPProfilerBranchStat::CreateChildrenForPins()
+{
+	TWeakObjectPtr<const UObject> ObjectPtr = ObjectContext.GetObjectPtr();
+	bAreExecutionPathsSequential = IsBranchSequential();
+	if (const UEdGraphNode* GraphNode = Cast<const UEdGraphNode>(ObjectPtr.Get()))
+	{
+		for (auto Pin : GraphNode->Pins)
+		{
+			if (Pin->Direction == EGPD_Output && Pin->PinType.PinCategory == UEdGraphSchema_K2::PC_Exec)
+			{
+				FBPProfilerStatPtr& Result = CachedChildren.FindOrAdd(Pin);
+				if (!Result.IsValid())
+				{
+					FScriptExecEvent NewExitPin(Pin, ObjectContext);
+					Result = MakeShareable(new FBPProfilerStat(AsShared(), NewExitPin));
+				}
+			}
+		}
+	}
+}
+
+FBPProfilerStatPtr FBPProfilerBranchStat::FindOrAddChildByContext(FScriptExecEvent& Context)
+{
+	FBPProfilerStatPtr Result;
+	check (CachedExecPin.IsValid());
+	if (FBPProfilerStatPtr* PinExecutionPath = CachedChildren.Find(CachedExecPin))
+	{
+		// Defer child creation down to the execution pin's children.
+		// This is an artificial step to increase and manage the parent child relationship and execution depth
+		// so we can navigate the stats correctly.
+		FBPProfilerStatPtr& NewChild = (*PinExecutionPath)->CachedChildren.FindOrAdd(Context.GetObjectPtr());
+		if (!NewChild.IsValid())
+		{
+			if (Context.IsBranch())
+			{
+				TSharedPtr<FBPProfilerBranchStat> Branch = MakeShareable(new FBPProfilerBranchStat(ParentStat, Context));
+				Branch->CreateChildrenForPins();
+				NewChild = StaticCastSharedPtr<FBPProfilerStat>(Branch);
+			}
+			else
+			{
+				NewChild = MakeShareable(new FBPProfilerStat(ParentStat, Context));
+			}
+		}
+		Result = NewChild;
+	}
+	return Result;
+}
+
+void FBPProfilerBranchStat::UpdateProfilingStats()
+{
+	if (PerformanceStats.IsDataDirty() || !PerformanceStats.IsDataValid())
+	{
+		// Update local performance stats
+		PerformanceStats.Update();
+		// Reset stats
+		const bool bIsLocalDataValid = PerformanceStats.IsDataValid();
+		InclusiveTime		= bIsLocalDataValid ? PerformanceStats.GetExclusiveTime() : 0.0;
+		InclusiveTotalTime	= bIsLocalDataValid ? PerformanceStats.GetTotalTime() : 0.0;
+		InclusiveSamples	= bIsLocalDataValid ? PerformanceStats.GetSampleCount() : 0.0;
+		// Find Inclusive data
+		if (bAreExecutionPathsSequential)
+		{
+			for (auto Child : CachedChildren)
+			{
+				FBPProfilerStatPtr& ChildPtr = Child.Value;
+				ChildPtr->UpdateProfilingStats();
+				InclusiveTime += ChildPtr->GetInclusiveTime();
+				InclusiveTotalTime += ChildPtr->GetInclusiveTotalTime();
+			}
+			MaxInclusiveTime = FMath::Max<double>(MaxInclusiveTime, InclusiveTime);
+			MinInclusiveTime = FMath::Min<double>(MinInclusiveTime, InclusiveTime);
+		}
+		else
+		{
+			double WorstInclusiveTime = 0.0;
+			double WorstInclusiveTotalTime = 0.0;
+			for (auto Child : CachedChildren)
+			{
+				FBPProfilerStatPtr& ChildPtr = Child.Value;
+				ChildPtr->UpdateProfilingStats();
+				const double ChildInclusiveTime = ChildPtr->GetInclusiveTime();
+				WorstInclusiveTime = FMath::Max<double>(WorstInclusiveTime, ChildInclusiveTime);
+				WorstInclusiveTotalTime = FMath::Max<double>(WorstInclusiveTotalTime, ChildPtr->GetInclusiveTotalTime());
+				// pick best/worst case path
+				MaxInclusiveTime = FMath::Max<float>(MaxInclusiveTime, ChildInclusiveTime);
+				MinInclusiveTime = FMath::Min<float>(MinInclusiveTime, ChildInclusiveTime);
+			}
+			InclusiveTime = WorstInclusiveTime;
+		}
+	}
+}
+
+double FBPProfilerBranchStat::GetMaxTime() const
+{
+	double Result = 0.0;
+	if (PerformanceStats.IsDataValid())
+	{
+		Result += PerformanceStats.GetMaxTime();
+	}
+	Result += MaxInclusiveTime;
+	return Result;
+}
+
+double FBPProfilerBranchStat::GetMinTime() const
+{
+	double Result = 0.0;
+	if (PerformanceStats.IsDataValid())
+	{
+		Result += PerformanceStats.GetMinTime();
+	}
+	else
+	{
+		Result += MinInclusiveTime;
+	}
+	return Result;
+}
+
+void FBPProfilerBranchStat::SubmitData(FScriptExecEvent& EventData)
+{
+	check (EventData == ObjectContext);
+	// Cache exec path
+	CachedExecPin = EventData.GetGraphPinPtr();
+	// Submit data
+	if (EventData.CanSubmitData())
+	{
+		PerformanceStats.AddEventTiming(EventData.GetTime());
+	}
+}
+
+#undef LOCTEXT_NAMESPACE
diff --git a/Engine/Source/Editor/Kismet/Private/Profiler/BPProfilerStatisticWidgets.h b/Engine/Source/Editor/Kismet/Private/Profiler/BPProfilerStatisticWidgets.h
new file mode 100644
--- /dev/null
+++ b/Engine/Source/Editor/Kismet/Private/Profiler/BPProfilerStatisticWidgets.h
@@ -0,0 +1,288 @@
+// Copyright 1998-2015 Epic Games, Inc. All Rights Reserved.
+
+#pragma once
+
+#include "ScriptEventExecution.h"
+
+//////////////////////////////////////////////////////////////////////////
+// SProfilerStatRow
+
+/** Enum for each column type in the execution tree widget */
+namespace EBlueprintProfilerStat
+{
+	enum Type
+	{
+		Name = 0,
+		TotalTime,
+		InclusiveTime,
+		ExclusiveTime,
+		MaxTime,
+		MinTime,
+		Samples
+	};
+};
+
+class SProfilerStatRow : public SMultiColumnTableRow<FBPProfilerStatPtr>
+{
+protected:
+
+	/** The profiler statistic associated with this stat row */
+	FBPProfilerStatPtr ItemToEdit;
+
+public:
+
+	SLATE_BEGIN_ARGS(SProfilerStatRow){}
+	SLATE_END_ARGS()
+
+	/** Constructs the statistic row item widget */
+	void Construct(const FArguments& InArgs, TSharedRef<STableViewBase> OwnerTableView, FBPProfilerStatPtr InItemToEdit);
+
+	/** Generates and returns a widget for the specified column name and the associated statistic */
+	virtual TSharedRef<SWidget> GenerateWidgetForColumn(const FName& ColumnName) override;
+
+	/** Converts the profiler stat enum into the associated column FName */
+	static const FName GetStatName(const EBlueprintProfilerStat::Type StatId);
+
+	/** Converts the profiler stat enum into the associated column FText */
+	static const FText GetStatText(const EBlueprintProfilerStat::Type StatId);
+
+};
+
+//////////////////////////////////////////////////////////////////////////
+// FBPPerformanceData
+
+struct FBPPerformanceData
+{
+public:
+
+	FBPPerformanceData()
+		: MaxTime(-MAX_dbl)
+		, MinTime(MAX_dbl)
+		, TotalTime(0.0)
+		, Samples(StatSampleSize)
+		, NumSamples(0)
+		, bDirty(false)
+	{
+	}
+
+	/** Add a single event node timing into the dataset */
+	void AddEventTiming(const double NodeTiming);
+
+	/** Update and average the current timings from the buffer */
+	void Update();
+
+	/** Reset the current data buffer and all derived data stats */
+	void Reset();
+
+	/** Returns if this data structure has any valid data */
+	bool IsDataValid() const { return MaxTime != -MAX_dbl && MinTime != MAX_dbl; }
+
+	/** Marks the data is dirty and signals an update is required */
+	void MarkDataDirty() { bDirty = true; }
+
+	/** Signals that the data is dirty and requires and update */
+	bool IsDataDirty() const { return bDirty; }
+
+	/** Returns the various stats this container holds */
+	double GetMaxTime() const { return MaxTime; }
+	double GetMinTime() const { return MinTime; }
+	double GetTotalTime() const { return TotalTime; }
+	double GetExclusiveTime() const { return CachedExclusiveTime; }
+	int32 GetSampleCount() const { return NumSamples; }
+	int32 GetSafeSampleCount() const { return Samples.Num() ? Samples.Num() : 1; }
+
+public:
+
+	/** Controls how many samples are used to generate the averaged samples */
+	static int32 StatSampleSize;
+
+private:
+
+	/** Sample buffer */
+	TSimpleRingBuffer<double> Samples;
+	/** Cached exclusive timing */
+	double CachedExclusiveTime;
+	/** Max exclusive timing */
+	double MaxTime;
+	/** Min exclusive timing */
+	double MinTime;
+	/** Total time accrued */
+	double TotalTime;
+	/** The current inclusive sample count */
+	int32 NumSamples;
+	/** Flag to indicate the data wants the average recalculated */
+	bool bDirty;
+
+};
+
+//////////////////////////////////////////////////////////////////////////
+// FBPProfilerStat
+
+class FBPProfilerStat : public TSharedFromThis<FBPProfilerStat>
+{
+public:
+
+	FBPProfilerStat(FBPProfilerStatPtr InParentStat)
+		: InclusiveTime(0.0)
+		, MinInclusiveTime(MAX_dbl)
+		, MaxInclusiveTime(-MAX_dbl)
+		, InclusiveTotalTime(0.0)
+		, InclusiveSamples(0)
+		, ParentStat(InParentStat)
+	{
+	}
+
+	FBPProfilerStat(FBPProfilerStatPtr InParentStat, FScriptExecEvent& Context)
+		: ObjectContext(Context)
+		, InclusiveTime(0.0)
+		, MinInclusiveTime(MAX_dbl)
+		, MaxInclusiveTime(-MAX_dbl)
+		, InclusiveTotalTime(0.0)
+		, InclusiveSamples(0)
+		, ParentStat(InParentStat)
+	{
+		UpdateStatName();
+	}
+
+	/** Release Data */
+	void ReleaseData()
+	{
+		// Need to investigate why this needs to be called manually.
+		for (auto Iter : CachedChildren)
+		{
+			Iter.Value->ReleaseData();
+			Iter.Value.Reset();
+		}
+		ParentStat.Reset();
+	}
+
+	/** Find context entry point based on the passed in exec event */
+	FBPProfilerStatPtr FindEntryPoint(FScriptExecEvent& EventEntry);
+
+	/** Returns the parent container statistic */
+	FBPProfilerStatPtr GetParent() { return ParentStat; }
+
+	/** Find or add a child stat based in exec context */
+	virtual FBPProfilerStatPtr FindOrAddChildByContext(FScriptExecEvent& Context);
+
+	/** Gather children for list/tree widget */
+	void GatherChildren(TArray<FBPProfilerStatPtr>& OutChildren);
+
+	/** Update the current profiling statistics */
+	virtual void UpdateProfilingStats();
+
+	/** Returns if this stat has valid data */
+	bool HasValidData() { return PerformanceStats.IsDataValid(); }
+
+	/** Generate name widget */
+	TSharedRef<SWidget> GenerateNameWidget();
+
+	/** Create text attribute for column */
+	TAttribute<FText> GetTextAttribute(FName ColumnName) const;
+
+	/** Returns the exec path object context for this statistic */
+	FScriptExecEvent& GetObjectContext() { return ObjectContext; }
+
+	/** Submits new event data to this profiler statistic */
+	virtual void SubmitData(class FScriptExecEvent& EventData);
+
+	/** Returns the call depth for this statistic */
+	int32 GetCallDepth() const { return ObjectContext.IsValid() ? ObjectContext.GetCallDepth() : -1; }
+
+	// Return current statisic values with safety checking */
+	virtual double GetMaxTime() const;
+	virtual double GetMinTime() const;
+	float GetTotalTime() const;
+	float GetInclusiveTotalTime() const;
+	float GetInclusiveTime() const;
+	float GetExclusiveTime() const;
+	float GetSampleCount() const;
+	float GetSafeSampleCount() const;
+	int32 GetInclusiveSampleCount() const { return InclusiveSamples; }
+
+protected:
+
+	friend class FBPProfilerBranchStat;
+
+	/** Cache the statistic text from the exec event context */
+	virtual void UpdateStatName();
+
+	/** Returns the slate icon from the node referenced by the exec event */
+	const FSlateBrush* GetIcon();
+
+	/** Returns the icon color override */
+	FLinearColor GetIconColor() const;
+
+	/** Returns the statistic tool tip */
+	const FText GetToolTip();
+
+	/** Navigates to the exec event node in the graph */
+	void OnNavigateToObject();
+
+	/** Returns the various statistic text values for the stats */
+	virtual FText GetInclusiveTimeText() const;
+	virtual FText GetExclusiveTimeText() const;
+	virtual FText GetTotalTimeText() const;
+	virtual FText GetMaxTimeText() const;
+	virtual FText GetMinTimeText() const;
+	virtual FText GetSamplesText() const;
+
+protected:
+
+	/** Local profiling statistics */
+	FBPPerformanceData PerformanceStats;
+	/** Container profiling statistics */
+	double InclusiveTime;
+	double MinInclusiveTime;
+	double MaxInclusiveTime;
+	double InclusiveTotalTime;
+	int32 InclusiveSamples;
+	/** Script execution event context */
+	FScriptExecEvent ObjectContext;
+	/** Cached statistic name */
+	FText StatName;
+	/** Statistic icon color */
+	FLinearColor IconColor;
+	/** Parent statistic */
+	FBPProfilerStatPtr ParentStat;
+	/** Child statistics */
+	TMap<TWeakObjectPtr<const UObject>, FBPProfilerStatPtr> CachedChildren;
+
+};
+
+//////////////////////////////////////////////////////////////////////////
+// FBPProfilerBranchStat
+
+class FBPProfilerBranchStat : public FBPProfilerStat
+{
+public:
+
+	FBPProfilerBranchStat(FBPProfilerStatPtr InParentStat, FScriptExecEvent& Context)
+		: FBPProfilerStat(InParentStat, Context)
+		, bAreExecutionPathsSequential(false)
+	{
+	}
+
+	//~ FBPProfilerStat Start
+	virtual FBPProfilerStatPtr FindOrAddChildByContext(FScriptExecEvent& Context) override;
+	void CreateChildrenForPins();
+	virtual void UpdateProfilingStats() override;
+	virtual double GetMaxTime() const override;
+	virtual double GetMinTime() const override;
+	virtual void SubmitData(class FScriptExecEvent& EventData) override;
+	//~ FBPProfilerStat End
+
+protected:
+
+	/** Returns whether different branches execute in sequence rather exclusively */
+	bool IsBranchSequential();
+
+private:
+
+	/** Cached exec pin */
+	TWeakObjectPtr<const UObject> CachedExecPin;
+
+	/** Signals that execution paths are sequential */
+	bool bAreExecutionPathsSequential;
+
+};
\ No newline at end of file
diff --git a/Engine/Source/Editor/Kismet/Private/Profiler/SBlueprintProfilerView.cpp b/Engine/Source/Editor/Kismet/Private/Profiler/SBlueprintProfilerView.cpp
new file mode 100644
--- /dev/null
+++ b/Engine/Source/Editor/Kismet/Private/Profiler/SBlueprintProfilerView.cpp
@@ -0,0 +1,227 @@
+// Copyright 1998-2015 Epic Games, Inc. All Rights Reserved.
+
+
+#include "BlueprintEditorPrivatePCH.h"
+#include "BPProfilerStatisticWidgets.h"
+#include "SBlueprintProfilerView.h"
+
+//////////////////////////////////////////////////////////////////////////
+// SBlueprintProfilerView
+
+FNumberFormattingOptions SBlueprintProfilerView::NumberFormat;
+
+SBlueprintProfilerView::~SBlueprintProfilerView()
+{
+	// Register for Profiler toggle events
+	FBlueprintCoreDelegates::OnToggleScriptProfiler.RemoveAll(this);
+}
+
+void SBlueprintProfilerView::Construct(const FArguments& InArgs)
+{	
+	CurrentViewType = InArgs._ProfileViewType;
+	// Register for Profiler toggle events
+	FBlueprintCoreDelegates::OnToggleScriptProfiler.AddRaw(this, &SBlueprintProfilerView::OnToggleProfiler);
+	// Initialise the number format.
+	NumberFormat.MinimumFractionalDigits = 4;
+	NumberFormat.MaximumFractionalDigits = 4;
+	NumberFormat.UseGrouping = false;
+	// Update the average sample retention count.
+	FBPPerformanceData::StatSampleSize = GetDefault<UEditorExperimentalSettings>()->BlueprintProfilerAverageSampleCount;
+	// Create the profiler view widgets
+	bTreeViewIsDirty = false;
+	CreateProfilerWidgets();
+}
+
+TSharedRef<ITableRow> SBlueprintProfilerView::OnGenerateRowForProfiler(FBPProfilerStatPtr InItem, const TSharedRef<STableViewBase>& OwnerTable)
+{
+	return SNew(SProfilerStatRow, OwnerTable, InItem);
+}
+
+void SBlueprintProfilerView::OnGetChildrenForProfiler(FBPProfilerStatPtr InParent, TArray<FBPProfilerStatPtr>& OutChildren)
+{
+	if (InParent.IsValid())
+	{
+		InParent->GatherChildren(OutChildren);
+	}
+}
+
+void SBlueprintProfilerView::OnToggleProfiler(bool bEnabled)
+{
+	if (!bEnabled)
+	{
+		// Clear profiler data, this requires work as I think the stats are currently leaking.
+		for (auto Iter : RootTreeItems)
+		{
+			Iter->ReleaseData();
+			Iter.Reset();
+		}
+		RootTreeItems.Reset();
+		ClassContexts.Reset();
+		if (ExecutionStatTree.IsValid())
+		{
+			bTreeViewIsDirty = true;
+		}
+	}
+}
+
+void SBlueprintProfilerView::CreateProfilerWidgets()
+{
+	ChildSlot
+	[
+		SNew(SVerticalBox)
+		+SVerticalBox::Slot()
+		.FillHeight(0.2f)
+		[
+			SNew(SBorder)
+			.Padding(0)
+			.BorderImage(FEditorStyle::GetBrush("NoBorder"))
+			.Visibility(this, &SBlueprintProfilerView::IsProfilerHidden)
+			.HAlign(HAlign_Center)
+			.VAlign(VAlign_Center)
+			[
+				SNew(STextBlock)
+				.Text(NSLOCTEXT("InactiveProfilerWidget", "InactiveText", "The Blueprint Profiler is currently Inactive"))
+			]
+		]
+#if 0 // To do stat ranking list
+		+SVerticalBox::Slot()
+		[
+			SAssignNew(BlueprintStatList, SListView<FBPProfilerStatPtr>)
+			.ListItemsSource(&RootTreeItems)
+			.SelectionMode(ESelectionMode::Single)
+			.OnGenerateRow(this, &SBlueprintProfilerView::OnGenerateRowForProfiler)
+//			.OnSelectionChanged(this, &SMessagingTypes::HandleTypeListSelectionChanged)
+			.HeaderRow
+			(
+				SNew(SHeaderRow)
+				+SHeaderRow::Column(SProfilerStatRow::GetStatName(EBlueprintProfilerStat::Name))
+				.DefaultLabel(SProfilerStatRow::GetStatText(EBlueprintProfilerStat::Name))
+				.ManualWidth(450)
+				+SHeaderRow::Column(SProfilerStatRow::GetStatName(EBlueprintProfilerStat::InclusiveTime))
+				.DefaultLabel(SProfilerStatRow::GetStatText(EBlueprintProfilerStat::InclusiveTime))
+				.ManualWidth(120)
+				+SHeaderRow::Column(SProfilerStatRow::GetStatName(EBlueprintProfilerStat::ExclusiveTime))
+				.DefaultLabel(SProfilerStatRow::GetStatText(EBlueprintProfilerStat::ExclusiveTime))
+				.ManualWidth(120)
+				+SHeaderRow::Column(SProfilerStatRow::GetStatName(EBlueprintProfilerStat::MaxTime))
+				.DefaultLabel(SProfilerStatRow::GetStatText(EBlueprintProfilerStat::MaxTime))
+				.ManualWidth(90)
+				+SHeaderRow::Column(SProfilerStatRow::GetStatName(EBlueprintProfilerStat::MinTime))
+				.DefaultLabel(SProfilerStatRow::GetStatText(EBlueprintProfilerStat::MinTime))
+				.ManualWidth(90)
+				+SHeaderRow::Column(SProfilerStatRow::GetStatName(EBlueprintProfilerStat::TotalTime))
+				.DefaultLabel(SProfilerStatRow::GetStatText(EBlueprintProfilerStat::TotalTime))
+				.ManualWidth(80)
+				+SHeaderRow::Column(SProfilerStatRow::GetStatName(EBlueprintProfilerStat::Samples))
+				.DefaultLabel(SProfilerStatRow::GetStatText(EBlueprintProfilerStat::Samples))
+				.ManualWidth(60)
+			)
+		]
+#endif
+		+SVerticalBox::Slot()
+		[
+			SNew(SBorder)
+			.Padding(0)
+			.BorderImage(FEditorStyle::GetBrush("NoBorder"))
+			.Visibility(this, &SBlueprintProfilerView::IsProfilerVisible)
+			[
+				SNew(SVerticalBox)
+				+SVerticalBox::Slot()
+				.AutoHeight()
+				.HAlign(HAlign_Left)
+				.VAlign(VAlign_Center)
+				[
+					SNew(STextBlock)
+					.Text(NSLOCTEXT("BueprintProfilerView", "StatisticTreeTitle", "Execution Tree Performance Statistics"))
+				]
+				+SVerticalBox::Slot()
+				.FillHeight(0.8f)
+				[
+					SAssignNew(ExecutionStatTree, STreeView<FBPProfilerStatPtr>)
+					.TreeItemsSource(&RootTreeItems)
+					.SelectionMode(ESelectionMode::Single)
+					.OnGetChildren(this, &SBlueprintProfilerView::OnGetChildrenForProfiler)
+					.OnGenerateRow(this, &SBlueprintProfilerView::OnGenerateRowForProfiler)
+					.HeaderRow
+					(
+						SNew(SHeaderRow)
+						+SHeaderRow::Column(SProfilerStatRow::GetStatName(EBlueprintProfilerStat::Name))
+						.DefaultLabel(SProfilerStatRow::GetStatText(EBlueprintProfilerStat::Name))
+						.ManualWidth(450)
+						+SHeaderRow::Column(SProfilerStatRow::GetStatName(EBlueprintProfilerStat::InclusiveTime))
+						.DefaultLabel(SProfilerStatRow::GetStatText(EBlueprintProfilerStat::InclusiveTime))
+						.ManualWidth(120)
+						+SHeaderRow::Column(SProfilerStatRow::GetStatName(EBlueprintProfilerStat::ExclusiveTime))
+						.DefaultLabel(SProfilerStatRow::GetStatText(EBlueprintProfilerStat::ExclusiveTime))
+						.ManualWidth(120)
+						+SHeaderRow::Column(SProfilerStatRow::GetStatName(EBlueprintProfilerStat::MaxTime))
+						.DefaultLabel(SProfilerStatRow::GetStatText(EBlueprintProfilerStat::MaxTime))
+						.ManualWidth(90)
+						+SHeaderRow::Column(SProfilerStatRow::GetStatName(EBlueprintProfilerStat::MinTime))
+						.DefaultLabel(SProfilerStatRow::GetStatText(EBlueprintProfilerStat::MinTime))
+						.ManualWidth(90)
+						+SHeaderRow::Column(SProfilerStatRow::GetStatName(EBlueprintProfilerStat::TotalTime))
+						.DefaultLabel(SProfilerStatRow::GetStatText(EBlueprintProfilerStat::TotalTime))
+						.ManualWidth(80)
+						+SHeaderRow::Column(SProfilerStatRow::GetStatName(EBlueprintProfilerStat::Samples))
+						.DefaultLabel(SProfilerStatRow::GetStatText(EBlueprintProfilerStat::Samples))
+						.ManualWidth(60)
+					)
+				]
+			]
+		]
+	];
+}
+
+void SBlueprintProfilerView::Tick(const FGeometry& AllottedGeometry, const double InCurrentTime, const float InDeltaTime)
+{
+	if (GetDefault<UEditorExperimentalSettings>()->bBlueprintPerformanceAnalysisTools)
+	{
+		IBlueprintProfilerInterface* ProfilerInterface = FModuleManager::GetModulePtr<IBlueprintProfilerInterface>("BlueprintProfiler");
+		if (ProfilerInterface && ProfilerInterface->IsProfilerEnabled())
+		{
+			// Find and process execution path timing data.
+			bool bRequestRefresh = false;
+			TArray<FBlueprintInstrumentedEvent>& ProfilerData = ProfilerInterface->GetProfilingData();
+			int32 ExecStart = 0;
+
+			for (int32 EventIdx = 0; EventIdx < ProfilerData.Num(); ++EventIdx)
+			{
+				switch(ProfilerData[EventIdx].GetType())
+				{
+					case EScriptInstrumentation::Class:
+					{
+						CurrentClassPath = ProfilerData[EventIdx].GetObjectPath();
+						ExecStart = EventIdx;
+						break;
+					}
+					case EScriptInstrumentation::Stop:
+					{
+						check(!CurrentClassPath.IsEmpty());
+						FProfilerClassContext* ActiveClass = ClassContexts.Find(CurrentClassPath);
+						const bool bNewClassContext = !ActiveClass;
+						if (!ActiveClass)
+						{
+							ActiveClass = &ClassContexts.Add(CurrentClassPath);
+						}
+						ActiveClass->ProcessExecutionPath(ProfilerData, ExecStart, EventIdx);
+						if (bNewClassContext)
+						{
+							RootTreeItems.Add(ActiveClass->GetRootStat());
+							bRequestRefresh = true;
+						}
+						ExecStart = EventIdx+1;
+						break;
+					}
+				}
+			}
+			// Reset profiler data
+			ProfilerData.Reset();
+		}
+	}
+	// Refresh Tree
+	if (bTreeViewIsDirty && ExecutionStatTree.IsValid())
+	{
+		ExecutionStatTree->RequestTreeRefresh();
+	}
+}
diff --git a/Engine/Source/Editor/Kismet/Private/Profiler/SBlueprintProfilerView.h b/Engine/Source/Editor/Kismet/Private/Profiler/SBlueprintProfilerView.h
new file mode 100644
--- /dev/null
+++ b/Engine/Source/Editor/Kismet/Private/Profiler/SBlueprintProfilerView.h
@@ -0,0 +1,95 @@
+// Copyright 1998-2015 Epic Games, Inc. All Rights Reserved.
+
+#pragma once
+
+#include "Developer/BlueprintProfiler/Public/BlueprintProfilerModule.h"
+#include "ScriptEventExecution.h"
+
+// Shared pointer to a performance stat
+typedef TSharedPtr<class FBPProfilerStat> FBPProfilerStatPtr;
+
+/** Enum to describe type of performance stats view */
+enum EPerformanceViewType
+{
+	PVT_LiveView = 0,
+	PVT_Blueprint,
+	PVT_BlueprintInstance,
+	PVT_BlueprintTrace
+};
+
+//////////////////////////////////////////////////////////////////////////
+// SBlueprintProfilerView
+
+typedef TSharedPtr<class FBPProfilerStat> FBPProfilerStatPtr;
+
+class SBlueprintProfilerView : public SCompoundWidget
+{
+public:
+	SLATE_BEGIN_ARGS( SBlueprintProfilerView )
+		: _ProfileViewType(PVT_Blueprint)
+		{}
+
+		SLATE_ARGUMENT( EPerformanceViewType, ProfileViewType )
+	SLATE_END_ARGS()
+
+
+public:
+
+	~SBlueprintProfilerView();
+
+	void Construct( const FArguments& InArgs );
+
+	// SWidget interface
+	virtual void Tick( const FGeometry& AllottedGeometry, const double InCurrentTime, const float InDeltaTime ) override;
+	// End of SWidget interface
+
+	/** Returns FText number format specifier */
+	static const FNumberFormattingOptions& GetNumberFormat() { return NumberFormat; }
+
+protected:
+
+	/** Called when the profiler state is toggled */
+	void OnToggleProfiler(bool bEnabled);
+
+	/** Called to generate a stat row for the execution stat tree */
+	TSharedRef<ITableRow> OnGenerateRowForProfiler(FBPProfilerStatPtr InItem, const TSharedRef<STableViewBase>& OwnerTable);
+
+	/** Called to generate children for given profiler stat */
+	void OnGetChildrenForProfiler(FBPProfilerStatPtr InParent, TArray<FBPProfilerStatPtr>& OutChildren);
+
+	/** Returns is the profiler widgets are currently visible */
+	EVisibility IsProfilerVisible() const { return RootTreeItems.Num() > 0 ? EVisibility::Visible : EVisibility::Collapsed; }
+
+	/** Returns is the profiler widgets are currently not visible, used to display the inactive message */
+	EVisibility IsProfilerHidden() const { return RootTreeItems.Num() == 0 ? EVisibility::Visible : EVisibility::Collapsed; }
+
+	/** Called to create the child profiler widgets */
+	void CreateProfilerWidgets();
+
+protected:
+
+	/** Number formatting options */
+	static FNumberFormattingOptions NumberFormat;
+
+	/** Current profiler view type */
+	EPerformanceViewType CurrentViewType;
+
+	/** Sorted stat list for identifying expensive items */
+	TSharedPtr<SListView<FBPProfilerStatPtr>> BlueprintStatList;
+
+	/** Sorted stat list for identifying expensive items */
+	TSharedPtr<STreeView<FBPProfilerStatPtr>> ExecutionStatTree;
+
+	/** Source item array for the treeview */
+	TArray<FBPProfilerStatPtr> RootTreeItems;
+
+	/** Active Instance Contexts */
+	TMap<FString, FProfilerClassContext> ClassContexts;
+
+	/** Cached active class statistic path */
+	FString CurrentClassPath;
+
+	/** Flag to control treeview refresh */
+	bool bTreeViewIsDirty;
+
+};
