Clang fixes, fallout from integration from Dev-Blueprints

diff --git a/Engine/Source/Developer/BlueprintProfiler/Public/BlueprintProfilerModule.h b/Engine/Source/Developer/BlueprintProfiler/Public/BlueprintProfilerModule.h
--- a/Engine/Source/Developer/BlueprintProfiler/Public/BlueprintProfilerModule.h
+++ b/Engine/Source/Developer/BlueprintProfiler/Public/BlueprintProfilerModule.h
@@ -2,6 +2,7 @@
 
 #pragma once
 
+#include "ModuleManager.h"
 #include "Script.h"
 #include "BlueprintProfilerSupport.h"
 
diff --git a/Engine/Source/Editor/Kismet/Private/Profiler/BPProfilerStatisticWidgets.h b/Engine/Source/Editor/Kismet/Private/Profiler/BPProfilerStatisticWidgets.h
--- a/Engine/Source/Editor/Kismet/Private/Profiler/BPProfilerStatisticWidgets.h
+++ b/Engine/Source/Editor/Kismet/Private/Profiler/BPProfilerStatisticWidgets.h
@@ -144,6 +144,8 @@ public:
 		UpdateStatName();
 	}
 
+	virtual ~FBPProfilerStat() {}
+
 	/** Release Data */
 	void ReleaseData()
 	{
@@ -263,6 +265,8 @@ public:
 	{
 	}
 
+	virtual ~FBPProfilerBranchStat() {}
+
 	//~ FBPProfilerStat Start
 	virtual FBPProfilerStatPtr FindOrAddChildByContext(FScriptExecEvent& Context) override;
 	void CreateChildrenForPins();
diff --git a/Engine/Source/Editor/Kismet/Private/SMyBlueprint.cpp b/Engine/Source/Editor/Kismet/Private/SMyBlueprint.cpp
--- a/Engine/Source/Editor/Kismet/Private/SMyBlueprint.cpp
+++ b/Engine/Source/Editor/Kismet/Private/SMyBlueprint.cpp
@@ -1218,7 +1218,7 @@ void SMyBlueprint::CollectAllActions(FGraphActionListBuilderBase& OutAllActions)
 	TArray<FName> OverridableFunctionNames;
 
 	// Cache potentially overridable functions
-	UClass* ParentClass = BlueprintObj->SkeletonGeneratedClass? BlueprintObj->SkeletonGeneratedClass->GetSuperClass() : BlueprintObj->ParentClass;
+	UClass* ParentClass = BlueprintObj->SkeletonGeneratedClass ? BlueprintObj->SkeletonGeneratedClass->GetSuperClass() : *BlueprintObj->ParentClass;
 	for ( TFieldIterator<UFunction> FunctionIt(ParentClass, EFieldIteratorFlags::IncludeSuper); FunctionIt; ++FunctionIt )
 	{
 		const UFunction* Function = *FunctionIt;
