More fixes for latest clang failures, still fallout from Dev-Blueprints merge

diff --git a/Engine/Source/Developer/BlueprintProfiler/Public/BlueprintProfilerSupport.h b/Engine/Source/Developer/BlueprintProfiler/Public/BlueprintProfilerSupport.h
--- a/Engine/Source/Developer/BlueprintProfiler/Public/BlueprintProfilerSupport.h
+++ b/Engine/Source/Developer/BlueprintProfiler/Public/BlueprintProfilerSupport.h
@@ -2,6 +2,9 @@
 
 #pragma once
 
+#include "Core.h"
+#include "Script.h"
+
 //////////////////////////////////////////////////////////////////////////
 // FBlueprintInstrumentedEvent
 
diff --git a/Engine/Source/Developer/BlueprintProfiler/classes/BlueprintProfiler.cpp b/Engine/Source/Developer/BlueprintProfiler/classes/BlueprintProfiler.cpp
--- a/Engine/Source/Developer/BlueprintProfiler/classes/BlueprintProfiler.cpp
+++ b/Engine/Source/Developer/BlueprintProfiler/classes/BlueprintProfiler.cpp
@@ -2,6 +2,8 @@
 
 #include "Private/BlueprintProfilerPCH.h"
 
+#include "BlueprintProfiler.h"
+
 //////////////////////////////////////////////////////////////////////////
 // FBlueprintProfiler
 
diff --git a/Engine/Source/Developer/BlueprintProfiler/classes/BlueprintProfilerSupport.cpp b/Engine/Source/Developer/BlueprintProfiler/classes/BlueprintProfilerSupport.cpp
--- a/Engine/Source/Developer/BlueprintProfiler/classes/BlueprintProfilerSupport.cpp
+++ b/Engine/Source/Developer/BlueprintProfiler/classes/BlueprintProfilerSupport.cpp
@@ -2,6 +2,8 @@
 
 #include "Private/BlueprintProfilerPCH.h"
 
+#include "BlueprintProfilerSupport.h"
+
 //////////////////////////////////////////////////////////////////////////
 // FBlueprintInstrumentedEvent
 
diff --git a/Engine/Source/Editor/Kismet/Private/Profiler/BPProfilerStatisticWidgets.h b/Engine/Source/Editor/Kismet/Private/Profiler/BPProfilerStatisticWidgets.h
--- a/Engine/Source/Editor/Kismet/Private/Profiler/BPProfilerStatisticWidgets.h
+++ b/Engine/Source/Editor/Kismet/Private/Profiler/BPProfilerStatisticWidgets.h
@@ -56,10 +56,11 @@ struct FBPPerformanceData
 public:
 
 	FBPPerformanceData()
-		: MaxTime(-MAX_dbl)
+		: Samples(StatSampleSize)
+		, CachedExclusiveTime(0.0)
+		, MaxTime(-MAX_dbl)
 		, MinTime(MAX_dbl)
 		, TotalTime(0.0)
-		, Samples(StatSampleSize)
 		, NumSamples(0)
 		, bDirty(false)
 	{
@@ -133,12 +134,12 @@ public:
 	}
 
 	FBPProfilerStat(FBPProfilerStatPtr InParentStat, FScriptExecEvent& Context)
-		: ObjectContext(Context)
-		, InclusiveTime(0.0)
+		: InclusiveTime(0.0)
 		, MinInclusiveTime(MAX_dbl)
 		, MaxInclusiveTime(-MAX_dbl)
 		, InclusiveTotalTime(0.0)
 		, InclusiveSamples(0)
+		, ObjectContext(Context)
 		, ParentStat(InParentStat)
 	{
 		UpdateStatName();
diff --git a/Engine/Source/Editor/Kismet/Private/Profiler/ScriptEventExecution.cpp b/Engine/Source/Editor/Kismet/Private/Profiler/ScriptEventExecution.cpp
--- a/Engine/Source/Editor/Kismet/Private/Profiler/ScriptEventExecution.cpp
+++ b/Engine/Source/Editor/Kismet/Private/Profiler/ScriptEventExecution.cpp
@@ -14,25 +14,25 @@ FScriptExecEvent::FScriptExecEvent()
 }
 
 FScriptExecEvent::FScriptExecEvent(TWeakObjectPtr<const UObject> InObjectPtr, const FScriptExecEvent& ContextToCopy)
-	: EventType(EScriptInstrumentation::NodePin)
+	: CallDepth(ContextToCopy.CallDepth)
+	, NumExecOutPins(1)
+	, EventType(EScriptInstrumentation::NodePin)
 	, ScriptCodeOffset(-1)
 	, EventTime(0.0)
 	, ContextObject(InObjectPtr)
 	, PinIndex(ContextToCopy.PinIndex)
-	, CallDepth(ContextToCopy.CallDepth)
-	, NumExecOutPins(1)
 {
 }
 
 FScriptExecEvent::FScriptExecEvent(TWeakObjectPtr<const UObject> InObjectPtr, TWeakObjectPtr<const UObject> InMacroObjectPtr, const FBlueprintInstrumentedEvent& InEvent, const int32 InPinIndex)
-	: EventType(InEvent.GetType())
+	: CallDepth(0)
+	, NumExecOutPins(-1)
+	, EventType(InEvent.GetType())
 	, ScriptCodeOffset(InEvent.GetScriptCodeOffset())
 	, EventTime(InEvent.GetTime())
 	, ContextObject(InObjectPtr)
 	, MacroContextObject(InMacroObjectPtr)
 	, PinIndex(InPinIndex)
-	, CallDepth(0)
-	, NumExecOutPins(-1)
 {
 }
 
