Fixed shadow variable issue from CL# 2815106.

----
Merging CL 2820921 from //UE4/Release-4.11 to Dev-Main (//UE4/Dev-Main)

diff --git a/Engine/Source/Editor/Kismet/Private/BlueprintDetailsCustomization.cpp b/Engine/Source/Editor/Kismet/Private/BlueprintDetailsCustomization.cpp
--- a/Engine/Source/Editor/Kismet/Private/BlueprintDetailsCustomization.cpp
+++ b/Engine/Source/Editor/Kismet/Private/BlueprintDetailsCustomization.cpp
@@ -1996,9 +1996,9 @@ void FBlueprintVarActionDetails::OnFinishedChangingProperties(const FPropertyCha
 bool FBlueprintVarActionDetails::IsVariableInheritedByBlueprint() const
 {
 	UClass* PropertyOwnerClass = nullptr;
-	if (UBlueprint* PropertyOwnerBlueprint = GetPropertyOwnerBlueprint())
+	if (UBlueprint* PropertyOwnerBP = GetPropertyOwnerBlueprint())
 	{
-		PropertyOwnerClass = PropertyOwnerBlueprint->SkeletonGeneratedClass;
+		PropertyOwnerClass = PropertyOwnerBP->SkeletonGeneratedClass;
 	}
 	else if (CachedVariableProperty.IsValid())
 	{
