Copying //UE4/Dev-Editor to Dev-Main (//UE4/Dev-Main)

==========================
MAJOR FEATURES + CHANGES
==========================

Change 2771249 on 2015/11/18 by [REDACTED_PERSON]

	Ensure that UCircularThrobber's Radius determines the widget's desired size when a child of UCanvasPanelSlot.

Change 2794402 on 2015/12/08 by [REDACTED_PERSON]

	Reverting recent changes to Circular throbber. It's unintuative to enforce Size To Content. Will find some other solution.

Change 2803507 on 2015/12/15 by [REDACTED_PERSON]

	BSP poly extrusion can now only be done in the normal direction of the poly.

Change 2803510 on 2015/12/15 by [REDACTED_PERSON]

	Building new static mesh LODs now initializes override vertex colors based on LOD0.

Change 2808877 on 2015/12/18 by [REDACTED_PERSON]

	Make sure the delta scale sign is swap when we have multiple axis with different sign current axis value

Change 2810114 on 2015/12/21 by [REDACTED_PERSON]

	We now expose a message telling the user that we found some mesh that are not reference by any scene node in the fbx file.

Change 2810211 on 2015/12/21 by [REDACTED_PERSON]

	Fixed issue with Show Only Selected not showing members of actor groups.

Change 2811035 on 2015/12/22 by [REDACTED_PERSON]

	Polish UI

Change 2811123 on 2015/12/22 by [REDACTED_PERSON]

	We now can decide which fbx sdk compatibility version we can use when exportting to a fbx file.

Change 2812830 on 2015/12/28 by [REDACTED_PERSON]

	Prevent engine assets' properties from having project assets assigned to them.

Change 2812854 on 2015/12/28 by [REDACTED_PERSON]

	Fixed issue where floating windows were having their border size erroneously added again and again.  Allowed PIE windows to not respect work area bounds if they are created centered, so that they can overlap off the edge of the screen.

Change 2812875 on 2015/12/28 by [REDACTED_PERSON]

	first implementation for skeletal mesh scene import and reimport. Small refator to remove duplicate code in different fbx list ui.

Change 2813172 on 2015/12/29 by [REDACTED_PERSON]

	Partial submit, the base code is there to add all light type with there properties.

Change 2813403 on 2015/12/30 by [REDACTED_PERSON]

	PIE in New Editor Window now respects the Game Gets Mouse Control setting.  This provides a workaround for UE-24824 where attempting to drag a PIE window fails due to the viewport capturing and locking the mouse to itself in FSceneViewport::OnFocusReceived.

Change 2813429 on 2015/12/30 by [REDACTED_PERSON]

	-spotlight and point light support fbx attenuation
	-fix the light orientation so now directional and spotlight point to the same direction of the fbx

Change 2813456 on 2015/12/30 by [REDACTED_PERSON]

	-Import the camera from fbx

Change 2813457 on 2015/12/30 by [REDACTED_PERSON]

	Fixed issues with the code which determines whether the user is attempting to assign a game asset/class to an engine asset's property.

Change 2813475 on 2015/12/30 by [REDACTED_PERSON]

	Removed erroneous debug code.

Change 2814451 on 2016/01/04 by [REDACTED_PERSON]

	Fixed Tint colour for circular throbber.

Change 2814546 on 2016/01/04 by [REDACTED_PERSON]

	Force Message Log to update its category list if a new category is added while it is open.

Change 2814613 on 2016/01/04 by [REDACTED_PERSON]

diff --git a/Engine/Config/BaseEditorPerProjectUserSettings.ini b/Engine/Config/BaseEditorPerProjectUserSettings.ini
--- a/Engine/Config/BaseEditorPerProjectUserSettings.ini
+++ b/Engine/Config/BaseEditorPerProjectUserSettings.ini
@@ -55,6 +55,7 @@ bSCSEditorShowGrid=true
 
 ; Export
 bKeepAttachHierarchy=true
+FbxExportCompatibility=3
 
 ; Misc
 HoverHighlightIntensity=0.000000
diff --git a/Engine/Source/Developer/AssetTools/Private/AssetTypeActions/AssetTypeActions_StaticMesh.cpp b/Engine/Source/Developer/AssetTools/Private/AssetTypeActions/AssetTypeActions_StaticMesh.cpp
--- a/Engine/Source/Developer/AssetTools/Private/AssetTypeActions/AssetTypeActions_StaticMesh.cpp
+++ b/Engine/Source/Developer/AssetTools/Private/AssetTypeActions/AssetTypeActions_StaticMesh.cpp
@@ -53,10 +53,12 @@ void FAssetTypeActions_StaticMesh::GetActions( const TArray<UObject*>& InObjects
 		);
 
 	MenuBuilder.AddSubMenu(
-		NSLOCTEXT("AssetTypeActions_StaticMesh", "StaticMesh_ImportLOD", "Import LOD"),
-		NSLOCTEXT("AssetTypeActions_StaticMesh", "StaticMesh_ImportLODtooltip", "Imports meshes into the LODs"),
-		FNewMenuDelegate::CreateSP( this, &FAssetTypeActions_StaticMesh::GetImportLODMenu, Meshes )
-	);
+		NSLOCTEXT("AssetTypeActions_StaticMesh", "StaticMesh_LODMenu", "Level Of Detail"),
+		NSLOCTEXT("AssetTypeActions_StaticMesh", "StaticMesh_LODTooltip", "LOD Options and Tools"),
+		FNewMenuDelegate::CreateSP(this, &FAssetTypeActions_StaticMesh::GetLODMenu, Meshes),
+		false,
+		FSlateIcon(FEditorStyle::GetStyleSetName(), "ContentBrowser.AssetActions")
+		);
 }
 
 void FAssetTypeActions_StaticMesh::OpenAssetEditor( const TArray<UObject*>& InObjects, TSharedPtr<IToolkitHost> EditWithinLevelEditor )
@@ -118,6 +120,130 @@ void FAssetTypeActions_StaticMesh::GetImportLODMenu(class FMenuBuilder& MenuBuil
 	}
 }
 
+void FAssetTypeActions_StaticMesh::GetLODMenu(class FMenuBuilder& MenuBuilder, TArray<TWeakObjectPtr<UStaticMesh>> Meshes)
+{
+
+	MenuBuilder.AddSubMenu(
+		NSLOCTEXT("AssetTypeActions_StaticMesh", "StaticMesh_ImportLOD", "Import LOD"),
+		NSLOCTEXT("AssetTypeActions_StaticMesh", "StaticMesh_ImportLODtooltip", "Imports meshes into the LODs"),
+		FNewMenuDelegate::CreateSP(this, &FAssetTypeActions_StaticMesh::GetImportLODMenu, Meshes)
+		);
+
+	MenuBuilder.AddMenuSeparator();
+
+	MenuBuilder.AddMenuEntry(
+		NSLOCTEXT("AssetTypeActions_StaticMesh", "StaticMesh_CopyLOD", "Copy LOD"),
+		NSLOCTEXT("AssetTypeActions_StaticMesh", "StaticMesh_CopyLODTooltip", "Copies the LOD settings from the selected mesh."),
+		FSlateIcon(),
+		FUIAction(
+		FExecuteAction::CreateSP(this, &FAssetTypeActions_StaticMesh::ExecuteCopyLODSettings, Meshes),
+		FCanExecuteAction::CreateSP(this, &FAssetTypeActions_StaticMesh::CanCopyLODSettings, Meshes)
+		)
+		);
+
+	FText PasteLabel = FText(LOCTEXT("StaticMesh_PasteLODWithName", "Paste LOD"));
+	if (LODCopyMesh.IsValid())
+	{
+		PasteLabel = FText::Format(LOCTEXT("StaticMesh_PasteLODWithName", "Paste LOD from {0}"), FText::FromString(LODCopyMesh->GetName()));
+	}
+	
+	MenuBuilder.AddMenuEntry(
+		PasteLabel,
+		NSLOCTEXT("AssetTypeActions_StaticMesh", "StaticMesh_PasteLODToltip", "Pastes LOD settings to the selected mesh(es)."),
+		FSlateIcon(),
+		FUIAction(
+		FExecuteAction::CreateSP(this, &FAssetTypeActions_StaticMesh::ExecutePasteLODSettings, Meshes),
+		FCanExecuteAction::CreateSP(this, &FAssetTypeActions_StaticMesh::CanPasteLODSettings, Meshes)
+		)
+		);
+
+
+}
+
+void FAssetTypeActions_StaticMesh::ExecuteCopyLODSettings(TArray<TWeakObjectPtr<UStaticMesh>> Objects)
+{
+	LODCopyMesh = Objects[0];
+}
+
+
+bool FAssetTypeActions_StaticMesh::CanCopyLODSettings(TArray<TWeakObjectPtr<UStaticMesh>> Objects) const
+{
+	return Objects.Num() == 1;
+}
+
+void FAssetTypeActions_StaticMesh::ExecutePasteLODSettings(TArray<TWeakObjectPtr<UStaticMesh>> Objects)
+{
+	if (!LODCopyMesh.IsValid())
+	{
+		return;
+	}
+
+	// Retrieve LOD settings from source mesh
+	struct FLODSettings
+	{
+		FMeshReductionSettings		ReductionSettings;
+		float						ScreenSize;
+	};
+
+	TArray<FLODSettings> LODSettings;
+	LODSettings.AddZeroed(LODCopyMesh->SourceModels.Num());
+	for (int32 i = 0; i < LODCopyMesh->SourceModels.Num(); i++)
+	{
+		LODSettings[i].ReductionSettings = LODCopyMesh->SourceModels[i].ReductionSettings;
+		LODSettings[i].ScreenSize = LODCopyMesh->SourceModels[i].ScreenSize;
+	}
+
+	const bool bAutoComputeLODScreenSize = LODCopyMesh->bAutoComputeLODScreenSize;
+	const float StreamingDistanceMultiplier = LODCopyMesh->StreamingDistanceMultiplier;
+
+	// Copy LOD settings over to selected objects in content browser (meshes)
+	for (TWeakObjectPtr<UStaticMesh> MeshPtr : Objects)
+	{
+		if (MeshPtr.IsValid())
+		{
+			UStaticMesh* Mesh = MeshPtr.Get();
+
+			const int32 LODCount = LODSettings.Num();
+			if (Mesh->SourceModels.Num() > LODCount)
+			{
+				const int32 NumToRemove = Mesh->SourceModels.Num() - LODCount;
+				Mesh->SourceModels.RemoveAt(LODCount, NumToRemove);
+			}
+
+			while (Mesh->SourceModels.Num() < LODCount)
+			{
+				new(Mesh->SourceModels) FStaticMeshSourceModel();
+			}
+
+			for (int32 i = 0; i < LODCount; i++)
+			{
+				Mesh->SourceModels[i].ReductionSettings = LODSettings[i].ReductionSettings;
+				Mesh->SourceModels[i].ScreenSize = LODSettings[i].ScreenSize;
+			}
+
+			for (int32 i = LODCount; i < LODSettings.Num(); ++i)
+			{
+				FStaticMeshSourceModel& SrcModel = Mesh->SourceModels[i];
+
+				SrcModel.ReductionSettings = LODSettings[i].ReductionSettings;
+				SrcModel.ScreenSize = LODSettings[i].ScreenSize;
+			}
+
+			Mesh->bAutoComputeLODScreenSize = bAutoComputeLODScreenSize;
+			Mesh->StreamingDistanceMultiplier = StreamingDistanceMultiplier;
+
+			Mesh->PostEditChange();
+			Mesh->MarkPackageDirty();
+		}
+	}
+
+}
+
+bool FAssetTypeActions_StaticMesh::CanPasteLODSettings(TArray<TWeakObjectPtr<UStaticMesh>> Objects) const
+{
+	return LODCopyMesh.IsValid();
+}
+
 void FAssetTypeActions_StaticMesh::ExecuteCreateDestructibleMesh(TArray<TWeakObjectPtr<UStaticMesh>> Objects)
 {
 	TArray< UObject* > Assets;
diff --git a/Engine/Source/Developer/AssetTools/Private/AssetTypeActions/AssetTypeActions_StaticMesh.h b/Engine/Source/Developer/AssetTools/Private/AssetTypeActions/AssetTypeActions_StaticMesh.h
--- a/Engine/Source/Developer/AssetTools/Private/AssetTypeActions/AssetTypeActions_StaticMesh.h
+++ b/Engine/Source/Developer/AssetTools/Private/AssetTypeActions/AssetTypeActions_StaticMesh.h
@@ -24,9 +24,27 @@ private:
 	/** Handler for when CreateDestructibleMesh is selected */
 	void ExecuteCreateDestructibleMesh(TArray<TWeakObjectPtr<UStaticMesh>> Objects);
 
+	/** Handler for when CopyLODDData is selected */
+	void ExecuteCopyLODSettings(TArray<TWeakObjectPtr<UStaticMesh>> Objects);
+	
+	/** Whether there is a valid static mesh to copy LOD from */
+	bool CanCopyLODSettings(TArray<TWeakObjectPtr<UStaticMesh>> Objects) const;
+
+	/** Handler for when PasteLODDData is selected */
+	void ExecutePasteLODSettings(TArray<TWeakObjectPtr<UStaticMesh>> Objects);
+
+	/** Whether there is a valid static meshes to copy LOD to*/
+	bool CanPasteLODSettings(TArray<TWeakObjectPtr<UStaticMesh>> Objects) const;
+
 	// Handler for when SaveGeneratedLODsInPackage is selected
 	void ExecuteSaveGeneratedLODsInPackage(TArray<TWeakObjectPtr<UStaticMesh>> Objects);
 
 	/** Handler to provide the list of LODs that can be imported or reimported */
 	void GetImportLODMenu(class FMenuBuilder& MenuBuilder,TArray<TWeakObjectPtr<UStaticMesh>> Objects);
+
+	/** Handler to provide the LOD sub-menu. Hides away LOD actions - includes Import LOD sub menu */
+	void GetLODMenu(class FMenuBuilder& MenuBuilder, TArray<TWeakObjectPtr<UStaticMesh>> Meshes);
+private:
+
+	TWeakObjectPtr<UStaticMesh> LODCopyMesh;	
 };
\ No newline at end of file
diff --git a/Engine/Source/Developer/DirectoryWatcher/Private/FileCache.cpp b/Engine/Source/Developer/DirectoryWatcher/Private/FileCache.cpp
--- a/Engine/Source/Developer/DirectoryWatcher/Private/FileCache.cpp
+++ b/Engine/Source/Developer/DirectoryWatcher/Private/FileCache.cpp
@@ -564,9 +564,6 @@ TOptional<FString> FFileCache::GetTransactionPath(const FString& InAbsolutePath)
 
 void FFileCache::DiffDirtyFiles(TMap<FImmutableString, FFileData>& InDirtyFiles, TArray<FUpdateCacheTransaction>& OutTransactions, const FDirectoryState* InFileSystemState) const
 {
-	TArray<uint8> ScratchBuffer;
-	ScratchBuffer.SetNumUninitialized(1024*1024);
-
 	TMap<FImmutableString, FFileData> AddedFiles, ModifiedFiles;
 	TSet<FImmutableString> RemovedFiles, InvalidDirtyFiles;
 
diff --git a/Engine/Source/Developer/FunctionalTesting/Private/Tests/ClientFuncTestPerforming.cpp b/Engine/Source/Developer/FunctionalTesting/Private/Tests/ClientFuncTestPerforming.cpp
--- a/Engine/Source/Developer/FunctionalTesting/Private/Tests/ClientFuncTestPerforming.cpp
+++ b/Engine/Source/Developer/FunctionalTesting/Private/Tests/ClientFuncTestPerforming.cpp
@@ -3,6 +3,8 @@
 #include "FunctionalTestingPrivatePCH.h"
 #include "AutomationCommon.h"
 
+#if WITH_DEV_AUTOMATION_TESTS
+
 #define LOCTEXT_NAMESPACE "FunctionalTesting"
 
 DEFINE_LOG_CATEGORY_STATIC(LogFunctionalTesting, Log, All);
@@ -107,4 +109,6 @@ bool FClientFunctionalTestingMapsBase::RunTest(const FString& Parameters)
 	return false;
 }
 
-#undef LOCTEXT_NAMESPACE
\ No newline at end of file
+#undef LOCTEXT_NAMESPACE
+
+#endif //WITH_DEV_AUTOMATION_TESTS
\ No newline at end of file
diff --git a/Engine/Source/Developer/MessageLog/Private/Presentation/MessageLogViewModel.cpp b/Engine/Source/Developer/MessageLog/Private/Presentation/MessageLogViewModel.cpp
--- a/Engine/Source/Developer/MessageLog/Private/Presentation/MessageLogViewModel.cpp
+++ b/Engine/Source/Developer/MessageLog/Private/Presentation/MessageLogViewModel.cpp
@@ -144,6 +144,8 @@ void FMessageLogViewModel::UpdateListingViewModelArray()
 	{
 		ViewModelArray.Add(It.Value());
 	}
+
+	Update();
 }
 
 #undef LOCTEXT_NAMESPACE
diff --git a/Engine/Source/Developer/MessageLog/Private/UserInterface/SMessageLog.cpp b/Engine/Source/Developer/MessageLog/Private/UserInterface/SMessageLog.cpp
--- a/Engine/Source/Developer/MessageLog/Private/UserInterface/SMessageLog.cpp
+++ b/Engine/Source/Developer/MessageLog/Private/UserInterface/SMessageLog.cpp
@@ -91,6 +91,7 @@ void SMessageLog::Construct( const FArguments& InArgs, const TSharedRef<FMessage
 	}
 
 	ViewModel->OnSelectionChanged().AddSP(this, &SMessageLog::HandleSelectionUpdated);
+	ViewModel->OnChanged().AddSP(this, &SMessageLog::RefreshCategoryList);
 }
 END_SLATE_FUNCTION_BUILD_OPTIMIZATION
 
@@ -104,6 +105,11 @@ void SMessageLog::HandleSelectionUpdated()
 	}
 }
 
+void SMessageLog::RefreshCategoryList()
+{
+	CategoriesListView->RequestListRefresh();
+}
+
 /* SWidget overrides
  *****************************************************************************/
 
diff --git a/Engine/Source/Developer/MessageLog/Private/UserInterface/SMessageLog.h b/Engine/Source/Developer/MessageLog/Private/UserInterface/SMessageLog.h
--- a/Engine/Source/Developer/MessageLog/Private/UserInterface/SMessageLog.h
+++ b/Engine/Source/Developer/MessageLog/Private/UserInterface/SMessageLog.h
@@ -42,9 +42,12 @@ private:
 	/** Callback for changing the selected item in the categories list. */
 	void HandleCategoriesListSelectionChanged( IMessageLogListingPtr Selection, ESelectInfo::Type SelectInfo );
 
-	/** Handles the broadcast message sent the current log listing is changed */
+	/** Handles the broadcast message sent when the current log listing selection is changed */
 	void HandleSelectionUpdated();
 
+	/** Handles the broadcast message sent when the current log listing is updated */
+	void RefreshCategoryList();
+
 private:	
 
 	/** Holds the log categories list view widget. */
diff --git a/Engine/Source/Developer/SharedSettingsWidgets/Private/SExternalImageReference.cpp b/Engine/Source/Developer/SharedSettingsWidgets/Private/SExternalImageReference.cpp
--- a/Engine/Source/Developer/SharedSettingsWidgets/Private/SExternalImageReference.cpp
+++ b/Engine/Source/Developer/SharedSettingsWidgets/Private/SExternalImageReference.cpp
@@ -17,6 +17,9 @@ void SExternalImageReference::Construct(const FArguments& InArgs, const FString&
 	FileDescription = InArgs._FileDescription;
 	OnPreExternalImageCopy = InArgs._OnPreExternalImageCopy;
 	OnPostExternalImageCopy = InArgs._OnPostExternalImageCopy;
+	BaseFilename = InBaseFilename;
+	OverrideFilename = InOverrideFilename;
+	bDeleteTargetWhenDefaultChosen = InArgs._DeleteTargetWhenDefaultChosen;
 
 	FExternalImagePickerConfiguration ImageReferenceConfig;
 	ImageReferenceConfig.TargetImagePath = InOverrideFilename;
@@ -36,6 +39,47 @@ void SExternalImageReference::Construct(const FArguments& InArgs, const FString&
 
 bool SExternalImageReference::HandleExternalImagePicked(const FString& InChosenImage, const FString& InTargetImage)
 {
+	if (bDeleteTargetWhenDefaultChosen && InChosenImage == BaseFilename)
+	{
+		// We elect to remove the target image completely if the default image is chosen (thus allowing us to distinguish the default from a non-default image)
+		if (ISourceControlModule::Get().IsEnabled())
+		{
+			ISourceControlProvider& SourceControlProvider = ISourceControlModule::Get().GetProvider();
+
+			const FSourceControlStatePtr SourceControlState = SourceControlProvider.GetState(InTargetImage, EStateCacheUsage::ForceUpdate);
+			const bool bIsSourceControlled = SourceControlState.IsValid() && SourceControlState->IsSourceControlled();
+
+			if (bIsSourceControlled)
+			{
+				// The file is managed by source control. Delete it through there.
+				TArray<FString> DeleteFilenames;
+				DeleteFilenames.Add(InTargetImage);
+
+				// Revert the file if it is checked out
+				const bool bIsAdded = SourceControlState->IsAdded();
+				if (SourceControlState->IsCheckedOut() || bIsAdded || SourceControlState->IsDeleted())
+				{
+					SourceControlProvider.Execute(ISourceControlOperation::Create<FRevert>(), DeleteFilenames);
+				}
+
+				// If it wasn't already marked as an add, we can ask the source control provider to delete the file
+				if (!bIsAdded)
+				{
+					// Open the file for delete
+					SourceControlProvider.Execute(ISourceControlOperation::Create<FDelete>(), DeleteFilenames);
+				}
+			}
+		}
+
+		IFileManager& FileManager = IFileManager::Get();
+		const bool bRequireExists = false;
+		const bool bEvenIfReadOnly = true;
+		const bool bQuiet = true;
+		FileManager.Delete(*InTargetImage, bRequireExists, bEvenIfReadOnly, bQuiet);
+
+		return true;
+	}
+
 	if(OnPreExternalImageCopy.IsBound())
 	{
 		if(!OnPreExternalImageCopy.Execute(InChosenImage))
diff --git a/Engine/Source/Developer/SharedSettingsWidgets/Public/SExternalImageReference.h b/Engine/Source/Developer/SharedSettingsWidgets/Public/SExternalImageReference.h
--- a/Engine/Source/Developer/SharedSettingsWidgets/Public/SExternalImageReference.h
+++ b/Engine/Source/Developer/SharedSettingsWidgets/Public/SExternalImageReference.h
@@ -30,6 +30,7 @@ class SHAREDSETTINGSWIDGETS_API SExternalImageReference : public SCompoundWidget
 		: _FileDescription(NSLOCTEXT("SExternalImageReference", "FileDescription", "External Image"))
 		, _MaxDisplaySize(400.0f, 400.0f)
 		, _RequiredSize(-1, -1)
+		, _DeleteTargetWhenDefaultChosen(false)
 		{}
 
 		/** The description of the file, used in error messages/notifications */
@@ -52,6 +53,10 @@ class SHAREDSETTINGSWIDGETS_API SExternalImageReference : public SCompoundWidget
 
 		/** A property handle to use if required */
 		SLATE_ARGUMENT(TSharedPtr<class IPropertyHandle>, PropertyHandle)
+
+		/** If true, the target image will be deleted if the default is chosen */
+		SLATE_ARGUMENT(bool, DeleteTargetWhenDefaultChosen)
+
 	SLATE_END_ARGS()
 
 public:
@@ -81,4 +86,7 @@ protected:
 
 	/** A property handle to use if required */
 	TSharedPtr<class IPropertyHandle> PropertyHandle;
+
+	/** If true, the target image will be deleted if the default is chosen */
+	bool bDeleteTargetWhenDefaultChosen;
 };
diff --git a/Engine/Source/Developer/SlateReflector/Private/Widgets/SWidgetReflector.cpp b/Engine/Source/Developer/SlateReflector/Private/Widgets/SWidgetReflector.cpp
--- a/Engine/Source/Developer/SlateReflector/Private/Widgets/SWidgetReflector.cpp
+++ b/Engine/Source/Developer/SlateReflector/Private/Widgets/SWidgetReflector.cpp
@@ -298,10 +298,7 @@ private:
 	void HandleReflectorTreeGetChildren( TSharedRef<FWidgetReflectorNodeBase> InReflectorNode, TArray<TSharedRef<FWidgetReflectorNodeBase>>& OutChildren );
 
 	/** Callback for when the selection in the reflector tree has changed. */
-	void HandleReflectorTreeSelectionChanged( TSharedPtr<FWidgetReflectorNodeBase>, ESelectInfo::Type /*SelectInfo*/ )
-	{
-		SelectedNodes = ReflectorTree->GetSelectedItems();
-	}
+	void HandleReflectorTreeSelectionChanged( TSharedPtr<FWidgetReflectorNodeBase>, ESelectInfo::Type /*SelectInfo*/ );
 
 	TSharedRef<ITableRow> GenerateEventLogRow( TSharedRef<FLoggedEvent> InReflectorNode, const TSharedRef<STableViewBase>& OwnerTable );
 
@@ -1636,6 +1633,17 @@ void SWidgetReflector::HandleReflectorTreeGetChildren(TSharedRef<FWidgetReflecto
 }
 
 
+void SWidgetReflector::HandleReflectorTreeSelectionChanged( TSharedPtr<FWidgetReflectorNodeBase>, ESelectInfo::Type /*SelectInfo*/ )
+{
+	SelectedNodes = ReflectorTree->GetSelectedItems();
+
+	if (CurrentUIMode == EWidgetReflectorUIMode::Snapshot)
+	{
+		WidgetSnapshotVisualizer->SetSelectedWidgets(SelectedNodes);
+	}
+}
+
+
 TSharedRef<ITableRow> SWidgetReflector::GenerateEventLogRow( TSharedRef<FLoggedEvent> InLoggedEvent, const TSharedRef<STableViewBase>& OwnerTable )
 {
 	return SNew(STableRow<TSharedRef<FLoggedEvent>>, OwnerTable)
diff --git a/Engine/Source/Developer/SlateReflector/Private/Widgets/SWidgetSnapshotVisualizer.cpp b/Engine/Source/Developer/SlateReflector/Private/Widgets/SWidgetSnapshotVisualizer.cpp
--- a/Engine/Source/Developer/SlateReflector/Private/Widgets/SWidgetSnapshotVisualizer.cpp
+++ b/Engine/Source/Developer/SlateReflector/Private/Widgets/SWidgetSnapshotVisualizer.cpp
@@ -89,6 +89,11 @@ public:
 		return bIsPicking;
 	}
 
+	void SetSelectedWidgets(const TArray<TSharedRef<FWidgetReflectorNodeBase>>& InSelectedWidgets)
+	{
+		SelectedWidgets = InSelectedWidgets;
+	}
+
 	virtual void OnArrangeChildren(const FGeometry& AllottedGeometry, FArrangedChildren& ArrangedChildren) const override
 	{
 		CachedSize = AllottedGeometry.GetLocalSize();
@@ -232,37 +237,20 @@ public:
 		LayerId = SPanel::OnPaint(Args, AllottedGeometry, MyClippingRect, OutDrawElements, LayerId, InWidgetStyle, bParentEnabled);
 		LayerId = ScrollyZoomy.PaintSoftwareCursorIfNeeded(AllottedGeometry, MyClippingRect, OutDrawElements, LayerId);
 
-		if (PickedWidgets.Num() > 0)
+		TSharedPtr<FWidgetReflectorNodeBase> Window = SnapshotDataPtr->GetWindow(SelectedWindowIndex);
+		if (Window.IsValid())
 		{
-			TSharedPtr<FWidgetReflectorNodeBase> Window = SnapshotDataPtr->GetWindow(SelectedWindowIndex);
-			if (Window.IsValid())
+			const FVector2D RootDrawOffset = PhysicalOffset - Window->GetAccumulatedLayoutTransform().GetTranslation();
+			if (bIsPicking)
 			{
-				const FVector2D RootDrawOffset = PhysicalOffset - Window->GetAccumulatedLayoutTransform().GetTranslation();
-				if (bIsPicking)
-				{
-					const FLinearColor TopmostWidgetColor(1.0f, 0.0f, 0.0f);
-					const FLinearColor LeafmostWidgetColor(0.0f, 1.0f, 0.0f);
+				const FLinearColor TopmostWidgetColor(1.0f, 0.0f, 0.0f);
+				const FLinearColor LeafmostWidgetColor(0.0f, 1.0f, 0.0f);
 
-					for (int32 WidgetIndex = 0; WidgetIndex < PickedWidgets.Num(); ++WidgetIndex)
-					{
-						const TSharedRef<FWidgetReflectorNodeBase>& PickedWidget = PickedWidgets[WidgetIndex];
-						const float ColorFactor = static_cast<float>(WidgetIndex)/PickedWidgets.Num();
-						const FLinearColor Tint(1.0f - ColorFactor, ColorFactor, 0.0f, 1.0f);
-
-						FSlateDrawElement::MakeBox(
-							OutDrawElements,
-							++LayerId,
-							AllottedGeometry.ToPaintGeometry(RootDrawOffset + PickedWidget->GetAccumulatedLayoutTransform().GetTranslation(), TransformPoint(PickedWidget->GetAccumulatedLayoutTransform().GetScale(), PickedWidget->GetLocalSize())),
-							FCoreStyle::Get().GetBrush(TEXT("Debug.Border")),
-							MyClippingRect,
-							ESlateDrawEffect::None,
-							FMath::Lerp(TopmostWidgetColor, LeafmostWidgetColor, ColorFactor)
-						);
-					}
-				}
-				else
+				for (int32 WidgetIndex = 0; WidgetIndex < PickedWidgets.Num(); ++WidgetIndex)
 				{
-					const TSharedRef<FWidgetReflectorNodeBase>& PickedWidget = PickedWidgets.Last();
+					const TSharedRef<FWidgetReflectorNodeBase>& PickedWidget = PickedWidgets[WidgetIndex];
+					const float ColorFactor = static_cast<float>(WidgetIndex)/PickedWidgets.Num();
+					const FLinearColor Tint(1.0f - ColorFactor, ColorFactor, 0.0f, 1.0f);
 
 					FSlateDrawElement::MakeBox(
 						OutDrawElements,
@@ -271,7 +259,22 @@ public:
 						FCoreStyle::Get().GetBrush(TEXT("Debug.Border")),
 						MyClippingRect,
 						ESlateDrawEffect::None,
-						PickedWidget->GetTint()
+						FMath::Lerp(TopmostWidgetColor, LeafmostWidgetColor, ColorFactor)
+					);
+				}
+			}
+			else
+			{
+				for (const auto& SelectedWidget : SelectedWidgets)
+				{
+					FSlateDrawElement::MakeBox(
+						OutDrawElements,
+						++LayerId,
+						AllottedGeometry.ToPaintGeometry(RootDrawOffset + SelectedWidget->GetAccumulatedLayoutTransform().GetTranslation(), TransformPoint(SelectedWidget->GetAccumulatedLayoutTransform().GetScale(), SelectedWidget->GetLocalSize())),
+						FCoreStyle::Get().GetBrush(TEXT("Debug.Border")),
+						MyClippingRect,
+						ESlateDrawEffect::None,
+						SelectedWidget->GetTint()
 					);
 				}
 			}
@@ -352,6 +355,8 @@ private:
 
 	bool bIsPicking;
 	TArray<TSharedRef<FWidgetReflectorNodeBase>> PickedWidgets;
+
+	TArray<TSharedRef<FWidgetReflectorNodeBase>> SelectedWidgets;
 };
 
 
@@ -792,6 +797,14 @@ void SWidgetSnapshotVisualizer::SnapshotDataUpdated()
 	}
 }
 
+void SWidgetSnapshotVisualizer::SetSelectedWidgets(const TArray<TSharedRef<FWidgetReflectorNodeBase>>& InSelectedWidgets)
+{
+	if (SnapshotImage.IsValid())
+	{
+		SnapshotImage->SetSelectedWidgets(InSelectedWidgets);
+	}
+}
+
 FReply SWidgetSnapshotVisualizer::OnPreviewKeyDown(const FGeometry& MyGeometry, const FKeyEvent& InKeyEvent)
 {
 	if (SnapshotImage.IsValid() && InKeyEvent.GetKey() == EKeys::Escape)
diff --git a/Engine/Source/Developer/SlateReflector/Private/Widgets/SWidgetSnapshotVisualizer.h b/Engine/Source/Developer/SlateReflector/Private/Widgets/SWidgetSnapshotVisualizer.h
--- a/Engine/Source/Developer/SlateReflector/Private/Widgets/SWidgetSnapshotVisualizer.h
+++ b/Engine/Source/Developer/SlateReflector/Private/Widgets/SWidgetSnapshotVisualizer.h
@@ -117,6 +117,9 @@ public:
 	/** Called when the snapshot data we're observing is changed. This allows us to update our view. */
 	void SnapshotDataUpdated();
 
+	/** Called to update the list of selected widgets */
+	void SetSelectedWidgets(const TArray<TSharedRef<FWidgetReflectorNodeBase>>& InSelectedWidgets);
+
 	// SWidget interface
 	virtual FReply OnPreviewKeyDown(const FGeometry& MyGeometry, const FKeyEvent& InKeyEvent) override;
 
diff --git a/Engine/Source/Developer/SourceControl/Private/Tests/SourceControlTests.cpp b/Engine/Source/Developer/SourceControl/Private/Tests/SourceControlTests.cpp
--- a/Engine/Source/Developer/SourceControl/Private/Tests/SourceControlTests.cpp
+++ b/Engine/Source/Developer/SourceControl/Private/Tests/SourceControlTests.cpp
@@ -13,6 +13,8 @@
 
 #include "Tests/SourceControlAutomationCommon.h"
 
+#if WITH_DEV_AUTOMATION_TESTS
+
 static void GetProviders(TArray<FString>& OutBeautifiedNames, TArray<FString>& OutTestCommands)
 {
 	// we want to use all the providers we can find except 'None'
@@ -945,4 +947,6 @@ bool FGetRevisionTest::RunTest(const FString& Parameters)
 	return true;
 }
 
+#endif //WITH_DEV_AUTOMATION_TESTS
+
 #endif	// #if WITH_EDITOR
\ No newline at end of file
diff --git a/Engine/Source/Editor/BlueprintGraph/Private/EdGraphSchema_K2.cpp b/Engine/Source/Editor/BlueprintGraph/Private/EdGraphSchema_K2.cpp
--- a/Engine/Source/Editor/BlueprintGraph/Private/EdGraphSchema_K2.cpp
+++ b/Engine/Source/Editor/BlueprintGraph/Private/EdGraphSchema_K2.cpp
@@ -4112,8 +4112,12 @@ void UEdGraphSchema_K2::TrySetDefaultValue(UEdGraphPin& Pin, const FString& NewD
 	}
 	else if(Pin.PinType.PinCategory == PC_Text)
 	{
-		// Set Text pins by string as if it were text.
-		TrySetDefaultText(Pin, FText::FromString(NewDefaultValue));
+		FText NewTextValue;
+		if (!FTextStringHelper::ReadFromString(*NewDefaultValue, NewTextValue))
+		{
+			NewTextValue = FText::FromString(NewDefaultValue);
+		}
+		TrySetDefaultText(Pin, NewTextValue);
 		UseDefaultObject = nullptr;
 		UseDefaultValue.Empty();
 		return;
@@ -4170,21 +4174,7 @@ void UEdGraphSchema_K2::TrySetDefaultText(UEdGraphPin& InPin, const FText& InNew
 		{
 			InPin.DefaultObject = NULL;
 			InPin.DefaultValue = NULL;
-			if(InNewDefaultText.IsEmpty())
-			{
-				InPin.DefaultTextValue = InNewDefaultText;
-			}
-			else
-			{
-				if(InNewDefaultText.IsCultureInvariant())
-				{
-					InPin.DefaultTextValue = InNewDefaultText;
-				}
-				else
-				{
-					InPin.DefaultTextValue = FText::ChangeKey(TEXT(""), FGuid::NewGuid().ToString(), InNewDefaultText);
-				}
-			}
+			InPin.DefaultTextValue = InNewDefaultText;
 		}
 
 		UEdGraphNode* Node = InPin.GetOwningNode();
@@ -4296,7 +4286,10 @@ void UEdGraphSchema_K2::SetPinDefaultValue(UEdGraphPin* Pin, const UFunction* Fu
 		{
 			if (Pin->PinType.PinCategory == PC_Text)
 			{
-				Pin->DefaultTextValue = FText::AsCultureInvariant(Pin->AutogeneratedDefaultValue);
+				if (!FTextStringHelper::ReadFromString(*Pin->AutogeneratedDefaultValue, Pin->DefaultTextValue))
+				{
+					Pin->DefaultTextValue = FText::FromString(Pin->AutogeneratedDefaultValue);
+				}
 			}
 			else
 			{
diff --git a/Engine/Source/Editor/BlueprintGraph/Private/K2Node_CallFunction.cpp b/Engine/Source/Editor/BlueprintGraph/Private/K2Node_CallFunction.cpp
--- a/Engine/Source/Editor/BlueprintGraph/Private/K2Node_CallFunction.cpp
+++ b/Engine/Source/Editor/BlueprintGraph/Private/K2Node_CallFunction.cpp
@@ -2047,7 +2047,7 @@ void UK2Node_CallFunction::ExpandNode(class FKismetCompilerContext& CompilerCont
 						auto ValuePin = InnerHandleAutoCreateRef(this, Pin, CompilerContext, SourceGraph, bHasDefaultValue);
 						if ( ValuePin )
 						{
-							if (!DefaultObject && DefaultTextValue.IsEmpty() && (DefaultValue == AutogeneratedDefaultValue))
+							if (!DefaultObject && DefaultTextValue.IsEmpty() && DefaultValue.Equals(AutogeneratedDefaultValue, ESearchCase::CaseSensitive))
 							{
 								// Use the latest code to set default value
 								Schema->SetPinDefaultValueBasedOnType(ValuePin);
diff --git a/Engine/Source/Editor/BlueprintGraph/Private/K2Node_MakeStruct.cpp b/Engine/Source/Editor/BlueprintGraph/Private/K2Node_MakeStruct.cpp
--- a/Engine/Source/Editor/BlueprintGraph/Private/K2Node_MakeStruct.cpp
+++ b/Engine/Source/Editor/BlueprintGraph/Private/K2Node_MakeStruct.cpp
@@ -28,7 +28,10 @@ void UK2Node_MakeStruct::FMakeStructPinManager::CustomizePinData(UEdGraphPin* Pi
 			InPin.AutogeneratedDefaultValue = Value;
 			if (bIsText)
 			{
-				InPin.DefaultTextValue = FText::FromString(Value);
+				if (!FTextStringHelper::ReadFromString(*Value, InPin.DefaultTextValue))
+				{
+					InPin.DefaultTextValue = FText::FromString(Value);
+				}
 			}
 			else if (bIsObject)
 			{
diff --git a/Engine/Source/Editor/BlueprintGraph/Private/K2Node_Select.cpp b/Engine/Source/Editor/BlueprintGraph/Private/K2Node_Select.cpp
--- a/Engine/Source/Editor/BlueprintGraph/Private/K2Node_Select.cpp
+++ b/Engine/Source/Editor/BlueprintGraph/Private/K2Node_Select.cpp
@@ -1090,7 +1090,7 @@ void UK2Node_Select::ExpandNode(class FKismetCompilerContext& CompilerContext, U
 		auto ValuePin = UK2Node_CallFunction::InnerHandleAutoCreateRef(this, Pin, CompilerContext, SourceGraph, true);
 		if (ValuePin)
 		{
-			if (!DefaultObject && DefaultTextValue.IsEmpty() && (DefaultValue == AutogeneratedDefaultValue))
+			if (!DefaultObject && DefaultTextValue.IsEmpty() && DefaultValue.Equals(AutogeneratedDefaultValue, ESearchCase::CaseSensitive))
 			{
 				// Use the latest code to set default value
 				Schema->SetPinDefaultValueBasedOnType(ValuePin);
diff --git a/Engine/Source/Editor/Cascade/Private/Tests/ParticleSystemTests.cpp b/Engine/Source/Editor/Cascade/Private/Tests/ParticleSystemTests.cpp
--- a/Engine/Source/Editor/Cascade/Private/Tests/ParticleSystemTests.cpp
+++ b/Engine/Source/Editor/Cascade/Private/Tests/ParticleSystemTests.cpp
@@ -15,6 +15,8 @@
 
 #include "ScopedTransaction.h"
 
+#if WITH_DEV_AUTOMATION_TESTS
+
 #define LOCTEXT_NAMESPACE "ParticleEditorPromotionTests"
 
 DEFINE_LOG_CATEGORY_STATIC(LogParticleEditorPromotionTests, Log, All);
@@ -235,4 +237,6 @@ bool FParticleEditorPromotionTest::RunTest(const FString& Parameters)
 	return true;
 }
 
-#undef LOCTEXT_NAMESPACE
\ No newline at end of file
+#undef LOCTEXT_NAMESPACE
+
+#endif //WITH_DEV_AUTOMATION_TESTS
\ No newline at end of file
diff --git a/Engine/Source/Editor/ClassViewer/Private/SClassViewer.cpp b/Engine/Source/Editor/ClassViewer/Private/SClassViewer.cpp
--- a/Engine/Source/Editor/ClassViewer/Private/SClassViewer.cpp
+++ b/Engine/Source/Editor/ClassViewer/Private/SClassViewer.cpp
@@ -1814,6 +1814,9 @@ void FClassHierarchy::LoadUnloadedTagData(TSharedPtr<FClassViewerNode>& InOutCla
 			{
 				if(!CurrentString.StartsWith(TEXT("Graphs=(")))
 				{
+					// The interfaces end with ' because of the path reference, so remove it.
+					InterfaceName.RemoveFromEnd("'");
+
 					UnloadedBlueprintData->AddImplementedInterfaces(InterfaceName);
 				}
 			}
diff --git a/Engine/Source/Editor/ContentBrowser/Private/SAssetView.cpp b/Engine/Source/Editor/ContentBrowser/Private/SAssetView.cpp
--- a/Engine/Source/Editor/ContentBrowser/Private/SAssetView.cpp
+++ b/Engine/Source/Editor/ContentBrowser/Private/SAssetView.cpp
@@ -3496,6 +3496,14 @@ bool SAssetView::AssetVerifyRenameCommit(const TSharedPtr<FAssetViewItem>& Item,
 			return true;
 		}
 	}
+	else
+	{
+		const TSharedPtr<FAssetViewFolder>& ItemAsFolder = StaticCastSharedPtr<FAssetViewFolder>(Item);
+		if (NewNameString == ItemAsFolder->FolderName.ToString())
+		{
+			return true;
+		}
+	}
 
 	if ( bIsAssetType )
 	{
diff --git a/Engine/Source/Editor/DetailCustomizations/Private/ComponentTransformDetails.cpp b/Engine/Source/Editor/DetailCustomizations/Private/ComponentTransformDetails.cpp
--- a/Engine/Source/Editor/DetailCustomizations/Private/ComponentTransformDetails.cpp
+++ b/Engine/Source/Editor/DetailCustomizations/Private/ComponentTransformDetails.cpp
@@ -667,7 +667,22 @@ void FComponentTransformDetails::OnToggleAbsoluteLocation( bool bEnable )
 				}
 
 				SceneComponent->bAbsoluteLocation = !SceneComponent->bAbsoluteLocation;
-
+				
+				// Update RelativeLocation to maintain/stabilize position when switching between relative and world.
+				if (SceneComponent->AttachParent)
+				{
+					if (SceneComponent->bAbsoluteLocation)
+					{
+						SceneComponent->RelativeLocation = SceneComponent->ComponentToWorld.GetTranslation();
+					}
+					else
+					{
+						FTransform ParentToWorld = SceneComponent->AttachParent->GetSocketTransform(SceneComponent->AttachSocketName);
+						FTransform RelativeTM = SceneComponent->ComponentToWorld.GetRelativeTransform(ParentToWorld);
+						SceneComponent->RelativeLocation = RelativeTM.GetTranslation();
+					}
+				}
+				
 				FPropertyChangedEvent PropertyChangedEvent( AbsoluteLocationProperty );
 				Object->PostEditChangeProperty(PropertyChangedEvent);
 
diff --git a/Engine/Source/Editor/DetailCustomizations/Private/TextCustomization.cpp b/Engine/Source/Editor/DetailCustomizations/Private/TextCustomization.cpp
--- a/Engine/Source/Editor/DetailCustomizations/Private/TextCustomization.cpp
+++ b/Engine/Source/Editor/DetailCustomizations/Private/TextCustomization.cpp
@@ -3,209 +3,123 @@
 #include "DetailCustomizationsPrivatePCH.h"
 #include "TextCustomization.h"
 #include "IPropertyUtilities.h"
+#include "STextPropertyEditableTextBox.h"
 
 namespace
 {
-	class STextPropertyWidget : public SCompoundWidget
+	/** Allows STextPropertyEditableTextBox to edit a property handle */
+	class FEditableTextPropertyHandle : public IEditableTextProperty
 	{
-		SLATE_BEGIN_ARGS(STextPropertyWidget)
-			: _Font( FEditorStyle::GetFontStyle( TEXT("PropertyWindow.NormalFont") ) ) 
-			{}
-			SLATE_ATTRIBUTE( FSlateFontInfo, Font )
-		SLATE_END_ARGS()
-
 	public:
-		void Construct(const FArguments& Arguments, const TSharedRef<IPropertyHandle>& InPropertyHandle, const TSharedPtr<IPropertyUtilities>& InPropertyUtilities);
-		void GetDesiredWidth( float& OutMinDesiredWidth, float& OutMaxDesiredWidth );
-		bool SupportsKeyboardFocus() const;
-		FReply OnFocusReceived( const FGeometry& MyGeometry, const FFocusEvent& InFocusEvent );
-		void Tick( const FGeometry& AllottedGeometry, const double InCurrentTime, const float InDeltaTime );
-		bool CanEdit() const;
-		bool IsReadOnly() const;
-
-	private:
-		TSharedPtr<IPropertyHandle> PropertyHandle;
-
-		TSharedPtr<IPropertyUtilities> PropertyUtilities;
-
-		TSharedPtr< class SWidget > PrimaryWidget;
-
-		TSharedPtr<SMultiLineEditableTextBox> MultiLineWidget;
-
-		TSharedPtr<SEditableTextBox> SingleLineWidget;
+		FEditableTextPropertyHandle(const TSharedRef<IPropertyHandle>& InPropertyHandle, const TSharedPtr<IPropertyUtilities>& InPropertyUtilities)
+			: PropertyHandle(InPropertyHandle)
+			, PropertyUtilities(InPropertyUtilities)
+		{
+			RefreshRawData();
+		}
 
-		TOptional< float > PreviousHeight;
+		virtual bool IsMultiLineText() const override
+		{
+			return PropertyHandle->IsValidHandle() && PropertyHandle->GetBoolMetaData("MultiLine");
+		}
 
-		bool bIsMultiLine;
+		virtual bool IsPassword() const override
+		{
+			return PropertyHandle->IsValidHandle() && PropertyHandle->GetBoolMetaData("PasswordField");
+		}
 
-		static FText MultipleValuesText;
-	};
+		virtual bool IsReadOnly() const override
+		{
+			return !PropertyHandle->IsValidHandle() || PropertyHandle->IsEditConst();
+		}
 
-	FText STextPropertyWidget::MultipleValuesText(NSLOCTEXT("PropertyEditor", "MultipleValues", "Multiple Values"));
+		virtual bool IsDefaultValue() const override
+		{
+			return PropertyHandle->IsValidHandle() && !PropertyHandle->DiffersFromDefault();
+		}
 
-	void STextPropertyWidget::Construct(const FArguments& InArgs, const TSharedRef<IPropertyHandle>& InPropertyHandle, const TSharedPtr<IPropertyUtilities>& InPropertyUtilities)
-	{
-		PropertyHandle = InPropertyHandle;
-		PropertyUtilities = InPropertyUtilities;
+		virtual FText GetToolTipText() const override
+		{
+			return (PropertyHandle->IsValidHandle())
+				? PropertyHandle->GetToolTipText()
+				: FText::GetEmpty();
+		}
 
-		const auto& GetTextValue = [this]() -> FText
+		virtual int32 GetNumTexts() const override
 		{
-			FText TextValue;
+			return (PropertyHandle->IsValidHandle())
+				? RawTextData.Num() 
+				: 0;
+		}
 
-			TArray<const void*> RawData;
-			PropertyHandle->AccessRawData(RawData);
-			if (RawData.Num() == 1)
-			{
-				const FText* RawDatum = reinterpret_cast<const FText*>(RawData.Top());
-				if (RawDatum)
-				{
-					TextValue = *RawDatum;
-				}
-			}
-			else if(RawData.Num() > 1)
+		virtual FText GetText(const int32 InIndex) const override
+		{
+			if (PropertyHandle->IsValidHandle())
 			{
-				TextValue = MultipleValuesText;
+				check(RawTextData.IsValidIndex(InIndex));
+				return *RawTextData[InIndex];
 			}
-			else
+			return FText::GetEmpty();
+		}
+
+		virtual void SetText(const int32 InIndex, const FText& InText) override
+		{
+			if (PropertyHandle->IsValidHandle())
 			{
-				TextValue = FText::GetEmpty();
+				check(RawTextData.IsValidIndex(InIndex));
+				*RawTextData[InIndex] = InText;
 			}
+		}
 
-			return TextValue;
-		};
-
-		const auto& OnTextCommitted = [this](const FText& NewText, ETextCommit::Type CommitInfo)
+		virtual void PreEdit() override
 		{
-			TArray<void*> RawData;
-			PropertyHandle->AccessRawData(RawData);
-
-			// Don't commit the Multiple Values text if there are multiple properties being set
-			if (RawData.Num() > 0 && (RawData.Num() == 1 || NewText.ToString() != MultipleValuesText.ToString()))
+			if (PropertyHandle->IsValidHandle())
 			{
 				PropertyHandle->NotifyPreChange();
+			}
+		}
 
-				for (void* const RawDatum : RawData)
-				{
-					FText& PropertyValue = *(reinterpret_cast<FText* const>(RawDatum));
-
-					// FText::FromString on the result of FText::ToString is intentional. For now, we want to nuke any namespace/key info and let it get regenerated from scratch,
-					// rather than risk adopting whatever came through some chain of calls. This will be replaced when preserving of identity is implemented.
-					PropertyValue = FText::FromString(NewText.ToString());
-				}
-
+		virtual void PostEdit() override
+		{
+			if (PropertyHandle->IsValidHandle())
+			{
 				PropertyHandle->NotifyPostChange();
 				PropertyHandle->NotifyFinishedChangingProperties();
 			}
-		};
 
-		TSharedPtr<SHorizontalBox> HorizontalBox;
-
-		bool bIsPassword = PropertyHandle->GetBoolMetaData("PasswordField");
-		bIsMultiLine = PropertyHandle->GetBoolMetaData("MultiLine");
-		if(bIsMultiLine)
-		{
-			ChildSlot
-				[
-					SAssignNew(HorizontalBox, SHorizontalBox)
-					+SHorizontalBox::Slot()
-					.FillWidth(1.0f)
-					[
-						SNew( SBox )
-						.MaxDesiredHeight(300.f)
-						[
-							SAssignNew(MultiLineWidget, SMultiLineEditableTextBox)
-							.Text_Lambda(GetTextValue)
-							.Font(InArgs._Font)
-							.SelectAllTextWhenFocused(false)
-							.ClearKeyboardFocusOnCommit(false)
-							.OnTextCommitted_Lambda(OnTextCommitted)
-							.SelectAllTextOnCommit(false)
-							.IsReadOnly(this, &STextPropertyWidget::IsReadOnly)
-							.AutoWrapText(true)
-							.ModiferKeyForNewLine(EModifierKey::Shift)
-							.IsPassword(bIsPassword)
-						]
-					]
-				];
-
-			PrimaryWidget = MultiLineWidget;
-		}
-		else
-		{
-			ChildSlot
-				[
-					SAssignNew(HorizontalBox, SHorizontalBox)
-					+SHorizontalBox::Slot()
-					.FillWidth(1.0f)
-					[
-						SAssignNew( SingleLineWidget, SEditableTextBox )
-						.Text_Lambda(GetTextValue)
-						.Font( InArgs._Font )
-						.SelectAllTextWhenFocused( true )
-						.ClearKeyboardFocusOnCommit(false)
-						.OnTextCommitted_Lambda(OnTextCommitted)
-						.SelectAllTextOnCommit( true )
-						.IsReadOnly(this, &STextPropertyWidget::IsReadOnly)
-						.IsPassword(bIsPassword)
-
-					]
-				];
-
-			PrimaryWidget = SingleLineWidget;
+			RefreshRawData();
 		}
 
-		SetEnabled( TAttribute<bool>( this, &STextPropertyWidget::CanEdit ) );
-	}
-
-	void STextPropertyWidget::GetDesiredWidth( float& OutMinDesiredWidth, float& OutMaxDesiredWidth )
-	{
-		if(bIsMultiLine)
+		virtual void RequestRefresh() override
 		{
-			OutMinDesiredWidth = 250.0f;
-		}
-		else
-		{
-			OutMinDesiredWidth = 125.0f;
+			if (PropertyUtilities.IsValid())
+			{
+				PropertyUtilities->RequestRefresh();
+			}
 		}
 
-		OutMaxDesiredWidth = 600.0f;
-	}
-
-	bool STextPropertyWidget::SupportsKeyboardFocus() const
-	{
-		return PrimaryWidget.IsValid() && PrimaryWidget->SupportsKeyboardFocus();
-	}
-
-	FReply STextPropertyWidget::OnFocusReceived( const FGeometry& MyGeometry, const FFocusEvent& InFocusEvent )
-	{
-		// Forward keyboard focus to our editable text widget
-		return FReply::Handled().SetUserFocus(PrimaryWidget.ToSharedRef(), InFocusEvent.GetCause());
-	}
-
-	void STextPropertyWidget::Tick( const FGeometry& AllottedGeometry, const double InCurrentTime, const float InDeltaTime )
-	{
-		const float CurrentHeight = AllottedGeometry.GetLocalSize().Y;
-		if (bIsMultiLine && PreviousHeight.IsSet() && PreviousHeight.GetValue() != CurrentHeight && PropertyUtilities.IsValid())
+	private:
+		void RefreshRawData()
 		{
-			PropertyUtilities->RequestRefresh();
-		}
-		PreviousHeight = CurrentHeight;
-	}
+			RawTextData.Empty();
 
-	bool STextPropertyWidget::CanEdit() const
-	{
-		return PropertyHandle.IsValid() ? !PropertyHandle->IsEditConst() : true;
-	}
+			if (PropertyHandle->IsValidHandle())
+			{
+				auto& RawData = reinterpret_cast<TArray<void*>&>(RawTextData);
+				PropertyHandle->AccessRawData(RawData);
+			}
+		}
 
-	bool STextPropertyWidget::IsReadOnly() const
-	{
-		return !CanEdit();
-	}
+		TSharedRef<IPropertyHandle> PropertyHandle;
+		TSharedPtr<IPropertyUtilities> PropertyUtilities;
+		TArray<FText*> RawTextData;
+	};
 }
 
 void FTextCustomization::CustomizeHeader( TSharedRef<class IPropertyHandle> InPropertyHandle, class FDetailWidgetRow& HeaderRow, IPropertyTypeCustomizationUtils& PropertyTypeCustomizationUtils )
 {
-	const bool bIsMultiLine = InPropertyHandle->GetProperty()->GetBoolMetaData("MultiLine");
+	TSharedRef<IEditableTextProperty> EditableTextProperty = MakeShareable(new FEditableTextPropertyHandle(InPropertyHandle, PropertyTypeCustomizationUtils.GetPropertyUtilities()));
+	const bool bIsMultiLine = EditableTextProperty->IsMultiLineText();
 
 	HeaderRow.FilterString(InPropertyHandle->GetPropertyDisplayName())
 		.NameContent()
@@ -216,11 +130,13 @@ void FTextCustomization::CustomizeHeader( TSharedRef<class IPropertyHandle> InPr
 		.MinDesiredWidth(bIsMultiLine ? 250.f : 125.f)
 		.MaxDesiredWidth(600.f)
 		[
-			SNew(STextPropertyWidget, InPropertyHandle, PropertyTypeCustomizationUtils.GetPropertyUtilities())
+			SNew(STextPropertyEditableTextBox, EditableTextProperty)
+				.Font(FEditorStyle::GetFontStyle("PropertyWindow.NormalFont"))
+				.AutoWrapText(true)
 		];
 }
 
 void FTextCustomization::CustomizeChildren( TSharedRef<class IPropertyHandle> InPropertyHandle, class IDetailChildrenBuilder& StructBuilder, IPropertyTypeCustomizationUtils& PropertyTypeCustomizationUtils )
 {
 
-}
\ No newline at end of file
+}
diff --git a/Engine/Source/Editor/DetailCustomizations/Private/WindowsTargetSettingsDetails.cpp b/Engine/Source/Editor/DetailCustomizations/Private/WindowsTargetSettingsDetails.cpp
--- a/Engine/Source/Editor/DetailCustomizations/Private/WindowsTargetSettingsDetails.cpp
+++ b/Engine/Source/Editor/DetailCustomizations/Private/WindowsTargetSettingsDetails.cpp
@@ -170,6 +170,7 @@ void FWindowsTargetSettingsDetails::CustomizeDetails( IDetailLayoutBuilder& Deta
 			.FileDescription(EditorSplashDesc)
 			.OnGetPickerPath(FOnGetPickerPath::CreateSP(this, &FWindowsTargetSettingsDetails::GetPickerPath))
 			.OnPostExternalImageCopy(FOnPostExternalImageCopy::CreateSP(this, &FWindowsTargetSettingsDetails::HandlePostExternalIconCopy))
+			.DeleteTargetWhenDefaultChosen(true)
 		]
 	];
 
diff --git a/Engine/Source/Editor/EditorWidgets/Private/STextPropertyEditableTextBox.cpp b/Engine/Source/Editor/EditorWidgets/Private/STextPropertyEditableTextBox.cpp
new file mode 100644
--- /dev/null
+++ b/Engine/Source/Editor/EditorWidgets/Private/STextPropertyEditableTextBox.cpp
@@ -0,0 +1,357 @@
+// Copyright 1998-2016 Epic Games, Inc. All Rights Reserved.
+
+#include "EditorWidgetsPrivatePCH.h"
+#include "STextPropertyEditableTextBox.h"
+
+#define LOCTEXT_NAMESPACE "STextPropertyEditableTextBox"
+
+FText STextPropertyEditableTextBox::MultipleValuesText(NSLOCTEXT("PropertyEditor", "MultipleValues", "Multiple Values"));
+
+void STextPropertyEditableTextBox::Construct(const FArguments& InArgs, const TSharedRef<IEditableTextProperty>& InEditableTextProperty)
+{
+	EditableTextProperty = InEditableTextProperty;
+
+	TSharedPtr<SHorizontalBox> HorizontalBox;
+
+	const bool bIsPassword = EditableTextProperty->IsPassword();
+	bIsMultiLine = EditableTextProperty->IsMultiLineText();
+	if (bIsMultiLine)
+	{
+		ChildSlot
+		[
+			SAssignNew(HorizontalBox, SHorizontalBox)
+			+SHorizontalBox::Slot()
+			.FillWidth(1.0f)
+			[
+				SNew(SBox)
+				.MinDesiredWidth(InArgs._MinDesiredWidth)
+				.MaxDesiredHeight(InArgs._MaxDesiredHeight)
+				[
+					SAssignNew(MultiLineWidget, SMultiLineEditableTextBox)
+					.Text(this, &STextPropertyEditableTextBox::GetTextValue)
+					.ToolTipText(this, &STextPropertyEditableTextBox::GetToolTipText)
+					.Style(InArgs._Style)
+					.Font(InArgs._Font)
+					.ForegroundColor(InArgs._ForegroundColor)
+					.SelectAllTextWhenFocused(false)
+					.ClearKeyboardFocusOnCommit(false)
+					.OnTextCommitted(this, &STextPropertyEditableTextBox::OnTextCommitted)
+					.SelectAllTextOnCommit(false)
+					.IsReadOnly(this, &STextPropertyEditableTextBox::IsReadOnly)
+					.AutoWrapText(InArgs._AutoWrapText)
+					.WrapTextAt(InArgs._WrapTextAt)
+					.ModiferKeyForNewLine(EModifierKey::Shift)
+					.IsPassword(bIsPassword)
+				]
+			]
+		];
+
+		PrimaryWidget = MultiLineWidget;
+	}
+	else
+	{
+		ChildSlot
+		[
+			SAssignNew(HorizontalBox, SHorizontalBox)
+			+SHorizontalBox::Slot()
+			.FillWidth(1.0f)
+			[
+				SNew(SBox)
+				.MinDesiredWidth(InArgs._MinDesiredWidth)
+				[
+					SAssignNew(SingleLineWidget, SEditableTextBox)
+					.Text(this, &STextPropertyEditableTextBox::GetTextValue)
+					.ToolTipText(this, &STextPropertyEditableTextBox::GetToolTipText)
+					.Style(InArgs._Style)
+					.Font(InArgs._Font)
+					.ForegroundColor(InArgs._ForegroundColor)
+					.SelectAllTextWhenFocused(true)
+					.ClearKeyboardFocusOnCommit(false)
+					.OnTextCommitted(this, &STextPropertyEditableTextBox::OnTextCommitted)
+					.SelectAllTextOnCommit(true)
+					.IsReadOnly(this, &STextPropertyEditableTextBox::IsReadOnly)
+					.IsPassword(bIsPassword)
+				]
+			]
+		];
+
+		PrimaryWidget = SingleLineWidget;
+	}
+
+	HorizontalBox->AddSlot()
+		.AutoWidth()
+		[
+			SNew(SComboButton)
+			.HAlign(HAlign_Center)
+			.VAlign(VAlign_Center)
+			.ContentPadding(FMargin(4, 0))
+			.ButtonStyle(FEditorStyle::Get(), "HoverHintOnly")
+			.ForegroundColor(FSlateColor::UseForeground())
+			.MenuContent()
+			[
+				SNew(SBox)
+				.WidthOverride(340)
+				.Padding(4)
+				[
+					SNew(SGridPanel)
+					.FillColumn(1, 1.0f)
+
+					// Namespace
+					+SGridPanel::Slot(0, 0)
+					.Padding(2)
+					.HAlign(HAlign_Right)
+					[
+						SNew(STextBlock)
+						.Text(LOCTEXT("TextNamespaceLabel", "Namespace:"))
+					]
+					+SGridPanel::Slot(1, 0)
+					.Padding(2)
+					[
+						SNew(SEditableTextBox)
+						.Text(this, &STextPropertyEditableTextBox::GetNamespaceValue)
+						.SelectAllTextWhenFocused(true)
+						.ClearKeyboardFocusOnCommit(false)
+						.OnTextCommitted(this, &STextPropertyEditableTextBox::OnNamespaceCommitted)
+						.SelectAllTextOnCommit(true)
+						.IsReadOnly(this, &STextPropertyEditableTextBox::IsReadOnly)
+					]
+
+					// Key
+					+SGridPanel::Slot(0, 1)
+					.Padding(2)
+					.HAlign(HAlign_Right)
+					[
+						SNew(STextBlock)
+						.Text(LOCTEXT("TextKeyLabel", "Key:"))
+					]
+					+SGridPanel::Slot(1, 1)
+					.Padding(2)
+					[
+						SNew(SEditableTextBox)
+						.Text(this, &STextPropertyEditableTextBox::GetKeyValue)
+						.IsReadOnly(true)
+					]
+				]
+			]
+		];
+
+	SetEnabled(TAttribute<bool>(this, &STextPropertyEditableTextBox::CanEdit));
+}
+
+void STextPropertyEditableTextBox::GetDesiredWidth(float& OutMinDesiredWidth, float& OutMaxDesiredWidth)
+{
+	if (bIsMultiLine)
+	{
+		OutMinDesiredWidth = 250.0f;
+	}
+	else
+	{
+		OutMinDesiredWidth = 125.0f;
+	}
+
+	OutMaxDesiredWidth = 600.0f;
+}
+
+bool STextPropertyEditableTextBox::SupportsKeyboardFocus() const
+{
+	return PrimaryWidget.IsValid() && PrimaryWidget->SupportsKeyboardFocus();
+}
+
+FReply STextPropertyEditableTextBox::OnFocusReceived(const FGeometry& MyGeometry, const FFocusEvent& InFocusEvent)
+{
+	// Forward keyboard focus to our editable text widget
+	return FReply::Handled().SetUserFocus(PrimaryWidget.ToSharedRef(), InFocusEvent.GetCause());
+}
+
+void STextPropertyEditableTextBox::Tick(const FGeometry& AllottedGeometry, const double InCurrentTime, const float InDeltaTime)
+{
+	const float CurrentHeight = AllottedGeometry.GetLocalSize().Y;
+	if (bIsMultiLine && PreviousHeight.IsSet() && PreviousHeight.GetValue() != CurrentHeight)
+	{
+		EditableTextProperty->RequestRefresh();
+	}
+	PreviousHeight = CurrentHeight;
+}
+
+bool STextPropertyEditableTextBox::CanEdit() const
+{
+	return !EditableTextProperty->IsReadOnly();
+}
+
+bool STextPropertyEditableTextBox::IsReadOnly() const
+{
+	return EditableTextProperty->IsReadOnly();
+}
+
+FText STextPropertyEditableTextBox::GetToolTipText() const
+{
+	FText LocalizedTextToolTip;
+	const int32 NumTexts = EditableTextProperty->GetNumTexts();
+	if (NumTexts == 1)
+	{
+		const FText TextValue = EditableTextProperty->GetText(0);
+
+		bool bIsLocalized = false;
+		FString Namespace;
+		FString Key;
+		const FString* SourceString = FTextInspector::GetSourceString(TextValue);
+
+		if (SourceString && TextValue.ShouldGatherForLocalization())
+		{
+			bIsLocalized = FTextLocalizationManager::Get().FindNamespaceAndKeyFromDisplayString(FTextInspector::GetSharedDisplayString(TextValue), Namespace, Key);
+		}
+
+		if (bIsLocalized)
+		{
+			LocalizedTextToolTip = FText::Format(LOCTEXT("LocalizedTextToolTipFmt", "--- Localized Text ---\nNamespace: {0}\nKey: {1}\nSource: {2}"), FText::FromString(Namespace), FText::FromString(Key), FText::FromString(*SourceString));
+		}
+	}
+	
+	const FText BaseToolTipText = EditableTextProperty->GetToolTipText();
+	if (LocalizedTextToolTip.IsEmptyOrWhitespace())
+	{
+		return BaseToolTipText;
+	}
+	if (BaseToolTipText.IsEmptyOrWhitespace())
+	{
+		return LocalizedTextToolTip;
+	}
+
+	return FText::Format(LOCTEXT("ToolTipCompleteFmt", "{0}\n\n{1}"), BaseToolTipText, LocalizedTextToolTip);
+}
+
+FText STextPropertyEditableTextBox::GetTextValue() const
+{
+	FText TextValue;
+
+	const int32 NumTexts = EditableTextProperty->GetNumTexts();
+	if (NumTexts == 1)
+	{
+		TextValue = EditableTextProperty->GetText(0);
+	}
+	else if (NumTexts > 1)
+	{
+		TextValue = MultipleValuesText;
+	}
+
+	return TextValue;
+}
+
+void STextPropertyEditableTextBox::OnTextCommitted(const FText& NewText, ETextCommit::Type CommitInfo)
+{
+	const int32 NumTexts = EditableTextProperty->GetNumTexts();
+
+	// Don't commit the Multiple Values text if there are multiple properties being set
+	if (NumTexts > 0 && (NumTexts == 1 || NewText.ToString() != MultipleValuesText.ToString()))
+	{
+		EditableTextProperty->PreEdit();
+
+		const FString& SourceString = NewText.ToString();
+		for (int32 TextIndex = 0; TextIndex < NumTexts; ++TextIndex)
+		{
+			const FText PropertyValue = EditableTextProperty->GetText(TextIndex);
+
+			// Only apply the change if the new text is different - we want to keep the keys stable where possible
+			if (PropertyValue.ToString().Equals(NewText.ToString(), ESearchCase::CaseSensitive))
+			{
+				continue;
+			}
+
+			// We want to preserve the namespace set on this property if it's *not* the default value
+			FString NewNamespace;
+			if (!EditableTextProperty->IsDefaultValue())
+			{
+				// Some properties report that they're not the default, but still haven't been set from a property, so we also check the property key to see if it's a valid GUID before allowing the namespace to persist
+				FGuid TmpGuid;
+				if (FGuid::Parse(FTextInspector::GetKey(PropertyValue).Get(TEXT("")), TmpGuid))
+				{
+					NewNamespace = FTextInspector::GetNamespace(PropertyValue).Get(TEXT(""));
+				}
+			}
+
+			EditableTextProperty->SetText(TextIndex, FText::ChangeKey(NewNamespace, FGuid::NewGuid().ToString(), NewText));
+		}
+
+		EditableTextProperty->PostEdit();
+	}
+}
+
+FText STextPropertyEditableTextBox::GetNamespaceValue() const
+{
+	FText NamespaceValue;
+
+	const int32 NumTexts = EditableTextProperty->GetNumTexts();
+	if (NumTexts == 1)
+	{
+		const FText PropertyValue = EditableTextProperty->GetText(0);
+		TOptional<FString> FoundNamespace = FTextInspector::GetNamespace(PropertyValue);
+		if (FoundNamespace.IsSet())
+		{
+			NamespaceValue = FText::FromString(FoundNamespace.GetValue());
+		}
+	}
+	else if (NumTexts > 1)
+	{
+		NamespaceValue = MultipleValuesText;
+	}
+
+	return NamespaceValue;
+}
+
+void STextPropertyEditableTextBox::OnNamespaceCommitted(const FText& NewText, ETextCommit::Type CommitInfo)
+{
+	const int32 NumTexts = EditableTextProperty->GetNumTexts();
+
+	// Don't commit the Multiple Values text if there are multiple properties being set
+	if (NumTexts > 0 && (NumTexts == 1 || NewText.ToString() != MultipleValuesText.ToString()))
+	{
+		EditableTextProperty->PreEdit();
+
+		const FString& Namespace = NewText.ToString();
+		for (int32 TextIndex = 0; TextIndex < NumTexts; ++TextIndex)
+		{
+			const FText PropertyValue = EditableTextProperty->GetText(TextIndex);
+
+			// Only apply the change if the new namespace is different - we want to keep the keys stable where possible
+			if (FTextInspector::GetNamespace(PropertyValue).Get(TEXT("")).Equals(Namespace, ESearchCase::CaseSensitive))
+			{
+				continue;
+			}
+
+			// If the current key is a GUID, then we can preserve that when setting the new namespace
+			FGuid NewKeyGuid;
+			if (!FGuid::Parse(FTextInspector::GetKey(PropertyValue).Get(TEXT("")), NewKeyGuid))
+			{
+				NewKeyGuid = FGuid::NewGuid();
+			}
+
+			EditableTextProperty->SetText(TextIndex, FText::ChangeKey(Namespace, NewKeyGuid.ToString(), PropertyValue));
+		}
+
+		EditableTextProperty->PostEdit();
+	}
+}
+
+FText STextPropertyEditableTextBox::GetKeyValue() const
+{
+	FText KeyValue;
+
+	const int32 NumTexts = EditableTextProperty->GetNumTexts();
+	if (NumTexts == 1)
+	{
+		const FText PropertyValue = EditableTextProperty->GetText(0);
+		TOptional<FString> FoundKey = FTextInspector::GetKey(PropertyValue);
+		if (FoundKey.IsSet())
+		{
+			KeyValue = FText::FromString(FoundKey.GetValue());
+		}
+	}
+	else if (NumTexts > 1)
+	{
+		KeyValue = MultipleValuesText;
+	}
+
+	return KeyValue;
+}
+
+#undef LOCTEXT_NAMESPACE
diff --git a/Engine/Source/Editor/EditorWidgets/Public/STextPropertyEditableTextBox.h b/Engine/Source/Editor/EditorWidgets/Public/STextPropertyEditableTextBox.h
new file mode 100644
--- /dev/null
+++ b/Engine/Source/Editor/EditorWidgets/Public/STextPropertyEditableTextBox.h
@@ -0,0 +1,106 @@
+// Copyright 1998-2016 Epic Games, Inc. All Rights Reserved.
+
+#pragma once
+
+/** Interface to allow STextPropertyEditableTextBox to be used to edit both properties and Blueprint pins */
+class IEditableTextProperty
+{
+public:
+	virtual ~IEditableTextProperty() {}
+
+	/** Are the text properties being edited marked as multi-line? */
+	virtual bool IsMultiLineText() const = 0;
+
+	/** Are the text properties being edited marked as password fields? */
+	virtual bool IsPassword() const = 0;
+
+	/** Are the text properties being edited read-only? */
+	virtual bool IsReadOnly() const = 0;
+
+	/** Is the value associated with the properties the default value? */
+	virtual bool IsDefaultValue() const = 0;
+
+	/** Get the tooltip text associated with the property being edited */
+	virtual FText GetToolTipText() const = 0;
+
+	/** Get the number of FText instances being edited by this property */
+	virtual int32 GetNumTexts() const = 0;
+
+	/** Get the text at the given index (check against GetNumTexts) */
+	virtual FText GetText(const int32 InIndex) const = 0;
+
+	/** Set the text at the given index (check against GetNumTexts) */
+	virtual void SetText(const int32 InIndex, const FText& InText) = 0;
+
+	/** Called prior to editing the value of the texts associated with the property being edited */
+	virtual void PreEdit() = 0;
+
+	/** Called after editing the value of the texts associated with the property being edited */
+	virtual void PostEdit() = 0;
+
+	/** Request a refresh of the property UI (eg, due to a size change) */
+	virtual void RequestRefresh() = 0;
+};
+
+/** A widget that can be used for editing FText instances */
+class EDITORWIDGETS_API STextPropertyEditableTextBox : public SCompoundWidget
+{
+	SLATE_BEGIN_ARGS(STextPropertyEditableTextBox)
+		: _Style(&FCoreStyle::Get().GetWidgetStyle<FEditableTextBoxStyle>("NormalEditableTextBox"))
+		, _Font()
+		, _ForegroundColor()
+		, _WrapTextAt(0.0f)
+		, _AutoWrapText(false)
+		, _MinDesiredWidth()
+		, _MaxDesiredHeight(300.0f)
+		{}
+		/** The styling of the textbox */
+		SLATE_STYLE_ARGUMENT(FEditableTextBoxStyle, Style)
+		/** Font color and opacity (overrides Style) */
+		SLATE_ATTRIBUTE(FSlateFontInfo, Font)
+		/** Text color and opacity (overrides Style) */
+		SLATE_ATTRIBUTE(FSlateColor, ForegroundColor)
+		/** Whether text wraps onto a new line when it's length exceeds this width; if this value is zero or negative, no wrapping occurs */
+		SLATE_ATTRIBUTE(float, WrapTextAt)
+		/** Whether to wrap text automatically based on the widget's computed horizontal space */
+		SLATE_ATTRIBUTE(bool, AutoWrapText)
+		/** When specified, will report the MinDesiredWidth if larger than the content's desired width */
+		SLATE_ATTRIBUTE(FOptionalSize, MinDesiredWidth)
+		/** When specified, will report the MaxDesiredHeight if smaller than the content's desired height */
+		SLATE_ATTRIBUTE(FOptionalSize, MaxDesiredHeight)
+	SLATE_END_ARGS()
+
+public:
+	void Construct(const FArguments& Arguments, const TSharedRef<IEditableTextProperty>& InEditableTextProperty);
+	virtual bool SupportsKeyboardFocus() const override;
+	virtual FReply OnFocusReceived(const FGeometry& MyGeometry, const FFocusEvent& InFocusEvent) override;
+	virtual void Tick(const FGeometry& AllottedGeometry, const double InCurrentTime, const float InDeltaTime) override;
+
+private:
+	void GetDesiredWidth(float& OutMinDesiredWidth, float& OutMaxDesiredWidth);
+	bool CanEdit() const;
+	bool IsReadOnly() const;
+	FText GetToolTipText() const;
+
+	FText GetTextValue() const;
+	void OnTextCommitted(const FText& NewText, ETextCommit::Type CommitInfo);
+
+	FText GetNamespaceValue() const;
+	void OnNamespaceCommitted(const FText& NewText, ETextCommit::Type CommitInfo);
+
+	FText GetKeyValue() const;
+
+	TSharedPtr<IEditableTextProperty> EditableTextProperty;
+
+	TSharedPtr<class SWidget> PrimaryWidget;
+
+	TSharedPtr<SMultiLineEditableTextBox> MultiLineWidget;
+
+	TSharedPtr<SEditableTextBox> SingleLineWidget;
+
+	TOptional<float> PreviousHeight;
+
+	bool bIsMultiLine;
+
+	static FText MultipleValuesText;
+};
diff --git a/Engine/Source/Editor/GameProjectGeneration/Private/SProjectBrowser.cpp b/Engine/Source/Editor/GameProjectGeneration/Private/SProjectBrowser.cpp
--- a/Engine/Source/Editor/GameProjectGeneration/Private/SProjectBrowser.cpp
+++ b/Engine/Source/Editor/GameProjectGeneration/Private/SProjectBrowser.cpp
@@ -259,7 +259,7 @@ void SProjectBrowser::Construct( const FArguments& InArgs )
 
 				// Auto-load project
 				+ SHorizontalBox::Slot()
-				.FillWidth(1.0f)
+				.AutoWidth()
 				.VAlign(VAlign_Center)
 				[
 					SNew(SCheckBox)			
@@ -271,6 +271,11 @@ void SProjectBrowser::Construct( const FArguments& InArgs )
 					]
 				]
 
+				+SHorizontalBox::Slot()
+				.FillWidth(1.0f)
+				[
+					SNullWidget::NullWidget
+				]
 				// Browse Button
 				+SHorizontalBox::Slot()
 				.AutoWidth()
diff --git a/Engine/Source/Editor/GameProjectGeneration/Private/Tests/GameProjectAutomationTests.cpp b/Engine/Source/Editor/GameProjectGeneration/Private/Tests/GameProjectAutomationTests.cpp
--- a/Engine/Source/Editor/GameProjectGeneration/Private/Tests/GameProjectAutomationTests.cpp
+++ b/Engine/Source/Editor/GameProjectGeneration/Private/Tests/GameProjectAutomationTests.cpp
@@ -11,6 +11,8 @@
 #include "../TemplateCategory.h"
 #include "../TemplateItem.h"
 
+#if WITH_DEV_AUTOMATION_TESTS
+
 DEFINE_LOG_CATEGORY_STATIC(LogGameProjectGenerationTests, Log, All);
 
 namespace GameProjectAutomationUtils
@@ -376,4 +378,6 @@ bool FCreateCPPTemplateProjectAutomationTests::RunTest(const FString& Parameters
 
 	return (OutMatchedProjectsDesk == OutCreatedProjectsDesk) && (OutMatchedProjectsMob == OutCreatedProjectsMob);
 
-}
\ No newline at end of file
+}
+
+#endif //WITH_DEV_AUTOMATION_TESTS
\ No newline at end of file
diff --git a/Engine/Source/Editor/GeometryMode/Classes/GeomModifier_Extrude.h b/Engine/Source/Editor/GeometryMode/Classes/GeomModifier_Extrude.h
--- a/Engine/Source/Editor/GeometryMode/Classes/GeomModifier_Extrude.h
+++ b/Engine/Source/Editor/GeometryMode/Classes/GeomModifier_Extrude.h
@@ -24,6 +24,7 @@ class UGeomModifier_Extrude : public UGeomModifier_Edit
 
 
 	//~ Begin UGeomModifier Interface
+	virtual bool InputDelta(class FEditorViewportClient* InViewportClient, FViewport* InViewport, FVector& InDrag, FRotator& InRot, FVector& InScale) override;
 	virtual bool Supports() override;
 	virtual void Initialize() override;
 	virtual void WasActivated() override;
diff --git a/Engine/Source/Editor/GeometryMode/Private/GeometryModifiers.cpp b/Engine/Source/Editor/GeometryMode/Private/GeometryModifiers.cpp
--- a/Engine/Source/Editor/GeometryMode/Private/GeometryModifiers.cpp
+++ b/Engine/Source/Editor/GeometryMode/Private/GeometryModifiers.cpp
@@ -658,6 +658,33 @@ UGeomModifier_Extrude::UGeomModifier_Extrude(const FObjectInitializer& ObjectIni
 	Segments = 1;
 }
 
+bool UGeomModifier_Extrude::InputDelta(FEditorViewportClient* InViewportClient, FViewport* InViewport, FVector& InDrag, FRotator& InRot, FVector& InScale)
+{
+	FEdModeGeometry* Mode = (FEdModeGeometry*)GLevelEditorModeTools().GetActiveMode(FBuiltinEditorModes::EM_Geometry);
+
+	const bool bGetRawValue = true;
+	const bool bIsLocalCoords = GLevelEditorModeTools().GetCoordSystem(bGetRawValue) == COORD_Local;
+
+	if (!bIsLocalCoords)
+	{
+		// Before the modal dialog has popped up, force tracking to stop and reset the focus
+		InViewportClient->LostFocus(InViewport);
+		InViewportClient->ReceivedFocus(InViewport);
+		CheckCoordinatesMode();
+	}
+
+	if (!bIsLocalCoords || Mode->GetCurrentWidgetAxis() != EAxisList::X)
+	{
+		InDrag = FVector::ZeroVector;
+		InRot = FRotator::ZeroRotator;
+		InScale = FVector::ZeroVector;
+		return false;
+	}
+
+	return Super::InputDelta(InViewportClient, InViewport, InDrag, InRot, InScale);
+}
+
+
 bool UGeomModifier_Extrude::Supports()
 {
 	FEdModeGeometry* mode = (FEdModeGeometry*)GLevelEditorModeTools().GetActiveMode(FBuiltinEditorModes::EM_Geometry);
diff --git a/Engine/Source/Editor/GraphEditor/GraphEditor.Build.cs b/Engine/Source/Editor/GraphEditor/GraphEditor.Build.cs
--- a/Engine/Source/Editor/GraphEditor/GraphEditor.Build.cs
+++ b/Engine/Source/Editor/GraphEditor/GraphEditor.Build.cs
@@ -25,12 +25,6 @@ public class GraphEditor : ModuleRules
             }
         );
 
-		PrivateIncludePathModuleNames.AddRange(
-			new string[] { 
-				"EditorWidgets"
-			}
-		);
-
 		PrivateDependencyModuleNames.AddRange(
 			new string[] {
                 "AppFramework",
@@ -41,6 +35,7 @@ public class GraphEditor : ModuleRules
 				"Slate",
 				"SlateCore",
                 "EditorStyle",
+				"EditorWidgets",
 				"UnrealEd",
 				"AssetRegistry",
 				"ClassViewer",
@@ -56,7 +51,6 @@ public class GraphEditor : ModuleRules
 		DynamicallyLoadedModuleNames.AddRange(
 			new string[] {
 				"ContentBrowser",
-				"EditorWidgets"
 			}
 		);
 	}
diff --git a/Engine/Source/Editor/GraphEditor/Private/KismetPins/SGraphPinText.cpp b/Engine/Source/Editor/GraphEditor/Private/KismetPins/SGraphPinText.cpp
--- a/Engine/Source/Editor/GraphEditor/Private/KismetPins/SGraphPinText.cpp
+++ b/Engine/Source/Editor/GraphEditor/Private/KismetPins/SGraphPinText.cpp
@@ -3,50 +3,98 @@
 
 #include "GraphEditorCommon.h"
 #include "SGraphPinText.h"
-#include "Editor/UnrealEd/Public/ScopedTransaction.h"
+#include "ScopedTransaction.h"
+#include "STextPropertyEditableTextBox.h"
 
-void SGraphPinText::Construct(const FArguments& InArgs, UEdGraphPin* InGraphPinObj)
+namespace
 {
-	SGraphPin::Construct(SGraphPin::FArguments(), InGraphPinObj);
-}
+	/** Allows STextPropertyEditableTextBox to edit a graph pin */
+	class FEditableTextGraphPin : public IEditableTextProperty
+	{
+	public:
+		FEditableTextGraphPin(UEdGraphPin* InGraphPinObj)
+			: GraphPinObj(InGraphPinObj)
+		{
+		}
 
-TSharedRef<SWidget>	SGraphPinText::GetDefaultValueWidget()
-{
-	return SNew(SBox)
-		.MinDesiredWidth(18.0f)
-		.MaxDesiredHeight(200)
-		[
-			SNew(SMultiLineEditableTextBox)
-			.Style(FEditorStyle::Get(), "Graph.EditableTextBox")
-			.Text(this, &SGraphPinText::GetTypeInValue)
-			.SelectAllTextWhenFocused(true)
-			.Visibility(this, &SGraphPin::GetDefaultValueVisibility)
-			.IsReadOnly(this, &SGraphPinText::GetDefaultValueIsReadOnly)
-			.OnTextCommitted(this, &SGraphPinText::SetTypeInValue)
-			.ForegroundColor(FSlateColor::UseForeground())
-			.WrapTextAt(400)
-			.ModiferKeyForNewLine(EModifierKey::Shift)
-		];
-}
+		virtual bool IsMultiLineText() const override
+		{
+			return true;
+		}
 
-FText SGraphPinText::GetTypeInValue() const
-{
-	return GraphPinObj->DefaultTextValue;
+		virtual bool IsPassword() const override
+		{
+			return false;
+		}
+
+		virtual bool IsReadOnly() const override
+		{
+			return GraphPinObj->bDefaultValueIsReadOnly;
+		}
+
+		virtual bool IsDefaultValue() const override
+		{
+			FString TextAsString;
+			if (FTextStringHelper::WriteToString(TextAsString, GraphPinObj->DefaultTextValue))
+			{
+				return TextAsString.Equals(GraphPinObj->AutogeneratedDefaultValue, ESearchCase::CaseSensitive);
+			}
+			return false;
+		}
+
+		virtual FText GetToolTipText() const override
+		{
+			return FText::GetEmpty();
+		}
+
+		virtual int32 GetNumTexts() const override
+		{
+			return 1;
+		}
+
+		virtual FText GetText(const int32 InIndex) const override
+		{
+			check(InIndex == 0);
+			return GraphPinObj->DefaultTextValue;
+		}
+
+		virtual void SetText(const int32 InIndex, const FText& InText) override
+		{
+			check(InIndex == 0);
+			const FScopedTransaction Transaction(NSLOCTEXT("GraphEditor", "ChangeTxtPinValue", "Change Text Pin Value"));
+			GraphPinObj->Modify();
+			GraphPinObj->GetSchema()->TrySetDefaultText(*GraphPinObj, InText);
+		}
+
+		virtual void PreEdit() override
+		{
+		}
+
+		virtual void PostEdit() override
+		{
+		}
+
+		virtual void RequestRefresh() override
+		{
+		}
+
+	private:
+		UEdGraphPin* GraphPinObj;
+	};
 }
 
-void SGraphPinText::SetTypeInValue(const FText& NewTypeInValue, ETextCommit::Type /*CommitInfo*/)
+void SGraphPinText::Construct(const FArguments& InArgs, UEdGraphPin* InGraphPinObj)
 {
-	if(!GraphPinObj->DefaultTextValue.EqualTo(NewTypeInValue))
-	{
-		const FScopedTransaction Transaction( NSLOCTEXT("GraphEditor", "ChangeTxtPinValue", "Change Text Pin Value" ) );
-		GraphPinObj->Modify();
-
-		// Update the localization string.
-		GraphPinObj->GetSchema()->TrySetDefaultText(*GraphPinObj, NewTypeInValue);
-	}
+	SGraphPin::Construct(SGraphPin::FArguments(), InGraphPinObj);
 }
 
-bool SGraphPinText::GetDefaultValueIsReadOnly() const
+TSharedRef<SWidget>	SGraphPinText::GetDefaultValueWidget()
 {
-	return GraphPinObj->bDefaultValueIsReadOnly;
-}
\ No newline at end of file
+	return SNew(STextPropertyEditableTextBox, MakeShareable(new FEditableTextGraphPin(GraphPinObj)))
+		.Style(FEditorStyle::Get(), "Graph.EditableTextBox")
+		.Visibility(this, &SGraphPin::GetDefaultValueVisibility)
+		.ForegroundColor(FSlateColor::UseForeground())
+		.WrapTextAt(400)
+		.MinDesiredWidth(18.0f)
+		.MaxDesiredHeight(200);
+}
diff --git a/Engine/Source/Editor/GraphEditor/Private/KismetPins/SGraphPinText.h b/Engine/Source/Editor/GraphEditor/Private/KismetPins/SGraphPinText.h
--- a/Engine/Source/Editor/GraphEditor/Private/KismetPins/SGraphPinText.h
+++ b/Engine/Source/Editor/GraphEditor/Private/KismetPins/SGraphPinText.h
@@ -14,10 +14,4 @@ protected:
 	//~ Begin SGraphPin Interface
 	virtual TSharedRef<SWidget>	GetDefaultValueWidget() override;
 	//~ End SGraphPin Interface
-
-	FText GetTypeInValue() const;
-	virtual void SetTypeInValue(const FText& NewTypeInValue, ETextCommit::Type CommitInfo);
-
-	/** @return True if the pin default value field is read-only */
-	bool GetDefaultValueIsReadOnly() const;
 };
diff --git a/Engine/Source/Editor/Kismet/Private/SSCSEditor.cpp b/Engine/Source/Editor/Kismet/Private/SSCSEditor.cpp
--- a/Engine/Source/Editor/Kismet/Private/SSCSEditor.cpp
+++ b/Engine/Source/Editor/Kismet/Private/SSCSEditor.cpp
@@ -342,6 +342,7 @@ FReply FSCSRowDragDropOp::DroppedOnPanel(const TSharedRef< class SWidget >& Pane
 FSCSEditorTreeNode::FSCSEditorTreeNode(FSCSEditorTreeNode::ENodeType InNodeType)
 	: NodeType(InNodeType)
 	, bNonTransactionalRename(false)
+	, ComponentTemplate(nullptr)
 {
 }
 
@@ -379,11 +380,6 @@ class USCS_Node* FSCSEditorTreeNode::GetSCSNode() const
 	return nullptr;
 }
 
-UActorComponent* FSCSEditorTreeNode::GetComponentTemplate() const
-{
-	return nullptr;
-}
-
 UActorComponent* FSCSEditorTreeNode::GetEditableComponentTemplate(UBlueprint* ActualEditedBlueprint)
 {
 	return nullptr;
@@ -421,6 +417,16 @@ FSCSEditorTreeNode::ENodeType FSCSEditorTreeNode::GetNodeType() const
 	return NodeType;
 }
 
+UActorComponent* FSCSEditorTreeNode::GetComponentTemplate(bool bEvenIfPendingKill) const
+{
+	return ComponentTemplate.Get(bEvenIfPendingKill);
+}
+
+void FSCSEditorTreeNode::SetComponentTemplate(UActorComponent* Component)
+{
+	ComponentTemplate = Component;
+}
+
 bool FSCSEditorTreeNode::IsAttachedTo(FSCSEditorTreeNodePtrType InNodePtr) const
 { 
 	FSCSEditorTreeNodePtrType TestParentPtr = ParentNodePtr;
@@ -820,6 +826,20 @@ FSCSEditorTreeNodeInstancedInheritedComponent::FSCSEditorTreeNodeInstancedInheri
 	check(InstancedComponentName != NAME_None);	// ...otherwise IsRootActor() can return a false positive.
 
 	InstancedComponentOwnerPtr = Owner;
+
+	TInlineComponentArray<UActorComponent*> Components;
+	Owner->GetComponents(Components);
+
+	ComponentTemplate = nullptr;
+	for (auto It = Components.CreateConstIterator(); It; ++It)
+	{
+		UActorComponent* ComponentInstance = *It;
+		if (ComponentInstance->GetFName() == InstancedComponentName)
+		{
+			ComponentTemplate = ComponentInstance;
+			break;
+		}
+	}
 }
 
 bool FSCSEditorTreeNodeInstancedInheritedComponent::IsNative() const
@@ -836,11 +856,11 @@ bool FSCSEditorTreeNodeInstancedInheritedComponent::IsNative() const
 
 bool FSCSEditorTreeNodeInstancedInheritedComponent::IsRootComponent() const
 {
-	UActorComponent* ComponentTemplate = GetComponentTemplate();
+	UActorComponent* Template = GetComponentTemplate();
 
 	if (AActor* OwnerActor = InstancedComponentOwnerPtr.Get())
 	{
-		if (OwnerActor->GetRootComponent() == ComponentTemplate)
+		if (OwnerActor->GetRootComponent() == Template)
 		{
 			return true;
 		}
@@ -861,8 +881,8 @@ bool FSCSEditorTreeNodeInstancedInheritedComponent::IsDefaultSceneRoot() const
 
 bool FSCSEditorTreeNodeInstancedInheritedComponent::CanEditDefaults() const
 {
-	UActorComponent* ComponentTemplate = GetComponentTemplate();
-	return (ComponentTemplate ? ComponentTemplate->IsEditableWhenInherited() : false);
+	UActorComponent* Template = GetComponentTemplate();
+	return (Template ? Template->IsEditableWhenInherited() : false);
 }
 
 FText FSCSEditorTreeNodeInstancedInheritedComponent::GetDisplayName() const
@@ -876,26 +896,6 @@ FText FSCSEditorTreeNodeInstancedInheritedComponent::GetDisplayName() const
 	return FText::GetEmpty();
 }
 
-UActorComponent* FSCSEditorTreeNodeInstancedInheritedComponent::GetComponentTemplate() const
-{
-	if (InstancedComponentOwnerPtr.IsValid())
-	{
-		TInlineComponentArray<UActorComponent*> Components;
-		InstancedComponentOwnerPtr.Get()->GetComponents(Components);
-
-		for (auto It = Components.CreateConstIterator(); It; ++It)
-		{
-			UActorComponent* ComponentInstance = *It;
-			if (ComponentInstance->GetFName() == InstancedComponentName)
-			{
-				return ComponentInstance;
-			}
-		}
-	}
-
-	return nullptr;
-}
-
 UActorComponent* FSCSEditorTreeNodeInstancedInheritedComponent::GetEditableComponentTemplate(UBlueprint* ActualEditedBlueprint)
 {
 	if (CanEditDefaults())
@@ -915,20 +915,34 @@ FSCSEditorTreeNodeInstanceAddedComponent::FSCSEditorTreeNodeInstanceAddedCompone
 	check(InstancedComponentName != NAME_None);	// ...otherwise IsRootActor() can return a false positive.
 
 	InstancedComponentOwnerPtr = Owner;
+
+	TInlineComponentArray<UActorComponent*> Components;
+	Owner->GetComponents(Components);
+
+	ComponentTemplate = nullptr;
+	for (auto It = Components.CreateConstIterator(); It; ++It)
+	{
+		UActorComponent* ComponentInstance = *It;
+		if (ComponentInstance->GetFName() == InstancedComponentName)
+		{
+			ComponentTemplate = ComponentInstance;
+			break;
+		}
+	}
 }
 
 bool FSCSEditorTreeNodeInstanceAddedComponent::IsRootComponent() const
 {
 	bool bIsRoot = true;
-	UActorComponent* ComponentTemplate = GetComponentTemplate();
+	UActorComponent* Template = GetComponentTemplate();
 
-	if (ComponentTemplate != NULL)
+	if (Template != NULL)
 	{
-		AActor* CDO = ComponentTemplate->GetOwner();
+		AActor* CDO = Template->GetOwner();
 		if (CDO != NULL)
 		{
 			// Evaluate to TRUE if we have a valid component reference that matches the native root component
-			bIsRoot = (ComponentTemplate == CDO->GetRootComponent());
+			bIsRoot = (Template == CDO->GetRootComponent());
 		}
 	}
 
@@ -955,26 +969,6 @@ FText FSCSEditorTreeNodeInstanceAddedComponent::GetDisplayName() const
 	return FText::FromName(InstancedComponentName);
 }
 
-UActorComponent* FSCSEditorTreeNodeInstanceAddedComponent::GetComponentTemplate() const
-{
-	if (InstancedComponentOwnerPtr.IsValid())
-	{
-		TInlineComponentArray<UActorComponent*> Components;
-		InstancedComponentOwnerPtr.Get()->GetComponents(Components);
-
-		for (auto It = Components.CreateConstIterator(); It; ++It)
-		{
-			UActorComponent* ComponentInstance = *It;
-			if (ComponentInstance->GetFName() == InstancedComponentName)
-			{
-				return ComponentInstance;
-			}
-		}
-	}
-
-	return nullptr;
-}
-
 UActorComponent* FSCSEditorTreeNodeInstanceAddedComponent::GetEditableComponentTemplate(UBlueprint* ActualEditedBlueprint)
 {
 	return GetComponentTemplate();
@@ -1022,17 +1016,17 @@ void FSCSEditorTreeNodeInstanceAddedComponent::OnCompleteRename(const FText& InN
 FSCSEditorTreeNodeComponent::FSCSEditorTreeNodeComponent(USCS_Node* InSCSNode, bool bInIsInheritedSCS)
 	: bIsInheritedSCS(bInIsInheritedSCS)
 	, SCSNodePtr(InSCSNode)
-	, ComponentTemplatePtr((InSCSNode != nullptr) ? InSCSNode->ComponentTemplate : nullptr)
 {
+	ComponentTemplate = (InSCSNode != nullptr) ? InSCSNode->ComponentTemplate : nullptr;
 }
 
 FSCSEditorTreeNodeComponent::FSCSEditorTreeNodeComponent(UActorComponent* InComponentTemplate)
 	: bIsInheritedSCS(false)
 	, SCSNodePtr(nullptr)
-	, ComponentTemplatePtr(InComponentTemplate)
 {
 	check(InComponentTemplate != nullptr);
 
+	ComponentTemplate = InComponentTemplate;
 	AActor* Owner = InComponentTemplate->GetOwner();
 	if (Owner != nullptr)
 	{
@@ -1124,11 +1118,6 @@ class USCS_Node* FSCSEditorTreeNodeComponent::GetSCSNode() const
 	return SCSNodePtr.Get();
 }
 
-UActorComponent* FSCSEditorTreeNodeComponent::GetComponentTemplate() const
-{
-	return ComponentTemplatePtr.Get();
-}
-
 UActorComponent* FSCSEditorTreeNodeComponent::GetEditableComponentTemplate(UBlueprint* ActualEditedBlueprint)
 {
 	if (CanEditDefaults())
@@ -3492,6 +3481,7 @@ void SSCSEditor::Construct( const FArguments& InArgs )
 	if (EditorMode == EComponentEditorMode::ActorInstance)
 	{
 		GEngine->OnLevelComponentRequestRename().AddSP(this, &SSCSEditor::OnLevelComponentRequestRename);
+		GEditor->OnObjectsReplaced().AddSP(this, &SSCSEditor::OnObjectsReplaced);
 	}
 }
 END_SLATE_FUNCTION_BUILD_OPTIMIZATION
@@ -3508,93 +3498,48 @@ void SSCSEditor::OnLevelComponentRequestRename(const UActorComponent* InComponen
 	}
 }
 
-UBlueprint* SSCSEditor::GetBlueprint() const
+void SSCSEditor::OnObjectsReplaced(const TMap<UObject*, UObject*>& OldToNewInstanceMap)
 {
-	if (AActor* Actor = GetActorContext())
-	{
-		UClass* ActorClass = Actor->GetClass();
-		check(ActorClass != nullptr);
-
-		return Cast<UBlueprint>(ActorClass->ClassGeneratedBy);
-	}
-
-	return nullptr;
+	ReplaceComponentReferencesInTree(GetRootComponentNodes(), OldToNewInstanceMap);
 }
 
-void SSCSEditor::Tick( const FGeometry& AllottedGeometry, const double InCurrentTime, const float InDeltaTime )
+void SSCSEditor::ReplaceComponentReferencesInTree(const TArray<FSCSEditorTreeNodePtrType>& Nodes, const TMap<UObject*, UObject*>& OldToNewInstanceMap)
 {
-	SCompoundWidget::Tick(AllottedGeometry, InCurrentTime, InDeltaTime);
-
-	TFunction<bool (const TArray<FSCSEditorTreeNodePtrType>&, int32&)> AreAnyNodesInvalidLambda = [&](const TArray<FSCSEditorTreeNodePtrType>& InNodes, int32& OutNumValidNodes) -> bool
+	for (const FSCSEditorTreeNodePtrType& Node : Nodes)
 	{
-		bool bFoundInvalidNode = false;
-		for(auto NodeIt = InNodes.CreateConstIterator(); NodeIt && !bFoundInvalidNode && !AreAnyNodesInvalidLambda((*NodeIt)->GetChildren(), OutNumValidNodes); ++NodeIt)
+		if (Node.IsValid())
 		{
-			const FSCSEditorTreeNodePtrType Node = *NodeIt;
-			bool bIsComponent = Node->GetNodeType() == FSCSEditorTreeNode::ComponentNode;
-			if (bIsComponent)
+			// We need to get the actual pointer to the old component which will be marked for pending kill, as these are the references which need updating
+			const bool bEvenIfPendingKill = true;
+			UActorComponent* ComponentTemplate = Node->GetComponentTemplate(bEvenIfPendingKill);
+			if (ComponentTemplate)
 			{
-				// A component node is considered to be invalid if one of the following conditions holds true:
-				//	(a) The component instance is invalid (i.e. NULL).
-				//	(b) The component instance is pending garbage collection.
-				//	(c) The component instance is a USceneComponent derivative and either is not the scene "root" node or otherwise is not contained within the scene component node hierarchy.
-				//	(d) The component instance is not a USceneComponent derivative and either matches the scene "root" node or is otherwise contained within the scene component node hierarchy.
-				//		=> (c) and (d) could happen if the user changes the parent class of a custom component's Blueprint class, for example.
-				//
-				FSCSEditorTreeNodePtrType ParentNodePtr = Node->GetParent();
-				const UActorComponent* InstancedComponent = Node->GetComponentTemplate();
-				bFoundInvalidNode = !InstancedComponent
-					|| InstancedComponent->IsPendingKill()
-					|| (InstancedComponent->IsA<USceneComponent>() && (Node != SceneRootNodePtr && (!ParentNodePtr.IsValid() || ParentNodePtr->GetNodeType() != FSCSEditorTreeNode::ComponentNode)))
-					|| (!InstancedComponent->IsA<USceneComponent>() && (Node == SceneRootNodePtr || (ParentNodePtr.IsValid() && ParentNodePtr->GetNodeType() == FSCSEditorTreeNode::ComponentNode)));
-					
-				if(!bFoundInvalidNode)
+				UObject* const* NewComponentTemplatePtr = OldToNewInstanceMap.Find(ComponentTemplate);
+				if (NewComponentTemplatePtr)
 				{
-					++OutNumValidNodes;
+					if (UActorComponent* NewComponentTemplate = Cast<UActorComponent>(*NewComponentTemplatePtr))
+					{
+						Node->SetComponentTemplate(NewComponentTemplate);
+					}
 				}
 			}
-		}
 
-		return bFoundInvalidNode;
-	};
+			ReplaceComponentReferencesInTree(Node->GetChildren(), OldToNewInstanceMap);
+		}
+	}
+}
 
-	// Only run this in "instanced" mode or otherwise only when the option to show UCS-instanced components is enabled
-	const bool bShowInstancedUCSComponents = !GetDefault<UBlueprintEditorSettings>()->bHideConstructionScriptComponentsInDetailsView;
-	if(EditorMode == EComponentEditorMode::ActorInstance || bShowInstancedUCSComponents)
+UBlueprint* SSCSEditor::GetBlueprint() const
+{
+	if (AActor* Actor = GetActorContext())
 	{
-		const AActor* ActorInstance = (EditorMode == EComponentEditorMode::ActorInstance) ? GetActorContext() : PreviewActor.Get(nullptr);
-		if (ActorInstance)
-		{
-			int32 NumComponentInstances = 0;
-			for (auto CompIt = ActorInstance->GetComponents().CreateConstIterator(); CompIt; ++CompIt)
-			{
-				// Don't count editor-only components in instanced mode (unless it's the root component), because we don't show them there. Also, don't count UCS-added components if the option to hide them is enabled, and don't count components that are based on the CDO in the BP editor context (i.e. preview world).
-				const UActorComponent* CompInst = *CompIt;
-				if ((!CompInst->IsEditorOnly() || EditorMode != EComponentEditorMode::ActorInstance || CompInst == ActorInstance->GetRootComponent())
-					&& (CompInst->CreationMethod != EComponentCreationMethod::UserConstructionScript || bShowInstancedUCSComponents)
-					&& (EditorMode != EComponentEditorMode::BlueprintSCS || CompInst->GetArchetype() != CompInst->GetClass()->GetDefaultObject()))
-				{
-					++NumComponentInstances;
-				}
-			}
-
-			int32 NumComponentNodes = 0;
-			static int32 UpdateCounter = 0;
-			if (AreAnyNodesInvalidLambda(GetRootNodes(), NumComponentNodes) || NumComponentNodes != NumComponentInstances)
-			{
-				if(++UpdateCounter % 60)
-				{
-					UE_LOG(LogSCSEditor, Warning, TEXT("Calling UpdateTree() from Tick() more frequently than expected due to invalid tree node state. This might be a performance issue."));
-				}
+		UClass* ActorClass = Actor->GetClass();
+		check(ActorClass != nullptr);
 
-				UpdateTree();
-			}
-			else
-			{
-				UpdateCounter = 0;
-			}
-		}
+		return Cast<UBlueprint>(ActorClass->ClassGeneratedBy);
 	}
+
+	return nullptr;
 }
 
 FReply SSCSEditor::OnKeyDown( const FGeometry& MyGeometry, const FKeyEvent& InKeyEvent )
diff --git a/Engine/Source/Editor/Kismet/Private/Tests/BlueprintEditorTests.cpp b/Engine/Source/Editor/Kismet/Private/Tests/BlueprintEditorTests.cpp
--- a/Engine/Source/Editor/Kismet/Private/Tests/BlueprintEditorTests.cpp
+++ b/Engine/Source/Editor/Kismet/Private/Tests/BlueprintEditorTests.cpp
@@ -22,6 +22,8 @@
 #include "AutomationEditorCommon.h"
 #include "AutomationEditorPromotionCommon.h"
 
+#if WITH_DEV_AUTOMATION_TESTS
+
 #define LOCTEXT_NAMESPACE "BlueprintEditorPromotionTests"
 
 DEFINE_LOG_CATEGORY_STATIC(LogBlueprintEditorPromotionTests, Log, All);
@@ -1427,3 +1429,5 @@ bool FBlueprintEditorPromotionTest::RunTest(const FString& Parameters)
 
 
 #undef LOCTEXT_NAMESPACE
+
+#endif //WITH_DEV_AUTOMATION_TESTS
\ No newline at end of file
diff --git a/Engine/Source/Editor/Kismet/Public/SSCSEditor.h b/Engine/Source/Editor/Kismet/Public/SSCSEditor.h
--- a/Engine/Source/Editor/Kismet/Public/SSCSEditor.h
+++ b/Engine/Source/Editor/Kismet/Public/SSCSEditor.h
@@ -56,10 +56,6 @@ public:
 	 * @return The SCS node that is represented by this object, or NULL if there is no SCS node associated with the component template.
 	 */
 	virtual class USCS_Node* GetSCSNode() const;
-	/**
-	 * @return The component template or instance represented by this object.
-	 */
-	virtual UActorComponent* GetComponentTemplate() const;
 	/**
 	 * @param ActualEditedBlueprint currently edited blueprint
 	 *
@@ -85,6 +81,14 @@ public:
 	 * @return Type of node
 	 */
 	ENodeType GetNodeType() const;
+	/**
+	 * @return The component template or instance represented by this node, if it's a component node.
+	 */
+	UActorComponent* GetComponentTemplate(bool bEvenIfPendingKill = false) const;
+	/**
+	 * Set the component template represented by this node, if it's a component node.
+	 */
+	void SetComponentTemplate(UActorComponent* Component);
 	/**
 	 * @return Whether or not this node is a direct child of the given node.
 	 */
@@ -268,6 +272,9 @@ protected:
 		return bResult;
 	}
 
+	// Component template represented by this node, if it's a component node, otherwise invalid
+	TWeakObjectPtr<UActorComponent> ComponentTemplate;
+
 private:
 	// The type of component tree node
 	ENodeType NodeType;
@@ -328,7 +335,6 @@ public:
 	//virtual FName GetVariableName() const override;
 	//virtual FString GetDisplayString() const override;
 	virtual FText GetDisplayName() const override;
-	virtual UActorComponent* GetComponentTemplate() const override;
 	virtual UActorComponent* GetEditableComponentTemplate(UBlueprint* ActualEditedBlueprint) override;
 	// End of FSCSEditorTreeNode public interface
 
@@ -361,7 +367,6 @@ public:
 	virtual FName GetVariableName() const override { return NAME_None; }
 	virtual FString GetDisplayString() const override;
 	virtual FText GetDisplayName() const override;
-	virtual UActorComponent* GetComponentTemplate() const override;
 	virtual UActorComponent* GetEditableComponentTemplate(UBlueprint* ActualEditedBlueprint) override;
 	virtual void OnCompleteRename(const FText& InNewName) override;
 	// End of FSCSEditorTreeNode public interface
@@ -374,6 +379,7 @@ protected:
 private:
 	FName InstancedComponentName;
 	TWeakObjectPtr<AActor> InstancedComponentOwnerPtr;
+	TWeakObjectPtr<UActorComponent> ComponentTemplate;
 };
 
 //////////////////////////////////////////////////////////////////////////
@@ -410,7 +416,6 @@ public:
 	//virtual FString GetDisplayString() const override;
 	virtual FText GetDisplayName() const override;
 	virtual class USCS_Node* GetSCSNode() const override;
-	virtual UActorComponent* GetComponentTemplate() const override;
 	virtual UActorComponent* GetEditableComponentTemplate(UBlueprint* ActualEditedBlueprint) override;
 	virtual void OnCompleteRename(const FText& InNewName) override;
 	// End of FSCSEditorTreeNode public interface
@@ -429,8 +434,6 @@ private:
 
 	// Is this the template coming from an SCS node?
 	TWeakObjectPtr<class USCS_Node> SCSNodePtr;
-
-	TWeakObjectPtr<UActorComponent> ComponentTemplatePtr;
 };
 
 
@@ -731,7 +734,6 @@ public:
 	void Construct(const FArguments& InArgs);
 
 	/** SWidget interface */
-	virtual void Tick( const FGeometry& AllottedGeometry, const double InCurrentTime, const float InDeltaTime );
 	virtual FReply OnKeyDown( const FGeometry& MyGeometry, const FKeyEvent& InKeyEvent );
 
 	/** Used by tree control - make a widget for a table row from a node */
@@ -900,6 +902,7 @@ public:
 
 	/** Try to handle a drag-drop operation */
 	FReply TryHandleAssetDragDropOperation(const FDragDropEvent& DragDropEvent);
+
 protected:
 	FString GetSelectedClassText() const;
 
@@ -920,6 +923,12 @@ protected:
 	 */
 	void OnRenameComponent(bool bTransactional);
 
+	/** Called when component objects are replaced following construction script execution */
+	void OnObjectsReplaced(const TMap<UObject*, UObject*>& OldToNewInstanceMap);
+
+	/** Update component pointers held by tree nodes if components have been replaced following construction script execution */
+	void ReplaceComponentReferencesInTree(const TArray<FSCSEditorTreeNodePtrType>& Nodes, const TMap<UObject*, UObject*>& OldToNewInstanceMap);
+
 	/**
 	 * Function to create events for the current selection
 	 * @param Blueprint						the active blueprint context
diff --git a/Engine/Source/Editor/KismetCompiler/Private/KismetCompilerVMBackend.cpp b/Engine/Source/Editor/KismetCompiler/Private/KismetCompilerVMBackend.cpp
--- a/Engine/Source/Editor/KismetCompiler/Private/KismetCompilerVMBackend.cpp
+++ b/Engine/Source/Editor/KismetCompiler/Private/KismetCompilerVMBackend.cpp
@@ -525,7 +525,8 @@ public:
 							Prop->ExportText_InContainer(ArrayIter, NewTerm.Name, StructData, StructData, NULL, PPF_None);
 							if (Prop->IsA(UTextProperty::StaticClass()))
 							{
-								NewTerm.TextLiteral = FText::FromString(NewTerm.Name);
+								NewTerm.TextLiteral = Cast<UTextProperty>(Prop)->GetPropertyValue_InContainer(StructData, ArrayIter);
+								NewTerm.Name = NewTerm.TextLiteral.ToString();
 							}
 							else if (Prop->IsA(UObjectProperty::StaticClass()))
 							{
@@ -560,7 +561,8 @@ public:
 					InnerProp->ExportText_Direct(NewTerm.Name, RawElemData, RawElemData, NULL, PPF_None);
 					if (InnerProp->IsA(UTextProperty::StaticClass()))
 					{
-						NewTerm.TextLiteral = FText::FromString(NewTerm.Name);
+						NewTerm.TextLiteral = Cast<UTextProperty>(InnerProp)->GetPropertyValue(RawElemData);
+						NewTerm.Name = NewTerm.TextLiteral.ToString();
 					}
 					else if (InnerProp->IsA(UObjectPropertyBase::StaticClass()))
 					{
diff --git a/Engine/Source/Editor/LandscapeEditor/Private/Tests/LandscapeAutomationTests.cpp b/Engine/Source/Editor/LandscapeEditor/Private/Tests/LandscapeAutomationTests.cpp
--- a/Engine/Source/Editor/LandscapeEditor/Private/Tests/LandscapeAutomationTests.cpp
+++ b/Engine/Source/Editor/LandscapeEditor/Private/Tests/LandscapeAutomationTests.cpp
@@ -10,6 +10,8 @@
 #include "ScopedTransaction.h"
 #include "Landscape.h"
 
+#if WITH_DEV_AUTOMATION_TESTS
+
 #define LOCTEXT_NAMESPACE "LandscapeEditor.NewLandscape"
 
 DEFINE_LOG_CATEGORY_STATIC(LogLandscapeAutomationTests, Log, All);
@@ -170,3 +172,5 @@ bool FLandscapeEditorTest::RunTest(const FString& Parameters)
 }
 
 #undef LOCTEXT_NAMESPACE
+
+#endif //WITH_DEV_AUTOMATION_TESTS
\ No newline at end of file
diff --git a/Engine/Source/Editor/LevelEditor/Private/LevelEditorActions.cpp b/Engine/Source/Editor/LevelEditor/Private/LevelEditorActions.cpp
--- a/Engine/Source/Editor/LevelEditor/Private/LevelEditorActions.cpp
+++ b/Engine/Source/Editor/LevelEditor/Private/LevelEditorActions.cpp
@@ -1080,6 +1080,9 @@ void FLevelEditorActionCallbacks::EditAsset_Clicked( const EToolkitMode::Type To
 
 		if (bShouldOpenEditors)
 		{
+			// Clear focus so the level viewport can receive its focus lost call (and clear pending keyup events which wouldn't arrive)
+			FSlateApplication::Get().ClearKeyboardFocus(EFocusCause::WindowActivate);
+
 			auto LevelEditorSharedPtr = LevelEditor.Pin();
 
 			if (LevelEditorSharedPtr.IsValid())
@@ -2160,8 +2163,10 @@ void FLevelEditorActionCallbacks::CheckOutProjectSettingsConfig( )
 void FLevelEditorActionCallbacks::OnShowOnlySelectedActors()
 {
 	const FScopedTransaction Transaction( NSLOCTEXT( "LevelEditorCommands", "ShowOnlySelectedActors", "Show Only Selected Actors" ) );
-	GUnrealEd->edactUnhideSelected( GetWorld() );
+	// First hide unselected as this will also hide group actor members
 	GUnrealEd->edactHideUnselected( GetWorld() );
+	// Then unhide selected to ensure that everything that's selected will be unhidden
+	GUnrealEd->edactUnhideSelected(GetWorld());
 }
 
 void FLevelEditorActionCallbacks::OnToggleTransformWidgetVisibility()
diff --git a/Engine/Source/Editor/MainFrame/Private/Frame/MainFrameActions.cpp b/Engine/Source/Editor/MainFrame/Private/Frame/MainFrameActions.cpp
--- a/Engine/Source/Editor/MainFrame/Private/Frame/MainFrameActions.cpp
+++ b/Engine/Source/Editor/MainFrame/Private/Frame/MainFrameActions.cpp
@@ -895,7 +895,7 @@ bool FMainFrameActionCallbacks::Redo_CanExecute()
 
 void FMainFrameActionCallbacks::ExecuteExecCommand( FString Command )
 {
-	GUnrealEd->Exec( GEditor->GetEditorWorldContext(true).World(), *Command );
+	GUnrealEd->Exec( GEditor->GetEditorWorldContext(false).World(), *Command );
 }
 
 
diff --git a/Engine/Source/Editor/MaterialEditor/Private/MaterialEditor.cpp b/Engine/Source/Editor/MaterialEditor/Private/MaterialEditor.cpp
--- a/Engine/Source/Editor/MaterialEditor/Private/MaterialEditor.cpp
+++ b/Engine/Source/Editor/MaterialEditor/Private/MaterialEditor.cpp
@@ -1086,9 +1086,19 @@ bool FMaterialEditor::SetPreviewAssetByName(const TCHAR* InAssetName)
 
 void FMaterialEditor::SetPreviewMaterial(UMaterialInterface* InMaterialInterface)
 {
-	if (PreviewViewport.IsValid())
+	if (Material->IsUIMaterial())
+	{
+		if (PreviewUIViewport.IsValid())
+		{
+			PreviewUIViewport->SetPreviewMaterial(InMaterialInterface);
+		}
+	}
+	else
 	{
-		PreviewViewport->SetPreviewMaterial(InMaterialInterface);
+		if (PreviewViewport.IsValid())
+		{
+			PreviewViewport->SetPreviewMaterial(InMaterialInterface);
+		}
 	}
 }
 
@@ -2308,6 +2318,7 @@ void FMaterialEditor::OnPreviewNode()
 			UMaterialGraphNode* GraphNode = Cast<UMaterialGraphNode>(*NodeIt);
 			if (GraphNode)
 			{
+				GraphEditor->NotifyGraphChanged();
 				SetPreviewExpression(GraphNode->MaterialExpression);
 			}
 		}
@@ -2830,8 +2841,12 @@ void FMaterialEditor::SetPreviewExpression(UMaterialExpression* NewPreviewExpres
 		if( ExpressionPreviewMaterial == NULL )
 		{
 			// Create the expression preview material if it hasnt already been created
-			ExpressionPreviewMaterial = NewObject<UMaterial>(GetTransientPackage(), NAME_None, RF_Public);
+			ExpressionPreviewMaterial = NewObject<UPreviewMaterial>(GetTransientPackage(), NAME_None, RF_Public);
 			ExpressionPreviewMaterial->bIsPreviewMaterial = true;
+			if (Material->IsUIMaterial())
+			{
+				ExpressionPreviewMaterial->MaterialDomain = MD_UI;
+			}
 		}
 
 		if (FunctionOutput)
@@ -3375,6 +3390,7 @@ void FMaterialEditor::PasteNodesHere(const FVector2D& Location)
 			}
 
 			NewExpression->UpdateParameterGuid(true, true);
+			Material->AddExpressionParameter(NewExpression, Material->EditorParameters);
 
 			UMaterialExpressionFunctionInput* FunctionInput = Cast<UMaterialExpressionFunctionInput>( NewExpression );
 			if( FunctionInput )
@@ -3562,6 +3578,20 @@ void FMaterialEditor::NotifyPostChange( const FPropertyChangedEvent& PropertyCha
 			TArray<TWeakObjectPtr<UObject>> SelectedObjects = MaterialDetailsView->GetSelectedObjects();
 			MaterialDetailsView->SetObjects( SelectedObjects, true );
 
+			if (ExpressionPreviewMaterial)
+			{
+				if (Material->IsUIMaterial())
+				{
+					ExpressionPreviewMaterial->MaterialDomain = MD_UI;
+				}
+				else
+				{
+					ExpressionPreviewMaterial->MaterialDomain = MD_Surface;
+				}
+
+				SetPreviewMaterial(ExpressionPreviewMaterial);
+			}
+
 			UpdatePreviewViewportsVisibility();
 		}
 
diff --git a/Engine/Source/Editor/MaterialEditor/Private/SMaterialEditorViewport.cpp b/Engine/Source/Editor/MaterialEditor/Private/SMaterialEditorViewport.cpp
--- a/Engine/Source/Editor/MaterialEditor/Private/SMaterialEditorViewport.cpp
+++ b/Engine/Source/Editor/MaterialEditor/Private/SMaterialEditorViewport.cpp
@@ -611,12 +611,14 @@ public:
 	float GetZoomLevel() const;
 
 	void SetPreviewSize( const FVector2D PreviewSize );
+	void SetPreviewMaterial(UMaterialInterface* InPreviewMaterial);
 private:
 	mutable FVector2D CachedSize;
 	float ZoomLevel;
 
 	FMaterialPreviewPanelSlot ChildSlot;
 	TSharedPtr<FSlateMaterialBrush> PreviewBrush;
+	TSharedPtr<SImage> ImageWidget;
 };
 
 
@@ -626,7 +628,7 @@ void SMaterialEditorUIPreviewZoomer::Construct( const FArguments& InArgs, UMater
 
 	ChildSlot
 	[
-		SNew( SImage )
+		SAssignNew( ImageWidget, SImage )
 		.Image( PreviewBrush.Get() )
 	];
 
@@ -698,6 +700,11 @@ void SMaterialEditorUIPreviewZoomer::SetPreviewSize( const FVector2D PreviewSize
 	PreviewBrush->ImageSize = PreviewSize;
 }
 
+void SMaterialEditorUIPreviewZoomer::SetPreviewMaterial(UMaterialInterface* InPreviewMaterial)
+{
+	PreviewBrush = MakeShareable(new FSlateMaterialBrush(*InPreviewMaterial, PreviewBrush->ImageSize));
+	ImageWidget->SetImage(PreviewBrush.Get());
+}
 
 void SMaterialEditorUIPreviewViewport::Construct( const FArguments& InArgs, UMaterialInterface* PreviewMaterial )
 {
@@ -786,6 +793,11 @@ void SMaterialEditorUIPreviewViewport::Construct( const FArguments& InArgs, UMat
 	PreviewZoomer->SetPreviewSize( FVector2D(PreviewSize) );
 }
 
+void SMaterialEditorUIPreviewViewport::SetPreviewMaterial(UMaterialInterface* InMaterialInterface)
+{
+	PreviewZoomer->SetPreviewMaterial(InMaterialInterface);
+}
+
 void SMaterialEditorUIPreviewViewport::OnPreviewXChanged( int32 NewValue )
 {
 	PreviewSize.X = NewValue;
diff --git a/Engine/Source/Editor/MaterialEditor/Private/SMaterialEditorViewport.h b/Engine/Source/Editor/MaterialEditor/Private/SMaterialEditorViewport.h
--- a/Engine/Source/Editor/MaterialEditor/Private/SMaterialEditorViewport.h
+++ b/Engine/Source/Editor/MaterialEditor/Private/SMaterialEditorViewport.h
@@ -111,6 +111,7 @@ public:
 	SLATE_END_ARGS()
 
 	void Construct( const FArguments& InArgs, UMaterialInterface* PreviewMaterial );
+	void SetPreviewMaterial(UMaterialInterface* InMaterialInterface);
 
 private:
 	void OnPreviewXChanged( int32 NewValue );
diff --git a/Engine/Source/Editor/PropertyEditor/Private/PropertyHandleImpl.cpp b/Engine/Source/Editor/PropertyEditor/Private/PropertyHandleImpl.cpp
--- a/Engine/Source/Editor/PropertyEditor/Private/PropertyHandleImpl.cpp
+++ b/Engine/Source/Editor/PropertyEditor/Private/PropertyHandleImpl.cpp
@@ -238,6 +238,12 @@ FString FPropertyValueImpl::GetPropertyValueArray() const
 
 bool FPropertyValueImpl::SendTextToObjectProperty( const FString& Text, EPropertyValueSetFlags::Type Flags )
 {
+	if (IsAnyOuterObjectFromEngine() && !FPackageName::IsEnginePackageName(Text))
+	{
+		UE_LOG(LogPropertyNode, Log, TEXT("Cannot assign a Project object %s to an Engine property."), *Text);
+		return false;
+	}
+
 	TSharedPtr<FPropertyNode> PropertyNodePin = PropertyNode.Pin();
 	if( PropertyNodePin.IsValid() )
 	{
@@ -657,6 +663,28 @@ FPropertyAccess::Result FPropertyValueImpl::GetValueAsDisplayText( FText& OutTex
 	return Res;
 }
 
+bool FPropertyValueImpl::IsAnyOuterObjectFromEngine() const
+{
+	TSharedPtr<FPropertyNode> PropertyNodePin = PropertyNode.Pin();
+	if (PropertyNodePin.IsValid())
+	{
+		FObjectPropertyNode* ObjectNode = PropertyNodePin->FindObjectItemParent();
+		if (ObjectNode)
+		{
+			for (int32 ObjectIndex = 0; ObjectIndex < ObjectNode->GetNumObjects(); ++ObjectIndex)
+			{
+				const FAssetData AssetData(ObjectNode->GetUObject(ObjectIndex));
+				if (FPackageName::IsEnginePackageName(AssetData.ObjectPath.ToString()))
+				{
+					return true;
+				}
+			}
+		}
+	}
+
+	return false;
+}
+
 FPropertyAccess::Result FPropertyValueImpl::SetValueAsString( const FString& InValue, EPropertyValueSetFlags::Type Flags )
 {
 	FPropertyAccess::Result Result = FPropertyAccess::Fail;
@@ -798,6 +826,7 @@ FPropertyAccess::Result FPropertyValueImpl::OnUseSelected()
 				if (!InterfaceThatMustBeImplemented || SelectedClass->ImplementsInterface(InterfaceThatMustBeImplemented))
 				{
 					FString const ClassPathName = SelectedClass->GetPathName();
+					const bool bIsEngineClass = FPackageName::IsEnginePackageName(ClassPathName);
 
 					TArray<FText> RestrictReasons;
 					if (PropertyNodePin->IsRestricted(ClassPathName, RestrictReasons))
@@ -805,7 +834,13 @@ FPropertyAccess::Result FPropertyValueImpl::OnUseSelected()
 						check(RestrictReasons.Num() > 0);
 						FMessageDialog::Open(EAppMsgType::Ok, RestrictReasons[0]);
 					}
-					else 
+					else if (IsAnyOuterObjectFromEngine() && !bIsEngineClass)
+					{
+						FMessageDialog::Open(EAppMsgType::Ok, FText::Format(
+							NSLOCTEXT("UnrealEd", "ClassAssignmentsFailed", "Cannot assign Project class {0} to an Engine property."),
+							FText::FromString(ClassPathName)));
+					}
+					else
 					{
 						
 						Res = SetValueAsString(ClassPathName, EPropertyValueSetFlags::DefaultFlags);
diff --git a/Engine/Source/Editor/PropertyEditor/Private/PropertyHandleImpl.h b/Engine/Source/Editor/PropertyEditor/Private/PropertyHandleImpl.h
--- a/Engine/Source/Editor/PropertyEditor/Private/PropertyHandleImpl.h
+++ b/Engine/Source/Editor/PropertyEditor/Private/PropertyHandleImpl.h
@@ -317,6 +317,8 @@ protected:
 	 */
 	FPropertyAccess::Result GetPropertyValueText( FText& OutText, FPropertyNode* InPropertyNode, const bool bAllowAlternateDisplayValue ) const;
 
+	bool IsAnyOuterObjectFromEngine() const;
+
 protected:
 	/** Property node used to access UProperty and address of object to change */
 	TWeakPtr<FPropertyNode> PropertyNode;
