Fixing the shadow variable warnings in the editor on mac, in the SCSEditor.

diff --git a/Engine/Source/Editor/Kismet/Private/SSCSEditor.cpp b/Engine/Source/Editor/Kismet/Private/SSCSEditor.cpp
--- a/Engine/Source/Editor/Kismet/Private/SSCSEditor.cpp
+++ b/Engine/Source/Editor/Kismet/Private/SSCSEditor.cpp
@@ -340,9 +340,9 @@ FReply FSCSRowDragDropOp::DroppedOnPanel(const TSharedRef< class SWidget >& Pane
 // FSCSEditorTreeNode
 
 FSCSEditorTreeNode::FSCSEditorTreeNode(FSCSEditorTreeNode::ENodeType InNodeType)
-	: NodeType(InNodeType)
+	: ComponentTemplatePtr(nullptr)
+	, NodeType(InNodeType)
 	, bNonTransactionalRename(false)
-	, ComponentTemplate(nullptr)
 {
 }
 
@@ -419,12 +419,12 @@ FSCSEditorTreeNode::ENodeType FSCSEditorTreeNode::GetNodeType() const
 
 UActorComponent* FSCSEditorTreeNode::GetComponentTemplate(bool bEvenIfPendingKill) const
 {
-	return ComponentTemplate.Get(bEvenIfPendingKill);
+	return ComponentTemplatePtr.Get(bEvenIfPendingKill);
 }
 
 void FSCSEditorTreeNode::SetComponentTemplate(UActorComponent* Component)
 {
-	ComponentTemplate = Component;
+	ComponentTemplatePtr = Component;
 }
 
 bool FSCSEditorTreeNode::IsAttachedTo(FSCSEditorTreeNodePtrType InNodePtr) const
@@ -830,13 +830,13 @@ FSCSEditorTreeNodeInstancedInheritedComponent::FSCSEditorTreeNodeInstancedInheri
 	TInlineComponentArray<UActorComponent*> Components;
 	Owner->GetComponents(Components);
 
-	ComponentTemplate = nullptr;
+	SetComponentTemplate(nullptr);
 	for (auto It = Components.CreateConstIterator(); It; ++It)
 	{
 		UActorComponent* ComponentInstance = *It;
 		if (ComponentInstance->GetFName() == InstancedComponentName)
 		{
-			ComponentTemplate = ComponentInstance;
+			SetComponentTemplate(ComponentInstance);
 			break;
 		}
 	}
@@ -1017,7 +1017,7 @@ FSCSEditorTreeNodeComponent::FSCSEditorTreeNodeComponent(USCS_Node* InSCSNode, b
 	: bIsInheritedSCS(bInIsInheritedSCS)
 	, SCSNodePtr(InSCSNode)
 {
-	ComponentTemplate = (InSCSNode != nullptr) ? InSCSNode->ComponentTemplate : nullptr;
+	SetComponentTemplate(( InSCSNode != nullptr ) ? InSCSNode->ComponentTemplate : nullptr);
 }
 
 FSCSEditorTreeNodeComponent::FSCSEditorTreeNodeComponent(UActorComponent* InComponentTemplate)
@@ -1026,7 +1026,7 @@ FSCSEditorTreeNodeComponent::FSCSEditorTreeNodeComponent(UActorComponent* InComp
 {
 	check(InComponentTemplate != nullptr);
 
-	ComponentTemplate = InComponentTemplate;
+	SetComponentTemplate(InComponentTemplate);
 	AActor* Owner = InComponentTemplate->GetOwner();
 	if (Owner != nullptr)
 	{
diff --git a/Engine/Source/Editor/Kismet/Public/SSCSEditor.h b/Engine/Source/Editor/Kismet/Public/SSCSEditor.h
--- a/Engine/Source/Editor/Kismet/Public/SSCSEditor.h
+++ b/Engine/Source/Editor/Kismet/Public/SSCSEditor.h
@@ -273,7 +273,7 @@ protected:
 	}
 
 	// Component template represented by this node, if it's a component node, otherwise invalid
-	TWeakObjectPtr<UActorComponent> ComponentTemplate;
+	TWeakObjectPtr<UActorComponent> ComponentTemplatePtr;
 
 private:
 	// The type of component tree node
