Copying //UE4/Dev-Blueprints to Dev-Main (//UE4/Dev-Main)

diff --git a/Engine/Source/Developer/BlueprintProfiler/Public/BlueprintProfilerModule.h b/Engine/Source/Developer/BlueprintProfiler/Public/BlueprintProfilerModule.h
--- a/Engine/Source/Developer/BlueprintProfiler/Public/BlueprintProfilerModule.h
+++ b/Engine/Source/Developer/BlueprintProfiler/Public/BlueprintProfilerModule.h
@@ -3,8 +3,12 @@
 #pragma once
 
 #include "ModuleManager.h"
-#include "Script.h"
-#include "BlueprintProfilerSupport.h"
+
+class UBlueprint;
+struct EScriptInstrumentationEvent;
+
+/** Delegate to broadcast structural stats changes */
+DECLARE_MULTICAST_DELEGATE_OneParam(FOnBPStatGraphLayoutChanged, TWeakObjectPtr<UBlueprint>);
 
 //////////////////////////////////////////////////////////////////////////
 // IBlueprintProfilerInterface
@@ -22,10 +26,41 @@ public:
 	/** Toggles the profiling event capture state */
 	virtual void ToggleProfilingCapture() {}
 
-	/** Returns the current profiling event data */
-	virtual TArray<FBlueprintInstrumentedEvent>& GetProfilingData() = 0;
-
 	/** Instruments and add a profiling event to the data */
 	virtual void InstrumentEvent(const EScriptInstrumentationEvent& Event) {}
 
-};
+#if WITH_EDITOR
+
+	/** Returns the multicast delegate to signal blueprint layout changes to stats */
+	virtual FOnBPStatGraphLayoutChanged& GetGraphLayoutChangedDelegate() = 0;
+
+	/** Returns the cached object context or creates a new entry */
+	virtual TSharedPtr<class FScriptExecutionContext> GetCachedObjectContext(const FString& ObjectPath) = 0;
+
+	/** Returns the current profiling event data for node */
+	virtual TSharedPtr<class FScriptExecutionNode> GetProfilerDataForNode(const FName NodeName) = 0;
+
+	/** Returns or creates an entry for the node */
+	virtual TSharedPtr<class FScriptExecutionNode> GetOrCreateProfilerDataForNode(const FName NodeName) = 0;
+
+	/** Returns if the profiler currently retains any data for the specified instance */
+	virtual bool HasDataForInstance(const UObject* Instance) { return false; }
+
+	/** Get Valid Instance Path */
+	virtual const FString GetValidInstancePath(const FString PathIn) = 0;
+
+	/** Process profiling event data */
+	virtual void ProcessEventProfilingData() {}
+
+	/** Navigate to blueprint */
+	virtual void NavigateToBlueprint(const FName NameIn) const = 0;
+
+	/** Navigate to instance */
+	virtual void NavigateToInstance(const FName NameIn) const = 0;
+
+	/** Navigate to profiler node */
+	virtual void NavigateToNode(const FName NameIn) const = 0;
+
+#endif // WITH_EDITOR
+
+};
\ No newline at end of file
diff --git a/Engine/Source/Developer/BlueprintProfiler/classes/BlueprintInstrumentationSupport.cpp b/Engine/Source/Developer/BlueprintProfiler/classes/BlueprintInstrumentationSupport.cpp
new file mode 100644
--- /dev/null
+++ b/Engine/Source/Developer/BlueprintProfiler/classes/BlueprintInstrumentationSupport.cpp
@@ -0,0 +1,55 @@
+// Copyright 1998-2016 Epic Games, Inc. All Rights Reserved.
+
+#include "Private/BlueprintProfilerPCH.h"
+
+//////////////////////////////////////////////////////////////////////////
+// FBlueprintInstrumentedEvent
+
+void FBlueprintInstrumentedEvent::SetData(EScriptInstrumentation::Type InEventType, const FString& InPathData, const int32 InCodeOffset)
+{
+	EventType = InEventType;
+	PathData = InPathData;
+	CodeOffset = InCodeOffset;
+	Time = FPlatformTime::Seconds();
+}
+
+FName FBlueprintInstrumentedEvent::GetFunctionName() const
+{
+	FName Result = NAME_None;
+	int32 FunctionStart = INDEX_NONE;
+	if (PathData.FindLastChar(':', FunctionStart))
+	{
+		FunctionStart++;
+		Result = FName(*PathData.Right(PathData.Len() - FunctionStart));
+	}
+	return Result;
+}
+
+//////////////////////////////////////////////////////////////////////////
+// FKismetInstrumentContext
+
+void FKismetInstrumentContext::UpdateContext(const UObject* InContextObject, TArray<FBlueprintInstrumentedEvent>& ProfilingData)
+{
+	// Handle instance context switch
+	if (ContextObject != InContextObject && InContextObject)
+	{
+		// Handle class context switch
+		if (const UClass* NewClass = InContextObject->GetClass())
+		{
+			if (NewClass != ContextClass)
+			{
+				ProfilingData.Add(FBlueprintInstrumentedEvent(EScriptInstrumentation::Class, NewClass->GetPathName(), -1));
+				ContextClass = NewClass;
+			}
+		}
+		// Handle instance change
+		ProfilingData.Add(FBlueprintInstrumentedEvent(EScriptInstrumentation::Instance, InContextObject->GetPathName(), -1));
+		ContextObject = InContextObject;
+	}
+}
+
+void FKismetInstrumentContext::ResetContext()
+{
+	ContextClass.Reset();
+	ContextObject.Reset();
+}
diff --git a/Engine/Source/Developer/BlueprintProfiler/classes/BlueprintProfiler.cpp b/Engine/Source/Developer/BlueprintProfiler/classes/BlueprintProfiler.cpp
--- a/Engine/Source/Developer/BlueprintProfiler/classes/BlueprintProfiler.cpp
+++ b/Engine/Source/Developer/BlueprintProfiler/classes/BlueprintProfiler.cpp
@@ -1,19 +1,21 @@
 // Copyright 1998-2016 Epic Games, Inc. All Rights Reserved.
 
 #include "Private/BlueprintProfilerPCH.h"
+#include "EditorStyleSet.h"
 
-#include "BlueprintProfiler.h"
-
-#if WITH_EDITOR
-#include "Editor.h"
-#include "Script.h"
-#endif
+#define LOCTEXT_NAMESPACE "BlueprintProfiler"
 
 //////////////////////////////////////////////////////////////////////////
 // FBlueprintProfiler
 
 IMPLEMENT_MODULE(FBlueprintProfiler, BlueprintProfiler);
 
+DECLARE_STATS_GROUP(TEXT("BlueprintProfiler"), STATGROUP_BlueprintProfiler, STATCAT_Advanced);
+DECLARE_CYCLE_STAT(TEXT("Tick"), STAT_ProfilerTickCost, STATGROUP_BlueprintProfiler);
+DECLARE_CYCLE_STAT(TEXT("Processing Events"), STAT_ProfilerInstrumentationCost, STATGROUP_BlueprintProfiler);
+DECLARE_CYCLE_STAT(TEXT("Statistic Update"), STAT_StatUpdateCost, STATGROUP_BlueprintProfiler);
+DECLARE_CYCLE_STAT(TEXT("Node Lookup"), STAT_NodeLookupCost, STATGROUP_BlueprintProfiler);
+
 FBlueprintProfiler::FBlueprintProfiler()
 	: bProfilerActive(false)
 	, bProfilingCaptureActive(false)
@@ -54,21 +56,315 @@ void FBlueprintProfiler::ToggleProfilingCapture()
 
 void FBlueprintProfiler::InstrumentEvent(const EScriptInstrumentationEvent& Event)
 {
-//	if (bProfilingCaptureActive)
+	#if WITH_EDITOR
+	SCOPE_CYCLE_COUNTER(STAT_ProfilerInstrumentationCost);
+	#endif
+	// Handle context switching events
+	ActiveContext.UpdateContext(Event.GetContextObject(), InstrumentationEvents);
+	// Add instrumented event
+	const int32 ScriptCodeOffset = GetScriptCodeOffset(Event);
+	const FString CurrentFunctionPath = GetCurrentFunctionPath(Event);
+	InstrumentationEvents.Add(FBlueprintInstrumentedEvent(Event.GetType(), CurrentFunctionPath, ScriptCodeOffset));
+	// Reset context on event end
+	if (Event.GetType() == EScriptInstrumentation::Stop)
+	{
+		ActiveContext.ResetContext();
+	}
+}
+
+int32 FBlueprintProfiler::GetScriptCodeOffset(const EScriptInstrumentationEvent& Event) const
+{
+	int32 ScriptCodeOffset = -1;
+	if (Event.IsStackFrameValid())
+	{
+		const FFrame& StackFrame = Event.GetStackFrame();
+		ScriptCodeOffset = StackFrame.Code - StackFrame.Node->Script.GetData() - 1;
+	}
+	return ScriptCodeOffset;
+}
+
+FString FBlueprintProfiler::GetCurrentFunctionPath(const EScriptInstrumentationEvent& Event) const
+{
+	FString CurrentFunctionPath;
+	if (Event.IsStackFrameValid())
+	{
+		const FFrame& StackFrame = Event.GetStackFrame();
+		CurrentFunctionPath = StackFrame.Node->GetPathName();
+	}
+	else
+	{
+		const FString ClassName = Event.GetContextObject()->GetClass()->GetPathName();
+		CurrentFunctionPath = FString::Printf( TEXT( "%s:%s" ), *ClassName, *Event.GetEventName().ToString());
+	}
+	return CurrentFunctionPath;
+}
+
+#if WITH_EDITOR
+
+TSharedPtr<FScriptExecutionContext> FBlueprintProfiler::GetCachedObjectContext(const FString& ObjectPath)
+{
+	SCOPE_CYCLE_COUNTER(STAT_NodeLookupCost);
+	TSharedPtr<FScriptExecutionContext>& Result = PathToObjectContext.FindOrAdd(ObjectPath);
+	if (!Result.IsValid())
+	{
+		Result = MakeShareable<FScriptExecutionContext>(new FScriptExecutionContext);
+	}
+	return Result;
+}
+
+TSharedPtr<FScriptExecutionNode> FBlueprintProfiler::GetProfilerDataForNode(const FName NodeName)
+{
+	SCOPE_CYCLE_COUNTER(STAT_NodeLookupCost);
+	TSharedPtr<FScriptExecutionNode> Result;
+	if (TSharedPtr<FScriptExecutionNode>* Entry = NodeProfilingData.Find(NodeName))
+	{
+		Result = *Entry;
+	}
+	return Result;
+}
+
+TSharedPtr<FScriptExecutionNode> FBlueprintProfiler::GetOrCreateProfilerDataForNode(const FName NodeName)
+{
+	TSharedPtr<FScriptExecutionNode>& Result = NodeProfilingData.FindOrAdd(NodeName);
+	if (!Result.IsValid())
+	{
+		Result = MakeShareable(new FScriptExecutionNode);
+	}
+	return Result;
+}
+
+bool FBlueprintProfiler::HasDataForInstance(const UObject* Instance)
+{
+	return false;
+}
+
+void FBlueprintProfiler::Tick(float DeltaSeconds)
+{
+	SCOPE_CYCLE_COUNTER(STAT_ProfilerTickCost);
+	ProcessEventProfilingData();
+}
+
+bool FBlueprintProfiler::IsTickable() const
+{
+	return InstrumentationEvents.Num() > 0;
+}
+
+void FBlueprintProfiler::ProcessEventProfilingData()
+{
+	// Really not keen on this function its all very ugly, it mostly works but I'll try and refactor as I get chance.
+
+	// Structure to track trace path stacks
+	struct LocalStruct
+	{
+		const UObject* Object;
+		EScriptInstrumentation::Type Type;
+		int32 ScriptCodeOffset;
+		double EventTime;
+		FTracePath TracePath;
+
+		bool operator == (const LocalStruct& Other) const
+		{
+			return Object == Other.Object; 
+		}
+
+	};
+	// Processing context info
+	int32 ExecStart = 0;
+	int32 ProcessedIndex = 0;
+	TArray<LocalStruct> TraceStack;
+	LocalStruct NewEntry;
+	FName InstanceName = NAME_None;
+	TSharedPtr<FScriptExecutionContext> ExecContext;
+	FTracePath TracePath;
+	// Process incoming event data. This needs pulling out into events and processing.
+	for (int32 EventIdx = 0; EventIdx < InstrumentationEvents.Num(); ++EventIdx)
+	{
+		const FBlueprintInstrumentedEvent& CurrEvent = InstrumentationEvents[EventIdx];
+		NewEntry.Type = CurrEvent.GetType();
+		NewEntry.EventTime = CurrEvent.GetTime();
+		NewEntry.ScriptCodeOffset = CurrEvent.GetScriptCodeOffset();
+
+		switch(CurrEvent.GetType())
+		{
+			case EScriptInstrumentation::Class:
+			{
+				ExecContext = MapBlueprintExecutionFlow(CurrEvent.GetObjectPath());
+				ExecStart = EventIdx;
+				break;
+			}
+			case EScriptInstrumentation::Instance:
+			{
+				MapBlueprintInstance(CurrEvent.GetObjectPath());
+				ExecStart = EventIdx;
+				InstanceName = FName(*GetValidInstancePath(CurrEvent.GetObjectPath()));
+				TracePath.Reset();
+				break;
+			}
+			case EScriptInstrumentation::NodeEntry:
+			case EScriptInstrumentation::PureNodeEntry:
+			{
+				UFunction* NodeFunction = ExecContext->BlueprintClass->FindFunctionByName(CurrEvent.GetFunctionName());
+				NewEntry.Object = ExecContext.IsValid() ? ExecContext->GetNodeFromCodeLocation(CurrEvent.GetScriptCodeOffset(), NodeFunction) : nullptr;
+				NewEntry.TracePath = TracePath;
+				// Add new event.
+				TraceStack.Add(NewEntry);
+				break;
+			}
+			case EScriptInstrumentation::NodeExit:
+			{
+				if (TraceStack.Num() > 0)
+				{
+					const int32 ScriptCodeOffset = CurrEvent.GetScriptCodeOffset();
+					UFunction* NodeFunction = ExecContext->BlueprintClass->FindFunctionByName(CurrEvent.GetFunctionName());
+					const UEdGraphNode* GraphNode = ExecContext.IsValid() ? ExecContext->GetNodeFromCodeLocation(ScriptCodeOffset, NodeFunction) : nullptr;
+					// Jumps in scripts can cause multiple node exits on branches, check here and ignore the first exit if two are present.
+					if (EventIdx + 1 < InstrumentationEvents.Num() && InstrumentationEvents[EventIdx+1].GetType() == EScriptInstrumentation::NodeExit)
+					{
+						UFunction* NextNodeFunction = ExecContext->BlueprintClass->FindFunctionByName(InstrumentationEvents[EventIdx+1].GetFunctionName());
+						const UEdGraphNode* NextGraphNode = ExecContext.IsValid() ? ExecContext->GetNodeFromCodeLocation(InstrumentationEvents[EventIdx+1].GetScriptCodeOffset(), NextNodeFunction) : nullptr;
+						if (NextGraphNode && NextGraphNode == GraphNode)
+						{
+							// Skip this branch jump.
+							continue;
+						}
+					}
+					// Group up timings and submit.
+					if (TraceStack.Last().Object == GraphNode && TraceStack.Last().Type == EScriptInstrumentation::NodeEntry)
+					{
+						// Find node time and pool pure node timings.
+						const double NodeTime = NewEntry.EventTime - TraceStack.Last().EventTime;
+						double PureNodeTime = 0.0;
+						const int32 NodeIndex = TraceStack.Num() - 1;
+						int32 PureNodeIndex = NodeIndex - 1;
+						for (; PureNodeIndex >= 0 && TraceStack[PureNodeIndex] == TraceStack[NodeIndex]; --PureNodeIndex)
+						{
+							PureNodeTime = TraceStack[NodeIndex].EventTime - TraceStack[PureNodeIndex].EventTime;
+						}
+						// Find exec node from graph node
+						TSharedPtr<FScriptExecutionNode> ExecNode = GetProfilerDataForNode(FName(*GraphNode->GetPathName()));
+
+						if (ExecNode.IsValid())
+						{
+							if (ExecNode->GetNumChildren() > 0)
+							{
+								// Restore TraceStack when exiting macro/function
+								TracePath = TraceStack.Last().TracePath;
+							}
+							// Update performance data on the tracepath.
+							TSharedPtr<FScriptPerfData> PerfData = ExecNode->GetPerfDataByInstanceAndTracePath(InstanceName, TracePath);
+							PerfData->AddEventTiming(PureNodeTime, NodeTime);
+						}
+						// Flush the trace stack of the current events.
+						TraceStack.RemoveAt(PureNodeIndex+1, NodeIndex - PureNodeIndex, false);
+						// Apply new node timing
+						ProcessedIndex = EventIdx;
+					}
+					// Update exit pin
+					TracePath.AddExitPin(ScriptCodeOffset);
+				}
+				break;
+			}
+			case EScriptInstrumentation::Event:
+			{
+				TracePath.Reset();
+				// Start event timing
+				NewEntry.Object = ExecContext->BlueprintClass->FindFunctionByName(CurrEvent.GetFunctionName());
+				NewEntry.TracePath.Reset();
+				// Add new event.
+				TraceStack.Add(NewEntry);
+				break;
+			}
+			case EScriptInstrumentation::Stop:
+			{
+				// Add event timing
+				if (TraceStack.Num() > 0 && TraceStack.Last().Type == EScriptInstrumentation::Event)
+				{
+					// Find node time and pool pure node timings.
+					const double NodeTime = NewEntry.EventTime - TraceStack.Last().EventTime;
+					// Find event exec node
+					FName EventPath(*FString::Printf( TEXT( "%s:%s" ), *ExecContext->ContextObject.Get()->GetName(), *TraceStack.Last().Object->GetName()));
+					TSharedPtr<FScriptExecutionNode> ExecNode = GetProfilerDataForNode(EventPath);
+					TracePath.Reset();
+
+					if (ExecNode.IsValid())
+					{
+						// Update performance data on the tracepath.
+						TSharedPtr<FScriptPerfData> PerfData = ExecNode->GetPerfDataByInstanceAndTracePath(InstanceName, TracePath);
+						PerfData->AddEventTiming(0.0, NodeTime);
+					}
+					// Flush the trace stack of the current events.
+					TraceStack.SetNum(0);
+				}
+				// Update stats
+				if (ExecContext.IsValid())
+				{
+					SCOPE_CYCLE_COUNTER(STAT_StatUpdateCost);
+					ExecContext->UpdateConnectedStats();
+				}
+				// Update processed index.
+				ProcessedIndex = EventIdx;
+				break;
+			}
+		}
+	}
+	// Reset profiler data
+	if (ProcessedIndex > 0)
+	{
+		InstrumentationEvents.RemoveAt(0, ProcessedIndex+1);
+	}
+}
+
+void FBlueprintProfiler::RemoveAllBlueprintReferences(UBlueprint* Blueprint)
+{
+	// Remove all blueprint references
+	if (Blueprint)
 	{
-		// Handle context switching events
-		ActiveContext.UpdateContext(Event.GetContextObject(), ProfilingData);
-		// Add instrumented event
-		const int32 ScriptCodeOffset = GetScriptCodeOffset(Event);
-		const FString CurrentFunctionPath = GetCurrentFunctionPath(Event);
-		ProfilingData.Add(FBlueprintInstrumentedEvent(Event.GetType(), CurrentFunctionPath, ScriptCodeOffset));
+		const FString BlueprintPath = Blueprint->GeneratedClass->GetPathName();
+		// Remove Object Context
+		if (TSharedPtr<FScriptExecutionContext>* ContextResult = PathToObjectContext.Find(BlueprintPath))
+		{
+			TSharedPtr<FScriptExecutionContext> ContextToRemove = *ContextResult;
+			// Clean up instances
+			TSharedPtr<FScriptExecutionBlueprint> BlueprintRoot = StaticCastSharedPtr<FScriptExecutionBlueprint>(ContextToRemove->ExecutionNode);
+			if (BlueprintRoot.IsValid())
+			{
+				// Remove PIE remaps
+				for (int32 InstanceIdx = 0; InstanceIdx < BlueprintRoot->GetInstanceCount(); ++InstanceIdx)
+				{
+					const FString InstancePath(BlueprintRoot->GetInstanceByIndex(InstanceIdx)->GetName().ToString());
+					PIEInstancePathMap.Remove(InstancePath);
+					if (const FString* Key = PIEInstancePathMap.FindKey(InstancePath))
+					{
+						PIEInstancePathMap.Remove(*Key);
+					}
+				}
+			}
+			// Clean up nodes
+			TMap<FName, TSharedPtr<FScriptExecutionNode>> ExecNodes;
+			ContextToRemove->GetAllExecNodes(ExecNodes);
+			// Remove all exec nodes.
+			for (auto ExecNodeIter : ExecNodes)
+			{
+				NodeProfilingData.Remove(ExecNodeIter.Key);
+				PathToObjectContext.Remove(ExecNodeIter.Key.ToString());
+			}
+			// Remove compilation hook, this will get added again if the blueprint is profiled
+			Blueprint->OnCompiled().RemoveAll(this);
+			// Broadcast changes to stats structure
+			GraphLayoutChangedDelegate.Broadcast(Blueprint);
+		}
 	}
 }
 
+#endif // WITH_EDITOR
+
 void FBlueprintProfiler::ResetProfilingData()
 {
-	ProfilingData.Reset();
+	InstrumentationEvents.Reset();
+	NodeProfilingData.Reset();
+	PathToObjectContext.Reset();
 	ActiveContext.ResetContext();
+	PIEInstancePathMap.Reset();
 }
 
 void FBlueprintProfiler::RegisterDelegates(bool bEnabled)
@@ -101,7 +397,6 @@ void FBlueprintProfiler::BeginPIE(bool bIsSimulating)
 {
 #if WITH_EDITOR
 	bPIEActive = true;
-	ResetProfilingData();
 	bProfilingCaptureActive = bProfilerActive && bPIEActive;
 #endif // WITH_EDITOR
 }
@@ -114,24 +409,414 @@ void FBlueprintProfiler::EndPIE(bool bIsSimulating)
 #endif // WITH_EDITOR
 }
 
-int32 FBlueprintProfiler::GetScriptCodeOffset(const EScriptInstrumentationEvent& Event) const
+TSharedPtr<FScriptExecutionContext> FBlueprintProfiler::MapBlueprintExecutionFlow(const FString& BlueprintPath)
 {
-	int32 ScriptCodeOffset = -1;
-	if (Event.IsStackFrameValid())
+	TSharedPtr<FScriptExecutionContext> Result = GetCachedObjectContext(BlueprintPath);
+
+	if (!Result->IsContextValid())
 	{
-		const FFrame& StackFrame = Event.GetStackFrame();
-		ScriptCodeOffset = StackFrame.Code - StackFrame.Node->Script.GetData() - 1;
+		// Create new blueprint root node
+		TSharedPtr<FScriptExecutionNode>& BlueprintEntry = NodeProfilingData.FindOrAdd(FName(*BlueprintPath));
+		if (!BlueprintEntry.IsValid())
+		{
+			BlueprintEntry = MakeShareable<FScriptExecutionNode>(new FScriptExecutionBlueprint);
+		}
+		Result->ExecutionNode = BlueprintEntry;
+		BlueprintEntry->SetNameByString(BlueprintPath);
+		BlueprintEntry->SetGraphName(NAME_None);
+		BlueprintEntry->SetFlags(EScriptExecutionNodeFlags::Class);
+		const FSlateBrush* Icon_Normal = FEditorStyle::GetBrush(TEXT("BlueprintProfiler.BPIcon_Normal"));
+		BlueprintEntry->SetIcon(const_cast<FSlateBrush*>(Icon_Normal));
+		BlueprintEntry->SetIconColor(FLinearColor(0.46f, 0.54f, 0.81f));
+		BlueprintEntry->SetToolTipText(LOCTEXT("NavigateToBlueprintHyperlink_ToolTip", "Navigate to the Blueprint"));
+		// Attempt to locate the blueprint
+		UBlueprint* BlueprintToMap = nullptr;
+		UBlueprintGeneratedClass* BPGC = nullptr;
+
+		if (UObject* ObjectPtr = FindObject<UObject>(nullptr, *BlueprintPath))
+		{
+			BPGC = Cast<UBlueprintGeneratedClass>(ObjectPtr);
+			if (BPGC)
+			{
+				Result->BlueprintClass = BPGC;
+				Result->BlueprintFunction = BPGC->UberGraphFunction;
+				BlueprintToMap = Cast<UBlueprint>(BPGC->ClassGeneratedBy);
+				if (BlueprintToMap)
+				{
+					Result->ContextObject = BlueprintToMap;
+					BlueprintEntry->SetDisplayName(FText::FromName(BlueprintToMap->GetFName()));
+				}
+			}
+		}
+		if (BlueprintToMap && BPGC)
+		{
+			// Register delegate to handle updates
+			BlueprintToMap->OnCompiled().AddRaw(this, &FBlueprintProfiler::RemoveAllBlueprintReferences);
+			// Find all event nodes for event graph and UCS
+			TArray<UEdGraph*> Graphs;
+			TArray<UK2Node_Event*> EventNodes;
+			TArray<UK2Node_FunctionEntry*> FunctionEntryNodes;
+			BlueprintToMap->GetAllGraphs(Graphs);
+			for (auto GraphIter : Graphs)
+			{
+				if (GraphIter->GetFName() == UEdGraphSchema_K2::GN_EventGraph)
+				{
+					GraphIter->GetNodesOfClass<UK2Node_Event>(EventNodes);
+				}
+				if (GraphIter->GetFName() == UEdGraphSchema_K2::FN_UserConstructionScript)
+				{
+					GraphIter->GetNodesOfClass<UK2Node_FunctionEntry>(FunctionEntryNodes);
+				}
+			}
+			// Map execution paths for each event node
+			for (auto EventIter : EventNodes)
+			{
+				
+				const FString EventPath = FString::Printf(TEXT("%s:%s"), *BlueprintToMap->GetName(), *EventIter->GetFunctionName().ToString());
+				FName EventName(*EventPath);
+				TSharedPtr<FScriptExecutionContext> EventContext = MakeShareable<FScriptExecutionContext>(new FScriptExecutionContext);
+				*EventContext = *Result;
+				EventContext->ContextObject = EventIter;
+				EventContext->ExecutionNode = GetOrCreateProfilerDataForNode(EventName);
+				EventContext->ExecutionNode->SetName(EventName);
+				EventContext->ExecutionNode->SetGraphName(EventIter->GetTypedOuter<UEdGraph>()->GetFName());
+				EventContext->ExecutionNode->SetFlags(EScriptExecutionNodeFlags::Event);
+				FLinearColor IconColor = FLinearColor(0.91f, 0.16f, 0.16f);
+				FText ToolTipText = LOCTEXT("NavigateToEventLocationHyperlink_ToolTip", "Navigate to the Event");
+				EventContext->ExecutionNode->SetDisplayName(EventIter->GetNodeTitle(ENodeTitleType::ListView));
+				EventContext->ExecutionNode->SetToolTipText(ToolTipText);
+				const FSlateBrush* Icon_Node = EventIter->ShowPaletteIconOnNode() ?	FEditorStyle::GetBrush(EventIter->GetPaletteIcon(IconColor)) :
+																				FEditorStyle::GetBrush(TEXT("BlueprintProfiler.BPNode"));
+				EventContext->ExecutionNode->SetIcon(const_cast<FSlateBrush*>(Icon_Node));
+				EventContext->ExecutionNode->SetIconColor(IconColor);
+				// Add event to mapped blueprint
+				Result->ExecutionNode->AddChildNode(EventContext->ExecutionNode);
+				// Map the event execution flow
+				TSharedPtr<FScriptExecutionNode> MappedEvent = MapNodeExecutionFlow(EventContext, EventIter);
+				if (MappedEvent.IsValid())
+				{
+					EventContext->ExecutionNode->AddChildNode(MappedEvent);
+				}
+			}
+			// Map execution paths for each function entry node
+			for (auto FunctionIter : FunctionEntryNodes)
+			{
+				const FString FunctionPath = FString::Printf(TEXT("%s:%s"), *BlueprintToMap->GetName(), *FunctionIter->SignatureName.ToString());
+				FName FunctionName(*FunctionPath);
+				TSharedPtr<FScriptExecutionContext> FunctionContext = MakeShareable<FScriptExecutionContext>(new FScriptExecutionContext);
+				*FunctionContext = *Result;
+				FunctionContext->ContextObject = FunctionIter;
+				FunctionContext->BlueprintFunction = BPGC->FindFunctionByName(FunctionIter->SignatureName, EIncludeSuperFlag::ExcludeSuper);
+				FunctionContext->ExecutionNode = GetOrCreateProfilerDataForNode(FunctionName);
+				FunctionContext->ExecutionNode->SetName(FunctionName);
+				FunctionContext->ExecutionNode->SetGraphName(FunctionIter->GetTypedOuter<UEdGraph>()->GetFName());
+				FunctionContext->ExecutionNode->SetFlags(EScriptExecutionNodeFlags::Event);
+				FLinearColor IconColor = FLinearColor(0.91f, 0.16f, 0.16f);
+				FText ToolTipText = LOCTEXT("NavigateToFunctionLocationHyperlink_ToolTip", "Navigate to the Function");
+				FunctionContext->ExecutionNode->SetDisplayName(FunctionIter->GetNodeTitle(ENodeTitleType::ListView));
+				FunctionContext->ExecutionNode->SetToolTipText(ToolTipText);
+				const FSlateBrush* Icon = FunctionIter->ShowPaletteIconOnNode() ?	FEditorStyle::GetBrush(FunctionIter->GetPaletteIcon(IconColor)) :
+																					FEditorStyle::GetBrush(TEXT("BlueprintProfiler.BPNode"));
+				FunctionContext->ExecutionNode->SetIcon(const_cast<FSlateBrush*>(Icon));
+				FunctionContext->ExecutionNode->SetIconColor(IconColor);
+				// Add Function to mapped blueprint
+				Result->ExecutionNode->AddChildNode(FunctionContext->ExecutionNode);
+				// Map the Function execution flow
+				TSharedPtr<FScriptExecutionNode> MappedFunction = MapNodeExecutionFlow(FunctionContext, FunctionIter);
+				if (MappedFunction.IsValid())
+				{
+					FunctionContext->ExecutionNode->AddChildNode(MappedFunction);
+				}
+			}
+		}
 	}
-	return ScriptCodeOffset;
+	return Result;
 }
 
-FString FBlueprintProfiler::GetCurrentFunctionPath(const EScriptInstrumentationEvent& Event) const
+TSharedPtr<FScriptExecutionNode> FBlueprintProfiler::MapNodeExecutionFlow(TSharedPtr<FScriptExecutionContext> InFunctionContext, UEdGraphNode* NodeToMap)
 {
-	FString CurrentFunctionPath;
-	if (Event.IsStackFrameValid())
+	TSharedPtr<FScriptExecutionNode> MappedNode;
+
+	if (NodeToMap)
 	{
-		const FFrame& StackFrame = Event.GetStackFrame();
-		CurrentFunctionPath = StackFrame.Node->GetPathName();
+		TSharedPtr<FScriptExecutionContext> NodeContext = GetCachedObjectContext(NodeToMap->GetPathName());
+		NodeContext->ContextObject = NodeToMap;
+		NodeContext->BlueprintClass = InFunctionContext->BlueprintClass;
+		NodeContext->BlueprintFunction = InFunctionContext->BlueprintFunction;
+		UBlueprintGeneratedClass* BlueprintClass = NodeContext->BlueprintClass.Get();
+		UFunction* BlueprintFunction = NodeContext->BlueprintFunction.Get();
+
+		if (NodeContext->ExecutionNode.IsValid())
+		{
+			MappedNode = NodeContext->ExecutionNode;
+		}
+		else
+		{
+			FName NodeName(FName(*NodeToMap->GetPathName()));
+			NodeContext->ExecutionNode = GetOrCreateProfilerDataForNode(NodeName);
+			NodeContext->ExecutionNode->SetName(NodeName);
+			NodeContext->ExecutionNode->SetGraphName(NodeToMap->GetTypedOuter<UEdGraph>()->GetFName());
+			NodeContext->ExecutionNode->SetFlags(EScriptExecutionNodeFlags::Node);
+			FLinearColor IconColor = FLinearColor(1.f, 1.f, 1.f, 0.8f);
+			FText ToolTipText = LOCTEXT("NavigateToNodeLocationHyperlink_ToolTip", "Navigate to the Node");
+			MappedNode = NodeContext->ExecutionNode;
+			// Evaluate Execution Pins
+			TArray<UEdGraphPin*> ExecPins;
+			for (auto Pin : NodeToMap->Pins)
+			{
+				if (Pin->Direction == EGPD_Output && Pin->PinType.PinCategory == UEdGraphSchema_K2::PC_Exec)
+				{
+					ExecPins.Add(Pin);
+				}
+			}
+			const bool bBranchNode = ExecPins.Num() > 1;
+			if (bBranchNode)
+			{
+				// Determine branch exec type
+				IconColor = FLinearColor(1.f, 1.f, 1.f, 0.8f);
+				if (NodeToMap->IsA<UK2Node_ExecutionSequence>())
+				{
+					MappedNode->AddFlags(EScriptExecutionNodeFlags::SequentialBranch);
+				}
+				else
+				{
+					MappedNode->AddFlags(EScriptExecutionNodeFlags::ConditionalBranch);
+				}
+			}
+			for (auto Pin : ExecPins)
+			{
+				const UEdGraphPin* LinkedPin = Pin->LinkedTo.Num() > 0 ? Pin->LinkedTo[0] : nullptr;
+
+				if (bBranchNode)
+				{
+					const int32 PinScriptCodeOffset = BlueprintClass->GetDebugData().FindPinToCodeAssociation(Pin, BlueprintFunction);
+					check(PinScriptCodeOffset != INDEX_NONE);
+					FName PinName(FName(*Pin->GetPathName()));
+					TSharedPtr<FScriptExecutionNode> LinkedNode = GetOrCreateProfilerDataForNode(PinName);
+					LinkedNode->SetName(PinName);
+					LinkedNode->SetGraphName(Pin->GetTypedOuter<UEdGraph>()->GetFName());
+					LinkedNode->SetDisplayName(Pin->GetDisplayName());
+					LinkedNode->SetFlags(EScriptExecutionNodeFlags::ExecPin);
+					LinkedNode->SetToolTipText(LOCTEXT("ExecPin_ToolTip", "Expand execution path"));
+					const FSlateBrush* Icon = LinkedPin ?	FEditorStyle::GetBrush(TEXT("BlueprintProfiler.BPPinConnected")) : 
+															FEditorStyle::GetBrush(TEXT("BlueprintProfiler.BPPinDisconnected"));
+					LinkedNode->SetIcon(const_cast<FSlateBrush*>(Icon));
+					LinkedNode->SetIconColor(FLinearColor(1.f, 1.f, 1.f, 0.8f));
+					MappedNode->AddLinkedNode(PinScriptCodeOffset, LinkedNode);
+					// Map pin exec flow
+					if (LinkedPin)
+					{
+						TSharedPtr<FScriptExecutionNode> ExecPinNode = MapNodeExecutionFlow(InFunctionContext, LinkedPin->GetOwningNode());
+						if (ExecPinNode.IsValid())
+						{
+							LinkedNode->AddChildNode(ExecPinNode);
+						}
+					}
+				}
+				else if (LinkedPin)
+				{
+					TSharedPtr<FScriptExecutionNode> ExecLinkedNode = MapNodeExecutionFlow(InFunctionContext, LinkedPin->GetOwningNode());
+					if (ExecLinkedNode.IsValid())
+					{
+						const int32 PinScriptCodeOffset = BlueprintClass->GetDebugData().FindPinToCodeAssociation(Pin, BlueprintFunction);
+						check(PinScriptCodeOffset != INDEX_NONE);
+						MappedNode->AddLinkedNode(PinScriptCodeOffset, ExecLinkedNode);
+					}
+				}
+			}
+			// Evaluate Children for call sites
+			if (UK2Node_CallFunction* FunctionCallSite = Cast<UK2Node_CallFunction>(NodeToMap))
+			{
+				const UEdGraphNode* EventNode = nullptr; 
+				if (UEdGraph* CalledGraph = FunctionCallSite->GetFunctionGraph(EventNode))
+				{
+					// Update Exec node
+					NodeContext->ExecutionNode->AddFlags(EScriptExecutionNodeFlags::FunctionCall);
+					IconColor = FLinearColor(0.46f, 0.54f, 0.95f);
+					ToolTipText = LOCTEXT("NavigateToFunctionLocationHyperlink_ToolTip", "Navigate to the Function Callsite");
+					// Create temporary function context
+					TSharedPtr<FScriptExecutionContext> NewFunctionContext = MakeShareable<FScriptExecutionContext>(new FScriptExecutionContext);
+					*NewFunctionContext = *InFunctionContext;
+					NewFunctionContext->BlueprintFunction = FunctionCallSite->FunctionReference.ResolveMember<UFunction>(BlueprintClass);
+					NewFunctionContext->ExecutionNode = NodeContext->ExecutionNode;
+					// Grab function entries
+					TArray<UK2Node_FunctionEntry*> FunctionEntryNodes;
+					CalledGraph->GetNodesOfClass<UK2Node_FunctionEntry>(FunctionEntryNodes);
+	
+					for (auto EntryNodeIter : FunctionEntryNodes)
+					{
+						TSharedPtr<FScriptExecutionNode> FunctionEntry = MapNodeExecutionFlow(NewFunctionContext, EntryNodeIter);
+						NewFunctionContext->ExecutionNode->AddChildNode(FunctionEntry);
+					}
+				}
+			}
+			else if (UK2Node_MacroInstance* MacroCallSite = Cast<UK2Node_MacroInstance>(NodeToMap))
+			{
+				if (UEdGraph* MacroGraph = MacroCallSite->GetMacroGraph())
+				{
+					// Update Exec node
+					NodeContext->ExecutionNode->AddFlags(EScriptExecutionNodeFlags::MacroCall);
+					IconColor = FLinearColor(0.36f, 0.34f, 0.71f);
+					ToolTipText = LOCTEXT("NavigateToMacroLocationHyperlink_ToolTip", "Navigate to the Macro Callsite");
+					// Create temporary macro context
+					TSharedPtr<FScriptExecutionContext> NewMacroContext = MakeShareable<FScriptExecutionContext>(new FScriptExecutionContext);
+					*NewMacroContext = *InFunctionContext;
+					UBlueprint* MacroBP = MacroGraph->GetTypedOuter<UBlueprint>();
+					UBlueprintGeneratedClass* MacroBPGC = Cast<UBlueprintGeneratedClass>(MacroBP->GeneratedClass);
+					NewMacroContext->ContextObject = MacroBP;
+					NewMacroContext->BlueprintClass = MacroBPGC;
+					NewMacroContext->BlueprintFunction = MacroBPGC->UberGraphFunction;
+					NewMacroContext->ExecutionNode = NodeContext->ExecutionNode;
+					// Grab function entries
+					TArray<UK2Node_FunctionEntry*> FunctionEntryNodes;
+					MacroGraph->GetNodesOfClass<UK2Node_FunctionEntry>(FunctionEntryNodes);
+	
+					for (auto EntryNodeIter : FunctionEntryNodes)
+					{
+						TSharedPtr<FScriptExecutionNode> MacroEntry = MapNodeExecutionFlow(NewMacroContext, EntryNodeIter);
+						NewMacroContext->ExecutionNode->AddChildNode(MacroEntry);
+					}
+				}
+			}
+			// Finalise exec node data.
+			MappedNode->SetDisplayName(NodeToMap->GetNodeTitle(ENodeTitleType::ListView));
+			MappedNode->SetToolTipText(ToolTipText);
+			const FSlateBrush* Icon = NodeToMap->ShowPaletteIconOnNode() ?	FEditorStyle::GetBrush(NodeToMap->GetPaletteIcon(IconColor)) :
+																			FEditorStyle::GetBrush(TEXT("BlueprintProfiler.BPNode"));
+			MappedNode->SetIcon(const_cast<FSlateBrush*>(Icon));
+			MappedNode->SetIconColor(IconColor);
+		}
 	}
-	return CurrentFunctionPath;
+	return MappedNode;
+}
+
+void FBlueprintProfiler::MapBlueprintInstance(const FString& InstancePath)
+{
+	FString RemappedInstancePath = GetValidInstancePath(InstancePath);
+	TSharedPtr<FScriptExecutionContext> Result = GetCachedObjectContext(RemappedInstancePath);
+
+	if (!Result->ExecutionNode.IsValid())
+	{
+		// Create new instance node
+		Result->ExecutionNode = GetOrCreateProfilerDataForNode(FName(*RemappedInstancePath));
+		Result->ExecutionNode->SetNameByString(RemappedInstancePath);
+		Result->ExecutionNode->SetGraphName(NAME_None);
+		Result->ExecutionNode->SetFlags(EScriptExecutionNodeFlags::Instance);
+		const FSlateBrush* Icon = FEditorStyle::GetBrush(TEXT("BlueprintProfiler.Actor"));
+		Result->ExecutionNode->SetIcon(const_cast<FSlateBrush*>(Icon));
+		Result->ExecutionNode->SetIconColor(FLinearColor(1.f, 1.f, 1.f, 0.8f));
+		Result->ExecutionNode->SetToolTipText(LOCTEXT("NavigateToInstanceHyperlink_ToolTip", "Navigate to the Instance"));
+		// Update widget data
+		if (const AActor* Actor = Cast<AActor>(Result->ContextObject.Get()))
+		{
+			Result->ExecutionNode->SetDisplayName(FText::FromString(Actor->GetActorLabel()));
+		}
+		else
+		{
+			Result->ExecutionNode->SetDisplayName(FText::FromString(Result->ContextObject.Get()->GetName()));
+		}
+		// Link to parent blueprint entry
+		if (UBlueprint* Blueprint = Cast<UBlueprint>(Result->ContextObject->GetClass()->ClassGeneratedBy))
+		{
+			Result->BlueprintClass = Cast<UBlueprintGeneratedClass>(Blueprint->GeneratedClass);
+			const FString ClassPath = Result->BlueprintClass->GetPathName();
+			TSharedPtr<FScriptExecutionNode> BlueprintNode = GetProfilerDataForNode(FName(*ClassPath));
+
+			if (BlueprintNode.IsValid() && BlueprintNode->IsRootNode())
+			{
+				TSharedPtr<FScriptExecutionBlueprint> BlueprintRootNode = StaticCastSharedPtr<FScriptExecutionBlueprint>(BlueprintNode);
+				BlueprintRootNode->AddInstance(Result->ExecutionNode);
+				Result->ExecutionNode->SetParentNode(BlueprintNode);
+				// Fill out events from the blueprint root node
+				for (int32 NodeIdx = 0; NodeIdx < BlueprintRootNode->GetNumChildren(); ++NodeIdx)
+				{
+					Result->ExecutionNode->AddChildNode(BlueprintRootNode->GetChildByIndex(NodeIdx));
+				}
+			}
+		}
+	}
+}
+
+const FString FBlueprintProfiler::GetValidInstancePath(const FString PathIn)
+{
+	FString Result(PathIn);
+	if (const FString* SearchResult = PIEInstancePathMap.Find(PathIn))
+	{
+		Result = *SearchResult;
+	}
+	else
+	{
+		// Attempt to locate the instance and map PIE objects to editor world objects
+		if (const UObject* ObjectPtr = FindObject<UObject>(nullptr, *PathIn))
+		{
+			// Get Outer world
+			if (UWorld* ObjectWorld = ObjectPtr->GetTypedOuter<UWorld>())
+			{
+				if (ObjectWorld->WorldType == EWorldType::PIE)
+				{
+					FWorldContext& EditorWorldContext = GEditor->GetEditorWorldContext();
+					if (UWorld* EditorWorld = EditorWorldContext.World())
+					{
+						for (auto LevelIter : EditorWorld->GetLevels())
+						{
+							if (UObject* EditorObject = FindObject<UObject>(LevelIter, *ObjectPtr->GetName()))
+							{
+								Result = EditorObject->GetPathName();
+								PIEInstancePathMap.Add(PathIn) = Result;
+								TSharedPtr<FScriptExecutionContext> ObjContext = GetCachedObjectContext(Result);
+								ObjContext->ContextObject = EditorObject;
+								break;
+							}
+						}
+					}
+				}
+				else
+				{
+					PIEInstancePathMap.Add(Result) = Result;
+					TSharedPtr<FScriptExecutionContext> ObjContext = GetCachedObjectContext(PathIn);
+					ObjContext->ContextObject = ObjectPtr;
+				}
+			}
+		}
+
+	}
+	return Result;
 }
+
+void FBlueprintProfiler::NavigateToBlueprint(const FName NameIn) const
+{
+	const TSharedPtr<FScriptExecutionContext>* Result = PathToObjectContext.Find(NameIn.ToString());
+	if (Result && Result->IsValid())
+	{
+		if ((*Result)->ContextObject.IsValid())
+		{
+			// Stubbed - complete...
+		}
+	}
+}
+
+void FBlueprintProfiler::NavigateToInstance(const FName NameIn) const
+{
+	const TSharedPtr<FScriptExecutionContext>* Result = PathToObjectContext.Find(NameIn.ToString());
+	if (Result && Result->IsValid())
+	{
+		if ((*Result)->ContextObject.IsValid())
+		{
+			// Stubbed - complete...
+		}
+	}
+}
+
+void FBlueprintProfiler::NavigateToNode(const FName NameIn) const
+{
+	const TSharedPtr<FScriptExecutionContext>* Result = PathToObjectContext.Find(NameIn.ToString());
+	if (Result && Result->IsValid())
+	{
+		if ((*Result)->ContextObject.IsValid())
+		{
+			FKismetEditorUtilities::BringKismetToFocusAttentionOnObject((*Result)->ContextObject.Get());
+		}
+	}
+}
+
+#undef LOCTEXT_NAMESPACE
diff --git a/Engine/Source/Editor/BlueprintGraph/Classes/EdGraphSchema_K2.h b/Engine/Source/Editor/BlueprintGraph/Classes/EdGraphSchema_K2.h
--- a/Engine/Source/Editor/BlueprintGraph/Classes/EdGraphSchema_K2.h
+++ b/Engine/Source/Editor/BlueprintGraph/Classes/EdGraphSchema_K2.h
@@ -2,6 +2,7 @@
 
 #pragma once
 #include "EdGraph/EdGraphSchema.h"
+#include "EdGraph/EdGraphPin.h"
 #include "EdGraphSchema_K2.generated.h"
 
 /** Reference to an structure (only used in 'docked' palette) */
diff --git a/Engine/Source/Editor/BlueprintGraph/Classes/K2Node.h b/Engine/Source/Editor/BlueprintGraph/Classes/K2Node.h
--- a/Engine/Source/Editor/BlueprintGraph/Classes/K2Node.h
+++ b/Engine/Source/Editor/BlueprintGraph/Classes/K2Node.h
@@ -6,6 +6,7 @@
 #include "EngineLogs.h"
 #include "K2Node.generated.h"
 
+class UActorComponent;
 class UBlueprintNodeSpawner;
 class FBlueprintActionDatabaseRegistrar;
 class UDynamicBlueprintBinding;
@@ -109,7 +110,7 @@ class UK2Node : public UEdGraphNode
 	BLUEPRINTGRAPH_API virtual FLinearColor GetNodeTitleColor() const override;
 	BLUEPRINTGRAPH_API virtual void AutowireNewNode(UEdGraphPin* FromPin) override;
 	BLUEPRINTGRAPH_API void PinConnectionListChanged(UEdGraphPin* Pin) override;
-	virtual UObject* GetJumpTargetForDoubleClick() const override { return GetReferencedLevelActor(); }
+    BLUEPRINTGRAPH_API virtual UObject* GetJumpTargetForDoubleClick() const override;
 	BLUEPRINTGRAPH_API virtual FString GetDocumentationLink() const override;
 	BLUEPRINTGRAPH_API virtual void GetPinHoverText(const UEdGraphPin& Pin, FString& HoverTextOut) const override;
 	BLUEPRINTGRAPH_API virtual bool ShowPaletteIconOnNode() const override { return true; }
@@ -120,6 +121,9 @@ class UK2Node : public UEdGraphNode
 
 	// K2Node interface
 
+	UPROPERTY()
+	TMap<UEdGraphPin*, UEdGraphPin*> ByRefMatchupPins;
+
 	// Reallocate pins during reconstruction; by default ignores the old pins and calls AllocateDefaultPins()
 	BLUEPRINTGRAPH_API virtual void ReallocatePinsDuringReconstruction(TArray<UEdGraphPin*>& OldPins);
 
diff --git a/Engine/Source/Editor/BlueprintGraph/Classes/K2Node_AddComponent.h b/Engine/Source/Editor/BlueprintGraph/Classes/K2Node_AddComponent.h
--- a/Engine/Source/Editor/BlueprintGraph/Classes/K2Node_AddComponent.h
+++ b/Engine/Source/Editor/BlueprintGraph/Classes/K2Node_AddComponent.h
@@ -57,11 +57,7 @@ class UK2Node_AddComponent : public UK2Node_CallFunction
 	}
 
 	/** Static name of function to call */
-	static FName GetAddComponentFunctionName()
-	{
-		static const FName AddComponentFunctionName(GET_FUNCTION_NAME_CHECKED(AActor, AddComponent));
-		return AddComponentFunctionName;
-	}
+	static BLUEPRINTGRAPH_API FName GetAddComponentFunctionName();
 
 private: 
 	UEdGraphPin* GetTemplateNamePin() const
diff --git a/Engine/Source/Editor/BlueprintGraph/Classes/K2Node_InputKey.h b/Engine/Source/Editor/BlueprintGraph/Classes/K2Node_InputKey.h
--- a/Engine/Source/Editor/BlueprintGraph/Classes/K2Node_InputKey.h
+++ b/Engine/Source/Editor/BlueprintGraph/Classes/K2Node_InputKey.h
@@ -2,8 +2,11 @@
 
 
 #pragma once
-#include "K2Node.h"
+
 #include "EdGraph/EdGraphNodeUtils.h" // for FNodeTextCache
+#include "InputCoreTypes.h"
+#include "K2Node.h"
+
 #include "K2Node_InputKey.generated.h"
 
 UCLASS(MinimalAPI)
diff --git a/Engine/Source/Editor/BlueprintGraph/Private/EdGraphSchema_K2.cpp b/Engine/Source/Editor/BlueprintGraph/Private/EdGraphSchema_K2.cpp
--- a/Engine/Source/Editor/BlueprintGraph/Private/EdGraphSchema_K2.cpp
+++ b/Engine/Source/Editor/BlueprintGraph/Private/EdGraphSchema_K2.cpp
@@ -1672,6 +1672,8 @@ void UEdGraphSchema_K2::OnReplaceVariableForVariableNode( UK2Node_Variable* Vari
 		}
 		Pin->PinName = VariableName;
 		Variable->ReconstructNode();
+
+		FBlueprintEditorUtils::MarkBlueprintAsStructurallyModified(OwnerBlueprint);
 	}
 }
 
@@ -4032,9 +4034,9 @@ void UEdGraphSchema_K2::ReconstructNode(UEdGraphNode& TargetNode, bool bIsBatchR
 				|| TargetNode.Pins.Contains(Pin);
 			if (!bIsValidPin)
 			{
-				UE_LOG(LogBlueprint, Error, 
-					TEXT("Broken Node: %s keeps removed/invalid pin: %s. Try refresh all nodes."),
-					*TargetNode.GetFullName(), *Pin->GetFullName());
+				UE_LOG(LogBlueprint, Warning,
+					TEXT("Invalid pin: '%s' thinks it belongs to a node ('%s' - '%s') which doesn't have record of it. Try refreshing the node, then compile and resave the Blueprint to hopefully aleviate the problem."),
+					*Pin->PinName, *TargetNode.GetNodeTitle(ENodeTitleType::MenuTitle).ToString(), *TargetNode.GetFullName());
 			}
 		}
 	}
@@ -6040,6 +6042,9 @@ void UEdGraphSchema_K2::RecombinePin(UEdGraphPin* Pin) const
 
 	ParentPin->bHidden = false;
 
+	UEdGraph* Graph = CastChecked<UEdGraph>(GraphNode->GetOuter());
+	UBlueprint* Blueprint = FBlueprintEditorUtils::FindBlueprintForGraphChecked(Graph);
+
 	for (int32 SubPinIndex = 0; SubPinIndex < ParentPin->SubPins.Num(); ++SubPinIndex)
 	{
 		UEdGraphPin* SubPin = ParentPin->SubPins[SubPinIndex];
@@ -6050,6 +6055,7 @@ void UEdGraphSchema_K2::RecombinePin(UEdGraphPin* Pin) const
 		}
 
 		GraphNode->Pins.Remove(SubPin);
+		Blueprint->PinWatches.Remove(SubPin);
 	}
 
 	if (Pin->Direction == EGPD_Input)
@@ -6084,13 +6090,9 @@ void UEdGraphSchema_K2::RecombinePin(UEdGraphPin* Pin) const
 		}
 	}
 
+	ParentPin->SubPins.Empty();	
 
-	ParentPin->SubPins.Empty();
-
-	UEdGraph* Graph = CastChecked<UEdGraph>(GraphNode->GetOuter());
 	Graph->NotifyGraphChanged();
-
-	UBlueprint* Blueprint = FBlueprintEditorUtils::FindBlueprintForGraphChecked(Graph);
 	FBlueprintEditorUtils::MarkBlueprintAsModified(Blueprint);
 }
 
diff --git a/Engine/Source/Editor/BlueprintGraph/Private/K2Node.cpp b/Engine/Source/Editor/BlueprintGraph/Private/K2Node.cpp
--- a/Engine/Source/Editor/BlueprintGraph/Private/K2Node.cpp
+++ b/Engine/Source/Editor/BlueprintGraph/Private/K2Node.cpp
@@ -280,6 +280,11 @@ void UK2Node::PinConnectionListChanged(UEdGraphPin* Pin)
 	NotifyPinConnectionListChanged(Pin);
 }
 
+UObject* UK2Node::GetJumpTargetForDoubleClick() const
+{
+    return GetReferencedLevelActor();
+}
+
 void UK2Node::ReallocatePinsDuringReconstruction(TArray<UEdGraphPin*>& OldPins)
 {
 	AllocateDefaultPins();
diff --git a/Engine/Source/Editor/BlueprintGraph/Private/K2Node_AddComponent.cpp b/Engine/Source/Editor/BlueprintGraph/Private/K2Node_AddComponent.cpp
--- a/Engine/Source/Editor/BlueprintGraph/Private/K2Node_AddComponent.cpp
+++ b/Engine/Source/Editor/BlueprintGraph/Private/K2Node_AddComponent.cpp
@@ -89,6 +89,12 @@ void UK2Node_AddComponent::AllocatePinsForExposedVariables()
 	TransformPin->SafeSetHidden(bHideTransformPins);
 }
 
+FName UK2Node_AddComponent::GetAddComponentFunctionName()
+{
+    static const FName AddComponentFunctionName(GET_FUNCTION_NAME_CHECKED(AActor, AddComponent));
+    return AddComponentFunctionName;
+}
+
 const UClass* UK2Node_AddComponent::GetSpawnedType() const
 {
 	const UActorComponent* TemplateComponent = GetTemplateFromNode();
diff --git a/Engine/Source/Editor/BlueprintGraph/Private/K2Node_BreakStruct.cpp b/Engine/Source/Editor/BlueprintGraph/Private/K2Node_BreakStruct.cpp
--- a/Engine/Source/Editor/BlueprintGraph/Private/K2Node_BreakStruct.cpp
+++ b/Engine/Source/Editor/BlueprintGraph/Private/K2Node_BreakStruct.cpp
@@ -6,6 +6,7 @@
 #include "EditorCategoryUtils.h"
 #include "BlueprintActionDatabaseRegistrar.h"
 #include "BlueprintActionFilter.h"	// for FBlueprintActionContext
+#include "BlueprintEditorSettings.h"
 
 #define LOCTEXT_NAMESPACE "K2Node_BreakStruct"
 
@@ -87,7 +88,13 @@ public:
 				CompilerContext.MessageLog.Warning(*Message.ToString(), Net->GetOuter());
 			}
 
+			UBlueprintEditorSettings* Settings = GetMutableDefault<UBlueprintEditorSettings>();
 			FBPTerminal* Term = Context.CreateLocalTerminalFromPinAutoChooseScope(Net, Net->PinName);
+
+			if (Settings->bAllowReferencePassThroughsToCacheNonReferences)
+			{
+				Term->bPassedByReference = ContextTerm->bPassedByReference;
+			}
 			Term->AssociatedVarProperty = BoundProperty;
 			Context.NetMap.Add(Net, Term);
 			Term->Context = ContextTerm;
diff --git a/Engine/Source/Editor/BlueprintGraph/Private/K2Node_Event.cpp b/Engine/Source/Editor/BlueprintGraph/Private/K2Node_Event.cpp
--- a/Engine/Source/Editor/BlueprintGraph/Private/K2Node_Event.cpp
+++ b/Engine/Source/Editor/BlueprintGraph/Private/K2Node_Event.cpp
@@ -5,6 +5,7 @@
 #include "KismetCompiler.h"
 #include "EventEntryHandler.h"
 #include "GraphEditorSettings.h"
+#include "BlueprintsObjectVersion.h"
 
 const FString UK2Node_Event::DelegateOutputName(TEXT("OutputDelegate"));
 
@@ -42,6 +43,7 @@ void UK2Node_Event::Serialize(FArchive& Ar)
 {
 	Super::Serialize(Ar);
 
+	Ar.UsingCustomVersion(FBlueprintsObjectVersion::GUID);
 	// Fix up legacy nodes that may not yet have a delegate pin
 	if(Ar.IsLoading())
 	{
@@ -50,9 +52,7 @@ void UK2Node_Event::Serialize(FArchive& Ar)
 			EventReference.SetExternalMember(EventSignatureName_DEPRECATED, EventSignatureClass_DEPRECATED);
 		}
 
-		// @TODO: Ar.IsTransacting() is no longer needed after the version bump done in Dev-BP
-		//if (Ar.CustomVer(FBlueprintsObjectVersion::GUID) < FBlueprintsObjectVersion::OverridenEventReferenceFixup)
-		if (!Ar.IsTransacting())
+		if (Ar.CustomVer(FBlueprintsObjectVersion::GUID) < FBlueprintsObjectVersion::OverridenEventReferenceFixup)
 		{
 			FixupEventReference();
 		}
diff --git a/Engine/Source/Editor/BlueprintGraph/Private/K2Node_FunctionEntry.cpp b/Engine/Source/Editor/BlueprintGraph/Private/K2Node_FunctionEntry.cpp
--- a/Engine/Source/Editor/BlueprintGraph/Private/K2Node_FunctionEntry.cpp
+++ b/Engine/Source/Editor/BlueprintGraph/Private/K2Node_FunctionEntry.cpp
@@ -4,6 +4,7 @@
 #include "BlueprintGraphPrivatePCH.h"
 #include "KismetCompiler.h"
 #include "K2Node_CallFunction.h"
+#include "BlueprintsObjectVersion.h"
 
 #define LOCTEXT_NAMESPACE "K2Node_FunctionEntry"
 
@@ -119,6 +120,8 @@ void UK2Node_FunctionEntry::Serialize(FArchive& Ar)
 {
 	Super::Serialize(Ar);
 
+	Ar.UsingCustomVersion(FBlueprintsObjectVersion::GUID);
+
 	if (Ar.IsLoading())
 	{
 		if (Ar.UE4Ver() < VER_UE4_BLUEPRINT_ENFORCE_CONST_IN_FUNCTION_OVERRIDES)
@@ -127,8 +130,7 @@ void UK2Node_FunctionEntry::Serialize(FArchive& Ar)
 			bEnforceConstCorrectness = false;
 		}
 
-		// @TODO: Dev-BP=>Main; gate this with a version check once it makes its way into main
-		//if (Ar.UE4Ver() < VER_UE4_CLEAN_BLUEPRINT_FUNC_FLAGS))
+		if (Ar.CustomVer(FBlueprintsObjectVersion::GUID) < FBlueprintsObjectVersion::CleanBlueprintFunctionFlags)
 		{
 			// Flags we explicitly use ExtraFlags for (at the time this fix was made):
 			//     FUNC_Public, FUNC_Protected, FUNC_Private, 
diff --git a/Engine/Source/Editor/BlueprintGraph/Private/K2Node_SetFieldsInStruct.cpp b/Engine/Source/Editor/BlueprintGraph/Private/K2Node_SetFieldsInStruct.cpp
--- a/Engine/Source/Editor/BlueprintGraph/Private/K2Node_SetFieldsInStruct.cpp
+++ b/Engine/Source/Editor/BlueprintGraph/Private/K2Node_SetFieldsInStruct.cpp
@@ -6,6 +6,7 @@
 #include "CompilerResultsLog.h"
 #include "KismetCompiler.h"
 #include "Editor/PropertyEditor/Public/PropertyCustomizationHelpers.h"
+#include "BlueprintEditorSettings.h"
 
 #define LOCTEXT_NAMESPACE "K2Node_MakeStruct"
 
@@ -25,7 +26,15 @@ struct SetFieldsInStructHelper
 class FKCHandler_SetFieldsInStruct : public FKCHandler_MakeStruct
 {
 public:
-	FKCHandler_SetFieldsInStruct(FKismetCompilerContext& InCompilerContext) : FKCHandler_MakeStruct(InCompilerContext) {}
+	FKCHandler_SetFieldsInStruct(FKismetCompilerContext& InCompilerContext)
+		: FKCHandler_MakeStruct(InCompilerContext)
+	{
+		UBlueprintEditorSettings* Settings = GetMutableDefault<UBlueprintEditorSettings>();
+		if (Settings->bAllowReferencePassThroughsToCacheNonReferences)
+		{
+			bAutoGenerateGotoForPure = false;
+		}
+	}
 
 	virtual UEdGraphPin* FindStructPinChecked(UEdGraphNode* InNode) const override
 	{
@@ -34,6 +43,94 @@ public:
 		check(EGPD_Input == FoundPin->Direction);
 		return FoundPin;
 	}
+
+	virtual void RegisterNet(FKismetFunctionContext& Context, UEdGraphPin* Net) override
+	{
+		UBlueprintEditorSettings* Settings = GetMutableDefault<UBlueprintEditorSettings>();
+
+		if (!Settings->bAllowReferencePassThroughsToCacheNonReferences)
+		{
+			FKCHandler_MakeStruct::RegisterNet(Context, Net);
+		}
+		else
+		{
+			if (Net->Direction == EGPD_Output)
+			{
+				if (Net->ReferencePassThroughConnection)
+				{
+					UEdGraphPin* InputPinNet = FEdGraphUtilities::GetNetFromPin(Net->ReferencePassThroughConnection);
+					FBPTerminal** InputPinTerm = Context.NetMap.Find(InputPinNet);
+					if (InputPinTerm && !(*InputPinTerm)->bPassedByReference)
+					{
+						// We need a net for the output pin which we have thus far prevented from being registered
+						FKCHandler_MakeStruct::RegisterNet(Context, Net);
+					}
+				}
+			}
+		}
+	}
+
+	virtual void RegisterNets(FKismetFunctionContext& Context, UEdGraphNode* Node) override
+	{
+		FKCHandler_MakeStruct::RegisterNets(Context, Node);
+
+		UBlueprintEditorSettings* Settings = GetMutableDefault<UBlueprintEditorSettings>();
+		if (Settings->bAllowReferencePassThroughsToCacheNonReferences)
+		{
+			UEdGraphPin* ReturnPin = Node->FindPin(SetFieldsInStructHelper::StructOutPinName());
+			UEdGraphPin* ReturnStructNet = FEdGraphUtilities::GetNetFromPin(ReturnPin);
+			FBPTerminal** ReturnTerm = Context.NetMap.Find(ReturnStructNet);
+
+			UEdGraphPin* InputPin = Node->FindPinChecked(SetFieldsInStructHelper::StructRefPinName());
+			UEdGraphPin* InputPinNet = FEdGraphUtilities::GetNetFromPin(InputPin);
+			FBPTerminal** InputTerm = Context.NetMap.Find(InputPinNet);
+
+			if (InputTerm == nullptr)
+			{
+				CompilerContext.MessageLog.Error(*LOCTEXT("MakeStruct_NoTerm_Error", "Failed to generate a term for the @@ pin; was it a struct reference that was left unset?").ToString(), InputPin);
+			}
+			else
+			{
+				if ((*InputTerm)->bPassedByReference) //InputPinNet->PinType.bIsReference)
+				{
+					// Forward the net to the output pin because it's being passed by-ref and this pin is a by-ref pin
+					Context.NetMap.Add(ReturnStructNet, *InputTerm);
+				}
+			}
+		}
+	}
+
+	virtual void Compile(FKismetFunctionContext& Context, UEdGraphNode* Node) override
+	{
+		FKCHandler_MakeStruct::Compile(Context, Node);
+
+		UBlueprintEditorSettings* Settings = GetMutableDefault<UBlueprintEditorSettings>();
+		if (Settings->bAllowReferencePassThroughsToCacheNonReferences)
+		{
+			{
+				UEdGraphPin* InputPin = Node->FindPinChecked(SetFieldsInStructHelper::StructRefPinName());
+				UEdGraphPin* InputPinNet = FEdGraphUtilities::GetNetFromPin(InputPin);
+				FBPTerminal** InputTerm = Context.NetMap.Find(InputPinNet);
+
+				// If the InputTerm was not a by-ref, then we need to place the modified structure into the local output term with an AssignStatement
+				if (InputTerm && !(*InputTerm)->bPassedByReference)
+				{
+					UEdGraphPin* ReturnPin = Node->FindPin(SetFieldsInStructHelper::StructOutPinName());
+					UEdGraphPin* ReturnStructNet = FEdGraphUtilities::GetNetFromPin(ReturnPin);
+					FBPTerminal** ReturnTerm = Context.NetMap.Find(ReturnStructNet);
+
+					FBlueprintCompiledStatement& AssignStatement = Context.AppendStatementForNode(Node);
+					AssignStatement.Type = KCST_Assignment;
+					// The return term is a reference no matter the way we received it.
+					(*ReturnTerm)->bPassedByReference = true;
+					AssignStatement.LHS = *ReturnTerm;
+					AssignStatement.RHS.Add(*InputTerm);
+				}
+			}
+
+			GenerateSimpleThenGoto(Context, *Node);
+		}
+	}
 };
 
 UK2Node_SetFieldsInStruct::UK2Node_SetFieldsInStruct(const FObjectInitializer& ObjectInitializer)
@@ -50,9 +147,14 @@ void UK2Node_SetFieldsInStruct::AllocateDefaultPins()
 		CreatePin(EGPD_Input, Schema->PC_Exec, TEXT(""), NULL, false, false, Schema->PN_Execute);
 		CreatePin(EGPD_Output, Schema->PC_Exec, TEXT(""), NULL, false, false, Schema->PN_Then);
 
-		CreatePin(EGPD_Input, Schema->PC_Struct, TEXT(""), StructType, false, true, SetFieldsInStructHelper::StructRefPinName());
+		UEdGraphPin* InPin = CreatePin(EGPD_Input, Schema->PC_Struct, TEXT(""), StructType, false, true, SetFieldsInStructHelper::StructRefPinName());
+
+		UEdGraphPin* OutPin = CreatePin(EGPD_Output, Schema->PC_Struct, TEXT(""), StructType, false, true, SetFieldsInStructHelper::StructOutPinName());
+		ByRefMatchupPins.Add(InPin, OutPin);
+
+		// Input pin will forward the ref to the output, if the input value is not a reference connection, a copy is made and modified instead and provided as a reference until the function is called again.
+		InPin->AssignByRefPassThroughConnection(OutPin);
 
-		auto OutPin = CreatePin(EGPD_Output, Schema->PC_Struct, TEXT(""), StructType, false, true, SetFieldsInStructHelper::StructOutPinName());
 		OutPin->PinToolTip = LOCTEXT("SetFieldsInStruct_OutPinTooltip", "Reference to the input struct").ToString();
 		{
 			FStructOnScope StructOnScope(Cast<UScriptStruct>(StructType));
@@ -214,20 +316,24 @@ void UK2Node_SetFieldsInStruct::ExpandNode(class FKismetCompilerContext& Compile
 {
 	Super::ExpandNode(CompilerContext, SourceGraph);
 
-	UEdGraphPin* OutPin = FindPin(SetFieldsInStructHelper::StructOutPinName());
-	if (OutPin && OutPin->LinkedTo.Num())
+	UBlueprintEditorSettings* Settings = GetMutableDefault<UBlueprintEditorSettings>();
+	if (!Settings->bAllowReferencePassThroughsToCacheNonReferences)
 	{
-		UEdGraphPin* InPin = FindPin(SetFieldsInStructHelper::StructRefPinName());
-		UEdGraphPin* SourcePin = (InPin && (1 == InPin->LinkedTo.Num())) ? InPin->LinkedTo[0] : nullptr;
-		auto Schema = CompilerContext.GetSchema();
-		const bool bCopied = SourcePin && Schema->MovePinLinks(*OutPin, *SourcePin, true).CanSafeConnect();
-		
-		if (!bCopied)
+		UEdGraphPin* OutPin = FindPin(SetFieldsInStructHelper::StructOutPinName());
+		if (OutPin && OutPin->LinkedTo.Num())
 		{
-			CompilerContext.MessageLog.Error(*LOCTEXT("ExpansionError", "Cannot copy links from @@").ToString(), OutPin);
+			UEdGraphPin* InPin = FindPin(SetFieldsInStructHelper::StructRefPinName());
+			UEdGraphPin* SourcePin = (InPin && (1 == InPin->LinkedTo.Num())) ? InPin->LinkedTo[0] : nullptr;
+			auto Schema = CompilerContext.GetSchema();
+			const bool bCopied = SourcePin && Schema->MovePinLinks(*OutPin, *SourcePin, true).CanSafeConnect();
+
+			if (!bCopied)
+			{
+				CompilerContext.MessageLog.Error(*LOCTEXT("ExpansionError", "Cannot copy links from @@").ToString(), OutPin);
+			}
 		}
+		Pins.Remove(OutPin);
 	}
-	Pins.Remove(OutPin);
 }
 
 void UK2Node_SetFieldsInStruct::Serialize(FArchive& Ar)
diff --git a/Engine/Source/Editor/BlueprintGraph/Private/MakeStructHandler.cpp b/Engine/Source/Editor/BlueprintGraph/Private/MakeStructHandler.cpp
--- a/Engine/Source/Editor/BlueprintGraph/Private/MakeStructHandler.cpp
+++ b/Engine/Source/Editor/BlueprintGraph/Private/MakeStructHandler.cpp
@@ -30,6 +30,7 @@ UEdGraphPin* FKCHandler_MakeStruct::FindStructPinChecked(UEdGraphNode* Node) con
 
 FKCHandler_MakeStruct::FKCHandler_MakeStruct(FKismetCompilerContext& InCompilerContext)
 	: FNodeHandlingFunctor(InCompilerContext)
+	, bAutoGenerateGotoForPure(true)
 {
 }
 
@@ -164,7 +165,7 @@ void FKCHandler_MakeStruct::Compile(FKismetFunctionContext& Context, UEdGraphNod
 		}
 	}
 
-	if (!Node->IsNodePure())
+	if (bAutoGenerateGotoForPure && !Node->IsNodePure())
 	{
 		GenerateSimpleThenGoto(Context, *Node);
 	}
diff --git a/Engine/Source/Editor/BlueprintGraph/Private/MakeStructHandler.h b/Engine/Source/Editor/BlueprintGraph/Private/MakeStructHandler.h
--- a/Engine/Source/Editor/BlueprintGraph/Private/MakeStructHandler.h
+++ b/Engine/Source/Editor/BlueprintGraph/Private/MakeStructHandler.h
@@ -18,4 +18,8 @@ public:
 	virtual void RegisterNet(FKismetFunctionContext& Context, UEdGraphPin* Net) override;
 
 	virtual void Compile(FKismetFunctionContext& Context, UEdGraphNode* InNode) override;
+
+protected:
+	/** Prevents this handler from auto-generating a GoTo to jump to the end of the node's compiled code so that child handlers can override the functionality */
+	bool bAutoGenerateGotoForPure;
 };
\ No newline at end of file
diff --git a/Engine/Source/Editor/BlueprintGraph/Public/BlueprintEditorSettings.h b/Engine/Source/Editor/BlueprintGraph/Public/BlueprintEditorSettings.h
--- a/Engine/Source/Editor/BlueprintGraph/Public/BlueprintEditorSettings.h
+++ b/Engine/Source/Editor/BlueprintGraph/Public/BlueprintEditorSettings.h
@@ -88,6 +88,10 @@ public:
 	UPROPERTY(EditAnywhere, config, Category=Compiler)
 	bool bJumpToNodeErrors;
 
+	/** If enabled, "Set Members..." nodes will cache any non-reference passed in the input on the output by-ref pin. This allows it to modify literals and other by-value items passed into it. */
+	UPROPERTY(EditAnywhere, config, Category = Experimental, AdvancedDisplay)
+	bool bAllowReferencePassThroughsToCacheNonReferences;
+
 	/** If enabled, nodes can be explicitly disabled via context menu when right-clicking on impure nodes in the Blueprint editor. Disabled nodes will not be compiled, but also will not break existing connections. */
 	UPROPERTY(EditAnywhere, config, Category=Experimental, AdvancedDisplay)
 	bool bAllowExplicitImpureNodeDisabling;
diff --git a/Engine/Source/Editor/EditorStyle/Private/SlateEditorStyle.cpp b/Engine/Source/Editor/EditorStyle/Private/SlateEditorStyle.cpp
--- a/Engine/Source/Editor/EditorStyle/Private/SlateEditorStyle.cpp
+++ b/Engine/Source/Editor/EditorStyle/Private/SlateEditorStyle.cpp
@@ -4672,6 +4672,7 @@ void FSlateEditorStyle::FStyle::SetupLevelEditorStyle()
 #if WITH_EDITOR
 	{
 		// Profiler stat window
+		Set( "BlueprintProfiler.ViewToolBar", new BOX_BRUSH("/Docking/Tab_Foreground", 4/16.0f));
 		Set( "BlueprintProfiler.BPIcon_Normal", new IMAGE_BRUSH("Icons/AssetIcons/Blueprint_16x", Icon16x16, FLinearColor( 0.46f, 0.54f, 0.81f )));
 		Set( "BlueprintProfiler.BPIcon_Macro", new IMAGE_BRUSH("Icons/AssetIcons/BlueprintMacroLibrary_16x", Icon16x16, FLinearColor( 0.46f, 0.54f, 0.81f )));
 		Set( "BlueprintProfiler.BPIcon_Interface", new IMAGE_BRUSH("Icons/AssetIcons/BlueprintInterface_16x", Icon16x16, FLinearColor( 0.46f, 0.54f, 0.81f )));
diff --git a/Engine/Source/Editor/Kismet/Private/BPVariableDragDropAction.cpp b/Engine/Source/Editor/Kismet/Private/BPVariableDragDropAction.cpp
--- a/Engine/Source/Editor/Kismet/Private/BPVariableDragDropAction.cpp
+++ b/Engine/Source/Editor/Kismet/Private/BPVariableDragDropAction.cpp
@@ -404,6 +404,8 @@ FReply FKismetVariableDragDropAction::DroppedOnNode(FVector2D ScreenPosition, FV
 				}
 			}
 
+			FBlueprintEditorUtils::MarkBlueprintAsStructurallyModified(DropOnBlueprint);
+
 			return FReply::Handled();
 		}
 	}
diff --git a/Engine/Source/Editor/Kismet/Private/BlueprintDetailsCustomization.cpp b/Engine/Source/Editor/Kismet/Private/BlueprintDetailsCustomization.cpp
--- a/Engine/Source/Editor/Kismet/Private/BlueprintDetailsCustomization.cpp
+++ b/Engine/Source/Editor/Kismet/Private/BlueprintDetailsCustomization.cpp
@@ -591,10 +591,28 @@ void FBlueprintVarActionDetails::CustomizeDetails( IDetailLayoutBuilder& DetailL
 	// Add in default value editing for properties that can be edited, local properties cannot be edited
 	if ((Blueprint != NULL) && (Blueprint->GeneratedClass != NULL))
 	{
+		bool bVariableRenamed = false;
 		if (VariableProperty != NULL && IsVariableInBlueprint())
 		{
+			// Determine the current property name on the CDO is stale
+			if (PropertyOwnerBlueprint.IsValid() && VariableProperty)
+			{
+				UBlueprint* PropertyBlueprint = PropertyOwnerBlueprint.Get();
+				const int32 VarIndex = FBlueprintEditorUtils::FindNewVariableIndex(PropertyBlueprint, CachedVariableName);
+				if (VarIndex != INDEX_NONE)
+				{
+					const FGuid VarGuid = PropertyBlueprint->NewVariables[VarIndex].VarGuid;
+					if (UBlueprintGeneratedClass* AuthoritiveBPGC = Cast<UBlueprintGeneratedClass>(PropertyBlueprint->GeneratedClass))
+					{
+						if (const FName* OldName = AuthoritiveBPGC->PropertyGuids.FindKey(VarGuid))
+						{
+							bVariableRenamed = CachedVariableName != *OldName;
+						}
+					}
+				}
+			}
 			const UProperty* OriginalProperty = nullptr;
-			
+		
 			if(!IsALocalVariable(VariableProperty))
 			{
 				OriginalProperty = FindField<UProperty>(Blueprint->GeneratedClass, VariableProperty->GetFName());
@@ -604,7 +622,7 @@ void FBlueprintVarActionDetails::CustomizeDetails( IDetailLayoutBuilder& DetailL
 				OriginalProperty = VariableProperty;
 			}
 
-			if (OriginalProperty == NULL)
+			if (OriginalProperty == NULL || bVariableRenamed)
 			{
 				// Prevent editing the default value of a skeleton property
 				VariableProperty = NULL;
diff --git a/Engine/Source/Editor/Kismet/Private/BlueprintEditor.cpp b/Engine/Source/Editor/Kismet/Private/BlueprintEditor.cpp
--- a/Engine/Source/Editor/Kismet/Private/BlueprintEditor.cpp
+++ b/Engine/Source/Editor/Kismet/Private/BlueprintEditor.cpp
@@ -101,6 +101,9 @@
 
 #include "IMenu.h"
 
+// Blueprint Profiler
+#include "Developer/BlueprintProfiler/Public/BlueprintProfilerModule.h"
+
 #define LOCTEXT_NAMESPACE "BlueprintEditor"
 
 /////////////////////////////////////////////////////
@@ -2205,7 +2208,8 @@ void FBlueprintEditor::CreateDefaultTabContents(const TArray<UBlueprint*>& InBlu
 
 		this->BlueprintProfiler =
 			SNew(SBlueprintProfilerView)
-			.ProfileViewType(PVT_LiveView);
+			.AssetEditor(SharedThis(this))
+			.ProfileViewType(EBlueprintPerfViewType::ExecutionGraph);
 
 		this->Palette = 
 			SNew(SBlueprintPalette, SharedThis(this))
diff --git a/Engine/Source/Editor/Kismet/Private/Profiler/BPProfilerStatisticWidgets.cpp b/Engine/Source/Editor/Kismet/Private/Profiler/BPProfilerStatisticWidgets.cpp
--- a/Engine/Source/Editor/Kismet/Private/Profiler/BPProfilerStatisticWidgets.cpp
+++ b/Engine/Source/Editor/Kismet/Private/Profiler/BPProfilerStatisticWidgets.cpp
@@ -4,6 +4,8 @@
 #include "SBlueprintProfilerView.h"
 #include "BPProfilerStatisticWidgets.h"
 #include "SHyperlink.h"
+#include "EventExecution.h"
+#include "Developer/BlueprintProfiler/Public/BlueprintProfilerModule.h"
 
 #define LOCTEXT_NAMESPACE "BlueprintProfilerViewTypesUI"
 
@@ -12,27 +14,29 @@
 
 namespace BlueprintProfilerStatText
 {
-	const FName ColumnId_Name( "Name" );
-	const FName ColumnId_InclusiveTime( "InclusiveTime" );
-	const FName ColumnId_ExclusiveTime( "ExclusiveTime" );
-	const FName ColumnId_MaxTime( "MaxTime" );
-	const FName ColumnId_MinTime( "MinTime" );
-	const FName ColumnId_Samples( "Samples" );
-	const FName ColumnId_TotalTime( "TotalTime" );
-
-	const FText ColumnText_Name( NSLOCTEXT("BlueprintProfilerViewUI", "Name", "Name") );
-	const FText ColumnText_InclusiveTime( NSLOCTEXT("BlueprintProfilerViewUI", "InclusiveTime", "Inclusive Time (ms)") );
-	const FText ColumnText_ExclusiveTime( NSLOCTEXT("BlueprintProfilerViewUI", "ExclusiveTime", "Exclusive Time (ms)") );
-	const FText ColumnText_MaxTime( NSLOCTEXT("BlueprintProfilerViewUI", "MaxTime", "Max Time (ms)") );
-	const FText ColumnText_MinTime( NSLOCTEXT("BlueprintProfilerViewUI", "MinTime", "Min Time (ms)") );
-	const FText ColumnText_Samples( NSLOCTEXT("BlueprintProfilerViewUI", "Samples", "Samples") );
-	const FText ColumnText_TotalTime( NSLOCTEXT("BlueprintProfilerViewUI", "TotalTime", "Total Time (s)") );
+	const FName ColumnId_Name("Name");
+	const FName ColumnId_InclusiveTime("InclusiveTime");
+	const FName ColumnId_Time("Time");
+	const FName ColumnId_PureTime("PureTime");
+	const FName ColumnId_MaxTime("MaxTime");
+	const FName ColumnId_MinTime("MinTime");
+	const FName ColumnId_Samples("Samples");
+	const FName ColumnId_TotalTime("TotalTime");
+
+	const FText ColumnText_Name(LOCTEXT("Name", "Name") );
+	const FText ColumnText_InclusiveTime(LOCTEXT("InclusiveTime", "Inclusive Time (ms)"));
+	const FText ColumnText_Time(LOCTEXT("Time", "Time (ms)"));
+	const FText ColumnText_PureTime(LOCTEXT("PureTime", "Pure Time (ms)"));
+	const FText ColumnText_MaxTime(LOCTEXT("MaxTime", "Max Time (ms)"));
+	const FText ColumnText_MinTime(LOCTEXT("MinTime", "Min Time (ms)"));
+	const FText ColumnText_Samples(LOCTEXT("Samples", "Samples"));
+	const FText ColumnText_TotalTime(LOCTEXT("TotalTime", "Total Time (s)"));
 };
 
 //////////////////////////////////////////////////////////////////////////
 // SProfilerStatRow
 
-TSharedRef<SWidget> SProfilerStatRow::GenerateWidgetForColumn( const FName& ColumnName )
+TSharedRef<SWidget> SProfilerStatRow::GenerateWidgetForColumn(const FName& ColumnName)
 {
 	if (ColumnName == BlueprintProfilerStatText::ColumnId_Name)
 	{
@@ -49,21 +53,20 @@ TSharedRef<SWidget> SProfilerStatRow::GenerateWidgetForColumn( const FName& Colu
 			.HAlign(HAlign_Left)
 			.VAlign(VAlign_Center)
 			[
-				ItemToEdit->GenerateNameWidget()
+				ItemToEdit->GenerateColumnWidget(ColumnName)
 			];
 	}
 	else
 	{
-		return SNew(STextBlock)
-			.Text(ItemToEdit->GetTextAttribute(ColumnName));
+		return ItemToEdit->GenerateColumnWidget(ColumnName);
 	}
 }
 
-void SProfilerStatRow::Construct(const FArguments& InArgs, TSharedRef<STableViewBase> OwnerTableView, FBPProfilerStatPtr InItemToEdit)
+void SProfilerStatRow::Construct(const FArguments& InArgs, TSharedRef<STableViewBase> OwnerTableView, FBPStatWidgetPtr InItemToEdit)
 {
 	check(InItemToEdit.IsValid());
 	ItemToEdit = InItemToEdit;
-	SMultiColumnTableRow<FBPProfilerStatPtr>::Construct( FSuperRowType::FArguments(), OwnerTableView );	
+	SMultiColumnTableRow<FBPStatWidgetPtr>::Construct( FSuperRowType::FArguments(), OwnerTableView );	
 }
 
 const FName SProfilerStatRow::GetStatName(const EBlueprintProfilerStat::Type StatId)
@@ -73,7 +76,8 @@ const FName SProfilerStatRow::GetStatName(const EBlueprintProfilerStat::Type Sta
 		case EBlueprintProfilerStat::Name:			return BlueprintProfilerStatText::ColumnId_Name;
 		case EBlueprintProfilerStat::TotalTime:		return BlueprintProfilerStatText::ColumnId_TotalTime;
 		case EBlueprintProfilerStat::InclusiveTime:	return BlueprintProfilerStatText::ColumnId_InclusiveTime;
-		case EBlueprintProfilerStat::ExclusiveTime:	return BlueprintProfilerStatText::ColumnId_ExclusiveTime;
+		case EBlueprintProfilerStat::Time:			return BlueprintProfilerStatText::ColumnId_Time;
+		case EBlueprintProfilerStat::PureTime:		return BlueprintProfilerStatText::ColumnId_PureTime;
 		case EBlueprintProfilerStat::MaxTime:		return BlueprintProfilerStatText::ColumnId_MaxTime;
 		case EBlueprintProfilerStat::MinTime:		return BlueprintProfilerStatText::ColumnId_MinTime;
 		case EBlueprintProfilerStat::Samples:		return BlueprintProfilerStatText::ColumnId_Samples;
@@ -88,7 +92,8 @@ const FText SProfilerStatRow::GetStatText(const EBlueprintProfilerStat::Type Sta
 		case EBlueprintProfilerStat::Name:			return BlueprintProfilerStatText::ColumnText_Name;
 		case EBlueprintProfilerStat::TotalTime:		return BlueprintProfilerStatText::ColumnText_TotalTime;
 		case EBlueprintProfilerStat::InclusiveTime:	return BlueprintProfilerStatText::ColumnText_InclusiveTime;
-		case EBlueprintProfilerStat::ExclusiveTime:	return BlueprintProfilerStatText::ColumnText_ExclusiveTime;
+		case EBlueprintProfilerStat::Time:			return BlueprintProfilerStatText::ColumnText_Time;
+		case EBlueprintProfilerStat::PureTime:		return BlueprintProfilerStatText::ColumnText_PureTime;
 		case EBlueprintProfilerStat::MaxTime:		return BlueprintProfilerStatText::ColumnText_MaxTime;
 		case EBlueprintProfilerStat::MinTime:		return BlueprintProfilerStatText::ColumnText_MinTime;
 		case EBlueprintProfilerStat::Samples:		return BlueprintProfilerStatText::ColumnText_Samples;
@@ -97,632 +102,265 @@ const FText SProfilerStatRow::GetStatText(const EBlueprintProfilerStat::Type Sta
 }
 
 //////////////////////////////////////////////////////////////////////////
-// FBPPerformanceData
-
-int32 FBPPerformanceData::StatSampleSize = 20;
-
-void FBPPerformanceData::AddEventTiming(const double NodeTiming)
-{
-	const double TimingInMs = NodeTiming * 1000;
-	Samples.WriteNewElementInitialized() = TimingInMs;
-	TotalTime += NodeTiming;
-	NumSamples++;
-	MinTime = FMath::Min<double>(MinTime, TimingInMs);
-	MaxTime = FMath::Max<double>(MaxTime, TimingInMs);
-	bDirty = true;
-}
-
-void FBPPerformanceData::Update()
-{
-	if (bDirty)
-	{
-		CachedExclusiveTime = 0.0;
-		for (int32 Idx = 0; Idx < Samples.Num(); ++Idx)
-		{
-			CachedExclusiveTime += Samples(Idx);
-		}
-		CachedExclusiveTime /= GetSafeSampleCount();
-		bDirty = false;
-	}
-}
-
-void FBPPerformanceData::Reset()
-{
-	MaxTime = -MAX_dbl;
-	MinTime = MAX_dbl;
-	CachedExclusiveTime = 0.0;
-	TotalTime = 0.0;
-	NumSamples = 0;
-	bDirty = false;
-}
-
-//////////////////////////////////////////////////////////////////////////
-// FBPProfilerStat
-
-TAttribute<FText> FBPProfilerStat::GetTextAttribute(FName ColumnName) const
-{
-	if (ColumnName == BlueprintProfilerStatText::ColumnId_TotalTime)
-	{
-		return TAttribute<FText>(this, &FBPProfilerStat::GetTotalTimeText);
-	}
-	else if (ColumnName == BlueprintProfilerStatText::ColumnId_InclusiveTime)
-	{
-		return TAttribute<FText>(this, &FBPProfilerStat::GetInclusiveTimeText);
-	}
-	else if (ColumnName == BlueprintProfilerStatText::ColumnId_ExclusiveTime)
-	{
-		return TAttribute<FText>(this, &FBPProfilerStat::GetExclusiveTimeText);
-	}
-	else if (ColumnName == BlueprintProfilerStatText::ColumnId_MaxTime)
-	{
-		return TAttribute<FText>(this, &FBPProfilerStat::GetMaxTimeText);
-	}
-	else if (ColumnName == BlueprintProfilerStatText::ColumnId_MinTime)
-	{
-		return TAttribute<FText>(this, &FBPProfilerStat::GetMinTimeText);
-	}
-	else if (ColumnName == BlueprintProfilerStatText::ColumnId_Samples)
-	{
-		return TAttribute<FText>(this, &FBPProfilerStat::GetSamplesText);
-	}
-	return TAttribute<FText>();
-}
-
-FText FBPProfilerStat::GetTotalTimeText() const
-{
-	// Covert back up to seconds for total time.
-	return FText::AsNumber(GetTotalTime(), &SBlueprintProfilerView::GetNumberFormat());
-}
-
-FText FBPProfilerStat::GetInclusiveTimeText() const
-{
-	return FText::AsNumber(GetInclusiveTime(), &SBlueprintProfilerView::GetNumberFormat());
-}
-
-FText FBPProfilerStat::GetExclusiveTimeText() const
-{
-	return FText::AsNumber(GetExclusiveTime(), &SBlueprintProfilerView::GetNumberFormat());
-}
-
-FText FBPProfilerStat::GetMaxTimeText() const
-{
-	return FText::AsNumber(GetMaxTime(), &SBlueprintProfilerView::GetNumberFormat());
-}
-
-FText FBPProfilerStat::GetMinTimeText() const
-{
-	return FText::AsNumber(GetMinTime(), &SBlueprintProfilerView::GetNumberFormat());
-}
-
-FText FBPProfilerStat::GetSamplesText() const
-{
-	return FText::AsNumber(GetSampleCount());
-}
-
-double FBPProfilerStat::GetMaxTime() const
-{
-	double Result = 0.0;
-	if (CachedChildren.Num())
-	{
-		Result += MaxInclusiveTime;
-	}
-	else if (PerformanceStats.IsDataValid())
-	{
-		Result += PerformanceStats.GetMaxTime();
-	}
-	return Result;
-}
-
-double FBPProfilerStat::GetMinTime() const
-{
-	double Result = 0.0;
-	if (CachedChildren.Num())
-	{
-		Result += MinInclusiveTime;
-	}
-	else if (PerformanceStats.IsDataValid())
-	{
-		Result += PerformanceStats.GetMinTime();
-	}
-	return Result;
-}
-
-float FBPProfilerStat::GetTotalTime() const
-{
-	float Result = 0.f;
-	if (CachedChildren.Num())
-	{
-		Result += InclusiveTotalTime;
-	}
-	else if (PerformanceStats.IsDataValid())
-	{
-		Result += PerformanceStats.GetTotalTime();
-	}
-	return Result;
-}
-
-float FBPProfilerStat::GetInclusiveTotalTime() const
-{
-	float Result = 0.f;
-	if (CachedChildren.Num())
-	{
-		Result += InclusiveTotalTime;
-	}
-	else if (PerformanceStats.IsDataValid())
-	{
-		Result += PerformanceStats.GetTotalTime();
-	}
-	return Result;
-}
-
-float FBPProfilerStat::GetExclusiveTime() const
-{
-	return PerformanceStats.IsDataValid() ? PerformanceStats.GetExclusiveTime() : 0.f;
-}
-
-float FBPProfilerStat::GetInclusiveTime() const
-{
-	float Result = 0.f;
-	Result += InclusiveTime;
-	return Result;
-}
-
-float FBPProfilerStat::GetSampleCount() const
-{
-	float Result = 0.f;
-	if (PerformanceStats.IsDataValid())
-	{
-		Result += PerformanceStats.GetSampleCount();
-	}
-	else
-	{
-		Result += InclusiveSamples;
-	}
-	return Result;
-}
-
-float FBPProfilerStat::GetSafeSampleCount() const
-{
-	float Result = 0.f;
-	if (PerformanceStats.IsDataValid())
-	{
-		Result += PerformanceStats.GetSampleCount();
-	}
-	return Result != 0.f ? Result : 1.f;
-}
-
-FBPProfilerStatPtr FBPProfilerStat::FindEntryPoint(FScriptExecEvent& EventEntry)
-{
-	FBPProfilerStatPtr Result;
-	if (EventEntry == ObjectContext)
-	{
-		Result = AsShared();
-	}
-	else
-	{
-		for (auto Iter : CachedChildren)
-		{
-			Result = Iter.Value->FindEntryPoint(EventEntry);
-			if (Result.IsValid())
-			{
-				break;
-			}
-		}
-	}
-	return Result;
-}
+// FBPProfilerStatWidget
 
-FBPProfilerStatPtr FBPProfilerStat::FindOrAddChildByContext(FScriptExecEvent& Context)
+TSharedRef<SWidget> FBPProfilerStatWidget::GenerateColumnWidget(FName ColumnName)
 {
-	FBPProfilerStatPtr& Result = CachedChildren.FindOrAdd(Context.GetObjectPtr());
-	if (!Result.IsValid())
+	if (ExecNode.IsValid())
 	{
-		if (Context.IsBranch())
+		if (ColumnName == BlueprintProfilerStatText::ColumnId_Name)
 		{
-			TSharedPtr<FBPProfilerBranchStat> Branch = MakeShareable(new FBPProfilerBranchStat(AsShared(), Context));
-			Branch->CreateChildrenForPins();
-			Result = StaticCastSharedPtr<FBPProfilerStat>(Branch);
+			return SNew(SHorizontalBox)
+				+SHorizontalBox::Slot()
+				.AutoWidth()
+				[
+					SNew(SImage)
+					.Image(ExecNode->GetIcon())
+					.ColorAndOpacity(ExecNode->GetIconColor())
+				]
+				+SHorizontalBox::Slot()
+				.AutoWidth()
+				.Padding(FMargin(5,0))
+				[
+					SNew(SHyperlink)
+					.Text(ExecNode->GetDisplayName())
+					.Style(FEditorStyle::Get(), "HoverOnlyHyperlink")
+					.ToolTipText(ExecNode->GetToolTipText())
+					.OnNavigate(this, &FBPProfilerStatWidget::NavigateTo)
+				];
 		}
 		else
 		{
-			Result = MakeShareable(new FBPProfilerStat(AsShared(), Context));
-		}
-	}
-	return Result;
-}
+			TAttribute<FText> TextAttr(LOCTEXT("NonApplicableStat", ""));
+			TAttribute<FSlateColor> ColorAttr;
+			const uint32 NonNodeStats = EScriptExecutionNodeFlags::Class|EScriptExecutionNodeFlags::Instance|EScriptExecutionNodeFlags::Event|
+										EScriptExecutionNodeFlags::BranchNode|EScriptExecutionNodeFlags::CallSite|EScriptExecutionNodeFlags::ExecPin;
 
-void FBPProfilerStat::SubmitData(class FScriptExecEvent& EventData)
-{
-	check (EventData == ObjectContext);
-	if (EventData.CanSubmitData())
-	{
-		PerformanceStats.AddEventTiming(EventData.GetTime());
-	}
-}
-
-TSharedRef<SWidget> FBPProfilerStat::GenerateNameWidget()
-{
-	return SNew(SHorizontalBox)
-		+SHorizontalBox::Slot()
-		.AutoWidth()
-		[
-			SNew(SImage)
-			.Image(GetIcon())
-			.ColorAndOpacity(GetIconColor())
-		]
-		+SHorizontalBox::Slot()
-		.AutoWidth()
-		.Padding(FMargin(5,0))
-		[
-			SNew(SHyperlink)
-			.Text(StatName)
-			.Style(FEditorStyle::Get(), "HoverOnlyHyperlink")
-		//	.ToolTipText(GetToolTipText())
-			.OnNavigate(this, &FBPProfilerStat::OnNavigateToObject)
-		];
-}
-
-void FBPProfilerStat::GatherChildren(TArray<FBPProfilerStatPtr>& OutChildren)
-{
-	for (auto Child : CachedChildren)
-	{
-		OutChildren.Add(Child.Value);
-	}
-}
-
-void FBPProfilerStat::UpdateProfilingStats()
-{
-	if (PerformanceStats.IsDataDirty() || !PerformanceStats.IsDataValid())
-	{
-		// Update local performance stats
-		PerformanceStats.Update();
-		const bool bIsLocalDataValid = PerformanceStats.IsDataValid();
-
-		switch(ObjectContext.GetType())
-		{
-			case EScriptInstrumentation::Class:
-			case EScriptInstrumentation::Instance:
+			if (ColumnName == BlueprintProfilerStatText::ColumnId_TotalTime)
 			{
-				// Reset stats
-				InclusiveTime		= 0.0;
-				InclusiveTotalTime	= 0.0;
-				InclusiveSamples	= 0.0;
-				// Re-Acquire data from children.
-				for (auto Child : CachedChildren)
+				if (!ExecNode->HasFlags(EScriptExecutionNodeFlags::ExecPin))
 				{
-					FBPProfilerStatPtr& ChildPtr = Child.Value;
-					ChildPtr->UpdateProfilingStats();
-					InclusiveTime += ChildPtr->GetInclusiveTime();
-					InclusiveTotalTime += ChildPtr->GetInclusiveTotalTime();
-					InclusiveSamples += !bIsLocalDataValid ? ChildPtr->GetInclusiveSampleCount() : 0;
+					TextAttr = TAttribute<FText>(PerformanceStats.Get(), &FScriptPerfData::GetTotalTimingText);
+					ColorAttr = TAttribute<FSlateColor>(FLinearColor::White);
 				}
-				MaxInclusiveTime = FMath::Max<double>(MaxInclusiveTime, InclusiveTime);
-				MinInclusiveTime = FMath::Min<double>(MinInclusiveTime, InclusiveTime);
-				break;
 			}
-			case EScriptInstrumentation::Event:
+			else if (ColumnName == BlueprintProfilerStatText::ColumnId_InclusiveTime)
 			{
-				// Reset stats
-				InclusiveTime		= 0.0;
-				InclusiveTotalTime	= 0.0;
-				InclusiveSamples	= bIsLocalDataValid ? PerformanceStats.GetSampleCount() : 0.0;
-				// Re-Acquire data from children.
-				for (auto Child : CachedChildren)
+				if (!ExecNode->HasFlags(NonNodeStats))
 				{
-					FBPProfilerStatPtr& ChildPtr = Child.Value;
-					ChildPtr->UpdateProfilingStats();
-					InclusiveTime += ChildPtr->GetInclusiveTime();
-					InclusiveTotalTime += ChildPtr->GetInclusiveTotalTime();
-					InclusiveSamples += !bIsLocalDataValid ? ChildPtr->GetInclusiveSampleCount() : 0;
+					TextAttr = TAttribute<FText>(PerformanceStats.Get(), &FScriptPerfData::GetInclusiveTimingText);
+					ColorAttr = TAttribute<FSlateColor>(PerformanceStats.Get(), &FScriptPerfData::GetInclusiveHeatColor);
 				}
-				MaxInclusiveTime = FMath::Max<double>(MaxInclusiveTime, InclusiveTime);
-				MinInclusiveTime = FMath::Min<double>(MinInclusiveTime, InclusiveTime);
-				InclusiveSamples	= bIsLocalDataValid ? PerformanceStats.GetSampleCount() : 0.0;
-
-				break;
 			}
-			default:
+			else if (ColumnName == BlueprintProfilerStatText::ColumnId_Time)
 			{
-				// Reset stats
-				InclusiveTime		= bIsLocalDataValid ? PerformanceStats.GetExclusiveTime() : 0.0;
-				InclusiveTotalTime	= bIsLocalDataValid ? PerformanceStats.GetTotalTime() : 0.0;
-				InclusiveSamples	= bIsLocalDataValid ? PerformanceStats.GetSampleCount() : 0.0;
-				// Re-Acquire data from children.
-				for (auto Child : CachedChildren)
+				if (!ExecNode->HasFlags(EScriptExecutionNodeFlags::ExecPin))
 				{
-					FBPProfilerStatPtr& ChildPtr = Child.Value;
-					ChildPtr->UpdateProfilingStats();
-					InclusiveTime += ChildPtr->GetInclusiveTime();
-					InclusiveTotalTime += ChildPtr->GetInclusiveTotalTime();
-					InclusiveSamples += !bIsLocalDataValid ? ChildPtr->GetInclusiveSampleCount() : 0;
+					TextAttr = TAttribute<FText>(PerformanceStats.Get(), &FScriptPerfData::GetNodeTimingText);
+					ColorAttr = TAttribute<FSlateColor>(PerformanceStats.Get(), &FScriptPerfData::GetNodeHeatColor);
 				}
-				MaxInclusiveTime = FMath::Max<double>(MaxInclusiveTime, InclusiveTime);
-				MinInclusiveTime = FMath::Min<double>(MinInclusiveTime, InclusiveTime);
-				break;
 			}
-		}
-	}
-}
-
-void FBPProfilerStat::UpdateStatName()
-{
-	switch(ObjectContext.GetType())
-	{
-		case EScriptInstrumentation::Instance:
-		{
-			TWeakObjectPtr<const UObject> ObjectPtr = ObjectContext.GetObjectPtr();
-			if (ObjectPtr.IsValid())
+			else if (ColumnName == BlueprintProfilerStatText::ColumnId_PureTime)
 			{
-				if (const AActor* Actor = Cast<AActor>(ObjectPtr.Get()))
+				if (!ExecNode->HasFlags(NonNodeStats))
 				{
-					StatName = FText::FromString(Actor->GetActorLabel());
+					TextAttr = TAttribute<FText>(PerformanceStats.Get(), &FScriptPerfData::GetPureNodeTimingText);
+					ColorAttr = TAttribute<FSlateColor>(PerformanceStats.Get(), &FScriptPerfData::GetInclusiveHeatColor);
 				}
-				else
+			}
+			else if (ColumnName == BlueprintProfilerStatText::ColumnId_MaxTime)
+			{
+				if (!ExecNode->HasFlags(EScriptExecutionNodeFlags::ExecPin))
 				{
-					StatName = FText::FromString(ObjectPtr.Get()->GetName());
+					TextAttr = TAttribute<FText>(PerformanceStats.Get(), &FScriptPerfData::GetMaxTimingText);
+					ColorAttr = TAttribute<FSlateColor>(PerformanceStats.Get(), &FScriptPerfData::GetMaxTimeHeatColor);
 				}
 			}
-			break;
-		}
-		case EScriptInstrumentation::NodePin:
-		{
-			TWeakObjectPtr<const UEdGraphPin> Pin = ObjectContext.GetTypedObjectPtr<UEdGraphPin>();
-			if (Pin.IsValid())
+			else if (ColumnName == BlueprintProfilerStatText::ColumnId_MinTime)
 			{
-				StatName = Pin.Get()->GetDisplayName();
+				if (!ExecNode->HasFlags(EScriptExecutionNodeFlags::ExecPin))
+				{
+					TextAttr = TAttribute<FText>(PerformanceStats.Get(), &FScriptPerfData::GetMinTimingText);
+					ColorAttr = TAttribute<FSlateColor>(FLinearColor::White);
+				}
+			}
+			else if (ColumnName == BlueprintProfilerStatText::ColumnId_Samples)
+			{
+				TextAttr = TAttribute<FText>(PerformanceStats.Get(), &FScriptPerfData::GetSamplesText);
+				ColorAttr = TAttribute<FSlateColor>(FLinearColor::White);
 			}
-			break;
+			// Create the actual widget
+			return SNew(STextBlock)
+				.Text(TextAttr)
+				.ColorAndOpacity(ColorAttr);
 		}
-		case EScriptInstrumentation::NodeEntry:
-		case EScriptInstrumentation::Event:
-		case EScriptInstrumentation::Function:
-		case EScriptInstrumentation::Branch:
-		case EScriptInstrumentation::Macro:
+	}
+	return SNullWidget::NullWidget;
+}
+
+void FBPProfilerStatWidget::NavigateTo() const
+{
+	if (ExecNode.IsValid())
+	{
+		if (IBlueprintProfilerInterface* Profiler = FModuleManager::GetModulePtr<IBlueprintProfilerInterface>("BlueprintProfiler"))
 		{
-			TWeakObjectPtr<const UEdGraphNode> Node = ObjectContext.GetTypedObjectPtr<UEdGraphNode>();
-			if (Node.IsValid())
+			if (ExecNode->IsClass())
 			{
-				StatName = Node.Get()->GetNodeTitle(ENodeTitleType::ListView);
+				Profiler->NavigateToBlueprint(ExecNode->GetName());
 			}
-			else
+			else if (ExecNode->IsInstance())
 			{
-				TWeakObjectPtr<const UObject> ObjectPtr = ObjectContext.GetObjectPtr();
-				if (ObjectPtr.IsValid())
-				{
-					StatName = FText::FromString(ObjectPtr.Get()->GetName());
-				}
+				Profiler->NavigateToInstance(ExecNode->GetName());
 			}
-			break;
-		}
-		default:
-		{
-			TWeakObjectPtr<const UObject> ObjectPtr = ObjectContext.GetObjectPtr();
-			if (ObjectPtr.IsValid())
+			else
 			{
-				StatName = FText::FromString(ObjectPtr.Get()->GetName());
+				Profiler->NavigateToNode(ExecNode->GetName());
 			}
-			break;
 		}
 	}
 }
 
-const FSlateBrush* FBPProfilerStat::GetIcon()
+void FBPProfilerStatWidget::GenerateExecNodeWidgets(const FName InstanceName, const FName FilterGraph)
 {
-	switch(ObjectContext.GetType())
+	if (ExecNode.IsValid())
 	{
-		case EScriptInstrumentation::Class:		return FEditorStyle::GetBrush(TEXT("BlueprintProfiler.BPIcon_Normal"));
-		case EScriptInstrumentation::Instance:	return FEditorStyle::GetBrush(TEXT("BlueprintProfiler.Actor"));
-		case EScriptInstrumentation::NodePin:	return CachedChildren.Num() > 0 ? FEditorStyle::GetBrush(TEXT("BlueprintProfiler.BPPinConnected")) : FEditorStyle::GetBrush(TEXT("BlueprintProfiler.BPPinDisconnected"));
-		default:
+		// Grab Performance Stats
+		PerformanceStats = ExecNode->GetPerfDataByInstanceAndTracePath(InstanceName, WidgetTracePath);
+		// Handle branch creation
+		CachedChildren.Reset(0);
+		if (ExecNode->IsBranch())
+		{
+			for (auto LinkIter : ExecNode->GetLinkedNodes())
+			{
+				TArray<FScriptNodeExecLinkage::FLinearExecPath> LinearExecNodes;
+				FTracePath LinkPath(WidgetTracePath);
+				LinkPath.AddExitPin(LinkIter.Key);
+				LinkIter.Value->GetLinearExecutionPath(LinearExecNodes, LinkPath);
+				for (auto LinearPathIter : LinearExecNodes)
+				{
+					TSharedPtr<FBPProfilerStatWidget> NewLinkedNode = MakeShareable<FBPProfilerStatWidget>(new FBPProfilerStatWidget(LinearPathIter.LinkedNode, LinearPathIter.TracePath));
+					NewLinkedNode->GenerateExecNodeWidgets(InstanceName, FilterGraph);
+					CachedChildren.Add(NewLinkedNode);
+				}
+			}
+		}
+		else
 		{
-			if (const UEdGraphNode* GraphNode = ObjectContext.GetTypedObjectPtr<UEdGraphNode>().Get())
+			for (auto Iter : ExecNode->GetChildNodes())
 			{
-				if (GraphNode->ShowPaletteIconOnNode())
+				// Filter out events based on graph
+				bool bFiltered = (Iter->IsEvent() && FilterGraph != NAME_None) ? IsEventFiltered(Iter, FilterGraph) : false;
+				if (!bFiltered)
 				{
-					return FEditorStyle::GetBrush(GraphNode->GetPaletteIcon(IconColor));
+					TArray<FScriptNodeExecLinkage::FLinearExecPath> LinearExecNodes;
+					FTracePath ChildTracePath(WidgetTracePath);
+					Iter->GetLinearExecutionPath(LinearExecNodes, ChildTracePath);
+					if (LinearExecNodes.Num() > 1)
+					{
+						for (auto LinearPathIter : LinearExecNodes)
+						{
+							TSharedPtr<FBPProfilerStatWidget> NewLinkedNode = MakeShareable<FBPProfilerStatWidget>(new FBPProfilerStatWidget(LinearPathIter.LinkedNode, LinearPathIter.TracePath));
+							NewLinkedNode->GenerateExecNodeWidgets(InstanceName, FilterGraph);
+							CachedChildren.Add(NewLinkedNode);
+						}
+					}
+					else
+					{
+						TSharedPtr<FBPProfilerStatWidget> NewChildNode = MakeShareable<FBPProfilerStatWidget>(new FBPProfilerStatWidget(Iter, ChildTracePath));
+						NewChildNode->GenerateExecNodeWidgets(InstanceName, FilterGraph);
+						CachedChildren.Add(NewChildNode);
+					}
 				}
 			}
-			return FEditorStyle::GetBrush(TEXT("BlueprintProfiler.BPNode"));
 		}
 	}
 }
 
-FLinearColor FBPProfilerStat::GetIconColor() const
+void FBPProfilerStatWidget::GatherChildren(TArray<TSharedPtr<FBPProfilerStatWidget>>& OutChildren)
 {
-	switch(ObjectContext.GetType())
+	if (CachedChildren.Num())
 	{
-		case EScriptInstrumentation::Class:		return FLinearColor( 0.46f, 0.54f, 0.81f );
-		case EScriptInstrumentation::Event:		return FLinearColor( 0.91f, 0.16f, 0.16f );
-		case EScriptInstrumentation::Function:	return FLinearColor( 0.46f, 0.54f, 0.95f );
-		case EScriptInstrumentation::Macro:		return FLinearColor( 0.36f, 0.34f, 0.71f );
-		case EScriptInstrumentation::Branch:	return FLinearColor( 1.f, 1.f, 1.f, 0.8f );
-		default:								return FLinearColor( 1.f, 1.f, 1.f, 0.8f );
+		OutChildren.Append(CachedChildren.GetData(), CachedChildren.Num());
 	}
 }
 
-const FText FBPProfilerStat::GetToolTip()
-{
-	switch(ObjectContext.GetType())
-	{
-		case EScriptInstrumentation::Class:		return LOCTEXT("NavigateToBlueprintLocationHyperlink_ToolTip", "Navigate to the Blueprint");
-		case EScriptInstrumentation::Instance:	return LOCTEXT("NavigateToInstanceLocationHyperlink_ToolTip", "Navigate to the Object Instance");
-		case EScriptInstrumentation::Event:		return LOCTEXT("NavigateToEventLocationHyperlink_ToolTip", "Navigate to the Event Node");
-		case EScriptInstrumentation::Function:	return LOCTEXT("NavigateToFunctionLocationHyperlink_ToolTip", "Navigate to the Function Callsite");
-		case EScriptInstrumentation::Macro:		return LOCTEXT("NavigateToMacroLocationHyperlink_ToolTip", "Navigate to the Macro Callsite");
-		default:								return LOCTEXT("NavigateToUnknownLocationHyperlink_ToolTip", "Navigate to the Object");
-	}
+bool FBPProfilerStatWidget::GetExpansionState() const 
+{ 
+	return ExecNode.IsValid() ? ExecNode->IsExpanded() : false;
 }
 
-void FBPProfilerStat::OnNavigateToObject()
+void FBPProfilerStatWidget::SetExpansionState(bool bExpansionStateIn)
 {
-	switch(ObjectContext.GetType())
+	if(ExecNode.IsValid())
 	{
-		case EScriptInstrumentation::Class:
-		{
-			break;
-		}
-		case EScriptInstrumentation::Instance:
-		{
-			break;
-		}
-		case EScriptInstrumentation::Event:
-		case EScriptInstrumentation::Function:
-		case EScriptInstrumentation::Branch:
-		case EScriptInstrumentation::Macro:
-		default:
-		{
-			TWeakObjectPtr<const UObject> ObjectPtr = ObjectContext.GetObjectPtr();
-			if (ObjectPtr.IsValid())
-			{
-				FKismetEditorUtilities::BringKismetToFocusAttentionOnObject(ObjectPtr.Get());
-			}
-			break;
-		}
+		ExecNode->SetExpanded(bExpansionStateIn);
 	}
 }
 
-bool FBPProfilerBranchStat::IsBranchSequential()
+void FBPProfilerStatWidget::ExpandWidgetState(TSharedPtr<STreeView<FBPStatWidgetPtr>> TreeView, bool bStateIn)
 {
-	TWeakObjectPtr<const UObject> ObjectPtr = ObjectContext.GetObjectPtr();
-	bool bIsSequential = false;
-	if (const UObject* Object = ObjectPtr.Get())
+	if (TreeView.IsValid())
 	{
-		if (Object->IsA<UK2Node_ExecutionSequence>())
-		{
-			bIsSequential = true;
-		}
+		TreeView->SetItemExpansion(AsShared(), bStateIn);
+	}
+	for (auto Iter : CachedChildren)
+	{
+		Iter->ExpandWidgetState(TreeView, bStateIn);
 	}
-	return bIsSequential;
 }
 
-void FBPProfilerBranchStat::CreateChildrenForPins()
+void FBPProfilerStatWidget::RestoreWidgetExpansionState(TSharedPtr<STreeView<FBPStatWidgetPtr>> TreeView)
 {
-	TWeakObjectPtr<const UObject> ObjectPtr = ObjectContext.GetObjectPtr();
-	bAreExecutionPathsSequential = IsBranchSequential();
-	if (const UEdGraphNode* GraphNode = Cast<const UEdGraphNode>(ObjectPtr.Get()))
+	if (TreeView.IsValid() && ExecNode.IsValid())
 	{
-		for (auto Pin : GraphNode->Pins)
-		{
-			if (Pin->Direction == EGPD_Output && Pin->PinType.PinCategory == UEdGraphSchema_K2::PC_Exec)
-			{
-				FBPProfilerStatPtr& Result = CachedChildren.FindOrAdd(Pin);
-				if (!Result.IsValid())
-				{
-					FScriptExecEvent NewExitPin(Pin, ObjectContext);
-					Result = MakeShareable(new FBPProfilerStat(AsShared(), NewExitPin));
-				}
-			}
-		}
+		TreeView->SetItemExpansion(AsShared(), ExecNode->IsExpanded());
+	}
+	for (auto Iter : CachedChildren)
+	{
+		Iter->RestoreWidgetExpansionState(TreeView);
 	}
 }
 
-FBPProfilerStatPtr FBPProfilerBranchStat::FindOrAddChildByContext(FScriptExecEvent& Context)
+bool FBPProfilerStatWidget::ProbeChildWidgetExpansionStates()
 {
-	FBPProfilerStatPtr Result;
-	check (CachedExecPin.IsValid());
-	if (FBPProfilerStatPtr* PinExecutionPath = CachedChildren.Find(CachedExecPin))
+	bool bIsExpanded = false;
+	if (ExecNode.IsValid() && ExecNode->IsExpanded())
 	{
-		// Defer child creation down to the execution pin's children.
-		// This is an artificial step to increase and manage the parent child relationship and execution depth
-		// so we can navigate the stats correctly.
-		FBPProfilerStatPtr& NewChild = (*PinExecutionPath)->CachedChildren.FindOrAdd(Context.GetObjectPtr());
-		if (!NewChild.IsValid())
+		bIsExpanded = true;
+	}
+	else
+	{
+		for (auto Iter : CachedChildren)
 		{
-			if (Context.IsBranch())
+			if (Iter->ProbeChildWidgetExpansionStates())
 			{
-				TSharedPtr<FBPProfilerBranchStat> Branch = MakeShareable(new FBPProfilerBranchStat(ParentStat, Context));
-				Branch->CreateChildrenForPins();
-				NewChild = StaticCastSharedPtr<FBPProfilerStat>(Branch);
-			}
-			else
-			{
-				NewChild = MakeShareable(new FBPProfilerStat(ParentStat, Context));
+				bIsExpanded = true;
+				break;
 			}
 		}
-		Result = NewChild;
 	}
-	return Result;
+	if (bIsExpanded)
+	{
+		ExecNode->SetExpanded(true);
+	}
+	return bIsExpanded;
 }
 
-void FBPProfilerBranchStat::UpdateProfilingStats()
+bool FBPProfilerStatWidget::IsEventFiltered(TSharedPtr<FScriptExecutionNode> InEventNode, FName GraphFilter) const
 {
-	if (PerformanceStats.IsDataDirty() || !PerformanceStats.IsDataValid())
+	bool bEventFiltered = false;
+	if (InEventNode.IsValid() && InEventNode->IsEvent())
 	{
-		// Update local performance stats
-		PerformanceStats.Update();
-		// Reset stats
-		const bool bIsLocalDataValid = PerformanceStats.IsDataValid();
-		InclusiveTime		= bIsLocalDataValid ? PerformanceStats.GetExclusiveTime() : 0.0;
-		InclusiveTotalTime	= bIsLocalDataValid ? PerformanceStats.GetTotalTime() : 0.0;
-		InclusiveSamples	= bIsLocalDataValid ? PerformanceStats.GetSampleCount() : 0.0;
-		// Find Inclusive data
-		if (bAreExecutionPathsSequential)
+		if (InEventNode->GetGraphName() == UEdGraphSchema_K2::FN_UserConstructionScript)
 		{
-			for (auto Child : CachedChildren)
-			{
-				FBPProfilerStatPtr& ChildPtr = Child.Value;
-				ChildPtr->UpdateProfilingStats();
-				InclusiveTime += ChildPtr->GetInclusiveTime();
-				InclusiveTotalTime += ChildPtr->GetInclusiveTotalTime();
-			}
-			MaxInclusiveTime = FMath::Max<double>(MaxInclusiveTime, InclusiveTime);
-			MinInclusiveTime = FMath::Min<double>(MinInclusiveTime, InclusiveTime);
+			bEventFiltered = GraphFilter != UEdGraphSchema_K2::FN_UserConstructionScript;
 		}
 		else
 		{
-			double WorstInclusiveTime = 0.0;
-			double WorstInclusiveTotalTime = 0.0;
-			for (auto Child : CachedChildren)
-			{
-				FBPProfilerStatPtr& ChildPtr = Child.Value;
-				ChildPtr->UpdateProfilingStats();
-				const double ChildInclusiveTime = ChildPtr->GetInclusiveTime();
-				WorstInclusiveTime = FMath::Max<double>(WorstInclusiveTime, ChildInclusiveTime);
-				WorstInclusiveTotalTime = FMath::Max<double>(WorstInclusiveTotalTime, ChildPtr->GetInclusiveTotalTime());
-				// pick best/worst case path
-				MaxInclusiveTime = FMath::Max<float>(MaxInclusiveTime, ChildInclusiveTime);
-				MinInclusiveTime = FMath::Min<float>(MinInclusiveTime, ChildInclusiveTime);
-			}
-			InclusiveTime = WorstInclusiveTime;
+			bEventFiltered = GraphFilter == UEdGraphSchema_K2::FN_UserConstructionScript;
 		}
 	}
-}
-
-double FBPProfilerBranchStat::GetMaxTime() const
-{
-	double Result = 0.0;
-	if (PerformanceStats.IsDataValid())
-	{
-		Result += PerformanceStats.GetMaxTime();
-	}
-	Result += MaxInclusiveTime;
-	return Result;
-}
-
-double FBPProfilerBranchStat::GetMinTime() const
-{
-	double Result = 0.0;
-	if (PerformanceStats.IsDataValid())
-	{
-		Result += PerformanceStats.GetMinTime();
-	}
-	else
-	{
-		Result += MinInclusiveTime;
-	}
-	return Result;
-}
-
-void FBPProfilerBranchStat::SubmitData(FScriptExecEvent& EventData)
-{
-	check (EventData == ObjectContext);
-	// Cache exec path
-	CachedExecPin = EventData.GetGraphPinPtr();
-	// Submit data
-	if (EventData.CanSubmitData())
-	{
-		PerformanceStats.AddEventTiming(EventData.GetTime());
-	}
+	return bEventFiltered;
 }
 
 #undef LOCTEXT_NAMESPACE
