Copying //UE4/Dev-Blueprints to Dev-Main (//UE4/Dev-Main)

diff --git a/Engine/Source/Editor/Kismet/Private/Profiler/BPProfilerStatisticWidgets.h b/Engine/Source/Editor/Kismet/Private/Profiler/BPProfilerStatisticWidgets.h
--- a/Engine/Source/Editor/Kismet/Private/Profiler/BPProfilerStatisticWidgets.h
+++ b/Engine/Source/Editor/Kismet/Private/Profiler/BPProfilerStatisticWidgets.h
@@ -2,7 +2,8 @@
 
 #pragma once
 
-#include "ScriptEventExecution.h"
+#include "ScriptPerfData.h"
+#include "TracePath.h"
 
 //////////////////////////////////////////////////////////////////////////
 // SProfilerStatRow
@@ -15,19 +16,23 @@ namespace EBlueprintProfilerStat
 		Name = 0,
 		TotalTime,
 		InclusiveTime,
-		ExclusiveTime,
+		Time,
+		PureTime,
 		MaxTime,
 		MinTime,
 		Samples
 	};
 };
 
-class SProfilerStatRow : public SMultiColumnTableRow<FBPProfilerStatPtr>
+// Shared pointer to a performance stat
+typedef TSharedPtr<class FBPProfilerStatWidget> FBPStatWidgetPtr;
+
+class SProfilerStatRow : public SMultiColumnTableRow<FBPStatWidgetPtr>
 {
 protected:
 
 	/** The profiler statistic associated with this stat row */
-	FBPProfilerStatPtr ItemToEdit;
+	FBPStatWidgetPtr ItemToEdit;
 
 public:
 
@@ -35,7 +40,7 @@ public:
 	SLATE_END_ARGS()
 
 	/** Constructs the statistic row item widget */
-	void Construct(const FArguments& InArgs, TSharedRef<STableViewBase> OwnerTableView, FBPProfilerStatPtr InItemToEdit);
+	void Construct(const FArguments& InArgs, TSharedRef<STableViewBase> OwnerTableView, FBPStatWidgetPtr InItemToEdit);
 
 	/** Generates and returns a widget for the specified column name and the associated statistic */
 	virtual TSharedRef<SWidget> GenerateWidgetForColumn(const FName& ColumnName) override;
@@ -49,245 +54,69 @@ public:
 };
 
 //////////////////////////////////////////////////////////////////////////
-// FBPPerformanceData
+// FBPProfilerStatWidget
 
-struct FBPPerformanceData
+class FBPProfilerStatWidget : public TSharedFromThis<FBPProfilerStatWidget>
 {
 public:
 
-	FBPPerformanceData()
-		: Samples(StatSampleSize)
-		, CachedExclusiveTime(0.0)
-		, MaxTime(-MAX_dbl)
-		, MinTime(MAX_dbl)
-		, TotalTime(0.0)
-		, NumSamples(0)
-		, bDirty(false)
-	{
-	}
-
-	/** Add a single event node timing into the dataset */
-	void AddEventTiming(const double NodeTiming);
-
-	/** Update and average the current timings from the buffer */
-	void Update();
-
-	/** Reset the current data buffer and all derived data stats */
-	void Reset();
-
-	/** Returns if this data structure has any valid data */
-	bool IsDataValid() const { return MaxTime != -MAX_dbl && MinTime != MAX_dbl; }
-
-	/** Marks the data is dirty and signals an update is required */
-	void MarkDataDirty() { bDirty = true; }
-
-	/** Signals that the data is dirty and requires and update */
-	bool IsDataDirty() const { return bDirty; }
-
-	/** Returns the various stats this container holds */
-	double GetMaxTime() const { return MaxTime; }
-	double GetMinTime() const { return MinTime; }
-	double GetTotalTime() const { return TotalTime; }
-	double GetExclusiveTime() const { return CachedExclusiveTime; }
-	int32 GetSampleCount() const { return NumSamples; }
-	int32 GetSafeSampleCount() const { return Samples.Num() ? Samples.Num() : 1; }
-
-public:
-
-	/** Controls how many samples are used to generate the averaged samples */
-	static int32 StatSampleSize;
-
-private:
-
-	/** Sample buffer */
-	TSimpleRingBuffer<double> Samples;
-	/** Cached exclusive timing */
-	double CachedExclusiveTime;
-	/** Max exclusive timing */
-	double MaxTime;
-	/** Min exclusive timing */
-	double MinTime;
-	/** Total time accrued */
-	double TotalTime;
-	/** The current inclusive sample count */
-	int32 NumSamples;
-	/** Flag to indicate the data wants the average recalculated */
-	bool bDirty;
-
-};
-
-//////////////////////////////////////////////////////////////////////////
-// FBPProfilerStat
-
-class FBPProfilerStat : public TSharedFromThis<FBPProfilerStat>
-{
-public:
-
-	FBPProfilerStat(FBPProfilerStatPtr InParentStat)
-		: InclusiveTime(0.0)
-		, MinInclusiveTime(MAX_dbl)
-		, MaxInclusiveTime(-MAX_dbl)
-		, InclusiveTotalTime(0.0)
-		, InclusiveSamples(0)
-		, ParentStat(InParentStat)
-	{
-	}
-
-	FBPProfilerStat(FBPProfilerStatPtr InParentStat, FScriptExecEvent& Context)
-		: InclusiveTime(0.0)
-		, MinInclusiveTime(MAX_dbl)
-		, MaxInclusiveTime(-MAX_dbl)
-		, InclusiveTotalTime(0.0)
-		, InclusiveSamples(0)
-		, ObjectContext(Context)
-		, ParentStat(InParentStat)
-	{
-		UpdateStatName();
-	}
-
-	virtual ~FBPProfilerStat() {}
-
-	/** Release Data */
-	void ReleaseData()
+	FBPProfilerStatWidget(TSharedPtr<class FScriptExecutionNode> InExecNode, const FTracePath& WidgetTracePathIn)
+		: WidgetTracePath(WidgetTracePathIn)
+		, ExecNode(InExecNode)
 	{
-		// Need to investigate why this needs to be called manually.
-		for (auto Iter : CachedChildren)
-		{
-			Iter.Value->ReleaseData();
-			Iter.Value.Reset();
-		}
-		ParentStat.Reset();
 	}
 
-	/** Find context entry point based on the passed in exec event */
-	FBPProfilerStatPtr FindEntryPoint(FScriptExecEvent& EventEntry);
-
-	/** Returns the parent container statistic */
-	FBPProfilerStatPtr GetParent() { return ParentStat; }
+	virtual ~FBPProfilerStatWidget() {}
 
-	/** Find or add a child stat based in exec context */
-	virtual FBPProfilerStatPtr FindOrAddChildByContext(FScriptExecEvent& Context);
+	/** Generate exec node widgets */
+	virtual void GenerateExecNodeWidgets(const FName InstanceName, const FName FilterGraph);
 
 	/** Gather children for list/tree widget */
-	void GatherChildren(TArray<FBPProfilerStatPtr>& OutChildren);
-
-	/** Update the current profiling statistics */
-	virtual void UpdateProfilingStats();
-
-	/** Returns if this stat has valid data */
-	bool HasValidData() { return PerformanceStats.IsDataValid(); }
+	void GatherChildren(TArray<TSharedPtr<FBPProfilerStatWidget>>& OutChildren);
 
-	/** Generate name widget */
-	TSharedRef<SWidget> GenerateNameWidget();
-
-	/** Create text attribute for column */
-	TAttribute<FText> GetTextAttribute(FName ColumnName) const;
+	/** Generate column widget */
+	TSharedRef<SWidget> GenerateColumnWidget(FName ColumnName);
 
 	/** Returns the exec path object context for this statistic */
-	FScriptExecEvent& GetObjectContext() { return ObjectContext; }
-
-	/** Submits new event data to this profiler statistic */
-	virtual void SubmitData(class FScriptExecEvent& EventData);
-
-	/** Returns the call depth for this statistic */
-	int32 GetCallDepth() const { return ObjectContext.IsValid() ? ObjectContext.GetCallDepth() : -1; }
+	TSharedPtr<class FScriptExecutionNode> GetExecNode() { return ExecNode; }
 
-	// Return current statisic values with safety checking */
-	virtual double GetMaxTime() const;
-	virtual double GetMinTime() const;
-	float GetTotalTime() const;
-	float GetInclusiveTotalTime() const;
-	float GetInclusiveTime() const;
-	float GetExclusiveTime() const;
-	float GetSampleCount() const;
-	float GetSafeSampleCount() const;
-	int32 GetInclusiveSampleCount() const { return InclusiveSamples; }
+	/** Returns the performance data object for this statistic */
+	TSharedPtr<FScriptPerfData> GetPerformanceData() { return PerformanceStats; }
 
-protected:
-
-	friend class FBPProfilerBranchStat;
-
-	/** Cache the statistic text from the exec event context */
-	virtual void UpdateStatName();
-
-	/** Returns the slate icon from the node referenced by the exec event */
-	const FSlateBrush* GetIcon();
+	/** Get expansion state */
+	bool GetExpansionState() const;
 
-	/** Returns the icon color override */
-	FLinearColor GetIconColor() const;
+	/** Set expansion state */
+	void SetExpansionState(bool bExpansionStateIn);
 
-	/** Returns the statistic tool tip */
-	const FText GetToolTip();
+	/** Expand widget state */
+	void ExpandWidgetState(TSharedPtr<STreeView<FBPStatWidgetPtr>> TreeView, bool bStateIn);
 
-	/** Navigates to the exec event node in the graph */
-	void OnNavigateToObject();
+	/** Restore widget Expansion state */
+	void RestoreWidgetExpansionState(TSharedPtr<STreeView<FBPStatWidgetPtr>> TreeView);
 
-	/** Returns the various statistic text values for the stats */
-	virtual FText GetInclusiveTimeText() const;
-	virtual FText GetExclusiveTimeText() const;
-	virtual FText GetTotalTimeText() const;
-	virtual FText GetMaxTimeText() const;
-	virtual FText GetMinTimeText() const;
-	virtual FText GetSamplesText() const;
+	/** Returns if any children expect to be expanded */
+	bool ProbeChildWidgetExpansionStates();
 
 protected:
 
-	/** Local profiling statistics */
-	FBPPerformanceData PerformanceStats;
-	/** Container profiling statistics */
-	double InclusiveTime;
-	double MinInclusiveTime;
-	double MaxInclusiveTime;
-	double InclusiveTotalTime;
-	int32 InclusiveSamples;
-	/** Script execution event context */
-	FScriptExecEvent ObjectContext;
-	/** Cached statistic name */
-	FText StatName;
-	/** Statistic icon color */
-	FLinearColor IconColor;
-	/** Parent statistic */
-	FBPProfilerStatPtr ParentStat;
-	/** Child statistics */
-	TMap<TWeakObjectPtr<const UObject>, FBPProfilerStatPtr> CachedChildren;
+	/** Returns if the exec node representing an event is filtered from view */
+	bool IsEventFiltered(TSharedPtr<FScriptExecutionNode> InEventNode, FName GraphFilter) const;
 
-};
-
-//////////////////////////////////////////////////////////////////////////
-// FBPProfilerBranchStat
-
-class FBPProfilerBranchStat : public FBPProfilerStat
-{
-public:
-
-	FBPProfilerBranchStat(FBPProfilerStatPtr InParentStat, FScriptExecEvent& Context)
-		: FBPProfilerStat(InParentStat, Context)
-		, bAreExecutionPathsSequential(false)
-	{
-	}
-
-	virtual ~FBPProfilerBranchStat() {}
-
-	//~ FBPProfilerStat Start
-	virtual FBPProfilerStatPtr FindOrAddChildByContext(FScriptExecEvent& Context) override;
-	void CreateChildrenForPins();
-	virtual void UpdateProfilingStats() override;
-	virtual double GetMaxTime() const override;
-	virtual double GetMinTime() const override;
-	virtual void SubmitData(class FScriptExecEvent& EventData) override;
-	//~ FBPProfilerStat End
+	/** Navigate to referenced node */
+	void NavigateTo() const;
 
 protected:
 
-	/** Returns whether different branches execute in sequence rather exclusively */
-	bool IsBranchSequential();
-
-private:
-
-	/** Cached exec pin */
-	TWeakObjectPtr<const UObject> CachedExecPin;
-
-	/** Signals that execution paths are sequential */
-	bool bAreExecutionPathsSequential;
+	/** Widget script tracepath */
+	FTracePath WidgetTracePath;
+	/** Script execution event node */
+	TSharedPtr<FScriptExecutionNode> ExecNode;
+	/** Script performance data */
+	TSharedPtr<FScriptPerfData> PerformanceStats;
+	/** Script code offset */
+	int32 ScriptCodeOffset;
+	/** Child statistics */
+	TArray<TSharedPtr<FBPProfilerStatWidget>> CachedChildren;
 
 };
\ No newline at end of file
diff --git a/Engine/Source/Editor/Kismet/Private/Profiler/EventExecution.cpp b/Engine/Source/Editor/Kismet/Private/Profiler/EventExecution.cpp
new file mode 100644
--- /dev/null
+++ b/Engine/Source/Editor/Kismet/Private/Profiler/EventExecution.cpp
@@ -0,0 +1,339 @@
+// Copyright 1998-2016 Epic Games, Inc. All Rights Reserved.
+
+#include "BlueprintEditorPrivatePCH.h"
+#include "EventExecution.h"
+#include "ScriptPerfData.h"
+
+//////////////////////////////////////////////////////////////////////////
+// FScriptNodeExecLinkage
+
+void FScriptNodeExecLinkage::AddLinkedNode(const int32 PinScriptOffset, TSharedPtr<class FScriptExecutionNode> LinkedNode)
+{
+	TSharedPtr<class FScriptExecutionNode>& Result = LinkedNodes.FindOrAdd(PinScriptOffset);
+	Result = LinkedNode;
+}
+
+//////////////////////////////////////////////////////////////////////////
+// FScriptNodePerfData
+
+TSharedPtr<FScriptPerfData> FScriptNodePerfData::GetPerfDataByInstanceAndTracePath(FName InstanceName, const FTracePath& TracePath)
+{ 
+	TMap<const uint32, TSharedPtr<FScriptPerfData>>& InstanceMap = InstanceInputPinToPerfDataMap.FindOrAdd(InstanceName);
+	TSharedPtr<FScriptPerfData>& Result = InstanceMap.FindOrAdd(TracePath);
+	if (!Result.IsValid())
+	{
+		Result = MakeShareable<FScriptPerfData>(new FScriptPerfData);
+	}
+	return Result;
+}
+
+void FScriptNodePerfData::GetInstancePerfDataByTracePath(const FTracePath& TracePath, TArray<TSharedPtr<FScriptPerfData>>& ResultsOut)
+{
+	for (auto InstanceIter : InstanceInputPinToPerfDataMap)
+	{
+		if (InstanceIter.Key != NAME_None)
+		{
+			TSharedPtr<FScriptPerfData> Result = InstanceIter.Value.FindOrAdd(TracePath);
+			if (Result.IsValid())
+			{
+				ResultsOut.Add(Result);
+			}
+		}
+	}
+}
+
+TSharedPtr<FScriptPerfData> FScriptNodePerfData::GetBlueprintPerfDataByTracePath(const FTracePath& TracePath)
+{
+	TMap<const uint32, TSharedPtr<FScriptPerfData>>& BlueprintMap = InstanceInputPinToPerfDataMap.FindOrAdd(NAME_None);
+	TSharedPtr<FScriptPerfData>& Result = BlueprintMap.FindOrAdd(TracePath);
+	if (!Result.IsValid())
+	{
+		Result = MakeShareable<FScriptPerfData>(new FScriptPerfData);
+	}
+	return Result;
+}
+
+//////////////////////////////////////////////////////////////////////////
+// FScriptExecutionNode
+
+FScriptExecutionNode::FScriptExecutionNode()
+	: NodeFlags(EScriptExecutionNodeFlags::None)
+	, bExpansionState(false)
+{
+}
+
+bool FScriptExecutionNode::operator == (const FScriptExecutionNode& NodeIn) const
+{
+	return NodeName == NodeIn.NodeName;
+}
+
+void FScriptExecutionNode::GetLinearExecutionPath(TArray<FLinearExecPath>& LinearExecutionNodes, const FTracePath& TracePath)
+{
+	LinearExecutionNodes.Add(FLinearExecPath(AsShared(), TracePath));
+	if (GetNumLinkedNodes() == 1)
+	{
+		FTracePath NewTracePath(TracePath);
+		for (auto NodeIter : LinkedNodes)
+		{
+			NewTracePath.AddExitPin(NodeIter.Key);
+			NodeIter.Value->GetLinearExecutionPath(LinearExecutionNodes, NewTracePath);
+		}
+	}
+}
+
+EScriptStatContainerType::Type FScriptExecutionNode::GetStatisticContainerType() const
+{
+	EScriptStatContainerType::Type Result = EScriptStatContainerType::Standard;
+	if (HasFlags(EScriptExecutionNodeFlags::Instance|EScriptExecutionNodeFlags::Event|EScriptExecutionNodeFlags::CallSite))
+	{
+		Result = EScriptStatContainerType::Container;
+	}
+	else if (HasFlags(EScriptExecutionNodeFlags::SequentialBranch))
+	{
+		Result = EScriptStatContainerType::SequentialBranch;
+	}
+	else if (HasFlags(EScriptExecutionNodeFlags::ExecPin))
+	{
+		Result = EScriptStatContainerType::NewExecutionPath;
+	}
+	return Result;
+}
+
+void FScriptExecutionNode::RefreshStats(const FTracePath& TracePath)
+{
+	bool bRefreshChildStats = false;
+	bool bRefreshLinkStats = true;
+	TArray<TSharedPtr<FScriptPerfData>> BlueprintPooledStats;
+	// Update stats based on type
+	switch(GetStatisticContainerType())
+	{
+		case EScriptStatContainerType::Container:
+		{
+			// Only consider stat data on this path
+			GetInstancePerfDataByTracePath(TracePath, BlueprintPooledStats);
+			bRefreshChildStats = true;
+			break;
+		}
+		case EScriptStatContainerType::Standard:
+		case EScriptStatContainerType::SequentialBranch:
+		{
+			// Only consider stat data on this path
+			GetInstancePerfDataByTracePath(TracePath, BlueprintPooledStats);
+			break;
+		}
+		//case EScriptStatContainerType::SequentialBranch:
+		//{
+		//	for (auto InstanceIter : InstanceInputPinToPerfDataMap)
+		//	{
+		//		if (InstanceIter.Key != NAME_None)
+		//		{
+		//			TSharedPtr<FScriptPerfData> InstancePerfData = GetPerfDataByInstanceAndTracePath(InstanceIter.Key, TracePath);
+		//			InstancePerfData->Reset();
+		//			BlueprintPooledStats.Add(InstancePerfData);
+		//		
+		//			for (auto LinkIter : LinkedNodes)
+		//			{
+		//				FTracePath LinkTracePath(TracePath);
+		//				LinkTracePath.AddExitPin(LinkIter.Key);
+		//				TSharedPtr<FScriptPerfData> LinkPerfData = LinkIter.Value->GetPerfDataByInstanceAndTracePath(InstanceIter.Key, LinkTracePath);
+		//				InstancePerfData->AddData(*LinkPerfData.Get());
+		//				LinkIter.Value->RefreshStats(LinkTracePath);
+		//			}
+		//		}
+		//	}
+		//	bRefreshLinkStats = false;
+		//	break;
+		//}
+		case EScriptStatContainerType::NewExecutionPath:
+		{
+			TSet<FName> AllInstances;
+			// Refresh child stats and copy as branch stats - this is a dummy entry
+			for (auto ChildIter : ChildNodes)
+			{
+				// Fill out instance data that can be missing because its a dummy node, or the code below can fail.
+				for (auto InstanceIter : ChildIter->InstanceInputPinToPerfDataMap)
+				{
+					if (InstanceIter.Key != NAME_None)
+					{
+						AllInstances.Add(InstanceIter.Key);
+					}
+				}
+				ChildIter->RefreshStats(TracePath);
+			}
+			for (auto InstanceIter : AllInstances)
+			{
+				if (InstanceIter != NAME_None)
+				{
+					TSharedPtr<FScriptPerfData> InstancePerfData = GetPerfDataByInstanceAndTracePath(InstanceIter, TracePath);
+					InstancePerfData->Reset();
+					BlueprintPooledStats.Add(InstancePerfData);
+				
+					for (auto ChildIter : ChildNodes)
+					{
+						TSharedPtr<FScriptPerfData> ChildPerfData = ChildIter->GetPerfDataByInstanceAndTracePath(InstanceIter, TracePath);
+						InstancePerfData->AddBranchData(*ChildPerfData.Get());
+					}
+				}
+			}
+			// Update Blueprint stats as branches
+			TSharedPtr<FScriptPerfData> BlueprintData = GetBlueprintPerfDataByTracePath(TracePath);
+			BlueprintData->Reset();
+
+			for (auto BlueprintChildDataIter : BlueprintPooledStats)
+			{
+				BlueprintData->AddBranchData(*BlueprintChildDataIter.Get());
+			}
+			BlueprintPooledStats.Reset(0);
+			bRefreshLinkStats = false;
+			break;
+		}
+	}
+	// Refresh Child Links
+	if (bRefreshChildStats)
+	{
+		for (auto ChildIter : ChildNodes)
+		{
+			ChildIter->RefreshStats(TracePath);
+		}
+	}
+	// Refresh All links
+	if (bRefreshLinkStats)
+	{
+		for (auto LinkIter : LinkedNodes)
+		{
+			FTracePath LinkTracePath(TracePath);
+			LinkTracePath.AddExitPin(LinkIter.Key);
+			LinkIter.Value->RefreshStats(LinkTracePath);
+		}
+	}
+	// Update the owning blueprint stats
+	if (BlueprintPooledStats.Num() > 0)
+	{
+		TSharedPtr<FScriptPerfData> BlueprintData = GetBlueprintPerfDataByTracePath(TracePath);
+		BlueprintData->Reset();
+
+		for (auto BlueprintChildDataIter : BlueprintPooledStats)
+		{
+			BlueprintData->AddData(*BlueprintChildDataIter.Get());
+		}
+	}
+}
+
+void FScriptExecutionNode::GetAllExecNodes(TMap<FName, TSharedPtr<FScriptExecutionNode>>& ExecNodesOut)
+{
+	for (auto ChildIter : ChildNodes)
+	{
+		ChildIter->GetAllExecNodes(ExecNodesOut);
+		ExecNodesOut.Add(ChildIter->GetName()) = ChildIter;
+	}
+	for (auto LinkIter : LinkedNodes)
+	{
+		LinkIter.Value->GetAllExecNodes(ExecNodesOut);
+		ExecNodesOut.Add(LinkIter.Value->GetName()) = LinkIter.Value;
+	}
+}
+
+//////////////////////////////////////////////////////////////////////////
+// FScriptExecutionBlueprint
+
+TSharedPtr<FScriptExecutionNode> FScriptExecutionBlueprint::GetInstanceByName(FName InstanceName)
+{
+	TSharedPtr<FScriptExecutionNode> Result;
+	for (auto Iter : Instances)
+	{
+		if (Iter->GetName() == InstanceName)
+		{
+			Result = Iter;
+			break;
+		}
+	}
+	return Result;
+}
+
+void FScriptExecutionBlueprint::RefreshStats(const FTracePath& TracePath)
+{
+	TArray<TSharedPtr<FScriptPerfData>> InstanceData;
+	// Update event stats
+	for (auto BlueprintEventIter : ChildNodes)
+	{
+		// This crawls through and updates all instance stats and pools the results into the blueprint node stats
+		// as an overall blueprint performance representation.
+		BlueprintEventIter->RefreshStats(TracePath);
+	}
+	// Update instance stats
+	for (auto InstanceIter : Instances)
+	{
+		TSharedPtr<FScriptPerfData> InstancePerfData = InstanceIter->GetPerfDataByInstanceAndTracePath(InstanceIter->GetName(), TracePath);
+		InstancePerfData->Reset();
+		// Update all top level instance stats now the events are up to date.
+		for (auto InstanceEventIter : InstanceIter->GetChildNodes())
+		{
+			TSharedPtr<FScriptPerfData> InstanceEventPerfData = InstanceEventIter->GetPerfDataByInstanceAndTracePath(InstanceIter->GetName(), TracePath);
+			InstancePerfData->AddData(*InstanceEventPerfData.Get());
+		}
+		// Add for consolidation at the bottom.
+		InstanceData.Add(InstancePerfData);
+	}
+	// Finally... update the root stats for the owning blueprint.
+	if (InstanceData.Num() > 0)
+	{
+		TSharedPtr<FScriptPerfData> BlueprintData = GetBlueprintPerfDataByTracePath(TracePath);
+		BlueprintData->Reset();
+
+		for (auto InstanceIter : InstanceData)
+		{
+			BlueprintData->AddData(*InstanceIter.Get());
+		}
+	}
+}
+
+void FScriptExecutionBlueprint::GetAllExecNodes(TMap<FName, TSharedPtr<FScriptExecutionNode>>& ExecNodesOut)
+{
+	for (auto BlueprintEventIter : ChildNodes)
+	{
+		BlueprintEventIter->GetAllExecNodes(ExecNodesOut);
+		ExecNodesOut.Add(BlueprintEventIter->GetName()) = BlueprintEventIter;
+	}
+	for (auto InstanceIter : Instances)
+	{
+		InstanceIter->GetAllExecNodes(ExecNodesOut);
+		ExecNodesOut.Add(InstanceIter->GetName()) = InstanceIter;
+	}
+	ExecNodesOut.Add(GetName()) = AsShared();
+}
+
+//////////////////////////////////////////////////////////////////////////
+// FScriptExecutionContext
+
+void FScriptExecutionContext::GetAllExecNodes(TMap<FName, TSharedPtr<FScriptExecutionNode>>& ExecNodesOut)
+{
+	if (ExecutionNode.IsValid())
+	{
+		ExecutionNode->GetAllExecNodes(ExecNodesOut);
+	}
+}
+
+const UEdGraphNode* FScriptExecutionContext::GetNodeFromCodeLocation(const int32 CodeLocation, UFunction* FunctionOverride) const
+{
+	if (BlueprintClass.IsValid())
+	{
+		if (FunctionOverride)
+		{
+			return BlueprintClass.Get()->GetDebugData().FindSourceNodeFromCodeLocation(FunctionOverride, CodeLocation, true);
+		}
+		else if (BlueprintFunction.IsValid())
+		{
+			return BlueprintClass.Get()->GetDebugData().FindSourceNodeFromCodeLocation(BlueprintFunction.Get(), CodeLocation, true);
+		}
+	}
+	return nullptr;
+}
+
+void FScriptExecutionContext::UpdateConnectedStats()
+{
+	if (ExecutionNode.IsValid())
+	{
+		FTracePath InitialTracePath;
+		ExecutionNode->RefreshStats(InitialTracePath);
+	}
+}
diff --git a/Engine/Source/Editor/Kismet/Private/Profiler/SBlueprintProfilerView.cpp b/Engine/Source/Editor/Kismet/Private/Profiler/SBlueprintProfilerView.cpp
--- a/Engine/Source/Editor/Kismet/Private/Profiler/SBlueprintProfilerView.cpp
+++ b/Engine/Source/Editor/Kismet/Private/Profiler/SBlueprintProfilerView.cpp
@@ -1,15 +1,19 @@
 // Copyright 1998-2016 Epic Games, Inc. All Rights Reserved.
 
-
 #include "BlueprintEditorPrivatePCH.h"
-#include "BPProfilerStatisticWidgets.h"
 #include "SBlueprintProfilerView.h"
 
+// Profiler View types
+#include "SProfilerOverview.h"
+#include "SLeastPerformantDisplay.h"
+#include "SGraphExecutionStatDisplay.h"
+#include "ScriptPerfData.h"
+
+#define LOCTEXT_NAMESPACE "BlueprintProfilerView"
+
 //////////////////////////////////////////////////////////////////////////
 // SBlueprintProfilerView
 
-FNumberFormattingOptions SBlueprintProfilerView::NumberFormat;
-
 SBlueprintProfilerView::~SBlueprintProfilerView()
 {
 	// Register for Profiler toggle events
@@ -17,225 +21,206 @@ SBlueprintProfilerView::~SBlueprintProfilerView()
 }
 
 void SBlueprintProfilerView::Construct(const FArguments& InArgs)
-{	
+{
 	CurrentViewType = InArgs._ProfileViewType;
+	BlueprintEditor = InArgs._AssetEditor;
 	// Register for Profiler toggle events
-	FBlueprintCoreDelegates::OnToggleScriptProfiler.AddRaw(this, &SBlueprintProfilerView::OnToggleProfiler);
+	FBlueprintCoreDelegates::OnToggleScriptProfiler.AddSP(this, &SBlueprintProfilerView::OnToggleProfiler);
 	// Initialise the number format.
-	NumberFormat.MinimumFractionalDigits = 4;
-	NumberFormat.MaximumFractionalDigits = 4;
-	NumberFormat.UseGrouping = false;
-	// Update the average sample retention count.
-	FBPPerformanceData::StatSampleSize = GetDefault<UEditorExperimentalSettings>()->BlueprintProfilerAverageSampleCount;
+	FNumberFormattingOptions StatisticNumberFormat;
+	StatisticNumberFormat.MinimumFractionalDigits = 4;
+	StatisticNumberFormat.MaximumFractionalDigits = 4;
+	StatisticNumberFormat.UseGrouping = false;
+	FScriptPerfData::SetNumberFormattingForStats(StatisticNumberFormat);
 	// Create the display text for the user
 	if (GetDefault<UEditorExperimentalSettings>()->bBlueprintPerformanceAnalysisTools)
 	{
 		IBlueprintProfilerInterface* ProfilerInterface = FModuleManager::GetModulePtr<IBlueprintProfilerInterface>("BlueprintProfiler");
 		const bool bProfilerEnabled	= ProfilerInterface && ProfilerInterface->IsProfilerEnabled();
-		StatusText = bProfilerEnabled ? NSLOCTEXT("ActiveProfilerWidget", "NoDataText", "The Blueprint Profiler is active but does not currently have any data to display") :
-										NSLOCTEXT("InactiveProfilerWidget", "InactiveText", "The Blueprint Profiler is currently Inactive");
+		StatusText = bProfilerEnabled ? LOCTEXT("ProfilerNoDataText", "The Blueprint Profiler is active but does not currently have any data to display") :
+										LOCTEXT("ProfilerInactiveText", "The Blueprint Profiler is currently Inactive");
 	}
 	else
 	{
-		StatusText = NSLOCTEXT("DisabledProfilerWidget", "DisabledProfilerText", "The Blueprint Profiler is experimental and is currently not enabled in the editor preferences");
+		StatusText = LOCTEXT("DisabledProfilerText", "The Blueprint Profiler is experimental and is currently not enabled in the editor preferences");
 	}
 	// Create the profiler view widgets
-	bTreeViewIsDirty = false;
-	CreateProfilerWidgets();
+	UpdateActiveProfilerWidget();
 }
 
-TSharedRef<ITableRow> SBlueprintProfilerView::OnGenerateRowForProfiler(FBPProfilerStatPtr InItem, const TSharedRef<STableViewBase>& OwnerTable)
+void SBlueprintProfilerView::OnToggleProfiler(bool bEnabled)
 {
-	return SNew(SProfilerStatRow, OwnerTable, InItem);
+	StatusText = bEnabled ? LOCTEXT("ProfilerNoDataText", "The Blueprint Profiler is active but does not currently have any data to display") :
+							LOCTEXT("ProfilerInactiveText", "The Blueprint Profiler is currently Inactive");
+	UpdateActiveProfilerWidget();
 }
 
-void SBlueprintProfilerView::OnGetChildrenForProfiler(FBPProfilerStatPtr InParent, TArray<FBPProfilerStatPtr>& OutChildren)
+void SBlueprintProfilerView::UpdateActiveProfilerWidget()
 {
-	if (InParent.IsValid())
-	{
-		InParent->GatherChildren(OutChildren);
+		FMenuBuilder ViewComboContent(true, NULL);
+		ViewComboContent.AddMenuEntry(	LOCTEXT("OverviewViewType", "Profiler Overview"), 
+										LOCTEXT("OverviewViewTypeDesc", "Displays High Level Profiling Information"), 
+										FSlateIcon(), 
+										FExecuteAction::CreateSP(this, &SBlueprintProfilerView::OnViewSelectionChanged, EBlueprintPerfViewType::Overview),
+										NAME_None, 
+										EUserInterfaceActionType::Button);
+		ViewComboContent.AddMenuEntry(	LOCTEXT("ExecutionGraphViewType", "Execution Graph"), 
+										LOCTEXT("ExecutionGraphViewTypeDesc", "Displays the Execution Graph with Statistics"), 
+										FSlateIcon(), 
+										FExecuteAction::CreateSP(this, &SBlueprintProfilerView::OnViewSelectionChanged, EBlueprintPerfViewType::ExecutionGraph),
+										NAME_None, 
+										EUserInterfaceActionType::Button);
+		ViewComboContent.AddMenuEntry(	LOCTEXT("LeastPerformantViewType", "Least Performant List"), 
+										LOCTEXT("LeastPerformantViewTypeDesc", "Displays a list of Least Performant Areas of the Blueprint"), 
+										FSlateIcon(), 
+										FExecuteAction::CreateSP(this, &SBlueprintProfilerView::OnViewSelectionChanged, EBlueprintPerfViewType::LeastPerformant),
+										NAME_None, 
+										EUserInterfaceActionType::Button);
+
+		ChildSlot
+		[
+			SNew(SVerticalBox)
+			+SVerticalBox::Slot()
+			.AutoHeight()
+			.Padding(0)
+			[
+				SNew(SBorder)
+				.BorderImage(FEditorStyle::GetBrush("BlueprintProfiler.ViewToolBar"))
+				.Padding(0)
+				.VAlign(VAlign_Top)
+				.HAlign(HAlign_Right)
+				[
+					SNew(SHorizontalBox)
+					+SHorizontalBox::Slot()
+					.AutoWidth()
+					.Padding(FMargin(5,0))
+					[
+						SAssignNew(ViewComboButton, SComboButton)
+						.ForegroundColor(this, &SBlueprintProfilerView::GetViewButtonForegroundColor)
+						.ToolTipText(LOCTEXT("BlueprintProfilerViewType", "View Type"))
+						.ButtonStyle(FEditorStyle::Get(), "ToggleButton")
+						.ContentPadding(2)
+						.MenuContent()
+						[
+							ViewComboContent.MakeWidget()
+						]
+						.ButtonContent()
+						[
+							CreateViewButton()
+						]
+					]
+				]
+			]
+			+SVerticalBox::Slot()
+			[
+				SNew(SBorder)
+				.Padding(FMargin(0,2,0,0))
+				.BorderImage(FEditorStyle::GetBrush("NoBorder"))
+				[
+					CreateActiveStatisticWidget()
+				]
+			]
+		];
 	}
-}
 
-void SBlueprintProfilerView::OnToggleProfiler(bool bEnabled)
+FText SBlueprintProfilerView::GetCurrentViewText() const
 {
-	StatusText = bEnabled ? NSLOCTEXT("ActiveProfilerWidget", "NoDataText", "The Blueprint Profiler is active but does not currently have any data to display") :
-							NSLOCTEXT("InactiveProfilerWidget", "InactiveText", "The Blueprint Profiler is currently Inactive");
-	if (!bEnabled)
+	switch(CurrentViewType)
 	{
-		// Clear profiler data, this requires work as I think the stats are currently leaking.
-		for (auto Iter : RootTreeItems)
-		{
-			Iter->ReleaseData();
-			Iter.Reset();
-		}
-		RootTreeItems.Reset();
-		ClassContexts.Reset();
-		if (ExecutionStatTree.IsValid())
-		{
-			bTreeViewIsDirty = true;
-		}
+		case EBlueprintPerfViewType::Overview:			return LOCTEXT("OverviewLabel", "Profiler Overview");
+		case EBlueprintPerfViewType::ExecutionGraph:	return LOCTEXT("ExecutionGraphLabel", "Execution Graph");
+		case EBlueprintPerfViewType::LeastPerformant:	return LOCTEXT("LeastPerformantLabel", "Least Performant");
 	}
+	return LOCTEXT("UnknownViewLabel", "Unknown ViewType");
 }
 
-void SBlueprintProfilerView::CreateProfilerWidgets()
+TSharedRef<SWidget> SBlueprintProfilerView::CreateViewButton() const
 {
-	ChildSlot
-	[
-		SNew(SVerticalBox)
-		+SVerticalBox::Slot()
-		.FillHeight(0.2f)
+	return SNew(SHorizontalBox)
+		+SHorizontalBox::Slot()
+		.AutoWidth()
+		.VAlign(VAlign_Center)
 		[
-			SNew(SBorder)
-			.Padding(0)
-			.BorderImage(FEditorStyle::GetBrush("NoBorder"))
-			.Visibility(this, &SBlueprintProfilerView::IsProfilerHidden)
-			.HAlign(HAlign_Center)
-			.VAlign(VAlign_Center)
-			[
-				SNew(STextBlock)
-				.Text(this, &SBlueprintProfilerView::GetProfilerStatusText)
-			]
-		]
-#if 0 // To do stat ranking list
-		+SVerticalBox::Slot()
-		[
-			SAssignNew(BlueprintStatList, SListView<FBPProfilerStatPtr>)
-			.ListItemsSource(&RootTreeItems)
-			.SelectionMode(ESelectionMode::Single)
-			.OnGenerateRow(this, &SBlueprintProfilerView::OnGenerateRowForProfiler)
-//			.OnSelectionChanged(this, &SMessagingTypes::HandleTypeListSelectionChanged)
-			.HeaderRow
-			(
-				SNew(SHeaderRow)
-				+SHeaderRow::Column(SProfilerStatRow::GetStatName(EBlueprintProfilerStat::Name))
-				.DefaultLabel(SProfilerStatRow::GetStatText(EBlueprintProfilerStat::Name))
-				.ManualWidth(450)
-				+SHeaderRow::Column(SProfilerStatRow::GetStatName(EBlueprintProfilerStat::InclusiveTime))
-				.DefaultLabel(SProfilerStatRow::GetStatText(EBlueprintProfilerStat::InclusiveTime))
-				.ManualWidth(120)
-				+SHeaderRow::Column(SProfilerStatRow::GetStatName(EBlueprintProfilerStat::ExclusiveTime))
-				.DefaultLabel(SProfilerStatRow::GetStatText(EBlueprintProfilerStat::ExclusiveTime))
-				.ManualWidth(120)
-				+SHeaderRow::Column(SProfilerStatRow::GetStatName(EBlueprintProfilerStat::MaxTime))
-				.DefaultLabel(SProfilerStatRow::GetStatText(EBlueprintProfilerStat::MaxTime))
-				.ManualWidth(90)
-				+SHeaderRow::Column(SProfilerStatRow::GetStatName(EBlueprintProfilerStat::MinTime))
-				.DefaultLabel(SProfilerStatRow::GetStatText(EBlueprintProfilerStat::MinTime))
-				.ManualWidth(90)
-				+SHeaderRow::Column(SProfilerStatRow::GetStatName(EBlueprintProfilerStat::TotalTime))
-				.DefaultLabel(SProfilerStatRow::GetStatText(EBlueprintProfilerStat::TotalTime))
-				.ManualWidth(80)
-				+SHeaderRow::Column(SProfilerStatRow::GetStatName(EBlueprintProfilerStat::Samples))
-				.DefaultLabel(SProfilerStatRow::GetStatText(EBlueprintProfilerStat::Samples))
-				.ManualWidth(60)
-			)
+			SNew(SImage)
+			.Image( FEditorStyle::GetBrush("GenericViewButton") )
 		]
-#endif
-		+SVerticalBox::Slot()
+		+SHorizontalBox::Slot()
+		.AutoWidth()
+		.Padding(2, 0, 0, 0)
+		.VAlign(VAlign_Center)
 		[
-			SNew(SBorder)
-			.Padding(0)
-			.BorderImage(FEditorStyle::GetBrush("NoBorder"))
-			.Visibility(this, &SBlueprintProfilerView::IsProfilerVisible)
-			[
-				SNew(SVerticalBox)
-				+SVerticalBox::Slot()
-				.AutoHeight()
-				.HAlign(HAlign_Left)
-				.VAlign(VAlign_Center)
-				[
-					SNew(STextBlock)
-					.Text(NSLOCTEXT("BueprintProfilerView", "StatisticTreeTitle", "Execution Tree Performance Statistics"))
-				]
-				+SVerticalBox::Slot()
-				.FillHeight(0.8f)
-				[
-					SAssignNew(ExecutionStatTree, STreeView<FBPProfilerStatPtr>)
-					.TreeItemsSource(&RootTreeItems)
-					.SelectionMode(ESelectionMode::Single)
-					.OnGetChildren(this, &SBlueprintProfilerView::OnGetChildrenForProfiler)
-					.OnGenerateRow(this, &SBlueprintProfilerView::OnGenerateRowForProfiler)
-					.HeaderRow
-					(
-						SNew(SHeaderRow)
-						+SHeaderRow::Column(SProfilerStatRow::GetStatName(EBlueprintProfilerStat::Name))
-						.DefaultLabel(SProfilerStatRow::GetStatText(EBlueprintProfilerStat::Name))
-						.ManualWidth(450)
-						+SHeaderRow::Column(SProfilerStatRow::GetStatName(EBlueprintProfilerStat::InclusiveTime))
-						.DefaultLabel(SProfilerStatRow::GetStatText(EBlueprintProfilerStat::InclusiveTime))
-						.ManualWidth(120)
-						+SHeaderRow::Column(SProfilerStatRow::GetStatName(EBlueprintProfilerStat::ExclusiveTime))
-						.DefaultLabel(SProfilerStatRow::GetStatText(EBlueprintProfilerStat::ExclusiveTime))
-						.ManualWidth(120)
-						+SHeaderRow::Column(SProfilerStatRow::GetStatName(EBlueprintProfilerStat::MaxTime))
-						.DefaultLabel(SProfilerStatRow::GetStatText(EBlueprintProfilerStat::MaxTime))
-						.ManualWidth(90)
-						+SHeaderRow::Column(SProfilerStatRow::GetStatName(EBlueprintProfilerStat::MinTime))
-						.DefaultLabel(SProfilerStatRow::GetStatText(EBlueprintProfilerStat::MinTime))
-						.ManualWidth(90)
-						+SHeaderRow::Column(SProfilerStatRow::GetStatName(EBlueprintProfilerStat::TotalTime))
-						.DefaultLabel(SProfilerStatRow::GetStatText(EBlueprintProfilerStat::TotalTime))
-						.ManualWidth(80)
-						+SHeaderRow::Column(SProfilerStatRow::GetStatName(EBlueprintProfilerStat::Samples))
-						.DefaultLabel(SProfilerStatRow::GetStatText(EBlueprintProfilerStat::Samples))
-						.ManualWidth(60)
-					)
-				]
-			]
-		]
-	];
+			SNew(STextBlock)
+			.Text(GetCurrentViewText())
+		];
 }
 
-void SBlueprintProfilerView::Tick(const FGeometry& AllottedGeometry, const double InCurrentTime, const float InDeltaTime)
+FSlateColor SBlueprintProfilerView::GetViewButtonForegroundColor() const
 {
-	if (GetDefault<UEditorExperimentalSettings>()->bBlueprintPerformanceAnalysisTools)
+	static const FName InvertedForegroundName("InvertedForeground");
+	static const FName DefaultForegroundName("DefaultForeground");
+	return ViewComboButton->IsHovered() ? FEditorStyle::GetSlateColor(InvertedForegroundName) : FEditorStyle::GetSlateColor(DefaultForegroundName);
+}
+
+void SBlueprintProfilerView::OnViewSelectionChanged(const EBlueprintPerfViewType::Type NewViewType)
+{
+	if (CurrentViewType != NewViewType)
 	{
-		IBlueprintProfilerInterface* ProfilerInterface = FModuleManager::GetModulePtr<IBlueprintProfilerInterface>("BlueprintProfiler");
-		if (ProfilerInterface && ProfilerInterface->IsProfilerEnabled())
-		{
-			// Find and process execution path timing data.
-			bool bRequestRefresh = false;
-			TArray<FBlueprintInstrumentedEvent>& ProfilerData = ProfilerInterface->GetProfilingData();
-			int32 ExecStart = 0;
+		CurrentViewType = NewViewType;
+		UpdateActiveProfilerWidget();
+	}
+}
 
-			for (int32 EventIdx = 0; EventIdx < ProfilerData.Num(); ++EventIdx)
-			{
-				switch(ProfilerData[EventIdx].GetType())
-				{
-					case EScriptInstrumentation::Class:
-					{
-						CurrentClassPath = ProfilerData[EventIdx].GetObjectPath();
-						ExecStart = EventIdx;
-						break;
-					}
-					case EScriptInstrumentation::Stop:
-					{
-						check(!CurrentClassPath.IsEmpty());
-						FProfilerClassContext* ActiveClass = ClassContexts.Find(CurrentClassPath);
-						const bool bNewClassContext = !ActiveClass;
-						if (!ActiveClass)
-						{
-							ActiveClass = &ClassContexts.Add(CurrentClassPath);
-						}
-						ActiveClass->ProcessExecutionPath(ProfilerData, ExecStart, EventIdx);
-						if (bNewClassContext)
-						{
-							RootTreeItems.Add(ActiveClass->GetRootStat());
-							bRequestRefresh = true;
-						}
-						ExecStart = EventIdx+1;
-						break;
-					}
-				}
-			}
-			// Reset profiler data
-			ProfilerData.Reset();
-		}
+TSharedRef<SWidget> SBlueprintProfilerView::CreateActiveStatisticWidget()
+{
+	// Get profiler status
+	EBlueprintPerfViewType::Type NewViewType = EBlueprintPerfViewType::None;
+	IBlueprintProfilerInterface* ProfilerInterface = FModuleManager::GetModulePtr<IBlueprintProfilerInterface>("BlueprintProfiler");
+	if (ProfilerInterface && ProfilerInterface->IsProfilerEnabled())
+	{
+		NewViewType = CurrentViewType;
 	}
-	// Refresh Tree
-	if (bTreeViewIsDirty && ExecutionStatTree.IsValid())
+	switch(NewViewType)
 	{
-		ExecutionStatTree->RequestTreeRefresh();
+		case EBlueprintPerfViewType::Overview:
+		{
+			return SNew(SProfilerOverview)
+				.AssetEditor(BlueprintEditor);
+		}
+		case EBlueprintPerfViewType::ExecutionGraph:
+		{
+			return SNew(SGraphExecutionStatDisplay)
+				.AssetEditor(BlueprintEditor);
+		}
+		case EBlueprintPerfViewType::LeastPerformant:
+		{
+			return SNew(SLeastPerformantDisplay)
+				.AssetEditor(BlueprintEditor);
+		}
+		default:
+		{
+			// Default place holder
+			return SNew(SVerticalBox)
+				+SVerticalBox::Slot()
+				.HAlign(HAlign_Center)
+				.VAlign(VAlign_Center)
+				.Padding(0)
+				[
+					SNew(SHorizontalBox)
+					+SHorizontalBox::Slot()
+					[
+						SNew(SBorder)
+						.Padding(0)
+						.BorderImage(FEditorStyle::GetBrush("NoBorder"))
+						.HAlign(HAlign_Center)
+						.VAlign(VAlign_Center)
+						[
+							SNew(STextBlock)
+							.Text(this, &SBlueprintProfilerView::GetProfilerStatusText)
+						]
+					]
+				];
+		}
 	}
 }
+
+#undef LOCTEXT_NAMESPACE
diff --git a/Engine/Source/Editor/Kismet/Private/Profiler/SBlueprintProfilerView.h b/Engine/Source/Editor/Kismet/Private/Profiler/SBlueprintProfilerView.h
--- a/Engine/Source/Editor/Kismet/Private/Profiler/SBlueprintProfilerView.h
+++ b/Engine/Source/Editor/Kismet/Private/Profiler/SBlueprintProfilerView.h
@@ -2,34 +2,31 @@
 
 #pragma once
 
-#include "Developer/BlueprintProfiler/Public/BlueprintProfilerModule.h"
-#include "ScriptEventExecution.h"
 
-// Shared pointer to a performance stat
-typedef TSharedPtr<class FBPProfilerStat> FBPProfilerStatPtr;
-
-/** Enum to describe type of performance stats view */
-enum EPerformanceViewType
+/** Blueprint performance view type */
+namespace EBlueprintPerfViewType
 {
-	PVT_LiveView = 0,
-	PVT_Blueprint,
-	PVT_BlueprintInstance,
-	PVT_BlueprintTrace
-};
+	enum Type
+	{
+		None = 0,
+		Overview,
+		ExecutionGraph,
+		LeastPerformant
+	};
+}
 
 //////////////////////////////////////////////////////////////////////////
 // SBlueprintProfilerView
 
-typedef TSharedPtr<class FBPProfilerStat> FBPProfilerStatPtr;
-
 class SBlueprintProfilerView : public SCompoundWidget
 {
 public:
-	SLATE_BEGIN_ARGS( SBlueprintProfilerView )
-		: _ProfileViewType(PVT_Blueprint)
+	SLATE_BEGIN_ARGS(SBlueprintProfilerView)
+		: _ProfileViewType(EBlueprintPerfViewType::ExecutionGraph)
 		{}
 
-		SLATE_ARGUMENT( EPerformanceViewType, ProfileViewType )
+		SLATE_ARGUMENT(TWeakPtr<FBlueprintEditor>, AssetEditor)
+		SLATE_ARGUMENT(EBlueprintPerfViewType::Type, ProfileViewType)
 	SLATE_END_ARGS()
 
 
@@ -37,65 +34,46 @@ public:
 
 	~SBlueprintProfilerView();
 
-	void Construct( const FArguments& InArgs );
-
-	// SWidget interface
-	virtual void Tick( const FGeometry& AllottedGeometry, const double InCurrentTime, const float InDeltaTime ) override;
-	// End of SWidget interface
-
-	/** Returns FText number format specifier */
-	static const FNumberFormattingOptions& GetNumberFormat() { return NumberFormat; }
+	void Construct(const FArguments& InArgs);
 
 protected:
 
 	/** Called when the profiler state is toggled */
 	void OnToggleProfiler(bool bEnabled);
 
-	/** Called to generate a stat row for the execution stat tree */
-	TSharedRef<ITableRow> OnGenerateRowForProfiler(FBPProfilerStatPtr InItem, const TSharedRef<STableViewBase>& OwnerTable);
-
-	/** Called to generate children for given profiler stat */
-	void OnGetChildrenForProfiler(FBPProfilerStatPtr InParent, TArray<FBPProfilerStatPtr>& OutChildren);
-
-	/** Returns is the profiler widgets are currently visible */
-	EVisibility IsProfilerVisible() const { return RootTreeItems.Num() > 0 ? EVisibility::Visible : EVisibility::Collapsed; }
-
-	/** Returns if the profiler widgets are currently not visible, used to display the inactive message */
-	EVisibility IsProfilerHidden() const { return RootTreeItems.Num() == 0 ? EVisibility::Visible : EVisibility::Collapsed; }
-
 	/** Returns the profiler current status message */
 	FText GetProfilerStatusText() const { return StatusText; }
 
 	/** Called to create the child profiler widgets */
-	void CreateProfilerWidgets();
+	void UpdateActiveProfilerWidget();
 
-protected:
+	/** Returns the current view label */
+	FText GetCurrentViewText() const;
 
-	/** Profiler status Text */
-	FText StatusText;
+	/** Constructs and returns the view button widget */
+	TSharedRef<SWidget> CreateViewButton() const;
 
-	/** Number formatting options */
-	static FNumberFormattingOptions NumberFormat;
+	/** Returns the active view button foreground color */
+	FSlateColor GetViewButtonForegroundColor() const;
 
-	/** Current profiler view type */
-	EPerformanceViewType CurrentViewType;
+	/** Called when the profiler view type is changed */
+	void OnViewSelectionChanged(const EBlueprintPerfViewType::Type NewViewType);
 
-	/** Sorted stat list for identifying expensive items */
-	TSharedPtr<SListView<FBPProfilerStatPtr>> BlueprintStatList;
+	/** Create active statistic display widget */
+	TSharedRef<SWidget> CreateActiveStatisticWidget();
 
-	/** Sorted stat list for identifying expensive items */
-	TSharedPtr<STreeView<FBPProfilerStatPtr>> ExecutionStatTree;
+protected:
 
-	/** Source item array for the treeview */
-	TArray<FBPProfilerStatPtr> RootTreeItems;
+	/** Profiler status Text */
+	FText StatusText;
 
-	/** Active Instance Contexts */
-	TMap<FString, FProfilerClassContext> ClassContexts;
+	/** Blueprint editor */
+	TWeakPtr<class FBlueprintEditor> BlueprintEditor;
 
-	/** Cached active class statistic path */
-	FString CurrentClassPath;
+	/** Current profiler view type */
+	EBlueprintPerfViewType::Type CurrentViewType;
 
-	/** Flag to control treeview refresh */
-	bool bTreeViewIsDirty;
+	/** View combo button widget */
+	TSharedPtr<SComboButton> ViewComboButton;
 
 };
diff --git a/Engine/Source/Editor/Kismet/Private/Profiler/SGraphExecutionStatDisplay.cpp b/Engine/Source/Editor/Kismet/Private/Profiler/SGraphExecutionStatDisplay.cpp
new file mode 100644
--- /dev/null
+++ b/Engine/Source/Editor/Kismet/Private/Profiler/SGraphExecutionStatDisplay.cpp
@@ -0,0 +1,370 @@
+// Copyright 1998-2016 Epic Games, Inc. All Rights Reserved.
+
+#include "BlueprintEditorPrivatePCH.h"
+#include "BPProfilerStatisticWidgets.h"
+#include "SGraphExecutionStatDisplay.h"
+#include "BlueprintEditor.h"
+#include "Public/Profiler/EventExecution.h"
+#include "SDockTab.h"
+
+#define LOCTEXT_NAMESPACE "BlueprintProfilerGraphExecView"
+
+//////////////////////////////////////////////////////////////////////////
+// SGraphExecutionStatDisplay
+
+SGraphExecutionStatDisplay::~SGraphExecutionStatDisplay()
+{
+	// Remove delegate for profiling toggle events
+	FBlueprintCoreDelegates::OnToggleScriptProfiler.RemoveAll(this);
+	// Remove delegate for graph structural changes
+	if (IBlueprintProfilerInterface* Profiler = FModuleManager::GetModulePtr<IBlueprintProfilerInterface>("BlueprintProfiler"))
+	{
+		Profiler->GetGraphLayoutChangedDelegate().RemoveAll(this);
+	}
+}
+
+void SGraphExecutionStatDisplay::Construct(const FArguments& InArgs)
+{	
+	BlueprintEditor = InArgs._AssetEditor;
+	// Register for profiling toggle events
+	FBlueprintCoreDelegates::OnToggleScriptProfiler.AddSP(this, &SGraphExecutionStatDisplay::OnToggleProfiler);
+	// Remove delegate for graph structural changes
+	if (IBlueprintProfilerInterface* Profiler = FModuleManager::GetModulePtr<IBlueprintProfilerInterface>("BlueprintProfiler"))
+	{
+		Profiler->GetGraphLayoutChangedDelegate().AddSP(this, &SGraphExecutionStatDisplay::OnGraphLayoutChanged);
+	}
+	// Set Default States
+	bDisplayInstances = false;
+	bScopeToActiveDebugInstance = false;
+	bAutoExpandStats = false;
+	bFilterEventsToGraph = true;
+
+	ChildSlot
+	[
+		SNew(SVerticalBox)
+		+SVerticalBox::Slot()
+		.Padding(FMargin(0, 0, 0, 2))
+		.AutoHeight()
+		.HAlign(HAlign_Fill)
+		.VAlign(VAlign_Top)
+		[
+			SNew(SBorder)
+			.Padding(4)
+			.BorderImage(FEditorStyle::GetBrush("BlueprintProfiler.ViewToolBar"))
+			[
+				SNew(SVerticalBox)
+				+SVerticalBox::Slot()
+				.HAlign(HAlign_Right)
+				.AutoHeight()
+				[
+					SNew(SHorizontalBox)
+					+SHorizontalBox::Slot()
+					.AutoWidth()
+					.Padding(FMargin(5,0))
+					[
+						SNew(SCheckBox)
+						.Content()
+						[
+							SNew(STextBlock)
+							.Text(LOCTEXT("FilterToGraph", "Filter Events to Current Graph"))
+						]
+						.IsChecked(this, &SGraphExecutionStatDisplay::GetGraphFilterCheck)
+						.OnCheckStateChanged(this, &SGraphExecutionStatDisplay::OnGraphFilterCheck)
+					]
+					+SHorizontalBox::Slot()
+					.AutoWidth()
+					.Padding(FMargin(5,0))
+					[
+						SNew(SCheckBox)
+						.Content()
+						[
+							SNew(STextBlock)
+							.Text(LOCTEXT("ShowInstancesCheck", "Display By Instance"))
+						]
+						.IsChecked(this, &SGraphExecutionStatDisplay::GetDisplayByInstanceCheck)
+						.OnCheckStateChanged(this, &SGraphExecutionStatDisplay::OnDisplayByInstanceCheck)
+					]
+					+SHorizontalBox::Slot()
+					.AutoWidth()
+					.Padding(FMargin(5,0))
+					[
+						SNew(SCheckBox)
+						.Content()
+						[
+							SNew(STextBlock)
+							.Text(LOCTEXT("InstanceFilterCheck", "Scope To Debug Filter"))
+						]
+						.IsChecked(this, &SGraphExecutionStatDisplay::GetInstanceFilterCheck)
+						.OnCheckStateChanged(this, &SGraphExecutionStatDisplay::OnInstanceFilterCheck)
+					]
+					+SHorizontalBox::Slot()
+					.AutoWidth()
+					.Padding(FMargin(5,0))
+					[
+						SNew(SCheckBox)
+						.Content()
+						[
+							SNew(STextBlock)
+							.Text(LOCTEXT("AutoItemExpansion", "Auto Expand Statistics"))
+						]
+						.IsChecked(this, &SGraphExecutionStatDisplay::GetAutoExpansionCheck)
+						.OnCheckStateChanged(this, &SGraphExecutionStatDisplay::OnAutoExpansionCheck)
+					]
+				]
+			]
+		]
+		+SVerticalBox::Slot()
+		[
+			SNew(SBorder)
+			.Padding(0)
+			.BorderImage(FEditorStyle::GetBrush("NoBorder"))
+			[
+				SAssignNew(ExecutionStatTree, STreeView<FBPStatWidgetPtr>)
+				.TreeItemsSource(&RootTreeItems)
+				.SelectionMode(ESelectionMode::Single)
+				.OnGetChildren(this, &SGraphExecutionStatDisplay::OnGetChildren)
+				.OnGenerateRow(this, &SGraphExecutionStatDisplay::OnGenerateRow)
+				.OnMouseButtonDoubleClick(this, &SGraphExecutionStatDisplay::OnDoubleClickStatistic)
+				.OnExpansionChanged(this, &SGraphExecutionStatDisplay::OnStatisticExpansionChanged)
+				.HeaderRow
+				(
+					SNew(SHeaderRow)
+					+SHeaderRow::Column(SProfilerStatRow::GetStatName(EBlueprintProfilerStat::Name))
+					.DefaultLabel(SProfilerStatRow::GetStatText(EBlueprintProfilerStat::Name))
+					.ManualWidth(450)
+					+SHeaderRow::Column(SProfilerStatRow::GetStatName(EBlueprintProfilerStat::Time))
+					.DefaultLabel(SProfilerStatRow::GetStatText(EBlueprintProfilerStat::Time))
+					.ManualWidth(70)
+					+SHeaderRow::Column(SProfilerStatRow::GetStatName(EBlueprintProfilerStat::PureTime))
+					.DefaultLabel(SProfilerStatRow::GetStatText(EBlueprintProfilerStat::PureTime))
+					.ManualWidth(100)
+					+SHeaderRow::Column(SProfilerStatRow::GetStatName(EBlueprintProfilerStat::InclusiveTime))
+					.DefaultLabel(SProfilerStatRow::GetStatText(EBlueprintProfilerStat::InclusiveTime))
+					.ManualWidth(120)
+					+SHeaderRow::Column(SProfilerStatRow::GetStatName(EBlueprintProfilerStat::MaxTime))
+					.DefaultLabel(SProfilerStatRow::GetStatText(EBlueprintProfilerStat::MaxTime))
+					.ManualWidth(90)
+					+SHeaderRow::Column(SProfilerStatRow::GetStatName(EBlueprintProfilerStat::MinTime))
+					.DefaultLabel(SProfilerStatRow::GetStatText(EBlueprintProfilerStat::MinTime))
+					.ManualWidth(90)
+					+SHeaderRow::Column(SProfilerStatRow::GetStatName(EBlueprintProfilerStat::TotalTime))
+					.DefaultLabel(SProfilerStatRow::GetStatText(EBlueprintProfilerStat::TotalTime))
+					.ManualWidth(80)
+					+SHeaderRow::Column(SProfilerStatRow::GetStatName(EBlueprintProfilerStat::Samples))
+					.DefaultLabel(SProfilerStatRow::GetStatText(EBlueprintProfilerStat::Samples))
+					.ManualWidth(60)
+				)
+			]
+		]
+	];
+}
+
+void SGraphExecutionStatDisplay::OnGraphLayoutChanged(TWeakObjectPtr<UBlueprint> Blueprint)
+{
+	if (CurrentBlueprint == Blueprint)
+	{
+		// Force a stat update.
+		CurrentBlueprint.Reset();
+	}
+}
+
+void SGraphExecutionStatDisplay::OnToggleProfiler(bool bEnabled)
+{
+	// Force a stat update.
+	CurrentBlueprint.Reset();
+}
+
+TSharedRef<ITableRow> SGraphExecutionStatDisplay::OnGenerateRow(FBPStatWidgetPtr InItem, const TSharedRef<STableViewBase>& OwnerTable)
+{
+	return SNew(SProfilerStatRow, OwnerTable, InItem);
+}
+
+void SGraphExecutionStatDisplay::OnGetChildren(FBPStatWidgetPtr InParent, TArray<FBPStatWidgetPtr>& OutChildren)
+{
+	if (InParent.IsValid())
+	{
+		InParent->GatherChildren(OutChildren);
+	}
+}
+
+void SGraphExecutionStatDisplay::OnGraphFilterCheck(ECheckBoxState NewState)
+{
+	const bool bOldState = bFilterEventsToGraph;
+	bFilterEventsToGraph = (NewState == ECheckBoxState::Checked) ? true : false;
+	// Force Update.
+	if (bFilterEventsToGraph != bOldState)
+	{
+		CurrentBlueprint.Reset();
+	}
+}
+
+void SGraphExecutionStatDisplay::OnInstanceFilterCheck(ECheckBoxState NewState)
+{
+	const bool bOldState = bScopeToActiveDebugInstance;
+	bScopeToActiveDebugInstance = (NewState == ECheckBoxState::Checked) ? true : false;
+	// Force Update.
+	if (bScopeToActiveDebugInstance != bOldState)
+	{
+		CurrentBlueprint.Reset();
+	}
+}
+
+void SGraphExecutionStatDisplay::OnDisplayByInstanceCheck(ECheckBoxState NewState)
+{
+	const bool bOldState = bDisplayInstances;
+	bDisplayInstances = (NewState == ECheckBoxState::Checked) ? true : false;
+	// Force Update.
+	if (bDisplayInstances != bOldState)
+	{
+		CurrentBlueprint.Reset();
+	}
+}
+
+void SGraphExecutionStatDisplay::OnAutoExpansionCheck(ECheckBoxState NewState)
+{
+	const bool bOldState = bAutoExpandStats;
+	bAutoExpandStats = (NewState == ECheckBoxState::Checked) ? true : false;
+	// Force Update.
+	if (bAutoExpandStats != bOldState)
+	{
+		if (bAutoExpandStats)
+		{
+			for (auto Iter : RootTreeItems)
+			{
+				Iter->ExpandWidgetState(ExecutionStatTree, bAutoExpandStats);
+			}
+		}
+		else
+		{
+			for (auto Iter : RootTreeItems)
+			{
+				Iter->RestoreWidgetExpansionState(ExecutionStatTree);
+			}
+		}
+	}
+}
+
+void SGraphExecutionStatDisplay::OnDoubleClickStatistic(FBPStatWidgetPtr Item)
+{
+	if (Item.IsValid())
+	{
+		Item->ExpandWidgetState(ExecutionStatTree, !Item->GetExpansionState());
+	}
+}
+
+void SGraphExecutionStatDisplay::OnStatisticExpansionChanged(FBPStatWidgetPtr Item, bool bExpanded)
+{
+	if (!bAutoExpandStats && Item.IsValid())
+	{
+		Item->SetExpansionState(bExpanded);
+	}
+}
+
+void SGraphExecutionStatDisplay::Tick(const FGeometry& AllottedGeometry, const double InCurrentTime, const float InDeltaTime)
+{
+	if (GetDefault<UEditorExperimentalSettings>()->bBlueprintPerformanceAnalysisTools)
+	{
+		if (IBlueprintProfilerInterface* ProfilerInterface = FModuleManager::GetModulePtr<IBlueprintProfilerInterface>("BlueprintProfiler"))
+		{
+			// Find and process execution path timing data.
+			if (BlueprintEditor.IsValid())
+			{
+				const TArray<UObject*>* Blueprints = BlueprintEditor.Pin()->GetObjectsCurrentlyBeingEdited();
+
+				if (Blueprints->Num() == 1)
+				{
+					UBlueprint* NewBlueprint = Cast<UBlueprint>((*Blueprints)[0]);
+					TWeakObjectPtr<const UObject> NewInstance = NewBlueprint->GetObjectBeingDebugged();
+					FName NewInstancePath = NewInstance.IsValid() ? FName(*ProfilerInterface->GetValidInstancePath(NewInstance->GetPathName())) : NAME_None;
+					const bool bBlueprintChanged = NewBlueprint != CurrentBlueprint.Get();
+					const bool bInstanceChanged = bScopeToActiveDebugInstance && NewInstancePath != CurrentInstancePath;
+					// Find active graph to filter on
+					bool bGraphFilterChanged = false;
+					TSharedPtr<SDockTab> DockTab = BlueprintEditor.Pin()->DocumentManager->GetActiveTab();
+					if (DockTab.IsValid())
+					{
+						TSharedRef<SGraphEditor> ActiveGraphEditor = StaticCastSharedRef<SGraphEditor>(DockTab->GetContent());
+						bGraphFilterChanged = ActiveGraphEditor->GetCurrentGraph() != CurrentGraph.Get();
+						CurrentGraph = ActiveGraphEditor->GetCurrentGraph();
+					}
+
+					if (bBlueprintChanged||bInstanceChanged||bGraphFilterChanged)
+					{
+						const FString BlueprintPath = NewBlueprint->GeneratedClass->GetPathName();
+						TSharedPtr<FScriptExecutionNode> BlueprintEntry = ProfilerInterface->GetProfilerDataForNode(FName(*BlueprintPath));
+						RootTreeItems.Reset(0);
+
+						if (BlueprintEntry.IsValid() && BlueprintEntry->IsRootNode())
+						{
+							// Cache Active blueprint and Instance
+							CurrentBlueprint = NewBlueprint;
+							CurrentInstancePath = NewInstancePath;
+							// Build Instance widget execution trees
+							TSharedPtr<FScriptExecutionBlueprint> BlueprintExecNode = StaticCastSharedPtr<FScriptExecutionBlueprint>(BlueprintEntry);
+
+							if (BlueprintExecNode.IsValid())
+							{
+								FName GraphFilter = CurrentGraph.IsValid() && bFilterEventsToGraph ? CurrentGraph.Get()->GetFName() : NAME_None;
+
+								if (bDisplayInstances && bScopeToActiveDebugInstance)
+								{
+									if (CurrentInstancePath != NAME_None)
+									{
+										TSharedPtr<FScriptExecutionNode> InstanceStat = BlueprintExecNode->GetInstanceByName(CurrentInstancePath);
+										if (InstanceStat.IsValid())
+										{
+											FTracePath InstanceTracePath;
+											FBPStatWidgetPtr InstanceWidget = MakeShareable<FBPProfilerStatWidget>(new FBPProfilerStatWidget(InstanceStat, InstanceTracePath));
+											InstanceWidget->GenerateExecNodeWidgets(InstanceStat->GetName(), GraphFilter);
+											RootTreeItems.Add(InstanceWidget);
+										}
+									}
+								}
+								else if (bDisplayInstances)
+								{
+									for (int32 InstanceIdx = 0; InstanceIdx < BlueprintExecNode->GetInstanceCount(); ++InstanceIdx)
+									{
+										TSharedPtr<FScriptExecutionNode> InstanceStat = BlueprintExecNode->GetInstanceByIndex(InstanceIdx);
+										FTracePath InstanceTracePath;
+										FBPStatWidgetPtr InstanceWidget = MakeShareable<FBPProfilerStatWidget>(new FBPProfilerStatWidget(InstanceStat, InstanceTracePath));
+										InstanceWidget->GenerateExecNodeWidgets(InstanceStat->GetName(), GraphFilter);
+										RootTreeItems.Add(InstanceWidget);
+									}
+								}
+								else
+								{
+									FTracePath TracePath;
+									FBPStatWidgetPtr BlueprintWidget = MakeShareable<FBPProfilerStatWidget>(new FBPProfilerStatWidget(BlueprintExecNode, TracePath));
+									BlueprintWidget->GenerateExecNodeWidgets(NAME_None, GraphFilter);
+									RootTreeItems.Add(BlueprintWidget);
+								}
+							}
+						}
+						// Refresh Tree
+						if (ExecutionStatTree.IsValid())
+						{
+							ExecutionStatTree->RequestTreeRefresh();
+							if (bAutoExpandStats)
+							{
+								for (auto Iter : RootTreeItems)
+								{
+									Iter->ExpandWidgetState(ExecutionStatTree, bAutoExpandStats);
+								}
+							}
+							else
+							{
+								for (auto Iter : RootTreeItems)
+								{
+									Iter->ProbeChildWidgetExpansionStates();
+									Iter->RestoreWidgetExpansionState(ExecutionStatTree);
+								}
+							}
+						}
+					}
+				}
+			}
+		}
+	}
+}
+
+#undef LOCTEXT_NAMESPACE
diff --git a/Engine/Source/Editor/Kismet/Private/Profiler/SGraphExecutionStatDisplay.h b/Engine/Source/Editor/Kismet/Private/Profiler/SGraphExecutionStatDisplay.h
new file mode 100644
--- /dev/null
+++ b/Engine/Source/Editor/Kismet/Private/Profiler/SGraphExecutionStatDisplay.h
@@ -0,0 +1,122 @@
+// Copyright 1998-2016 Epic Games, Inc. All Rights Reserved.
+
+#pragma once
+
+#include "Developer/BlueprintProfiler/Public/BlueprintProfilerModule.h"
+
+// Shared pointer to a performance stat
+typedef TSharedPtr<class FBPProfilerStatWidget> FBPStatWidgetPtr;
+
+//////////////////////////////////////////////////////////////////////////
+// SGraphExecutionStatDisplay
+
+class SGraphExecutionStatDisplay : public SCompoundWidget
+{
+public:
+	SLATE_BEGIN_ARGS(SGraphExecutionStatDisplay)
+		{}
+
+		SLATE_ARGUMENT(TWeakPtr<FBlueprintEditor>, AssetEditor)
+	SLATE_END_ARGS()
+
+
+public:
+
+	~SGraphExecutionStatDisplay();
+
+	void Construct(const FArguments& InArgs);
+
+	// SWidget interface
+	virtual void Tick(const FGeometry& AllottedGeometry, const double InCurrentTime, const float InDeltaTime) override;
+	// End of SWidget interface
+
+protected:
+
+	/** Called when blueprint graph layout changes occur */
+	void OnGraphLayoutChanged(TWeakObjectPtr<UBlueprint> Blueprint);
+
+	/** Called when the blueprint profiler is enabled/disabled */
+	void OnToggleProfiler(bool bEnabled);
+
+	/** Returns is the profiler widgets are currently visible */
+	EVisibility GetNoDataWidgetVisibility() const { return RootTreeItems.Num() == 0 ? EVisibility::Visible : EVisibility::Collapsed; }
+
+	/** Returns if the profiler widgets are currently not visible, used to display the inactive message */
+	EVisibility GetWidgetVisibility() const { return RootTreeItems.Num() > 0 ? EVisibility::Visible : EVisibility::Collapsed; }
+
+	/** Called to generate a stat row for the execution stat tree */
+	TSharedRef<ITableRow> OnGenerateRow(FBPStatWidgetPtr InItem, const TSharedRef<STableViewBase>& OwnerTable);
+
+	/** Called to generate children for given profiler stat */
+	void OnGetChildren(FBPStatWidgetPtr InParent, TArray<FBPStatWidgetPtr>& OutChildren);
+
+	/** Called to evaluate the graph filter checkbox */
+	ECheckBoxState GetGraphFilterCheck() const { return bFilterEventsToGraph ? ECheckBoxState::Checked : ECheckBoxState::Unchecked; }
+
+	/** Called to update the instance filter state */
+	void OnGraphFilterCheck(ECheckBoxState NewState);
+
+	/** Called to evaluate the instance filter checkbox */
+	ECheckBoxState GetInstanceFilterCheck() const { return bDisplayInstances && bScopeToActiveDebugInstance ? ECheckBoxState::Checked : ECheckBoxState::Unchecked; }
+
+	/** Called to update the instance filter state */
+	void OnInstanceFilterCheck(ECheckBoxState NewState);
+
+	/** Called to evaluate the display instances checkbox */
+	ECheckBoxState GetDisplayByInstanceCheck() const { return bDisplayInstances ? ECheckBoxState::Checked : ECheckBoxState::Unchecked; }
+
+	/** Called to update the display instances state */
+	void OnDisplayByInstanceCheck(ECheckBoxState NewState);
+
+	/** Called to evaluate the auto expansion checkbox */
+	ECheckBoxState GetAutoExpansionCheck() const { return bAutoExpandStats ? ECheckBoxState::Checked : ECheckBoxState::Unchecked; }
+
+	/** Called to update the auto expansion state */
+	void OnAutoExpansionCheck(ECheckBoxState NewState);
+
+	/** On double click statistic */
+	void OnDoubleClickStatistic(FBPStatWidgetPtr Item);
+
+	/** On statistic expansion state changed */
+	void OnStatisticExpansionChanged(FBPStatWidgetPtr Item, bool bExpanded);
+
+	/** Auto expand statistics */
+	void AutoExpandStatistics();
+
+
+protected:
+
+	/** Blueprint editor */
+	TWeakPtr<class FBlueprintEditor> BlueprintEditor;
+
+	/** Profiler status text attribute */
+	TAttribute<FText> ProfilerStatusText;
+
+	/** Sorted stat list for identifying expensive items */
+	TSharedPtr<STreeView<FBPStatWidgetPtr>> ExecutionStatTree;
+
+	/** Source item array for the treeview */
+	TArray<FBPStatWidgetPtr> RootTreeItems;
+
+	/** Current blueprint to display stats for */
+	TWeakObjectPtr<UBlueprint> CurrentBlueprint;
+
+	/** Current graph we filtered stats on */
+	TWeakObjectPtr<UEdGraph> CurrentGraph;
+
+	/** Current instance to display stats for */
+	FName CurrentInstancePath;
+
+	/** Display instances state */
+	bool bDisplayInstances;
+
+	/** Current instance filter state */
+	bool bScopeToActiveDebugInstance;
+
+	/** Current statistic expansion state */
+	bool bAutoExpandStats;
+
+	/** Filter events to graph */
+	bool bFilterEventsToGraph;
+
+};
diff --git a/Engine/Source/Editor/Kismet/Private/Profiler/SLeastPerformantDisplay.cpp b/Engine/Source/Editor/Kismet/Private/Profiler/SLeastPerformantDisplay.cpp
new file mode 100644
--- /dev/null
+++ b/Engine/Source/Editor/Kismet/Private/Profiler/SLeastPerformantDisplay.cpp
@@ -0,0 +1,109 @@
+// Copyright 1998-2016 Epic Games, Inc. All Rights Reserved.
+
+#include "BlueprintEditorPrivatePCH.h"
+#include "SLeastPerformantDisplay.h"
+#include "BPProfilerStatisticWidgets.h"
+
+//////////////////////////////////////////////////////////////////////////
+// SLeastPerformantDisplay
+
+SLeastPerformantDisplay::~SLeastPerformantDisplay()
+{
+}
+
+void SLeastPerformantDisplay::Construct(const FArguments& InArgs)
+{	
+	BlueprintEditor = InArgs._AssetEditor;
+
+	ChildSlot
+	[
+		SNew(SVerticalBox)
+		+SVerticalBox::Slot()
+		.Padding(FMargin(0, 0, 0, 2))
+		.AutoHeight()
+		.HAlign(HAlign_Fill)
+		.VAlign(VAlign_Top)
+		[
+			SNew(SBorder)
+			.Padding(4)
+			.BorderImage(FEditorStyle::GetBrush("BlueprintProfiler.ViewToolBar"))
+			[
+				SNew(SVerticalBox)
+				+SVerticalBox::Slot()
+				.HAlign(HAlign_Right)
+				.AutoHeight()
+				[
+					SNew(SHorizontalBox)
+				]
+			]
+		]
+		+SVerticalBox::Slot()
+		[
+			SNew(SBorder)
+			.Padding(0)
+			.BorderImage(FEditorStyle::GetBrush("NoBorder"))
+			[
+				SAssignNew(BlueprintStatList, SListView<FBPStatWidgetPtr>)
+				.ListItemsSource(&RootTreeItems)
+				.SelectionMode(ESelectionMode::Single)
+				.OnGenerateRow(this, &SLeastPerformantDisplay::OnGenerateRow)
+	//			.OnSelectionChanged(this, &SMessagingTypes::HandleTypeListSelectionChanged)
+				.HeaderRow
+				(
+					SNew(SHeaderRow)
+					+SHeaderRow::Column(SProfilerStatRow::GetStatName(EBlueprintProfilerStat::Name))
+					.DefaultLabel(SProfilerStatRow::GetStatText(EBlueprintProfilerStat::Name))
+					.ManualWidth(450)
+					+SHeaderRow::Column(SProfilerStatRow::GetStatName(EBlueprintProfilerStat::Time))
+					.DefaultLabel(SProfilerStatRow::GetStatText(EBlueprintProfilerStat::Time))
+					.ManualWidth(70)
+					+SHeaderRow::Column(SProfilerStatRow::GetStatName(EBlueprintProfilerStat::PureTime))
+					.DefaultLabel(SProfilerStatRow::GetStatText(EBlueprintProfilerStat::PureTime))
+					.ManualWidth(100)
+					+SHeaderRow::Column(SProfilerStatRow::GetStatName(EBlueprintProfilerStat::InclusiveTime))
+					.DefaultLabel(SProfilerStatRow::GetStatText(EBlueprintProfilerStat::InclusiveTime))
+					.ManualWidth(120)
+					+SHeaderRow::Column(SProfilerStatRow::GetStatName(EBlueprintProfilerStat::MaxTime))
+					.DefaultLabel(SProfilerStatRow::GetStatText(EBlueprintProfilerStat::MaxTime))
+					.ManualWidth(90)
+					+SHeaderRow::Column(SProfilerStatRow::GetStatName(EBlueprintProfilerStat::MinTime))
+					.DefaultLabel(SProfilerStatRow::GetStatText(EBlueprintProfilerStat::MinTime))
+					.ManualWidth(90)
+					+SHeaderRow::Column(SProfilerStatRow::GetStatName(EBlueprintProfilerStat::TotalTime))
+					.DefaultLabel(SProfilerStatRow::GetStatText(EBlueprintProfilerStat::TotalTime))
+					.ManualWidth(80)
+					+SHeaderRow::Column(SProfilerStatRow::GetStatName(EBlueprintProfilerStat::Samples))
+					.DefaultLabel(SProfilerStatRow::GetStatText(EBlueprintProfilerStat::Samples))
+					.ManualWidth(60)
+				)
+			]
+		]
+	];
+}
+
+TSharedRef<ITableRow> SLeastPerformantDisplay::OnGenerateRow(FBPStatWidgetPtr InItem, const TSharedRef<STableViewBase>& OwnerTable)
+{
+	return SNew(SProfilerStatRow, OwnerTable, InItem);
+}
+
+void SLeastPerformantDisplay::OnDoubleClickStatistic(FBPStatWidgetPtr Item)
+{
+	if (Item.IsValid())
+	{
+//		Item->ExpandWidgetState(ExecutionStatTree, !Item->GetExpansionState());
+	}
+}
+
+void SLeastPerformantDisplay::Tick(const FGeometry& AllottedGeometry, const double InCurrentTime, const float InDeltaTime)
+{
+	if (GetDefault<UEditorExperimentalSettings>()->bBlueprintPerformanceAnalysisTools)
+	{
+		if (IBlueprintProfilerInterface* Profiler = FModuleManager::GetModulePtr<IBlueprintProfilerInterface>("BlueprintProfiler"))
+		{
+			// Find and process execution path timing data.
+			if (BlueprintEditor.IsValid())
+			{
+			}
+		}
+	}
+}
diff --git a/Engine/Source/Editor/Kismet/Private/Profiler/SLeastPerformantDisplay.h b/Engine/Source/Editor/Kismet/Private/Profiler/SLeastPerformantDisplay.h
new file mode 100644
--- /dev/null
+++ b/Engine/Source/Editor/Kismet/Private/Profiler/SLeastPerformantDisplay.h
@@ -0,0 +1,55 @@
+// Copyright 1998-2016 Epic Games, Inc. All Rights Reserved.
+
+#pragma once
+
+#include "Developer/BlueprintProfiler/Public/BlueprintProfilerModule.h"
+
+// Shared pointer to a performance stat
+typedef TSharedPtr<class FBPProfilerStatWidget> FBPStatWidgetPtr;
+
+//////////////////////////////////////////////////////////////////////////
+// SLeastPerformantDisplay
+
+class SLeastPerformantDisplay : public SCompoundWidget
+{
+public:
+	SLATE_BEGIN_ARGS(SLeastPerformantDisplay)
+		{}
+
+		SLATE_ARGUMENT(TWeakPtr<FBlueprintEditor>, AssetEditor)
+	SLATE_END_ARGS()
+
+
+public:
+
+	~SLeastPerformantDisplay();
+
+	void Construct(const FArguments& InArgs);
+
+	// SWidget interface
+	virtual void Tick(const FGeometry& AllottedGeometry, const double InCurrentTime, const float InDeltaTime) override;
+	// End of SWidget interface
+
+protected:
+
+	/** Called to generate a stat row for the execution stat tree */
+	TSharedRef<ITableRow> OnGenerateRow(FBPStatWidgetPtr InItem, const TSharedRef<STableViewBase>& OwnerTable);
+
+	/** On double click statistic */
+	void OnDoubleClickStatistic(FBPStatWidgetPtr Item);
+
+protected:
+
+	/** Blueprint editor */
+	TWeakPtr<class FBlueprintEditor> BlueprintEditor;
+
+	/** Profiler status text attribute */
+	TAttribute<FText> ProfilerStatusText;
+
+	/** Sorted stat list for identifying expensive items */
+	TSharedPtr<SListView<FBPStatWidgetPtr>> BlueprintStatList;
+
+	/** Source item array for the treeview */
+	TArray<FBPStatWidgetPtr> RootTreeItems;
+
+};
diff --git a/Engine/Source/Editor/Kismet/Private/Profiler/SProfilerOverview.cpp b/Engine/Source/Editor/Kismet/Private/Profiler/SProfilerOverview.cpp
new file mode 100644
--- /dev/null
+++ b/Engine/Source/Editor/Kismet/Private/Profiler/SProfilerOverview.cpp
@@ -0,0 +1,105 @@
+// Copyright 1998-2016 Epic Games, Inc. All Rights Reserved.
+
+#include "BlueprintEditorPrivatePCH.h"
+#include "SProfilerOverview.h"
+#include "BPProfilerStatisticWidgets.h"
+
+//////////////////////////////////////////////////////////////////////////
+// SProfilerOverview
+
+SProfilerOverview::~SProfilerOverview()
+{
+}
+
+void SProfilerOverview::Construct(const FArguments& InArgs)
+{	
+	BlueprintEditor = InArgs._AssetEditor;
+
+	ChildSlot
+	[
+		SNew(SVerticalBox)
+		+SVerticalBox::Slot()
+		.Padding(FMargin(0, 0, 0, 2))
+		.AutoHeight()
+		.HAlign(HAlign_Fill)
+		.VAlign(VAlign_Top)
+		[
+			SNew(SBorder)
+			.Padding(4)
+			.BorderImage(FEditorStyle::GetBrush("BlueprintProfiler.ViewToolBar"))
+			[
+				SNew(SVerticalBox)
+				+SVerticalBox::Slot()
+				.HAlign(HAlign_Right)
+				.AutoHeight()
+				[
+					SNew(SHorizontalBox)
+				]
+			]
+		]
+		+SVerticalBox::Slot()
+		[
+			SNew(SBorder)
+			.Padding(0)
+			.BorderImage(FEditorStyle::GetBrush("NoBorder"))
+			[
+				SAssignNew(BlueprintStatList, SListView<FBPStatWidgetPtr>)
+				.ListItemsSource(&RootTreeItems)
+				.SelectionMode(ESelectionMode::Single)
+				.OnGenerateRow(this, &SProfilerOverview::OnGenerateRow)
+	//			.OnSelectionChanged(this, &SMessagingTypes::HandleTypeListSelectionChanged)
+				.HeaderRow
+				(
+					SNew(SHeaderRow)
+					+SHeaderRow::Column(SProfilerStatRow::GetStatName(EBlueprintProfilerStat::Name))
+					.DefaultLabel(SProfilerStatRow::GetStatText(EBlueprintProfilerStat::Name))
+					.ManualWidth(450)
+					+SHeaderRow::Column(SProfilerStatRow::GetStatName(EBlueprintProfilerStat::Time))
+					.DefaultLabel(SProfilerStatRow::GetStatText(EBlueprintProfilerStat::Time))
+					.ManualWidth(70)
+					+SHeaderRow::Column(SProfilerStatRow::GetStatName(EBlueprintProfilerStat::PureTime))
+					.DefaultLabel(SProfilerStatRow::GetStatText(EBlueprintProfilerStat::PureTime))
+					.ManualWidth(100)
+					+SHeaderRow::Column(SProfilerStatRow::GetStatName(EBlueprintProfilerStat::InclusiveTime))
+					.DefaultLabel(SProfilerStatRow::GetStatText(EBlueprintProfilerStat::InclusiveTime))
+					.ManualWidth(120)
+					+SHeaderRow::Column(SProfilerStatRow::GetStatName(EBlueprintProfilerStat::MaxTime))
+					.DefaultLabel(SProfilerStatRow::GetStatText(EBlueprintProfilerStat::MaxTime))
+					.ManualWidth(90)
+					+SHeaderRow::Column(SProfilerStatRow::GetStatName(EBlueprintProfilerStat::MinTime))
+					.DefaultLabel(SProfilerStatRow::GetStatText(EBlueprintProfilerStat::MinTime))
+					.ManualWidth(90)
+					+SHeaderRow::Column(SProfilerStatRow::GetStatName(EBlueprintProfilerStat::TotalTime))
+					.DefaultLabel(SProfilerStatRow::GetStatText(EBlueprintProfilerStat::TotalTime))
+					.ManualWidth(80)
+					+SHeaderRow::Column(SProfilerStatRow::GetStatName(EBlueprintProfilerStat::Samples))
+					.DefaultLabel(SProfilerStatRow::GetStatText(EBlueprintProfilerStat::Samples))
+					.ManualWidth(60)
+				)
+			]
+		]
+	];
+}
+
+TSharedRef<ITableRow> SProfilerOverview::OnGenerateRow(FBPStatWidgetPtr InItem, const TSharedRef<STableViewBase>& OwnerTable)
+{
+	return SNew(SProfilerStatRow, OwnerTable, InItem);
+}
+
+void SProfilerOverview::OnDoubleClickStatistic(FBPStatWidgetPtr Item)
+{
+	if (Item.IsValid())
+	{
+//		Item->ExpandWidgetState(ExecutionStatTree, !Item->GetExpansionState());
+	}
+}
+
+void SProfilerOverview::Tick(const FGeometry& AllottedGeometry, const double InCurrentTime, const float InDeltaTime)
+{
+	if (GetDefault<UEditorExperimentalSettings>()->bBlueprintPerformanceAnalysisTools)
+	{
+		if (IBlueprintProfilerInterface* ProfilerInterface = FModuleManager::GetModulePtr<IBlueprintProfilerInterface>("BlueprintProfiler"))
+		{
+		}
+	}
+}
diff --git a/Engine/Source/Editor/Kismet/Private/Profiler/SProfilerOverview.h b/Engine/Source/Editor/Kismet/Private/Profiler/SProfilerOverview.h
new file mode 100644
--- /dev/null
+++ b/Engine/Source/Editor/Kismet/Private/Profiler/SProfilerOverview.h
@@ -0,0 +1,58 @@
+// Copyright 1998-2016 Epic Games, Inc. All Rights Reserved.
+
+#pragma once
+
+#include "Developer/BlueprintProfiler/Public/BlueprintProfilerModule.h"
+
+// Shared pointer to a performance stat
+typedef TSharedPtr<class FBPProfilerStatWidget> FBPStatWidgetPtr;
+
+//////////////////////////////////////////////////////////////////////////
+// SProfilerOverview
+
+class SProfilerOverview : public SCompoundWidget
+{
+public:
+	SLATE_BEGIN_ARGS(SProfilerOverview)
+		{}
+
+		SLATE_ARGUMENT(TWeakPtr<FBlueprintEditor>, AssetEditor)
+	SLATE_END_ARGS()
+
+
+public:
+
+	~SProfilerOverview();
+
+	void Construct(const FArguments& InArgs);
+
+	// SWidget interface
+	virtual void Tick(const FGeometry& AllottedGeometry, const double InCurrentTime, const float InDeltaTime) override;
+	// End of SWidget interface
+
+protected:
+
+	/** Called to generate a stat row for the execution stat tree */
+	TSharedRef<ITableRow> OnGenerateRow(FBPStatWidgetPtr InItem, const TSharedRef<STableViewBase>& OwnerTable);
+
+	/** On double click statistic */
+	void OnDoubleClickStatistic(FBPStatWidgetPtr Item);
+
+protected:
+
+	/** Blueprint editor */
+	TWeakPtr<class FBlueprintEditor> BlueprintEditor;
+
+	/** Profiler status text attribute */
+	TAttribute<FText> ProfilerStatusText;
+
+	/** Sorted stat list for identifying expensive items */
+	TSharedPtr<SListView<FBPStatWidgetPtr>> BlueprintStatList;
+
+	/** Source item array for the treeview */
+	TArray<FBPStatWidgetPtr> RootTreeItems;
+
+	/** Current blueprint to display stats for */
+	TWeakObjectPtr<UBlueprint> CurrentBlueprint;
+
+};
diff --git a/Engine/Source/Editor/Kismet/Private/Profiler/ScriptPerfData.cpp b/Engine/Source/Editor/Kismet/Private/Profiler/ScriptPerfData.cpp
new file mode 100644
--- /dev/null
+++ b/Engine/Source/Editor/Kismet/Private/Profiler/ScriptPerfData.cpp
@@ -0,0 +1,168 @@
+// Copyright 1998-2016 Epic Games, Inc. All Rights Reserved.
+
+#include "BlueprintEditorPrivatePCH.h"
+#include "ScriptPerfData.h"
+
+//////////////////////////////////////////////////////////////////////////
+// FScriptPerfData
+
+FNumberFormattingOptions FScriptPerfData::StatNumberFormat;
+float FScriptPerfData::RecentSampleBias = 0.2f;
+float FScriptPerfData::HistoricalSampleBias = 0.8f;
+float FScriptPerfData::NodePerformanceThreshold = 1.f;
+float FScriptPerfData::PureNodePerformanceThreshold = 1.f;
+float FScriptPerfData::InclusivePerformanceThreshold = 1.5f;
+float FScriptPerfData::MaxPerformanceThreshold = 2.f;
+
+void FScriptPerfData::AddEventTiming(const double PureNodeTimingIn, const double NodeTimingIn)
+{
+	const double TimingInMs = NodeTimingIn * 1000;
+	const double PureNodeTimingInMs = PureNodeTimingIn * 1000;
+	NodeTiming = (TimingInMs * RecentSampleBias) + (NodeTiming * HistoricalSampleBias);
+	PureNodeTiming = (PureNodeTimingInMs * RecentSampleBias) + (PureNodeTiming * HistoricalSampleBias);
+	InclusiveTiming = NodeTiming + PureNodeTiming;
+	TotalTiming += NodeTimingIn + PureNodeTimingIn;
+	NumSamples++;
+	MaxTiming = FMath::Max<double>(MaxTiming, InclusiveTiming);
+	MinTiming = FMath::Min<double>(MinTiming, InclusiveTiming);
+}
+
+void FScriptPerfData::AddData(const FScriptPerfData& DataIn)
+{
+	if (DataIn.IsDataValid())
+	{
+		// Accumulate data, find min, max values
+		NodeTiming += DataIn.NodeTiming;
+		PureNodeTiming += DataIn.PureNodeTiming;
+		InclusiveTiming = NodeTiming + PureNodeTiming;
+		TotalTiming += DataIn.TotalTiming;
+		NumSamples += DataIn.NumSamples;
+		MaxTiming = FMath::Max<double>(MaxTiming, DataIn.MaxTiming);
+		MinTiming = FMath::Min<double>(MinTiming, DataIn.MinTiming);
+	}
+}
+
+void FScriptPerfData::AddBranchData(const FScriptPerfData& DataIn)
+{
+	if (DataIn.IsDataValid())
+	{
+		NumSamples += DataIn.NumSamples;
+	}
+}
+
+void FScriptPerfData::Reset()
+{
+	NodeTiming = 0.0;
+	PureNodeTiming = 0.0;
+	InclusiveTiming = 0.0;
+	TotalTiming = 0.0;
+	NumSamples = 0;
+	MaxTiming = -MAX_dbl;
+	MinTiming = MAX_dbl;
+}
+
+void FScriptPerfData::SetRecentSampleBias(const float RecentSampleBiasIn)
+{
+	RecentSampleBias = RecentSampleBiasIn;
+	HistoricalSampleBias = 1.f - RecentSampleBias;
+}
+
+void FScriptPerfData::SetNodePerformanceThreshold(const float NodePerformanceThresholdIn)
+{
+	NodePerformanceThreshold = NodePerformanceThresholdIn;
+}
+
+FSlateColor FScriptPerfData::GetNodeHeatColor() const
+{
+	const float Value = 1.f - FMath::Min<float>(NodeTiming / NodePerformanceThreshold, 1.f);
+	return FLinearColor(1.f, Value, Value);
+}
+
+void FScriptPerfData::SetPureNodePerformanceThreshold(const float PureNodePerformanceThresholdIn)
+{
+	PureNodePerformanceThreshold = PureNodePerformanceThresholdIn;
+}
+
+FSlateColor FScriptPerfData::GetPureNodeHeatColor() const
+{
+	const float Value = 1.f - FMath::Min<float>(PureNodeTiming / PureNodePerformanceThreshold, 1.f);
+	return FLinearColor(1.f, Value, Value);
+}
+
+void FScriptPerfData::SetInclusivePerformanceThreshold(const float InclusivePerformanceThresholdIn)
+{
+	InclusivePerformanceThreshold = InclusivePerformanceThresholdIn;
+}
+
+FSlateColor FScriptPerfData::GetInclusiveHeatColor() const
+{
+	const float Value = 1.f - FMath::Min<float>(InclusiveTiming / InclusivePerformanceThreshold, 1.f);
+	return FLinearColor(1.f, Value, Value);
+}
+
+void FScriptPerfData::SetMaxPerformanceThreshold(const float MaxPerformanceThresholdIn)
+{
+	MaxPerformanceThreshold = MaxPerformanceThresholdIn;
+}
+
+FSlateColor FScriptPerfData::GetMaxTimeHeatColor() const
+{
+	const float Value = 1.f - FMath::Min<float>(MaxTiming / MaxPerformanceThreshold, 1.f);
+	return FLinearColor(1.f, Value, Value);
+}
+
+FText FScriptPerfData::GetNodeTimingText() const
+{
+	return FText::AsNumber(GetNodeTiming(), &StatNumberFormat);
+}
+
+FText FScriptPerfData::GetPureNodeTimingText() const
+{
+	const double CurrentPureTiming = GetPureTiming();
+	if (CurrentPureTiming == 0.0)
+	{
+		return FText::GetEmpty();
+	}
+	else
+	{
+		return FText::AsNumber(CurrentPureTiming, &StatNumberFormat);
+	}
+}
+
+FText FScriptPerfData::GetInclusiveTimingText() const
+{
+	const double CurrentPureTiming = GetPureTiming();
+	if (CurrentPureTiming == 0.0)
+	{
+		return FText::GetEmpty();
+	}
+	else
+	{
+		return FText::AsNumber(GetInclusiveTiming(), &StatNumberFormat);
+	}
+}
+
+FText FScriptPerfData::GetMaxTimingText() const
+{
+	const double MaximumTime = GetMaxTiming();
+	const bool bMaxTimeInitialised = MaximumTime != -MAX_dbl;
+	return FText::AsNumber(bMaxTimeInitialised ? MaximumTime : 0.0, &StatNumberFormat);
+}
+
+FText FScriptPerfData::GetMinTimingText() const
+{
+	const double MinimumTime = GetMinTiming();
+	const bool bMinTimeInitialised = MinimumTime != MAX_dbl;
+	return FText::AsNumber(bMinTimeInitialised ? MinimumTime : 0.0, &StatNumberFormat);
+}
+
+FText FScriptPerfData::GetTotalTimingText() const
+{
+	// Covert back up to seconds for total time.
+	return FText::AsNumber(GetTotalTiming(), &StatNumberFormat);
+}
+
+FText FScriptPerfData::GetSamplesText() const
+{
+	return FText::AsNumber(GetSampleCount());
+}
\ No newline at end of file
diff --git a/Engine/Source/Editor/Kismet/Private/SBlueprintPalette.cpp b/Engine/Source/Editor/Kismet/Private/SBlueprintPalette.cpp
--- a/Engine/Source/Editor/Kismet/Private/SBlueprintPalette.cpp
+++ b/Engine/Source/Editor/Kismet/Private/SBlueprintPalette.cpp
@@ -147,7 +147,7 @@ static FString GetVarTooltip(UBlueprint* InBlueprint, UClass* VarClass, FName Va
  * @param  ColorOut		An output color, further denoting the specified action.
  * @param  ToolTipOut	The tooltip to display when the icon is hovered over (describing the sub-graph type).
  */
-static void GetSubGraphIcon(FEdGraphSchemaAction_K2Graph const* const ActionIn, UBlueprint const* BlueprintIn, FSlateBrush const*& IconOut, FSlateColor& ColorOut, FString& ToolTipOut)
+static void GetSubGraphIcon(FEdGraphSchemaAction_K2Graph const* const ActionIn, UBlueprint const* BlueprintIn, FSlateBrush const*& IconOut, FSlateColor& ColorOut, FText& ToolTipOut)
 {
 	check(BlueprintIn != NULL);
 
@@ -164,7 +164,7 @@ static void GetSubGraphIcon(FEdGraphSchemaAction_K2Graph const* const ActionIn,
 				IconOut = FEditorStyle::GetBrush(TEXT("GraphEditor.EventGraph_16x"));
 			}
 
-			ToolTipOut = FString::Printf( *LOCTEXT("EventGraph_ToolTip", "Event Graph").ToString() );
+			ToolTipOut = LOCTEXT("EventGraph_ToolTip", "Event Graph");
 		}
 		break;
 	case EEdGraphSchemaAction_K2Graph::Subgraph:
@@ -172,12 +172,12 @@ static void GetSubGraphIcon(FEdGraphSchemaAction_K2Graph const* const ActionIn,
 			if ( ActionIn->EdGraph != NULL && ActionIn->EdGraph->IsA(UAnimationStateMachineGraph::StaticClass()) )
 			{
 				IconOut = FEditorStyle::GetBrush(TEXT("GraphEditor.StateMachine_16x") );
-				ToolTipOut = FString::Printf( *LOCTEXT("AnimationStateMachineGraph_ToolTip", "Animation State Machine").ToString() );
+				ToolTipOut = LOCTEXT("AnimationStateMachineGraph_ToolTip", "Animation State Machine");
 			}
 			else if ( ActionIn->EdGraph != NULL && ActionIn->EdGraph->IsA(UAnimationStateGraph::StaticClass()) )
 			{
 				IconOut = FEditorStyle::GetBrush(TEXT("GraphEditor.State_16x") );
-				ToolTipOut = FString::Printf( *LOCTEXT("AnimationState_ToolTip", "Animation State").ToString() );
+				ToolTipOut = LOCTEXT("AnimationState_ToolTip", "Animation State");
 			}
 			else if ( ActionIn->EdGraph != NULL && ActionIn->EdGraph->IsA(UAnimationTransitionGraph::StaticClass()) )
 			{
@@ -185,18 +185,18 @@ static void GetSubGraphIcon(FEdGraphSchemaAction_K2Graph const* const ActionIn,
 				if ( EdGraphOuter != NULL && EdGraphOuter->IsA(UAnimStateConduitNode::StaticClass()) )
 				{
 					IconOut = FEditorStyle::GetBrush(TEXT("GraphEditor.Conduit_16x"));
-					ToolTipOut = FString::Printf( *LOCTEXT("ConduitGraph_ToolTip", "Conduit").ToString() );
+					ToolTipOut = LOCTEXT("ConduitGraph_ToolTip", "Conduit");
 				}
 				else
 				{
 					IconOut = FEditorStyle::GetBrush(TEXT("GraphEditor.Rule_16x"));
-					ToolTipOut = FString::Printf( *LOCTEXT("AnimationTransitionGraph_ToolTip", "Animation Transition Rule").ToString() );
+					ToolTipOut = LOCTEXT("AnimationTransitionGraph_ToolTip", "Animation Transition Rule");
 				}
 			}
 			else
 			{
 				IconOut = FEditorStyle::GetBrush(TEXT("GraphEditor.SubGraph_16x") );
-				ToolTipOut = FString::Printf( *LOCTEXT("EventSubgraph_ToolTip", "Event Subgraph").ToString() );
+				ToolTipOut = LOCTEXT("EventSubgraph_ToolTip", "Event Subgraph");
 			}
 		}
 		break;
@@ -205,7 +205,7 @@ static void GetSubGraphIcon(FEdGraphSchemaAction_K2Graph const* const ActionIn,
 			IconOut = FEditorStyle::GetBrush(TEXT("GraphEditor.Macro_16x"));
 			if ( ActionIn->EdGraph == NULL )
 			{
-				ToolTipOut = FString::Printf( *LOCTEXT("PotentialOverride_Tooltip", "Potential Override").ToString() );	
+				ToolTipOut = LOCTEXT("PotentialOverride_Tooltip", "Potential Override");	
 			}
 			else
 			{
@@ -213,11 +213,11 @@ static void GetSubGraphIcon(FEdGraphSchemaAction_K2Graph const* const ActionIn,
 				UFunction* OverrideFunc = FindField<UFunction>(BlueprintIn->ParentClass, ActionIn->FuncName);
 				if ( OverrideFunc == NULL )
 				{
-					ToolTipOut = FString::Printf( *LOCTEXT("Macro_Tooltip", "Macro").ToString() );
+					ToolTipOut = LOCTEXT("Macro_Tooltip", "Macro");
 				}
 				else 
 				{
-					ToolTipOut = FString::Printf( *LOCTEXT("Override_Tooltip", "Override").ToString() );
+					ToolTipOut = LOCTEXT("Override_Tooltip", "Override");
 				}
 			}
 		}
@@ -227,17 +227,16 @@ static void GetSubGraphIcon(FEdGraphSchemaAction_K2Graph const* const ActionIn,
 			IconOut = FEditorStyle::GetBrush(TEXT("GraphEditor.InterfaceFunction_16x"));
 			FFormatNamedArguments Args;
 			Args.Add(TEXT("InterfaceName"), FText::FromName(ActionIn->FuncName));
-			FText TooltipText = FText::Format(LOCTEXT("FunctionFromInterface_Tooltip", "Function (from Interface '{InterfaceName}')"), Args);
+			ToolTipOut = FText::Format(LOCTEXT("FunctionFromInterface_Tooltip", "Function (from Interface '{InterfaceName}')"), Args);
 			if (UFunction* OverrideFunc = FindField<UFunction>(BlueprintIn->SkeletonGeneratedClass, ActionIn->FuncName))
 			{
 				if (UEdGraphSchema_K2::FunctionCanBePlacedAsEvent(OverrideFunc))
 				{
-					Args.Add(TEXT("BaseTooltip"), TooltipText);
-					TooltipText = FText::Format(LOCTEXT("InterfaceFunctionExpectedAsEvent_Tooltip", "{BaseTooltip}\nInterface '{InterfaceName}' is already implemented as a function graph but is expected as an event. Remove the function graph and reimplement as an event."), Args);
+					Args.Add(TEXT("BaseTooltip"), ToolTipOut);
+					ToolTipOut = FText::Format(LOCTEXT("InterfaceFunctionExpectedAsEvent_Tooltip", "{BaseTooltip}\nInterface '{InterfaceName}' is already implemented as a function graph but is expected as an event. Remove the function graph and reimplement as an event."), Args);
 					ColorOut = FLinearColor::Yellow;
 				}
 			}
-			ToolTipOut = TooltipText.ToString();
 		}
 		break;
 	case EEdGraphSchemaAction_K2Graph::Function:
@@ -245,7 +244,7 @@ static void GetSubGraphIcon(FEdGraphSchemaAction_K2Graph const* const ActionIn,
 			if ( ActionIn->EdGraph == NULL )
 			{
 				IconOut = FEditorStyle::GetBrush(TEXT("GraphEditor.PotentialOverrideFunction_16x"));
-				ToolTipOut = FString::Printf( *LOCTEXT("PotentialOverride_Tooltip", "Potential Override").ToString() );	
+				ToolTipOut = LOCTEXT("PotentialOverride_Tooltip", "Potential Override");	
 			}
 			else
 			{
@@ -261,17 +260,17 @@ static void GetSubGraphIcon(FEdGraphSchemaAction_K2Graph const* const ActionIn,
 						IconOut = FEditorStyle::GetBrush(TEXT("GraphEditor.Function_16x"));
 						if ( ActionIn->EdGraph != NULL && ActionIn->EdGraph->IsA(UAnimationGraph::StaticClass()) )
 						{
-							ToolTipOut = FString::Printf( *LOCTEXT("AnimationGraph_Tooltip", "Animation Graph").ToString() );
+							ToolTipOut = LOCTEXT("AnimationGraph_Tooltip", "Animation Graph");
 						}
 						else
 						{
-							ToolTipOut = FString::Printf( *LOCTEXT("Function_Tooltip", "Function").ToString() );
+							ToolTipOut = LOCTEXT("Function_Tooltip", "Function");
 						}
 					}
 					else
 					{
 						IconOut = FEditorStyle::GetBrush(TEXT("GraphEditor.OverrideFunction_16x"));
-						ToolTipOut = FString::Printf( *LOCTEXT("Override_Tooltip", "Override").ToString() );
+						ToolTipOut = LOCTEXT("Override_Tooltip", "Override");
 					}
 				}
 			}
@@ -291,10 +290,10 @@ static void GetSubGraphIcon(FEdGraphSchemaAction_K2Graph const* const ActionIn,
  * @param  ColorOut		An output color, further denoting the specified action.
  * @param  ToolTipOut	An output tooltip, best describing the specified action type.
  */
-static void GetPaletteItemIcon(TSharedPtr<FEdGraphSchemaAction> ActionIn, UBlueprint const* BlueprintIn, FSlateBrush const*& BrushOut, FSlateColor& ColorOut, FString& ToolTipOut, FString& DocLinkOut, FString& DocExcerptOut)
+static void GetPaletteItemIcon(TSharedPtr<FEdGraphSchemaAction> ActionIn, UBlueprint const* BlueprintIn, FSlateBrush const*& BrushOut, FSlateColor& ColorOut, FText& ToolTipOut, FString& DocLinkOut, FString& DocExcerptOut)
 {
 	// Default to tooltip based on action supplied
-	ToolTipOut = (ActionIn->GetTooltipDescription().Len() > 0) ? ActionIn->GetTooltipDescription() : ActionIn->GetMenuDescription().ToString();
+	ToolTipOut = (ActionIn->GetTooltipDescription().Len() > 0) ? FText::FromString(ActionIn->GetTooltipDescription()) : ActionIn->GetMenuDescription();
 
 	if (ActionIn->GetTypeId() == FBlueprintActionMenuItem::StaticGetTypeId())
 	{
@@ -310,11 +309,15 @@ static void GetPaletteItemIcon(TSharedPtr<FEdGraphSchemaAction> ActionIn, UBluep
 	else if (UK2Node const* const NodeTemplate = FBlueprintActionMenuUtils::ExtractNodeTemplateFromAction(ActionIn))
 	{
 		// If the node wants to create tooltip text, use that instead, because its probably more detailed
-		FString NodeToolTipText = NodeTemplate->GetTooltipText().ToString();
-		if (NodeToolTipText.Len() > 0)
+		FText NodeToolTipText = NodeTemplate->GetTooltipText();
+		if (!NodeToolTipText.IsEmpty())
 		{
 			ToolTipOut = NodeToolTipText;
 		}
+		else
+		{
+			ToolTipOut = ToolTipOut;
+		}
 
 		// Ask node for a palette icon
 		FLinearColor IconLinearColor = FLinearColor::White;
@@ -332,7 +335,9 @@ static void GetPaletteItemIcon(TSharedPtr<FEdGraphSchemaAction> ActionIn, UBluep
 		FEdGraphSchemaAction_K2Delegate* DelegateAction = (FEdGraphSchemaAction_K2Delegate*)ActionIn.Get();
 
 		BrushOut = FEditorStyle::GetBrush(TEXT("GraphEditor.Delegate_16x"));
-		ToolTipOut = FString::Printf(*LOCTEXT("Delegate_Tooltip", "Event Dispatcher '%s'").ToString(), *DelegateAction->GetDelegateName().ToString());
+		FFormatNamedArguments Args;
+		Args.Add(TEXT("EventDispatcherName"), FText::FromName(DelegateAction->GetDelegateName()));
+		ToolTipOut = FText::Format(LOCTEXT("Delegate_Tooltip", "Event Dispatcher '{EventDispatcherName}'"), Args);
 	}
 	else if (ActionIn->GetTypeId() == FEdGraphSchemaAction_K2Var::StaticGetTypeId())
 	{
@@ -340,7 +345,7 @@ static void GetPaletteItemIcon(TSharedPtr<FEdGraphSchemaAction> ActionIn, UBluep
 
 		UClass* VarClass = VarAction->GetVariableClass();
 		BrushOut = FBlueprintEditor::GetVarIconAndColor(VarClass, VarAction->GetVariableName(), ColorOut);
-		ToolTipOut = GetVarType(VarClass, VarAction->GetVariableName(), true, true);
+		ToolTipOut = FText::FromString(GetVarType(VarClass, VarAction->GetVariableName(), true, true));
 
 		DocLinkOut = TEXT("Shared/Editor/Blueprint/VariableTypes");
 		DocExcerptOut = GetVarType(VarClass, VarAction->GetVariableName(), false, false);
@@ -351,7 +356,7 @@ static void GetPaletteItemIcon(TSharedPtr<FEdGraphSchemaAction> ActionIn, UBluep
 
 		UStruct* VarScope = LocalVarAction->GetVariableScope();
 		BrushOut = FBlueprintEditor::GetVarIconAndColor(VarScope, LocalVarAction->GetVariableName(), ColorOut);
-		ToolTipOut = GetVarType(VarScope, LocalVarAction->GetVariableName(), true);
+		ToolTipOut = FText::FromString(GetVarType(VarScope, LocalVarAction->GetVariableName(), true));
 
 		DocLinkOut = TEXT("Shared/Editor/Blueprint/VariableTypes");
 		DocExcerptOut = GetVarType(VarScope, LocalVarAction->GetVariableName(), false);
@@ -359,12 +364,12 @@ static void GetPaletteItemIcon(TSharedPtr<FEdGraphSchemaAction> ActionIn, UBluep
 	else if (ActionIn->GetTypeId() == FEdGraphSchemaAction_K2Enum::StaticGetTypeId())
 	{
 		BrushOut = FEditorStyle::GetBrush(TEXT("GraphEditor.EnumGlyph"));
-		ToolTipOut = FString::Printf(*LOCTEXT("Enum_Tooltip", "Enum Asset").ToString());
+		ToolTipOut = LOCTEXT("Enum_Tooltip", "Enum Asset");
 	}
 	else if (ActionIn->GetTypeId() == FEdGraphSchemaAction_K2Struct::StaticGetTypeId())
 	{
 		BrushOut = FEditorStyle::GetBrush(TEXT("GraphEditor.StructGlyph"));
-		ToolTipOut = FString::Printf(*LOCTEXT("Struct_Tooltip", "Struct Asset").ToString());
+		ToolTipOut = LOCTEXT("Struct_Tooltip", "Struct Asset");
 	}
 }
 
@@ -990,10 +995,10 @@ void SBlueprintPaletteItem::Construct(const FArguments& InArgs, FCreateWidgetFor
 	// construct the icon widget
 	FSlateBrush const* IconBrush   = FEditorStyle::GetBrush(TEXT("NoBrush"));
 	FSlateColor        IconColor   = FSlateColor::UseForeground();
-	FString            IconToolTip = GraphAction->GetTooltipDescription();
+	FText			   IconToolTip = FText::FromString(GraphAction->GetTooltipDescription());
 	FString			   IconDocLink, IconDocExcerpt;
 	GetPaletteItemIcon(GraphAction, Blueprint, IconBrush, IconColor, IconToolTip, IconDocLink, IconDocExcerpt);
-	TSharedRef<SWidget> IconWidget = CreateIconWidget(FText::FromString(IconToolTip), IconBrush, IconColor, IconDocLink, IconDocExcerpt);
+	TSharedRef<SWidget> IconWidget = CreateIconWidget(IconToolTip, IconBrush, IconColor, IconDocLink, IconDocExcerpt);
 	IconWidget->SetEnabled(!bIsFullyReadOnly);
 
 	// Setup a meta tag for this node
