Copying //UE4/Dev-Blueprints to Dev-Main (//UE4/Dev-Main)

diff --git a/Engine/Source/Editor/Kismet/Public/Profiler/EventExecution.h b/Engine/Source/Editor/Kismet/Public/Profiler/EventExecution.h
new file mode 100644
--- /dev/null
+++ b/Engine/Source/Editor/Kismet/Public/Profiler/EventExecution.h
@@ -0,0 +1,308 @@
+// Copyright 1998-2016 Epic Games, Inc. All Rights Reserved.
+
+#pragma once
+
+#include "TracePath.h"
+
+struct FSlateBrush;
+class UBlueprintGeneratedClass;
+
+/**  Execution node flags */
+namespace EScriptExecutionNodeFlags
+{
+	enum Type
+	{
+		None				= 0x00000000,	// No Flags
+		Class				= 0x00000001,	// Class
+		Instance			= 0x00000002,	// Instance
+		Event				= 0x00000004,	// Event
+		FunctionCall		= 0x00000008,	// Function Call
+		MacroCall			= 0x00000010,	// Macro Call
+		CallSite			= 0x00000018,	// Function / Macro call site
+		MacroNode			= 0x00000100,	// Macro Node
+		ConditionalBranch	= 0x00000200,	// Node has multiple exit pins using a jump
+		SequentialBranch	= 0x00000400,	// Node has multiple exit pins ran in sequence
+		BranchNode			= 0x00000600,	// BranchNode
+		Node				= 0x00000800,	// Node timing
+		ExecPin				= 0x00001000	// Exec pin dummy node
+	};
+}
+
+/** Stat Container Type */
+namespace EScriptStatContainerType
+{
+	enum Type
+	{
+		Standard = 0,
+		Container,
+		SequentialBranch,
+		NewExecutionPath
+	};
+}
+
+//////////////////////////////////////////////////////////////////////////
+// FScriptNodeExecLinkage
+
+class KISMET_API FScriptNodeExecLinkage
+{
+public:
+
+	struct FLinearExecPath
+	{
+		FLinearExecPath(TSharedPtr<class FScriptExecutionNode> NodeIn, const FTracePath TracePathIn)
+			: TracePath(TracePathIn)
+			, LinkedNode(NodeIn)
+		{
+		}
+
+		FTracePath TracePath;
+		TSharedPtr<class FScriptExecutionNode> LinkedNode;
+	};
+
+	/** Returns the parent node */
+	TSharedPtr<class FScriptExecutionNode> GetParentNode() { return ParentNode; }
+
+	/** Sets the parent node */
+	void SetParentNode(TSharedPtr<class FScriptExecutionNode> InParentNode) { ParentNode = InParentNode; }
+
+	/** Returns the linked nodes map */
+	TMap<int32, TSharedPtr<class FScriptExecutionNode>>& GetLinkedNodes() { return LinkedNodes; }
+
+	/** Returns the number of linked nodes */
+	int32 GetNumLinkedNodes() const { return LinkedNodes.Num(); }
+
+	/** Returns the linked node for the specified pin script offset */
+	TSharedPtr<class FScriptExecutionNode> GetLinkedNodeByPinScriptOffset(const int32 PinScriptOffset) { return LinkedNodes[PinScriptOffset]; }
+
+	/** Add linked node */
+	void AddLinkedNode(const int32 PinScriptOffset, TSharedPtr<class FScriptExecutionNode> LinkedNode);
+
+	/** Returns the child nodes map */
+	TArray<TSharedPtr<class FScriptExecutionNode>>& GetChildNodes() { return ChildNodes; }
+
+	/** Returns the number of children */
+	int32 GetNumChildren() const { return ChildNodes.Num(); }
+
+	/** Returns the child node for the specified index */
+	TSharedPtr<class FScriptExecutionNode> GetChildByIndex(const int32 ChildIndex) { return ChildNodes[ChildIndex]; }
+
+	/** Add child node */
+	void AddChildNode(TSharedPtr<class FScriptExecutionNode> ChildNode) { ChildNodes.Add(ChildNode); }
+
+protected:
+
+	/** Parent node */
+	TSharedPtr<class FScriptExecutionNode> ParentNode;
+	/** Linked nodes */
+	TMap<int32, TSharedPtr<class FScriptExecutionNode>> LinkedNodes;
+	/** Child nodes */
+	TArray<TSharedPtr<class FScriptExecutionNode>> ChildNodes;
+
+};
+
+//////////////////////////////////////////////////////////////////////////
+// FScriptNodePerfData
+
+class KISMET_API FScriptNodePerfData
+{
+public:
+
+	/** Get perf data for instance and tracepath */
+	TSharedPtr<class FScriptPerfData> GetPerfDataByInstanceAndTracePath(FName InstanceName, const FTracePath& TracePath);
+
+	/** Get all instance perf data for the trace path, excluding the global blueprint data */
+	void GetInstancePerfDataByTracePath(const FTracePath& TracePath, TArray<TSharedPtr<FScriptPerfData>>& ResultsOut);
+
+	/** Get global blueprint perf data for the trace path */
+	TSharedPtr<FScriptPerfData> GetBlueprintPerfDataByTracePath(const FTracePath& TracePath);
+
+protected:
+
+	/** FScriptExeutionPath hash to perf data */
+	TMap<FName, TMap<const uint32, TSharedPtr<FScriptPerfData>>> InstanceInputPinToPerfDataMap;
+
+};
+
+//////////////////////////////////////////////////////////////////////////
+// FScriptExecutionNode
+
+class KISMET_API FScriptExecutionNode : public FScriptNodeExecLinkage, public FScriptNodePerfData, public TSharedFromThis<FScriptExecutionNode>
+{
+public:
+
+	FScriptExecutionNode();
+    virtual ~FScriptExecutionNode() {}
+
+	bool operator == (const FScriptExecutionNode& NodeIn) const;
+
+	/** Get the node's name */
+	FName GetName() const { return NodeName; }
+
+	/** Set the node's name */
+	void SetName(FName NewName) { NodeName = NewName; }
+
+	/** Set the node's name by string */
+	void SetNameByString(const FString& NewName) { NodeName = FName(*NewName); }
+
+	/** Returns the owning graph name */
+	FName GetGraphName() const { return OwningGraphName; }
+
+	/** Sets the owning graph name */
+	void SetGraphName(FName GraphNameIn) { OwningGraphName = GraphNameIn; }
+
+	/** Set the node's flags */
+	void SetFlags(const uint32 NewFlags) { NodeFlags = NewFlags; }
+
+	/** Add to the node's flags */
+	void AddFlags(const uint32 NewFlags) { NodeFlags |= NewFlags; }
+
+	/** Remove from the node's flags */
+	void RemoveFlags(const uint32 NewFlags) { NodeFlags &= ~NewFlags; }
+
+	/** Does the node contain these flags */
+	bool HasFlags(const uint32 Flags) const { return (NodeFlags & Flags) != 0U; }
+
+	/** Returns if this exec event represents a change in class/blueprint */
+	bool IsClass() const { return (NodeFlags & EScriptExecutionNodeFlags::Class) != 0U; }
+
+	/** Returns if this exec event represents a change in instance */
+	bool IsInstance() const { return (NodeFlags & EScriptExecutionNodeFlags::Instance) != 0U; }
+
+	/** Returns if this exec event represents the start of an event execution path */
+	bool IsEvent() const { return (NodeFlags & EScriptExecutionNodeFlags::Event) != 0U; }
+
+	/** Returns if this event is a function callsite event */
+	bool IsFunctionCallSite() const { return (NodeFlags & EScriptExecutionNodeFlags::FunctionCall) != 0U; }
+
+	/** Returns if this event is a macro callsite event */
+	bool IsMacroCallSite() const { return (NodeFlags & EScriptExecutionNodeFlags::MacroCall) != 0U; }
+
+	/** Returns if this event happened inside a macro instance */
+	bool IsMacroNode() const { return (NodeFlags & EScriptExecutionNodeFlags::MacroNode) != 0U; }
+
+	/** Returns if this event potentially multiple exit sites */
+	bool IsBranch() const { return (NodeFlags & EScriptExecutionNodeFlags::BranchNode) != 0U; }
+
+	/** Returns the display name for widget UI */
+	const FText& GetDisplayName() const { return DisplayName; }
+
+	/** Sets the display name for widget UI */
+	void SetDisplayName(const FText InDisplayName) { DisplayName = InDisplayName; }
+
+	/** Returns the tooltip for widget UI */
+	const FText& GetToolTipText() const { return Tooltip; }
+
+	/** Sets the tooltip for widget UI */
+	void SetToolTipText(const FText InTooltip) { Tooltip = InTooltip; }
+
+	/** Returns the icon color for widget UI */
+	const FLinearColor& GetIconColor() const { return IconColor; }
+
+	/** Sets the icon color for widget UI */
+	void SetIconColor(FLinearColor InIconColor) { IconColor = InIconColor; }
+
+	/** Returns the icon for widget UI */
+	const FSlateBrush* GetIcon() const { return Icon; }
+
+	/** Sets the icon for widget UI */
+	void SetIcon(FSlateBrush* InIcon) { Icon = InIcon; }
+
+	/** Returns the profiler heat color */
+	FLinearColor GetHeatColor(const FTracePath& TracePath) const;
+
+	/** Returns the current expansion state for widget UI */
+	bool IsExpanded() const { return bExpansionState; }
+
+	/** Sets the current expansion state for widget UI */
+	void SetExpanded(bool bIsExpanded) { bExpansionState = bIsExpanded; }
+
+	/** Returns if this node represents a blueprint */
+	virtual bool IsRootNode() const { return false; }
+
+	/** Return the linear execution path from this node */
+	void GetLinearExecutionPath(TArray<FLinearExecPath>& LinearExecutionNodes, const FTracePath& TracePath);
+
+	/** Get statistic container type */
+	virtual EScriptStatContainerType::Type GetStatisticContainerType() const;
+
+	/** Refresh Stats */
+	virtual void RefreshStats(const FTracePath& TracePath);
+
+	/** Get all exec nodes */
+	virtual void GetAllExecNodes(TMap<FName, TSharedPtr<FScriptExecutionNode>>& ExecNodesOut);
+
+protected:
+
+	/** Node name */
+	FName NodeName;
+	/** Owning graph name */
+	FName OwningGraphName;
+	/** Node flags to describe the source graph node type */
+	uint32 NodeFlags;
+	/** Display name for widget UI */
+	FText DisplayName;
+	/** Tooltip for widget UI */
+	FText Tooltip;
+	/** Icon color for widget UI */
+	FLinearColor IconColor;
+	/** Icon for widget UI */
+	FSlateBrush* Icon;
+	/** Expansion state */
+	bool bExpansionState;
+
+};
+
+//////////////////////////////////////////////////////////////////////////
+// FScriptExecutionBlueprint
+
+class KISMET_API FScriptExecutionBlueprint : public FScriptExecutionNode
+{
+public:
+
+	/** Adds new blueprint instance node */
+	void AddInstance(TSharedPtr<class FScriptExecutionNode> Instance) { Instances.Add(Instance); }
+
+	/** Returns the current number of instance nodes */
+	int32 GetInstanceCount() { return Instances.Num(); }
+
+	/** Returns the instance node specified by Index */
+	TSharedPtr<class FScriptExecutionNode> GetInstanceByIndex(const int32 Index) { return Instances[Index]; }
+
+	/** Returns the instance that matches the supplied name if present */
+	TSharedPtr<FScriptExecutionNode> GetInstanceByName(FName InstanceName);
+
+	// ~Begin FScriptExecutionNode
+	virtual bool IsRootNode() const override { return true; }
+	virtual EScriptStatContainerType::Type GetStatisticContainerType() const override { return EScriptStatContainerType::Container; }
+	virtual void RefreshStats(const FTracePath& TracePath) override;
+	virtual void GetAllExecNodes(TMap<FName, TSharedPtr<FScriptExecutionNode>>& ExecNodesOut) override;
+	// ~End FScriptExecutionNode
+
+private:
+
+	/** Exec nodes representing all instances based on this blueprint */
+	TArray<TSharedPtr<class FScriptExecutionNode>> Instances;
+
+};
+
+//////////////////////////////////////////////////////////////////////////
+// FScriptExecutionContext
+
+class KISMET_API FScriptExecutionContext : public TSharedFromThis<FScriptExecutionContext>
+{
+public:
+
+	bool IsContextValid() const { return ContextObject.IsValid(); }
+	const class UEdGraphNode* GetGraphNode() const { return ContextObject.IsValid() ? Cast<UEdGraphNode>(ContextObject.Get()) : nullptr; }
+	const class UEdGraphNode* GetNodeFromCodeLocation(const int32 CodeLocation, UFunction* FunctionOverride = nullptr) const;
+	void GetAllExecNodes(TMap<FName, TSharedPtr<FScriptExecutionNode>>& ExecNodesOut);
+	void UpdateConnectedStats();
+
+public:
+
+	TSharedPtr<FScriptExecutionNode> ExecutionNode;
+	TWeakObjectPtr<const UObject> ContextObject;
+	TWeakObjectPtr<const UObject> MacroContextObject;
+	TWeakObjectPtr<UBlueprintGeneratedClass> BlueprintClass;
+	TWeakObjectPtr<UFunction> BlueprintFunction;
+};
diff --git a/Engine/Source/Editor/Kismet/Public/Profiler/ScriptPerfData.h b/Engine/Source/Editor/Kismet/Public/Profiler/ScriptPerfData.h
new file mode 100644
--- /dev/null
+++ b/Engine/Source/Editor/Kismet/Public/Profiler/ScriptPerfData.h
@@ -0,0 +1,105 @@
+// Copyright 1998-2016 Epic Games, Inc. All Rights Reserved.
+
+#pragma once
+
+#include "Styling/SlateColor.h"
+
+//////////////////////////////////////////////////////////////////////////
+// FScriptPerfData
+
+class KISMET_API FScriptPerfData : public TSharedFromThis<FScriptPerfData>
+{
+public:
+
+	FScriptPerfData()
+		: InclusiveTiming(0.0)
+		, NodeTiming(0.0)
+		, PureNodeTiming(0.0)
+		, MaxTiming(-MAX_dbl)
+		, MinTiming(MAX_dbl)
+		, TotalTiming(0.0)
+		, NumSamples(0)
+	{
+	}
+
+	/** Add a single event node timing into the dataset */
+	void AddEventTiming(const double PureNodeTimingIn, const double NodeTimingIn);
+
+	/** Add data */
+	void AddData(const FScriptPerfData& DataIn);
+
+	/** Add branch data */
+	void AddBranchData(const FScriptPerfData& DataIn);
+
+	/** Reset the current data buffer and all derived data stats */
+	void Reset();
+
+	// Returns if this data structure has any valid data
+	bool IsDataValid() const { return NumSamples > 0; }
+
+	// Updates the various thresholds that control stats calcs and display.
+	static void SetNumberFormattingForStats(const FNumberFormattingOptions& FormatIn) { StatNumberFormat = FormatIn; }
+	static void SetRecentSampleBias(const float RecentSampleBiasIn);
+	static void SetNodePerformanceThreshold(const float NodePerformanceThresholdIn);
+	static void SetPureNodePerformanceThreshold(const float PureNodePerformanceThresholdIn);
+	static void SetInclusivePerformanceThreshold(const float InclusivePerformanceThresholdIn);
+	static void SetMaxPerformanceThreshold(const float MaxPerformanceThresholdIn);
+
+	// Returns the various stats this container holds
+	double GetNodeTiming() const { return NodeTiming; }
+	double GetPureTiming() const { return PureNodeTiming; }
+	float GetInclusiveTiming() const { return InclusiveTiming; }
+	double GetMaxTiming() const { return MaxTiming; }
+	double GetMinTiming() const { return MinTiming; }
+	double GetTotalTiming() const { return TotalTiming; }
+	int32 GetSampleCount() const { return NumSamples; }
+
+	// Returns various performance colors for visual display
+	FSlateColor GetNodeHeatColor() const;
+	FSlateColor GetPureNodeHeatColor() const;
+	FSlateColor GetInclusiveHeatColor() const;
+	FSlateColor GetMaxTimeHeatColor() const;
+
+	// Returns the various stats this container holds in FText format
+	FText GetTotalTimingText() const;
+	FText GetInclusiveTimingText() const;
+	FText GetNodeTimingText() const;
+	FText GetPureNodeTimingText() const;
+	FText GetMaxTimingText() const;
+	FText GetMinTimingText() const;
+	FText GetSamplesText() const;
+
+private:
+
+	/** Inclusive timing, pure and node timings */
+	double InclusiveTiming;
+	/** Node timing */
+	double NodeTiming;
+	/** Pure node timing */
+	double PureNodeTiming;
+	/** Max exclusive timing */
+	double MaxTiming;
+	/** Min exclusive timing */
+	double MinTiming;
+	/** Total time accrued */
+	double TotalTiming;
+	/** The current inclusive sample count */
+	int32 NumSamples;
+
+	/** Number formatting options */
+	static FNumberFormattingOptions StatNumberFormat;
+
+	/** Controls the bias between new and older samples */
+	static float RecentSampleBias;
+	/** Cached Historical Sample Bias */
+	static float HistoricalSampleBias;
+	/** Cached node performance threshold */
+	static float NodePerformanceThreshold;
+	/** Cached pure node performance threshold */
+	static float PureNodePerformanceThreshold;
+	/** Cached inclusive performance threshold */
+	static float InclusivePerformanceThreshold;
+	/** Cached max performance threshold */
+	static float MaxPerformanceThreshold;
+
+};
diff --git a/Engine/Source/Editor/Kismet/Public/Profiler/TracePath.h b/Engine/Source/Editor/Kismet/Public/Profiler/TracePath.h
new file mode 100644
--- /dev/null
+++ b/Engine/Source/Editor/Kismet/Public/Profiler/TracePath.h
@@ -0,0 +1,42 @@
+// Copyright 1998-2016 Epic Games, Inc. All Rights Reserved.
+
+#pragma once
+
+//////////////////////////////////////////////////////////////////////////
+// FTracePath
+
+class FTracePath
+{
+public:
+
+	FTracePath()
+		: CachedCrc(0U)
+	{
+	}
+
+	FTracePath(const FTracePath& TracePathIn)
+		: CachedCrc(TracePathIn.CachedCrc)
+	{
+	}
+
+	/** Cast to uint32 operator */
+	operator const uint32() const { return CachedCrc; }
+
+	/** Add exit pin */
+	void AddExitPin(const int32 ExitPinScriptOffset)
+	{
+		CachedCrc = FCrc::MemCrc32(&ExitPinScriptOffset, 1, CachedCrc);
+	}
+
+	/** Reset */
+	void Reset()
+	{
+		CachedCrc = 0;
+	}
+
+private:
+
+	/** Cached crc */
+	uint32 CachedCrc;
+
+};
diff --git a/Engine/Source/Editor/KismetCompiler/Private/KismetCompiler.cpp b/Engine/Source/Editor/KismetCompiler/Private/KismetCompiler.cpp
--- a/Engine/Source/Editor/KismetCompiler/Private/KismetCompiler.cpp
+++ b/Engine/Source/Editor/KismetCompiler/Private/KismetCompiler.cpp
@@ -517,6 +517,23 @@ void FKismetCompilerContext::CreateClassVariablesFromBlueprint()
 
 	// Grab the blueprint variables
 	NewClass->NumReplicatedProperties = 0;	// Keep track of how many replicated variables this blueprint adds
+	// Clear out any existing property guids
+	const bool bRebuildPropertyMap = bIsFullCompile && !Blueprint->bIsRegeneratingOnLoad;
+	if (bRebuildPropertyMap)
+	{
+		NewClass->PropertyGuids.Reset();
+		// Add any chained parent blueprint map values
+		UBlueprint* ParentBP = Cast<UBlueprint>(Blueprint->ParentClass->ClassGeneratedBy);
+		while (ParentBP)
+		{
+			if (UBlueprintGeneratedClass* ParentBPGC = Cast<UBlueprintGeneratedClass>(ParentBP->GeneratedClass))
+			{
+				NewClass->PropertyGuids.Append(ParentBPGC->PropertyGuids);
+			}
+			ParentBP = Cast<UBlueprint>(ParentBP->ParentClass->ClassGeneratedBy);
+		}
+	}
+
 	for (int32 i = 0; i < Blueprint->NewVariables.Num(); ++i)
 	{
 		FBPVariableDescription& Variable = Blueprint->NewVariables[Blueprint->NewVariables.Num() - (i + 1)];
@@ -552,6 +569,11 @@ void FKismetCompilerContext::CreateClassVariablesFromBlueprint()
 					}
 				}
 			}
+			if (bRebuildPropertyMap)
+			{
+				// Update new class property guid map
+				NewClass->PropertyGuids.Add(Variable.VarName, Variable.VarGuid);
+			}
 		}
 	}
 
@@ -1471,9 +1493,9 @@ void FKismetCompilerContext::CompileFunction(FKismetFunctionContext& Context)
 		SortKeyMap.Add(Node, i);
 
 		const FString NodeComment = Node->NodeComment.IsEmpty() ? Node->GetName() : Node->NodeComment;
-
+		const bool bPureNode = IsNodePure(Node);
 		// Debug comments
-		if (KismetCompilerDebugOptions::EmitNodeComments)
+		if (KismetCompilerDebugOptions::EmitNodeComments && !Context.bGeneratingCpp)
 		{
 			FBlueprintCompiledStatement& Statement = Context.AppendStatementForNode(Node);
 			Statement.Type = KCST_Comment;
@@ -1481,23 +1503,41 @@ void FKismetCompilerContext::CompileFunction(FKismetFunctionContext& Context)
 		}
 
 		// Debug opcode insertion point
-		if (Context.IsDebuggingOrInstrumentationRequired() && !IsNodePure(Node))
+		if (Context.IsDebuggingOrInstrumentationRequired())
 		{
-			bool bEmitDebuggingSite = true;
-
-			if (Context.IsEventGraph() && (Node->IsA(UK2Node_FunctionEntry::StaticClass())))
+			if (!bPureNode)
 			{
-				// The entry point in the ubergraph is a non-visual construct, and will lead to some
-				// other 'fake' entry point such as an event or latent action.  Therefore, don't create
-				// debug data for the behind-the-scenes entry point, only for the user-visible ones.
-				bEmitDebuggingSite = false;
-			}
+				bool bEmitDebuggingSite = true;
 
-			if (bEmitDebuggingSite)
+				if (Context.IsEventGraph() && (Node->IsA(UK2Node_FunctionEntry::StaticClass())))
+				{
+					// The entry point in the ubergraph is a non-visual construct, and will lead to some
+					// other 'fake' entry point such as an event or latent action.  Therefore, don't create
+					// debug data for the behind-the-scenes entry point, only for the user-visible ones.
+					bEmitDebuggingSite = false;
+				}
+				if (bEmitDebuggingSite)
+				{
+					FBlueprintCompiledStatement& Statement = Context.AppendStatementForNode(Node);
+					Statement.Type = Context.GetBreakpointType();
+					if (Context.IsInstrumentationRequired())
+					{
+						for (auto Pin : Node->Pins)
+						{
+							if (Pin->Direction == EGPD_Input && Pin->PinType.PinCategory == UEdGraphSchema_K2::PC_Exec)
+							{
+								Statement.ExecContext = Pin;
+								break;
+							}
+						}
+					}
+					Statement.Comment = NodeComment;
+				}
+			}
+			else if (Context.IsInstrumentationRequired())
 			{
 				FBlueprintCompiledStatement& Statement = Context.AppendStatementForNode(Node);
-				Statement.Type = Context.GetBreakpointType();
-				Statement.Comment = NodeComment;
+				Statement.Type = KCST_InstrumentedPureNodeEntry;
 			}
 		}
 
@@ -1583,6 +1623,12 @@ void FKismetCompilerContext::CompileFunction(FKismetFunctionContext& Context)
 			++TestIndex;
 		}
 	}
+
+	if (Context.bIsUbergraph && CompileOptions.DoesRequireCppCodeGeneration())
+	{
+		Context.UnsortedSeparateExecutionGroups = FKismetCompilerUtilities::FindUnsortedSeparateExecutionGroups(Context.LinearExecutionList);
+	}
+
 }
 
 /**
diff --git a/Engine/Source/Editor/KismetCompiler/Private/KismetCompilerMisc.cpp b/Engine/Source/Editor/KismetCompiler/Private/KismetCompilerMisc.cpp
--- a/Engine/Source/Editor/KismetCompiler/Private/KismetCompilerMisc.cpp
+++ b/Engine/Source/Editor/KismetCompiler/Private/KismetCompilerMisc.cpp
@@ -14,6 +14,7 @@
 #include "Editor/UnrealEd/Public/ObjectTools.h"
 #include "DefaultValueHelper.h"
 #include "Engine/UserDefinedStruct.h"
+#include "BlueprintEditorSettings.h"
 
 #define LOCTEXT_NAMESPACE "KismetCompiler"
 
@@ -1234,10 +1235,12 @@ void FNodeHandlingFunctor::ResolveAndRegisterScopedTerm(FKismetFunctionContext&
 	UProperty* BoundProperty = FKismetCompilerUtilities::FindPropertyInScope(SearchScope, Net, CompilerContext.MessageLog, CompilerContext.GetSchema(), Context.NewClass);
 	if (BoundProperty != NULL)
 	{
+		UBlueprintEditorSettings* Settings = GetMutableDefault<UBlueprintEditorSettings>();
 		// Create the term in the list
 		FBPTerminal* Term = new (NetArray) FBPTerminal();
 		Term->CopyFromPin(Net, Net->PinName);
 		Term->AssociatedVarProperty = BoundProperty;
+		Term->bPassedByReference = Settings->bAllowReferencePassThroughsToCacheNonReferences;
 		Context.NetMap.Add(Net, Term);
 
 		// Check if the property is a local variable and mark it so
@@ -1470,6 +1473,19 @@ int32 FKismetFunctionContext::GetContextUniqueID()
 	return UUIDCounter++;
 }
 
+bool FKismetFunctionContext::DoesStatementRequiresSwitch(const FBlueprintCompiledStatement* Statement)
+{
+	return Statement && (
+		Statement->Type == KCST_UnconditionalGoto ||
+		Statement->Type == KCST_PushState ||
+		Statement->Type == KCST_GotoIfNot ||
+		Statement->Type == KCST_ComputedGoto ||
+		Statement->Type == KCST_EndOfThread ||
+		Statement->Type == KCST_EndOfThreadIfNot ||
+		Statement->Type == KCST_GotoReturn ||
+		Statement->Type == KCST_GotoReturnIfNot);
+}
+
 bool FKismetFunctionContext::MustUseSwitchState(const FBlueprintCompiledStatement* ExcludeThisOne) const
 {
 	for (auto Node : LinearExecutionList)
@@ -1479,15 +1495,7 @@ bool FKismetFunctionContext::MustUseSwitchState(const FBlueprintCompiledStatemen
 		{
 			for (auto Statement : (*StatementList))
 			{
-				if (Statement && (Statement != ExcludeThisOne) && (
-					Statement->Type == KCST_UnconditionalGoto ||
-					Statement->Type == KCST_PushState ||
-					Statement->Type == KCST_GotoIfNot ||
-					Statement->Type == KCST_ComputedGoto ||
-					Statement->Type == KCST_EndOfThread ||
-					Statement->Type == KCST_EndOfThreadIfNot ||
-					Statement->Type == KCST_GotoReturn ||
-					Statement->Type == KCST_GotoReturnIfNot))
+				if (Statement && (Statement != ExcludeThisOne) && DoesStatementRequiresSwitch(Statement))
 				{
 					return true;
 				}
@@ -1706,6 +1714,14 @@ void FKismetFunctionContext::FinalSortLinearExecList()
 	LinearExecutionList = SortedLinearExecutionList;
 }
 
+bool FKismetFunctionContext::DoesStatementRequiresFlowStack(const FBlueprintCompiledStatement* Statement)
+{
+	return Statement && (
+		(Statement->Type == KCST_EndOfThreadIfNot) ||
+		(Statement->Type == KCST_EndOfThread) ||
+		(Statement->Type == KCST_PushState));
+}
+
 void FKismetFunctionContext::ResolveStatements()
 {
 	BP_SCOPED_COMPILER_EVENT_STAT(EKismetCompilerStats_ResolveCompiledStatements);
@@ -1714,16 +1730,7 @@ void FKismetFunctionContext::ResolveStatements()
 	static const FBoolConfigValueHelper OptimizeExecutionFlowStack(TEXT("Kismet"), TEXT("bOptimizeExecutionFlowStack"), GEngineIni);
 	if (OptimizeExecutionFlowStack)
 	{
-		const bool bFlowStackIsRequired = AllGeneratedStatements.ContainsByPredicate(
-			[](const FBlueprintCompiledStatement* Statement)
-		{
-			return Statement && (
-				(Statement->Type == KCST_EndOfThreadIfNot) ||
-				(Statement->Type == KCST_EndOfThread) ||
-				(Statement->Type == KCST_PushState));
-		});
-
-		bUseFlowStack = bFlowStackIsRequired;
+		bUseFlowStack = AllGeneratedStatements.ContainsByPredicate(&FKismetFunctionContext::DoesStatementRequiresFlowStack);
 	}
 
 	ResolveGotoFixups();
@@ -1890,14 +1897,7 @@ FBPTerminal* FKismetFunctionContext::CreateLocalTerminalFromPinAutoChooseScope(U
 	static FBoolConfigValueHelper UseLocalGraphVariables(TEXT("Kismet"), TEXT("bUseLocalGraphVariables"), GEngineIni);
 	static FBoolConfigValueHelper UseLocalGraphVariablesInCpp(TEXT("BlueprintNativizationSettings"), TEXT("bUseLocalEventGraphVariables"));
 
-
-	auto PinRepresentsDelegate = [](const FEdGraphPinType& PinType) -> bool
-	{
-		return (PinType.PinCategory == UEdGraphSchema_K2::PC_Delegate)
-			|| (PinType.PinCategory == UEdGraphSchema_K2::PC_MCDelegate);
-	};
-	const bool bUseLocalGraphVariables = UseLocalGraphVariables 
-		|| (bGeneratingCpp && (UseLocalGraphVariablesInCpp || PinRepresentsDelegate(Net->PinType)));
+	const bool bUseLocalGraphVariables = UseLocalGraphVariables || (bGeneratingCpp && UseLocalGraphVariablesInCpp);
 
 	const bool OutputPin = EEdGraphPinDirection::EGPD_Output == Net->Direction;
 	if (bSharedTerm && bUseLocalGraphVariables && OutputPin)
@@ -1930,3 +1930,79 @@ void FBPTerminal::CopyFromPin(UEdGraphPin* Net, const FString& NewName)
 	const bool bStructSubCategoryObj = (NULL != Cast<UScriptStruct>(Net->PinType.PinSubCategoryObject.Get()));
 	SetContextTypeStruct(bStructCategory && bStructSubCategoryObj);
 }
+
+
+TArray<TSet<UEdGraphNode*>> FKismetCompilerUtilities::FindUnsortedSeparateExecutionGroups(const TArray<UEdGraphNode*>& Nodes)
+{
+	TArray<UEdGraphNode*> UnprocessedNodes;
+	for (UEdGraphNode* Node : Nodes)
+	{
+		UK2Node* K2Node = Cast<UK2Node>(Node);
+		if (K2Node && !K2Node->IsNodePure())
+		{
+			UnprocessedNodes.Add(Node);
+		}
+	}
+
+	TSet<UEdGraphNode*> AlreadyProcessed;
+	TArray<TSet<UEdGraphNode*>> Result;
+	while (UnprocessedNodes.Num())
+	{
+		Result.Emplace(TSet<UEdGraphNode*>());
+		TSet<UEdGraphNode*>& ExecutionGroup = Result.Last();
+		TSet<UEdGraphNode*> ToProcess;
+
+		UEdGraphNode* Seed = UnprocessedNodes.Pop();
+		ensure(!AlreadyProcessed.Contains(Seed));
+		ToProcess.Add(Seed);
+		ExecutionGroup.Add(Seed);
+		while (ToProcess.Num())
+		{
+			UEdGraphNode* Node = *ToProcess.CreateIterator();
+			// for each execution pin
+			for (UEdGraphPin* Pin : Node->Pins)
+			{
+				if (Pin && Pin->LinkedTo.Num() && (Pin->PinType.PinCategory == UEdGraphSchema_K2::PC_Exec))
+				{
+					for (UEdGraphPin* LinkedPin : Pin->LinkedTo)
+					{
+						if (!LinkedPin)
+						{
+							continue;
+						}
+						UEdGraphNode* LinkedNode = LinkedPin->GetOwningNodeUnchecked();
+						const bool bIsAlreadyProcessed = AlreadyProcessed.Contains(LinkedNode);
+						const bool bInCurrentExecutionGroup = ExecutionGroup.Contains(LinkedNode);
+						ensure(!bIsAlreadyProcessed || bInCurrentExecutionGroup);
+						ensure(bInCurrentExecutionGroup || UnprocessedNodes.Contains(LinkedNode));
+						if (!bIsAlreadyProcessed)
+						{
+							ToProcess.Add(LinkedNode);
+							ExecutionGroup.Add(LinkedNode);
+							UnprocessedNodes.Remove(LinkedNode);
+						}
+					}
+				}
+			}
+
+			const int32 WasRemovedFromProcess = ToProcess.Remove(Node);
+			ensure(0 != WasRemovedFromProcess);
+			bool bAlreadyProcessed = false;
+			AlreadyProcessed.Add(Node, &bAlreadyProcessed);
+			ensure(!bAlreadyProcessed);
+		}
+
+		if (1 == ExecutionGroup.Num())
+		{
+			UEdGraphNode* TheOnlyNode = *ExecutionGroup.CreateIterator();
+			if (!TheOnlyNode || TheOnlyNode->IsA<UK2Node_FunctionEntry>() || TheOnlyNode->IsA<UK2Node_Timeline>())
+			{
+				Result.Pop();
+			}
+		}
+	}
+
+
+
+	return Result;
+}
\ No newline at end of file
diff --git a/Engine/Source/Editor/KismetCompiler/Private/KismetCompilerVMBackend.cpp b/Engine/Source/Editor/KismetCompiler/Private/KismetCompilerVMBackend.cpp
--- a/Engine/Source/Editor/KismetCompiler/Private/KismetCompilerVMBackend.cpp
+++ b/Engine/Source/Editor/KismetCompiler/Private/KismetCompilerVMBackend.cpp
@@ -1322,8 +1322,9 @@ public:
 		int32 EventType = 0;
 		switch (Statement.Type)
 		{
-			case KCST_InstrumentedWireExit:		EventType = EScriptInstrumentation::NodeExit; break;
-			case KCST_InstrumentedWireEntry:	EventType = EScriptInstrumentation::NodeEntry; break;
+			case KCST_InstrumentedWireExit:			EventType = EScriptInstrumentation::NodeExit; break;
+			case KCST_InstrumentedWireEntry:		EventType = EScriptInstrumentation::NodeEntry; break;
+			case KCST_InstrumentedPureNodeEntry:	EventType = EScriptInstrumentation::PureNodeEntry; break;
 		}
 		Writer << EX_InstrumentationEvent;
 		Writer << EventType;
@@ -1484,6 +1485,7 @@ public:
 			EmitSwitchValue(Statement);
 			break;
 		case KCST_InstrumentedWireExit:
+		case KCST_InstrumentedWireEntry:
 			{
 				UEdGraphPin const* TrueSourcePin = Cast<UEdGraphPin const>(FunctionContext.MessageLog.FindSourceObject(Statement.ExecContext));
 				if (TrueSourcePin)
@@ -1491,10 +1493,6 @@ public:
 					int32 Offset = Writer.ScriptBuffer.Num() + sizeof(int32);
 					ClassBeingBuilt->GetDebugData().RegisterPinToCodeAssociation(TrueSourcePin, FunctionContext.Function, Offset);
 				}
-			}
-			// no break, continue down.
-		case KCST_InstrumentedWireEntry:
-			{
 				if (SourceNode != NULL)
 				{
 					// Record where this NOP is
@@ -1514,6 +1512,9 @@ public:
 						ClassBeingBuilt->GetDebugData().RegisterNodeToCodeAssociation(TrueSourceNode, MacroSourceNode, MacroInstanceNodes, FunctionContext.Function, Offset, false);
 					}
 				}
+			}
+		case KCST_InstrumentedPureNodeEntry:
+			{
 				// Emit Statement
 				EmitInstrumentation(Statement);
 				break;
diff --git a/Engine/Source/Editor/KismetCompiler/Public/BlueprintCompiledStatement.h b/Engine/Source/Editor/KismetCompiler/Public/BlueprintCompiledStatement.h
--- a/Engine/Source/Editor/KismetCompiler/Public/BlueprintCompiledStatement.h
+++ b/Engine/Source/Editor/KismetCompiler/Public/BlueprintCompiledStatement.h
@@ -38,8 +38,9 @@ enum EKismetCompiledStatementType
 	KCST_GotoReturnIfNot = 28, // [if (!TargetObject->TargetProperty)] goto TargetLabel
 	KCST_SwitchValue = 29,
 	// Kismet instrumentation extensions
-	KCST_InstrumentedWireEntry,	// Instrumented wiretrace entry
-	KCST_InstrumentedWireExit	// Instrumented wiretrace exit
+	KCST_InstrumentedPureNodeEntry,	// Instrumented pure node entry
+	KCST_InstrumentedWireEntry,		// Instrumented wiretrace entry
+	KCST_InstrumentedWireExit		// Instrumented wiretrace exit
 };
 
 //@TODO: Too rigid / icky design
diff --git a/Engine/Source/Editor/KismetCompiler/Public/KismetCompiledFunctionContext.h b/Engine/Source/Editor/KismetCompiler/Public/KismetCompiledFunctionContext.h
--- a/Engine/Source/Editor/KismetCompiler/Public/KismetCompiledFunctionContext.h
+++ b/Engine/Source/Editor/KismetCompiler/Public/KismetCompiledFunctionContext.h
@@ -60,6 +60,9 @@ public:
 	// Goto fixup requests (each statement (key) wants to goto the first statement attached to the exec out-pin (value))
 	TMap< FBlueprintCompiledStatement*, UEdGraphPin* > GotoFixupRequestMap;
 
+	// Used to split uber graph into subfunctions by C++ backend
+	TArray<TSet<UEdGraphNode*>> UnsortedSeparateExecutionGroups;
+
 	// Map from a net to an term (either a literal or a storage location)
 	TIndirectArray<FBPTerminal> Parameters;
 	TIndirectArray<FBPTerminal> Results;
@@ -91,9 +94,6 @@ public:
 	struct FNetNameMapping* NetNameMap;
 	bool bAllocatedNetNameMap;
 
-	// Stored calls of latent function (on current class), needed to tell if blueprint should be tickable
-	TArray< UK2Node_CallFunction* > LatentFunctionCalls;
-
 	//Skip some optimization. C++ code will be generated in this pass. 
 	bool bGeneratingCpp;
 
@@ -178,6 +178,11 @@ public:
 		return bInstrumentScriptCode || bCreateDebugData;
 	}
 
+	bool IsInstrumentationRequired() const
+	{
+		return bInstrumentScriptCode;
+	}
+
 	EKismetCompiledStatementType GetWireTraceType() const
 	{
 		return bInstrumentScriptCode ? KCST_InstrumentedWireExit : KCST_WireTraceSite;
@@ -282,7 +287,10 @@ public:
 		}
 		else
 		{
-			MessageLog.Warning(TEXT("A node that generated no code of it's own (@@) tried to inject code into @@"), Source, Destination);
+			MessageLog.Warning(
+				*FString::Printf(TEXT("A node that generated no code of it's own (@@) tried to inject code into @@. %s"), *GetFullNameSafe(Source))
+				, Source
+				, Destination);
 		}
 	}
 
@@ -306,7 +314,8 @@ private:
 	void ResolveGotoFixups();
 
 public:
-
+	static bool DoesStatementRequiresSwitch(const FBlueprintCompiledStatement* Statement);
+	static bool DoesStatementRequiresFlowStack(const FBlueprintCompiledStatement* Statement);
 	// The function links gotos, sorts statments, and merges adjacent ones. 
 	void ResolveStatements();
 
diff --git a/Engine/Source/Editor/KismetCompiler/Public/KismetCompilerMisc.h b/Engine/Source/Editor/KismetCompiler/Public/KismetCompilerMisc.h
--- a/Engine/Source/Editor/KismetCompiler/Public/KismetCompilerMisc.h
+++ b/Engine/Source/Editor/KismetCompiler/Public/KismetCompilerMisc.h
@@ -70,6 +70,9 @@ public:
 	 */
 	static UObject* CheckPropertyNameOnScope(UStruct* Scope, const FName& PropertyName);
 
+	// Find groups of nodes, that can be executed separately.
+	static TArray<TSet<UEdGraphNode*>> FindUnsortedSeparateExecutionGroups(const TArray<UEdGraphNode*>& Nodes);
+
 private:
 	/** Counter to ensure unique names in the transient package, to avoid GC collection issues with classes and their CDOs */
 	static uint32 ConsignToOblivionCounter;
diff --git a/Engine/Source/Editor/UnrealEd/Classes/Settings/EditorExperimentalSettings.h b/Engine/Source/Editor/UnrealEd/Classes/Settings/EditorExperimentalSettings.h
--- a/Engine/Source/Editor/UnrealEd/Classes/Settings/EditorExperimentalSettings.h
+++ b/Engine/Source/Editor/UnrealEd/Classes/Settings/EditorExperimentalSettings.h
@@ -65,9 +65,9 @@ public:
 	UPROPERTY(EditAnywhere, config, Category=Blueprints, meta=(DisplayName="Blueprint Performance Analysis Tools"))
 	bool bBlueprintPerformanceAnalysisTools;
 
-	/** The number of samples the blueprint profiler should use to average the current value. */
-	UPROPERTY(EditAnywhere, config, Category=Blueprints, meta=(DisplayName="Blueprint Profiler Sample Count"))
-	int32 BlueprintProfilerAverageSampleCount;
+	/** Bias to weight prominance of newer samples against hositorical samples. */
+	UPROPERTY(EditAnywhere, config, Category=Blueprints, meta=(DisplayName="Blueprint Performance Analysis Sample Bias", ClampMin=0.01, ClampMax=1.0))
+	float BlueprintProfilerRecentSampleBias;
 
 	/** Enables the visual diff tool for widget blueprints. WARNING: changes to the widget hierarchy will not be detected */
 	UPROPERTY(EditAnywhere, config, Category=Blueprints, meta=(DisplayName="Use the Diff Tool for Widget Blueprints"))
diff --git a/Engine/Source/Editor/UnrealEd/Private/Kismet2/EnumEditorUtils.cpp b/Engine/Source/Editor/UnrealEd/Private/Kismet2/EnumEditorUtils.cpp
--- a/Engine/Source/Editor/UnrealEd/Private/Kismet2/EnumEditorUtils.cpp
+++ b/Engine/Source/Editor/UnrealEd/Private/Kismet2/EnumEditorUtils.cpp
@@ -400,10 +400,13 @@ void FEnumEditorUtils::EnsureAllDisplayNamesExist(UUserDefinedEnum* Enum)
 	if (Enum)
 	{
 		const int32 EnumeratorsToEnsure = (Enum->NumEnums() > 0) ? (Enum->NumEnums() - 1) : 0;
-		Enum->DisplayNames.Empty(EnumeratorsToEnsure);
+		if (Enum->DisplayNames.Num() != EnumeratorsToEnsure)
+		{
+			Enum->DisplayNames.Empty(EnumeratorsToEnsure);
+		}
 		for	(int32 Index = 0; Index < EnumeratorsToEnsure; Index++)
 		{
-			FText DisplayNameMetaData = Enum->GetEnumText(Index);
+			FText DisplayNameMetaData = Enum->UEnum::GetEnumText(Index);
 			if (DisplayNameMetaData.IsEmpty())
 			{
 				const FString EnumName = Enum->GetEnumName(Index);
@@ -414,7 +417,14 @@ void FEnumEditorUtils::EnsureAllDisplayNamesExist(UUserDefinedEnum* Enum)
 				}
 				Enum->SetMetaData(FEnumEditorUtilsHelper::DisplayName(), *DisplayNameMetaData.ToString(), Index);
 			}
-			Enum->DisplayNames.Add(DisplayNameMetaData);
+			if (!Enum->DisplayNames.IsValidIndex(Index))
+			{
+				Enum->DisplayNames.Add(DisplayNameMetaData);
+			}
+			else if (!Enum->DisplayNames[Index].EqualTo(DisplayNameMetaData))
+			{
+				Enum->DisplayNames[Index] = DisplayNameMetaData;
+			}
 		}
 	}
 }
diff --git a/Engine/Source/Editor/UnrealEd/Private/Kismet2/KismetDebugUtilities.cpp b/Engine/Source/Editor/UnrealEd/Private/Kismet2/KismetDebugUtilities.cpp
--- a/Engine/Source/Editor/UnrealEd/Private/Kismet2/KismetDebugUtilities.cpp
+++ b/Engine/Source/Editor/UnrealEd/Private/Kismet2/KismetDebugUtilities.cpp
@@ -503,7 +503,7 @@ void FKismetDebugUtilities::CheckBreakConditions(UEdGraphNode* NodeStoppedAt, bo
 				Data.GraphStack.Push(NewGraph);
 			}
 		}
-		else
+		else if (Data.GraphStack.Num())
 		{
 			UEdGraph* CurrGraph = NodeStoppedAt->GetTypedOuter<UEdGraph>();
 			if (Data.GraphStack.Last() != CurrGraph)
diff --git a/Engine/Source/Editor/UnrealEd/Private/Settings/SettingsClasses.cpp b/Engine/Source/Editor/UnrealEd/Private/Settings/SettingsClasses.cpp
--- a/Engine/Source/Editor/UnrealEd/Private/Settings/SettingsClasses.cpp
+++ b/Engine/Source/Editor/UnrealEd/Private/Settings/SettingsClasses.cpp
@@ -15,6 +15,8 @@
 
 #include "SourceCodeNavigation.h"
 
+#include "ScriptPerfData.h"
+
 #define LOCTEXT_NAMESPACE "SettingsClasses"
 
 /* UContentBrowserSettings interface
@@ -67,7 +69,7 @@ UEditorExperimentalSettings::UEditorExperimentalSettings( const FObjectInitializ
 	, bUnifiedBlueprintEditor(true)
 	, bBlueprintableComponents(true)
 	, bBlueprintPerformanceAnalysisTools(false)
-	, BlueprintProfilerAverageSampleCount(20)
+	, BlueprintProfilerRecentSampleBias(0.2f)
 {
 }
 
@@ -90,6 +92,10 @@ void UEditorExperimentalSettings::PostEditChangeProperty( struct FPropertyChange
 			FModuleManager::Get().LoadModule(TEXT("EnvironmentQueryEditor"));
 		}
 	}
+	else if (Name == FName(TEXT("BlueprintProfilerRecentSampleBias")))
+	{
+		FScriptPerfData::SetRecentSampleBias(BlueprintProfilerRecentSampleBias);
+	}
 
 	if (!FUnrealEdMisc::Get().IsDeletePreferences())
 	{
diff --git a/Engine/Source/Programs/UnrealHeaderTool/Private/CodeGenerator.cpp b/Engine/Source/Programs/UnrealHeaderTool/Private/CodeGenerator.cpp
--- a/Engine/Source/Programs/UnrealHeaderTool/Private/CodeGenerator.cpp
+++ b/Engine/Source/Programs/UnrealHeaderTool/Private/CodeGenerator.cpp
@@ -1312,9 +1312,23 @@ void FNativeClassHeaderGenerator::ExportNativeGeneratedInitCode(FClass* Class, F
 	FUHTStringBuilder CallSingletons;
 	FString ApiString = GetAPIString();
 
+	TSet<FName> AlreadyIncludedNames;
 	TArray<UFunction*> FunctionsToExport;
 	for( TFieldIterator<UFunction> Function(Class,EFieldIteratorFlags::ExcludeSuper); Function; ++Function )
 	{
+		UFunction* LocalFunc = *Function;
+		FName TrueName = FNativeClassHeaderGenerator::GetOverriddenFName(LocalFunc);
+		bool bAlreadyIncluded = false;
+		AlreadyIncludedNames.Add(TrueName, &bAlreadyIncluded);
+		if (bAlreadyIncluded)
+		{
+			// In a dynamic class the same function signature may be used for a Multi- and a Single-cast delegate.
+			if (!LocalFunc->IsA<UDelegateFunction>() || !bIsDynamic)
+			{
+				FError::Throwf(TEXT("The same function linked twice. Function: %s Class: %s"), *LocalFunc->GetName(), *Class->GetName());
+			}
+			continue;
+		}
 		FunctionsToExport.Add(*Function);
 	}
 
@@ -5308,8 +5322,6 @@ void FNativeClassHeaderGenerator::ExportGeneratedCPP()
 
 	GeneratedLinkerFixupFunction.Logf(TEXT("void EmptyLinkFunctionForGeneratedCode%s() {}") LINE_TERMINATOR, *ModuleInfo->Name);
 
-	GeneratedCPPText.Log(*GeneratedPackageCPP);
-
 	{
 		// Autogenerate names (alphabetically sorted).
 		ReferencedNames.KeySort( TLess<FName>() );
@@ -5320,6 +5332,8 @@ void FNativeClassHeaderGenerator::ExportGeneratedCPP()
 		}
 	}
 
+	GeneratedCPPText.Log(*GeneratedPackageCPP);
+
 	GeneratedCPPEpilogue.Logf(
 		LINE_TERMINATOR
 		);
diff --git a/Engine/Source/Runtime/AIModule/Private/BehaviorTree/BTFunctionLibrary.cpp b/Engine/Source/Runtime/AIModule/Private/BehaviorTree/BTFunctionLibrary.cpp
--- a/Engine/Source/Runtime/AIModule/Private/BehaviorTree/BTFunctionLibrary.cpp
+++ b/Engine/Source/Runtime/AIModule/Private/BehaviorTree/BTFunctionLibrary.cpp
@@ -33,7 +33,7 @@ namespace
 	FORCEINLINE UBlackboardComponent* GetBlackboard(UBTNode& BTNode)
 	{
 		check(BTNode.GetOuter());
-		check(BTNode.GetClass()->HasAnyClassFlags(CLASS_CompiledFromBlueprint) && "This function call is valid only for BP-implemented BT nodes");
+		check(BTNode.GetClass()->HasAnyClassFlags(CLASS_CompiledFromBlueprint) || Cast<UDynamicClass>(BTNode.GetClass()) && "This function call is valid only for BP-implemented BT nodes");
 		check(Cast<UBehaviorTreeComponent>(BTNode.GetOuter()));
 
 		return ((UBehaviorTreeComponent*)BTNode.GetOuter())->GetBlackboardComponent();
diff --git a/Engine/Source/Runtime/Core/Public/UObject/BlueprintsObjectVersion.h b/Engine/Source/Runtime/Core/Public/UObject/BlueprintsObjectVersion.h
--- a/Engine/Source/Runtime/Core/Public/UObject/BlueprintsObjectVersion.h
+++ b/Engine/Source/Runtime/Core/Public/UObject/BlueprintsObjectVersion.h
@@ -8,6 +8,8 @@ struct CORE_API FBlueprintsObjectVersion
 	{
 		// Before any version changes were made
 		BeforeCustomVersionWasAdded = 0,
+		OverridenEventReferenceFixup,
+		CleanBlueprintFunctionFlags,
 
 
 		// -----<new versions can be added above this line>-------------------------------------------------
diff --git a/Engine/Source/Runtime/Core/Public/UObject/ObjectVersion.h b/Engine/Source/Runtime/Core/Public/UObject/ObjectVersion.h
--- a/Engine/Source/Runtime/Core/Public/UObject/ObjectVersion.h
+++ b/Engine/Source/Runtime/Core/Public/UObject/ObjectVersion.h
@@ -600,6 +600,8 @@ enum EUnrealEngineObjectUE4Version
 	VER_UE4_FIX_SLOT_NAME_DUPLICATION,
 	// Texture streaming using AABBs instead of Spheres
 	VER_UE4_STREAMABLE_TEXTURE_AABB,
+	// FPropertyTag stores GUID of property
+	VER_UE4_PROPERTY_GUID_IN_PROPERTY_TAG,
 	// -----<new versions can be added before this line>-------------------------------------------------
 	// - this needs to be the last line (see note below)
 	VER_UE4_AUTOMATIC_VERSION_PLUS_ONE,
diff --git a/Engine/Source/Runtime/CoreUObject/Private/UObject/Class.cpp b/Engine/Source/Runtime/CoreUObject/Private/UObject/Class.cpp
--- a/Engine/Source/Runtime/CoreUObject/Private/UObject/Class.cpp
+++ b/Engine/Source/Runtime/CoreUObject/Private/UObject/Class.cpp
@@ -870,6 +870,8 @@ void UStruct::SerializeTaggedProperties(FArchive& Ar, uint8* Data, UStruct* Defa
 
 	UClass* DefaultsClass = dynamic_cast<UClass*>(DefaultsStruct);
 	UScriptStruct* DefaultsScriptStruct = dynamic_cast<UScriptStruct*>(DefaultsStruct);
+	// Determine if this struct supports optional property guid's (UBlueprintGeneratedClasses Only)
+	const bool bArePropertyGuidsAvailable = (Ar.UE4Ver() >= VER_UE4_PROPERTY_GUID_IN_PROPERTY_TAG) && !FPlatformProperties::RequiresCookedData() && ArePropertyGuidsAvailable();
 
 	if( Ar.IsLoading() )
 	{
@@ -921,6 +923,16 @@ void UStruct::SerializeTaggedProperties(FArchive& Ar, uint8* Data, UStruct* Defa
 				RemainingArrayDim	= Property ? Property->ArrayDim : 0;
 			}
 			
+			// Optionally resolve properties using Guid Property tags in non cooked builds that support it.
+			if (bArePropertyGuidsAvailable && Tag.HasPropertyGuid)
+			{
+				// Use property guids from blueprint generated classes to redirect serialised data.
+				FName Result = FindPropertyNameFromGuid(Tag.PropertyGuid);
+				if (Result != NAME_None && Tag.Name != Result)
+				{
+					Tag.Name = Result;
+				}
+			}
 			// If this property is not the one we expect (e.g. skipped as it matches the default value), do the brute force search.
 			if( Property == NULL || Property->GetFName() != Tag.Name )
 			{
@@ -1431,6 +1443,12 @@ void UStruct::SerializeTaggedProperties(FArchive& Ar, uint8* Data, UStruct* Defa
 						FArchive::FScopeAddDebugData S(Ar, Property->GetFName());
 #endif
 						FPropertyTag Tag( Ar, Property, Idx, DataPtr, DefaultValue );
+						// If available use the property guid from BlueprintGeneratedClasses, provided we aren't cooking data.
+						if (bArePropertyGuidsAvailable && !Ar.IsCooking())
+						{
+							const FGuid PropertyGuid = FindPropertyGuidFromName(Tag.Name);
+							Tag.SetPropertyGuid(PropertyGuid);
+						}
 						Ar << Tag;
 
 						// need to know how much data this call to SerializeTaggedProperty consumes, so mark where we are
@@ -4794,6 +4812,11 @@ UScriptStruct* TBaseStructure<FGuid>::Get()
 	return ScriptStruct;
 }
 
+UScriptStruct* TBaseStructure<FBox2D>::Get()
+{
+	static auto ScriptStruct = StaticGetBaseStructureInternal(TEXT("Box2D"));
+	return ScriptStruct;
+}
 
 UScriptStruct* TBaseStructure<FFallbackStruct>::Get()
 {
diff --git a/Engine/Source/Runtime/CoreUObject/Private/UObject/PropertyBaseObject.cpp b/Engine/Source/Runtime/CoreUObject/Private/UObject/PropertyBaseObject.cpp
--- a/Engine/Source/Runtime/CoreUObject/Private/UObject/PropertyBaseObject.cpp
+++ b/Engine/Source/Runtime/CoreUObject/Private/UObject/PropertyBaseObject.cpp
@@ -455,10 +455,10 @@ void UObjectPropertyBase::CheckValidObject(void* Value) const
 #if USE_CIRCULAR_DEPENDENCY_LOAD_DEFERRING
 		FLinkerLoad* PropertyLinker = GetLinker();
 		bool const bIsDeferringValueLoad = ((PropertyLinker == nullptr) || (PropertyLinker->LoadFlags & LOAD_DeferDependencyLoads)) &&
-			Object->IsA<ULinkerPlaceholderExportObject>();
+			(Object->IsA<ULinkerPlaceholderExportObject>() || Object->IsA<ULinkerPlaceholderClass>());
 
 #if USE_DEFERRED_DEPENDENCY_CHECK_VERIFICATION_TESTS
-		check(bIsDeferringValueLoad || !Object->IsA<ULinkerPlaceholderExportObject>());
+		check( bIsDeferringValueLoad || (!Object->IsA<ULinkerPlaceholderExportObject>() && !Object->IsA<ULinkerPlaceholderClass>()) );
 #endif // USE_DEFERRED_DEPENDENCY_CHECK_VERIFICATION_TESTS
 
 #else  // USE_CIRCULAR_DEPENDENCY_LOAD_DEFERRING 
diff --git a/Engine/Source/Runtime/CoreUObject/Private/UObject/ScriptCore.cpp b/Engine/Source/Runtime/CoreUObject/Private/UObject/ScriptCore.cpp
--- a/Engine/Source/Runtime/CoreUObject/Private/UObject/ScriptCore.cpp
+++ b/Engine/Source/Runtime/CoreUObject/Private/UObject/ScriptCore.cpp
@@ -228,6 +228,13 @@ FString UnicodeToCPPIdentifier(const FString& InName, bool bDeprecated, const TC
 		Ret.InsertAt(0, TCHAR('_'));
 	}
 	Ret = PrefixStr + Ret + Postfix;
+
+	// Workaround for a strange compiler error
+	if (InName == TEXT("Replicate to server"))
+	{
+		Ret = TEXT("MagicNameWorkaround");
+	}
+
 	return bDeprecated ? Ret + TEXT("_DEPRECATED") : Ret;
 }
 
@@ -1068,7 +1075,7 @@ void UObject::ProcessEvent( UFunction* Function, void* Parms )
 #if !(UE_BUILD_SHIPPING || UE_BUILD_TEST)
 	if (GetClass()->HasInstrumentation())
 	{
-		EScriptInstrumentationEvent EventInstrumentationInfo(EScriptInstrumentation::Event, this);
+		EScriptInstrumentationEvent EventInstrumentationInfo(EScriptInstrumentation::Event, this, Function->GetFName());
 		FBlueprintCoreDelegates::InstrumentScriptEvent(EventInstrumentationInfo);
 	}
 #endif
@@ -1202,7 +1209,7 @@ void UObject::ProcessEvent( UFunction* Function, void* Parms )
 #if !(UE_BUILD_SHIPPING || UE_BUILD_TEST)
 	if (GetClass()->HasInstrumentation())
 	{
-		EScriptInstrumentationEvent EventInstrumentationInfo(EScriptInstrumentation::Stop, this);
+		EScriptInstrumentationEvent EventInstrumentationInfo(EScriptInstrumentation::Stop, this, NAME_None);
 		FBlueprintCoreDelegates::InstrumentScriptEvent(EventInstrumentationInfo);
 	}
 #if WITH_EDITORONLY_DATA
diff --git a/Engine/Source/Runtime/CoreUObject/Public/UObject/AssetPtr.h b/Engine/Source/Runtime/CoreUObject/Public/UObject/AssetPtr.h
--- a/Engine/Source/Runtime/CoreUObject/Public/UObject/AssetPtr.h
+++ b/Engine/Source/Runtime/CoreUObject/Public/UObject/AssetPtr.h
@@ -35,6 +35,15 @@ public:
 		(*this)=Other;
 	}
 
+	/**  
+	 * Construct from a unique object identifier
+	 * @param ObjectID Object identifier to create a weak pointer to
+	**/
+	explicit FORCEINLINE FAssetPtr(const FStringAssetReference &ObjectID)
+		: TPersistentObjectPtr<FStringAssetReference>(ObjectID)
+	{
+	}
+
 	/**  
 	 * Construct from another lazy pointer
 	 * @param Other lazy pointer to copy from
@@ -106,6 +115,15 @@ public:
 	{
 	}
 
+	/**  
+	 * Construct from a unique object identifier
+	 * @param ObjectID Object identifier to create a weak pointer to
+	**/
+	explicit FORCEINLINE TAssetPtr(const FStringAssetReference &ObjectID)
+		: AssetPtr(ObjectID)
+	{
+	}
+
 	/**
 	 * Reset the lazy pointer back to the NULL state
 	 */
@@ -326,6 +344,15 @@ public:
 	{
 	}
 
+	/**  
+	 * Construct from a unique object identifier
+	 * @param ObjectID Object identifier to create a weak pointer to
+	**/
+	explicit FORCEINLINE TAssetSubclassOf(const FStringAssetReference &ObjectID)
+		: AssetPtr(ObjectID)
+	{
+	}
+
 	/**
 	 * Reset the lazy pointer back to the NULL state
 	 */
diff --git a/Engine/Source/Runtime/CoreUObject/Public/UObject/Class.h b/Engine/Source/Runtime/CoreUObject/Public/UObject/Class.h
--- a/Engine/Source/Runtime/CoreUObject/Public/UObject/Class.h
+++ b/Engine/Source/Runtime/CoreUObject/Public/UObject/Class.h
@@ -427,6 +427,18 @@ public:
 	// Required by UHT makefiles for internal data serialization.
 	friend struct FStructArchiveProxy;
 #endif // HACK_HEADER_GENERATOR
+
+protected:
+
+	/** Returns the property name from the guid */
+	virtual FName FindPropertyNameFromGuid(const FGuid& PropertyGuid) const { return NAME_None; }
+
+	/** Find property guid */
+	virtual FGuid FindPropertyGuidFromName(const FName InName) const { return FGuid(); }
+
+	/** Returns if we have access to property guids */
+	virtual bool ArePropertyGuidsAvailable() const { return false; }
+
 };
 
 enum EStructFlags
@@ -1373,7 +1385,7 @@ public:
 		//@TODO: UCREMOVAL: CPF_ConstParm added as a hack to get blueprints compiling with a const DamageType parameter.
 		const uint64 IgnoreFlags = CPF_PersistentInstance | CPF_ExportObject | CPF_InstancedReference 
 			| CPF_ContainsInstancedReference | CPF_ComputedFlags | CPF_ConstParm | CPF_HasGetValueTypeHash | CPF_UObjectWrapper
-			| CPF_NativeAccessSpecifiers;
+			| CPF_NativeAccessSpecifiers | CPF_AdvancedDisplay;
 		return IgnoreFlags;
 	}
 
@@ -2921,6 +2933,11 @@ template<> struct TBaseStructure<FGuid>
 	COREUOBJECT_API static UScriptStruct* Get();
 };
 
+template<> struct TBaseStructure<FBox2D>
+{
+	COREUOBJECT_API static UScriptStruct* Get();
+};	
+
 struct FFallbackStruct;
 template<> struct TBaseStructure<FFallbackStruct>
 {
diff --git a/Engine/Source/Runtime/CoreUObject/Public/UObject/PersistentObjectPtr.h b/Engine/Source/Runtime/CoreUObject/Public/UObject/PersistentObjectPtr.h
--- a/Engine/Source/Runtime/CoreUObject/Public/UObject/PersistentObjectPtr.h
+++ b/Engine/Source/Runtime/CoreUObject/Public/UObject/PersistentObjectPtr.h
@@ -39,6 +39,17 @@ public:
 		(*this)=Other;
 	}
 
+	/**  
+	 * Construct from a unique object identifier
+	 * @param InObjectID a unique object identifier
+	 */
+	explicit FORCEINLINE TPersistentObjectPtr(const TObjectID& InObjectID)
+		: WeakPtr()
+		, TagAtLastTest(0)
+		, ObjectID(InObjectID)
+	{
+	}
+
 	/**  
 	 * Copy from a unique object identifier
 	 * @param ObjectID object identifier to create a weak pointer to
diff --git a/Engine/Source/Runtime/CoreUObject/Public/UObject/PropertyTag.h b/Engine/Source/Runtime/CoreUObject/Public/UObject/PropertyTag.h
--- a/Engine/Source/Runtime/CoreUObject/Public/UObject/PropertyTag.h
+++ b/Engine/Source/Runtime/CoreUObject/Public/UObject/PropertyTag.h
@@ -25,6 +25,8 @@ struct FPropertyTag
 	int32		ArrayIndex;	// Index if an array; else 0.
 	int32		SizeOffset;	// location in stream of tag size member
 	FGuid	StructGuid;
+	uint8	HasPropertyGuid;
+	FGuid	PropertyGuid;
 
 	// Constructors.
 	FPropertyTag()
@@ -37,6 +39,7 @@ struct FPropertyTag
 		, Size      (0)
 		, ArrayIndex(INDEX_NONE)
 		, SizeOffset(INDEX_NONE)
+		, HasPropertyGuid(0)
 	{}
 
 	FPropertyTag( FArchive& InSaveAr, UProperty* Property, int32 InIndex, uint8* Value, uint8* Defaults )
@@ -49,6 +52,7 @@ struct FPropertyTag
 		, Size      (0)
 		, ArrayIndex(InIndex)
 		, SizeOffset(INDEX_NONE)
+		, HasPropertyGuid(0)
 	{
 		if (Property)
 		{
@@ -76,6 +80,16 @@ struct FPropertyTag
 		}
 	}
 
+	// Set optional property guid
+	void SetPropertyGuid(const FGuid& InPropertyGuid)
+	{
+		if (InPropertyGuid.IsValid())
+		{
+			PropertyGuid = InPropertyGuid;
+			HasPropertyGuid = true;
+		}
+	}
+
 	// Serializer.
 	friend FArchive& operator<<( FArchive& Ar, FPropertyTag& Tag )
 	{
@@ -124,7 +138,16 @@ struct FPropertyTag
 				Ar << Tag.InnerType;
 			}
 		}
-
+		// Property tags to handle renamed blueprint properties effectively.
+		if (Ar.UE4Ver() >= VER_UE4_PROPERTY_GUID_IN_PROPERTY_TAG)
+		{
+			Ar << Tag.HasPropertyGuid;
+			if (Tag.HasPropertyGuid)
+			{
+				Ar << Tag.PropertyGuid;
+			}
+		}
+		
 		return Ar;
 	}
 
diff --git a/Engine/Source/Runtime/CoreUObject/Public/UObject/Script.h b/Engine/Source/Runtime/CoreUObject/Public/UObject/Script.h
--- a/Engine/Source/Runtime/CoreUObject/Public/UObject/Script.h
+++ b/Engine/Source/Runtime/CoreUObject/Public/UObject/Script.h
@@ -293,6 +293,7 @@ namespace EScriptInstrumentation
 		Function,
 		Branch,
 		Macro,
+		PureNodeEntry,
 		NodeEntry,
 		NodeExit,
 		NodePin,
@@ -337,13 +338,15 @@ public:
 	EScriptInstrumentationEvent(EScriptInstrumentation::Type InEventType, const UObject* InContextObject, const struct FFrame& InStackFrame)
 		: EventType(InEventType)
 		, ContextObject(InContextObject)
+		, EventName(NAME_None)
 		, StackFramePtr(&InStackFrame)
 	{
 	}
 
-	EScriptInstrumentationEvent(EScriptInstrumentation::Type InEventType, const UObject* InContextObject)
+	EScriptInstrumentationEvent(EScriptInstrumentation::Type InEventType, const UObject* InContextObject, FName InEventName)
 		: EventType(InEventType)
 		, ContextObject(InContextObject)
+		, EventName(InEventName)
 		, StackFramePtr(nullptr)
 	{
 	}
@@ -352,12 +355,14 @@ public:
 	void SetType(EScriptInstrumentation::Type InType) { EventType = InType;	}
 	const UObject* GetContextObject() const { return ContextObject; }
 	bool IsStackFrameValid() const { return StackFramePtr != nullptr; }
+	FName GetEventName() const { return EventName; }
 	const FFrame& GetStackFrame() const { return *StackFramePtr; }
 
 protected:
 
 	EScriptInstrumentation::Type EventType;
 	const UObject* ContextObject;
+	const FName EventName;
 	const struct FFrame* StackFramePtr;
 };
 
diff --git a/Engine/Source/Runtime/Engine/Classes/EdGraph/EdGraphPin.h b/Engine/Source/Runtime/Engine/Classes/EdGraph/EdGraphPin.h
--- a/Engine/Source/Runtime/Engine/Classes/EdGraph/EdGraphPin.h
+++ b/Engine/Source/Runtime/Engine/Classes/EdGraph/EdGraphPin.h
@@ -220,6 +220,10 @@ class UEdGraphPin : public UObject
 	UPROPERTY()
 	UEdGraphPin* ParentPin;
 
+	/** Pin that this pin uses for passing through reference connection */
+	UPROPERTY()
+	UEdGraphPin* ReferencePassThroughConnection;
+
 #if WITH_EDITORONLY_DATA
 	/** If true, this connector is currently hidden. */
 	UPROPERTY()
@@ -285,6 +289,9 @@ public:
 	 */
 	ENGINE_API void CopyPersistentDataFromOldPin(const UEdGraphPin& SourcePin);
 
+	/** Connects the two pins as by-ref pass-through, allowing the input to auto-forward to the output pin */
+	ENGINE_API void AssignByRefPassThroughConnection(UEdGraphPin* InTargetPin);
+
 	/** Returns the node that owns this pin */
 	class UEdGraphNode* GetOwningNode() const
 	{
diff --git a/Engine/Source/Runtime/Engine/Classes/Engine/Blueprint.h b/Engine/Source/Runtime/Engine/Classes/Engine/Blueprint.h
--- a/Engine/Source/Runtime/Engine/Classes/Engine/Blueprint.h
+++ b/Engine/Source/Runtime/Engine/Classes/Engine/Blueprint.h
@@ -4,6 +4,7 @@
 
 #include "EdGraph/EdGraphPin.h"
 #include "EdGraph/EdGraphNode.h"
+#include "Engine/EngineTypes.h"
 #include "BlueprintCore.h"
 #include "Blueprint.generated.h"
 
diff --git a/Engine/Source/Runtime/Engine/Classes/Engine/BlueprintGeneratedClass.h b/Engine/Source/Runtime/Engine/Classes/Engine/BlueprintGeneratedClass.h
--- a/Engine/Source/Runtime/Engine/Classes/Engine/BlueprintGeneratedClass.h
+++ b/Engine/Source/Runtime/Engine/Classes/Engine/BlueprintGeneratedClass.h
@@ -341,6 +341,17 @@ public:
 		PerFuncInfo.LineNumberToPinMap.Add(CodeOffset, ExecPin);
 	}
 
+	int32 FindPinToCodeAssociation(UEdGraphPin const* ExecPin, UFunction* InFunction)
+	{
+		FDebuggingInfoForSingleFunction& PerFuncInfo = PerFunctionLineNumbers.FindOrAdd(InFunction);
+		int32 Result = INDEX_NONE;
+		if (const int32* CodeOffset = PerFuncInfo.LineNumberToPinMap.FindKey(ExecPin))
+		{
+			Result = *CodeOffset;
+		}
+		return Result;
+	}
+
 	// Registers an association between an object (pin or node typically) and an associated class member property
 	void RegisterClassPropertyAssociation(class UObject* TrueSourceObject, class UProperty* AssociatedProperty)
 	{
@@ -517,7 +528,11 @@ public:
 #if WITH_EDITORONLY_DATA
 	UPROPERTY(Transient)
 	UObject* OverridenArchetypeForCDO;
-#endif //WITH_EDITOR
+
+	/** Property guid map */
+	UPROPERTY()
+	TMap<FName,FGuid> PropertyGuids;
+#endif //WITH_EDITORONLY_DATA
 
 	// Mapping of changed properties & data to apply when instancing components in a cooked build (one entry per named AddComponent node template for fast lookup at runtime).
 	// Note: This is not currently utilized by the editor; it is a runtime optimization for cooked builds only. It assumes that the component class structure does not change.
@@ -572,8 +587,16 @@ public:
 	{
 		return CustomPropertyListForPostConstruction.Num() > 0 ? *CustomPropertyListForPostConstruction.GetData() : nullptr;
 	}
+
+protected:
+
+	virtual FName FindPropertyNameFromGuid(const FGuid& PropertyGuid) const override;
+	virtual FGuid FindPropertyGuidFromName(const FName InName) const override;
+	virtual bool ArePropertyGuidsAvailable() const override;
 	// End UClass interface
 
+public:
+
 	/** Called when the custom list of properties used during post-construct initialization needs to be rebuilt (e.g. after serialization and recompilation). */
 	void UpdateCustomPropertyListForPostConstruction();
 
diff --git a/Engine/Source/Runtime/Engine/Private/BlueprintGeneratedClass.cpp b/Engine/Source/Runtime/Engine/Private/BlueprintGeneratedClass.cpp
--- a/Engine/Source/Runtime/Engine/Private/BlueprintGeneratedClass.cpp
+++ b/Engine/Source/Runtime/Engine/Private/BlueprintGeneratedClass.cpp
@@ -1104,4 +1104,38 @@ void FBlueprintCookedComponentInstancingData::LoadCachedPropertyDataForSerializa
 			bIsValid = false;
 		}
 	}
+}
+
+bool UBlueprintGeneratedClass::ArePropertyGuidsAvailable() const
+{
+#if WITH_EDITORONLY_DATA
+	// Property guid's are generated during compilation.
+	return PropertyGuids.Num() > 0;
+#else
+	return false;
+#endif // WITH_EDITORONLY_DATA
+}
+
+FName UBlueprintGeneratedClass::FindPropertyNameFromGuid(const FGuid& PropertyGuid) const
+{
+	FName RedirectedName = NAME_None;
+#if WITH_EDITORONLY_DATA
+	if (const FName* Result = PropertyGuids.FindKey(PropertyGuid))
+	{
+		RedirectedName = *Result;
+	}
+#endif // WITH_EDITORONLY_DATA
+	return RedirectedName;
+}
+
+FGuid UBlueprintGeneratedClass::FindPropertyGuidFromName(const FName InName) const
+{
+	FGuid PropertyGuid;
+#if WITH_EDITORONLY_DATA
+	if (const FGuid* Result = PropertyGuids.Find(InName))
+	{
+		PropertyGuid = *Result;
+	}
+#endif // WITH_EDITORONLY_DATA
+	return PropertyGuid;
 }
\ No newline at end of file
diff --git a/Engine/Source/Runtime/Engine/Private/EdGraph/EdGraphPin.cpp b/Engine/Source/Runtime/Engine/Private/EdGraph/EdGraphPin.cpp
--- a/Engine/Source/Runtime/Engine/Private/EdGraph/EdGraphPin.cpp
+++ b/Engine/Source/Runtime/Engine/Private/EdGraph/EdGraphPin.cpp
@@ -235,6 +235,40 @@ void UEdGraphPin::CopyPersistentDataFromOldPin(const UEdGraphPin& SourcePin)
 #endif // WITH_EDITORONLY_DATA
 }
 
+void UEdGraphPin::AssignByRefPassThroughConnection(UEdGraphPin* InTargetPin)
+{
+	if (InTargetPin)
+	{
+		if (GetOwningNode() != InTargetPin->GetOwningNode())
+		{
+			UE_LOG(LogBlueprint, Warning, TEXT("Pin '%s' is owned by node '%s' and pin '%s' is owned by '%s', they must be owned by the same node!"), *GetName(), *GetOwningNode()->GetName(), *InTargetPin->GetName(), *InTargetPin->GetOwningNode()->GetName());
+		}
+		else if (Direction == InTargetPin->Direction)
+		{
+			UE_LOG(LogBlueprint, Warning, TEXT("Both pin '%s' and pin '%s' on node '%s' go in the same direction, one must be an input and the other an output!"), *GetName(), *InTargetPin->GetName(), *GetOwningNode()->GetName());
+		}
+		else if (!PinType.bIsReference && !InTargetPin->PinType.bIsReference)
+		{
+			UEdGraphPin* InputPin = (Direction == EGPD_Input) ? this : InTargetPin;
+			UEdGraphPin* OutputPin = (Direction == EGPD_Input) ? InTargetPin : this;
+			UE_LOG(LogBlueprint, Warning, TEXT("Input pin '%s' is attempting to by-ref pass-through to output pin '%s' on node '%s', however neither pin is by-ref!"), *InputPin->GetName(), *OutputPin->GetName(), *InputPin->GetOwningNode()->GetName());
+		}
+		else if (!PinType.bIsReference)
+		{
+			UE_LOG(LogBlueprint, Warning, TEXT("Pin '%s' on node '%s' is not by-ref but it is attempting to pass-through to '%s'"), *GetName(), *GetOwningNode()->GetName(), *InTargetPin->GetName());
+		}
+		else if (!InTargetPin->PinType.bIsReference)
+		{
+			UE_LOG(LogBlueprint, Warning, TEXT("Pin '%s' on node '%s' is not by-ref but it is attempting to pass-through to '%s'"), *InTargetPin->GetName(), *InTargetPin->GetOwningNode()->GetName(), *GetName());
+		}
+		else
+		{
+			ReferencePassThroughConnection = InTargetPin;
+			InTargetPin->ReferencePassThroughConnection = this;
+		}
+	}
+}
+
 const class UEdGraphSchema* UEdGraphPin::GetSchema() const
 {
 #if WITH_EDITOR
diff --git a/Engine/Source/Runtime/Engine/Public/GeneratedCodeHelpers.h b/Engine/Source/Runtime/Engine/Public/GeneratedCodeHelpers.h
--- a/Engine/Source/Runtime/Engine/Public/GeneratedCodeHelpers.h
+++ b/Engine/Source/Runtime/Engine/Public/GeneratedCodeHelpers.h
@@ -24,6 +24,18 @@
 //For DOREPLIFETIME macros
 #include "Net/UnrealNetwork.h"
 
+//For Box2D
+#include "Runtime/Core/Public/Math/Box2D.h"
+
+inline FBox2D CreateFBox2D(FVector2D InMin, FVector2D InMax, bool InIsValid)
+{
+	FBox2D Result;
+	Result.Min = InMin;
+	Result.Max = InMax;
+	Result.bIsValid = InIsValid;
+	return Result;
+}
+
 template<class NativeType>
 inline NativeType* NoNativeCast(UClass* NoNativeClass, UObject* Object)
 {
diff --git a/Engine/Source/Runtime/UMG/Private/UserWidget.cpp b/Engine/Source/Runtime/UMG/Private/UserWidget.cpp
--- a/Engine/Source/Runtime/UMG/Private/UserWidget.cpp
+++ b/Engine/Source/Runtime/UMG/Private/UserWidget.cpp
@@ -67,7 +67,8 @@ bool UUserWidget::Initialize()
 		{
 			//TODO NickD: This is a hack, and should be undone later
 			UWidgetBlueprintGeneratedClass* SuperBGClass = Cast<UWidgetBlueprintGeneratedClass>(BGClass->GetSuperClass());
-			if (BGClass->WidgetTree->RootWidget == nullptr && SuperBGClass )
+			const bool bNoRootWidget = (nullptr == BGClass->WidgetTree) || (nullptr == BGClass->WidgetTree->RootWidget);
+			if (bNoRootWidget && SuperBGClass)
 			{
 				SuperBGClass->InitializeWidget(this);
 			}
