Blueprints: Fixed crash if Redo is invoked while transaction buffer was empty

diff --git a/Engine/Source/Editor/Kismet/Private/BlueprintEditorModule.cpp b/Engine/Source/Editor/Kismet/Private/BlueprintEditorModule.cpp
--- a/Engine/Source/Editor/Kismet/Private/BlueprintEditorModule.cpp
+++ b/Engine/Source/Editor/Kismet/Private/BlueprintEditorModule.cpp
@@ -158,7 +158,10 @@ void FBlueprintUndoRedoHandler::PostUndo(bool bSuccess)
 void FBlueprintUndoRedoHandler::PostRedo(bool bSuccess)
 {
 	// Note: We add 1 to get the correct slot, because the transaction buffer will have decremented the UndoCount prior to getting here.
-	FixSubObjectReferencesPostUndoRedo(GEditor->Trans->GetTransaction(GEditor->Trans->GetQueueLength() - (GEditor->Trans->GetUndoCount() + 1)));
+	if( GEditor->Trans->GetQueueLength() > 0 )
+	{
+		FixSubObjectReferencesPostUndoRedo(GEditor->Trans->GetTransaction(GEditor->Trans->GetQueueLength() - (GEditor->Trans->GetUndoCount() + 1)));
+	}
 }
 
 void FBlueprintEditorModule::StartupModule()
