VREditor - Making the blueprint editor able to handle not having a parent window.  Updating the tab embedding code to better embed the tabs, and fill the entire provided window instead of trying to nest a new SWindow in the existing virtual window.

diff --git a/Engine/Source/Editor/Kismet/Private/BlueprintEditor.cpp b/Engine/Source/Editor/Kismet/Private/BlueprintEditor.cpp
--- a/Engine/Source/Editor/Kismet/Private/BlueprintEditor.cpp
+++ b/Engine/Source/Editor/Kismet/Private/BlueprintEditor.cpp
@@ -8240,19 +8240,23 @@ void FBlueprintEditor::TryInvokingDetailsTab(bool bFlash)
 
 		// We don't want to force this tab into existence when the blueprint editor isn't in the foreground and actively
 		// being interacted with.  So we make sure the window it's in is focused and the tab is in the foreground.
-		if ( BlueprintTab.IsValid() && BlueprintTab->IsForeground() && BlueprintTab->GetParentWindow()->HasFocusedDescendants()  )
+		if ( BlueprintTab.IsValid() && BlueprintTab->IsForeground() )
 		{
-			if (!Inspector.IsValid() || !Inspector->GetOwnerTab().IsValid() || Inspector->GetOwnerTab()->GetDockArea().IsValid())
+			TSharedPtr<SWindow> ParentWindow = BlueprintTab->GetParentWindow();
+			if ( ParentWindow.IsValid() && ParentWindow->HasFocusedDescendants() )
 			{
-				// Show the details panel if it doesn't exist.
-				TabManager->InvokeTab(FBlueprintEditorTabs::DetailsID);
-
-				if(bFlash)
+				if ( !Inspector.IsValid() || !Inspector->GetOwnerTab().IsValid() || Inspector->GetOwnerTab()->GetDockArea().IsValid() )
 				{
-					TSharedPtr<SDockTab> OwnerTab = Inspector->GetOwnerTab();
-					if(OwnerTab.IsValid())
+					// Show the details panel if it doesn't exist.
+					TabManager->InvokeTab(FBlueprintEditorTabs::DetailsID);
+
+					if ( bFlash )
 					{
-						OwnerTab->FlashTab();
+						TSharedPtr<SDockTab> OwnerTab = Inspector->GetOwnerTab();
+						if ( OwnerTab.IsValid() )
+						{
+							OwnerTab->FlashTab();
+						}
 					}
 				}
 			}
diff --git a/Engine/Source/Editor/VREditor/VREditorUISystem.cpp b/Engine/Source/Editor/VREditor/VREditorUISystem.cpp
--- a/Engine/Source/Editor/VREditor/VREditorUISystem.cpp
+++ b/Engine/Source/Editor/VREditor/VREditorUISystem.cpp
@@ -799,12 +799,6 @@ void FVREditorUISystem::CreateUIs()
 			TSharedRef<SWindow> TabManagerWindowRef = TabManagerWindow.ToSharedRef();
 			ProxyTabManager = MakeShareable(new FProxyTabmanager(TabManagerWindowRef));
 
-			TSharedRef<FTabManager::FLayout> LoadedLayout = FTabManager::NewLayout(TEXT("VRTabManager"));
-			LoadedLayout->AddArea(FTabManager::NewPrimaryArea());
-
-			TSharedPtr<SWidget> DockLayout = FGlobalTabmanager::Get()->RestoreFrom(LoadedLayout, TabManagerWindowRef, false);
-			TabManagerWindowRef->SetContent(DockLayout.ToSharedRef());
-
 			// We're going to start stealing tabs from the global tab manager inserting them into the world instead.
 			FGlobalTabmanager::Get()->SetProxyTabManager(ProxyTabManager);
 		}
diff --git a/Engine/Source/Runtime/Slate/Private/Framework/Docking/TabManager.cpp b/Engine/Source/Runtime/Slate/Private/Framework/Docking/TabManager.cpp
--- a/Engine/Source/Runtime/Slate/Private/Framework/Docking/TabManager.cpp
+++ b/Engine/Source/Runtime/Slate/Private/Framework/Docking/TabManager.cpp
@@ -2012,14 +2012,16 @@ void FGlobalTabmanager::SetProxyTabManager(TSharedPtr<FProxyTabmanager> InProxyT
 void FProxyTabmanager::OpenUnmanagedTab(FName PlaceholderId, const FSearchPreference& SearchPreference, const TSharedRef<SDockTab>& UnmanagedTab)
 {
 	// No layout info about this tab found; start 
-	TSharedRef<FArea> NewAreaForTab = FTabManager::NewArea(FTabManager::FallbackWindowSize)
+	TSharedRef<FArea> NewAreaForTab = FTabManager::NewPrimaryArea()
 		->Split
 		(
 			FTabManager::NewStack()
 			->AddTab(UnmanagedTab->GetLayoutIdentifier(), ETabState::OpenedTab)
 		);
 
-	TSharedRef<SDockingArea> DockingArea = RestoreArea(NewAreaForTab, ParentWindow.Pin());
+	TSharedPtr<SWindow> ParentWindowPtr = ParentWindow.Pin();
+	TSharedRef<SDockingArea> DockingArea = RestoreArea(NewAreaForTab, ParentWindowPtr, false);
+	ParentWindowPtr->SetContent(DockingArea);
 
 	const TSharedPtr<SDockTab> NewlyOpenedTab = DockingArea->GetAllChildTabs()[0];
 	check(NewlyOpenedTab.IsValid());
