Re-Merge CL# 2910537 and 2910629 from Release-4.11 to Main

Falling back to old logic for SSCSEditor::FindTreeNode(), when GetEditableComponentTemplate() fails to return a component.

diff --git a/Engine/Source/Editor/Kismet/Private/SSCSEditor.cpp b/Engine/Source/Editor/Kismet/Private/SSCSEditor.cpp
--- a/Engine/Source/Editor/Kismet/Private/SSCSEditor.cpp
+++ b/Engine/Source/Editor/Kismet/Private/SSCSEditor.cpp
@@ -5365,10 +5365,18 @@ FSCSEditorTreeNodePtrType SSCSEditor::FindTreeNode(const UActorComponent* InComp
 		if(InStartNodePtr.IsValid())
 		{
 			// Check to see if the given component template matches the given tree node
+			// 
+			// For certain node types, GetEditableComponentTemplate() will handle retrieving 
+			// the "OverridenComponentTemplate" which may be what we're looking for in some 
+			// cases; if not, then we fall back to just checking GetComponentTemplate()
 			if (InStartNodePtr->GetEditableComponentTemplate(GetBlueprint()) == InComponent)
 			{
 				NodePtr = InStartNodePtr;
 			}
+			else if (InStartNodePtr->GetComponentTemplate() == InComponent)
+			{
+				NodePtr = InStartNodePtr;
+			}
 			else
 			{
 				// Recursively search for the node in our child set
