Copying //UE4/Dev-Blueprints to Dev-Main (//UE4/Dev-Main)

diff --git a/Engine/Source/Editor/BlueprintGraph/Private/K2Node.cpp b/Engine/Source/Editor/BlueprintGraph/Private/K2Node.cpp
--- a/Engine/Source/Editor/BlueprintGraph/Private/K2Node.cpp
+++ b/Engine/Source/Editor/BlueprintGraph/Private/K2Node.cpp
@@ -289,27 +289,7 @@ void UK2Node::ReallocatePinsDuringReconstruction(TArray<UEdGraphPin*>& OldPins)
 {
 	AllocateDefaultPins();
 
-	for (auto OldPin : OldPins)
-	{
-		if (OldPin->ParentPin)
-		{
-			// find the new pin that corresponds to parent, and split it if it isn't already split
-			for (auto NewPin : Pins)
-			{
-				if (FCString::Stricmp(*(NewPin->PinName), *(OldPin->ParentPin->PinName)) == 0)
-				{
-					// Make sure we're not dealing with a menu node
-					UEdGraph* OuterGraph = GetGraph();
-					if (OuterGraph && OuterGraph->Schema && NewPin->SubPins.Num() == 0)
-					{
-						NewPin->PinType = OldPin->ParentPin->PinType;
-						GetSchema()->SplitPin(NewPin);
-						break;
-					}
-				}
-			}
-		}
-	}
+	RestoreSplitPins(OldPins);
 }
 
 void UK2Node::PostReconstructNode()
@@ -477,6 +457,32 @@ UK2Node::ERedirectType UK2Node::ShouldRedirectParam(const TArray<FString>& OldPi
 	return ERedirectType_None;
 }
 
+void UK2Node::RestoreSplitPins(TArray<UEdGraphPin*>& OldPins)
+{
+	// necessary to recreate split pins and keep their wires
+	for (auto OldPin : OldPins)
+	{
+		if (OldPin->ParentPin)
+		{
+			// find the new pin that corresponds to parent, and split it if it isn't already split
+			for (auto NewPin : Pins)
+			{
+				if (FCString::Stricmp(*(NewPin->PinName), *(OldPin->ParentPin->PinName)) == 0)
+				{
+					// Make sure we're not dealing with a menu node
+					UEdGraph* OuterGraph = GetGraph();
+					if (OuterGraph && OuterGraph->Schema && NewPin->SubPins.Num() == 0)
+					{
+						NewPin->PinType = OldPin->ParentPin->PinType;
+						GetSchema()->SplitPin(NewPin);
+						break;
+					}
+				}
+			}
+		}
+	}
+}
+
 UK2Node::ERedirectType UK2Node::DoPinsMatchForReconstruction(const UEdGraphPin* NewPin, int32 NewPinIndex, const UEdGraphPin* OldPin, int32 OldPinIndex) const
 {
 	ERedirectType RedirectType = ERedirectType_None;
diff --git a/Engine/Source/Editor/BlueprintGraph/Private/K2Node_ActorBoundEvent.cpp b/Engine/Source/Editor/BlueprintGraph/Private/K2Node_ActorBoundEvent.cpp
--- a/Engine/Source/Editor/BlueprintGraph/Private/K2Node_ActorBoundEvent.cpp
+++ b/Engine/Source/Editor/BlueprintGraph/Private/K2Node_ActorBoundEvent.cpp
@@ -4,6 +4,7 @@
 
 #include "KismetCompiler.h"
 #include "EventEntryHandler.h"
+#include "Kismet2/KismetEditorUtilities.h"
 
 #define LOCTEXT_NAMESPACE "K2Node_ActorBoundEvent"
 
@@ -81,6 +82,29 @@ void UK2Node_ActorBoundEvent::DestroyNode()
 	Super::DestroyNode();
 }
 
+bool UK2Node_ActorBoundEvent::CanPasteHere(const UEdGraph* TargetGraph) const
+{
+	// By default, to be safe, we don't allow events to be pasted, except under special circumstances (see below)
+	bool bDisallowPaste = !Super::CanPasteHere(TargetGraph);
+	if (!bDisallowPaste)
+	{
+		void* LevelBP = GetReferencedLevelActor()->GetLevel()->GetLevelScriptBlueprint(true);
+		if (FBlueprintEditorUtils::FindBlueprintForGraph(TargetGraph) == LevelBP)
+		{
+			if (const UK2Node_Event* PreExistingNode = FKismetEditorUtilities::FindBoundEventForActor(GetReferencedLevelActor(), DelegatePropertyName))
+			{
+				//UE_LOG(LogBlueprint, Log, TEXT("Cannot paste event node (%s) directly because it is flagged as an internal event."), *GetFName().ToString());
+				bDisallowPaste = true;
+			}
+		}
+		else
+		{
+			///UE_LOG(LogBlueprint, Log, TEXT("Cannot paste event node (%s) directly because it is flagged as an internal event."), *GetFName().ToString());
+		}
+	}
+	return !bDisallowPaste;
+}
+
 FText UK2Node_ActorBoundEvent::GetNodeTitle(ENodeTitleType::Type TitleType) const
 {
 	if (EventOwner == nullptr)
diff --git a/Engine/Source/Editor/BlueprintGraph/Private/K2Node_BitmaskLiteral.cpp b/Engine/Source/Editor/BlueprintGraph/Private/K2Node_BitmaskLiteral.cpp
new file mode 100644
--- /dev/null
+++ b/Engine/Source/Editor/BlueprintGraph/Private/K2Node_BitmaskLiteral.cpp
@@ -0,0 +1,94 @@
+// Copyright 1998-2016 Epic Games, Inc. All Rights Reserved.
+
+
+#include "BlueprintGraphPrivatePCH.h"
+#include "K2Node_BitmaskLiteral.h"
+#include "EditorCategoryUtils.h"
+#include "KismetCompiler.h"
+#include "../../../Runtime/Engine/Classes/Kismet/KismetSystemLibrary.h"
+#include "BlueprintFieldNodeSpawner.h"
+#include "BlueprintActionDatabaseRegistrar.h"
+
+#define LOCTEXT_NAMESPACE "UK2Node_BitmaskLiteral"
+
+const FString& UK2Node_BitmaskLiteral::GetBitmaskInputPinName()
+{
+	static const FString PinName(TEXT("Bitmask"));
+	return PinName;
+}
+
+UK2Node_BitmaskLiteral::UK2Node_BitmaskLiteral(const FObjectInitializer& ObjectInitializer)
+	: Super(ObjectInitializer)
+	, BitflagsEnum(nullptr)
+{
+	
+}
+
+void UK2Node_BitmaskLiteral::PostEditChangeProperty(struct FPropertyChangedEvent& PropertyChangedEvent)
+{
+	// Rebuild the node if the bitflags enum property was changed
+	FName PropertyName = (PropertyChangedEvent.Property != NULL) ? PropertyChangedEvent.Property->GetFName() : NAME_None;
+	if ((PropertyName == TEXT("BitflagsEnum")))
+	{
+		GetSchema()->ReconstructNode(*this);
+	}
+
+	Super::PostEditChangeProperty(PropertyChangedEvent);
+}
+
+void UK2Node_BitmaskLiteral::AllocateDefaultPins()
+{
+	const UEdGraphSchema_K2* Schema = GetDefault<UEdGraphSchema_K2>();
+
+	UEdGraphPin* InputPin = CreatePin(EGPD_Input, Schema->PC_Int, Schema->PSC_Bitmask, BitflagsEnum, false, false, GetBitmaskInputPinName());
+	Schema->SetPinDefaultValueBasedOnType(InputPin);
+
+	CreatePin(EGPD_Output, Schema->PC_Int, Schema->PSC_Bitmask, BitflagsEnum, false, false, Schema->PN_ReturnValue);
+
+	Super::AllocateDefaultPins();
+}
+
+void UK2Node_BitmaskLiteral::ExpandNode(class FKismetCompilerContext& CompilerContext, UEdGraph* SourceGraph)
+{
+	const UEdGraphSchema_K2* Schema = CompilerContext.GetSchema();
+	check(nullptr != Schema);
+
+	const FName FunctionName = GET_FUNCTION_NAME_CHECKED(UKismetSystemLibrary, MakeLiteralInt);
+	UK2Node_CallFunction* MakeLiteralInt = CompilerContext.SpawnIntermediateNode<UK2Node_CallFunction>(this, SourceGraph);
+	MakeLiteralInt->SetFromFunction(UKismetSystemLibrary::StaticClass()->FindFunctionByName(FunctionName));
+	MakeLiteralInt->AllocateDefaultPins();
+
+	UEdGraphPin* OrgInputPin = FindPinChecked(GetBitmaskInputPinName());
+	UEdGraphPin* NewInputPin = MakeLiteralInt->FindPinChecked(TEXT("Value"));
+	check(nullptr != NewInputPin);
+	CompilerContext.MovePinLinksToIntermediate(*OrgInputPin, *NewInputPin);
+
+	UEdGraphPin* OrgReturnPin = FindPinChecked(Schema->PN_ReturnValue);
+	UEdGraphPin* NewReturnPin = MakeLiteralInt->GetReturnValuePin();
+	check(nullptr != NewReturnPin);
+	CompilerContext.MovePinLinksToIntermediate(*OrgReturnPin, *NewReturnPin);
+}
+
+FText UK2Node_BitmaskLiteral::GetNodeTitle(ENodeTitleType::Type TitleType) const
+{
+	return LOCTEXT("NodeTitle", "Make Bitmask");
+}
+
+void UK2Node_BitmaskLiteral::GetMenuActions(FBlueprintActionDatabaseRegistrar& ActionRegistrar) const
+{
+	UClass* ActionKey = GetClass();
+	if (ActionRegistrar.IsOpenForRegistration(ActionKey))
+	{
+		UBlueprintNodeSpawner* NodeSpawner = UBlueprintNodeSpawner::Create(GetClass());
+		check(NodeSpawner != nullptr);
+
+		ActionRegistrar.AddBlueprintAction(ActionKey, NodeSpawner);
+	}
+}
+
+FText UK2Node_BitmaskLiteral::GetMenuCategory() const
+{
+	return FEditorCategoryUtils::GetCommonCategory(FCommonEditorCategory::Math);
+}
+
+#undef LOCTEXT_NAMESPACE
diff --git a/Engine/Source/Editor/BlueprintGraph/Private/K2Node_ComponentBoundEvent.cpp b/Engine/Source/Editor/BlueprintGraph/Private/K2Node_ComponentBoundEvent.cpp
--- a/Engine/Source/Editor/BlueprintGraph/Private/K2Node_ComponentBoundEvent.cpp
+++ b/Engine/Source/Editor/BlueprintGraph/Private/K2Node_ComponentBoundEvent.cpp
@@ -3,6 +3,7 @@
 #include "BlueprintGraphPrivatePCH.h"
 #include "Engine/ComponentDelegateBinding.h"
 #include "K2Node_ComponentBoundEvent.h"
+#include "Kismet2/KismetEditorUtilities.h"
 
 #define LOCTEXT_NAMESPACE "K2Node"
 
@@ -18,6 +19,21 @@ bool UK2Node_ComponentBoundEvent::Modify(bool bAlwaysMarkDirty)
 	return Super::Modify(bAlwaysMarkDirty);
 }
 
+bool UK2Node_ComponentBoundEvent::CanPasteHere(const UEdGraph* TargetGraph) const
+{
+	// By default, to be safe, we don't allow events to be pasted, except under special circumstances (see below)
+	bool bDisallowPaste = !Super::CanPasteHere(TargetGraph);
+	if (!bDisallowPaste)
+	{
+		if (const UK2Node_Event* PreExistingNode = FKismetEditorUtilities::FindBoundEventForComponent(FBlueprintEditorUtils::FindBlueprintForGraph(TargetGraph), DelegatePropertyName, ComponentPropertyName))
+		{
+			//UE_LOG(LogBlueprint, Log, TEXT("Cannot paste event node (%s) directly because it is flagged as an internal event."), *GetFName().ToString());
+			bDisallowPaste = true;
+		}
+	}
+	return !bDisallowPaste;
+
+}
 FText UK2Node_ComponentBoundEvent::GetNodeTitle(ENodeTitleType::Type TitleType) const
 {
 	if (CachedNodeTitle.IsOutOfDate(this))
diff --git a/Engine/Source/Editor/BlueprintGraph/Private/K2Node_ConstructObjectFromClass.cpp b/Engine/Source/Editor/BlueprintGraph/Private/K2Node_ConstructObjectFromClass.cpp
--- a/Engine/Source/Editor/BlueprintGraph/Private/K2Node_ConstructObjectFromClass.cpp
+++ b/Engine/Source/Editor/BlueprintGraph/Private/K2Node_ConstructObjectFromClass.cpp
@@ -54,16 +54,13 @@ void UK2Node_ConstructObjectFromClass::AllocateDefaultPins()
 
 	// Add blueprint pin
 	UEdGraphPin* ClassPin = CreatePin(EGPD_Input, K2Schema->PC_Class, TEXT(""), GetClassPinBaseClass(), false, false, FK2Node_ConstructObjectFromClassHelper::ClassPinName);
-	SetPinToolTip(*ClassPin, LOCTEXT("ClassPinDescription", "The object class you want to spawn"));
-
+	
 	// Result pin
 	UEdGraphPin* ResultPin = CreatePin(EGPD_Output, K2Schema->PC_Object, TEXT(""), GetClassPinBaseClass(), false, false, K2Schema->PN_ReturnValue);
-	SetPinToolTip(*ResultPin, LOCTEXT("ResultPinDescription", "The spawned object"));
-
+	
 	if (UseOuter())
 	{
 		UEdGraphPin* OuterPin = CreatePin(EGPD_Input, K2Schema->PC_Object, TEXT(""), UObject::StaticClass(), false, false, FK2Node_ConstructObjectFromClassHelper::OuterPinName);
-		SetPinToolTip(*OuterPin, LOCTEXT("OuterPinDescription", "Owner of the created object"));
 	}
 
 	Super::AllocateDefaultPins();
@@ -169,6 +166,7 @@ void UK2Node_ConstructObjectFromClass::ReallocatePinsDuringReconstruction(TArray
 	{
 		CreatePinsForClass(UseSpawnClass);
 	}
+	RestoreSplitPins(OldPins);
 }
 
 bool UK2Node_ConstructObjectFromClass::IsSpawnVarPin(UEdGraphPin* Pin)
@@ -232,6 +230,15 @@ void UK2Node_ConstructObjectFromClass::PinConnectionListChanged(UEdGraphPin* Pin
 	}
 }
 
+void UK2Node_ConstructObjectFromClass::GetPinHoverText(const UEdGraphPin& Pin, FString& HoverTextOut) const
+{
+	SetPinToolTip(*GetClassPin(), LOCTEXT("ClassPinDescription", "The object class you want to construct"));
+	SetPinToolTip(*GetResultPin(), LOCTEXT("ResultPinDescription", "The constructed object"));
+	SetPinToolTip(*GetOuterPin(), LOCTEXT("OuterPinDescription", "Owner of the constructed object"));
+
+	return Super::GetPinHoverText(Pin, HoverTextOut);
+}
+
 void UK2Node_ConstructObjectFromClass::PinDefaultValueChanged(UEdGraphPin* ChangedPin) 
 {
 	if (ChangedPin && (ChangedPin->PinName == FK2Node_ConstructObjectFromClassHelper::ClassPinName))
diff --git a/Engine/Source/Editor/BlueprintGraph/Private/K2Node_Event.cpp b/Engine/Source/Editor/BlueprintGraph/Private/K2Node_Event.cpp
--- a/Engine/Source/Editor/BlueprintGraph/Private/K2Node_Event.cpp
+++ b/Engine/Source/Editor/BlueprintGraph/Private/K2Node_Event.cpp
@@ -505,10 +505,17 @@ bool UK2Node_Event::CanPasteHere(const UEdGraph* TargetGraph) const
 						&& ExcludedEventNames.Contains(EventReference.GetMemberName().ToString());
 					if(!bDisallowPaste)
 					{
+						TArray<UK2Node_Event*> DisabledEventNodesToStomp;
 						// If the event function is already handled in this Blueprint, don't paste this event
 						for(int32 i = 0; i < ExistingEventNodes.Num() && !bDisallowPaste; ++i)
 						{
-							bDisallowPaste = ExistingEventNodes[i]->bOverrideFunction && ExistingEventNodes[i]->IsNodeEnabled() && AreEventNodesIdentical(this, ExistingEventNodes[i]);
+							bDisallowPaste = ExistingEventNodes[i]->bOverrideFunction /*&& ExistingEventNodes[i]->IsNodeEnabled() */&& AreEventNodesIdentical(this, ExistingEventNodes[i]);
+
+							if (bDisallowPaste && !ExistingEventNodes[i]->IsNodeEnabled())
+							{
+								DisabledEventNodesToStomp.Add(ExistingEventNodes[i]);
+								bDisallowPaste = false;
+							}
 						}
 
 						// We need to also check for 'const' BPIE methods that might already be implemented as functions with a read-only 'self' context (these were previously implemented as events)
@@ -543,6 +550,14 @@ bool UK2Node_Event::CanPasteHere(const UEdGraph* TargetGraph) const
 						{
 							UE_LOG(LogBlueprint, Log, TEXT("Cannot paste event node (%s) directly because the event function (%s) is already handled."), *GetFName().ToString(), *EventReference.GetMemberName().ToString());
 						}
+
+						if (!bDisallowPaste)
+						{
+							for (UK2Node_Event* EventNode : DisabledEventNodesToStomp)
+							{
+								EventNode->DestroyNode();
+							}
+						}
 					}
 					else
 					{
diff --git a/Engine/Source/Editor/BlueprintGraph/Private/K2Node_FunctionEntry.cpp b/Engine/Source/Editor/BlueprintGraph/Private/K2Node_FunctionEntry.cpp
--- a/Engine/Source/Editor/BlueprintGraph/Private/K2Node_FunctionEntry.cpp
+++ b/Engine/Source/Editor/BlueprintGraph/Private/K2Node_FunctionEntry.cpp
@@ -333,6 +333,14 @@ void UK2Node_FunctionEntry::ExpandNode(class FKismetCompilerContext& CompilerCon
 
 		// Find the associated UFunction
 		UFunction* Function = FindField<UFunction>(CompilerContext.Blueprint->SkeletonGeneratedClass, *OriginalNode->GetOuter()->GetName());
+
+		// When regenerating on load, we may need to import text on certain properties to force load the assets
+		TSharedPtr<FStructOnScope> LocalVarData;
+		if (CompilerContext.Blueprint->bIsRegeneratingOnLoad)
+		{
+			LocalVarData = MakeShareable(new FStructOnScope(Function));
+		}
+
 		for (TFieldIterator<UProperty> It(Function); It; ++It)
 		{
 			if (const UProperty* Property = *It)
@@ -400,6 +408,16 @@ void UK2Node_FunctionEntry::ExpandNode(class FKismetCompilerContext& CompilerCon
 							}
 							else
 							{
+								if (CompilerContext.Blueprint->bIsRegeneratingOnLoad)
+								{
+									// When regenerating on load, we want to force load assets referenced by local variables.
+									// This functionality is already handled when generating Terms in the Kismet Compiler for Arrays and Structs, so we do not have to worry about them.
+									if (LocalVar.VarType.PinCategory == Schema->PC_Object || LocalVar.VarType.PinCategory == Schema->PC_Class || LocalVar.VarType.PinCategory == Schema->PC_Interface)
+									{
+										FBlueprintEditorUtils::PropertyValueFromString(Property, LocalVar.DefaultValue, LocalVarData->GetStructMemory());
+									}
+								}
+
 								// Set the default value
 								Schema->TrySetDefaultValue(*SetPin, LocalVar.DefaultValue);
 							}
diff --git a/Engine/Source/Editor/BlueprintGraph/Private/K2Node_GetClassDefaults.cpp b/Engine/Source/Editor/BlueprintGraph/Private/K2Node_GetClassDefaults.cpp
--- a/Engine/Source/Editor/BlueprintGraph/Private/K2Node_GetClassDefaults.cpp
+++ b/Engine/Source/Editor/BlueprintGraph/Private/K2Node_GetClassDefaults.cpp
@@ -271,7 +271,7 @@ bool UK2Node_GetClassDefaults::HasExternalDependencies(TArray<class UStruct*>* O
 
 void UK2Node_GetClassDefaults::ReallocatePinsDuringReconstruction(TArray<UEdGraphPin*>& OldPins) 
 {
-	Super::ReallocatePinsDuringReconstruction(OldPins);
+	AllocateDefaultPins();
 
 	// Recreate output pins based on the previous input class
 	UEdGraphPin* OldClassPin = FindClassPin(OldPins);
@@ -279,6 +279,8 @@ void UK2Node_GetClassDefaults::ReallocatePinsDuringReconstruction(TArray<UEdGrap
 	{
 		CreateOutputPins(InputClass);
 	}
+
+	RestoreSplitPins(OldPins);
 }
 
 FNodeHandlingFunctor* UK2Node_GetClassDefaults::CreateNodeHandler(FKismetCompilerContext& CompilerContext) const
diff --git a/Engine/Source/Editor/BlueprintGraph/Private/K2Node_SpawnActor.cpp b/Engine/Source/Editor/BlueprintGraph/Private/K2Node_SpawnActor.cpp
--- a/Engine/Source/Editor/BlueprintGraph/Private/K2Node_SpawnActor.cpp
+++ b/Engine/Source/Editor/BlueprintGraph/Private/K2Node_SpawnActor.cpp
@@ -110,6 +110,8 @@ void UK2Node_SpawnActor::ReallocatePinsDuringReconstruction(TArray<UEdGraphPin*>
 	{
 		CreatePinsForClass(UseSpawnClass);
 	}
+
+	RestoreSplitPins(OldPins);
 }
 
 bool UK2Node_SpawnActor::IsSpawnVarPin(UEdGraphPin* Pin)
diff --git a/Engine/Source/Editor/BlueprintGraph/Private/K2Node_SpawnActorFromClass.cpp b/Engine/Source/Editor/BlueprintGraph/Private/K2Node_SpawnActorFromClass.cpp
--- a/Engine/Source/Editor/BlueprintGraph/Private/K2Node_SpawnActorFromClass.cpp
+++ b/Engine/Source/Editor/BlueprintGraph/Private/K2Node_SpawnActorFromClass.cpp
@@ -49,32 +49,25 @@ void UK2Node_SpawnActorFromClass::AllocateDefaultPins()
 
 	// Add blueprint pin
 	UEdGraphPin* ClassPin = CreatePin(EGPD_Input, K2Schema->PC_Class, TEXT(""), AActor::StaticClass(), false, false, FK2Node_SpawnActorFromClassHelper::ClassPinName);
-	K2Schema->ConstructBasicPinTooltip(*ClassPin, LOCTEXT("ClassPinDescription", "The Actor class you want to spawn"), ClassPin->PinToolTip);
 
 	// Transform pin
 	UScriptStruct* TransformStruct = TBaseStructure<FTransform>::Get();
 	UEdGraphPin* TransformPin = CreatePin(EGPD_Input, K2Schema->PC_Struct, TEXT(""), TransformStruct, false, false, FK2Node_SpawnActorFromClassHelper::SpawnTransformPinName);
-	K2Schema->ConstructBasicPinTooltip(*TransformPin, LOCTEXT("TransformPinDescription", "The transform to spawn the Actor with"), TransformPin->PinToolTip);
 
 	// Collision handling method pin
 	UEnum* const MethodEnum = FindObjectChecked<UEnum>(ANY_PACKAGE, TEXT("ESpawnActorCollisionHandlingMethod"), true);
 	UEdGraphPin* const CollisionHandlingOverridePin = CreatePin(EGPD_Input, K2Schema->PC_Byte, TEXT(""), MethodEnum, false, false, FK2Node_SpawnActorFromClassHelper::CollisionHandlingOverridePinName);
 	CollisionHandlingOverridePin->DefaultValue = MethodEnum->GetEnumName(static_cast<int>(ESpawnActorCollisionHandlingMethod::Undefined));
 
-	K2Schema->ConstructBasicPinTooltip(*CollisionHandlingOverridePin, LOCTEXT("CollisionHandlingOverridePinDescription", "Specifies how to handle collisions at the spawn point. If undefined, uses actor class settings."), CollisionHandlingOverridePin->PinToolTip);
-
-
 	UEdGraphPin* OwnerPin = CreatePin(EGPD_Input, K2Schema->PC_Object, TEXT(""), AActor::StaticClass(),/*bIsArray =*/false, /*bIsReference =*/false, FK2Node_SpawnActorFromClassHelper::OwnerPinName);
 	OwnerPin->bAdvancedView = true;
 	if (ENodeAdvancedPins::NoPins == AdvancedPinDisplay)
 	{
 		AdvancedPinDisplay = ENodeAdvancedPins::Hidden;
 	}
-	K2Schema->ConstructBasicPinTooltip(*OwnerPin, LOCTEXT("OwnerPinDescription", "Can be left empty; primarily used for replication (bNetUseOwnerRelevancy and bOnlyRelevantToOwner), or visibility (PrimitiveComponent's bOwnerNoSee/bOnlyOwnerSee)"), OwnerPin->PinToolTip);
 
 	// Result pin
 	UEdGraphPin* ResultPin = CreatePin(EGPD_Output, K2Schema->PC_Object, TEXT(""), AActor::StaticClass(), false, false, K2Schema->PN_ReturnValue);
-	K2Schema->ConstructBasicPinTooltip(*ResultPin, LOCTEXT("ResultPinDescription", "The spawned Actor"), ResultPin->PinToolTip);
 
 	Super::AllocateDefaultPins();
 }
@@ -278,6 +271,7 @@ void UK2Node_SpawnActorFromClass::ReallocatePinsDuringReconstruction(TArray<UEdG
 	}
 
 	MaybeUpdateCollisionPin(OldPins);
+	RestoreSplitPins(OldPins);
 }
 
 void UK2Node_SpawnActorFromClass::PostPlacedNewNode()
@@ -355,8 +349,6 @@ void UK2Node_SpawnActorFromClass::OnClassPinChanged()
 		K2Schema->TryCreateConnection(ResultPin, Connections);
 	}
 
-	K2Schema->ConstructBasicPinTooltip(*ResultPin, LOCTEXT("ResultPinDescription", "The spawned Actor"), ResultPin->PinToolTip);
-
 	// Rewire the old pins to the new pins so connections are maintained if possible
 	RewireOldPinsToNewPins(OldClassPins, NewClassPins);
 
@@ -379,6 +371,24 @@ void UK2Node_SpawnActorFromClass::PinConnectionListChanged(UEdGraphPin* ChangedP
 	}
 }
 
+void UK2Node_SpawnActorFromClass::GetPinHoverText(const UEdGraphPin& Pin, FString& HoverTextOut) const
+{
+	const UEdGraphSchema_K2* K2Schema = GetDefault<UEdGraphSchema_K2>();
+
+	UEdGraphPin* ClassPin = GetClassPin();
+	K2Schema->ConstructBasicPinTooltip(*ClassPin, LOCTEXT("ClassPinDescription", "The Actor class you want to spawn"), ClassPin->PinToolTip);
+	UEdGraphPin* TransformPin = GetSpawnTransformPin();
+	K2Schema->ConstructBasicPinTooltip(*TransformPin, LOCTEXT("TransformPinDescription", "The transform to spawn the Actor with"), TransformPin->PinToolTip);
+	UEdGraphPin* CollisionHandlingOverridePin = GetCollisionHandlingOverridePin();
+	K2Schema->ConstructBasicPinTooltip(*CollisionHandlingOverridePin, LOCTEXT("CollisionHandlingOverridePinDescription", "Specifies how to handle collisions at the spawn point. If undefined, uses actor class settings."), CollisionHandlingOverridePin->PinToolTip);
+	UEdGraphPin* OwnerPin = GetOwnerPin();
+	K2Schema->ConstructBasicPinTooltip(*OwnerPin, LOCTEXT("OwnerPinDescription", "Can be left empty; primarily used for replication (bNetUseOwnerRelevancy and bOnlyRelevantToOwner), or visibility (PrimitiveComponent's bOwnerNoSee/bOnlyOwnerSee)"), OwnerPin->PinToolTip);
+	UEdGraphPin* ResultPin = GetResultPin();
+	K2Schema->ConstructBasicPinTooltip(*ResultPin, LOCTEXT("ResultPinDescription", "The spawned Actor"), ResultPin->PinToolTip);
+
+	return Super::GetPinHoverText(Pin, HoverTextOut);
+}
+
 void UK2Node_SpawnActorFromClass::PinDefaultValueChanged(UEdGraphPin* ChangedPin) 
 {
 	if (ChangedPin && (ChangedPin->PinName == FK2Node_SpawnActorFromClassHelper::ClassPinName))
diff --git a/Engine/Source/Editor/BlueprintGraph/Private/K2Node_VariableGet.cpp b/Engine/Source/Editor/BlueprintGraph/Private/K2Node_VariableGet.cpp
--- a/Engine/Source/Editor/BlueprintGraph/Private/K2Node_VariableGet.cpp
+++ b/Engine/Source/Editor/BlueprintGraph/Private/K2Node_VariableGet.cpp
@@ -168,6 +168,8 @@ void UK2Node_VariableGet::ReallocatePinsDuringReconstruction(TArray<UEdGraphPin*
 			}
 		}
 		CreatePinForSelf();
+
+		RestoreSplitPins(OldPins);
 	}
 }
 
@@ -487,4 +489,4 @@ void UK2Node_VariableGet::Serialize(FArchive& Ar)
 }
 
 
-#undef LOCTEXT_NAMESPACE
\ No newline at end of file
+#undef LOCTEXT_NAMESPACE
diff --git a/Engine/Source/Editor/BlueprintGraph/Private/K2Node_VariableSet.cpp b/Engine/Source/Editor/BlueprintGraph/Private/K2Node_VariableSet.cpp
--- a/Engine/Source/Editor/BlueprintGraph/Private/K2Node_VariableSet.cpp
+++ b/Engine/Source/Editor/BlueprintGraph/Private/K2Node_VariableSet.cpp
@@ -123,6 +123,8 @@ void UK2Node_VariableSet::ReallocatePinsDuringReconstruction(TArray<UEdGraphPin*
 		CreateOutputPinTooltip();
 		CreatePinForSelf();
 	}
+
+	RestoreSplitPins(OldPins);
 }
 
 
diff --git a/Engine/Source/Editor/GraphEditor/Private/KismetPins/SGraphPinInteger.cpp b/Engine/Source/Editor/GraphEditor/Private/KismetPins/SGraphPinInteger.cpp
--- a/Engine/Source/Editor/GraphEditor/Private/KismetPins/SGraphPinInteger.cpp
+++ b/Engine/Source/Editor/GraphEditor/Private/KismetPins/SGraphPinInteger.cpp
@@ -13,6 +13,141 @@ void SGraphPinInteger::Construct(const FArguments& InArgs, UEdGraphPin* InGraphP
 	SGraphPinNum::Construct(SGraphPinNum::FArguments(), InGraphPinObj);
 }
 
+TSharedRef<SWidget>	SGraphPinInteger::GetDefaultValueWidget()
+{
+	check(GraphPinObj);
+
+	if (const UEdGraphSchema_K2* K2Schema = Cast<const UEdGraphSchema_K2>(GraphPinObj->GetSchema()))
+	{
+		FString PinSubCategory = GraphPinObj->PinType.PinSubCategory;
+		const UEnum* BitmaskEnum = Cast<UEnum>(GraphPinObj->PinType.PinSubCategoryObject.Get());
+
+		if (PinSubCategory == K2Schema->PSC_Bitmask)
+		{
+			struct BitmaskFlagInfo
+			{
+				int32 Value;
+				FText DisplayName;
+				FText ToolTipText;
+			};
+
+			const int32 BitmaskBitCount = sizeof(int32) << 3;
+
+			TArray<BitmaskFlagInfo> BitmaskFlags;
+			BitmaskFlags.Reserve(BitmaskBitCount);
+
+			if(BitmaskEnum)
+			{
+				for(int32 BitmaskEnumIndex = 0; BitmaskEnumIndex < BitmaskEnum->NumEnums() - 1; ++BitmaskEnumIndex)
+				{
+					int8 BitmaskFlagIndex = BitmaskEnum->GetValueByIndex(BitmaskEnumIndex);
+					if(BitmaskFlagIndex >= 0 && BitmaskFlagIndex < BitmaskBitCount)
+					{
+						BitmaskFlagInfo* BitmaskFlag = new(BitmaskFlags) BitmaskFlagInfo();
+
+						BitmaskFlag->Value = static_cast<int32>(1 << BitmaskFlagIndex);
+
+						BitmaskFlag->DisplayName = BitmaskEnum->GetDisplayNameText(BitmaskEnumIndex);
+						if(BitmaskFlag->DisplayName.IsEmpty())
+						{
+							BitmaskFlag->DisplayName = FText::FromString(BitmaskEnum->GetEnumName(BitmaskEnumIndex));
+						}
+
+						BitmaskFlag->ToolTipText = BitmaskEnum->GetToolTipText(BitmaskEnumIndex);
+						if(BitmaskFlag->ToolTipText.IsEmpty())
+						{
+							BitmaskFlag->ToolTipText = FText::Format(NSLOCTEXT("GraphEditor", "BitmaskDefaultFlagToolTipText", "Toggle {0} on/off"), BitmaskFlag->DisplayName);
+						}
+					}
+				}
+			}
+			else
+			{
+				for(int BitmaskFlagIndex = 0; BitmaskFlagIndex < BitmaskBitCount; ++BitmaskFlagIndex)
+				{
+					BitmaskFlagInfo* BitmaskFlag = new(BitmaskFlags) BitmaskFlagInfo();
+
+					BitmaskFlag->Value = static_cast<int32>(1 << BitmaskFlagIndex);
+					BitmaskFlag->DisplayName = FText::Format(NSLOCTEXT("GraphEditor", "BitmaskDefaultFlagDisplayName", "Flag {0}"), FText::AsNumber(BitmaskFlagIndex + 1));
+					BitmaskFlag->ToolTipText = FText::Format(NSLOCTEXT("GraphEditor", "BitmaskDefaultFlagToolTipText", "Toggle {0} on/off"), BitmaskFlag->DisplayName);
+				}
+			}
+
+			const auto& GetComboButtonText = [=]() -> FText
+			{
+				int32 BitmaskValue = FCString::Atoi(*GraphPinObj->GetDefaultAsString());
+				if(BitmaskValue > 0)
+				{
+					if(BitmaskValue & (BitmaskValue - 1))
+					{
+						return NSLOCTEXT("GraphEditor", "BitmaskButtonContentMultipleBitsSet", "(Multiple)");
+					}
+					else
+					{
+						for(int i = 0; i < BitmaskFlags.Num(); ++i)
+						{
+							if(BitmaskValue & BitmaskFlags[i].Value)
+							{
+								return BitmaskFlags[i].DisplayName;
+							}
+						}
+					}
+				}
+
+				return NSLOCTEXT("GraphEditor", "BitmaskButtonContentNoFlagsSet", "(No Flags)");
+			};
+
+			return SNew(SComboButton)
+				.ContentPadding(3)
+				.MenuPlacement(MenuPlacement_BelowAnchor)
+				.Visibility(this, &SGraphPin::GetDefaultValueVisibility)
+				.ButtonContent()
+				[
+					// Wrap in configurable box to restrain height/width of menu
+					SNew(SBox)
+					.MinDesiredWidth(84.0f)
+					[
+						SNew(STextBlock)
+						.Text_Lambda(GetComboButtonText)
+						.Font(FEditorStyle::GetFontStyle(TEXT("PropertyWindow.NormalFont")))
+					]
+				]
+				.OnGetMenuContent_Lambda([=]()
+				{
+					FMenuBuilder MenuBuilder(false, nullptr);
+
+					for(int i = 0; i < BitmaskFlags.Num(); ++i)
+					{
+						MenuBuilder.AddMenuEntry(
+							BitmaskFlags[i].DisplayName,
+							BitmaskFlags[i].ToolTipText,
+							FSlateIcon(),
+							FUIAction
+							(
+								FExecuteAction::CreateLambda([=]()
+								{
+									const int32 CurValue = FCString::Atoi(*GraphPinObj->GetDefaultAsString());
+									GraphPinObj->GetSchema()->TrySetDefaultValue(*GraphPinObj, FString::FromInt(CurValue ^ BitmaskFlags[i].Value));
+								}),
+								FCanExecuteAction(),
+								FIsActionChecked::CreateLambda([=]() -> bool
+								{
+									const int32 CurValue = FCString::Atoi(*GraphPinObj->GetDefaultAsString());
+									return (CurValue & BitmaskFlags[i].Value) != 0;
+								})
+							),
+							NAME_None,
+							EUserInterfaceActionType::Check);
+					}
+
+					return MenuBuilder.MakeWidget();
+				});
+		}
+	}
+	
+	return SGraphPinNum::GetDefaultValueWidget();
+}
+
 void SGraphPinInteger::SetTypeInValue(const FText& NewTypeInValue, ETextCommit::Type /*CommitInfo*/)
 {
 	const FString TypeValueString = NewTypeInValue.ToString();
diff --git a/Engine/Source/Editor/GraphEditor/Private/KismetPins/SGraphPinInteger.h b/Engine/Source/Editor/GraphEditor/Private/KismetPins/SGraphPinInteger.h
--- a/Engine/Source/Editor/GraphEditor/Private/KismetPins/SGraphPinInteger.h
+++ b/Engine/Source/Editor/GraphEditor/Private/KismetPins/SGraphPinInteger.h
@@ -12,6 +12,7 @@ public:
 
 protected:
 	//~ Begin SGraphPinString Interface
+	virtual TSharedRef<SWidget>	GetDefaultValueWidget() override;
 	virtual void SetTypeInValue(const FText& NewTypeInValue, ETextCommit::Type CommitInfo) override;
 	//~ End SGraphPinString Interface
 };
diff --git a/Engine/Source/Editor/GraphEditor/Private/SGraphPanel.cpp b/Engine/Source/Editor/GraphEditor/Private/SGraphPanel.cpp
--- a/Engine/Source/Editor/GraphEditor/Private/SGraphPanel.cpp
+++ b/Engine/Source/Editor/GraphEditor/Private/SGraphPanel.cpp
@@ -105,6 +105,8 @@ void SGraphPanel::Construct( const SGraphPanel::FArguments& InArgs )
 
 	SavedMousePosForOnPaintEventLocalSpace = FVector2D::ZeroVector;
 	PreviousFrameSavedMousePosForSplineOverlap = FVector2D::ZeroVector;
+
+	TimeLeftToInvalidatePerTick = 0.0f;
 }
 
 SGraphPanel::~SGraphPanel()
@@ -683,6 +685,15 @@ void SGraphPanel::AddPinToHoverSet(UEdGraphPin* HoveredPin)
 {
 	CurrentHoveredPins.Add(HoveredPin);
 	TimeWhenMouseEnteredPin = FSlateApplication::Get().GetCurrentTime();
+
+	// About covers the fade in time when highlighting pins or splines.
+	TimeLeftToInvalidatePerTick += 1.5f;
+
+	// This handle should always be for this function
+	if (!ActiveTimerHandleInvalidatePerTick.IsValid())
+	{
+		ActiveTimerHandleInvalidatePerTick = RegisterActiveTimer(0.f, FWidgetActiveTimerDelegate::CreateSP(this, &SGraphPanel::InvalidatePerTick));
+	}
 }
 
 void SGraphPanel::RemovePinFromHoverSet(UEdGraphPin* UnhoveredPin)
@@ -1662,3 +1673,19 @@ void SGraphPanel::AddReferencedObjects(FReferenceCollector& Collector)
 	Collector.AddReferencedObject( GraphObj );
 	Collector.AddReferencedObject( GraphObjToDiff );
 }
+
+EActiveTimerReturnType SGraphPanel::InvalidatePerTick(double InCurrentTime, float InDeltaTime)
+{
+	// Invalidate the layout so it will redraw.
+	Invalidate(EInvalidateWidget::Layout);
+
+	TimeLeftToInvalidatePerTick -= InDeltaTime;
+
+	// When the time is done, stop the invalidation per tick because the UI will be static once more.
+	if (TimeLeftToInvalidatePerTick <= 0.0f)
+	{
+		TimeLeftToInvalidatePerTick = 0.0f;
+		return EActiveTimerReturnType::Stop;
+	}
+	return EActiveTimerReturnType::Continue;
+}
diff --git a/Engine/Source/Editor/GraphEditor/Private/SGraphPanel.h b/Engine/Source/Editor/GraphEditor/Private/SGraphPanel.h
--- a/Engine/Source/Editor/GraphEditor/Private/SGraphPanel.h
+++ b/Engine/Source/Editor/GraphEditor/Private/SGraphPanel.h
@@ -157,6 +157,9 @@ public:
 	/** Straighten any connections attached to the specified pin, optionally limiting to the specified pin to align */
 	void StraightenConnections(UEdGraphPin* SourcePin, UEdGraphPin* PinToAlign = nullptr);
 
+	/** When the graph panel needs to be dynamically refreshing for animations, this function is registered to tick and invalidate the UI. */
+	EActiveTimerReturnType InvalidatePerTick(double InCurrentTime, float InDeltaTime);
+
 protected:
 
 	void NotifyGraphChanged ( const struct FEdGraphEditAction& InAction);
@@ -286,4 +289,10 @@ private:
 
 	/** Returns the pin that we're considering as hovered if we are hovering over a spline; may be null */
 	class SGraphPin* GetBestPinFromHoveredSpline() const;
+
+	/** Handle to timer callback that allows the UI to refresh it's arrangement each tick, allows animations to occur within the UI */
+	TWeakPtr<FActiveTimerHandle> ActiveTimerHandleInvalidatePerTick;
+
+	/** Amount of time left to invalidate the UI per tick */
+	float TimeLeftToInvalidatePerTick;
 };
diff --git a/Engine/Source/Editor/Kismet/Private/BlueprintDetailsCustomization.cpp b/Engine/Source/Editor/Kismet/Private/BlueprintDetailsCustomization.cpp
--- a/Engine/Source/Editor/Kismet/Private/BlueprintDetailsCustomization.cpp
+++ b/Engine/Source/Editor/Kismet/Private/BlueprintDetailsCustomization.cpp
@@ -554,6 +554,58 @@ void FBlueprintVarActionDetails::CustomizeDetails( IDetailLayoutBuilder& DetailL
 			.Font(DetailFontInfo)
 		]
 	];
+
+	TSharedPtr<SToolTip> BitmaskTooltip = IDocumentation::Get()->CreateToolTip(LOCTEXT("VarBitmaskTooltip", "Whether or not to treat this variable as a bitmask."), nullptr, DocLink, TEXT("Bitmask"));
+
+	Category.AddCustomRow(LOCTEXT("IsVariableBitmaskLabel", "Bitmask"))
+	.Visibility(TAttribute<EVisibility>(this, &FBlueprintVarActionDetails::BitmaskVisibility))
+	.NameContent()
+	[
+		SNew(STextBlock)
+		.Text(LOCTEXT("IsVariableBitmaskLabel", "Bitmask"))
+		.ToolTip(BitmaskTooltip)
+		.Font(IDetailLayoutBuilder::GetDetailFont())
+	]
+	.ValueContent()
+	[
+		SNew(SCheckBox)
+		.IsChecked(this, &FBlueprintVarActionDetails::OnBitmaskCheckboxState)
+		.OnCheckStateChanged(this, &FBlueprintVarActionDetails::OnBitmaskChanged)
+		.IsEnabled(IsVariableInBlueprint())
+		.ToolTip(BitmaskTooltip)
+	];
+
+	BitmaskEnumTypeNames.Empty();
+	BitmaskEnumTypeNames.Add(MakeShareable(new FString(LOCTEXT("BitmaskEnumTypeName_None", "None").ToString())));
+	for (TObjectIterator<UEnum> EnumIt; EnumIt; ++EnumIt)
+	{
+		UEnum* CurrentEnum = *EnumIt;
+		if (UEdGraphSchema_K2::IsAllowableBlueprintVariableType(CurrentEnum) && CurrentEnum->HasMetaData(TEXT("Bitflags")))
+		{
+			BitmaskEnumTypeNames.Add(MakeShareable(new FString(CurrentEnum->GetFName().ToString())));
+		}
+	}
+
+	TSharedPtr<SToolTip> BitmaskEnumTooltip = IDocumentation::Get()->CreateToolTip(LOCTEXT("VarBitmaskEnumTooltip", "If this is a bitmask, choose an optional enumeration type for the flags."), nullptr, DocLink, TEXT("Bitmask Flags"));
+	
+	Category.AddCustomRow(LOCTEXT("BitmaskEnumLabel", "Bitmask Enum"))
+	.Visibility(TAttribute<EVisibility>(this, &FBlueprintVarActionDetails::BitmaskVisibility))
+	.NameContent()
+	[
+		SNew(STextBlock)
+		.Text(LOCTEXT("BitmaskEnumLabel", "Bitmask Enum"))
+		.ToolTip(BitmaskEnumTooltip)
+		.Font(IDetailLayoutBuilder::GetDetailFont())
+	]
+	.ValueContent()
+	[
+		SNew(STextComboBox)
+		.OptionsSource(&BitmaskEnumTypeNames)
+		.InitiallySelectedItem(GetBitmaskEnumTypeName())
+		.OnSelectionChanged(this, &FBlueprintVarActionDetails::OnBitmaskEnumTypeChanged)
+		.IsEnabled(IsVariableInBlueprint() && OnBitmaskCheckboxState() == ECheckBoxState::Checked)
+	];
+
 	ReplicationOptions.Empty();
 	ReplicationOptions.Add(MakeShareable(new FString("None")));
 	ReplicationOptions.Add(MakeShareable(new FString("Replicated")));
@@ -1710,6 +1762,109 @@ EVisibility FBlueprintVarActionDetails::RangeVisibility() const
 	return EVisibility::Collapsed;
 }
 
+EVisibility FBlueprintVarActionDetails::BitmaskVisibility() const
+{
+	UProperty* VariableProperty = CachedVariableProperty.Get();
+	if (VariableProperty && IsABlueprintVariable(VariableProperty) && VariableProperty->IsA(UIntProperty::StaticClass()))
+	{
+		return EVisibility::Visible;
+	}
+
+	return EVisibility::Collapsed;
+}
+
+ECheckBoxState FBlueprintVarActionDetails::OnBitmaskCheckboxState() const
+{
+	UProperty* Property = CachedVariableProperty.Get();
+	if (Property)
+	{
+		return (Property && Property->HasMetaData(FBlueprintMetadata::MD_Bitmask)) ? ECheckBoxState::Checked : ECheckBoxState::Unchecked;
+	}
+	return ECheckBoxState::Unchecked;
+}
+
+void FBlueprintVarActionDetails::OnBitmaskChanged(ECheckBoxState InNewState)
+{
+	const FName VarName = CachedVariableName;
+	if (VarName != NAME_None)
+	{
+		const bool bIsBitmask = (InNewState == ECheckBoxState::Checked);
+		if (bIsBitmask)
+		{
+			FBlueprintEditorUtils::SetBlueprintVariableMetaData(GetBlueprintObj(), VarName, nullptr, FBlueprintMetadata::MD_Bitmask, TEXT(""));
+		}
+		else
+		{
+			FBlueprintEditorUtils::RemoveBlueprintVariableMetaData(GetBlueprintObj(), VarName, nullptr, FBlueprintMetadata::MD_Bitmask);
+		}
+
+		TArray<UK2Node_Variable*> VariableNodes;
+		FBlueprintEditorUtils::GetAllNodesOfClass<UK2Node_Variable>(GetBlueprintObj(), VariableNodes);
+
+		for (TArray<UK2Node_Variable*>::TConstIterator NodeIt(VariableNodes); NodeIt; ++NodeIt)
+		{
+			UK2Node_Variable* CurrentNode = *NodeIt;
+			if (VarName == CurrentNode->GetVarName())
+			{
+				CurrentNode->ReconstructNode();
+			}
+		}
+	}
+}
+
+TSharedPtr<FString> FBlueprintVarActionDetails::GetBitmaskEnumTypeName() const
+{
+	TSharedPtr<FString> Result;
+	const FName VarName = CachedVariableName;
+
+	if (BitmaskEnumTypeNames.Num() > 0 && VarName != NAME_None)
+	{
+		Result = BitmaskEnumTypeNames[0];
+
+		FString OutValue;
+		FBlueprintEditorUtils::GetBlueprintVariableMetaData(GetBlueprintObj(), VarName, nullptr, FBlueprintMetadata::MD_BitmaskEnum, OutValue);
+
+		for (int32 i = 1; i < BitmaskEnumTypeNames.Num(); ++i)
+		{
+			if (OutValue == *BitmaskEnumTypeNames[i])
+			{
+				Result = BitmaskEnumTypeNames[i];
+				break;
+			}
+		}
+	}
+
+	return Result;
+}
+
+void FBlueprintVarActionDetails::OnBitmaskEnumTypeChanged(TSharedPtr<FString> ItemSelected, ESelectInfo::Type SelectInfo)
+{
+	const FName VarName = CachedVariableName;
+	if (VarName != NAME_None)
+	{
+		if (ItemSelected == BitmaskEnumTypeNames[0])
+		{
+			FBlueprintEditorUtils::RemoveBlueprintVariableMetaData(GetBlueprintObj(), VarName, nullptr, FBlueprintMetadata::MD_BitmaskEnum);
+		}
+		else if(ItemSelected.IsValid())
+		{
+			FBlueprintEditorUtils::SetBlueprintVariableMetaData(GetBlueprintObj(), VarName, nullptr, FBlueprintMetadata::MD_BitmaskEnum, *ItemSelected);
+		}
+
+		TArray<UK2Node_Variable*> VariableNodes;
+		FBlueprintEditorUtils::GetAllNodesOfClass<UK2Node_Variable>(GetBlueprintObj(), VariableNodes);
+
+		for (TArray<UK2Node_Variable*>::TConstIterator NodeIt(VariableNodes); NodeIt; ++NodeIt)
+		{
+			UK2Node_Variable* CurrentNode = *NodeIt;
+			if (VarName == CurrentNode->GetVarName())
+			{
+				CurrentNode->ReconstructNode();
+			}
+		}
+	}
+}
+
 TSharedPtr<FString> FBlueprintVarActionDetails::GetVariableReplicationType() const
 {
 	EVariableReplication::Type VariableReplication = EVariableReplication::None;
diff --git a/Engine/Source/Editor/Kismet/Private/BlueprintDetailsCustomization.h b/Engine/Source/Editor/Kismet/Private/BlueprintDetailsCustomization.h
--- a/Engine/Source/Editor/Kismet/Private/BlueprintDetailsCustomization.h
+++ b/Engine/Source/Editor/Kismet/Private/BlueprintDetailsCustomization.h
@@ -155,6 +155,13 @@ private:
 	FText OnGetMetaKeyValue(FName Key) const;
 	void OnMetaKeyValueChanged(const FText& NewMinValue, ETextCommit::Type CommitInfo, FName Key);
 	EVisibility RangeVisibility() const;
+
+	ECheckBoxState OnBitmaskCheckboxState() const;
+	EVisibility BitmaskVisibility() const;
+	void OnBitmaskChanged(ECheckBoxState InNewState);
+
+	TSharedPtr<FString> GetBitmaskEnumTypeName() const;
+	void OnBitmaskEnumTypeChanged(TSharedPtr<FString> ItemSelected, ESelectInfo::Type SelectInfo);
 	
 	TSharedPtr<FString> GetVariableReplicationType() const;
 	void OnChangeReplication(TSharedPtr<FString> ItemSelected, ESelectInfo::Type SelectInfo);
@@ -200,6 +207,9 @@ private:
 	/** Array of replication options for our combo text box */
 	TArray<TSharedPtr<FString>> ReplicationOptions;
 
+	/** Array of enum type names for integers used as bitmasks */
+	TArray<TSharedPtr<FString>> BitmaskEnumTypeNames;
+
 	/** The widget used when in variable name editing mode */ 
 	TSharedPtr<SEditableTextBox> VarNameEditableTextBox;
 
diff --git a/Engine/Source/Editor/Kismet/Private/BlueprintEditor.cpp b/Engine/Source/Editor/Kismet/Private/BlueprintEditor.cpp
--- a/Engine/Source/Editor/Kismet/Private/BlueprintEditor.cpp
+++ b/Engine/Source/Editor/Kismet/Private/BlueprintEditor.cpp
@@ -510,7 +510,7 @@ bool FBlueprintEditor::IsInAScriptingMode() const
 
 bool FBlueprintEditor::OnRequestClose()
 {
-	if (IsBlueprintProfilerSupported() && GetDefault<UEditorExperimentalSettings>()->bBlueprintPerformanceAnalysisTools)
+	if (IsProfilerAvailable())
 	{
 		TabManager->InvokeTab(FBlueprintEditorTabs::BlueprintProfilerID)->RequestCloseTab();
 	}
@@ -1682,7 +1682,7 @@ void FBlueprintEditor::InitBlueprintEditor(
 	PostLayoutBlueprintEditorInitialization();
 
 	// Ensure the profiler UI respects the current state if it had previously been enabled.
-	if (IsBlueprintProfilerSupported())
+	if (IsProfilerAvailable())
 	{
 		bool bOpenProfilerTab = false;
 		if (GetDefault<UEditorExperimentalSettings>()->bBlueprintPerformanceAnalysisTools)
@@ -2486,7 +2486,7 @@ void FBlueprintEditor::CreateDefaultCommands()
 	// Blueprint Profiler Commands
 	ToolkitCommands->MapAction(FFullBlueprintEditorCommands::Get().ToggleProfiler,
 		FExecuteAction::CreateSP(this, &FBlueprintEditor::ToggleProfiler),
-		FCanExecuteAction(),
+		FCanExecuteAction::CreateSP(this, &FBlueprintEditor::IsProfilerAvailable),
 		FIsActionChecked::CreateSP(this, &FBlueprintEditor::IsProfilerActive));
 
 	// New document actions
@@ -2572,9 +2572,29 @@ void FBlueprintEditor::CreateDefaultCommands()
 		FCanExecuteAction::CreateSP(this, &FBlueprintEditor::CanGenerateNativeCode));
 }
 
+bool FBlueprintEditor::IsProfilerAvailable() const
+{
+	bool bIsBPProfilerSupported = false;
+	if (GetDefault<UEditorExperimentalSettings>()->bBlueprintPerformanceAnalysisTools)
+	{
+		if (GetToolkitFName() == FName("BlueprintEditor") && GetCurrentMode() == FBlueprintEditorApplicationModes::StandardBlueprintEditorMode)
+		{
+			if (auto Blueprint = GetBlueprintObj())
+			{
+				if (Blueprint->BlueprintType == BPTYPE_Normal || Blueprint->BlueprintType == BPTYPE_LevelScript)
+				{
+					bIsBPProfilerSupported = true;
+				}
+			}
+		}
+	}
+	return bIsBPProfilerSupported;
+
+}
+
 bool FBlueprintEditor::IsProfilerActive() const
 {
-	if (IsBlueprintProfilerSupported() && GetDefault<UEditorExperimentalSettings>()->bBlueprintPerformanceAnalysisTools)
+	if (IsProfilerAvailable())
 	{
 		IBlueprintProfilerInterface& ProfilerModule = FModuleManager::LoadModuleChecked<IBlueprintProfilerInterface>("BlueprintProfiler");
 		return ProfilerModule.IsProfilerEnabled();
@@ -2584,7 +2604,7 @@ bool FBlueprintEditor::IsProfilerActive() const
 
 void FBlueprintEditor::ToggleProfiler()
 {
-	if (IsBlueprintProfilerSupported() && GetDefault<UEditorExperimentalSettings>()->bBlueprintPerformanceAnalysisTools)
+	if (IsProfilerAvailable())
 	{
 		IBlueprintProfilerInterface& ProfilerModule = FModuleManager::LoadModuleChecked<IBlueprintProfilerInterface>("BlueprintProfiler");
 		ProfilerModule.ToggleProfilingCapture();
@@ -2600,22 +2620,6 @@ void FBlueprintEditor::ToggleProfiler()
 	}
 }
 
-bool FBlueprintEditor::IsBlueprintProfilerSupported() const
-{
-	bool bIsBPProfilerSupported = false;
-	if (GetToolkitFName() == FName("BlueprintEditor") && GetCurrentMode() == FBlueprintEditorApplicationModes::StandardBlueprintEditorMode)
-	{
-		if (auto Blueprint = GetBlueprintObj())
-		{
-			if (Blueprint->BlueprintType == BPTYPE_Normal || Blueprint->BlueprintType == BPTYPE_LevelScript)
-			{
-				bIsBPProfilerSupported = true;
-			}
-		}
-	}
-	return bIsBPProfilerSupported;
-}
-
 void FBlueprintEditor::OpenNativeCodeGenerationTool()
 {
 	auto Blueprint = GetBlueprintObj();
@@ -5474,6 +5478,13 @@ void FBlueprintEditor::PasteNodes()
 	}
 
 	PasteNodesHere(FocusedGraphEd->GetCurrentGraph(), FocusedGraphEd->GetPasteLocation());
+
+	// Dump any temporary pre-compile warnings to the compiler log.
+	UBlueprint* BlueprintObj = GetBlueprintObj();
+	if (BlueprintObj->PreCompileLog.IsValid())
+	{
+		DumpMessagesToCompilerLog(BlueprintObj->PreCompileLog->Messages, true);
+	}
 }
 
 
diff --git a/Engine/Source/Editor/Kismet/Private/FindInBlueprintManager.cpp b/Engine/Source/Editor/Kismet/Private/FindInBlueprintManager.cpp
--- a/Engine/Source/Editor/Kismet/Private/FindInBlueprintManager.cpp
+++ b/Engine/Source/Editor/Kismet/Private/FindInBlueprintManager.cpp
@@ -1641,16 +1641,25 @@ FString FFindInBlueprintSearchManager::GatherBlueprintSearchMetadata(const UBlue
 
 	// Gather all graph searchable data
 	TArray< UEdGraph* > SubGraphs;
-	BlueprintSearchMetaDataHelpers::GatherGraphSearchData(Writer, Blueprint, Blueprint->UbergraphPages, FFindInBlueprintSearchTags::FiB_UberGraphs, &SubGraphs);
-	BlueprintSearchMetaDataHelpers::GatherGraphSearchData(Writer, Blueprint, Blueprint->FunctionGraphs, FFindInBlueprintSearchTags::FiB_Functions, &SubGraphs);
-	BlueprintSearchMetaDataHelpers::GatherGraphSearchData(Writer, Blueprint, Blueprint->MacroGraphs, FFindInBlueprintSearchTags::FiB_Macros, &SubGraphs);
 
-	// Gather all interface graphs as functions
-	for (const FBPInterfaceDescription& InterfaceDesc : Blueprint->ImplementedInterfaces)
+	// Gather normal event graphs
+	BlueprintSearchMetaDataHelpers::GatherGraphSearchData(Writer, Blueprint, Blueprint->UbergraphPages, FFindInBlueprintSearchTags::FiB_UberGraphs, &SubGraphs);
+	
+	// We have interface graphs and function graphs to put into the Functions category. We cannot do them separately, so we must compile the full list
 	{
-		BlueprintSearchMetaDataHelpers::GatherGraphSearchData(Writer, Blueprint, InterfaceDesc.Graphs, FFindInBlueprintSearchTags::FiB_Functions, &SubGraphs);
+		TArray<UEdGraph*> CompleteGraphList;
+		CompleteGraphList.Append(Blueprint->FunctionGraphs);
+		// Gather all interface graphs as functions
+		for (const FBPInterfaceDescription& InterfaceDesc : Blueprint->ImplementedInterfaces)
+		{
+			CompleteGraphList.Append(InterfaceDesc.Graphs);
+		}
+		BlueprintSearchMetaDataHelpers::GatherGraphSearchData(Writer, Blueprint, CompleteGraphList, FFindInBlueprintSearchTags::FiB_Functions, &SubGraphs);
 	}
 
+	// Gather Macros
+	BlueprintSearchMetaDataHelpers::GatherGraphSearchData(Writer, Blueprint, Blueprint->MacroGraphs, FFindInBlueprintSearchTags::FiB_Macros, &SubGraphs);
+
 	// Sub graphs are processed separately so that they do not become children in the TreeView, cluttering things up if the tree is deep
 	BlueprintSearchMetaDataHelpers::GatherGraphSearchData(Writer, Blueprint, SubGraphs, FFindInBlueprintSearchTags::FiB_SubGraphs, NULL);
 
diff --git a/Engine/Source/Editor/Kismet/Private/FindInBlueprints.cpp b/Engine/Source/Editor/Kismet/Private/FindInBlueprints.cpp
--- a/Engine/Source/Editor/Kismet/Private/FindInBlueprints.cpp
+++ b/Engine/Source/Editor/Kismet/Private/FindInBlueprints.cpp
@@ -1250,7 +1250,10 @@ FText SFindInBlueprints::GetUncachedBlueprintWarningText() const
 	}
 	else
 	{
-		ReturnDisplayText = FText::Format(LOCTEXT("UncachedBlueprints", "Search incomplete. {Count} Blueprints are not indexed!"), Args);
+		Args.Add(TEXT("UnindexedCount"), FindInBlueprintManager.GetNumberUncachedBlueprints());
+		Args.Add(TEXT("OutOfDateCount"), OutOfDateWithLastSearchBPCount);
+
+		ReturnDisplayText = FText::Format(LOCTEXT("UncachedBlueprints", "Search incomplete. {Count} ({UnindexedCount} Unindexed/{OutOfDateCount} Out-of-Date) Blueprints need to be indexed!"), Args);
 
 		if (FailedToCacheCount > 0)
 		{
diff --git a/Engine/Source/Editor/Kismet/Private/Profiler/BPProfilerStatisticWidgets.cpp b/Engine/Source/Editor/Kismet/Private/Profiler/BPProfilerStatisticWidgets.cpp
--- a/Engine/Source/Editor/Kismet/Private/Profiler/BPProfilerStatisticWidgets.cpp
+++ b/Engine/Source/Editor/Kismet/Private/Profiler/BPProfilerStatisticWidgets.cpp
@@ -17,7 +17,6 @@ namespace BlueprintProfilerStatText
 	const FName ColumnId_Name("Name");
 	const FName ColumnId_InclusiveTime("InclusiveTime");
 	const FName ColumnId_Time("Time");
-	const FName ColumnId_PureTime("PureTime");
 	const FName ColumnId_MaxTime("MaxTime");
 	const FName ColumnId_MinTime("MinTime");
 	const FName ColumnId_Samples("Samples");
@@ -26,7 +25,6 @@ namespace BlueprintProfilerStatText
 	const FText ColumnText_Name(LOCTEXT("Name", "Name") );
 	const FText ColumnText_InclusiveTime(LOCTEXT("InclusiveTime", "Inclusive Time (ms)"));
 	const FText ColumnText_Time(LOCTEXT("Time", "Time (ms)"));
-	const FText ColumnText_PureTime(LOCTEXT("PureTime", "Pure Time (ms)"));
 	const FText ColumnText_MaxTime(LOCTEXT("MaxTime", "Max Time (ms)"));
 	const FText ColumnText_MinTime(LOCTEXT("MinTime", "Min Time (ms)"));
 	const FText ColumnText_Samples(LOCTEXT("Samples", "Samples"));
@@ -78,7 +76,6 @@ const FName SProfilerStatRow::GetStatName(const EBlueprintProfilerStat::Type Sta
 		case EBlueprintProfilerStat::TotalTime:		return BlueprintProfilerStatText::ColumnId_TotalTime;
 		case EBlueprintProfilerStat::InclusiveTime:	return BlueprintProfilerStatText::ColumnId_InclusiveTime;
 		case EBlueprintProfilerStat::Time:			return BlueprintProfilerStatText::ColumnId_Time;
-		case EBlueprintProfilerStat::PureTime:		return BlueprintProfilerStatText::ColumnId_PureTime;
 		case EBlueprintProfilerStat::MaxTime:		return BlueprintProfilerStatText::ColumnId_MaxTime;
 		case EBlueprintProfilerStat::MinTime:		return BlueprintProfilerStatText::ColumnId_MinTime;
 		case EBlueprintProfilerStat::Samples:		return BlueprintProfilerStatText::ColumnId_Samples;
@@ -94,7 +91,6 @@ const FText SProfilerStatRow::GetStatText(const EBlueprintProfilerStat::Type Sta
 		case EBlueprintProfilerStat::TotalTime:		return BlueprintProfilerStatText::ColumnText_TotalTime;
 		case EBlueprintProfilerStat::InclusiveTime:	return BlueprintProfilerStatText::ColumnText_InclusiveTime;
 		case EBlueprintProfilerStat::Time:			return BlueprintProfilerStatText::ColumnText_Time;
-		case EBlueprintProfilerStat::PureTime:		return BlueprintProfilerStatText::ColumnText_PureTime;
 		case EBlueprintProfilerStat::MaxTime:		return BlueprintProfilerStatText::ColumnText_MaxTime;
 		case EBlueprintProfilerStat::MinTime:		return BlueprintProfilerStatText::ColumnText_MinTime;
 		case EBlueprintProfilerStat::Samples:		return BlueprintProfilerStatText::ColumnText_Samples;
@@ -102,6 +98,155 @@ const FText SProfilerStatRow::GetStatText(const EBlueprintProfilerStat::Type Sta
 	}
 }
 
+//////////////////////////////////////////////////////////////////////////
+// FBPProfilerStatOptions
+
+void FBPProfilerStatDiplayOptions::SetActiveInstance(const FName InstanceName)
+{
+	if (ActiveInstance != InstanceName && HasFlags(ScopeToDebugInstance))
+	{
+		Flags |= Modified;
+	}
+	ActiveInstance = InstanceName;
+}
+
+void FBPProfilerStatDiplayOptions::SetActiveGraph(const FName GraphName)
+{
+	if (ActiveGraph != GraphName && HasFlags(GraphFilter))
+	{
+		Flags |= Modified;
+	}
+	ActiveGraph = GraphName;
+}
+
+TSharedRef<SWidget> FBPProfilerStatDiplayOptions::CreateToolbar()
+{
+	return	
+		SNew(SVerticalBox)
+		+SVerticalBox::Slot()
+		.HAlign(HAlign_Right)
+		.AutoHeight()
+		[
+			SNew(SHorizontalBox)
+			+SHorizontalBox::Slot()
+			.AutoWidth()
+			.Padding(FMargin(5,0))
+			[
+				SNew(SCheckBox)
+				.Content()
+				[
+					SNew(STextBlock)
+					.Text(LOCTEXT("FilterToGraph", "Filter to Graph"))
+				]
+				.IsChecked<FBPProfilerStatDiplayOptions, uint32>(this, &FBPProfilerStatDiplayOptions::GetChecked, GraphFilter)
+				.OnCheckStateChanged<FBPProfilerStatDiplayOptions, uint32>(this, &FBPProfilerStatDiplayOptions::OnChecked, GraphFilter)
+			]
+			+SHorizontalBox::Slot()
+			.AutoWidth()
+			.Padding(FMargin(5,0))
+			[
+				SNew(SCheckBox)
+				.Content()
+				[
+					SNew(STextBlock)
+					.Text(LOCTEXT("DisplayPureStats", "Pure Timings"))
+				]
+				.IsChecked<FBPProfilerStatDiplayOptions, uint32>(this, &FBPProfilerStatDiplayOptions::GetChecked, DisplayPure)
+				.OnCheckStateChanged<FBPProfilerStatDiplayOptions, uint32>(this, &FBPProfilerStatDiplayOptions::OnChecked, DisplayPure)
+			]
+			+SHorizontalBox::Slot()
+			.AutoWidth()
+			.Padding(FMargin(5,0))
+			[
+				SNew(SCheckBox)
+				.Content()
+				[
+					SNew(STextBlock)
+					.Text(LOCTEXT("ShowInstancesCheck", "Show Instances"))
+				]
+				.IsChecked<FBPProfilerStatDiplayOptions, uint32>(this, &FBPProfilerStatDiplayOptions::GetChecked, DisplayByInstance)
+				.OnCheckStateChanged<FBPProfilerStatDiplayOptions, uint32>(this, &FBPProfilerStatDiplayOptions::OnChecked, DisplayByInstance)
+			]
+			+SHorizontalBox::Slot()
+			.AutoWidth()
+			.Padding(FMargin(5,0))
+			[
+				SNew(SCheckBox)
+				.Content()
+				[
+					SNew(STextBlock)
+					.Text(LOCTEXT("InstanceFilterCheck", "Debug Filter Scope"))
+				]
+				.IsChecked<FBPProfilerStatDiplayOptions, uint32>(this, &FBPProfilerStatDiplayOptions::GetChecked, ScopeToDebugInstance)
+				.OnCheckStateChanged<FBPProfilerStatDiplayOptions, uint32>(this, &FBPProfilerStatDiplayOptions::OnChecked, ScopeToDebugInstance)
+			]
+			+SHorizontalBox::Slot()
+			.AutoWidth()
+			.Padding(FMargin(5,0))
+			[
+				SNew(SCheckBox)
+				.Content()
+				[
+					SNew(STextBlock)
+					.Text(LOCTEXT("AutoItemExpansion", "Auto Expand Statistics"))
+				]
+				.IsChecked<FBPProfilerStatDiplayOptions, uint32>(this, &FBPProfilerStatDiplayOptions::GetChecked, AutoExpand)
+				.OnCheckStateChanged<FBPProfilerStatDiplayOptions, uint32>(this, &FBPProfilerStatDiplayOptions::OnChecked, AutoExpand)
+			]
+		];
+}
+
+ECheckBoxState FBPProfilerStatDiplayOptions::GetChecked(const uint32 FlagsIn) const
+{
+	ECheckBoxState CheckedState;
+	if (FlagsIn & ScopeToDebugInstance)
+	{
+		if (HasFlags(DisplayByInstance))
+		{
+			CheckedState = HasFlags(ScopeToDebugInstance) ? ECheckBoxState::Checked : ECheckBoxState::Unchecked;
+		}
+		else
+		{
+			CheckedState = ECheckBoxState::Undetermined;
+		}
+	}
+	else
+	{
+		CheckedState = HasAllFlags(FlagsIn) ? ECheckBoxState::Checked : ECheckBoxState::Unchecked;
+	}
+	return CheckedState;
+}
+
+void FBPProfilerStatDiplayOptions::OnChecked(ECheckBoxState NewState, const uint32 FlagsIn)
+{
+	if (NewState == ECheckBoxState::Checked)
+	{
+		Flags |= FlagsIn;
+	}
+	else
+	{
+		Flags &= ~FlagsIn;
+	}
+	Flags |= Modified;
+}
+
+bool FBPProfilerStatDiplayOptions::IsFiltered(TSharedPtr<FScriptExecutionNode> Node) const
+{
+	bool bFilteredOut = !HasFlags(EDisplayFlags::DisplayPure) && Node->HasFlags(EScriptExecutionNodeFlags::PureStats);
+	if (Node->IsEvent() && HasFlags(EDisplayFlags::GraphFilter))
+	{
+		if (Node->GetGraphName() == UEdGraphSchema_K2::FN_UserConstructionScript)
+		{
+			bFilteredOut = ActiveGraph != UEdGraphSchema_K2::FN_UserConstructionScript;
+		}
+		else
+		{
+			bFilteredOut = ActiveGraph == UEdGraphSchema_K2::FN_UserConstructionScript;
+		}
+	}
+	return bFilteredOut;
+}
+
 //////////////////////////////////////////////////////////////////////////
 // FBPProfilerStatWidget
 
@@ -124,7 +269,11 @@ TSharedRef<SWidget> FBPProfilerStatWidget::GenerateColumnWidget(FName ColumnName
 				.Padding(FMargin(5,0))
 				[
 					SNew(SHyperlink)
+		#if TRACEPATH_DEBUG
+					.Text(FText::FromName(ExecNode->GetName()))
+		#else
 					.Text(ExecNode->GetDisplayName())
+		#endif
 					.Style(FEditorStyle::Get(), "HoverOnlyHyperlink")
 					.ToolTipText(ExecNode->GetToolTipText())
 					.OnNavigate(this, &FBPProfilerStatWidget::NavigateTo)
@@ -143,8 +292,7 @@ TSharedRef<SWidget> FBPProfilerStatWidget::GenerateColumnWidget(FName ColumnName
 		{
 			TAttribute<FText> TextAttr(LOCTEXT("NonApplicableStat", ""));
 			TAttribute<FSlateColor> ColorAttr;
-			const uint32 NonNodeStats = EScriptExecutionNodeFlags::Class|EScriptExecutionNodeFlags::Instance|EScriptExecutionNodeFlags::Event|
-										EScriptExecutionNodeFlags::BranchNode|EScriptExecutionNodeFlags::CallSite|EScriptExecutionNodeFlags::ExecPin;
+			const uint32 NonNodeStats = EScriptExecutionNodeFlags::Container|EScriptExecutionNodeFlags::CallSite|EScriptExecutionNodeFlags::BranchNode|EScriptExecutionNodeFlags::ExecPin;
 
 			if (ColumnName == BlueprintProfilerStatText::ColumnId_TotalTime)
 			{
@@ -156,7 +304,7 @@ TSharedRef<SWidget> FBPProfilerStatWidget::GenerateColumnWidget(FName ColumnName
 			}
 			else if (ColumnName == BlueprintProfilerStatText::ColumnId_InclusiveTime)
 			{
-				if (!ExecNode->HasFlags(NonNodeStats))
+				if (ExecNode->HasFlags(NonNodeStats))
 				{
 					TextAttr = TAttribute<FText>(PerformanceStats.Get(), &FScriptPerfData::GetInclusiveTimingText);
 					ColorAttr = TAttribute<FSlateColor>(PerformanceStats.Get(), &FScriptPerfData::GetInclusiveHeatColor);
@@ -170,14 +318,6 @@ TSharedRef<SWidget> FBPProfilerStatWidget::GenerateColumnWidget(FName ColumnName
 					ColorAttr = TAttribute<FSlateColor>(PerformanceStats.Get(), &FScriptPerfData::GetNodeHeatColor);
 				}
 			}
-			else if (ColumnName == BlueprintProfilerStatText::ColumnId_PureTime)
-			{
-				if (!ExecNode->HasFlags(NonNodeStats))
-				{
-					TextAttr = TAttribute<FText>(PerformanceStats.Get(), &FScriptPerfData::GetPureNodeTimingText);
-					ColorAttr = TAttribute<FSlateColor>(PerformanceStats.Get(), &FScriptPerfData::GetInclusiveHeatColor);
-				}
-			}
 			else if (ColumnName == BlueprintProfilerStatText::ColumnId_MaxTime)
 			{
 				if (!ExecNode->HasFlags(EScriptExecutionNodeFlags::ExecPin))
@@ -216,79 +356,92 @@ void FBPProfilerStatWidget::NavigateTo() const
 	}
 }
 
-void FBPProfilerStatWidget::GenerateExecNodeWidgets(const FName InstanceName, const FName FilterGraph)
+void FBPProfilerStatWidget::GenerateExecNodeWidgets(const TSharedPtr<FBPProfilerStatDiplayOptions> DisplayOptions)
 {
 	if (ExecNode.IsValid())
 	{
 		// Grab Performance Stats
-		PerformanceStats = ExecNode->GetPerfDataByInstanceAndTracePath(InstanceName, WidgetTracePath);
-		// Handle branch creation
+		PerformanceStats = ExecNode->GetPerfDataByInstanceAndTracePath(DisplayOptions->GetActiveInstance(), WidgetTracePath);
 		CachedChildren.Reset(0);
-		if (ExecNode->IsBranch())
-		{
-			for (auto LinkIter : ExecNode->GetLinkedNodes())
-			{
-				TArray<FScriptNodeExecLinkage::FLinearExecPath> LinearExecNodes;
-				FTracePath LinkPath(WidgetTracePath);
-				LinkPath.AddExitPin(LinkIter.Key);
-				LinkIter.Value->GetLinearExecutionPath(LinearExecNodes, LinkPath);
-				for (auto LinearPathIter : LinearExecNodes)
-				{
-					TSharedPtr<FBPProfilerStatWidget> NewLinkedNode = MakeShareable<FBPProfilerStatWidget>(new FBPProfilerStatWidget(LinearPathIter.LinkedNode, LinearPathIter.TracePath));
-					NewLinkedNode->GenerateExecNodeWidgets(InstanceName, FilterGraph);
-					CachedChildren.Add(NewLinkedNode);
-				}
-			}
-		}
-		else
+
+		if (ExecNode->HasFlags(EScriptExecutionNodeFlags::PureStats))
 		{
-			// @TODO - filter out pure nodes?
-			if (!ExecNode->IsPureNode())
+			if (ExecNode->IsPureChain())
 			{
 				// Get the full pure node chain associated with this exec node.
 				TMap<int32, TSharedPtr<FScriptExecutionNode>> AllPureNodes;
 				ExecNode->GetAllPureNodes(AllPureNodes);
 
-				// Sort pure nodes by script offset (execution order).
-				AllPureNodes.KeySort(TLess<int32>());
-
 				// Build trace path, tree view node widget and register perf stats for tracking.
 				FTracePath PureTracePath(WidgetTracePath);
 				for (auto Iter : AllPureNodes)
 				{
 					PureTracePath.AddExitPin(Iter.Key);
 					TSharedPtr<FBPProfilerStatWidget> NewPureChildNode = MakeShareable<FBPProfilerStatWidget>(new FBPProfilerStatWidget(Iter.Value, PureTracePath));
-					NewPureChildNode->GenerateExecNodeWidgets(InstanceName, FilterGraph);
-					CachedChildren.Add(NewPureChildNode);
+					NewPureChildNode->GenerateExecNodeWidgets(DisplayOptions);
+
+					// Pure nodes are shown in reverse execution order.
+					CachedChildren.Insert(NewPureChildNode, 0);
 				}
 			}
-
+		}
+		else
+		{
 			for (auto Iter : ExecNode->GetChildNodes())
 			{
 				// Filter out events based on graph
-				bool bFiltered = (Iter->IsEvent() && FilterGraph != NAME_None) ? IsEventFiltered(Iter, FilterGraph) : false;
-				if (!bFiltered)
+				if (!DisplayOptions->IsFiltered(Iter))
 				{
 					TArray<FScriptNodeExecLinkage::FLinearExecPath> LinearExecNodes;
 					FTracePath ChildTracePath(WidgetTracePath);
 					Iter->GetLinearExecutionPath(LinearExecNodes, ChildTracePath);
 					if (LinearExecNodes.Num() > 1)
 					{
+						TSharedPtr<FBPProfilerStatWidget> ChildContainer = AsShared();
 						for (auto LinearPathIter : LinearExecNodes)
 						{
-							TSharedPtr<FBPProfilerStatWidget> NewLinkedNode = MakeShareable<FBPProfilerStatWidget>(new FBPProfilerStatWidget(LinearPathIter.LinkedNode, LinearPathIter.TracePath));
-							NewLinkedNode->GenerateExecNodeWidgets(InstanceName, FilterGraph);
-							CachedChildren.Add(NewLinkedNode);
+							if (!DisplayOptions->IsFiltered(LinearPathIter.LinkedNode))
+							{
+								TSharedPtr<FBPProfilerStatWidget> NewLinkedNode = MakeShareable<FBPProfilerStatWidget>(new FBPProfilerStatWidget(LinearPathIter.LinkedNode, LinearPathIter.TracePath));
+								NewLinkedNode->GenerateExecNodeWidgets(DisplayOptions);
+								ChildContainer->CachedChildren.Add(NewLinkedNode);
+								if (LinearPathIter.LinkedNode->HasFlags(EScriptExecutionNodeFlags::Container))
+								{
+									ChildContainer = NewLinkedNode;
+								}
+							}
 						}
 					}
 					else
 					{
 						TSharedPtr<FBPProfilerStatWidget> NewChildNode = MakeShareable<FBPProfilerStatWidget>(new FBPProfilerStatWidget(Iter, ChildTracePath));
-						NewChildNode->GenerateExecNodeWidgets(InstanceName, FilterGraph);
+						NewChildNode->GenerateExecNodeWidgets(DisplayOptions);
 						CachedChildren.Add(NewChildNode);
 					}
 				}
 			}
+			if (ExecNode->IsBranch())
+			{
+				for (auto LinkIter : ExecNode->GetLinkedNodes())
+				{
+					if (!DisplayOptions->IsFiltered(LinkIter.Value))
+					{
+						TArray<FScriptNodeExecLinkage::FLinearExecPath> LinearExecNodes;
+						FTracePath LinkPath(WidgetTracePath);
+						if (!LinkIter.Value->HasFlags(EScriptExecutionNodeFlags::InvalidTrace))
+						{
+							LinkPath.AddExitPin(LinkIter.Key);
+						}
+						LinkIter.Value->GetLinearExecutionPath(LinearExecNodes, LinkPath);
+						for (auto LinearPathIter : LinearExecNodes)
+						{
+							TSharedPtr<FBPProfilerStatWidget> NewLinkedNode = MakeShareable<FBPProfilerStatWidget>(new FBPProfilerStatWidget(LinearPathIter.LinkedNode, LinearPathIter.TracePath));
+							NewLinkedNode->GenerateExecNodeWidgets(DisplayOptions);
+							CachedChildren.Add(NewLinkedNode);
+						}
+					}
+				}
+			}
 		}
 	}
 }
@@ -363,21 +516,4 @@ bool FBPProfilerStatWidget::ProbeChildWidgetExpansionStates()
 	return bIsExpanded;
 }
 
-bool FBPProfilerStatWidget::IsEventFiltered(TSharedPtr<FScriptExecutionNode> InEventNode, FName GraphFilter) const
-{
-	bool bEventFiltered = false;
-	if (InEventNode.IsValid() && InEventNode->IsEvent())
-	{
-		if (InEventNode->GetGraphName() == UEdGraphSchema_K2::FN_UserConstructionScript)
-		{
-			bEventFiltered = GraphFilter != UEdGraphSchema_K2::FN_UserConstructionScript;
-		}
-		else
-		{
-			bEventFiltered = GraphFilter == UEdGraphSchema_K2::FN_UserConstructionScript;
-		}
-	}
-	return bEventFiltered;
-}
-
 #undef LOCTEXT_NAMESPACE
diff --git a/Engine/Source/Editor/Kismet/Private/Profiler/BPProfilerStatisticWidgets.h b/Engine/Source/Editor/Kismet/Private/Profiler/BPProfilerStatisticWidgets.h
--- a/Engine/Source/Editor/Kismet/Private/Profiler/BPProfilerStatisticWidgets.h
+++ b/Engine/Source/Editor/Kismet/Private/Profiler/BPProfilerStatisticWidgets.h
@@ -53,6 +53,81 @@ public:
 
 };
 
+//////////////////////////////////////////////////////////////////////////
+// FBPProfilerStatDiplayOptions
+
+struct FBPProfilerStatDiplayOptions : public TSharedFromThis<FBPProfilerStatDiplayOptions>
+{
+public:
+
+	/** Display flags controlling widget appearance */
+	enum EDisplayFlags
+	{
+		DisplayByInstance			= 0x00000001,	// Show instance stats
+		ScopeToDebugInstance		= 0x00000002,	// Scope stats to current debug instance
+		GraphFilter					= 0x00000004,	// Filter to current graph
+		DisplayPure					= 0x00000008,	// Display pure stats
+		AutoExpand					= 0x00000010,	// Auto expand stats
+		Modified					= 0x10000000,	// Display state was changed
+		Default						= DisplayByInstance|ScopeToDebugInstance|DisplayPure|GraphFilter|Modified
+	};
+
+	FBPProfilerStatDiplayOptions()
+		: ActiveInstance(NAME_None)
+		, ActiveGraph(NAME_None)
+		, Flags(Default)
+	{
+	}
+
+	/** Marks display state as modified */
+	void SetStateModified() { Flags |= Modified; }
+
+	/** Returns if the filter state options are modified */
+	bool IsStateModified() const { return (Flags & Modified) != 0U; }
+
+	/** Add flags */
+	void AddFlags(const uint32 FlagsIn) { Flags |= FlagsIn; }
+
+	/** Clear flags */
+	void ClearFlags(const uint32 FlagsIn) { Flags &= ~FlagsIn; }
+
+	/** Returns true if any flags match */
+	bool HasFlags(const uint32 FlagsIn) const { return (Flags & FlagsIn) != 0U; }
+
+	/** Returns true if all flags match */
+	bool HasAllFlags(const uint32 FlagsIn) const { return (Flags & FlagsIn) == FlagsIn; }
+
+	/** Set active instance name */
+	FName GetActiveInstance() const { return ActiveInstance; }
+
+	/** Set active instance name */
+	void SetActiveInstance(const FName InstanceName);
+
+	/** Set active graph name */
+	void SetActiveGraph(const FName GraphName);
+
+	/** Create toolbar widget */
+	TSharedRef<SWidget> CreateToolbar();
+
+	/** Returns if the UI should be checked for flags */
+	ECheckBoxState GetChecked(const uint32 FlagsIn) const;
+
+	/** Handles UI check states */
+	void OnChecked(ECheckBoxState NewState, const uint32 FlagsIn);
+
+	/** Returns true if the exec node passes current filters */
+	bool IsFiltered(TSharedPtr<class FScriptExecutionNode> Node) const;
+
+private:
+
+	/** The active instance name */
+	FName ActiveInstance;
+	/** The active graph name */
+	FName ActiveGraph;
+	/** The display flags */
+	uint32 Flags;
+};
+
 //////////////////////////////////////////////////////////////////////////
 // FBPProfilerStatWidget
 
@@ -69,7 +144,7 @@ public:
 	virtual ~FBPProfilerStatWidget() {}
 
 	/** Generate exec node widgets */
-	virtual void GenerateExecNodeWidgets(const FName InstanceName, const FName FilterGraph);
+	virtual void GenerateExecNodeWidgets(const TSharedPtr<FBPProfilerStatDiplayOptions> DisplayOptions);
 
 	/** Gather children for list/tree widget */
 	void GatherChildren(TArray<TSharedPtr<FBPProfilerStatWidget>>& OutChildren);
@@ -100,9 +175,6 @@ public:
 
 protected:
 
-	/** Returns if the exec node representing an event is filtered from view */
-	bool IsEventFiltered(TSharedPtr<FScriptExecutionNode> InEventNode, FName GraphFilter) const;
-
 	/** Navigate to referenced node */
 	void NavigateTo() const;
 
@@ -114,8 +186,6 @@ protected:
 	TSharedPtr<FScriptExecutionNode> ExecNode;
 	/** Script performance data */
 	TSharedPtr<FScriptPerfData> PerformanceStats;
-	/** Script code offset */
-	int32 ScriptCodeOffset;
 	/** Child statistics */
 	TArray<TSharedPtr<FBPProfilerStatWidget>> CachedChildren;
 
diff --git a/Engine/Source/Editor/Kismet/Private/Profiler/EventExecution.cpp b/Engine/Source/Editor/Kismet/Private/Profiler/EventExecution.cpp
--- a/Engine/Source/Editor/Kismet/Private/Profiler/EventExecution.cpp
+++ b/Engine/Source/Editor/Kismet/Private/Profiler/EventExecution.cpp
@@ -7,18 +7,22 @@
 //////////////////////////////////////////////////////////////////////////
 // FScriptNodeExecLinkage
 
-void FScriptNodeExecLinkage::AddPureNode(const int32 PinScriptOffset, TSharedPtr<class FScriptExecutionNode> PureNode)
-{
-	TSharedPtr<class FScriptExecutionNode>& Result = PureNodes.FindOrAdd(PinScriptOffset);
-	Result = PureNode;
-}
-
 void FScriptNodeExecLinkage::AddLinkedNode(const int32 PinScriptOffset, TSharedPtr<class FScriptExecutionNode> LinkedNode)
 {
 	TSharedPtr<class FScriptExecutionNode>& Result = LinkedNodes.FindOrAdd(PinScriptOffset);
 	Result = LinkedNode;
 }
 
+TSharedPtr<FScriptExecutionNode> FScriptNodeExecLinkage::GetLinkedNodeByScriptOffset(const int32 PinScriptOffset)
+{
+	TSharedPtr<FScriptExecutionNode> Result;
+	if (TSharedPtr<FScriptExecutionNode>* SearchResult = LinkedNodes.Find(PinScriptOffset))
+	{
+		Result = *SearchResult;
+	}
+	return Result;
+}
+
 //////////////////////////////////////////////////////////////////////////
 // FScriptNodePerfData
 
@@ -99,24 +103,67 @@ bool FScriptExecutionNode::operator == (const FScriptExecutionNode& NodeIn) cons
 
 void FScriptExecutionNode::GetLinearExecutionPath(TArray<FLinearExecPath>& LinearExecutionNodes, const FTracePath& TracePath)
 {
-	FExecNodeFilter Filter;
-	GetLinearExecutionPath_Internal(Filter, LinearExecutionNodes, TracePath);
+	LinearExecutionNodes.Add(FLinearExecPath(AsShared(), TracePath));
+	FTracePath NewTracePath(TracePath);
+	if (HasFlags(EScriptExecutionNodeFlags::FunctionTunnel))
+	{
+		MapTunnelLinearExecution(NewTracePath);
+	}
+	if (GetNumLinkedNodes() == 1)
+	{
+		for (auto NodeIter : LinkedNodes)
+		{
+			if (HasFlags(EScriptExecutionNodeFlags::PureStats))
+			{
+				continue;
+			}
+			else
+			{
+				if (NodeIter.Key != INDEX_NONE && !HasFlags(EScriptExecutionNodeFlags::InvalidTrace))
+				{
+					NewTracePath.AddExitPin(NodeIter.Key);
+				}
+				if (NodeIter.Value->HasFlags(EScriptExecutionNodeFlags::FunctionTunnel))
+				{
+					LinearExecutionNodes.Add(FLinearExecPath(NodeIter.Value, NewTracePath));
+					for (auto TunnelExit : NodeIter.Value->LinkedNodes)
+					{
+						if (TunnelExit.Value->HasFlags(EScriptExecutionNodeFlags::TunnelThenPin))
+						{
+							MapTunnelLinearExecution(NewTracePath);
+							NewTracePath.AddExitPin(TunnelExit.Key);
+							TunnelExit.Value->GetLinearExecutionPath(LinearExecutionNodes, NewTracePath);
+							break;
+						}
+					}
+				}
+				else
+				{
+					NodeIter.Value->GetLinearExecutionPath(LinearExecutionNodes, NewTracePath);
+				}
+			}
+		}
+	}
 }
 
-void FScriptExecutionNode::GetLinearExecutionPath_Internal(FExecNodeFilter& Filter, TArray<FLinearExecPath>& LinearExecutionNodes, const FTracePath& TracePath)
+void FScriptExecutionNode::MapTunnelLinearExecution(FTracePath& Trace) const
 {
-	if (!Filter.Contains(ObservedObject))
+	if (HasFlags(EScriptExecutionNodeFlags::FunctionTunnel))
+	{
+		for (auto ChildIter : ChildNodes)
+		{
+			ChildIter->MapTunnelLinearExecution(Trace);
+		}
+	}
+	else if (GetNumLinkedNodes() == 1)
 	{
-		Filter.Add(ObservedObject);
-		LinearExecutionNodes.Add(FLinearExecPath(AsShared(), TracePath));
-		if (GetNumLinkedNodes() == 1)
+		for (auto NodeIter : LinkedNodes)
 		{
-			FTracePath NewTracePath(TracePath);
-			for (auto NodeIter : LinkedNodes)
+			if (!NodeIter.Value->HasFlags(EScriptExecutionNodeFlags::InvalidTrace))
 			{
-				NewTracePath.AddExitPin(NodeIter.Key);
-				NodeIter.Value->GetLinearExecutionPath_Internal(Filter, LinearExecutionNodes, NewTracePath);
+				Trace.AddExitPin(NodeIter.Key);
 			}
+			NodeIter.Value->MapTunnelLinearExecution(Trace);
 		}
 	}
 }
@@ -124,7 +171,11 @@ void FScriptExecutionNode::GetLinearExecutionPath_Internal(FExecNodeFilter& Filt
 EScriptStatContainerType::Type FScriptExecutionNode::GetStatisticContainerType() const
 {
 	EScriptStatContainerType::Type Result = EScriptStatContainerType::Standard;
-	if (HasFlags(EScriptExecutionNodeFlags::Instance|EScriptExecutionNodeFlags::Event|EScriptExecutionNodeFlags::CallSite))
+	if (HasFlags(EScriptExecutionNodeFlags::ExecPin|EScriptExecutionNodeFlags::TunnelPin))
+	{
+		Result = EScriptStatContainerType::NewExecutionPath;
+	}
+	else if (HasFlags(EScriptExecutionNodeFlags::Container))
 	{
 		Result = EScriptStatContainerType::Container;
 	}
@@ -132,136 +183,128 @@ EScriptStatContainerType::Type FScriptExecutionNode::GetStatisticContainerType()
 	{
 		Result = EScriptStatContainerType::SequentialBranch;
 	}
-	else if (HasFlags(EScriptExecutionNodeFlags::ExecPin))
-	{
-		Result = EScriptStatContainerType::NewExecutionPath;
-	}
 	return Result;
 }
 
 void FScriptExecutionNode::RefreshStats(const FTracePath& TracePath)
 {
-	FExecNodeFilter StatFilter;
-	RefreshStats_Internal(TracePath, StatFilter);
-}
-
-void FScriptExecutionNode::RefreshStats_Internal(const FTracePath& TracePath, FExecNodeFilter& VisitedStats)
-{
-	if (!VisitedStats.Find(ObservedObject))
+	// Process stat update
+	bool bRefreshChildStats = false;
+	bool bRefreshLinkStats = true;
+	TArray<TSharedPtr<FScriptPerfData>> BlueprintPooledStats;
+	// Update stats based on type
+	switch(GetStatisticContainerType())
 	{
-		// Update visited stats
-		if (!HasFlags(EScriptExecutionNodeFlags::Container))
+		case EScriptStatContainerType::Container:
 		{
-			VisitedStats.Add(ObservedObject);
+			GetInstancePerfDataByTracePath(TracePath, BlueprintPooledStats);
+			bRefreshChildStats = true;
+			break;
 		}
-		// Process stat update
-		bool bRefreshChildStats = false;
-		bool bRefreshLinkStats = true;
-		TArray<TSharedPtr<FScriptPerfData>> BlueprintPooledStats;
-		// Update stats based on type
-		switch(GetStatisticContainerType())
+		case EScriptStatContainerType::Standard:
+		case EScriptStatContainerType::SequentialBranch:
 		{
-			case EScriptStatContainerType::Container:
-			{
-				// Only consider stat data on this path
-				GetInstancePerfDataByTracePath(TracePath, BlueprintPooledStats);
-				bRefreshChildStats = true;
-				break;
-			}
-			case EScriptStatContainerType::Standard:
-			case EScriptStatContainerType::SequentialBranch:
+			GetInstancePerfDataByTracePath(TracePath, BlueprintPooledStats);
+
+			// Refresh attached pure chain stats (if present)
+			TSharedPtr<FScriptExecutionNode> PureChainNode = GetPureChainNode();
+			if (PureChainNode.IsValid())
 			{
-				// Only consider stat data on this path
-				GetInstancePerfDataByTracePath(TracePath, BlueprintPooledStats);
-				break;
+				PureChainNode->RefreshStats(TracePath);
 			}
-			case EScriptStatContainerType::NewExecutionPath:
+			break;
+		}
+		case EScriptStatContainerType::NewExecutionPath:
+		{
+			TSet<FName> AllInstances;
+			// Refresh child stats and copy as branch stats - this is a dummy entry
+			for (auto ChildIter : ChildNodes)
 			{
-				TSet<FName> AllInstances;
-				// Refresh child stats and copy as branch stats - this is a dummy entry
-				for (auto ChildIter : ChildNodes)
+				// Fill out instance data that can be missing because its a dummy node, or the code below can fail.
+				for (auto InstanceIter : ChildIter->InstanceInputPinToPerfDataMap)
 				{
-					// Fill out instance data that can be missing because its a dummy node, or the code below can fail.
-					for (auto InstanceIter : ChildIter->InstanceInputPinToPerfDataMap)
+					if (InstanceIter.Key != NAME_None)
 					{
-						if (InstanceIter.Key != NAME_None)
-						{
-							AllInstances.Add(InstanceIter.Key);
-						}
+						AllInstances.Add(InstanceIter.Key);
 					}
-					ChildIter->RefreshStats_Internal(TracePath, VisitedStats);
 				}
-				for (auto InstanceIter : AllInstances)
+				ChildIter->RefreshStats(TracePath);
+			}
+			for (auto InstanceIter : AllInstances)
+			{
+				if (InstanceIter != NAME_None)
 				{
-					if (InstanceIter != NAME_None)
-					{
-						TSharedPtr<FScriptPerfData> InstancePerfData = GetPerfDataByInstanceAndTracePath(InstanceIter, TracePath);
-						InstancePerfData->Reset();
-						BlueprintPooledStats.Add(InstancePerfData);
+					TSharedPtr<FScriptPerfData> InstancePerfData = GetPerfDataByInstanceAndTracePath(InstanceIter, TracePath);
+					InstancePerfData->Reset();
+					BlueprintPooledStats.Add(InstancePerfData);
 
-						for (auto ChildIter : ChildNodes)
-						{
-							TSharedPtr<FScriptPerfData> ChildPerfData = ChildIter->GetPerfDataByInstanceAndTracePath(InstanceIter, TracePath);
-							InstancePerfData->AddBranchData(*ChildPerfData.Get());
-						}
+					for (auto ChildIter : ChildNodes)
+					{
+						TSharedPtr<FScriptPerfData> ChildPerfData = ChildIter->GetPerfDataByInstanceAndTracePath(InstanceIter, TracePath);
+						InstancePerfData->AddBranchData(*ChildPerfData.Get());
 					}
 				}
-				// Update Blueprint stats as branches
-				TSharedPtr<FScriptPerfData> BlueprintData = GetBlueprintPerfDataByTracePath(TracePath);
-				BlueprintData->Reset();
+			}
+			// Update Blueprint stats as branches
+			TSharedPtr<FScriptPerfData> BlueprintData = GetBlueprintPerfDataByTracePath(TracePath);
+			BlueprintData->Reset();
 
-				for (auto BlueprintChildDataIter : BlueprintPooledStats)
-				{
-					BlueprintData->AddBranchData(*BlueprintChildDataIter.Get());
-				}
-				BlueprintPooledStats.Reset(0);
-				bRefreshLinkStats = false;
-				break;
+			for (auto BlueprintChildDataIter : BlueprintPooledStats)
+			{
+				BlueprintData->AddBranchData(*BlueprintChildDataIter.Get());
 			}
+			BlueprintPooledStats.Reset(0);
+			bRefreshLinkStats = false;
+			break;
+		}
+	}
+	// Refresh Child Links
+	if (bRefreshChildStats)
+	{
+		for (auto ChildIter : ChildNodes)
+		{
+			ChildIter->RefreshStats(TracePath);
 		}
-		// Refresh Child Links
-		if (bRefreshChildStats)
+	}
+	// Refresh All links
+	if (bRefreshLinkStats)
+	{
+		if (IsPureChain())
 		{
-			for (auto ChildIter : ChildNodes)
+			TMap<int32, TSharedPtr<FScriptExecutionNode>> AllPureNodes;
+			GetAllPureNodes(AllPureNodes);
+
+			FTracePath PureTracePath(TracePath);
+			for (auto PureIter : AllPureNodes)
 			{
-				ChildIter->RefreshStats_Internal(TracePath, VisitedStats);
+				PureTracePath.AddExitPin(PureIter.Key);
+				PureIter.Value->RefreshStats(PureTracePath);
 			}
 		}
-		// Refresh All links
-		if (bRefreshLinkStats)
+		else
 		{
-			if (!IsPureNode())
-			{
-				TMap<int32, TSharedPtr<FScriptExecutionNode>> AllPureNodes;
-				GetAllPureNodes(AllPureNodes);
-
-				AllPureNodes.KeySort(TLess<int32>());
-
-				FTracePath PureTracePath(TracePath);
-				for (auto PureIter : AllPureNodes)
-				{
-					PureTracePath.AddExitPin(PureIter.Key);
-					PureIter.Value->RefreshStats(PureTracePath);
-				}
-			}
-
 			for (auto LinkIter : LinkedNodes)
 			{
 				FTracePath LinkTracePath(TracePath);
+				if (HasFlags(EScriptExecutionNodeFlags::FunctionTunnel))
+				{
+					MapTunnelLinearExecution(LinkTracePath);
+				}
+				//if (!LinkIter.Value->HasFlags(EScriptExecutionNodeFlags::DummyTrace))
 				LinkTracePath.AddExitPin(LinkIter.Key);
-				LinkIter.Value->RefreshStats_Internal(LinkTracePath, VisitedStats);
+				LinkIter.Value->RefreshStats(LinkTracePath);
 			}
 		}
-		// Update the owning blueprint stats
-		if (BlueprintPooledStats.Num() > 0)
-		{
-			TSharedPtr<FScriptPerfData> BlueprintData = GetBlueprintPerfDataByTracePath(TracePath);
-			BlueprintData->Reset();
+	}
+	// Update the owning blueprint stats
+	if (BlueprintPooledStats.Num() > 0)
+	{
+		TSharedPtr<FScriptPerfData> BlueprintData = GetBlueprintPerfDataByTracePath(TracePath);
+		BlueprintData->Reset();
 
-			for (auto BlueprintChildDataIter : BlueprintPooledStats)
-			{
-				BlueprintData->AddData(*BlueprintChildDataIter.Get());
-			}
+		for (auto BlueprintChildDataIter : BlueprintPooledStats)
+		{
+			BlueprintData->AddData(*BlueprintChildDataIter.Get());
 		}
 	}
 }
@@ -283,21 +326,39 @@ void FScriptExecutionNode::GetAllExecNodes(TMap<FName, TSharedPtr<FScriptExecuti
 void FScriptExecutionNode::GetAllPureNodes(TMap<int32, TSharedPtr<class FScriptExecutionNode>>& PureNodesOut)
 {
 	GetAllPureNodes_Internal(PureNodesOut, PureNodeScriptCodeRange);
+
+	// Sort pure nodes by script offset (execution order).
+	PureNodesOut.KeySort(TLess<int32>());
 }
 
 void FScriptExecutionNode::GetAllPureNodes_Internal(TMap<int32, TSharedPtr<class FScriptExecutionNode>>& PureNodesOut, const FInt32Range& ScriptCodeRange)
 {
-	for (auto PureIter : PureNodes)
+	for (auto LinkIter : LinkedNodes)
 	{
-		PureIter.Value->GetAllPureNodes_Internal(PureNodesOut, ScriptCodeRange);
+		LinkIter.Value->GetAllPureNodes_Internal(PureNodesOut, ScriptCodeRange);
 
-		if (ScriptCodeRange.Contains(PureIter.Key))
+		if (LinkIter.Value->IsPureNode() && ScriptCodeRange.Contains(LinkIter.Key))
 		{
-			PureNodesOut.Add(PureIter.Key, PureIter.Value);
+			PureNodesOut.Add(LinkIter.Key, LinkIter.Value);
 		}
 	}
 }
 
+TSharedPtr<FScriptExecutionNode> FScriptExecutionNode::GetPureChainNode()
+{
+	TSharedPtr<FScriptExecutionNode> Result;
+
+	for (auto ChildIter : ChildNodes)
+	{
+		if (ChildIter->IsPureChain())
+		{
+			Result = ChildIter;
+		}
+	}
+
+	return Result;
+}
+
 void FScriptExecutionNode::NavigateToObject() const
 {
 	if (ObservedObject.IsValid())
diff --git a/Engine/Source/Editor/Kismet/Private/Profiler/SGraphExecutionStatDisplay.cpp b/Engine/Source/Editor/Kismet/Private/Profiler/SGraphExecutionStatDisplay.cpp
--- a/Engine/Source/Editor/Kismet/Private/Profiler/SGraphExecutionStatDisplay.cpp
+++ b/Engine/Source/Editor/Kismet/Private/Profiler/SGraphExecutionStatDisplay.cpp
@@ -27,6 +27,10 @@ SGraphExecutionStatDisplay::~SGraphExecutionStatDisplay()
 void SGraphExecutionStatDisplay::Construct(const FArguments& InArgs)
 {	
 	BlueprintEditor = InArgs._AssetEditor;
+	if (!DisplayOptions.IsValid())
+	{
+		DisplayOptions = MakeShareable(new FBPProfilerStatDiplayOptions);
+	}
 	// Register for profiling toggle events
 	FBlueprintCoreDelegates::OnToggleScriptProfiler.AddSP(this, &SGraphExecutionStatDisplay::OnToggleProfiler);
 	// Remove delegate for graph structural changes
@@ -34,11 +38,6 @@ void SGraphExecutionStatDisplay::Construct(const FArguments& InArgs)
 	{
 		Profiler->GetGraphLayoutChangedDelegate().AddSP(this, &SGraphExecutionStatDisplay::OnGraphLayoutChanged);
 	}
-	// Set Default States
-	bDisplayInstances = false;
-	bScopeToActiveDebugInstance = false;
-	bAutoExpandStats = false;
-	bFilterEventsToGraph = true;
 
 	ChildSlot
 	[
@@ -53,65 +52,7 @@ void SGraphExecutionStatDisplay::Construct(const FArguments& InArgs)
 			.Padding(4)
 			.BorderImage(FEditorStyle::GetBrush("BlueprintProfiler.ViewToolBar"))
 			[
-				SNew(SVerticalBox)
-				+SVerticalBox::Slot()
-				.HAlign(HAlign_Right)
-				.AutoHeight()
-				[
-					SNew(SHorizontalBox)
-					+SHorizontalBox::Slot()
-					.AutoWidth()
-					.Padding(FMargin(5,0))
-					[
-						SNew(SCheckBox)
-						.Content()
-						[
-							SNew(STextBlock)
-							.Text(LOCTEXT("FilterToGraph", "Filter Events to Current Graph"))
-						]
-						.IsChecked(this, &SGraphExecutionStatDisplay::GetGraphFilterCheck)
-						.OnCheckStateChanged(this, &SGraphExecutionStatDisplay::OnGraphFilterCheck)
-					]
-					+SHorizontalBox::Slot()
-					.AutoWidth()
-					.Padding(FMargin(5,0))
-					[
-						SNew(SCheckBox)
-						.Content()
-						[
-							SNew(STextBlock)
-							.Text(LOCTEXT("ShowInstancesCheck", "Display By Instance"))
-						]
-						.IsChecked(this, &SGraphExecutionStatDisplay::GetDisplayByInstanceCheck)
-						.OnCheckStateChanged(this, &SGraphExecutionStatDisplay::OnDisplayByInstanceCheck)
-					]
-					+SHorizontalBox::Slot()
-					.AutoWidth()
-					.Padding(FMargin(5,0))
-					[
-						SNew(SCheckBox)
-						.Content()
-						[
-							SNew(STextBlock)
-							.Text(LOCTEXT("InstanceFilterCheck", "Scope To Debug Filter"))
-						]
-						.IsChecked(this, &SGraphExecutionStatDisplay::GetInstanceFilterCheck)
-						.OnCheckStateChanged(this, &SGraphExecutionStatDisplay::OnInstanceFilterCheck)
-					]
-					+SHorizontalBox::Slot()
-					.AutoWidth()
-					.Padding(FMargin(5,0))
-					[
-						SNew(SCheckBox)
-						.Content()
-						[
-							SNew(STextBlock)
-							.Text(LOCTEXT("AutoItemExpansion", "Auto Expand Statistics"))
-						]
-						.IsChecked(this, &SGraphExecutionStatDisplay::GetAutoExpansionCheck)
-						.OnCheckStateChanged(this, &SGraphExecutionStatDisplay::OnAutoExpansionCheck)
-					]
-				]
+				DisplayOptions->CreateToolbar()
 			]
 		]
 		+SVerticalBox::Slot()
@@ -136,9 +77,6 @@ void SGraphExecutionStatDisplay::Construct(const FArguments& InArgs)
 					+SHeaderRow::Column(SProfilerStatRow::GetStatName(EBlueprintProfilerStat::Time))
 					.DefaultLabel(SProfilerStatRow::GetStatText(EBlueprintProfilerStat::Time))
 					.ManualWidth(70)
-					+SHeaderRow::Column(SProfilerStatRow::GetStatName(EBlueprintProfilerStat::PureTime))
-					.DefaultLabel(SProfilerStatRow::GetStatText(EBlueprintProfilerStat::PureTime))
-					.ManualWidth(100)
 					+SHeaderRow::Column(SProfilerStatRow::GetStatName(EBlueprintProfilerStat::InclusiveTime))
 					.DefaultLabel(SProfilerStatRow::GetStatText(EBlueprintProfilerStat::InclusiveTime))
 					.ManualWidth(120)
@@ -165,14 +103,14 @@ void SGraphExecutionStatDisplay::OnGraphLayoutChanged(TWeakObjectPtr<UBlueprint>
 	if (CurrentBlueprint == Blueprint)
 	{
 		// Force a stat update.
-		CurrentBlueprint.Reset();
+		DisplayOptions->SetStateModified();
 	}
 }
 
 void SGraphExecutionStatDisplay::OnToggleProfiler(bool bEnabled)
 {
 	// Force a stat update.
-	CurrentBlueprint.Reset();
+	DisplayOptions->SetStateModified();
 }
 
 TSharedRef<ITableRow> SGraphExecutionStatDisplay::OnGenerateRow(FBPStatWidgetPtr InItem, const TSharedRef<STableViewBase>& OwnerTable)
@@ -188,63 +126,6 @@ void SGraphExecutionStatDisplay::OnGetChildren(FBPStatWidgetPtr InParent, TArray
 	}
 }
 
-void SGraphExecutionStatDisplay::OnGraphFilterCheck(ECheckBoxState NewState)
-{
-	const bool bOldState = bFilterEventsToGraph;
-	bFilterEventsToGraph = (NewState == ECheckBoxState::Checked) ? true : false;
-	// Force Update.
-	if (bFilterEventsToGraph != bOldState)
-	{
-		CurrentBlueprint.Reset();
-	}
-}
-
-void SGraphExecutionStatDisplay::OnInstanceFilterCheck(ECheckBoxState NewState)
-{
-	const bool bOldState = bScopeToActiveDebugInstance;
-	bScopeToActiveDebugInstance = (NewState == ECheckBoxState::Checked) ? true : false;
-	// Force Update.
-	if (bScopeToActiveDebugInstance != bOldState)
-	{
-		CurrentBlueprint.Reset();
-	}
-}
-
-void SGraphExecutionStatDisplay::OnDisplayByInstanceCheck(ECheckBoxState NewState)
-{
-	const bool bOldState = bDisplayInstances;
-	bDisplayInstances = (NewState == ECheckBoxState::Checked) ? true : false;
-	// Force Update.
-	if (bDisplayInstances != bOldState)
-	{
-		CurrentBlueprint.Reset();
-	}
-}
-
-void SGraphExecutionStatDisplay::OnAutoExpansionCheck(ECheckBoxState NewState)
-{
-	const bool bOldState = bAutoExpandStats;
-	bAutoExpandStats = (NewState == ECheckBoxState::Checked) ? true : false;
-	// Force Update.
-	if (bAutoExpandStats != bOldState)
-	{
-		if (bAutoExpandStats)
-		{
-			for (auto Iter : RootTreeItems)
-			{
-				Iter->ExpandWidgetState(ExecutionStatTree, bAutoExpandStats);
-			}
-		}
-		else
-		{
-			for (auto Iter : RootTreeItems)
-			{
-				Iter->RestoreWidgetExpansionState(ExecutionStatTree);
-			}
-		}
-	}
-}
-
 void SGraphExecutionStatDisplay::OnDoubleClickStatistic(FBPStatWidgetPtr Item)
 {
 	if (Item.IsValid())
@@ -255,7 +136,7 @@ void SGraphExecutionStatDisplay::OnDoubleClickStatistic(FBPStatWidgetPtr Item)
 
 void SGraphExecutionStatDisplay::OnStatisticExpansionChanged(FBPStatWidgetPtr Item, bool bExpanded)
 {
-	if (!bAutoExpandStats && Item.IsValid())
+	if (Item.IsValid() && !DisplayOptions->HasFlags(FBPProfilerStatDiplayOptions::AutoExpand))
 	{
 		Item->SetExpansionState(bExpanded);
 	}
@@ -281,63 +162,70 @@ void SGraphExecutionStatDisplay::Tick(const FGeometry& AllottedGeometry, const d
 
 					if (BlueprintExecContext.IsValid())
 					{
+						// Check blueprint
+						if (CurrentBlueprint.Get() != NewBlueprint)
+						{
+							CurrentBlueprint = NewBlueprint;
+							DisplayOptions->SetStateModified();
+						}
+						// Find Current instance
 						FName NewInstancePath = BlueprintExecContext->RemapInstancePath(FName(*NewInstance->GetPathName()));
-						const bool bBlueprintChanged = NewBlueprint != CurrentBlueprint.Get();
-						const bool bInstanceChanged = bScopeToActiveDebugInstance && NewInstancePath != CurrentInstancePath;
+						DisplayOptions->SetActiveInstance(NewInstancePath);
 						// Find active graph to filter on
-						bool bGraphFilterChanged = false;
 						TSharedPtr<SDockTab> DockTab = BlueprintEditor.Pin()->DocumentManager->GetActiveTab();
 						if (DockTab.IsValid())
 						{
 							TSharedRef<SGraphEditor> ActiveGraphEditor = StaticCastSharedRef<SGraphEditor>(DockTab->GetContent());
-							bGraphFilterChanged = ActiveGraphEditor->GetCurrentGraph() != CurrentGraph.Get();
-							CurrentGraph = ActiveGraphEditor->GetCurrentGraph();
+							const UEdGraph* CurrentGraph = ActiveGraphEditor->GetCurrentGraph();
+							DisplayOptions->SetActiveGraph(CurrentGraph ? CurrentGraph->GetFName() : NAME_None);
 						}
-
-						if (bBlueprintChanged||bInstanceChanged||bGraphFilterChanged)
+						if (DisplayOptions->IsStateModified())
 						{
-							const FString BlueprintPath = NewBlueprint->GeneratedClass->GetPathName();
 							TSharedPtr<FScriptExecutionBlueprint> BlueprintExecNode = BlueprintExecContext->GetBlueprintExecNode();
 							RootTreeItems.Reset(0);
 
 							if (BlueprintExecNode.IsValid())
 							{
-								// Cache Active blueprint and Instance
-								CurrentBlueprint = NewBlueprint;
-								CurrentInstancePath = NewInstancePath;
 								// Build Instance widget execution trees
-								FName GraphFilter = CurrentGraph.IsValid() && bFilterEventsToGraph ? CurrentGraph.Get()->GetFName() : NAME_None;
+								DisplayOptions->ClearFlags(FBPProfilerStatDiplayOptions::Modified);
+								// Cache Active blueprint and Instance
+								const FName CurrentInstancePath = DisplayOptions->GetActiveInstance();
 
-								if (bDisplayInstances && bScopeToActiveDebugInstance)
+								if (DisplayOptions->HasFlags(FBPProfilerStatDiplayOptions::DisplayByInstance))
 								{
-									if (CurrentInstancePath != NAME_None)
+									if (DisplayOptions->HasFlags(FBPProfilerStatDiplayOptions::ScopeToDebugInstance))
+									{
+										if (CurrentInstancePath != NAME_None)
+										{
+											TSharedPtr<FScriptExecutionNode> InstanceStat = BlueprintExecNode->GetInstanceByName(CurrentInstancePath);
+											if (InstanceStat.IsValid())
+											{
+												FTracePath InstanceTracePath;
+												DisplayOptions->SetActiveInstance(InstanceStat->GetName());
+												FBPStatWidgetPtr InstanceWidget = MakeShareable<FBPProfilerStatWidget>(new FBPProfilerStatWidget(InstanceStat, InstanceTracePath));
+												InstanceWidget->GenerateExecNodeWidgets(DisplayOptions);
+												RootTreeItems.Add(InstanceWidget);
+											}
+										}
+									}
+									else
 									{
-										TSharedPtr<FScriptExecutionNode> InstanceStat = BlueprintExecNode->GetInstanceByName(CurrentInstancePath);
-										if (InstanceStat.IsValid())
+										for (int32 InstanceIdx = 0; InstanceIdx < BlueprintExecNode->GetInstanceCount(); ++InstanceIdx)
 										{
+											TSharedPtr<FScriptExecutionNode> InstanceStat = BlueprintExecNode->GetInstanceByIndex(InstanceIdx);
 											FTracePath InstanceTracePath;
+											DisplayOptions->SetActiveInstance(InstanceStat->GetName());
 											FBPStatWidgetPtr InstanceWidget = MakeShareable<FBPProfilerStatWidget>(new FBPProfilerStatWidget(InstanceStat, InstanceTracePath));
-											InstanceWidget->GenerateExecNodeWidgets(InstanceStat->GetName(), GraphFilter);
+											InstanceWidget->GenerateExecNodeWidgets(DisplayOptions);
 											RootTreeItems.Add(InstanceWidget);
 										}
 									}
 								}
-								else if (bDisplayInstances)
-								{
-									for (int32 InstanceIdx = 0; InstanceIdx < BlueprintExecNode->GetInstanceCount(); ++InstanceIdx)
-									{
-										TSharedPtr<FScriptExecutionNode> InstanceStat = BlueprintExecNode->GetInstanceByIndex(InstanceIdx);
-										FTracePath InstanceTracePath;
-										FBPStatWidgetPtr InstanceWidget = MakeShareable<FBPProfilerStatWidget>(new FBPProfilerStatWidget(InstanceStat, InstanceTracePath));
-										InstanceWidget->GenerateExecNodeWidgets(InstanceStat->GetName(), GraphFilter);
-										RootTreeItems.Add(InstanceWidget);
-									}
-								}
 								else
 								{
 									FTracePath TracePath;
 									FBPStatWidgetPtr BlueprintWidget = MakeShareable<FBPProfilerStatWidget>(new FBPProfilerStatWidget(BlueprintExecNode, TracePath));
-									BlueprintWidget->GenerateExecNodeWidgets(NAME_None, GraphFilter);
+									BlueprintWidget->GenerateExecNodeWidgets(DisplayOptions);
 									RootTreeItems.Add(BlueprintWidget);
 								}
 							}
@@ -345,11 +233,11 @@ void SGraphExecutionStatDisplay::Tick(const FGeometry& AllottedGeometry, const d
 							if (ExecutionStatTree.IsValid())
 							{
 								ExecutionStatTree->RequestTreeRefresh();
-								if (bAutoExpandStats)
+								if (DisplayOptions->HasFlags(FBPProfilerStatDiplayOptions::AutoExpand))
 								{
 									for (auto Iter : RootTreeItems)
 									{
-										Iter->ExpandWidgetState(ExecutionStatTree, bAutoExpandStats);
+										Iter->ExpandWidgetState(ExecutionStatTree, true);
 									}
 								}
 								else
diff --git a/Engine/Source/Editor/Kismet/Private/Profiler/SGraphExecutionStatDisplay.h b/Engine/Source/Editor/Kismet/Private/Profiler/SGraphExecutionStatDisplay.h
--- a/Engine/Source/Editor/Kismet/Private/Profiler/SGraphExecutionStatDisplay.h
+++ b/Engine/Source/Editor/Kismet/Private/Profiler/SGraphExecutionStatDisplay.h
@@ -50,40 +50,12 @@ protected:
 	/** Called to generate children for given profiler stat */
 	void OnGetChildren(FBPStatWidgetPtr InParent, TArray<FBPStatWidgetPtr>& OutChildren);
 
-	/** Called to evaluate the graph filter checkbox */
-	ECheckBoxState GetGraphFilterCheck() const { return bFilterEventsToGraph ? ECheckBoxState::Checked : ECheckBoxState::Unchecked; }
-
-	/** Called to update the instance filter state */
-	void OnGraphFilterCheck(ECheckBoxState NewState);
-
-	/** Called to evaluate the instance filter checkbox */
-	ECheckBoxState GetInstanceFilterCheck() const { return bDisplayInstances && bScopeToActiveDebugInstance ? ECheckBoxState::Checked : ECheckBoxState::Unchecked; }
-
-	/** Called to update the instance filter state */
-	void OnInstanceFilterCheck(ECheckBoxState NewState);
-
-	/** Called to evaluate the display instances checkbox */
-	ECheckBoxState GetDisplayByInstanceCheck() const { return bDisplayInstances ? ECheckBoxState::Checked : ECheckBoxState::Unchecked; }
-
-	/** Called to update the display instances state */
-	void OnDisplayByInstanceCheck(ECheckBoxState NewState);
-
-	/** Called to evaluate the auto expansion checkbox */
-	ECheckBoxState GetAutoExpansionCheck() const { return bAutoExpandStats ? ECheckBoxState::Checked : ECheckBoxState::Unchecked; }
-
-	/** Called to update the auto expansion state */
-	void OnAutoExpansionCheck(ECheckBoxState NewState);
-
 	/** On double click statistic */
 	void OnDoubleClickStatistic(FBPStatWidgetPtr Item);
 
 	/** On statistic expansion state changed */
 	void OnStatisticExpansionChanged(FBPStatWidgetPtr Item, bool bExpanded);
 
-	/** Auto expand statistics */
-	void AutoExpandStatistics();
-
-
 protected:
 
 	/** Blueprint editor */
@@ -101,22 +73,6 @@ protected:
 	/** Current blueprint to display stats for */
 	TWeakObjectPtr<UBlueprint> CurrentBlueprint;
 
-	/** Current graph we filtered stats on */
-	TWeakObjectPtr<UEdGraph> CurrentGraph;
-
-	/** Current instance to display stats for */
-	FName CurrentInstancePath;
-
-	/** Display instances state */
-	bool bDisplayInstances;
-
-	/** Current instance filter state */
-	bool bScopeToActiveDebugInstance;
-
-	/** Current statistic expansion state */
-	bool bAutoExpandStats;
-
-	/** Filter events to graph */
-	bool bFilterEventsToGraph;
-
+	/** Current display options */
+	TSharedPtr<struct FBPProfilerStatDiplayOptions> DisplayOptions;
 };
diff --git a/Engine/Source/Editor/Kismet/Private/Profiler/SLeastPerformantDisplay.cpp b/Engine/Source/Editor/Kismet/Private/Profiler/SLeastPerformantDisplay.cpp
--- a/Engine/Source/Editor/Kismet/Private/Profiler/SLeastPerformantDisplay.cpp
+++ b/Engine/Source/Editor/Kismet/Private/Profiler/SLeastPerformantDisplay.cpp
@@ -57,9 +57,6 @@ void SLeastPerformantDisplay::Construct(const FArguments& InArgs)
 					+SHeaderRow::Column(SProfilerStatRow::GetStatName(EBlueprintProfilerStat::Time))
 					.DefaultLabel(SProfilerStatRow::GetStatText(EBlueprintProfilerStat::Time))
 					.ManualWidth(70)
-					+SHeaderRow::Column(SProfilerStatRow::GetStatName(EBlueprintProfilerStat::PureTime))
-					.DefaultLabel(SProfilerStatRow::GetStatText(EBlueprintProfilerStat::PureTime))
-					.ManualWidth(100)
 					+SHeaderRow::Column(SProfilerStatRow::GetStatName(EBlueprintProfilerStat::InclusiveTime))
 					.DefaultLabel(SProfilerStatRow::GetStatText(EBlueprintProfilerStat::InclusiveTime))
 					.ManualWidth(120)
diff --git a/Engine/Source/Editor/Kismet/Private/Profiler/SProfilerOverview.cpp b/Engine/Source/Editor/Kismet/Private/Profiler/SProfilerOverview.cpp
--- a/Engine/Source/Editor/Kismet/Private/Profiler/SProfilerOverview.cpp
+++ b/Engine/Source/Editor/Kismet/Private/Profiler/SProfilerOverview.cpp
@@ -57,9 +57,6 @@ void SProfilerOverview::Construct(const FArguments& InArgs)
 					+SHeaderRow::Column(SProfilerStatRow::GetStatName(EBlueprintProfilerStat::Time))
 					.DefaultLabel(SProfilerStatRow::GetStatText(EBlueprintProfilerStat::Time))
 					.ManualWidth(70)
-					+SHeaderRow::Column(SProfilerStatRow::GetStatName(EBlueprintProfilerStat::PureTime))
-					.DefaultLabel(SProfilerStatRow::GetStatText(EBlueprintProfilerStat::PureTime))
-					.ManualWidth(100)
 					+SHeaderRow::Column(SProfilerStatRow::GetStatName(EBlueprintProfilerStat::InclusiveTime))
 					.DefaultLabel(SProfilerStatRow::GetStatText(EBlueprintProfilerStat::InclusiveTime))
 					.ManualWidth(120)
