Copying //UE4/Dev-Framework to //UE4/Dev-Main (Source: //UE4/Dev-Framework @ 2964666)

diff --git a/Engine/Source/Editor/Cascade/Private/CascadePreviewViewportClient.cpp b/Engine/Source/Editor/Cascade/Private/CascadePreviewViewportClient.cpp
--- a/Engine/Source/Editor/Cascade/Private/CascadePreviewViewportClient.cpp
+++ b/Engine/Source/Editor/Cascade/Private/CascadePreviewViewportClient.cpp
@@ -190,7 +190,7 @@ FCascadeEdPreviewViewportClient::~FCascadeEdPreviewViewportClient()
 {
 }
 
-void FCascadeEdPreviewViewportClient::Draw(FViewport* Viewport, FCanvas* Canvas)
+void FCascadeEdPreviewViewportClient::Draw(FViewport* InViewport, FCanvas* Canvas)
 {
 	if (!CascadePtr.IsValid())
 	{
@@ -265,7 +265,7 @@ void FCascadeEdPreviewViewportClient::Draw(FViewport* Viewport, FCanvas* Canvas)
 	CascadePreviewScene.AddComponent(LineBatcher,FTransform::Identity);
 
 
-	FEditorViewportClient::Draw(Viewport, Canvas);
+	FEditorViewportClient::Draw(InViewport, Canvas);
 
 	EngineShowFlags = SavedEngineShowFlags;
 	FCanvasTextItem TextItem( FVector2D::ZeroVector, FText::GetEmpty(), GEngine->GetTinyFont(), FLinearColor::White );
@@ -273,8 +273,8 @@ void FCascadeEdPreviewViewportClient::Draw(FViewport* Viewport, FCanvas* Canvas)
 	{
 		// 'Up' from the lower left...
 		FString strOutput;
-		const int32 XPosition = Viewport->GetSizeXY().X - 5;
-		int32 YPosition = Viewport->GetSizeXY().Y - (GetDrawElement(ParticleMemory) ? 15 : 5);
+		const int32 XPosition = InViewport->GetSizeXY().X - 5;
+		int32 YPosition = InViewport->GetSizeXY().Y - (GetDrawElement(ParticleMemory) ? 15 : 5);
 
 		UParticleSystemComponent* PartComp = CascadePtr.Pin()->GetParticleSystemComponent();
 
@@ -333,7 +333,7 @@ void FCascadeEdPreviewViewportClient::Draw(FViewport* Viewport, FCanvas* Canvas)
 
 			if (GetDrawElement(ParticleMemory))
 			{
-				YPosition = Viewport->GetSizeXY().Y - 5;
+				YPosition = InViewport->GetSizeXY().Y - 5;
 				FString MemoryOutput = FString::Printf(TEXT("Template: %.0f KByte / Instance: %.0f KByte"), 
 					(float)ParticleSystemRootSize / 1024.0f + (float)ParticleModuleMemSize / 1024.0f,
 					(float)PSysCompRootSize / 1024.0f + (float)PSysCompResourceSize / 1024.0f);
@@ -385,7 +385,7 @@ void FCascadeEdPreviewViewportClient::Draw(FViewport* Viewport, FCanvas* Canvas)
 
 			if (GetDrawElement(ParticleMemory))
 			{
-				YPosition = Viewport->GetSizeXY().Y - 5;
+				YPosition = InViewport->GetSizeXY().Y - 5;
 				FString MemoryOutput = FString::Printf(TEXT("Template: %.0f KByte / Instance: %.0f KByte"), 
 					(float)ParticleSystemRootSize / 1024.0f + (float)ParticleModuleMemSize / 1024.0f,
 					(float)PSysCompRootSize / 1024.0f + (float)PSysCompResourceSize / 1024.0f);
@@ -404,7 +404,7 @@ void FCascadeEdPreviewViewportClient::Draw(FViewport* Viewport, FCanvas* Canvas)
 				TextItem.Text = LOCTEXT("SystemCompleted", "Completed");
 				TextItem.bCentreX = true;
 				TextItem.bCentreY = true;
-				Canvas->DrawItem(TextItem, Viewport->GetSizeXY().X * 0.5f, Viewport->GetSizeXY().Y - 10);
+				Canvas->DrawItem(TextItem, InViewport->GetSizeXY().X * 0.5f, InViewport->GetSizeXY().Y - 10);
 				TextItem.bCentreX = false;
 				TextItem.bCentreY = false;
 			}
@@ -436,7 +436,7 @@ void FCascadeEdPreviewViewportClient::Draw(FViewport* Viewport, FCanvas* Canvas)
 				if(bIsAGPUEmitter)
 				{
 					const int32 XPosition = 5;
-					const int32 YPosition = Viewport->GetSizeXY().Y - 75.0f;
+					const int32 YPosition = InViewport->GetSizeXY().Y - 75.0f;
 					FString strOutput = NSLOCTEXT("Cascade", "NoFixedBounds_Warning", "WARNING: This particle system has no fixed bounding box and contains a GPU emitter.").ToString();
 					TextItem.SetColor( FLinearColor::White );
 					TextItem.Text = FText::FromString( strOutput );
@@ -454,14 +454,14 @@ void FCascadeEdPreviewViewportClient::Draw(FViewport* Viewport, FCanvas* Canvas)
 		FString DetailModeOutput = FString::Printf(TEXT("DETAIL MODE: %s"), (DetailMode == DM_Medium)? TEXT("MEDIUM"): TEXT("LOW"));
 		TextItem.SetColor( FLinearColor::Red );
 		TextItem.Text = FText::FromString( DetailModeOutput );
-		Canvas->DrawItem( TextItem, 5.0f, Viewport->GetSizeXY().Y - 90.0f );		
+		Canvas->DrawItem( TextItem, 5.0f, InViewport->GetSizeXY().Y - 90.0f );		
 	}
 
 	if (GEngine->bEnableEditorPSysRealtimeLOD)
 	{
 		TextItem.SetColor( FLinearColor(0.25f, 0.25f, 1.0f) );
 		TextItem.Text = LOCTEXT("LODPREVIEWMODEENABLED","LOD PREVIEW MODE ENABLED");
-		Canvas->DrawItem( TextItem,  5.0f, Viewport->GetSizeXY().Y - 105.0f );		
+		Canvas->DrawItem( TextItem,  5.0f, InViewport->GetSizeXY().Y - 105.0f );		
 	}
 
 	EParticleSignificanceLevel ReqSignificance = CascadePtr.Pin()->GetRequiredSignificance();
@@ -470,17 +470,17 @@ void FCascadeEdPreviewViewportClient::Draw(FViewport* Viewport, FCanvas* Canvas)
 		FString ReqSigOutput = FString::Printf(TEXT("REQUIRED SIGNIFICANCE: %s"), (ReqSignificance == EParticleSignificanceLevel::Medium) ? TEXT("MEDIUM") : ((ReqSignificance == EParticleSignificanceLevel::High) ? TEXT("HIGH") : TEXT("CRITICAL")));
 		TextItem.SetColor(FLinearColor::Red);
 		TextItem.Text = FText::FromString(ReqSigOutput);
-		Canvas->DrawItem(TextItem, 5.0f, Viewport->GetSizeXY().Y - 120.0f);
+		Canvas->DrawItem(TextItem, 5.0f, InViewport->GetSizeXY().Y - 120.0f);
 	}
 
 	if (bCaptureScreenShot)
 	{
 		UParticleSystem* ParticleSystem = CascadePtr.Pin()->GetParticleSystem();
-		int32 SrcWidth = Viewport->GetSizeXY().X;
-		int32 SrcHeight = Viewport->GetSizeXY().Y;
+		int32 SrcWidth = InViewport->GetSizeXY().X;
+		int32 SrcHeight = InViewport->GetSizeXY().Y;
 		// Read the contents of the viewport into an array.
 		TArray<FColor> OrigBitmap;
-		if (Viewport->ReadPixels(OrigBitmap))
+		if (InViewport->ReadPixels(OrigBitmap))
 		{
 			check(OrigBitmap.Num() == SrcWidth * SrcHeight);
 
@@ -557,21 +557,21 @@ void FCascadeEdPreviewViewportClient::Draw(const FSceneView* View, FPrimitiveDra
 	DrawPreviewLightVisualization(View, PDI);
 }
 
-bool FCascadeEdPreviewViewportClient::InputKey(FViewport* Viewport, int32 ControllerId, FKey Key, EInputEvent Event, float AmountDepressed, bool Gamepad)
+bool FCascadeEdPreviewViewportClient::InputKey(FViewport* InViewport, int32 ControllerId, FKey Key, EInputEvent Event, float AmountDepressed, bool Gamepad)
 {
 	//Update cursor 
 	UpdateAndApplyCursorVisibility();
 
 	bool bHandled = false;
-	const int32 HitX = Viewport->GetMouseX();
-	const int32 HitY = Viewport->GetMouseY();
+	const int32 HitX = InViewport->GetMouseX();
+	const int32 HitY = InViewport->GetMouseY();
 
 	if(Key == EKeys::LeftMouseButton)
 	{
 		if (Event == IE_Pressed)
 		{
-			Viewport->InvalidateHitProxy();
-			HHitProxy* HitResult = Viewport->GetHitProxy(HitX,HitY);
+			InViewport->InvalidateHitProxy();
+			HHitProxy* HitResult = InViewport->GetHitProxy(HitX,HitY);
 			if (HitResult && HitResult->IsA(HWidgetUtilProxy::StaticGetType()))
 			{
 				HWidgetUtilProxy* WidgetProxy = (HWidgetUtilProxy*)HitResult;
@@ -582,7 +582,7 @@ bool FCascadeEdPreviewViewportClient::InputKey(FViewport* Viewport, int32 Contro
 				WidgetAxis = WidgetProxy->Axis;
 
 				// Calculate the scree-space directions for this drag.
-				FSceneViewFamilyContext ViewFamily( FSceneViewFamily::ConstructionValues( Viewport, GetScene(), EngineShowFlags ));
+				FSceneViewFamilyContext ViewFamily( FSceneViewFamily::ConstructionValues( InViewport, GetScene(), EngineShowFlags ));
 				FSceneView* View = CalcSceneView(&ViewFamily);
 				WidgetProxy->CalcVectors(View, FViewportClick(View, this, Key, Event, HitX, HitY), LocalManipulateDir, WorldManipulateDir, DragX, DragY);
 				bHandled = true;
@@ -611,7 +611,7 @@ bool FCascadeEdPreviewViewportClient::InputKey(FViewport* Viewport, int32 Contro
 
 	if( !bHandled )
 	{
-		bHandled = FEditorViewportClient::InputKey(Viewport,ControllerId,Key,Event,AmountDepressed,Gamepad);
+		bHandled = FEditorViewportClient::InputKey(InViewport,ControllerId,Key,Event,AmountDepressed,Gamepad);
 	}
 
 
@@ -624,7 +624,7 @@ static auto CVarCascadeDragSpeed = TAutoConsoleVariable<float>(TEXT("CascadeDrag
 static TAutoConsoleVariable<float> CVarCascadeRotateSpeed(TEXT("CascadeRotateSpeed"),0.005f,TEXT("Cascade drag speed."));
 static TAutoConsoleVariable<float> CVarCascadeScaleSpeed(TEXT("CascadeScaleSpeed"),1.0f,TEXT("Cascade scale speed."));
 
-bool FCascadeEdPreviewViewportClient::InputAxis(FViewport* Viewport, int32 ControllerId, FKey Key, float Delta, float DeltaTime, int32 NumSamples, bool bGamepad)
+bool FCascadeEdPreviewViewportClient::InputAxis(FViewport* InViewport, int32 ControllerId, FKey Key, float Delta, float DeltaTime, int32 NumSamples, bool bGamepad)
 {
 	bool bHandled = false;
 
@@ -660,12 +660,12 @@ bool FCascadeEdPreviewViewportClient::InputAxis(FViewport* Viewport, int32 Contr
 	}
 	else
 	{
-		bHandled = FEditorViewportClient::InputAxis(Viewport,ControllerId,Key,Delta,DeltaTime,NumSamples,bGamepad);
+		bHandled = FEditorViewportClient::InputAxis(InViewport,ControllerId,Key,Delta,DeltaTime,NumSamples,bGamepad);
 	}
 
 	if (!IsRealtime() && !FMath::IsNearlyZero(Delta))
 	{
-		Viewport->Invalidate();
+		InViewport->Invalidate();
 	}
 
 	return bHandled;
diff --git a/Engine/Source/Editor/ContentBrowser/Private/SAssetView.cpp b/Engine/Source/Editor/ContentBrowser/Private/SAssetView.cpp
--- a/Engine/Source/Editor/ContentBrowser/Private/SAssetView.cpp
+++ b/Engine/Source/Editor/ContentBrowser/Private/SAssetView.cpp
@@ -2033,14 +2033,14 @@ void SAssetView::SetMajorityAssetType(FName NewMajorityAssetType)
 					{
 						if ( TagIt->Type != UObject::FAssetRegistryTag::TT_Hidden )
 						{
-							const FName& Tag = TagIt->Name;
+							const FName TagName = TagIt->Name;
 
-							if ( !OnAssetTagWantsToBeDisplayed.IsBound() || OnAssetTagWantsToBeDisplayed.Execute(NewMajorityAssetType, Tag) )
+							if ( !OnAssetTagWantsToBeDisplayed.IsBound() || OnAssetTagWantsToBeDisplayed.Execute(NewMajorityAssetType, TagName) )
 							{
 								// Get tag metadata
 								TMap<FName, UObject::FAssetRegistryTagMetadata> MetadataMap;
 								CDO->GetAssetRegistryTagMetadata(MetadataMap);
-								const UObject::FAssetRegistryTagMetadata* Metadata = MetadataMap.Find(Tag);
+								const UObject::FAssetRegistryTagMetadata* Metadata = MetadataMap.Find(TagName);
 
 								FText DisplayName;
 								if (Metadata != nullptr && !Metadata->DisplayName.IsEmpty())
@@ -2049,7 +2049,7 @@ void SAssetView::SetMajorityAssetType(FName NewMajorityAssetType)
 								}
 								else
 								{
-									DisplayName = FText::FromName(Tag);
+									DisplayName = FText::FromName(TagName);
 								}
 
 								FText TooltipText;
@@ -2060,14 +2060,14 @@ void SAssetView::SetMajorityAssetType(FName NewMajorityAssetType)
 								else
 								{
 									// If the tag name corresponds to a property name, use the property tooltip
-									UProperty* Property = FindField<UProperty>(TypeClass, Tag);
-									TooltipText = (Property != nullptr) ? Property->GetToolTipText() : FText::FromString(FName::NameToDisplayString(Tag.ToString(), false));
+									UProperty* Property = FindField<UProperty>(TypeClass, TagName);
+									TooltipText = (Property != nullptr) ? Property->GetToolTipText() : FText::FromString(FName::NameToDisplayString(TagName.ToString(), false));
 								}
 
 								ColumnView->GetHeaderRow()->AddColumn(
-										SHeaderRow::Column(Tag)
-										.SortMode( TAttribute< EColumnSortMode::Type >::Create( TAttribute< EColumnSortMode::Type >::FGetter::CreateSP( this, &SAssetView::GetColumnSortMode, Tag ) ) )
-										.SortPriority(TAttribute< EColumnSortPriority::Type >::Create(TAttribute< EColumnSortPriority::Type >::FGetter::CreateSP(this, &SAssetView::GetColumnSortPriority, Tag)))
+										SHeaderRow::Column(TagName)
+										.SortMode( TAttribute< EColumnSortMode::Type >::Create( TAttribute< EColumnSortMode::Type >::FGetter::CreateSP( this, &SAssetView::GetColumnSortMode, TagName ) ) )
+										.SortPriority(TAttribute< EColumnSortPriority::Type >::Create(TAttribute< EColumnSortPriority::Type >::FGetter::CreateSP(this, &SAssetView::GetColumnSortPriority, TagName)))
 										.OnSort( FOnSortModeChanged::CreateSP( this, &SAssetView::OnSortColumnHeader ) )
 										.DefaultLabel( DisplayName )
 										.DefaultTooltip( TooltipText )
@@ -2078,7 +2078,7 @@ void SAssetView::SetMajorityAssetType(FName NewMajorityAssetType)
 								// If we found a tag the matches the column we are currently sorting on, there will be no need to change the column
 								for (int32 SortIdx = 0; SortIdx < CurrentSortOrder.Num(); SortIdx++)
 								{
-									if (Tag == CurrentSortOrder[SortIdx].SortColumn)
+									if (TagName == CurrentSortOrder[SortIdx].SortColumn)
 									{
 										CurrentSortOrder[SortIdx].bSortRelevant = true;
 									}
@@ -2159,7 +2159,6 @@ void SAssetView::ProcessRecentlyAddedAssets()
 
 	if (FilteredRecentlyAddedAssets.Num() > 0)
 	{
-		const static float MaxSecondsPerFrame = 0.015;
 		double TickStartTime = FPlatformTime::Seconds();
 		bool bNeedsRefresh = false;
 
diff --git a/Engine/Source/Editor/CurveAssetEditor/Private/CurveAssetEditor.cpp b/Engine/Source/Editor/CurveAssetEditor/Private/CurveAssetEditor.cpp
--- a/Engine/Source/Editor/CurveAssetEditor/Private/CurveAssetEditor.cpp
+++ b/Engine/Source/Editor/CurveAssetEditor/Private/CurveAssetEditor.cpp
@@ -17,21 +17,21 @@
 
 const FName FCurveAssetEditor::CurveTabId( TEXT( "CurveAssetEditor_Curve" ) );
 
-void FCurveAssetEditor::RegisterTabSpawners(const TSharedRef<class FTabManager>& TabManager)
+void FCurveAssetEditor::RegisterTabSpawners(const TSharedRef<class FTabManager>& InTabManager)
 {
-	WorkspaceMenuCategory = TabManager->AddLocalWorkspaceMenuCategory(LOCTEXT("WorkspaceMenu_CurveAssetEditor", "Curve Asset Editor"));
+	WorkspaceMenuCategory = InTabManager->AddLocalWorkspaceMenuCategory(LOCTEXT("WorkspaceMenu_CurveAssetEditor", "Curve Asset Editor"));
 
-	FAssetEditorToolkit::RegisterTabSpawners(TabManager);
+	FAssetEditorToolkit::RegisterTabSpawners(InTabManager);
 
-	TabManager->RegisterTabSpawner( CurveTabId, FOnSpawnTab::CreateSP(this, &FCurveAssetEditor::SpawnTab_CurveAsset) )
+	InTabManager->RegisterTabSpawner( CurveTabId, FOnSpawnTab::CreateSP(this, &FCurveAssetEditor::SpawnTab_CurveAsset) )
 		.SetDisplayName( LOCTEXT("CurveTab", "Curve") )
 		.SetGroup(WorkspaceMenuCategory.ToSharedRef())
 		.SetIcon(FSlateIcon(FEditorStyle::GetStyleSetName(), "ClassIcon.CurveBase"));
 }
 
-void FCurveAssetEditor::UnregisterTabSpawners(const TSharedRef<class FTabManager>& TabManager)
+void FCurveAssetEditor::UnregisterTabSpawners(const TSharedRef<class FTabManager>& InTabManager)
 {
-	TabManager->UnregisterTabSpawner( CurveTabId );
+	InTabManager->UnregisterTabSpawner( CurveTabId );
 }
 
 void FCurveAssetEditor::InitCurveAssetEditor( const EToolkitMode::Type Mode, const TSharedPtr< class IToolkitHost >& InitToolkitHost, UCurveBase* CurveToEdit )
diff --git a/Engine/Source/Editor/CurveTableEditor/Private/CurveTableEditor.cpp b/Engine/Source/Editor/CurveTableEditor/Private/CurveTableEditor.cpp
--- a/Engine/Source/Editor/CurveTableEditor/Private/CurveTableEditor.cpp
+++ b/Engine/Source/Editor/CurveTableEditor/Private/CurveTableEditor.cpp
@@ -52,19 +52,19 @@ private:
 };
 
 
-void FCurveTableEditor::RegisterTabSpawners(const TSharedRef<class FTabManager>& TabManager)
+void FCurveTableEditor::RegisterTabSpawners(const TSharedRef<class FTabManager>& InTabManager)
 {
-	WorkspaceMenuCategory = TabManager->AddLocalWorkspaceMenuCategory(LOCTEXT("WorkspaceMenu_CurveTableEditor", "Curve Table Editor"));
+	WorkspaceMenuCategory = InTabManager->AddLocalWorkspaceMenuCategory(LOCTEXT("WorkspaceMenu_CurveTableEditor", "Curve Table Editor"));
 
-	TabManager->RegisterTabSpawner( CurveTableTabId, FOnSpawnTab::CreateSP(this, &FCurveTableEditor::SpawnTab_CurveTable) )
+	InTabManager->RegisterTabSpawner( CurveTableTabId, FOnSpawnTab::CreateSP(this, &FCurveTableEditor::SpawnTab_CurveTable) )
 		.SetDisplayName( LOCTEXT("CurveTableTab", "Curve Table") )
 		.SetGroup( WorkspaceMenuCategory.ToSharedRef() );
 }
 
 
-void FCurveTableEditor::UnregisterTabSpawners(const TSharedRef<class FTabManager>& TabManager)
+void FCurveTableEditor::UnregisterTabSpawners(const TSharedRef<class FTabManager>& InTabManager)
 {
-	TabManager->UnregisterTabSpawner( CurveTableTabId );
+	InTabManager->UnregisterTabSpawner( CurveTableTabId );
 }
 
 
diff --git a/Engine/Source/Editor/DataTableEditor/Private/DataTableEditor.cpp b/Engine/Source/Editor/DataTableEditor/Private/DataTableEditor.cpp
--- a/Engine/Source/Editor/DataTableEditor/Private/DataTableEditor.cpp
+++ b/Engine/Source/Editor/DataTableEditor/Private/DataTableEditor.cpp
@@ -56,23 +56,23 @@ private:
 	FDataTableEditorRowListViewDataPtr Item;
 };
 
-void FDataTableEditor::RegisterTabSpawners(const TSharedRef<class FTabManager>& TabManager)
+void FDataTableEditor::RegisterTabSpawners(const TSharedRef<class FTabManager>& InTabManager)
 {
-	WorkspaceMenuCategory = TabManager->AddLocalWorkspaceMenuCategory(LOCTEXT("WorkspaceMenu_Data Table Editor", "Data Table Editor"));
+	WorkspaceMenuCategory = InTabManager->AddLocalWorkspaceMenuCategory(LOCTEXT("WorkspaceMenu_Data Table Editor", "Data Table Editor"));
 
-	TabManager->RegisterTabSpawner( DataTableTabId, FOnSpawnTab::CreateSP(this, &FDataTableEditor::SpawnTab_DataTable) )
+	InTabManager->RegisterTabSpawner( DataTableTabId, FOnSpawnTab::CreateSP(this, &FDataTableEditor::SpawnTab_DataTable) )
 		.SetDisplayName( LOCTEXT("DataTableTab", "Data Table") )
 		.SetGroup( WorkspaceMenuCategory.ToSharedRef() );
 
-	TabManager->RegisterTabSpawner(RowEditorTabId, FOnSpawnTab::CreateSP(this, &FDataTableEditor::SpawnTab_RowEditor))
+	InTabManager->RegisterTabSpawner(RowEditorTabId, FOnSpawnTab::CreateSP(this, &FDataTableEditor::SpawnTab_RowEditor))
 		.SetDisplayName(LOCTEXT("RowEditorTab", "Row Editor"))
 		.SetGroup(WorkspaceMenuCategory.ToSharedRef());
 }
 
-void FDataTableEditor::UnregisterTabSpawners(const TSharedRef<class FTabManager>& TabManager)
+void FDataTableEditor::UnregisterTabSpawners(const TSharedRef<class FTabManager>& InTabManager)
 {
-	TabManager->UnregisterTabSpawner( DataTableTabId );
-	TabManager->UnregisterTabSpawner(RowEditorTabId);
+	InTabManager->UnregisterTabSpawner( DataTableTabId );
+	InTabManager->UnregisterTabSpawner(RowEditorTabId);
 }
 
 FDataTableEditor::FDataTableEditor()
diff --git a/Engine/Source/Editor/DestructibleMeshEditor/Private/DestructibleMeshEditor.cpp b/Engine/Source/Editor/DestructibleMeshEditor/Private/DestructibleMeshEditor.cpp
--- a/Engine/Source/Editor/DestructibleMeshEditor/Private/DestructibleMeshEditor.cpp
+++ b/Engine/Source/Editor/DestructibleMeshEditor/Private/DestructibleMeshEditor.cpp
@@ -42,42 +42,42 @@ FDestructibleMeshEditor::~FDestructibleMeshEditor()
 	GEditor->OnObjectReimported().RemoveAll(this);
 }
 
-void FDestructibleMeshEditor::RegisterTabSpawners(const TSharedRef<class FTabManager>& TabManager)
+void FDestructibleMeshEditor::RegisterTabSpawners(const TSharedRef<class FTabManager>& InTabManager)
 {
-	WorkspaceMenuCategory = TabManager->AddLocalWorkspaceMenuCategory(LOCTEXT("WorkspaceMenu_DestructibleMeshEditor", "Destructible Mesh Editor"));
+	WorkspaceMenuCategory = InTabManager->AddLocalWorkspaceMenuCategory(LOCTEXT("WorkspaceMenu_DestructibleMeshEditor", "Destructible Mesh Editor"));
 	auto WorkspaceMenuCategoryRef = WorkspaceMenuCategory.ToSharedRef();
 
-	FAssetEditorToolkit::RegisterTabSpawners(TabManager);
+	FAssetEditorToolkit::RegisterTabSpawners(InTabManager);
 
-	TabManager->RegisterTabSpawner( ViewportTabId, FOnSpawnTab::CreateSP(this, &FDestructibleMeshEditor::SpawnTab_Viewport) )
+	InTabManager->RegisterTabSpawner( ViewportTabId, FOnSpawnTab::CreateSP(this, &FDestructibleMeshEditor::SpawnTab_Viewport) )
 		.SetDisplayName( LOCTEXT("ViewportTab", "Viewport") )
 		.SetGroup(WorkspaceMenuCategoryRef)
 		.SetIcon(FSlateIcon(FEditorStyle::GetStyleSetName(), "LevelEditor.Tabs.Viewports"));
 
-	TabManager->RegisterTabSpawner( PropertiesTabId, FOnSpawnTab::CreateSP(this, &FDestructibleMeshEditor::SpawnTab_Properties) )
+	InTabManager->RegisterTabSpawner( PropertiesTabId, FOnSpawnTab::CreateSP(this, &FDestructibleMeshEditor::SpawnTab_Properties) )
 		.SetDisplayName( LOCTEXT("PropertiesTab", "Destructible Settings") )
 		.SetGroup(WorkspaceMenuCategoryRef)
 		.SetIcon(FSlateIcon(FEditorStyle::GetStyleSetName(), "DestructibleMeshEditor.Tabs.DestructibleSettings"));
 
-	TabManager->RegisterTabSpawner( FractureSettingsTabId, FOnSpawnTab::CreateSP(this, &FDestructibleMeshEditor::SpawnTab_FractureSettings) )
+	InTabManager->RegisterTabSpawner( FractureSettingsTabId, FOnSpawnTab::CreateSP(this, &FDestructibleMeshEditor::SpawnTab_FractureSettings) )
 		.SetDisplayName( LOCTEXT("FractureSettingsTab", "Fracture Settings") )
 		.SetGroup(WorkspaceMenuCategoryRef)
 		.SetIcon(FSlateIcon(FEditorStyle::GetStyleSetName(), "DestructibleMeshEditor.Tabs.FractureSettings"));
 
-	TabManager->RegisterTabSpawner( ChunkParametersTabId, FOnSpawnTab::CreateSP(this, &FDestructibleMeshEditor::SpawnTab_ChunkParameters) )
+	InTabManager->RegisterTabSpawner( ChunkParametersTabId, FOnSpawnTab::CreateSP(this, &FDestructibleMeshEditor::SpawnTab_ChunkParameters) )
 		.SetDisplayName( LOCTEXT("ChunkParametersTab", "Chunk Parameters") )
 		.SetGroup(WorkspaceMenuCategoryRef)
 		.SetIcon(FSlateIcon(FEditorStyle::GetStyleSetName(), "DestructibleMeshEditor.Tabs.ChunkParameters"));
 }
 
-void FDestructibleMeshEditor::UnregisterTabSpawners(const TSharedRef<class FTabManager>& TabManager)
+void FDestructibleMeshEditor::UnregisterTabSpawners(const TSharedRef<class FTabManager>& InTabManager)
 {
-	FAssetEditorToolkit::UnregisterTabSpawners(TabManager);
+	FAssetEditorToolkit::UnregisterTabSpawners(InTabManager);
 
-	TabManager->UnregisterTabSpawner( ViewportTabId );
-	TabManager->UnregisterTabSpawner( PropertiesTabId );
-	TabManager->UnregisterTabSpawner( FractureSettingsTabId );
-	TabManager->UnregisterTabSpawner( ChunkParametersTabId );
+	InTabManager->UnregisterTabSpawner( ViewportTabId );
+	InTabManager->UnregisterTabSpawner( PropertiesTabId );
+	InTabManager->UnregisterTabSpawner( FractureSettingsTabId );
+	InTabManager->UnregisterTabSpawner( ChunkParametersTabId );
 }
 
 void FDestructibleMeshEditor::InitDestructibleMeshEditor( const EToolkitMode::Type Mode, const TSharedPtr< class IToolkitHost >& InitToolkitHost, UDestructibleMesh* InDestructibleMesh )
diff --git a/Engine/Source/Editor/DetailCustomizations/Private/PhysicsConstraintComponentDetails.cpp b/Engine/Source/Editor/DetailCustomizations/Private/PhysicsConstraintComponentDetails.cpp
--- a/Engine/Source/Editor/DetailCustomizations/Private/PhysicsConstraintComponentDetails.cpp
+++ b/Engine/Source/Editor/DetailCustomizations/Private/PhysicsConstraintComponentDetails.cpp
@@ -282,9 +282,7 @@ void FPhysicsConstraintComponentDetails::AddAngularLimits(IDetailLayoutBuilder&
 
 	if (bInPhat == false)
 	{
-		AngularLimitCat.AddProperty(ConstraintInstance->GetChildHandle(GET_MEMBER_NAME_CHECKED(FConstraintInstance, AngularRotationOffset)).ToSharedRef())
-			.IsEnabled(TAttribute<bool>::Create(TAttribute<bool>::FGetter::CreateSP(this, &FPhysicsConstraintComponentDetails::IsPropertyEnabled, EPropertyType::AngularAnyLimit)));
-
+	    AngularLimitCat.AddProperty(ConstraintInstance->GetChildHandle(GET_MEMBER_NAME_CHECKED(FConstraintInstance, AngularRotationOffset)).ToSharedRef());
 	}
 	else
 	{
diff --git a/Engine/Source/Editor/FontEditor/Private/FontEditor.cpp b/Engine/Source/Editor/FontEditor/Private/FontEditor.cpp
--- a/Engine/Source/Editor/FontEditor/Private/FontEditor.cpp
+++ b/Engine/Source/Editor/FontEditor/Private/FontEditor.cpp
@@ -74,51 +74,51 @@ void FFontEditorCommands::RegisterCommands()
 	UI_COMMAND(FontForegroundColor, "Foreground", "Changes the foreground color of the previewer.", EUserInterfaceActionType::Button, FInputChord());
 }
 
-void FFontEditor::RegisterTabSpawners(const TSharedRef<class FTabManager>& TabManager)
+void FFontEditor::RegisterTabSpawners(const TSharedRef<class FTabManager>& InTabManager)
 {
-	WorkspaceMenuCategory = TabManager->AddLocalWorkspaceMenuCategory(LOCTEXT("WorkspaceMenu_FontEditor", "Font Editor"));
+	WorkspaceMenuCategory = InTabManager->AddLocalWorkspaceMenuCategory(LOCTEXT("WorkspaceMenu_FontEditor", "Font Editor"));
 	auto WorkspaceMenuCategoryRef = WorkspaceMenuCategory.ToSharedRef();
 
-	FAssetEditorToolkit::RegisterTabSpawners(TabManager);
+	FAssetEditorToolkit::RegisterTabSpawners(InTabManager);
 
-	TabManager->RegisterTabSpawner( TexturePagesViewportTabId, FOnSpawnTab::CreateSP(this, &FFontEditor::SpawnTab_TexturePagesViewport) )
+	InTabManager->RegisterTabSpawner( TexturePagesViewportTabId, FOnSpawnTab::CreateSP(this, &FFontEditor::SpawnTab_TexturePagesViewport) )
 		.SetDisplayName( LOCTEXT("TexturePagesViewportTab", "Texture Pages") )
 		.SetGroup(WorkspaceMenuCategoryRef)
 		.SetIcon(FSlateIcon(FEditorStyle::GetStyleSetName(), "LevelEditor.Tabs.Viewports"))
 		.SetMenuType( TAttribute<ETabSpawnerMenuType::Type>::Create(TAttribute<ETabSpawnerMenuType::Type>::FGetter::CreateSP(this, &FFontEditor::GetTabSpawnerMenuType, TexturePagesViewportTabId)) );
 
-	TabManager->RegisterTabSpawner( CompositeFontEditorTabId, FOnSpawnTab::CreateSP(this, &FFontEditor::SpawnTab_CompositeFontEditor) )
+	InTabManager->RegisterTabSpawner( CompositeFontEditorTabId, FOnSpawnTab::CreateSP(this, &FFontEditor::SpawnTab_CompositeFontEditor) )
 		.SetDisplayName( LOCTEXT("CompositeFontEditorTab", "Composite Font") )
 		.SetGroup(WorkspaceMenuCategoryRef)
 		.SetIcon(FSlateIcon(FEditorStyle::GetStyleSetName(), "FontEditor.Tabs.PageProperties"))
 		.SetMenuType( TAttribute<ETabSpawnerMenuType::Type>::Create(TAttribute<ETabSpawnerMenuType::Type>::FGetter::CreateSP(this, &FFontEditor::GetTabSpawnerMenuType, CompositeFontEditorTabId)) );
 
-	TabManager->RegisterTabSpawner( PreviewTabId,		FOnSpawnTab::CreateSP(this, &FFontEditor::SpawnTab_Preview) )
+	InTabManager->RegisterTabSpawner( PreviewTabId,		FOnSpawnTab::CreateSP(this, &FFontEditor::SpawnTab_Preview) )
 		.SetDisplayName( LOCTEXT("PreviewTab", "Preview") )
 		.SetGroup(WorkspaceMenuCategoryRef)
 		.SetIcon(FSlateIcon(FEditorStyle::GetStyleSetName(), "FontEditor.Tabs.Preview"));
 
-	TabManager->RegisterTabSpawner( PropertiesTabId,	FOnSpawnTab::CreateSP(this, &FFontEditor::SpawnTab_Properties) )
+	InTabManager->RegisterTabSpawner( PropertiesTabId,	FOnSpawnTab::CreateSP(this, &FFontEditor::SpawnTab_Properties) )
 		.SetDisplayName( LOCTEXT("PropertiesTabId", "Details") )
 		.SetGroup(WorkspaceMenuCategoryRef)
 		.SetIcon(FSlateIcon(FEditorStyle::GetStyleSetName(), "LevelEditor.Tabs.Details"));
 
-	TabManager->RegisterTabSpawner( PagePropertiesTabId,FOnSpawnTab::CreateSP(this, &FFontEditor::SpawnTab_PageProperties) )
+	InTabManager->RegisterTabSpawner( PagePropertiesTabId,FOnSpawnTab::CreateSP(this, &FFontEditor::SpawnTab_PageProperties) )
 		.SetDisplayName( LOCTEXT("PagePropertiesTab", "Page Details") )
 		.SetGroup(WorkspaceMenuCategoryRef)
 		.SetIcon(FSlateIcon(FEditorStyle::GetStyleSetName(), "FontEditor.Tabs.PageProperties"))
 		.SetMenuType( TAttribute<ETabSpawnerMenuType::Type>::Create(TAttribute<ETabSpawnerMenuType::Type>::FGetter::CreateSP(this, &FFontEditor::GetTabSpawnerMenuType, PagePropertiesTabId)) );
 }
 
-void FFontEditor::UnregisterTabSpawners(const TSharedRef<class FTabManager>& TabManager)
+void FFontEditor::UnregisterTabSpawners(const TSharedRef<class FTabManager>& InTabManager)
 {
-	FAssetEditorToolkit::UnregisterTabSpawners(TabManager);
+	FAssetEditorToolkit::UnregisterTabSpawners(InTabManager);
 
-	TabManager->UnregisterTabSpawner( TexturePagesViewportTabId );
-	TabManager->UnregisterTabSpawner( CompositeFontEditorTabId );
-	TabManager->UnregisterTabSpawner( PreviewTabId );	
-	TabManager->UnregisterTabSpawner( PropertiesTabId );
-	TabManager->UnregisterTabSpawner( PagePropertiesTabId );
+	InTabManager->UnregisterTabSpawner( TexturePagesViewportTabId );
+	InTabManager->UnregisterTabSpawner( CompositeFontEditorTabId );
+	InTabManager->UnregisterTabSpawner( PreviewTabId );	
+	InTabManager->UnregisterTabSpawner( PropertiesTabId );
+	InTabManager->UnregisterTabSpawner( PagePropertiesTabId );
 }
 
 FFontEditor::~FFontEditor()
diff --git a/Engine/Source/Editor/FontEditor/Private/SFontEditorViewport.cpp b/Engine/Source/Editor/FontEditor/Private/SFontEditorViewport.cpp
--- a/Engine/Source/Editor/FontEditor/Private/SFontEditorViewport.cpp
+++ b/Engine/Source/Editor/FontEditor/Private/SFontEditorViewport.cpp
@@ -751,11 +751,11 @@ const FColor& SFontEditorViewport::GetPreviewBackgroundColor() const
 	return FColor::Black;
 }
 
-void SFontEditorViewport::SetPreviewForegroundColor(const FColor& ForegroundColor)
+void SFontEditorViewport::SetPreviewForegroundColor(const FColor& InForegroundColor)
 {
 	if (ViewportClient.IsValid())
 	{
-		ViewportClient->SetForegroundColor(ForegroundColor);
+		ViewportClient->SetForegroundColor(InForegroundColor);
 
 		RefreshViewport();
 	}
diff --git a/Engine/Source/Editor/Kismet/Private/BlueprintDetailsCustomization.cpp b/Engine/Source/Editor/Kismet/Private/BlueprintDetailsCustomization.cpp
--- a/Engine/Source/Editor/Kismet/Private/BlueprintDetailsCustomization.cpp
+++ b/Engine/Source/Editor/Kismet/Private/BlueprintDetailsCustomization.cpp
@@ -632,19 +632,19 @@ void FBlueprintVarActionDetails::CustomizeDetails( IDetailLayoutBuilder& DetailL
 		.ToolTip(ReplicationTooltip)
 	];
 
-	UBlueprint* Blueprint = GetBlueprintObj();
+	UBlueprint* BlueprintObj = GetBlueprintObj();
 
 	// Handle event generation
-	if ( FBlueprintEditorUtils::DoesSupportEventGraphs(Blueprint) )
+	if ( FBlueprintEditorUtils::DoesSupportEventGraphs(BlueprintObj) )
 	{
 		AddEventsCategory(DetailLayout, VariableProperty);
 	}
 
 	// Add in default value editing for properties that can be edited, local properties cannot be edited
-	if ((Blueprint != NULL) && (Blueprint->GeneratedClass != NULL))
+	if ((BlueprintObj != nullptr) && (BlueprintObj->GeneratedClass != nullptr))
 	{
 		bool bVariableRenamed = false;
-		if (VariableProperty != NULL && IsVariableInBlueprint())
+		if (VariableProperty != nullptr && IsVariableInBlueprint())
 		{
 			// Determine the current property name on the CDO is stale
 			if (PropertyOwnerBlueprint.IsValid() && VariableProperty)
@@ -667,7 +667,7 @@ void FBlueprintVarActionDetails::CustomizeDetails( IDetailLayoutBuilder& DetailL
 		
 			if(!IsALocalVariable(VariableProperty))
 			{
-				OriginalProperty = FindField<UProperty>(Blueprint->GeneratedClass, VariableProperty->GetFName());
+				OriginalProperty = FindField<UProperty>(BlueprintObj->GeneratedClass, VariableProperty->GetFName());
 			}
 			else
 			{
@@ -685,20 +685,20 @@ void FBlueprintVarActionDetails::CustomizeDetails( IDetailLayoutBuilder& DetailL
 				auto BGStruct = Cast<const UUserDefinedStruct>(StructProperty->Struct);
 				if (BGStruct && (EUserDefinedStructureStatus::UDSS_UpToDate != BGStruct->Status))
 				{
-					VariableProperty = NULL;
+					VariableProperty = nullptr;
 				}
 			}
 		}
 
 		// Find the class containing the variable
-		UClass* VariableClass = (VariableProperty != NULL) ? VariableProperty->GetTypedOuter<UClass>() : NULL;
+		UClass* VariableClass = (VariableProperty != NULL) ? VariableProperty->GetTypedOuter<UClass>() : nullptr;
 
 		FText ErrorMessage;
 		IDetailCategoryBuilder& DefaultValueCategory = DetailLayout.EditCategory(TEXT("DefaultValueCategory"), LOCTEXT("DefaultValueCategoryHeading", "Default Value"));
 
-		if (VariableProperty == NULL)
+		if (VariableProperty == nullptr)
 		{
-			if (Blueprint->Status != BS_UpToDate)
+			if (BlueprintObj->Status != BS_UpToDate)
 			{
 				ErrorMessage = LOCTEXT("VariableMissing_DirtyBlueprint", "Please compile the blueprint");
 			}
@@ -961,23 +961,23 @@ bool FBlueprintVarActionDetails::GetVariableNameChangeEnabled() const
 {
 	bool bIsReadOnly = true;
 
-	UBlueprint* Blueprint = GetBlueprintObj();
-	check(Blueprint != NULL);
+	UBlueprint* BlueprintObj = GetBlueprintObj();
+	check(BlueprintObj != nullptr);
 
 	UProperty* VariableProperty = CachedVariableProperty.Get();
 	if(VariableProperty != nullptr && IsVariableInBlueprint())
 	{
-		if(FBlueprintEditorUtils::FindNewVariableIndex(Blueprint, CachedVariableName) != INDEX_NONE)
+		if(FBlueprintEditorUtils::FindNewVariableIndex(BlueprintObj, CachedVariableName) != INDEX_NONE)
 		{
 			bIsReadOnly = false;
 		}
-		else if(Blueprint->FindTimelineTemplateByVariableName(CachedVariableName))
+		else if(BlueprintObj->FindTimelineTemplateByVariableName(CachedVariableName))
 		{
 			bIsReadOnly = false;
 		}
-		else if(IsASCSVariable(VariableProperty) && Blueprint->SimpleConstructionScript != NULL)
+		else if(IsASCSVariable(VariableProperty) && BlueprintObj->SimpleConstructionScript != nullptr)
 		{
-			if (USCS_Node* Node = Blueprint->SimpleConstructionScript->FindSCSNode(CachedVariableName))
+			if (USCS_Node* Node = BlueprintObj->SimpleConstructionScript->FindSCSNode(CachedVariableName))
 			{
 				bIsReadOnly = !FComponentEditorUtils::IsValidVariableNameString(Node->ComponentTemplate, Node->VariableName.ToString());
 			}
@@ -1000,13 +1000,13 @@ void FBlueprintVarActionDetails::OnVarNameChanged(const FText& InNewText)
 {
 	bIsVarNameInvalid = true;
 
-	UBlueprint* Blueprint = GetBlueprintObj();
-	check(Blueprint != NULL);
+	UBlueprint* BlueprintObj = GetBlueprintObj();
+	check(BlueprintObj != nullptr);
 
 	UProperty* VariableProperty = CachedVariableProperty.Get();
-	if(VariableProperty && IsASCSVariable(VariableProperty) && Blueprint->SimpleConstructionScript != NULL)
+	if(VariableProperty && IsASCSVariable(VariableProperty) && BlueprintObj->SimpleConstructionScript != nullptr)
 	{
-		for (USCS_Node* Node : Blueprint->SimpleConstructionScript->GetAllNodes())
+		for (USCS_Node* Node : BlueprintObj->SimpleConstructionScript->GetAllNodes())
 		{
 			if (Node && Node->VariableName == CachedVariableName && !FComponentEditorUtils::IsValidVariableNameString(Node->ComponentTemplate, InNewText.ToString()))
 			{
@@ -1016,7 +1016,7 @@ void FBlueprintVarActionDetails::OnVarNameChanged(const FText& InNewText)
 		}
 	}
 
-	TSharedPtr<INameValidatorInterface> NameValidator = MakeShareable(new FKismetNameValidator(Blueprint, CachedVariableName, GetLocalVariableScope(VariableProperty)));
+	TSharedPtr<INameValidatorInterface> NameValidator = MakeShareable(new FKismetNameValidator(BlueprintObj, CachedVariableName, GetLocalVariableScope(VariableProperty)));
 
 	EValidatorResult ValidatorResult = NameValidator->IsValid(InNewText.ToString());
 	if(ValidatorResult == EValidatorResult::AlreadyInUse)
@@ -1479,8 +1479,8 @@ void FBlueprintVarActionDetails::OnEditableChanged(ECheckBoxState InNewState)
 	// Toggle the flag on the blueprint's version of the variable description, based on state
 	const bool bVariableIsExposed = InNewState == ECheckBoxState::Checked;
 
-	UBlueprint* Blueprint = MyBlueprint.Pin()->GetBlueprintObj();
-	FBlueprintEditorUtils::SetBlueprintOnlyEditableFlag(Blueprint, VarName, !bVariableIsExposed);
+	UBlueprint* BlueprintObj = MyBlueprint.Pin()->GetBlueprintObj();
+	FBlueprintEditorUtils::SetBlueprintOnlyEditableFlag(BlueprintObj, VarName, !bVariableIsExposed);
 }
 
 ECheckBoxState FBlueprintVarActionDetails::OnCreateWidgetCheckboxState() const
@@ -1664,15 +1664,15 @@ EVisibility FBlueprintVarActionDetails::ExposeToMatineeVisibility() const
 
 ECheckBoxState FBlueprintVarActionDetails::OnGetConfigVariableCheckboxState() const
 {
-	UBlueprint* Blueprint = GetPropertyOwnerBlueprint();
+	UBlueprint* BlueprintObj = GetPropertyOwnerBlueprint();
 	const FName VarName = CachedVariableName;
 	ECheckBoxState CheckboxValue = ECheckBoxState::Unchecked;
 
-	if( Blueprint && VarName != NAME_None )
+	if( BlueprintObj && VarName != NAME_None )
 	{
-		const int32 VarIndex = FBlueprintEditorUtils::FindNewVariableIndex( Blueprint, VarName );
+		const int32 VarIndex = FBlueprintEditorUtils::FindNewVariableIndex( BlueprintObj, VarName );
 
-		if( VarIndex != INDEX_NONE && Blueprint->NewVariables[ VarIndex ].PropertyFlags & CPF_Config )
+		if( VarIndex != INDEX_NONE && BlueprintObj->NewVariables[ VarIndex ].PropertyFlags & CPF_Config )
 		{
 			CheckboxValue = ECheckBoxState::Checked;
 		}
@@ -1682,24 +1682,24 @@ ECheckBoxState FBlueprintVarActionDetails::OnGetConfigVariableCheckboxState() co
 
 void FBlueprintVarActionDetails::OnSetConfigVariableState( ECheckBoxState InNewState )
 {
-	UBlueprint* Blueprint = GetBlueprintObj();
+	UBlueprint* BlueprintObj = GetBlueprintObj();
 	const FName VarName = CachedVariableName;
 
-	if( Blueprint && VarName != NAME_None )
+	if( BlueprintObj && VarName != NAME_None )
 	{
-		const int32 VarIndex = FBlueprintEditorUtils::FindNewVariableIndex( Blueprint, VarName );
+		const int32 VarIndex = FBlueprintEditorUtils::FindNewVariableIndex( BlueprintObj, VarName );
 
 		if( VarIndex != INDEX_NONE )
 		{
 			if( InNewState == ECheckBoxState::Checked )
 			{
-				Blueprint->NewVariables[ VarIndex ].PropertyFlags |= CPF_Config;
+				BlueprintObj->NewVariables[ VarIndex ].PropertyFlags |= CPF_Config;
 			}
 			else
 			{
-				Blueprint->NewVariables[ VarIndex ].PropertyFlags &= ~CPF_Config;
+				BlueprintObj->NewVariables[ VarIndex ].PropertyFlags &= ~CPF_Config;
 			}
-			FBlueprintEditorUtils::MarkBlueprintAsStructurallyModified( Blueprint );
+			FBlueprintEditorUtils::MarkBlueprintAsStructurallyModified( BlueprintObj );
 		}
 	}
 }
@@ -1722,10 +1722,10 @@ FText FBlueprintVarActionDetails::OnGetMetaKeyValue(FName Key) const
 	FName VarName = CachedVariableName;
 	if (VarName != NAME_None)
 	{
-		if ( UBlueprint* Blueprint = GetPropertyOwnerBlueprint() )
+		if ( UBlueprint* BlueprintObj = GetPropertyOwnerBlueprint() )
 		{
 			FString Result;
-			FBlueprintEditorUtils::GetBlueprintVariableMetaData(GetPropertyOwnerBlueprint(), VarName, GetLocalVariableScope(CachedVariableProperty.Get()), Key, /*out*/ Result);
+			FBlueprintEditorUtils::GetBlueprintVariableMetaData(BlueprintObj, VarName, GetLocalVariableScope(CachedVariableProperty.Get()), Key, /*out*/ Result);
 
 			return FText::FromString(Result);
 		}
@@ -4780,8 +4780,8 @@ void FBlueprintComponentDetails::CustomizeDetails(IDetailLayoutBuilder& DetailLa
 	check( BlueprintEditorPtr.IsValid() );
 	TSharedPtr<SSCSEditor> Editor = BlueprintEditorPtr.Pin()->GetSCSEditor();
 	check( Editor.IsValid() );
-	const UBlueprint* Blueprint = GetBlueprintObj();
-	check(Blueprint != NULL);
+	const UBlueprint* BlueprintObj = GetBlueprintObj();
+	check(BlueprintObj != nullptr);
 
 	TArray<FSCSEditorTreeNodePtrType> Nodes = Editor->GetSelectedNodes();
 
@@ -4928,10 +4928,10 @@ void FBlueprintComponentDetails::CustomizeDetails(IDetailLayoutBuilder& DetailLa
 	}
 
 	// Handle event generation
-	if ( FBlueprintEditorUtils::DoesSupportEventGraphs(Blueprint) && Nodes.Num() == 1 )
+	if ( FBlueprintEditorUtils::DoesSupportEventGraphs(BlueprintObj) && Nodes.Num() == 1 )
 	{
 		FName PropertyName = CachedNodePtr->GetVariableName();
-		UObjectProperty* VariableProperty = FindField<UObjectProperty>(Blueprint->SkeletonGeneratedClass, PropertyName);
+		UObjectProperty* VariableProperty = FindField<UObjectProperty>(BlueprintObj->SkeletonGeneratedClass, PropertyName);
 
 		AddEventsCategory(DetailLayout, VariableProperty);
 	}
@@ -5093,13 +5093,13 @@ TSharedRef< ITableRow > FBlueprintComponentDetails::MakeVariableCategoryViewWidg
 
 void FBlueprintComponentDetails::PopulateVariableCategories()
 {
-	UBlueprint* Blueprint = GetBlueprintObj();
+	UBlueprint* BlueprintObj = GetBlueprintObj();
 
-	check(Blueprint);
-	check(Blueprint->SkeletonGeneratedClass);
+	check(BlueprintObj);
+	check(BlueprintObj->SkeletonGeneratedClass);
 
 	TArray<FName> VisibleVariables;
-	for (TFieldIterator<UProperty> PropertyIt(Blueprint->SkeletonGeneratedClass, EFieldIteratorFlags::IncludeSuper); PropertyIt; ++PropertyIt)
+	for (TFieldIterator<UProperty> PropertyIt(BlueprintObj->SkeletonGeneratedClass, EFieldIteratorFlags::IncludeSuper); PropertyIt; ++PropertyIt)
 	{
 		UProperty* Property = *PropertyIt;
 
@@ -5109,14 +5109,14 @@ void FBlueprintComponentDetails::PopulateVariableCategories()
 		}
 	}
 
-	FBlueprintEditorUtils::GetSCSVariableNameList(Blueprint, VisibleVariables);
+	FBlueprintEditorUtils::GetSCSVariableNameList(BlueprintObj, VisibleVariables);
 
 	VariableCategorySource.Empty();
 	VariableCategorySource.Add(MakeShareable(new FText(LOCTEXT("Default", "Default"))));
 	for (int32 i = 0; i < VisibleVariables.Num(); ++i)
 	{
-		FText Category = FBlueprintEditorUtils::GetBlueprintVariableCategory(Blueprint, VisibleVariables[i], NULL);
-		if (!Category.IsEmpty() && !Category.EqualTo(FText::FromString(Blueprint->GetName())))
+		FText Category = FBlueprintEditorUtils::GetBlueprintVariableCategory(BlueprintObj, VisibleVariables[i], nullptr);
+		if (!Category.IsEmpty() && !Category.EqualTo(FText::FromString(BlueprintObj->GetName())))
 		{
 			bool bNewCategory = true;
 			for (int32 j = 0; j < VariableCategorySource.Num() && bNewCategory; ++j)
diff --git a/Engine/Source/Editor/Kismet/Private/BlueprintEditor.cpp b/Engine/Source/Editor/Kismet/Private/BlueprintEditor.cpp
--- a/Engine/Source/Editor/Kismet/Private/BlueprintEditor.cpp
+++ b/Engine/Source/Editor/Kismet/Private/BlueprintEditor.cpp
@@ -1599,12 +1599,12 @@ void FBlueprintEditor::LoadLibrariesFromAssetRegistry()
 	}
 }
 
-void FBlueprintEditor::RegisterTabSpawners(const TSharedRef<class FTabManager>& TabManager)
+void FBlueprintEditor::RegisterTabSpawners(const TSharedRef<class FTabManager>& InTabManager)
 {
 	//@TODO: Can't we do this sooner?
-	DocumentManager->SetTabManager(TabManager);
+	DocumentManager->SetTabManager(InTabManager);
 
-	FWorkflowCentricApplication::RegisterTabSpawners(TabManager);
+	FWorkflowCentricApplication::RegisterTabSpawners(InTabManager);
 }
 
 void FBlueprintEditor::SetCurrentMode(FName NewMode)
@@ -3568,9 +3568,9 @@ bool FBlueprintEditor::CanAccessComponentsMode() const
 	return bCanAccess;
 }
 
-void FBlueprintEditor::RegisterToolbarTab(const TSharedRef<class FTabManager>& TabManager)
+void FBlueprintEditor::RegisterToolbarTab(const TSharedRef<class FTabManager>& InTabManager)
 {
-	FAssetEditorToolkit::RegisterTabSpawners(TabManager);
+	FAssetEditorToolkit::RegisterTabSpawners(InTabManager);
 }
 
 void FBlueprintEditor::LogSimpleMessage(const FText& MessageText)
@@ -7367,7 +7367,7 @@ FText FBlueprintEditor::GetBaseToolkitName() const
 
 FText FBlueprintEditor::GetToolkitName() const
 {
-	const auto EditingObjects = GetEditingObjects();
+	const TArray<UObject*>& EditingObjs = GetEditingObjects();
 
 	if( IsEditingSingleBlueprint() )
 	{
@@ -7392,7 +7392,7 @@ FText FBlueprintEditor::GetToolkitName() const
 
 	TSubclassOf< UObject > SharedParentClass = NULL;
 
-	for( auto ObjectIter = EditingObjects.CreateConstIterator(); ObjectIter; ++ObjectIter )
+	for( auto ObjectIter = EditingObjs.CreateConstIterator(); ObjectIter; ++ObjectIter )
 	{
 		UBlueprint* Blueprint = Cast<UBlueprint>( *ObjectIter );;
 		check( Blueprint );
@@ -7412,14 +7412,14 @@ FText FBlueprintEditor::GetToolkitName() const
 	}
 
 	FFormatNamedArguments Args;
-	Args.Add( TEXT("NumberOfObjects"), EditingObjects.Num() );
+	Args.Add( TEXT("NumberOfObjects"), EditingObjs.Num() );
 	Args.Add( TEXT("ObjectName"), FText::FromString( SharedParentClass->GetName() ) );
 	return FText::Format( NSLOCTEXT("KismetEditor", "ToolkitTitle_UniqueLayerName", "{NumberOfObjects} {ClassName} - Class Defaults"), Args );
 }
 
 FText FBlueprintEditor::GetToolkitToolTipText() const
 {
-	const auto EditingObjects = GetEditingObjects();
+	const TArray<UObject*>& EditingObjs = GetEditingObjects();
 
 	if( IsEditingSingleBlueprint() )
 	{
@@ -7439,7 +7439,7 @@ FText FBlueprintEditor::GetToolkitToolTipText() const
 
 	TSubclassOf< UObject > SharedParentClass = NULL;
 
-	for( auto ObjectIter = EditingObjects.CreateConstIterator(); ObjectIter; ++ObjectIter )
+	for( auto ObjectIter = EditingObjs.CreateConstIterator(); ObjectIter; ++ObjectIter )
 	{
 		UBlueprint* Blueprint = Cast<UBlueprint>( *ObjectIter );;
 		check( Blueprint );
@@ -7449,7 +7449,7 @@ FText FBlueprintEditor::GetToolkitToolTipText() const
 		{
 			SharedParentClass = Blueprint->ParentClass;
 		}
-
+		 
 		// If we've encountered an object that's not a subclass of the current best baseclass,
 		// climb up a step in the class hierarchy.
 		while( !Blueprint->ParentClass->IsChildOf( SharedParentClass ) )
@@ -7459,7 +7459,7 @@ FText FBlueprintEditor::GetToolkitToolTipText() const
 	}
 
 	FFormatNamedArguments Args;
-	Args.Add( TEXT("NumberOfObjects"), EditingObjects.Num() );
+	Args.Add( TEXT("NumberOfObjects"), EditingObjs.Num() );
 	Args.Add( TEXT("ObjectName"), FText::FromString( SharedParentClass->GetName() ) );
 	return FText::Format( NSLOCTEXT("KismetEditor", "ToolkitTitle_UniqueLayerName", "{NumberOfObjects} {ClassName} - Class Defaults"), Args );
 }
diff --git a/Engine/Source/Editor/Kismet/Private/SSCSEditor.cpp b/Engine/Source/Editor/Kismet/Private/SSCSEditor.cpp
--- a/Engine/Source/Editor/Kismet/Private/SSCSEditor.cpp
+++ b/Engine/Source/Editor/Kismet/Private/SSCSEditor.cpp
@@ -2951,7 +2951,7 @@ TSharedRef<SWidget> SSCS_RowWidget_ActorRoot::GenerateWidgetForColumn(const FNam
 	ensure(ColumnName == SCS_ColumnName_ComponentClass);
 
 	// Create the name field
-	TSharedPtr<SInlineEditableTextBlock> InlineWidget =
+	TSharedPtr<SInlineEditableTextBlock> InlineEditableWidget =
 		SNew(SInlineEditableTextBlock)
 		.Text(this, &SSCS_RowWidget_ActorRoot::GetActorDisplayText)
 		.OnVerifyTextChanged(this, &SSCS_RowWidget_ActorRoot::OnVerifyActorLabelChanged)
@@ -2959,7 +2959,7 @@ TSharedRef<SWidget> SSCS_RowWidget_ActorRoot::GenerateWidgetForColumn(const FNam
 		.IsSelected(this, &SSCS_RowWidget_ActorRoot::IsSelectedExclusively)
 		.IsReadOnly(!NodePtr->CanRename() || (SCSEditor.IsValid() && !SCSEditor.Pin()->IsEditingAllowed()));
 
-	NodePtr->SetRenameRequestedDelegate(FSCSEditorTreeNode::FOnRenameRequested::CreateSP(InlineWidget.Get(), &SInlineEditableTextBlock::EnterEditingMode));
+	NodePtr->SetRenameRequestedDelegate(FSCSEditorTreeNode::FOnRenameRequested::CreateSP(InlineEditableWidget.Get(), &SInlineEditableTextBlock::EnterEditingMode));
 
 	return SNew(SHorizontalBox)
 		.ToolTip(CreateToolTipWidget())
@@ -2979,7 +2979,7 @@ TSharedRef<SWidget> SSCS_RowWidget_ActorRoot::GenerateWidgetForColumn(const FNam
 		.VAlign(VAlign_Center)
 		.Padding(0.0f, 0.0f)
 		[
-			InlineWidget.ToSharedRef()
+			InlineEditableWidget.ToSharedRef()
 		]
 
 	+SHorizontalBox::Slot()
@@ -5492,30 +5492,30 @@ void SSCSEditor::GetCollapsedNodes(const FSCSEditorTreeNodePtrType& InNodePtr, T
 
 EVisibility SSCSEditor::GetPromoteToBlueprintButtonVisibility() const
 {
-	EVisibility Visibility = EVisibility::Collapsed;
+	EVisibility ButtonVisibility = EVisibility::Collapsed;
 	if (EditorMode == EComponentEditorMode::ActorInstance)
 	{
 		if (GetBlueprint() == nullptr)
 		{
-			Visibility = EVisibility::Visible;
+			ButtonVisibility = EVisibility::Visible;
 		}
 	}
 
-	return Visibility;
+	return ButtonVisibility;
 }
 
 EVisibility SSCSEditor::GetEditBlueprintButtonVisibility() const
 {
-	EVisibility Visibility = EVisibility::Collapsed;
+	EVisibility ButtonVisibility = EVisibility::Collapsed;
 	if (EditorMode == EComponentEditorMode::ActorInstance)
 	{
 		if (GetBlueprint() != nullptr)
 		{
-			Visibility = EVisibility::Visible;
+			ButtonVisibility = EVisibility::Visible;
 		}
 	}
 
-	return Visibility;
+	return ButtonVisibility;
 }
 
 FText SSCSEditor::OnGetApplyChangesToBlueprintTooltip() const
diff --git a/Engine/Source/Editor/Kismet/Private/UserDefinedEnumEditor.cpp b/Engine/Source/Editor/Kismet/Private/UserDefinedEnumEditor.cpp
--- a/Engine/Source/Editor/Kismet/Private/UserDefinedEnumEditor.cpp
+++ b/Engine/Source/Editor/Kismet/Private/UserDefinedEnumEditor.cpp
@@ -14,23 +14,23 @@
 const FName FUserDefinedEnumEditor::EnumeratorsTabId( TEXT( "UserDefinedEnum_EnumeratorEditor" ) );
 const FName FUserDefinedEnumEditor::UserDefinedEnumEditorAppIdentifier( TEXT( "UserDefinedEnumEditorApp" ) );
 
-void FUserDefinedEnumEditor::RegisterTabSpawners(const TSharedRef<class FTabManager>& TabManager)
+void FUserDefinedEnumEditor::RegisterTabSpawners(const TSharedRef<class FTabManager>& InTabManager)
 {
-	WorkspaceMenuCategory = TabManager->AddLocalWorkspaceMenuCategory(LOCTEXT("WorkspaceMenu_UserDefinedEnumEditor", "User-Defined Enum Editor"));
+	WorkspaceMenuCategory = InTabManager->AddLocalWorkspaceMenuCategory(LOCTEXT("WorkspaceMenu_UserDefinedEnumEditor", "User-Defined Enum Editor"));
 
-	FAssetEditorToolkit::RegisterTabSpawners(TabManager);
+	FAssetEditorToolkit::RegisterTabSpawners(InTabManager);
 
-	TabManager->RegisterTabSpawner( EnumeratorsTabId, FOnSpawnTab::CreateSP(this, &FUserDefinedEnumEditor::SpawnEnumeratorsTab) )
+	InTabManager->RegisterTabSpawner( EnumeratorsTabId, FOnSpawnTab::CreateSP(this, &FUserDefinedEnumEditor::SpawnEnumeratorsTab) )
 		.SetDisplayName( LOCTEXT("EnumeratorEditor", "Enumerators") )
 		.SetGroup(WorkspaceMenuCategory.ToSharedRef())
 		.SetIcon(FSlateIcon(FEditorStyle::GetStyleSetName(), "GraphEditor.Enum_16x"));
 }
 
-void FUserDefinedEnumEditor::UnregisterTabSpawners(const TSharedRef<class FTabManager>& TabManager)
+void FUserDefinedEnumEditor::UnregisterTabSpawners(const TSharedRef<class FTabManager>& InTabManager)
 {
-	FAssetEditorToolkit::UnregisterTabSpawners(TabManager);
+	FAssetEditorToolkit::UnregisterTabSpawners(InTabManager);
 
-	TabManager->UnregisterTabSpawner( EnumeratorsTabId );
+	InTabManager->UnregisterTabSpawner( EnumeratorsTabId );
 }
 
 void FUserDefinedEnumEditor::InitEditor(const EToolkitMode::Type Mode, const TSharedPtr< class IToolkitHost >& InitToolkitHost, UUserDefinedEnum* EnumToEdit)
@@ -77,10 +77,10 @@ TSharedRef<SDockTab> FUserDefinedEnumEditor::SpawnEnumeratorsTab(const FSpawnTab
 	check( Args.GetTabId() == EnumeratorsTabId );
 
 	UUserDefinedEnum* EditedEnum = NULL;
-	const auto EditingObjects = GetEditingObjects();
-	if (EditingObjects.Num())
+	const TArray<UObject*>& EditingObjs = GetEditingObjects();
+	if (EditingObjs.Num())
 	{
-		EditedEnum = Cast<UUserDefinedEnum>(EditingObjects[ 0 ]);
+		EditedEnum = Cast<UUserDefinedEnum>(EditingObjs[ 0 ]);
 	}
 
 	// Create a property view
diff --git a/Engine/Source/Editor/Kismet/Private/UserDefinedStructureEditor.cpp b/Engine/Source/Editor/Kismet/Private/UserDefinedStructureEditor.cpp
--- a/Engine/Source/Editor/Kismet/Private/UserDefinedStructureEditor.cpp
+++ b/Engine/Source/Editor/Kismet/Private/UserDefinedStructureEditor.cpp
@@ -96,8 +96,8 @@ public:
 		ViewArgs.NotifyHook = this;
 
 		DetailsView = PropertyModule.CreateDetailView(ViewArgs);
-		TWeakPtr< FStructureDefaultValueView > WeakThis = SharedThis(this);
-		FOnGetDetailCustomizationInstance LayoutStructDetails = FOnGetDetailCustomizationInstance::CreateStatic(&FDefaultValueDetails::MakeInstance, WeakThis, StructData);
+		TWeakPtr< FStructureDefaultValueView > LocalWeakThis = SharedThis(this);
+		FOnGetDetailCustomizationInstance LayoutStructDetails = FOnGetDetailCustomizationInstance::CreateStatic(&FDefaultValueDetails::MakeInstance, LocalWeakThis, StructData);
 		DetailsView->RegisterInstancedCustomPropertyLayout(UUserDefinedStruct::StaticClass(), LayoutStructDetails);
 		DetailsView->SetObject(UserDefinedStruct.Get());
 	}
@@ -246,23 +246,23 @@ private:
 const FName FUserDefinedStructureEditor::MemberVariablesTabId( TEXT( "UserDefinedStruct_MemberVariablesEditor" ) );
 const FName FUserDefinedStructureEditor::UserDefinedStructureEditorAppIdentifier( TEXT( "UserDefinedStructEditorApp" ) );
 
-void FUserDefinedStructureEditor::RegisterTabSpawners(const TSharedRef<class FTabManager>& TabManager)
+void FUserDefinedStructureEditor::RegisterTabSpawners(const TSharedRef<class FTabManager>& InTabManager)
 {
-	WorkspaceMenuCategory = TabManager->AddLocalWorkspaceMenuCategory(LOCTEXT("WorkspaceMenu_UserDefinedStructureEditor", "User-Defined Structure Editor"));
+	WorkspaceMenuCategory = InTabManager->AddLocalWorkspaceMenuCategory(LOCTEXT("WorkspaceMenu_UserDefinedStructureEditor", "User-Defined Structure Editor"));
 
-	FAssetEditorToolkit::RegisterTabSpawners(TabManager);
+	FAssetEditorToolkit::RegisterTabSpawners(InTabManager);
 
-	TabManager->RegisterTabSpawner( MemberVariablesTabId, FOnSpawnTab::CreateSP(this, &FUserDefinedStructureEditor::SpawnStructureTab) )
+	InTabManager->RegisterTabSpawner( MemberVariablesTabId, FOnSpawnTab::CreateSP(this, &FUserDefinedStructureEditor::SpawnStructureTab) )
 		.SetDisplayName( LOCTEXT("MemberVariablesEditor", "Member Variables") )
 		.SetGroup(WorkspaceMenuCategory.ToSharedRef())
 		.SetIcon(FSlateIcon(FEditorStyle::GetStyleSetName(), "Kismet.Tabs.Variables"));
 }
 
-void FUserDefinedStructureEditor::UnregisterTabSpawners(const TSharedRef<class FTabManager>& TabManager)
+void FUserDefinedStructureEditor::UnregisterTabSpawners(const TSharedRef<class FTabManager>& InTabManager)
 {
-	FAssetEditorToolkit::UnregisterTabSpawners(TabManager);
+	FAssetEditorToolkit::UnregisterTabSpawners(InTabManager);
 
-	TabManager->UnregisterTabSpawner( MemberVariablesTabId );
+	InTabManager->UnregisterTabSpawner( MemberVariablesTabId );
 }
 
 void FUserDefinedStructureEditor::InitEditor(const EToolkitMode::Type Mode, const TSharedPtr< class IToolkitHost >& InitToolkitHost, UUserDefinedStruct* Struct)
@@ -341,10 +341,10 @@ TSharedRef<SDockTab> FUserDefinedStructureEditor::SpawnStructureTab(const FSpawn
 	check( Args.GetTabId() == MemberVariablesTabId );
 
 	UUserDefinedStruct* EditedStruct = NULL;
-	const auto EditingObjects = GetEditingObjects();
-	if (EditingObjects.Num())
+	const TArray<UObject*>& EditingObjs = GetEditingObjects();
+	if (EditingObjs.Num())
 	{
-		EditedStruct = Cast<UUserDefinedStruct>(EditingObjects[ 0 ]);
+		EditedStruct = Cast<UUserDefinedStruct>(EditingObjs[ 0 ]);
 	}
 
 	auto Box = SNew(SHorizontalBox);
diff --git a/Engine/Source/Editor/Kismet/Private/WorkflowCentricApplication.cpp b/Engine/Source/Editor/Kismet/Private/WorkflowCentricApplication.cpp
--- a/Engine/Source/Editor/Kismet/Private/WorkflowCentricApplication.cpp
+++ b/Engine/Source/Editor/Kismet/Private/WorkflowCentricApplication.cpp
@@ -11,17 +11,17 @@
 
 TArray<FWorkflowApplicationModeExtender> FWorkflowCentricApplication::ModeExtenderList;
 
-void FWorkflowCentricApplication::RegisterTabSpawners(const TSharedRef<class FTabManager>& TabManager)
+void FWorkflowCentricApplication::RegisterTabSpawners(const TSharedRef<class FTabManager>& InTabManager)
 {
 	if (CurrentAppModePtr.IsValid())
 	{
-		CurrentAppModePtr->RegisterTabFactories(TabManager);
+		CurrentAppModePtr->RegisterTabFactories(InTabManager);
 	}
 }
 
-void FWorkflowCentricApplication::UnregisterTabSpawners(const TSharedRef<class FTabManager>& TabManager)
+void FWorkflowCentricApplication::UnregisterTabSpawners(const TSharedRef<class FTabManager>& InTabManager)
 {
-	TabManager->UnregisterAllTabSpawners();
+	InTabManager->UnregisterAllTabSpawners();
 }
 
 
diff --git a/Engine/Source/Editor/MaterialEditor/Private/MaterialEditor.cpp b/Engine/Source/Editor/MaterialEditor/Private/MaterialEditor.cpp
--- a/Engine/Source/Editor/MaterialEditor/Private/MaterialEditor.cpp
+++ b/Engine/Source/Editor/MaterialEditor/Private/MaterialEditor.cpp
@@ -173,65 +173,65 @@ void FMatExpressionPreview::NotifyCompilationFinished()
 // FMaterialEditor //
 /////////////////////
 
-void FMaterialEditor::RegisterTabSpawners(const TSharedRef<class FTabManager>& TabManager)
+void FMaterialEditor::RegisterTabSpawners(const TSharedRef<class FTabManager>& InTabManager)
 {
-	WorkspaceMenuCategory = TabManager->AddLocalWorkspaceMenuCategory(LOCTEXT("WorkspaceMenu_MaterialEditor", "Material Editor"));
+	WorkspaceMenuCategory = InTabManager->AddLocalWorkspaceMenuCategory(LOCTEXT("WorkspaceMenu_MaterialEditor", "Material Editor"));
 	auto WorkspaceMenuCategoryRef = WorkspaceMenuCategory.ToSharedRef();
 
-	FAssetEditorToolkit::RegisterTabSpawners(TabManager);
+	FAssetEditorToolkit::RegisterTabSpawners(InTabManager);
 	
-	TabManager->RegisterTabSpawner( PreviewTabId, FOnSpawnTab::CreateSP(this, &FMaterialEditor::SpawnTab_Preview) )
+	InTabManager->RegisterTabSpawner( PreviewTabId, FOnSpawnTab::CreateSP(this, &FMaterialEditor::SpawnTab_Preview) )
 		.SetDisplayName( LOCTEXT("ViewportTab", "Viewport") )
 		.SetGroup( WorkspaceMenuCategoryRef )
 		.SetIcon(FSlateIcon(FEditorStyle::GetStyleSetName(), "LevelEditor.Tabs.Viewports"));
 	
-	TabManager->RegisterTabSpawner( GraphCanvasTabId, FOnSpawnTab::CreateSP(this, &FMaterialEditor::SpawnTab_GraphCanvas) )
+	InTabManager->RegisterTabSpawner( GraphCanvasTabId, FOnSpawnTab::CreateSP(this, &FMaterialEditor::SpawnTab_GraphCanvas) )
 		.SetDisplayName( LOCTEXT("GraphCanvasTab", "Graph") )
 		.SetGroup( WorkspaceMenuCategoryRef )
 		.SetIcon(FSlateIcon(FEditorStyle::GetStyleSetName(), "GraphEditor.EventGraph_16x"));
 	
-	TabManager->RegisterTabSpawner( PropertiesTabId, FOnSpawnTab::CreateSP(this, &FMaterialEditor::SpawnTab_MaterialProperties) )
+	InTabManager->RegisterTabSpawner( PropertiesTabId, FOnSpawnTab::CreateSP(this, &FMaterialEditor::SpawnTab_MaterialProperties) )
 		.SetDisplayName( LOCTEXT("DetailsTab", "Details") )
 		.SetGroup( WorkspaceMenuCategoryRef )
 		.SetIcon(FSlateIcon(FEditorStyle::GetStyleSetName(), "LevelEditor.Tabs.Details"));
 
-	TabManager->RegisterTabSpawner( PaletteTabId, FOnSpawnTab::CreateSP(this, &FMaterialEditor::SpawnTab_Palette) )
+	InTabManager->RegisterTabSpawner( PaletteTabId, FOnSpawnTab::CreateSP(this, &FMaterialEditor::SpawnTab_Palette) )
 		.SetDisplayName( LOCTEXT("PaletteTab", "Palette") )
 		.SetGroup( WorkspaceMenuCategoryRef )
 		.SetIcon(FSlateIcon(FEditorStyle::GetStyleSetName(), "Kismet.Tabs.Palette"));
 
-	TabManager->RegisterTabSpawner( StatsTabId, FOnSpawnTab::CreateSP(this, &FMaterialEditor::SpawnTab_Stats) )
+	InTabManager->RegisterTabSpawner( StatsTabId, FOnSpawnTab::CreateSP(this, &FMaterialEditor::SpawnTab_Stats) )
 		.SetDisplayName( LOCTEXT("StatsTab", "Stats") )
 		.SetGroup( WorkspaceMenuCategoryRef )
 		.SetIcon(FSlateIcon(FEditorStyle::GetStyleSetName(), "LevelEditor.Tabs.StatsViewer"));
 	
-	TabManager->RegisterTabSpawner(FindTabId, FOnSpawnTab::CreateSP(this, &FMaterialEditor::SpawnTab_Find))
+	InTabManager->RegisterTabSpawner(FindTabId, FOnSpawnTab::CreateSP(this, &FMaterialEditor::SpawnTab_Find))
 		.SetDisplayName(LOCTEXT("FindTab", "Find Results"))
 		.SetGroup( WorkspaceMenuCategoryRef )
 		.SetIcon(FSlateIcon(FEditorStyle::GetStyleSetName(), "Kismet.Tabs.FindResults"));
 
-	TabManager->RegisterTabSpawner( HLSLCodeTabId, FOnSpawnTab::CreateSP(this, &FMaterialEditor::SpawnTab_HLSLCode) )
+	InTabManager->RegisterTabSpawner( HLSLCodeTabId, FOnSpawnTab::CreateSP(this, &FMaterialEditor::SpawnTab_HLSLCode) )
 		.SetDisplayName( LOCTEXT("HLSLCodeTab", "HLSL Code") )
 		.SetGroup( WorkspaceMenuCategoryRef )
 		.SetIcon(FSlateIcon(FEditorStyle::GetStyleSetName(), "MaterialEditor.Tabs.HLSLCode"));
 
-	OnRegisterTabSpawners().Broadcast(TabManager);
+	OnRegisterTabSpawners().Broadcast(InTabManager);
 }
 
 
-void FMaterialEditor::UnregisterTabSpawners(const TSharedRef<class FTabManager>& TabManager)
+void FMaterialEditor::UnregisterTabSpawners(const TSharedRef<class FTabManager>& InTabManager)
 {
-	FAssetEditorToolkit::UnregisterTabSpawners(TabManager);
+	FAssetEditorToolkit::UnregisterTabSpawners(InTabManager);
 
-	TabManager->UnregisterTabSpawner( PreviewTabId );
-	TabManager->UnregisterTabSpawner( GraphCanvasTabId );
-	TabManager->UnregisterTabSpawner( PropertiesTabId );
-	TabManager->UnregisterTabSpawner( PaletteTabId );
-	TabManager->UnregisterTabSpawner( StatsTabId );
-	TabManager->UnregisterTabSpawner( FindTabId );
-	TabManager->UnregisterTabSpawner( HLSLCodeTabId );
+	InTabManager->UnregisterTabSpawner( PreviewTabId );
+	InTabManager->UnregisterTabSpawner( GraphCanvasTabId );
+	InTabManager->UnregisterTabSpawner( PropertiesTabId );
+	InTabManager->UnregisterTabSpawner( PaletteTabId );
+	InTabManager->UnregisterTabSpawner( StatsTabId );
+	InTabManager->UnregisterTabSpawner( FindTabId );
+	InTabManager->UnregisterTabSpawner( HLSLCodeTabId );
 
-	OnUnregisterTabSpawners().Broadcast(TabManager);
+	OnUnregisterTabSpawners().Broadcast(InTabManager);
 }
 
 void FMaterialEditor::InitEditorForMaterial(UMaterial* InMaterial)
diff --git a/Engine/Source/Editor/MaterialEditor/Private/MaterialInstanceEditor.cpp b/Engine/Source/Editor/MaterialEditor/Private/MaterialInstanceEditor.cpp
--- a/Engine/Source/Editor/MaterialEditor/Private/MaterialInstanceEditor.cpp
+++ b/Engine/Source/Editor/MaterialEditor/Private/MaterialInstanceEditor.cpp
@@ -109,40 +109,40 @@ protected:
 
 
 
-void FMaterialInstanceEditor::RegisterTabSpawners(const TSharedRef<class FTabManager>& TabManager)
+void FMaterialInstanceEditor::RegisterTabSpawners(const TSharedRef<class FTabManager>& InTabManager)
 {
-	WorkspaceMenuCategory = TabManager->AddLocalWorkspaceMenuCategory(LOCTEXT("WorkspaceMenu_MaterialInstanceEditor", "Material Instance Editor"));
+	WorkspaceMenuCategory = InTabManager->AddLocalWorkspaceMenuCategory(LOCTEXT("WorkspaceMenu_MaterialInstanceEditor", "Material Instance Editor"));
 	auto WorkspaceMenuCategoryRef = WorkspaceMenuCategory.ToSharedRef();
 
-	FAssetEditorToolkit::RegisterTabSpawners(TabManager);
+	FAssetEditorToolkit::RegisterTabSpawners(InTabManager);
 	
-	TabManager->RegisterTabSpawner( PreviewTabId, FOnSpawnTab::CreateSP( this, &FMaterialInstanceEditor::SpawnTab_Preview ) )
+	InTabManager->RegisterTabSpawner( PreviewTabId, FOnSpawnTab::CreateSP( this, &FMaterialInstanceEditor::SpawnTab_Preview ) )
 		.SetDisplayName( LOCTEXT( "ViewportTab", "Viewport" ) )
 		.SetGroup( WorkspaceMenuCategoryRef )
 		.SetIcon( FSlateIcon( FEditorStyle::GetStyleSetName(), "LevelEditor.Tabs.Viewports" ) );
 	
-	TabManager->RegisterTabSpawner( PropertiesTabId, FOnSpawnTab::CreateSP( this, &FMaterialInstanceEditor::SpawnTab_Properties ) )
+	InTabManager->RegisterTabSpawner( PropertiesTabId, FOnSpawnTab::CreateSP( this, &FMaterialInstanceEditor::SpawnTab_Properties ) )
 		.SetDisplayName( LOCTEXT( "PropertiesTab", "Details" ) )
 		.SetGroup( WorkspaceMenuCategoryRef )
 		.SetIcon( FSlateIcon( FEditorStyle::GetStyleSetName(), "LevelEditor.Tabs.Details" ) );
 	
-	TabManager->RegisterTabSpawner( ParentsTabId, FOnSpawnTab::CreateSP(this, &FMaterialInstanceEditor::SpawnTab_Parents) )
+	InTabManager->RegisterTabSpawner( ParentsTabId, FOnSpawnTab::CreateSP(this, &FMaterialInstanceEditor::SpawnTab_Parents) )
 		.SetDisplayName( LOCTEXT("ParentsTab", "Parents") )
 		.SetGroup( WorkspaceMenuCategoryRef )
 		.SetIcon( FSlateIcon( FEditorStyle::GetStyleSetName(), "Kismet.Tabs.Palette" ) );
 
-	OnRegisterTabSpawners().Broadcast(TabManager);
+	OnRegisterTabSpawners().Broadcast(InTabManager);
 }
 
-void FMaterialInstanceEditor::UnregisterTabSpawners(const TSharedRef<class FTabManager>& TabManager)
+void FMaterialInstanceEditor::UnregisterTabSpawners(const TSharedRef<class FTabManager>& InTabManager)
 {
-	FAssetEditorToolkit::UnregisterTabSpawners(TabManager);
+	FAssetEditorToolkit::UnregisterTabSpawners(InTabManager);
 
-	TabManager->UnregisterTabSpawner( PreviewTabId );		
-	TabManager->UnregisterTabSpawner( PropertiesTabId );	
-	TabManager->UnregisterTabSpawner( ParentsTabId );		
+	InTabManager->UnregisterTabSpawner( PreviewTabId );		
+	InTabManager->UnregisterTabSpawner( PropertiesTabId );	
+	InTabManager->UnregisterTabSpawner( ParentsTabId );		
 
-	OnUnregisterTabSpawners().Broadcast(TabManager);
+	OnUnregisterTabSpawners().Broadcast(InTabManager);
 }
 
 
diff --git a/Engine/Source/Editor/MaterialEditor/Private/SMaterialEditorViewport.cpp b/Engine/Source/Editor/MaterialEditor/Private/SMaterialEditorViewport.cpp
--- a/Engine/Source/Editor/MaterialEditor/Private/SMaterialEditorViewport.cpp
+++ b/Engine/Source/Editor/MaterialEditor/Private/SMaterialEditorViewport.cpp
@@ -92,10 +92,10 @@ void FMaterialEditorViewportClient::Tick(float DeltaSeconds)
 }
 
 
-void FMaterialEditorViewportClient::Draw(FViewport* Viewport,FCanvas* Canvas)
+void FMaterialEditorViewportClient::Draw(FViewport* InViewport,FCanvas* Canvas)
 {
-	FEditorViewportClient::Draw(Viewport, Canvas);
-	MaterialEditorPtr.Pin()->DrawMessages(Viewport, Canvas);
+	FEditorViewportClient::Draw(InViewport, Canvas);
+	MaterialEditorPtr.Pin()->DrawMessages(InViewport, Canvas);
 }
 
 bool FMaterialEditorViewportClient::ShouldOrbitCamera() const
diff --git a/Engine/Source/Editor/Matinee/Private/Matinee.cpp b/Engine/Source/Editor/Matinee/Private/Matinee.cpp
--- a/Engine/Source/Editor/Matinee/Private/Matinee.cpp
+++ b/Engine/Source/Editor/Matinee/Private/Matinee.cpp
@@ -128,43 +128,43 @@ static const FName MatineeCurveEdName("Matinee_CurveEditor");
 static const FName MatineeTrackWindowName("Matinee_TrackWindow");
 static const FName MatineePropertyWindowName("Matinee_PropertyWindow");
 
-void FMatinee::RegisterTabSpawners(const TSharedRef<class FTabManager>& TabManager)
+void FMatinee::RegisterTabSpawners(const TSharedRef<class FTabManager>& InTabManager)
 {
-	WorkspaceMenuCategory = TabManager->AddLocalWorkspaceMenuCategory(LOCTEXT("WorkspaceMenu_MatineeEditor", "Matinee"));
+	WorkspaceMenuCategory = InTabManager->AddLocalWorkspaceMenuCategory(LOCTEXT("WorkspaceMenu_MatineeEditor", "Matinee"));
 	auto WorkspaceMenuCategoryRef = WorkspaceMenuCategory.ToSharedRef();
 
-	FAssetEditorToolkit::RegisterTabSpawners(TabManager);
+	FAssetEditorToolkit::RegisterTabSpawners(InTabManager);
 
-	TabManager->RegisterTabSpawner(MatineeRecordingViewportName, FOnSpawnTab::CreateRaw(this, &FMatinee::SpawnRecordingViewport))
+	InTabManager->RegisterTabSpawner(MatineeRecordingViewportName, FOnSpawnTab::CreateRaw(this, &FMatinee::SpawnRecordingViewport))
 		.SetDisplayName(NSLOCTEXT("Matinee", "RecordingViewport", "Matinee Recorder"))
 		.SetGroup(WorkspaceMenuCategoryRef)
 		.SetIcon(FSlateIcon(FEditorStyle::GetStyleSetName(), "Matinee.Tabs.RecordingViewport"));
 
-	TabManager->RegisterTabSpawner(MatineeCurveEdName, FOnSpawnTab::CreateSP(this, &FMatinee::SpawnTab, MatineeCurveEdName))
+	InTabManager->RegisterTabSpawner(MatineeCurveEdName, FOnSpawnTab::CreateSP(this, &FMatinee::SpawnTab, MatineeCurveEdName))
 		.SetDisplayName(NSLOCTEXT("Matinee", "CurveEditorTitle", "Curve Editor"))
 		.SetGroup(WorkspaceMenuCategoryRef)
 		.SetIcon(FSlateIcon(FEditorStyle::GetStyleSetName(), "Matinee.Tabs.CurveEditor"));
 
-	TabManager->RegisterTabSpawner(MatineeTrackWindowName, FOnSpawnTab::CreateSP(this, &FMatinee::SpawnTab, MatineeTrackWindowName))
+	InTabManager->RegisterTabSpawner(MatineeTrackWindowName, FOnSpawnTab::CreateSP(this, &FMatinee::SpawnTab, MatineeTrackWindowName))
 		.SetDisplayName(NSLOCTEXT("Matinee", "TrackViewEditorTitle", "Tracks"))
 		.SetGroup(WorkspaceMenuCategoryRef)
 		.SetIcon(FSlateIcon(FEditorStyle::GetStyleSetName(), "Matinee.Tabs.Tracks"));
 
 
-	TabManager->RegisterTabSpawner( MatineePropertyWindowName, FOnSpawnTab::CreateSP(this, &FMatinee::SpawnTab, MatineePropertyWindowName) )
+	InTabManager->RegisterTabSpawner( MatineePropertyWindowName, FOnSpawnTab::CreateSP(this, &FMatinee::SpawnTab, MatineePropertyWindowName) )
 		.SetDisplayName(NSLOCTEXT("Matinee", "PropertiesEditorTitle", "Details"))
 		.SetGroup(WorkspaceMenuCategoryRef)
 		.SetIcon(FSlateIcon(FEditorStyle::GetStyleSetName(), "LevelEditor.Tabs.Details"));
 
 }
 
-void FMatinee::UnregisterTabSpawners(const TSharedRef<class FTabManager>& TabManager)
+void FMatinee::UnregisterTabSpawners(const TSharedRef<class FTabManager>& InTabManager)
 {
-	FAssetEditorToolkit::UnregisterTabSpawners(TabManager);
+	FAssetEditorToolkit::UnregisterTabSpawners(InTabManager);
 
-	TabManager->UnregisterTabSpawner(MatineeCurveEdName);
-	TabManager->UnregisterTabSpawner(MatineeTrackWindowName);
-	TabManager->UnregisterTabSpawner(MatineePropertyWindowName);
+	InTabManager->UnregisterTabSpawner(MatineeCurveEdName);
+	InTabManager->UnregisterTabSpawner(MatineeTrackWindowName);
+	InTabManager->UnregisterTabSpawner(MatineePropertyWindowName);
 }
 
 TSharedRef<SDockTab> FMatinee::SpawnTab( const FSpawnTabArgs& TabSpawnArgs, FName TabIdentifier )
diff --git a/Engine/Source/Editor/Matinee/Private/MatineeDraw.cpp b/Engine/Source/Editor/Matinee/Private/MatineeDraw.cpp
--- a/Engine/Source/Editor/Matinee/Private/MatineeDraw.cpp
+++ b/Engine/Source/Editor/Matinee/Private/MatineeDraw.cpp
@@ -498,10 +498,10 @@ FString FMatinee::MakeTimecodeString( float InTime, bool bIncludeMinutes ) const
 
 
 /** Draw gridlines and time labels. */
-void FMatineeViewportClient::DrawGrid(FViewport* Viewport, FCanvas* Canvas, bool bDrawTimeline)
+void FMatineeViewportClient::DrawGrid(FViewport* InViewport, FCanvas* Canvas, bool bDrawTimeline)
 {
-	const int32 ViewX = Viewport->GetSizeXY().X;
-	const int32 ViewY = Viewport->GetSizeXY().Y;
+	const int32 ViewX = InViewport->GetSizeXY().X;
+	const int32 ViewY = InViewport->GetSizeXY().Y;
 
 	// Calculate desired grid spacing.
 	int32 MinPixelsPerGrid = 35;
@@ -606,10 +606,10 @@ void FMatineeViewportClient::DrawGrid(FViewport* Viewport, FCanvas* Canvas, bool
 }
 
 /** Draw the timeline control at the bottom of the editor. */
-void FMatineeViewportClient::DrawTimeline(FViewport* Viewport, FCanvas* Canvas)
+void FMatineeViewportClient::DrawTimeline(FViewport* InViewport, FCanvas* Canvas)
 {
-	int32 ViewX = Viewport->GetSizeXY().X;
-	int32 ViewY = Viewport->GetSizeXY().Y;
+	int32 ViewX = InViewport->GetSizeXY().X;
+	int32 ViewY = InViewport->GetSizeXY().Y;
 
 	//////// DRAW TIMELINE
 	// Entire length is clickable.
@@ -618,12 +618,12 @@ void FMatineeViewportClient::DrawTimeline(FViewport* Viewport, FCanvas* Canvas)
 	Canvas->DrawTile(InterpEd->LabelWidth, ViewY - TotalBarHeight, ViewX - InterpEd->LabelWidth, TimelineHeight, 0.f, 0.f, 0.f, 0.f, FColor(80,80,80) );
 	if( Canvas->IsHitTesting() ) Canvas->SetHitProxy( NULL );
 
-	DrawGrid(Viewport, Canvas, true);
+	DrawGrid(InViewport, Canvas, true);
 
 	// Draw black line separating nav from timeline.
 	Canvas->DrawTile(0, ViewY - TotalBarHeight, ViewX, 1, 0.f, 0.f, 0.f, 0.f, FLinearColor::Black );
 
-	DrawMarkers(Viewport, Canvas);
+	DrawMarkers(InViewport, Canvas);
 
 	//////// DRAW NAVIGATOR
 	{
@@ -879,10 +879,10 @@ void FMatineeViewportClient::DrawTimeline(FViewport* Viewport, FCanvas* Canvas)
 }
 
 /** Draw various markers on the timeline */
-void FMatineeViewportClient::DrawMarkers(FViewport* Viewport, FCanvas* Canvas)
+void FMatineeViewportClient::DrawMarkers(FViewport* InViewport, FCanvas* Canvas)
 {
-	int32 ViewX = Viewport->GetSizeXY().X;
-	int32 ViewY = Viewport->GetSizeXY().Y;
+	int32 ViewX = InViewport->GetSizeXY().X;
+	int32 ViewY = InViewport->GetSizeXY().Y;
 	int32 ScaleTopY = ViewY - TotalBarHeight + 1;
 
 	// Calculate screen X position that indicates current position in track.
@@ -1713,7 +1713,7 @@ void FMatineeViewportClient::DrawTrackLabel( FCanvas* Canvas, UInterpTrack* Trac
 }
 
 /** Draw the track editor using the supplied 2D RenderInterface. */
-void FMatineeViewportClient::Draw(FViewport* Viewport, FCanvas* Canvas)
+void FMatineeViewportClient::Draw(FViewport* InViewport, FCanvas* Canvas)
 {
 	if (!ParentTab.IsValid())
 	{
@@ -1726,8 +1726,8 @@ void FMatineeViewportClient::Draw(FViewport* Viewport, FCanvas* Canvas)
 	// Erase background
 	Canvas->Clear( FColor(162,162,162) );
 
-	const int32 ViewX = Viewport->GetSizeXY().X;
-	const int32 ViewY = Viewport->GetSizeXY().Y;
+	const int32 ViewX = InViewport->GetSizeXY().X;
+	const int32 ViewY = InViewport->GetSizeXY().Y;
 
 	// @todo frick: Weird to compute this here and storing it in parent
 	InterpEd->TrackViewSizeX = ViewX - InterpEd->LabelWidth;
@@ -1737,7 +1737,7 @@ void FMatineeViewportClient::Draw(FViewport* Viewport, FCanvas* Canvas)
 	InterpEd->PixelsPerSec = FMath::Max( 1.0f, float(ViewX - InterpEd->LabelWidth)/float(InterpEd->ViewEndTime - InterpEd->ViewStartTime) );
 	InterpEd->NavPixelsPerSecond = FMath::Max( 0.0f, float(ViewX - InterpEd->LabelWidth)/InterpEd->IData->InterpLength );
 
-	DrawGrid(Viewport, Canvas, false);
+	DrawGrid(InViewport, Canvas, false);
 
 	FCanvasTileItem TileItem( FVector2D::ZeroVector, FVector2D::ZeroVector, NullRegionColor );
 	TileItem.BlendMode = SE_BLEND_Translucent;
@@ -2163,7 +2163,7 @@ void FMatineeViewportClient::Draw(FViewport* Viewport, FCanvas* Canvas)
 	if( bWantTimeline )
 	{
 		// Draw grid and timeline stuff.
-		DrawTimeline(Viewport, Canvas);
+		DrawTimeline(InViewport, Canvas);
 	}
 
 	// Draw line between title block and track view for empty space
diff --git a/Engine/Source/Editor/Matinee/Private/MatineeViewportClient.cpp b/Engine/Source/Editor/Matinee/Private/MatineeViewportClient.cpp
--- a/Engine/Source/Editor/Matinee/Private/MatineeViewportClient.cpp
+++ b/Engine/Source/Editor/Matinee/Private/MatineeViewportClient.cpp
@@ -158,18 +158,18 @@ void FMatineeViewportClient::AddKeysFromHitProxy( HHitProxy* HitProxy, TArray<FI
 	}
 }
 
-bool FMatineeViewportClient::InputKey(FViewport* Viewport, int32 ControllerId, FKey Key, EInputEvent Event,float /*AmountDepressed*/,bool /*Gamepad*/)
+bool FMatineeViewportClient::InputKey(FViewport* InViewport, int32 ControllerId, FKey Key, EInputEvent Event,float /*AmountDepressed*/,bool /*Gamepad*/)
 {
 	UpdateAndApplyCursorVisibility();
 
-	const bool bCtrlDown = Viewport->KeyState(EKeys::LeftControl) || Viewport->KeyState(EKeys::RightControl);
-	const bool bShiftDown = Viewport->KeyState(EKeys::LeftShift) || Viewport->KeyState(EKeys::RightShift);
-	const bool bAltDown = Viewport->KeyState(EKeys::LeftAlt) || Viewport->KeyState(EKeys::RightAlt);
-	const bool bCmdDown = Viewport->KeyState(EKeys::LeftCommand) || Viewport->KeyState(EKeys::RightCommand);
-	const bool bCapsDown = Viewport->KeyState(EKeys::CapsLock);
+	const bool bCtrlDown = InViewport->KeyState(EKeys::LeftControl) || InViewport->KeyState(EKeys::RightControl);
+	const bool bShiftDown = InViewport->KeyState(EKeys::LeftShift) || InViewport->KeyState(EKeys::RightShift);
+	const bool bAltDown = InViewport->KeyState(EKeys::LeftAlt) || InViewport->KeyState(EKeys::RightAlt);
+	const bool bCmdDown = InViewport->KeyState(EKeys::LeftCommand) || InViewport->KeyState(EKeys::RightCommand);
+	const bool bCapsDown = InViewport->KeyState(EKeys::CapsLock);
 
-	const int32 HitX = Viewport->GetMouseX();
-	const int32 HitY = Viewport->GetMouseY();
+	const int32 HitX = InViewport->GetMouseX();
+	const int32 HitY = InViewport->GetMouseY();
 
 	bool bClickedTrackViewport = false;
 
@@ -181,7 +181,7 @@ bool FMatineeViewportClient::InputKey(FViewport* Viewport, int32 ControllerId, F
 			{
 				if(DragObject == NULL)
 				{
-					HHitProxy*	HitResult = Viewport->GetHitProxy(HitX,HitY);
+					HHitProxy*	HitResult = InViewport->GetHitProxy(HitX,HitY);
 
 					if(HitResult)
 					{
@@ -601,7 +601,7 @@ bool FMatineeViewportClient::InputKey(FViewport* Viewport, int32 ControllerId, F
 						}
 					}
 
-					Viewport->LockMouseToViewport(true);
+					InViewport->LockMouseToViewport(true);
 
 					bMouseDown = true;
 					OldMouseX = HitX;
@@ -612,9 +612,9 @@ bool FMatineeViewportClient::InputKey(FViewport* Viewport, int32 ControllerId, F
 			break;
 		case IE_DoubleClick:
 			{
-				Viewport->InvalidateHitProxy();
+				InViewport->InvalidateHitProxy();
 
-				HHitProxy*	HitResult = Viewport->GetHitProxy(HitX,HitY);
+				HHitProxy*	HitResult = InViewport->GetHitProxy(HitX,HitY);
 
 				if(HitResult)
 				{
@@ -632,7 +632,7 @@ bool FMatineeViewportClient::InputKey(FViewport* Viewport, int32 ControllerId, F
 			break;
 		case IE_Released:
 			{
-				Viewport->InvalidateHitProxy();
+				InViewport->InvalidateHitProxy();
 
 				if(bBoxSelecting)
 				{
@@ -648,7 +648,7 @@ bool FMatineeViewportClient::InputKey(FViewport* Viewport, int32 ControllerId, F
 
 					// We read back the hit proxy map for the required region.
 					TArray<HHitProxy*> ProxyMap;
-					Viewport->GetHitProxyMap(FIntRect(MinX, MinY, MaxX + 1, MaxY + 1), ProxyMap);
+					InViewport->GetHitProxyMap(FIntRect(MinX, MinY, MaxX + 1, MaxY + 1), ProxyMap);
 
 					TArray<FInterpEdSelKey>	NewSelection;
 
@@ -679,7 +679,7 @@ bool FMatineeViewportClient::InputKey(FViewport* Viewport, int32 ControllerId, F
 				}
 				else if(DragObject)
 				{
-					HHitProxy*	HitResult = Viewport->GetHitProxy(HitX,HitY);
+					HHitProxy*	HitResult = InViewport->GetHitProxy(HitX,HitY);
 
 					if(HitResult)
 					{
@@ -698,7 +698,7 @@ bool FMatineeViewportClient::InputKey(FViewport* Viewport, int32 ControllerId, F
 				}
 				else if(DistanceDragged < 4)
 				{
-					HHitProxy*	HitResult = Viewport->GetHitProxy(HitX,HitY);
+					HHitProxy*	HitResult = InViewport->GetHitProxy(HitX,HitY);
 
 					// If mouse didn't really move since last time, and we released over empty space, deselect everything.
 					if(!HitResult)
@@ -767,7 +767,7 @@ bool FMatineeViewportClient::InputKey(FViewport* Viewport, int32 ControllerId, F
 					bGrabbingMarker = false;
 				}
 
-				Viewport->LockMouseToViewport(false);
+				InViewport->LockMouseToViewport(false);
 
 				DistanceDragged = 0;
 
@@ -786,16 +786,16 @@ bool FMatineeViewportClient::InputKey(FViewport* Viewport, int32 ControllerId, F
 		{
 		case IE_Pressed:
 			{
-				HHitProxy*	HitResult = Viewport->GetHitProxy(HitX,HitY);
+				HHitProxy*	HitResult = InViewport->GetHitProxy(HitX,HitY);
 
 				if(HitResult)
 				{
 					// User right-click somewhere in the track editor
-					TSharedPtr<SWidget> Menu = InterpEd->CreateContextMenu( Viewport, HitResult, bIsDirectorTrackWindow );
+					TSharedPtr<SWidget> Menu = InterpEd->CreateContextMenu( InViewport, HitResult, bIsDirectorTrackWindow );
 					if (Menu.IsValid())
 					{
 						// Redraw the viewport so the user can see which object was right clicked on
-						Viewport->Draw();
+						InViewport->Draw();
 						FlushRenderingCommands();
 
 						TSharedPtr< SWindow > Parent = FSlateApplication::Get().GetActiveTopLevelWindow(); 
@@ -815,7 +815,7 @@ bool FMatineeViewportClient::InputKey(FViewport* Viewport, int32 ControllerId, F
 
 		case IE_Released:
 			{
-				Viewport->InvalidateHitProxy();
+				InViewport->InvalidateHitProxy();
 			}
 			break;
 		}
@@ -838,15 +838,15 @@ bool FMatineeViewportClient::InputKey(FViewport* Viewport, int32 ControllerId, F
 	}
 
 	// Handle viewport screenshot.
-	InputTakeScreenshot( Viewport, Key, Event );
+	InputTakeScreenshot( InViewport, Key, Event );
 
 	return true;
 }
 
 // X and Y here are the new screen position of the cursor.
-void FMatineeViewportClient::MouseMove(FViewport* Viewport, int32 X, int32 Y)
+void FMatineeViewportClient::MouseMove(FViewport* InViewport, int32 X, int32 Y)
 {
-	bool bCtrlDown = Viewport->KeyState(EKeys::LeftControl) || Viewport->KeyState(EKeys::RightControl);
+	bool bCtrlDown = InViewport->KeyState(EKeys::LeftControl) || InViewport->KeyState(EKeys::RightControl);
 
 	int32 DeltaX = OldMouseX - X;
 	int32 DeltaY = OldMouseY - Y;
@@ -960,25 +960,25 @@ void FMatineeViewportClient::MouseMove(FViewport* Viewport, int32 X, int32 Y)
 	}
 }
 
-bool FMatineeViewportClient::InputAxis(FViewport* Viewport, int32 ControllerId, FKey Key, float Delta, float DeltaTime, int32 NumSamples, bool bGamepad)
+bool FMatineeViewportClient::InputAxis(FViewport* InViewport, int32 ControllerId, FKey Key, float Delta, float DeltaTime, int32 NumSamples, bool bGamepad)
 {
 	if ( Key == EKeys::MouseX || Key == EKeys::MouseY )
 	{
-		int32 X = Viewport->GetMouseX();
-		int32 Y = Viewport->GetMouseY();
-		MouseMove(Viewport, X, Y);
+		int32 X = InViewport->GetMouseX();
+		int32 Y = InViewport->GetMouseY();
+		MouseMove(InViewport, X, Y);
 		return true;
 	}
 	return false;
 }
 
-EMouseCursor::Type FMatineeViewportClient::GetCursor(FViewport* Viewport,int32 X,int32 Y)
+EMouseCursor::Type FMatineeViewportClient::GetCursor(FViewport* InViewport,int32 X,int32 Y)
 {
 	EMouseCursor::Type Result = EMouseCursor::Crosshairs;
 
-	if(DragObject==NULL)
+	if(DragObject==nullptr)
 	{
-		HHitProxy*	HitProxy = Viewport->GetHitProxy(X,Y);
+		HHitProxy*	HitProxy = InViewport->GetHitProxy(X,Y);
 		
 		if(HitProxy)
 		{
diff --git a/Engine/Source/Editor/NiagaraEditor/Private/NiagaraEditor.cpp b/Engine/Source/Editor/NiagaraEditor/Private/NiagaraEditor.cpp
--- a/Engine/Source/Editor/NiagaraEditor/Private/NiagaraEditor.cpp
+++ b/Engine/Source/Editor/NiagaraEditor/Private/NiagaraEditor.cpp
@@ -20,34 +20,34 @@ const FName FNiagaraEditor::NodeGraphTabId(TEXT("NiagaraEditor_NodeGraph"));
 const FName FNiagaraEditor::FlattenedNodeGraphTabId(TEXT("NiagaraEditor_FlattenedNodeGraph")); 
 const FName FNiagaraEditor::PropertiesTabId(TEXT("NiagaraEditor_MaterialProperties"));
 
-void FNiagaraEditor::RegisterTabSpawners(const TSharedRef<class FTabManager>& TabManager)
+void FNiagaraEditor::RegisterTabSpawners(const TSharedRef<class FTabManager>& InTabManager)
 {
-	WorkspaceMenuCategory = TabManager->AddLocalWorkspaceMenuCategory(LOCTEXT("WorkspaceMenu_NiagaraEditor", "Niagara"));
+	WorkspaceMenuCategory = InTabManager->AddLocalWorkspaceMenuCategory(LOCTEXT("WorkspaceMenu_NiagaraEditor", "Niagara"));
 
-	FAssetEditorToolkit::RegisterTabSpawners(TabManager);
+	FAssetEditorToolkit::RegisterTabSpawners(InTabManager);
 	
 	TSharedRef<FWorkspaceItem> WorkspaceMenuCategoryRef = WorkspaceMenuCategory.ToSharedRef();
 
-	TabManager->RegisterTabSpawner( NodeGraphTabId, FOnSpawnTab::CreateSP(this, &FNiagaraEditor::SpawnTab_NodeGraph) )
+	InTabManager->RegisterTabSpawner( NodeGraphTabId, FOnSpawnTab::CreateSP(this, &FNiagaraEditor::SpawnTab_NodeGraph) )
 		.SetDisplayName( LOCTEXT("NodeGraph", "Node Graph") )
 		.SetGroup(WorkspaceMenuCategoryRef); 
 
-	TabManager->RegisterTabSpawner( FlattenedNodeGraphTabId, FOnSpawnTab::CreateSP(this, &FNiagaraEditor::SpawnTab_FlattenedNodeGraph))
+	InTabManager->RegisterTabSpawner( FlattenedNodeGraphTabId, FOnSpawnTab::CreateSP(this, &FNiagaraEditor::SpawnTab_FlattenedNodeGraph))
 		.SetDisplayName(LOCTEXT("FlattenedNodeGraph", "Flattened Node Graph"))
 		.SetGroup(WorkspaceMenuCategoryRef); 
 
-	TabManager->RegisterTabSpawner(PropertiesTabId, FOnSpawnTab::CreateSP(this, &FNiagaraEditor::SpawnTab_NodeProperties))
+	InTabManager->RegisterTabSpawner(PropertiesTabId, FOnSpawnTab::CreateSP(this, &FNiagaraEditor::SpawnTab_NodeProperties))
 		.SetDisplayName(LOCTEXT("DetailsTab", "Details"))
 		.SetGroup(WorkspaceMenuCategoryRef)
 		.SetIcon(FSlateIcon(FEditorStyle::GetStyleSetName(), "LevelEditor.Tabs.Details"));
 }
 
-void FNiagaraEditor::UnregisterTabSpawners(const TSharedRef<class FTabManager>& TabManager)
+void FNiagaraEditor::UnregisterTabSpawners(const TSharedRef<class FTabManager>& InTabManager)
 {
-	FAssetEditorToolkit::UnregisterTabSpawners(TabManager);
+	FAssetEditorToolkit::UnregisterTabSpawners(InTabManager);
 
-	TabManager->UnregisterTabSpawner(NodeGraphTabId);
-	TabManager->UnregisterTabSpawner(PropertiesTabId);
+	InTabManager->UnregisterTabSpawner(NodeGraphTabId);
+	InTabManager->UnregisterTabSpawner(PropertiesTabId);
 }
 
 FNiagaraEditor::~FNiagaraEditor()
diff --git a/Engine/Source/Editor/NiagaraEditor/Private/NiagaraEffectEditor.cpp b/Engine/Source/Editor/NiagaraEditor/Private/NiagaraEffectEditor.cpp
--- a/Engine/Source/Editor/NiagaraEditor/Private/NiagaraEffectEditor.cpp
+++ b/Engine/Source/Editor/NiagaraEditor/Private/NiagaraEffectEditor.cpp
@@ -18,51 +18,51 @@ const FName FNiagaraEffectEditor::DevEmitterEditorTabID(TEXT("NiagaraEffectEdito
 const FName FNiagaraEffectEditor::CurveEditorTabID(TEXT("NiagaraEffectEditor_CurveEditor"));
 const FName FNiagaraEffectEditor::SequencerTabID(TEXT("NiagaraEffectEditor_Sequencer"));
 
-void FNiagaraEffectEditor::RegisterTabSpawners(const TSharedRef<class FTabManager>& TabManager)
+void FNiagaraEffectEditor::RegisterTabSpawners(const TSharedRef<class FTabManager>& InTabManager)
 {
-	WorkspaceMenuCategory = TabManager->AddLocalWorkspaceMenuCategory(LOCTEXT("WorkspaceMenu_NiagaraEffectEditor", "Niagara Effect"));
+	WorkspaceMenuCategory = InTabManager->AddLocalWorkspaceMenuCategory(LOCTEXT("WorkspaceMenu_NiagaraEffectEditor", "Niagara Effect"));
 
-	FAssetEditorToolkit::RegisterTabSpawners(TabManager);
+	FAssetEditorToolkit::RegisterTabSpawners(InTabManager);
 
-	TabManager->RegisterTabSpawner(ViewportTabID, FOnSpawnTab::CreateSP(this, &FNiagaraEffectEditor::SpawnTab_Viewport))
+	InTabManager->RegisterTabSpawner(ViewportTabID, FOnSpawnTab::CreateSP(this, &FNiagaraEffectEditor::SpawnTab_Viewport))
 		.SetDisplayName(LOCTEXT("Preview", "Preview"))
 		.SetGroup(WorkspaceMenuCategory.ToSharedRef())
 		.SetIcon(FSlateIcon(FEditorStyle::GetStyleSetName(), "LevelEditor.Tabs.Viewports"));
 
-	TabManager->RegisterTabSpawner(EmitterEditorTabID, FOnSpawnTab::CreateSP(this, &FNiagaraEffectEditor::SpawnTab_EmitterList))
+	InTabManager->RegisterTabSpawner(EmitterEditorTabID, FOnSpawnTab::CreateSP(this, &FNiagaraEffectEditor::SpawnTab_EmitterList))
 		.SetDisplayName(LOCTEXT("Emitters", "Emitters"))
 		.SetGroup(WorkspaceMenuCategory.ToSharedRef());
 
-	TabManager->RegisterTabSpawner(DevEmitterEditorTabID, FOnSpawnTab::CreateSP(this, &FNiagaraEffectEditor::SpawnTab_DevEmitterList))
+	InTabManager->RegisterTabSpawner(DevEmitterEditorTabID, FOnSpawnTab::CreateSP(this, &FNiagaraEffectEditor::SpawnTab_DevEmitterList))
 		.SetDisplayName(LOCTEXT("DevEmitters", "Emitters_Dev"))
 		.SetGroup(WorkspaceMenuCategory.ToSharedRef());
 
-	TabManager->RegisterTabSpawner(CurveEditorTabID, FOnSpawnTab::CreateSP(this, &FNiagaraEffectEditor::SpawnTab_CurveEd))
+	InTabManager->RegisterTabSpawner(CurveEditorTabID, FOnSpawnTab::CreateSP(this, &FNiagaraEffectEditor::SpawnTab_CurveEd))
 		.SetDisplayName(LOCTEXT("Curves", "Curves"))
 		.SetGroup(WorkspaceMenuCategory.ToSharedRef());
 
-	TabManager->RegisterTabSpawner(SequencerTabID, FOnSpawnTab::CreateSP(this, &FNiagaraEffectEditor::SpawnTab_Sequencer))
+	InTabManager->RegisterTabSpawner(SequencerTabID, FOnSpawnTab::CreateSP(this, &FNiagaraEffectEditor::SpawnTab_Sequencer))
 		.SetDisplayName(LOCTEXT("Timeline", "Timeline"))
 		.SetGroup(WorkspaceMenuCategory.ToSharedRef());
 
 
-	TabManager->RegisterTabSpawner(UpdateTabId, FOnSpawnTab::CreateSP(this, &FNiagaraEffectEditor::SpawnTab))
+	InTabManager->RegisterTabSpawner(UpdateTabId, FOnSpawnTab::CreateSP(this, &FNiagaraEffectEditor::SpawnTab))
 		.SetDisplayName(LOCTEXT("NiagaraEffect", "Niagara Effect"))
 		.SetGroup( WorkspaceMenuCategory.ToSharedRef() );
 
 
 }
 
-void FNiagaraEffectEditor::UnregisterTabSpawners(const TSharedRef<class FTabManager>& TabManager)
+void FNiagaraEffectEditor::UnregisterTabSpawners(const TSharedRef<class FTabManager>& InTabManager)
 {
-	FAssetEditorToolkit::UnregisterTabSpawners(TabManager);
-
-	TabManager->UnregisterTabSpawner(UpdateTabId);
-	TabManager->UnregisterTabSpawner(ViewportTabID);
-	TabManager->UnregisterTabSpawner(EmitterEditorTabID);
-	TabManager->UnregisterTabSpawner(DevEmitterEditorTabID);
-	TabManager->UnregisterTabSpawner(CurveEditorTabID);
-	TabManager->UnregisterTabSpawner(SequencerTabID);
+	FAssetEditorToolkit::UnregisterTabSpawners(InTabManager);
+
+	InTabManager->UnregisterTabSpawner(UpdateTabId);
+	InTabManager->UnregisterTabSpawner(ViewportTabID);
+	InTabManager->UnregisterTabSpawner(EmitterEditorTabID);
+	InTabManager->UnregisterTabSpawner(DevEmitterEditorTabID);
+	InTabManager->UnregisterTabSpawner(CurveEditorTabID);
+	InTabManager->UnregisterTabSpawner(SequencerTabID);
 }
 
 
diff --git a/Engine/Source/Editor/NiagaraEditor/Private/SNiagaraEffectEditorViewport.cpp b/Engine/Source/Editor/NiagaraEditor/Private/SNiagaraEffectEditorViewport.cpp
--- a/Engine/Source/Editor/NiagaraEditor/Private/SNiagaraEffectEditorViewport.cpp
+++ b/Engine/Source/Editor/NiagaraEditor/Private/SNiagaraEffectEditorViewport.cpp
@@ -71,9 +71,9 @@ void FNiagaraEffectEditorViewportClient::Tick(float DeltaSeconds)
 }
 
 
-void FNiagaraEffectEditorViewportClient::Draw(FViewport* Viewport,FCanvas* Canvas)
+void FNiagaraEffectEditorViewportClient::Draw(FViewport* InViewport,FCanvas* Canvas)
 {
-	FEditorViewportClient::Draw(Viewport, Canvas);
+	FEditorViewportClient::Draw(InViewport, Canvas);
 }
 
 bool FNiagaraEffectEditorViewportClient::ShouldOrbitCamera() const
diff --git a/Engine/Source/Editor/PhAT/Private/PhAT.cpp b/Engine/Source/Editor/PhAT/Private/PhAT.cpp
--- a/Engine/Source/Editor/PhAT/Private/PhAT.cpp
+++ b/Engine/Source/Editor/PhAT/Private/PhAT.cpp
@@ -83,36 +83,36 @@ static const FName PhATHierarchyName("PhAT_Hierarchy");
 
 
 
-void FPhAT::RegisterTabSpawners(const TSharedRef<class FTabManager>& TabManager)
+void FPhAT::RegisterTabSpawners(const TSharedRef<class FTabManager>& InTabManager)
 {
-	WorkspaceMenuCategory = TabManager->AddLocalWorkspaceMenuCategory( LOCTEXT( "WorkspaceMenu_PhAT", "Physics Asset Editor" ) );
+	WorkspaceMenuCategory = InTabManager->AddLocalWorkspaceMenuCategory( LOCTEXT( "WorkspaceMenu_PhAT", "Physics Asset Editor" ) );
 	auto WorkspaceMenuCategoryRef = WorkspaceMenuCategory.ToSharedRef();
 
-	FAssetEditorToolkit::RegisterTabSpawners( TabManager );
+	FAssetEditorToolkit::RegisterTabSpawners( InTabManager );
 
-	TabManager->RegisterTabSpawner( PhATPreviewViewportName, FOnSpawnTab::CreateSP( this, &FPhAT::SpawnTab, PhATPreviewViewportName ) )
+	InTabManager->RegisterTabSpawner( PhATPreviewViewportName, FOnSpawnTab::CreateSP( this, &FPhAT::SpawnTab, PhATPreviewViewportName ) )
 		.SetDisplayName( LOCTEXT( "ViewportTab", "Viewport" ) )
 		.SetGroup( WorkspaceMenuCategoryRef )
 		.SetIcon( FSlateIcon( FEditorStyle::GetStyleSetName(), "LevelEditor.Tabs.Viewports" ) );
 
-	TabManager->RegisterTabSpawner( PhATPropertiesName, FOnSpawnTab::CreateSP( this, &FPhAT::SpawnTab, PhATPropertiesName ) )
+	InTabManager->RegisterTabSpawner( PhATPropertiesName, FOnSpawnTab::CreateSP( this, &FPhAT::SpawnTab, PhATPropertiesName ) )
 		.SetDisplayName( LOCTEXT( "PropertiesTab", "Details" ) )
 		.SetGroup( WorkspaceMenuCategoryRef )
 		.SetIcon( FSlateIcon( FEditorStyle::GetStyleSetName(), "LevelEditor.Tabs.Details" ) );
 
-	TabManager->RegisterTabSpawner( PhATHierarchyName, FOnSpawnTab::CreateSP( this, &FPhAT::SpawnTab, PhATHierarchyName ) )
+	InTabManager->RegisterTabSpawner( PhATHierarchyName, FOnSpawnTab::CreateSP( this, &FPhAT::SpawnTab, PhATHierarchyName ) )
 		.SetDisplayName( LOCTEXT( "HierarchyTab", "Hierarchy" ) )
 		.SetGroup( WorkspaceMenuCategoryRef )
 		.SetIcon( FSlateIcon( FEditorStyle::GetStyleSetName(), "Kismet.Tabs.Palette" ) );
 }
 
-void FPhAT::UnregisterTabSpawners(const TSharedRef<class FTabManager>& TabManager)
+void FPhAT::UnregisterTabSpawners(const TSharedRef<class FTabManager>& InTabManager)
 {
-	FAssetEditorToolkit::UnregisterTabSpawners(TabManager);
+	FAssetEditorToolkit::UnregisterTabSpawners(InTabManager);
 
-	TabManager->UnregisterTabSpawner(PhATPreviewViewportName);
-	TabManager->UnregisterTabSpawner(PhATPropertiesName);
-	TabManager->UnregisterTabSpawner(PhATHierarchyName);
+	InTabManager->UnregisterTabSpawner(PhATPreviewViewportName);
+	InTabManager->UnregisterTabSpawner(PhATPropertiesName);
+	InTabManager->UnregisterTabSpawner(PhATHierarchyName);
 }
 
 TSharedRef<SDockTab> FPhAT::SpawnTab( const FSpawnTabArgs& TabSpawnArgs, FName TabIdentifier )
diff --git a/Engine/Source/Editor/PhAT/Private/PhATPreviewViewportClient.cpp b/Engine/Source/Editor/PhAT/Private/PhATPreviewViewportClient.cpp
--- a/Engine/Source/Editor/PhAT/Private/PhATPreviewViewportClient.cpp
+++ b/Engine/Source/Editor/PhAT/Private/PhATPreviewViewportClient.cpp
@@ -276,12 +276,12 @@ void FPhATEdPreviewViewportClient::Draw(const FSceneView* View, FPrimitiveDrawIn
 	}
 }
 
-bool FPhATEdPreviewViewportClient::InputKey(FViewport* Viewport, int32 ControllerId, FKey Key, EInputEvent Event, float AmountDepressed, bool Gamepad)
+bool FPhATEdPreviewViewportClient::InputKey(FViewport* InViewport, int32 ControllerId, FKey Key, EInputEvent Event, float AmountDepressed, bool Gamepad)
 {
-	int32 HitX = Viewport->GetMouseX();
-	int32 HitY = Viewport->GetMouseY();
-	bool bCtrlDown = Viewport->KeyState(EKeys::LeftControl) || Viewport->KeyState(EKeys::RightControl);
-	bool bShiftDown = Viewport->KeyState(EKeys::LeftShift) || Viewport->KeyState(EKeys::RightShift);
+	int32 HitX = InViewport->GetMouseX();
+	int32 HitY = InViewport->GetMouseY();
+	bool bCtrlDown = InViewport->KeyState(EKeys::LeftControl) || InViewport->KeyState(EKeys::RightControl);
+	bool bShiftDown = InViewport->KeyState(EKeys::LeftShift) || InViewport->KeyState(EKeys::RightShift);
 
 	bool bHandled = false;
 	if (SharedData->bRunningSimulation)
@@ -293,13 +293,13 @@ bool FPhATEdPreviewViewportClient::InputKey(FViewport* Viewport, int32 Controlle
 				if (bShiftDown)
 				{
 					bHandled = true;
-					SimMousePress(Viewport, false, Key);
+					SimMousePress(InViewport, false, Key);
 					bAllowedToMoveCamera = false;
 				}
 				else if (bCtrlDown)
 				{
 					bHandled = true;
-					SimMousePress(Viewport, true, Key);
+					SimMousePress(InViewport, true, Key);
 					bAllowedToMoveCamera = false;
 				}
 			}
@@ -331,7 +331,7 @@ bool FPhATEdPreviewViewportClient::InputKey(FViewport* Viewport, int32 Controlle
 
 	if( !bHandled )
 	{
-		bHandled = FEditorViewportClient::InputKey( Viewport, ControllerId, Key, Event, AmountDepressed, Gamepad );
+		bHandled = FEditorViewportClient::InputKey( InViewport, ControllerId, Key, Event, AmountDepressed, Gamepad );
 	}
 
 	if( bHandled )
@@ -342,13 +342,13 @@ bool FPhATEdPreviewViewportClient::InputKey(FViewport* Viewport, int32 Controlle
 	return bHandled;
 }
 
-bool FPhATEdPreviewViewportClient::InputAxis(FViewport* Viewport, int32 ControllerId, FKey Key, float Delta, float DeltaTime, int32 NumSamples, bool bGamepad)
+bool FPhATEdPreviewViewportClient::InputAxis(FViewport* InViewport, int32 ControllerId, FKey Key, float Delta, float DeltaTime, int32 NumSamples, bool bGamepad)
 {
 	bool bHandled = false;
 	// If we are 'manipulating' don't move the camera but do something else with mouse input.
 	if (SharedData->bManipulating)
 	{
-		bool bCtrlDown = Viewport->KeyState(EKeys::LeftControl) || Viewport->KeyState(EKeys::RightControl);
+		bool bCtrlDown = InViewport->KeyState(EKeys::LeftControl) || InViewport->KeyState(EKeys::RightControl);
 
 		if (SharedData->bRunningSimulation)
 		{
@@ -366,10 +366,10 @@ bool FPhATEdPreviewViewportClient::InputAxis(FViewport* Viewport, int32 Controll
 
 	if( !bHandled )
 	{
-		bHandled = FEditorViewportClient::InputAxis(Viewport,ControllerId,Key,Delta,DeltaTime,NumSamples,bGamepad);
+		bHandled = FEditorViewportClient::InputAxis(InViewport,ControllerId,Key,Delta,DeltaTime,NumSamples,bGamepad);
 	}
 
-	Viewport->Invalidate();
+	InViewport->Invalidate();
 
 	return bHandled;
 }
@@ -460,7 +460,7 @@ void FPhATEdPreviewViewportClient::ProcessClick(class FSceneView& View, class HH
 	}
 }
 
-bool FPhATEdPreviewViewportClient::InputWidgetDelta( FViewport* Viewport, EAxisList::Type CurrentAxis, FVector& Drag, FRotator& Rot, FVector& Scale )
+bool FPhATEdPreviewViewportClient::InputWidgetDelta( FViewport* InViewport, EAxisList::Type CurrentAxis, FVector& Drag, FRotator& Rot, FVector& Scale )
 {
 	bool bHandled = false;
 	TArray<FPhATSharedData::FSelection> & SelectedObjects = SharedData->EditingMode == FPhATSharedData::PEM_BodyEdit ? SharedData->SelectedBodies : SharedData->SelectedConstraints;
@@ -813,15 +813,15 @@ void FPhATEdPreviewViewportClient::EndManipulating()
 	}
 }
 
-void FPhATEdPreviewViewportClient::SimMousePress(FViewport* Viewport, bool bConstrainRotation, FKey Key)
+void FPhATEdPreviewViewportClient::SimMousePress(FViewport* InViewport, bool bConstrainRotation, FKey Key)
 {
-	bool bCtrlDown = Viewport->KeyState(EKeys::LeftControl) || Viewport->KeyState(EKeys::RightControl);
-	bool bShiftDown = Viewport->KeyState(EKeys::LeftShift) || Viewport->KeyState(EKeys::RightShift);
+	bool bCtrlDown = InViewport->KeyState(EKeys::LeftControl) || InViewport->KeyState(EKeys::RightControl);
+	bool bShiftDown = InViewport->KeyState(EKeys::LeftShift) || InViewport->KeyState(EKeys::RightShift);
 
-	FSceneViewFamilyContext ViewFamily(FSceneViewFamily::ConstructionValues( Viewport, GetScene(), EngineShowFlags ));
+	FSceneViewFamilyContext ViewFamily(FSceneViewFamily::ConstructionValues( InViewport, GetScene(), EngineShowFlags ));
 	FSceneView* View = CalcSceneView(&ViewFamily);
 
-	const FViewportClick Click(View, this, EKeys::Invalid, IE_Released, Viewport->GetMouseX(), Viewport->GetMouseY());
+	const FViewportClick Click(View, this, EKeys::Invalid, IE_Released, InViewport->GetMouseX(), InViewport->GetMouseY());
 #if DEBUG_CLICK_VIEWPORT	
 	SharedData->LastClickOrigin = Click.GetOrigin();
 	SharedData->LastClickDirection = Click.GetDirection();
diff --git a/Engine/Source/Editor/PropertyEditor/Private/DetailAdvancedDropdownNode.cpp b/Engine/Source/Editor/PropertyEditor/Private/DetailAdvancedDropdownNode.cpp
--- a/Engine/Source/Editor/PropertyEditor/Private/DetailAdvancedDropdownNode.cpp
+++ b/Engine/Source/Editor/PropertyEditor/Private/DetailAdvancedDropdownNode.cpp
@@ -30,10 +30,10 @@ public:
 
 		bDisplayShowAdvancedMessage = bInDisplayShowAdvancedMessage;
 
-		TSharedPtr<SWidget> Content;
+		TSharedPtr<SWidget> ContentWidget;
 		if( bIsTopNode )
 		{
-			Content =
+			ContentWidget =
 				SNew( SBorder )
 				.BorderImage( FEditorStyle::GetBrush("DetailsView.CategoryMiddle") )
 				.Padding( FMargin( 0.0f, 3.0f, SDetailTableRowBase::ScrollbarPaddingSize, 0.0f ) )
@@ -44,7 +44,7 @@ public:
 		}
 		else if( InArgs._ShouldShowAdvancedButton )
 		{
-			Content = 
+			ContentWidget = 
 				SNew( SBorder )
 				.BorderImage( FEditorStyle::GetBrush("DetailsView.AdvancedDropdownBorder") )
 				.Padding( FMargin( 0.0f, 3.0f, SDetailTableRowBase::ScrollbarPaddingSize, 0.0f ) )
@@ -123,7 +123,7 @@ public:
 
 			SplitterArea = SNew(SSpacer);
 
-			Content =
+			ContentWidget =
 				SNew( SBorder )
 				.BorderImage( FEditorStyle::GetBrush("DetailsView.CategoryBottom") )
 				.Padding( FMargin( 0.0f, 0.0f, SDetailTableRowBase::ScrollbarPaddingSize, 2.0f ) )
@@ -134,7 +134,7 @@ public:
 		
 		ChildSlot
 		[
-			Content.ToSharedRef()
+			ContentWidget.ToSharedRef()
 		];
 
 		STableRow< TSharedPtr< IDetailTreeNode > >::ConstructInternal(
diff --git a/Engine/Source/Editor/PropertyEditor/Private/DetailPropertyRow.cpp b/Engine/Source/Editor/PropertyEditor/Private/DetailPropertyRow.cpp
--- a/Engine/Source/Editor/PropertyEditor/Private/DetailPropertyRow.cpp
+++ b/Engine/Source/Editor/PropertyEditor/Private/DetailPropertyRow.cpp
@@ -34,8 +34,8 @@ FDetailPropertyRow::FDetailPropertyRow(TSharedPtr<FPropertyNode> InPropertyNode,
 		// Check if the property is valid for type customization.  Note: Static arrays of types will be a UProperty with array elements as children
 		if (!PropertyEditorHelpers::IsStaticArray(*PropertyNodeRef))
 		{
-			static FName PropertyEditor("PropertyEditor");
-			FPropertyEditorModule& PropertyEditorModule = FModuleManager::GetModuleChecked<FPropertyEditorModule>(PropertyEditor);
+			static FName NAME_PropertyEditor("PropertyEditor");
+			FPropertyEditorModule& PropertyEditorModule = FModuleManager::GetModuleChecked<FPropertyEditorModule>(NAME_PropertyEditor);
 			
 			IDetailsViewPrivate& DetailsView = InParentCategory->GetDetailsView();
 			
