Copying //UE4/Dev-Blueprints to //UE4/Dev-Main (Source: //UE4/Dev-Blueprints @ 3025888)

diff --git a/Engine/Source/Editor/Kismet/Private/Profiler/BPProfilerStatisticWidgets.cpp b/Engine/Source/Editor/Kismet/Private/Profiler/BPProfilerStatisticWidgets.cpp
--- a/Engine/Source/Editor/Kismet/Private/Profiler/BPProfilerStatisticWidgets.cpp
+++ b/Engine/Source/Editor/Kismet/Private/Profiler/BPProfilerStatisticWidgets.cpp
@@ -3,7 +3,6 @@
 #include "BlueprintEditorPrivatePCH.h"
 #include "SBlueprintProfilerView.h"
 #include "BPProfilerStatisticWidgets.h"
-#include "SHyperlink.h"
 #include "EventExecution.h"
 #include "Developer/BlueprintProfiler/Public/BlueprintProfilerModule.h"
 
@@ -15,16 +14,16 @@
 namespace BlueprintProfilerStatText
 {
 	const FName ColumnId_Name("Name");
+	const FName ColumnId_AverageTime("AverageTime");
 	const FName ColumnId_InclusiveTime("InclusiveTime");
-	const FName ColumnId_Time("Time");
 	const FName ColumnId_MaxTime("MaxTime");
 	const FName ColumnId_MinTime("MinTime");
 	const FName ColumnId_Samples("Samples");
 	const FName ColumnId_TotalTime("TotalTime");
 
 	const FText ColumnText_Name(LOCTEXT("Name", "Name") );
+	const FText ColumnText_AverageTime(LOCTEXT("AverageTime", "Avg Time (ms)"));
 	const FText ColumnText_InclusiveTime(LOCTEXT("InclusiveTime", "Inclusive Time (ms)"));
-	const FText ColumnText_Time(LOCTEXT("Time", "Time (ms)"));
 	const FText ColumnText_MaxTime(LOCTEXT("MaxTime", "Max Time (ms)"));
 	const FText ColumnText_MinTime(LOCTEXT("MinTime", "Min Time (ms)"));
 	const FText ColumnText_Samples(LOCTEXT("Samples", "Samples"));
@@ -74,8 +73,8 @@ const FName SProfilerStatRow::GetStatName(const EBlueprintProfilerStat::Type Sta
 	{
 		case EBlueprintProfilerStat::Name:			return BlueprintProfilerStatText::ColumnId_Name;
 		case EBlueprintProfilerStat::TotalTime:		return BlueprintProfilerStatText::ColumnId_TotalTime;
+		case EBlueprintProfilerStat::AverageTime:	return BlueprintProfilerStatText::ColumnId_AverageTime;
 		case EBlueprintProfilerStat::InclusiveTime:	return BlueprintProfilerStatText::ColumnId_InclusiveTime;
-		case EBlueprintProfilerStat::Time:			return BlueprintProfilerStatText::ColumnId_Time;
 		case EBlueprintProfilerStat::MaxTime:		return BlueprintProfilerStatText::ColumnId_MaxTime;
 		case EBlueprintProfilerStat::MinTime:		return BlueprintProfilerStatText::ColumnId_MinTime;
 		case EBlueprintProfilerStat::Samples:		return BlueprintProfilerStatText::ColumnId_Samples;
@@ -89,8 +88,8 @@ const FText SProfilerStatRow::GetStatText(const EBlueprintProfilerStat::Type Sta
 	{
 		case EBlueprintProfilerStat::Name:			return BlueprintProfilerStatText::ColumnText_Name;
 		case EBlueprintProfilerStat::TotalTime:		return BlueprintProfilerStatText::ColumnText_TotalTime;
+		case EBlueprintProfilerStat::AverageTime:	return BlueprintProfilerStatText::ColumnText_AverageTime;
 		case EBlueprintProfilerStat::InclusiveTime:	return BlueprintProfilerStatText::ColumnText_InclusiveTime;
-		case EBlueprintProfilerStat::Time:			return BlueprintProfilerStatText::ColumnText_Time;
 		case EBlueprintProfilerStat::MaxTime:		return BlueprintProfilerStatText::ColumnText_MaxTime;
 		case EBlueprintProfilerStat::MinTime:		return BlueprintProfilerStatText::ColumnText_MinTime;
 		case EBlueprintProfilerStat::Samples:		return BlueprintProfilerStatText::ColumnText_Samples;
@@ -98,155 +97,6 @@ const FText SProfilerStatRow::GetStatText(const EBlueprintProfilerStat::Type Sta
 	}
 }
 
-//////////////////////////////////////////////////////////////////////////
-// FBPProfilerStatOptions
-
-void FBPProfilerStatDiplayOptions::SetActiveInstance(const FName InstanceName)
-{
-	if (ActiveInstance != InstanceName && HasFlags(ScopeToDebugInstance))
-	{
-		Flags |= Modified;
-	}
-	ActiveInstance = InstanceName;
-}
-
-void FBPProfilerStatDiplayOptions::SetActiveGraph(const FName GraphName)
-{
-	if (ActiveGraph != GraphName && HasFlags(GraphFilter))
-	{
-		Flags |= Modified;
-	}
-	ActiveGraph = GraphName;
-}
-
-TSharedRef<SWidget> FBPProfilerStatDiplayOptions::CreateToolbar()
-{
-	return	
-		SNew(SVerticalBox)
-		+SVerticalBox::Slot()
-		.HAlign(HAlign_Right)
-		.AutoHeight()
-		[
-			SNew(SHorizontalBox)
-			+SHorizontalBox::Slot()
-			.AutoWidth()
-			.Padding(FMargin(5,0))
-			[
-				SNew(SCheckBox)
-				.Content()
-				[
-					SNew(STextBlock)
-					.Text(LOCTEXT("FilterToGraph", "Filter to Graph"))
-				]
-				.IsChecked<FBPProfilerStatDiplayOptions, uint32>(this, &FBPProfilerStatDiplayOptions::GetChecked, GraphFilter)
-				.OnCheckStateChanged<FBPProfilerStatDiplayOptions, uint32>(this, &FBPProfilerStatDiplayOptions::OnChecked, GraphFilter)
-			]
-			+SHorizontalBox::Slot()
-			.AutoWidth()
-			.Padding(FMargin(5,0))
-			[
-				SNew(SCheckBox)
-				.Content()
-				[
-					SNew(STextBlock)
-					.Text(LOCTEXT("DisplayPureStats", "Pure Timings"))
-				]
-				.IsChecked<FBPProfilerStatDiplayOptions, uint32>(this, &FBPProfilerStatDiplayOptions::GetChecked, DisplayPure)
-				.OnCheckStateChanged<FBPProfilerStatDiplayOptions, uint32>(this, &FBPProfilerStatDiplayOptions::OnChecked, DisplayPure)
-			]
-			+SHorizontalBox::Slot()
-			.AutoWidth()
-			.Padding(FMargin(5,0))
-			[
-				SNew(SCheckBox)
-				.Content()
-				[
-					SNew(STextBlock)
-					.Text(LOCTEXT("ShowInstancesCheck", "Show Instances"))
-				]
-				.IsChecked<FBPProfilerStatDiplayOptions, uint32>(this, &FBPProfilerStatDiplayOptions::GetChecked, DisplayByInstance)
-				.OnCheckStateChanged<FBPProfilerStatDiplayOptions, uint32>(this, &FBPProfilerStatDiplayOptions::OnChecked, DisplayByInstance)
-			]
-			+SHorizontalBox::Slot()
-			.AutoWidth()
-			.Padding(FMargin(5,0))
-			[
-				SNew(SCheckBox)
-				.Content()
-				[
-					SNew(STextBlock)
-					.Text(LOCTEXT("InstanceFilterCheck", "Debug Filter Scope"))
-				]
-				.IsChecked<FBPProfilerStatDiplayOptions, uint32>(this, &FBPProfilerStatDiplayOptions::GetChecked, ScopeToDebugInstance)
-				.OnCheckStateChanged<FBPProfilerStatDiplayOptions, uint32>(this, &FBPProfilerStatDiplayOptions::OnChecked, ScopeToDebugInstance)
-			]
-			+SHorizontalBox::Slot()
-			.AutoWidth()
-			.Padding(FMargin(5,0))
-			[
-				SNew(SCheckBox)
-				.Content()
-				[
-					SNew(STextBlock)
-					.Text(LOCTEXT("AutoItemExpansion", "Auto Expand Statistics"))
-				]
-				.IsChecked<FBPProfilerStatDiplayOptions, uint32>(this, &FBPProfilerStatDiplayOptions::GetChecked, AutoExpand)
-				.OnCheckStateChanged<FBPProfilerStatDiplayOptions, uint32>(this, &FBPProfilerStatDiplayOptions::OnChecked, AutoExpand)
-			]
-		];
-}
-
-ECheckBoxState FBPProfilerStatDiplayOptions::GetChecked(const uint32 FlagsIn) const
-{
-	ECheckBoxState CheckedState;
-	if (FlagsIn & ScopeToDebugInstance)
-	{
-		if (HasFlags(DisplayByInstance))
-		{
-			CheckedState = HasFlags(ScopeToDebugInstance) ? ECheckBoxState::Checked : ECheckBoxState::Unchecked;
-		}
-		else
-		{
-			CheckedState = ECheckBoxState::Undetermined;
-		}
-	}
-	else
-	{
-		CheckedState = HasAllFlags(FlagsIn) ? ECheckBoxState::Checked : ECheckBoxState::Unchecked;
-	}
-	return CheckedState;
-}
-
-void FBPProfilerStatDiplayOptions::OnChecked(ECheckBoxState NewState, const uint32 FlagsIn)
-{
-	if (NewState == ECheckBoxState::Checked)
-	{
-		Flags |= FlagsIn;
-	}
-	else
-	{
-		Flags &= ~FlagsIn;
-	}
-	Flags |= Modified;
-}
-
-bool FBPProfilerStatDiplayOptions::IsFiltered(TSharedPtr<FScriptExecutionNode> Node) const
-{
-	bool bFilteredOut = !HasFlags(EDisplayFlags::DisplayPure) && Node->HasFlags(EScriptExecutionNodeFlags::PureStats);
-	if (Node->IsEvent() && HasFlags(EDisplayFlags::GraphFilter))
-	{
-		if (Node->GetGraphName() == UEdGraphSchema_K2::FN_UserConstructionScript)
-		{
-			bFilteredOut = ActiveGraph != UEdGraphSchema_K2::FN_UserConstructionScript;
-		}
-		else
-		{
-			bFilteredOut = ActiveGraph == UEdGraphSchema_K2::FN_UserConstructionScript;
-		}
-	}
-	return bFilteredOut;
-}
-
 //////////////////////////////////////////////////////////////////////////
 // FBPProfilerStatWidget
 
@@ -270,23 +120,13 @@ TSharedRef<SWidget> FBPProfilerStatWidget::GenerateColumnWidget(FName ColumnName
 				+SHorizontalBox::Slot()
 				.AutoWidth()
 				[
-					SNew(SImage)
-					.Image(ExecNode->GetIcon())
-					.ColorAndOpacity(ExecNode->GetIconColor())
+					ExecNode->GetIconWidget()
 				]
 				+SHorizontalBox::Slot()
 				.AutoWidth()
 				.Padding(FMargin(5,0))
 				[
-					SNew(SHyperlink)
-		#if TRACEPATH_DEBUG
-					.Text(FText::FromName(ExecNode->GetName()))
-		#else
-					.Text(ExecNode->GetDisplayName())
-		#endif
-					.Style(FEditorStyle::Get(), "HoverOnlyHyperlink")
-					.ToolTipText(ExecNode->GetToolTipText())
-					.OnNavigate(this, &FBPProfilerStatWidget::NavigateTo)
+					ExecNode->GetHyperlinkWidget()
 		#if TRACEPATH_DEBUG
 				]
 				+SHorizontalBox::Slot()
@@ -306,43 +146,28 @@ TSharedRef<SWidget> FBPProfilerStatWidget::GenerateColumnWidget(FName ColumnName
 
 			if (ColumnName == BlueprintProfilerStatText::ColumnId_TotalTime)
 			{
-				if (!ExecNode->HasFlags(EScriptExecutionNodeFlags::ExecPin))
-				{
-					TextAttr = TAttribute<FText>(PerformanceStats.Get(), &FScriptPerfData::GetTotalTimingText);
-					ColorAttr = TAttribute<FSlateColor>(FLinearColor::White);
-				}
+				TextAttr = TAttribute<FText>(PerformanceStats.Get(), &FScriptPerfData::GetTotalTimingText);
+				ColorAttr = TAttribute<FSlateColor>(FLinearColor::White);
 			}
 			else if (ColumnName == BlueprintProfilerStatText::ColumnId_InclusiveTime)
 			{
-				if (ExecNode->HasFlags(NonNodeStats))
-				{
-					TextAttr = TAttribute<FText>(PerformanceStats.Get(), &FScriptPerfData::GetInclusiveTimingText);
-					ColorAttr = TAttribute<FSlateColor>(this, &FBPProfilerStatWidget::GetInclusiveHeatColor);
-				}
+				TextAttr = TAttribute<FText>(PerformanceStats.Get(), &FScriptPerfData::GetInclusiveTimingText);
+				ColorAttr = TAttribute<FSlateColor>(this, &FBPProfilerStatWidget::GetInclusiveHeatColor);
 			}
-			else if (ColumnName == BlueprintProfilerStatText::ColumnId_Time)
+			else if (ColumnName == BlueprintProfilerStatText::ColumnId_AverageTime)
 			{
-				if (!ExecNode->HasFlags(EScriptExecutionNodeFlags::ExecPin))
-				{
-					TextAttr = TAttribute<FText>(PerformanceStats.Get(), &FScriptPerfData::GetNodeTimingText);
-					ColorAttr = TAttribute<FSlateColor>(this, &FBPProfilerStatWidget::GetNodeHeatColor);
-				}
+				TextAttr = TAttribute<FText>(PerformanceStats.Get(), &FScriptPerfData::GetExclusiveTimingText);
+				ColorAttr = TAttribute<FSlateColor>(this, &FBPProfilerStatWidget::GetNodeHeatColor);
 			}
 			else if (ColumnName == BlueprintProfilerStatText::ColumnId_MaxTime)
 			{
-				if (!ExecNode->HasFlags(EScriptExecutionNodeFlags::ExecPin))
-				{
-					TextAttr = TAttribute<FText>(PerformanceStats.Get(), &FScriptPerfData::GetMaxTimingText);
-					ColorAttr = TAttribute<FSlateColor>(this, &FBPProfilerStatWidget::GetMaxTimeHeatColor);
-				}
+				TextAttr = TAttribute<FText>(PerformanceStats.Get(), &FScriptPerfData::GetMaxTimingText);
+				ColorAttr = TAttribute<FSlateColor>(this, &FBPProfilerStatWidget::GetMaxTimeHeatColor);
 			}
 			else if (ColumnName == BlueprintProfilerStatText::ColumnId_MinTime)
 			{
-				if (!ExecNode->HasFlags(EScriptExecutionNodeFlags::ExecPin))
-				{
-					TextAttr = TAttribute<FText>(PerformanceStats.Get(), &FScriptPerfData::GetMinTimingText);
-					ColorAttr = TAttribute<FSlateColor>(FLinearColor::White);
-				}
+				TextAttr = TAttribute<FText>(PerformanceStats.Get(), &FScriptPerfData::GetMinTimingText);
+				ColorAttr = TAttribute<FSlateColor>(FLinearColor::White);
 			}
 			else if (ColumnName == BlueprintProfilerStatText::ColumnId_Samples)
 			{
@@ -384,12 +209,12 @@ FSlateColor FBPProfilerStatWidget::GetMaxTimeHeatColor() const
 	return FLinearColor(1.f, Value, Value);
 }
 
-void FBPProfilerStatWidget::GenerateExecNodeWidgets(const TSharedPtr<FBPProfilerStatDiplayOptions> DisplayOptions)
+void FBPProfilerStatWidget::GenerateExecNodeWidgets(const TSharedPtr<FBlueprintProfilerStatOptions> DisplayOptions)
 {
 	if (ExecNode.IsValid())
 	{
 		// Grab Performance Stats
-		PerformanceStats = ExecNode->GetPerfDataByInstanceAndTracePath(DisplayOptions->GetActiveInstance(), WidgetTracePath);
+		PerformanceStats = ExecNode->GetOrAddPerfDataByInstanceAndTracePath(DisplayOptions->GetActiveInstance(), WidgetTracePath);
 		CachedChildren.Reset(0);
 
 		if (ExecNode->HasFlags(EScriptExecutionNodeFlags::PureStats))
@@ -458,6 +283,10 @@ void FBPProfilerStatWidget::GenerateExecNodeWidgets(const TSharedPtr<FBPProfiler
 					{
 						TArray<FScriptNodeExecLinkage::FLinearExecPath> LinearExecNodes;
 						FTracePath LinkPath(WidgetTracePath);
+						if (LinkIter.Value->HasFlags(EScriptExecutionNodeFlags::AsyncTaskDelegate))
+						{
+							LinkPath.ResetPath();
+						}
 						if (!LinkIter.Value->HasFlags(EScriptExecutionNodeFlags::InvalidTrace))
 						{
 							LinkPath.AddExitPin(LinkIter.Key);
diff --git a/Engine/Source/Editor/Kismet/Private/Profiler/BPProfilerStatisticWidgets.h b/Engine/Source/Editor/Kismet/Private/Profiler/BPProfilerStatisticWidgets.h
--- a/Engine/Source/Editor/Kismet/Private/Profiler/BPProfilerStatisticWidgets.h
+++ b/Engine/Source/Editor/Kismet/Private/Profiler/BPProfilerStatisticWidgets.h
@@ -15,8 +15,8 @@ namespace EBlueprintProfilerStat
 	{
 		Name = 0,
 		TotalTime,
+		AverageTime,
 		InclusiveTime,
-		Time,
 		PureTime,
 		MaxTime,
 		MinTime,
@@ -53,81 +53,6 @@ public:
 
 };
 
-//////////////////////////////////////////////////////////////////////////
-// FBPProfilerStatDiplayOptions
-
-struct FBPProfilerStatDiplayOptions : public TSharedFromThis<FBPProfilerStatDiplayOptions>
-{
-public:
-
-	/** Display flags controlling widget appearance */
-	enum EDisplayFlags
-	{
-		DisplayByInstance			= 0x00000001,	// Show instance stats
-		ScopeToDebugInstance		= 0x00000002,	// Scope stats to current debug instance
-		GraphFilter					= 0x00000004,	// Filter to current graph
-		DisplayPure					= 0x00000008,	// Display pure stats
-		AutoExpand					= 0x00000010,	// Auto expand stats
-		Modified					= 0x10000000,	// Display state was changed
-		Default						= DisplayByInstance|ScopeToDebugInstance|DisplayPure|GraphFilter|Modified
-	};
-
-	FBPProfilerStatDiplayOptions()
-		: ActiveInstance(NAME_None)
-		, ActiveGraph(NAME_None)
-		, Flags(Default)
-	{
-	}
-
-	/** Marks display state as modified */
-	void SetStateModified() { Flags |= Modified; }
-
-	/** Returns if the filter state options are modified */
-	bool IsStateModified() const { return (Flags & Modified) != 0U; }
-
-	/** Add flags */
-	void AddFlags(const uint32 FlagsIn) { Flags |= FlagsIn; }
-
-	/** Clear flags */
-	void ClearFlags(const uint32 FlagsIn) { Flags &= ~FlagsIn; }
-
-	/** Returns true if any flags match */
-	bool HasFlags(const uint32 FlagsIn) const { return (Flags & FlagsIn) != 0U; }
-
-	/** Returns true if all flags match */
-	bool HasAllFlags(const uint32 FlagsIn) const { return (Flags & FlagsIn) == FlagsIn; }
-
-	/** Set active instance name */
-	FName GetActiveInstance() const { return HasFlags(DisplayByInstance) ? ActiveInstance : NAME_None; }
-
-	/** Set active instance name */
-	void SetActiveInstance(const FName InstanceName);
-
-	/** Set active graph name */
-	void SetActiveGraph(const FName GraphName);
-
-	/** Create toolbar widget */
-	TSharedRef<SWidget> CreateToolbar();
-
-	/** Returns if the UI should be checked for flags */
-	ECheckBoxState GetChecked(const uint32 FlagsIn) const;
-
-	/** Handles UI check states */
-	void OnChecked(ECheckBoxState NewState, const uint32 FlagsIn);
-
-	/** Returns true if the exec node passes current filters */
-	bool IsFiltered(TSharedPtr<class FScriptExecutionNode> Node) const;
-
-private:
-
-	/** The active instance name */
-	FName ActiveInstance;
-	/** The active graph name */
-	FName ActiveGraph;
-	/** The display flags */
-	uint32 Flags;
-};
-
 //////////////////////////////////////////////////////////////////////////
 // FBPProfilerStatWidget
 
@@ -140,7 +65,7 @@ public:
 	virtual ~FBPProfilerStatWidget() {}
 
 	/** Generate exec node widgets */
-	virtual void GenerateExecNodeWidgets(const TSharedPtr<FBPProfilerStatDiplayOptions> DisplayOptions);
+	virtual void GenerateExecNodeWidgets(const TSharedPtr<struct FBlueprintProfilerStatOptions> DisplayOptions);
 
 	/** Gather children for list/tree widget */
 	void GatherChildren(TArray<TSharedPtr<FBPProfilerStatWidget>>& OutChildren);
diff --git a/Engine/Source/Editor/Kismet/Private/Profiler/EventExecution.cpp b/Engine/Source/Editor/Kismet/Private/Profiler/EventExecution.cpp
--- a/Engine/Source/Editor/Kismet/Private/Profiler/EventExecution.cpp
+++ b/Engine/Source/Editor/Kismet/Private/Profiler/EventExecution.cpp
@@ -2,6 +2,13 @@
 
 #include "BlueprintEditorPrivatePCH.h"
 #include "EventExecution.h"
+#include "BlueprintProfilerModule.h"
+#include "SHyperlink.h"
+
+// This forces checking against correct tracepaths when querying data
+#define STRICT_PERFDATA_CREATION 0
+
+#define LOCTEXT_NAMESPACE "BlueprintEventUI"
 
 //////////////////////////////////////////////////////////////////////////
 // FScriptNodeExecLinkage
@@ -55,15 +62,32 @@ bool FScriptNodePerfData::HasPerfDataForInstanceAndTracePath(FName InstanceName,
 	return false;
 }
 
+TSharedPtr<FScriptPerfData> FScriptNodePerfData::GetOrAddPerfDataByInstanceAndTracePath(FName InstanceName, const FTracePath& TracePath)
+{ 
+	TMap<const uint32, TSharedPtr<FScriptPerfData>>& InstanceMap = InstanceInputPinToPerfDataMap.FindOrAdd(InstanceName);
+	TSharedPtr<FScriptPerfData>& Result = InstanceMap.FindOrAdd(TracePath);
+	if (!Result.IsValid())
+	{
+		Result = MakeShareable<FScriptPerfData>(new FScriptPerfData(GetPerfDataType()));
+	}
+	return Result;
+}
+
 TSharedPtr<FScriptPerfData> FScriptNodePerfData::GetPerfDataByInstanceAndTracePath(FName InstanceName, const FTracePath& TracePath)
 { 
 	TMap<const uint32, TSharedPtr<FScriptPerfData>>& InstanceMap = InstanceInputPinToPerfDataMap.FindOrAdd(InstanceName);
+	#if STRICT_PERFDATA_CREATION
+	TSharedPtr<FScriptPerfData>* Result = InstanceMap.Find(TracePath);
+	check (Result != nullptr);
+	return *Result;
+	#else
 	TSharedPtr<FScriptPerfData>& Result = InstanceMap.FindOrAdd(TracePath);
 	if (!Result.IsValid())
 	{
 		Result = MakeShareable<FScriptPerfData>(new FScriptPerfData(GetPerfDataType()));
 	}
 	return Result;
+	#endif
 }
 
 void FScriptNodePerfData::GetInstancePerfDataByTracePath(const FTracePath& TracePath, TArray<TSharedPtr<FScriptPerfData>>& ResultsOut)
@@ -81,8 +105,25 @@ void FScriptNodePerfData::GetInstancePerfDataByTracePath(const FTracePath& Trace
 	}
 }
 
+TSharedPtr<FScriptPerfData> FScriptNodePerfData::GetOrAddBlueprintPerfDataByTracePath(const FTracePath& TracePath)
+{
+	TMap<const uint32, TSharedPtr<FScriptPerfData>>& BlueprintMap = InstanceInputPinToPerfDataMap.FindOrAdd(NAME_None);
+	TSharedPtr<FScriptPerfData>& Result = BlueprintMap.FindOrAdd(TracePath);
+	if (!Result.IsValid())
+	{
+		Result = MakeShareable<FScriptPerfData>(new FScriptPerfData(GetPerfDataType()));
+	}
+	return Result;
+}
+
 TSharedPtr<FScriptPerfData> FScriptNodePerfData::GetBlueprintPerfDataByTracePath(const FTracePath& TracePath)
 {
+	#if STRICT_PERFDATA_CREATION
+	TMap<const uint32, TSharedPtr<FScriptPerfData>>& BlueprintMap = InstanceInputPinToPerfDataMap.FindOrAdd(NAME_None);
+	TSharedPtr<FScriptPerfData>* Result = BlueprintMap.Find(TracePath);
+	check (Result != nullptr);
+	return *Result;
+	#else
 	TMap<const uint32, TSharedPtr<FScriptPerfData>>& BlueprintMap = InstanceInputPinToPerfDataMap.FindOrAdd(NAME_None);
 	TSharedPtr<FScriptPerfData>& Result = BlueprintMap.FindOrAdd(TracePath);
 	if (!Result.IsValid())
@@ -90,6 +131,7 @@ TSharedPtr<FScriptPerfData> FScriptNodePerfData::GetBlueprintPerfDataByTracePath
 		Result = MakeShareable<FScriptPerfData>(new FScriptPerfData(GetPerfDataType()));
 	}
 	return Result;
+	#endif
 }
 
 void FScriptNodePerfData::GetBlueprintPerfDataForAllTracePaths(FScriptPerfData& OutPerfData)
@@ -121,6 +163,7 @@ FScriptExecutionNode::FScriptExecutionNode(const FScriptExecNodeParams& InitPara
 	, OwningGraphName(InitParams.OwningGraphName)
 	, NodeFlags(InitParams.NodeFlags)
 	, ObservedObject(InitParams.ObservedObject)
+	, ObservedPin(InitParams.ObservedPin)
 	, DisplayName(InitParams.DisplayName)
 	, Tooltip(InitParams.Tooltip)
 	, IconColor(InitParams.IconColor)
@@ -134,6 +177,26 @@ bool FScriptExecutionNode::operator == (const FScriptExecutionNode& NodeIn) cons
 	return NodeName == NodeIn.NodeName;
 }
 
+TSharedRef<SWidget> FScriptExecutionNode::GetIconWidget()
+{
+	return SNew(SImage)
+			.Image(Icon)
+			.ColorAndOpacity(IconColor);
+}
+
+TSharedRef<SWidget> FScriptExecutionNode::GetHyperlinkWidget()
+{
+	return SNew(SHyperlink)
+		#if TRACEPATH_DEBUG
+		.Text(FText::FromName(NodeName))
+		#else
+		.Text(DisplayName)
+		#endif
+		.Style(FEditorStyle::Get(), "HoverOnlyHyperlink")
+		.ToolTipText(Tooltip)
+		.OnNavigate(this, &FScriptExecutionNode::NavigateToObject);
+}
+
 void FScriptExecutionNode::GetLinearExecutionPath(TArray<FLinearExecPath>& LinearExecutionNodes, const FTracePath& TracePath, const bool bIncludeChildren)
 {
 	LinearExecutionNodes.Add(FLinearExecPath(AsShared(), TracePath));
@@ -156,6 +219,10 @@ void FScriptExecutionNode::GetLinearExecutionPath(TArray<FLinearExecPath>& Linea
 			else
 			{
 				FTracePath NewTracePath(TracePath);
+				if (NodeIter.Value->HasFlags(EScriptExecutionNodeFlags::AsyncTaskDelegate))
+				{
+					NewTracePath.ResetPath();
+				}
 				if (NodeIter.Key != INDEX_NONE && !HasFlags(EScriptExecutionNodeFlags::InvalidTrace))
 				{
 					NewTracePath.AddExitPin(NodeIter.Key);
@@ -231,10 +298,12 @@ void FScriptExecutionNode::RefreshStats(const FTracePath& TracePath)
 		{
 			// Refresh child stats and copy as branch stats - this is a dummy entry
 			TSet<FName> AllInstances;
+			TArray<FLinearExecPath> ChildExecPaths;
 			for (auto ChildIter : ChildNodes)
 			{
 				ChildIter->GetValidInstanceNames(AllInstances);
 				ChildIter->RefreshStats(TracePath);
+				ChildIter->GetLinearExecutionPath(ChildExecPaths, TracePath, false);
 			}
 			for (auto InstanceIter : AllInstances)
 			{
@@ -244,15 +313,15 @@ void FScriptExecutionNode::RefreshStats(const FTracePath& TracePath)
 					InstancePerfData->Reset();
 					BlueprintPooledStats.Add(InstancePerfData);
 
-					for (auto ChildIter : ChildNodes)
+					for (auto ChildIter : ChildExecPaths)
 					{
-						TSharedPtr<FScriptPerfData> ChildPerfData = ChildIter->GetPerfDataByInstanceAndTracePath(InstanceIter, TracePath);
+						TSharedPtr<FScriptPerfData> ChildPerfData = ChildIter.LinkedNode->GetPerfDataByInstanceAndTracePath(InstanceIter, ChildIter.TracePath);
 						InstancePerfData->AddBranchData(*ChildPerfData.Get());
 					}
 				}
 			}
 			// Update Blueprint stats as branches
-			TSharedPtr<FScriptPerfData> BlueprintData = GetBlueprintPerfDataByTracePath(TracePath);
+			TSharedPtr<FScriptPerfData> BlueprintData = GetOrAddBlueprintPerfDataByTracePath(TracePath);
 			BlueprintData->Reset();
 
 			for (auto BlueprintChildDataIter : BlueprintPooledStats)
@@ -278,6 +347,10 @@ void FScriptExecutionNode::RefreshStats(const FTracePath& TracePath)
 		for (auto LinkIter : LinkedNodes)
 		{
 			FTracePath LinkTracePath(TracePath);
+			if (LinkIter.Value->HasFlags(EScriptExecutionNodeFlags::AsyncTaskDelegate))
+			{
+				LinkTracePath.ResetPath();
+			}
 			if (!LinkIter.Value->HasFlags(EScriptExecutionNodeFlags::InvalidTrace))
 			{
 				LinkTracePath.AddExitPin(LinkIter.Key);
@@ -288,14 +361,45 @@ void FScriptExecutionNode::RefreshStats(const FTracePath& TracePath)
 	// Update the owning blueprint stats
 	if (BlueprintPooledStats.Num() > 0)
 	{
-		TSharedPtr<FScriptPerfData> BlueprintData = GetBlueprintPerfDataByTracePath(TracePath);
-		BlueprintData->Reset();
+		TSharedPtr<FScriptPerfData> BlueprintData = GetOrAddBlueprintPerfDataByTracePath(TracePath);
+		BlueprintData->InitialiseFromDataSet(BlueprintPooledStats);
+	}
+}
+
+float FScriptExecutionNode::CalculateHottestPathStats(FScriptExecutionHottestPathParams HotPathParams)
+{
+	// Grab local perf data
+	TSharedPtr<FScriptPerfData> LocalPerfData = GetOrAddPerfDataByInstanceAndTracePath(HotPathParams.InstanceName, HotPathParams.TracePath);
+	float AccumulatedTime = 0.f;
 
-		for (auto BlueprintChildDataIter : BlueprintPooledStats)
+	if (!IsEvent())
+	{
+		// Subtract local inclusive cost from parent inclusive cost
+		const double NodeTime = IsBranch() ? LocalPerfData->GetExclusiveTiming() : LocalPerfData->GetInclusiveTiming();
+		HotPathParams.TimeSoFar += NodeTime;
+		AccumulatedTime += NodeTime;
+		const float HottestEndpointValue = static_cast<float>(HotPathParams.TimeSoFar / HotPathParams.EventTiming);
+		LocalPerfData->SetHottestEndpointHeatLevel(HottestEndpointValue);
+	}
+	// Update children
+	for (auto ChildIter : ChildNodes)
+	{
+		AccumulatedTime += ChildIter->CalculateHottestPathStats(HotPathParams);
+	}
+	// Update Links
+	for (auto LinkIter : LinkedNodes)
+	{
+		FScriptExecutionHottestPathParams LinkedHotPathParams(HotPathParams);
+		if (!LinkIter.Value->HasFlags(EScriptExecutionNodeFlags::InvalidTrace))
 		{
-			BlueprintData->AddData(*BlueprintChildDataIter.Get());
+			LinkedHotPathParams.TracePath.AddExitPin(LinkIter.Key);
 		}
+		AccumulatedTime += LinkIter.Value->CalculateHottestPathStats(LinkedHotPathParams);
 	}
+	const float HottestPathValue = static_cast<float>(AccumulatedTime / HotPathParams.EventTiming);
+	LocalPerfData->SetHottestPathHeatLevel(HottestPathValue);
+
+	return AccumulatedTime;
 }
 
 void FScriptExecutionNode::GetAllExecNodes(TMap<FName, TSharedPtr<FScriptExecutionNode>>& ExecNodesOut)
@@ -350,12 +454,24 @@ TSharedPtr<FScriptExecutionNode> FScriptExecutionNode::GetPureChainNode()
 
 void FScriptExecutionNode::NavigateToObject() const
 {
-	if (ObservedObject.IsValid())
+	if (IsObservedObjectValid())
 	{
-		FKismetEditorUtilities::BringKismetToFocusAttentionOnObject(ObservedObject.Get());
+		if (UEdGraphPin* Pin = ObservedPin.Get())
+		{
+			FKismetEditorUtilities::BringKismetToFocusAttentionOnPin(Pin);
+		}
+		else
+		{
+			FKismetEditorUtilities::BringKismetToFocusAttentionOnObject(ObservedObject.Get());
+		}
 	}
 }
 
+bool FScriptExecutionNode::IsObservedObjectValid() const
+{
+	return ObservedObject.IsValid() && !ObservedObject.Get()->IsPendingKill();
+}
+
 const EScriptPerfDataType FScriptExecutionNode::GetPerfDataType() const
 {
 	return IsEvent() ? EScriptPerfDataType::Event : EScriptPerfDataType::Node;
@@ -382,7 +498,7 @@ void FScriptExecutionTunnelEntry::GetLinearExecutionPath(TArray<FLinearExecPath>
 	LinearExecutionNodes.Add(FLinearExecPath(AsShared(), TracePath));
 	if (TunnelType != MultiIOTunnel)
 	{
-		FTracePath TunnelTracePath(TracePath, SharedThis<const FScriptExecutionTunnelEntry>(this));
+		FTracePath TunnelTracePath(TracePath, SharedThis<FScriptExecutionTunnelEntry>(this));
 		for (auto TunnelExit : LinkedNodes)
 		{
 			if (IsPathValidForTunnel(TunnelExit.Key))
@@ -401,7 +517,7 @@ void FScriptExecutionTunnelEntry::GetFilteredChildNodes(TArray<FLinearExecPath>&
 	{
 		case SimpleTunnel:
 		{
-			FTracePath ChildTrace(TracePath);
+			FTracePath ChildTrace(TracePath, SharedThis<FScriptExecutionTunnelEntry>(this));
 			for (auto Child : ChildNodes)
 			{
 				if (Child->HasFlags(EScriptExecutionNodeFlags::TunnelEntryPin))
@@ -421,7 +537,7 @@ void FScriptExecutionTunnelEntry::GetFilteredChildNodes(TArray<FLinearExecPath>&
 		}
 		case SinglePathTunnel:
 		{
-			FTracePath ChildTrace(TracePath);
+			FTracePath ChildTrace(TracePath, SharedThis<FScriptExecutionTunnelEntry>(this));
 			for (auto Child : ChildNodes)
 			{
 				ChildArrayInOut.Add(FLinearExecPath(Child, ChildTrace));
@@ -431,7 +547,7 @@ void FScriptExecutionTunnelEntry::GetFilteredChildNodes(TArray<FLinearExecPath>&
 		case MultiIOTunnel:
 		{
 			TArray<FLinearExecPath> Results;
-			FTracePath ChildTrace(TracePath);
+			FTracePath ChildTrace(TracePath, SharedThis<FScriptExecutionTunnelEntry>(this));
 			for (auto Child : ChildNodes)
 			{
 				const bool bIgnoreChildren = Child->HasFlags(EScriptExecutionNodeFlags::TunnelExitPin);
@@ -508,10 +624,12 @@ void FScriptExecutionTunnelEntry::RefreshStats(const FTracePath& TracePath)
 	TArray<TSharedPtr<FScriptPerfData>> BlueprintPooledStats;
 	// Update stats based on type
 	GetInstancePerfDataByTracePath(TracePath, BlueprintPooledStats);
+	// Create new trace
+	FTracePath TunnelTracePath(TracePath, SharedThis<FScriptExecutionTunnelEntry>(this));
 	// Refresh Child Links
 	for (auto ChildIter : ChildNodes)
 	{
-		ChildIter->RefreshStats(TracePath);
+		ChildIter->RefreshStats(TunnelTracePath);
 	}
 	// Refresh All links
 	for (auto LinkIter : LinkedNodes)
@@ -524,11 +642,7 @@ void FScriptExecutionTunnelEntry::RefreshStats(const FTracePath& TracePath)
 	if (BlueprintPooledStats.Num() > 0)
 	{
 		TSharedPtr<FScriptPerfData> BlueprintData = GetBlueprintPerfDataByTracePath(TracePath);
-		BlueprintData->Reset();
-		for (auto BlueprintChildDataIter : BlueprintPooledStats)
-		{
-			BlueprintData->AddData(*BlueprintChildDataIter.Get());
-		}
+		BlueprintData->InitialiseFromDataSet(BlueprintPooledStats);
 	}
 }
 
@@ -577,9 +691,42 @@ void FScriptExecutionTunnelExit::GetLinearExecutionPath(TArray<FLinearExecPath>&
 	}
 }
 
+void FScriptExecutionTunnelExit::RefreshStats(const FTracePath& TracePath)
+{
+	// Process stat update
+	TArray<TSharedPtr<FScriptPerfData>> BlueprintPooledStats;
+	TSharedPtr<const FScriptExecutionTunnelEntry> TunnelNode = TracePath.GetTunnel();
+	check (TunnelNode.IsValid());
+	// Update stats based on type
+	GetInstancePerfDataByTracePath(TracePath, BlueprintPooledStats);
+	// Refresh all valid tunnel links
+	for (auto LinkIter : LinkedNodes)
+	{
+		if (TunnelNode->IsPathValidForTunnel(LinkIter.Key))
+		{
+			FTracePath LinkTracePath(TracePath);
+			if (LinkIter.Value->HasFlags(EScriptExecutionNodeFlags::AsyncTaskDelegate))
+			{
+				LinkTracePath.ResetPath();
+			}
+			if (!LinkIter.Value->HasFlags(EScriptExecutionNodeFlags::InvalidTrace))
+			{
+				LinkTracePath.AddExitPin(LinkIter.Key);
+			}
+			LinkIter.Value->RefreshStats(LinkTracePath);
+		}
+	}
+	// Update the owning blueprint stats
+	if (BlueprintPooledStats.Num() > 0)
+	{
+		TSharedPtr<FScriptPerfData> BlueprintData = GetOrAddBlueprintPerfDataByTracePath(TracePath);
+		BlueprintData->InitialiseFromDataSet(BlueprintPooledStats);
+	}
+}
+
 void FScriptExecutionTunnelExit::NavigateToObject() const
 {
-	if (ObservedObject.IsValid())
+	if (IsObservedObjectValid())
 	{
 		FKismetEditorUtilities::BringKismetToFocusAttentionOnObject(ObservedObject.Get());
 	}
@@ -649,12 +796,7 @@ void FScriptExecutionPureNode::RefreshStats(const FTracePath& TracePath)
 	if (BlueprintPooledStats.Num() > 0)
 	{
 		TSharedPtr<FScriptPerfData> BlueprintData = GetBlueprintPerfDataByTracePath(TracePath);
-		BlueprintData->Reset();
-
-		for (auto BlueprintChildDataIter : BlueprintPooledStats)
-		{
-			BlueprintData->AddData(*BlueprintChildDataIter.Get());
-		}
+		BlueprintData->InitialiseFromDataSet(BlueprintPooledStats);
 	}
 }
 
@@ -685,16 +827,25 @@ void FScriptExecutionBlueprint::RefreshStats(const FTracePath& TracePath)
 		// as an overall blueprint performance representation.
 		BlueprintEventIter->RefreshStats(TracePath);
 	}
+	// Determine if we require the hottest path's updating
+	IBlueprintProfilerInterface& ProfilerModule = FModuleManager::LoadModuleChecked<IBlueprintProfilerInterface>("BlueprintProfiler");
+	const bool bUpdateHottestPaths = ProfilerModule.GetWireHeatMapDisplayMode() != EBlueprintProfilerHeatMapDisplayMode::None;
 	// Update instance stats
 	for (auto InstanceIter : Instances)
 	{
-		TSharedPtr<FScriptPerfData> InstancePerfData = InstanceIter->GetPerfDataByInstanceAndTracePath(InstanceIter->GetName(), TracePath);
+		TSharedPtr<FScriptPerfData> InstancePerfData = InstanceIter->GetOrAddPerfDataByInstanceAndTracePath(InstanceIter->GetName(), TracePath);
 		InstancePerfData->Reset();
 		// Update all top level instance stats now the events are up to date.
 		for (auto InstanceEventIter : InstanceIter->GetChildNodes())
 		{
-			TSharedPtr<FScriptPerfData> InstanceEventPerfData = InstanceEventIter->GetPerfDataByInstanceAndTracePath(InstanceIter->GetName(), TracePath);
+			TSharedPtr<FScriptPerfData> InstanceEventPerfData = InstanceEventIter->GetOrAddPerfDataByInstanceAndTracePath(InstanceIter->GetName(), TracePath);
 			InstancePerfData->AddData(*InstanceEventPerfData.Get());
+			// Update the hottest path stats
+			if (bUpdateHottestPaths)
+			{
+				FScriptExecutionHottestPathParams HotPathParams(InstanceIter->GetName(), InstanceEventPerfData->GetExclusiveTiming());
+				InstanceEventIter->CalculateHottestPathStats(HotPathParams);
+			}
 		}
 		// Add for consolidation at the bottom.
 		InstanceData.Add(InstancePerfData);
@@ -702,7 +853,7 @@ void FScriptExecutionBlueprint::RefreshStats(const FTracePath& TracePath)
 	// Finally... update the root stats for the owning blueprint.
 	if (InstanceData.Num() > 0)
 	{
-		TSharedPtr<FScriptPerfData> BlueprintData = GetBlueprintPerfDataByTracePath(TracePath);
+		TSharedPtr<FScriptPerfData> BlueprintData = GetOrAddBlueprintPerfDataByTracePath(TracePath);
 		BlueprintData->Reset();
 
 		for (auto InstanceIter : InstanceData)
@@ -729,7 +880,7 @@ void FScriptExecutionBlueprint::GetAllExecNodes(TMap<FName, TSharedPtr<FScriptEx
 
 void FScriptExecutionBlueprint::NavigateToObject() const
 {
-	if (ObservedObject.IsValid())
+	if (IsObservedObjectValid())
 	{
 		if (const UBlueprint* Blueprint = Cast<UBlueprint>(ObservedObject.Get()))
 		{
@@ -745,30 +896,64 @@ void FScriptExecutionBlueprint::NavigateToObject() const
 
 void FScriptExecutionInstance::NavigateToObject() const
 {
-	if (ObservedObject.IsValid())
+	if (const FWorldContext* PIEWorldContext = GEditor->GetPIEWorldContext())
 	{
-		if (const AActor* Actor = Cast<AActor>(ObservedObject.Get()))
+		if (const AActor* PIEActor = Cast<AActor>(PIEInstance.Get()))
 		{
-			if (const FWorldContext* PIEWorldContext = GEditor->GetPIEWorldContext())
-			{
-				if (const UWorld* PIEWorld = PIEWorldContext->World())
-				{
-					for (auto LevelIter : PIEWorld->GetLevels())
-					{
-						if (AActor* PIEActor = Cast<AActor>(FindObject<UObject>(LevelIter, *Actor->GetName())))
-						{
-							GEditor->SelectNone(false, false);
-							GEditor->SelectActor(const_cast<AActor*>(PIEActor), true, true, true);
-							break;
-						}
-					}
-				}
-			}
-			else
-			{
-				GEditor->SelectNone(false, false);
-				GEditor->SelectActor(const_cast<AActor*>(Actor), true, true, true);
-			}
+			GEditor->SelectNone(false, false);
+			GEditor->SelectActor(const_cast<AActor*>(PIEActor), true, true, true);
 		}
 	}
+	else
+	{
+		if (const AActor* EditorActor = Cast<AActor>(ObservedObject.Get()))
+		{
+			GEditor->SelectNone(false, false);
+			GEditor->SelectActor(const_cast<AActor*>(EditorActor), true, true, true);
+		}
+	}
+}
+
+bool FScriptExecutionInstance::IsObservedObjectValid() const
+{
+	if (const FWorldContext* PIEWorldContext = GEditor->GetPIEWorldContext())
+	{
+		return PIEInstance.IsValid() && !PIEInstance.Get()->IsPendingKill();
+	}
+	else
+	{
+		return ObservedObject.IsValid() && !ObservedObject.Get()->IsPendingKill();
+	}
+}
+
+TSharedRef<SWidget> FScriptExecutionInstance::GetIconWidget()
+{
+	return SNew(SImage)
+			.Image(Icon)
+			.ColorAndOpacity(this, &FScriptExecutionInstance::GetInstanceIconColor);
+}
+
+TSharedRef<SWidget> FScriptExecutionInstance::GetHyperlinkWidget()
+{
+	return SNew(SHyperlink)
+		#if TRACEPATH_DEBUG
+		.Text(FText::FromName(NodeName))
+		#else
+		.Text(DisplayName)
+		#endif
+		.Style(FEditorStyle::Get(), "HoverOnlyHyperlink")
+		.ToolTipText(this, &FScriptExecutionInstance::GetInstanceTooltip)
+		.OnNavigate(this, &FScriptExecutionInstance::NavigateToObject);
 }
+
+FSlateColor FScriptExecutionInstance::GetInstanceIconColor() const
+{
+	return FSlateColor(IsObservedObjectValid() ? IconColor : FLinearColor(1.f, 1.f, 1.f, 0.3f));
+}
+
+FText FScriptExecutionInstance::GetInstanceTooltip() const
+{
+	return IsObservedObjectValid() ? Tooltip : LOCTEXT("ActorInstanceInvalid_Tooltip", "Actor instance no longer exists");
+}
+
+#undef LOCTEXT_NAMESPACE
\ No newline at end of file
diff --git a/Engine/Source/Editor/Kismet/Private/Profiler/SBlueprintProfilerToolbar.cpp b/Engine/Source/Editor/Kismet/Private/Profiler/SBlueprintProfilerToolbar.cpp
new file mode 100644
--- /dev/null
+++ b/Engine/Source/Editor/Kismet/Private/Profiler/SBlueprintProfilerToolbar.cpp
@@ -0,0 +1,1213 @@
+// Copyright 1998-2016 Epic Games, Inc. All Rights Reserved.
+
+#include "BlueprintEditorPrivatePCH.h"
+#include "SBlueprintProfilerToolbar.h"
+//#include "BPProfilerStatisticWidgets.h"
+#include "EventExecution.h"
+#include "SNumericEntryBox.h"
+#include "Editor/UnrealEd/Classes/Settings/EditorExperimentalSettings.h"
+
+#define LOCTEXT_NAMESPACE "BlueprintProfilerToolbar"
+
+//////////////////////////////////////////////////////////////////////////
+// FBPProfilerStatOptions
+
+void FBlueprintProfilerStatOptions::SetActiveInstance(const FName InstanceName)
+{
+	if (ActiveInstance != InstanceName && HasFlags(ScopeToDebugInstance))
+	{
+		Flags |= Modified;
+	}
+	ActiveInstance = InstanceName;
+}
+
+void FBlueprintProfilerStatOptions::SetActiveGraph(const FName GraphName)
+{
+	if (ActiveGraph != GraphName && HasFlags(GraphFilter))
+	{
+		Flags |= Modified;
+	}
+	ActiveGraph = GraphName;
+}
+
+TSharedRef<SWidget> FBlueprintProfilerStatOptions::CreateToolbar()
+{
+	return	
+		SNew(SVerticalBox)
+		+SVerticalBox::Slot()
+		.HAlign(HAlign_Right)
+		.AutoHeight()
+		[
+			SNew(SHorizontalBox)
+			+SHorizontalBox::Slot()
+			.AutoWidth()
+			.Padding(FMargin(5,0))
+			[
+				SNew(SCheckBox)
+				.Content()
+				[
+					SNew(STextBlock)
+					.Text(LOCTEXT("FilterToGraph", "Filter to Graph"))
+				]
+				.IsChecked<FBlueprintProfilerStatOptions, uint32>(this, &FBlueprintProfilerStatOptions::GetChecked, GraphFilter)
+				.OnCheckStateChanged<FBlueprintProfilerStatOptions, uint32>(this, &FBlueprintProfilerStatOptions::OnChecked, GraphFilter)
+			]
+			+SHorizontalBox::Slot()
+			.AutoWidth()
+			.Padding(FMargin(5,0))
+			[
+				SNew(SCheckBox)
+				.Content()
+				[
+					SNew(STextBlock)
+					.Text(LOCTEXT("DisplayPureStats", "Pure Timings"))
+				]
+				.IsChecked<FBlueprintProfilerStatOptions, uint32>(this, &FBlueprintProfilerStatOptions::GetChecked, DisplayPure)
+				.OnCheckStateChanged<FBlueprintProfilerStatOptions, uint32>(this, &FBlueprintProfilerStatOptions::OnChecked, DisplayPure)
+			]
+			+SHorizontalBox::Slot()
+			.AutoWidth()
+			.Padding(FMargin(5,0))
+			[
+				SNew(SCheckBox)
+				.Content()
+				[
+					SNew(STextBlock)
+					.Text(LOCTEXT("ShowInstancesCheck", "Show Instances"))
+				]
+				.IsChecked<FBlueprintProfilerStatOptions, uint32>(this, &FBlueprintProfilerStatOptions::GetChecked, DisplayByInstance)
+				.OnCheckStateChanged<FBlueprintProfilerStatOptions, uint32>(this, &FBlueprintProfilerStatOptions::OnChecked, DisplayByInstance)
+			]
+			+SHorizontalBox::Slot()
+			.AutoWidth()
+			.Padding(FMargin(5,0))
+			[
+				SNew(SCheckBox)
+				.Content()
+				[
+					SNew(STextBlock)
+					.Text(LOCTEXT("InstanceFilterCheck", "Debug Filter Scope"))
+				]
+				.IsChecked<FBlueprintProfilerStatOptions, uint32>(this, &FBlueprintProfilerStatOptions::GetChecked, ScopeToDebugInstance)
+				.OnCheckStateChanged<FBlueprintProfilerStatOptions, uint32>(this, &FBlueprintProfilerStatOptions::OnChecked, ScopeToDebugInstance)
+			]
+			//+SHorizontalBox::Slot()
+			//.AutoWidth()
+			//.Padding(FMargin(5,0))
+			//[
+			//	SNew(SCheckBox)
+			//	.Content()
+			//	[
+			//		SNew(STextBlock)
+			//		.Text(LOCTEXT("AutoItemExpansion", "Auto Expand Statistics"))
+			//	]
+			//	.IsChecked<FBlueprintProfilerStatOptions, uint32>(this, &FBlueprintProfilerStatOptions::GetChecked, AutoExpand)
+			//	.OnCheckStateChanged<FBlueprintProfilerStatOptions, uint32>(this, &FBlueprintProfilerStatOptions::OnChecked, AutoExpand)
+			//]
+		];
+}
+
+TSharedRef<SWidget> FBlueprintProfilerStatOptions::CreateComboContent()
+{
+	return	
+		SNew(SVerticalBox)
+		+SVerticalBox::Slot()
+		.HAlign(HAlign_Left)
+		.AutoHeight()
+		[
+			SNew(SHorizontalBox)
+			+SHorizontalBox::Slot()
+			.AutoWidth()
+			.Padding(FMargin(5,2))
+			[
+				SNew(SCheckBox)
+				.Content()
+				[
+					SNew(STextBlock)
+					.Text(LOCTEXT("FilterToGraph_Label", "Filter to Graph"))
+					.ToolTipText(LOCTEXT("FilterToGraph_Tooltip", "Only Display events that are applicable to the current graph."))
+				]
+				.IsChecked<FBlueprintProfilerStatOptions, uint32>(this, &FBlueprintProfilerStatOptions::GetChecked, GraphFilter)
+				.OnCheckStateChanged<FBlueprintProfilerStatOptions, uint32>(this, &FBlueprintProfilerStatOptions::OnChecked, GraphFilter)
+			]
+		]
+		+SVerticalBox::Slot()
+		.HAlign(HAlign_Left)
+		.AutoHeight()
+		[
+			SNew(SHorizontalBox)
+			+SHorizontalBox::Slot()
+			.AutoWidth()
+			.Padding(FMargin(5,2))
+			[
+				SNew(SCheckBox)
+				.Content()
+				[
+					SNew(STextBlock)
+					.Text(LOCTEXT("AverageBlueprintStats_Label", "Use Averaged Blueprint Stats"))
+					.ToolTipText(LOCTEXT("AverageBlueprintStats_Tooltip", "Blueprint Statistics are Averaged or Summed By Instance."))
+				]
+				.IsChecked<FBlueprintProfilerStatOptions, uint32>(this, &FBlueprintProfilerStatOptions::GetChecked, AverageBlueprintStats)
+				.OnCheckStateChanged<FBlueprintProfilerStatOptions, uint32>(this, &FBlueprintProfilerStatOptions::OnChecked, AverageBlueprintStats)
+			]
+		]
+		+SVerticalBox::Slot()
+		.HAlign(HAlign_Left)
+		.AutoHeight()
+		[
+			SNew(SHorizontalBox)
+			+SHorizontalBox::Slot()
+			.AutoWidth()
+			.Padding(FMargin(5,2))
+			[
+				SNew(SCheckBox)
+				.Content()
+				[
+					SNew(STextBlock)
+					.Text(LOCTEXT("DisplayPureStats_Label", "Pure Timings"))
+					.ToolTipText(LOCTEXT("DisplayPureStats_Tooltip", "Display timings from pure nodes."))
+				]
+				.IsChecked<FBlueprintProfilerStatOptions, uint32>(this, &FBlueprintProfilerStatOptions::GetChecked, DisplayPure)
+				.OnCheckStateChanged<FBlueprintProfilerStatOptions, uint32>(this, &FBlueprintProfilerStatOptions::OnChecked, DisplayPure)
+			]
+		]
+		+SVerticalBox::Slot()
+		.HAlign(HAlign_Left)
+		.AutoHeight()
+		[
+			SNew(SHorizontalBox)
+			+SHorizontalBox::Slot()
+			.AutoWidth()
+			.Padding(FMargin(5,2))
+			[
+				SNew(SCheckBox)
+				.Content()
+				[
+					SNew(STextBlock)
+					.Text(LOCTEXT("ShowInstancesCheck_Label", "Show Instances"))
+					.ToolTipText(LOCTEXT("ShowInstancesCheck_Tooltip", "Display statistics by instance rather than by blueprint."))
+				]
+				.IsChecked<FBlueprintProfilerStatOptions, uint32>(this, &FBlueprintProfilerStatOptions::GetChecked, DisplayByInstance)
+				.OnCheckStateChanged<FBlueprintProfilerStatOptions, uint32>(this, &FBlueprintProfilerStatOptions::OnChecked, DisplayByInstance)
+			]
+		]
+		+SVerticalBox::Slot()
+		.HAlign(HAlign_Left)
+		.AutoHeight()
+		[
+			SNew(SHorizontalBox)
+			.Visibility<FBlueprintProfilerStatOptions, uint32>(this, &FBlueprintProfilerStatOptions::GetVisibility, ScopeToDebugInstance)
+			+SHorizontalBox::Slot()
+			.AutoWidth()
+			.Padding(FMargin(5,2))
+			[
+				SNew(SCheckBox)
+				.Content()
+				[
+					SNew(STextBlock)
+					.Text(LOCTEXT("InstanceFilterCheck_Label", "Debug Filter Scope"))
+					.ToolTipText(LOCTEXT("InstanceFilterCheck_Tooltip", "Restrict the displayed instances to match the blueprint debug filter."))
+				]
+				.IsChecked<FBlueprintProfilerStatOptions, uint32>(this, &FBlueprintProfilerStatOptions::GetChecked, ScopeToDebugInstance)
+				.OnCheckStateChanged<FBlueprintProfilerStatOptions, uint32>(this, &FBlueprintProfilerStatOptions::OnChecked, ScopeToDebugInstance)
+			]
+#if 0
+		]
+		+SVerticalBox::Slot()
+		.HAlign(HAlign_Left)
+		.AutoHeight()
+		[
+			SNew(SHorizontalBox)
+			+SHorizontalBox::Slot()
+			.AutoWidth()
+			.Padding(FMargin(5,2))
+			[
+				SNew(SCheckBox)
+				.Content()
+				[
+					SNew(STextBlock)
+					.Text(LOCTEXT("AutoItemExpansion_Label", "Auto Expand Statistics"))
+					.ToolTipText(LOCTEXT("AutoItemExpansion_Tooltip", "Auto expand the statistic trees."))
+				]
+				.IsChecked<FBlueprintProfilerStatOptions, uint32>(this, &FBlueprintProfilerStatOptions::GetChecked, AutoExpand)
+				.OnCheckStateChanged<FBlueprintProfilerStatOptions, uint32>(this, &FBlueprintProfilerStatOptions::OnChecked, AutoExpand)
+			]
+#endif
+		];
+}
+
+ECheckBoxState FBlueprintProfilerStatOptions::GetChecked(const uint32 FlagsIn) const
+{
+	ECheckBoxState CheckedState;
+	if (FlagsIn & ScopeToDebugInstance)
+	{
+		if (HasFlags(DisplayByInstance))
+		{
+			CheckedState = HasFlags(ScopeToDebugInstance) ? ECheckBoxState::Checked : ECheckBoxState::Unchecked;
+		}
+		else
+		{
+			CheckedState = ECheckBoxState::Undetermined;
+		}
+	}
+	else
+	{
+		CheckedState = HasAllFlags(FlagsIn) ? ECheckBoxState::Checked : ECheckBoxState::Unchecked;
+	}
+	return CheckedState;
+}
+
+void FBlueprintProfilerStatOptions::OnChecked(ECheckBoxState NewState, const uint32 FlagsIn)
+{
+	if (NewState == ECheckBoxState::Checked)
+	{
+		Flags |= FlagsIn;
+	}
+	else
+	{
+		Flags &= ~FlagsIn;
+	}
+	if (FlagsIn & AverageBlueprintStats)
+	{
+		FScriptPerfData::EnableBlueprintStatAverage(NewState == ECheckBoxState::Checked);
+	}
+	Flags |= Modified;
+}
+
+EVisibility FBlueprintProfilerStatOptions::GetVisibility(const uint32 FlagsIn) const
+{
+	EVisibility Result = EVisibility::Hidden;
+	if (FlagsIn & ScopeToDebugInstance)
+	{
+		Result = (Flags & DisplayByInstance) != 0 ? EVisibility::Visible : EVisibility::Collapsed;
+	}
+	return Result;
+}
+
+bool FBlueprintProfilerStatOptions::IsFiltered(TSharedPtr<FScriptExecutionNode> Node) const
+{
+	bool bFilteredOut = !HasFlags(EDisplayFlags::DisplayPure) && Node->HasFlags(EScriptExecutionNodeFlags::PureStats);
+	if (Node->IsEvent() && HasFlags(EDisplayFlags::GraphFilter))
+	{
+		if (Node->GetGraphName() == UEdGraphSchema_K2::FN_UserConstructionScript)
+		{
+			bFilteredOut = ActiveGraph != UEdGraphSchema_K2::FN_UserConstructionScript;
+		}
+		else
+		{
+			bFilteredOut = ActiveGraph == UEdGraphSchema_K2::FN_UserConstructionScript;
+		}
+	}
+	return bFilteredOut;
+}
+
+//////////////////////////////////////////////////////////////////////////
+// SBlueprintProfilerToolbar
+
+void SBlueprintProfilerToolbar::Construct(const FArguments& InArgs)
+{
+	CurrentViewType = InArgs._ProfileViewType;
+	DisplayOptions = InArgs._DisplayOptions;
+	ViewChangeDelegate = InArgs._OnViewChanged;
+
+	FMenuBuilder ViewComboContent(true, NULL);
+	ViewComboContent.AddMenuEntry(	LOCTEXT("OverviewViewType", "Profiler Overview"), 
+									LOCTEXT("OverviewViewTypeDesc", "Displays High Level Profiling Information"), 
+									FSlateIcon(), 
+									FExecuteAction::CreateSP(this, &SBlueprintProfilerToolbar::OnViewSelectionChanged, EBlueprintPerfViewType::Overview),
+									NAME_None, 
+									EUserInterfaceActionType::Button);
+	ViewComboContent.AddMenuEntry(	LOCTEXT("ExecutionGraphViewType", "Execution Graph"), 
+									LOCTEXT("ExecutionGraphViewTypeDesc", "Displays the Execution Graph with Statistics"), 
+									FSlateIcon(), 
+									FExecuteAction::CreateSP(this, &SBlueprintProfilerToolbar::OnViewSelectionChanged, EBlueprintPerfViewType::ExecutionGraph),
+									NAME_None, 
+									EUserInterfaceActionType::Button);
+	ViewComboContent.AddMenuEntry(	LOCTEXT("LeastPerformantViewType", "Least Performant List"), 
+									LOCTEXT("LeastPerformantViewTypeDesc", "Displays a list of Least Performant Areas of the Blueprint"), 
+									FSlateIcon(), 
+									FExecuteAction::CreateSP(this, &SBlueprintProfilerToolbar::OnViewSelectionChanged, EBlueprintPerfViewType::LeastPerformant),
+									NAME_None, 
+									EUserInterfaceActionType::Button);
+
+	FMenuBuilder HeatMapDisplayModeComboContent(true, nullptr);
+	HeatMapDisplayModeComboContent.AddMenuEntry(LOCTEXT("HeatMapDisplayMode_None", "Off"),
+		LOCTEXT("HeatMapDisplayMode_None_Desc", "Normal Display (No Heat Map)"),
+		FSlateIcon(),
+		FUIAction(
+			FExecuteAction::CreateSP(this, &SBlueprintProfilerToolbar::OnHeatMapDisplayModeChanged, EHeatmapControlId::NodeHeatmap, EBlueprintProfilerHeatMapDisplayMode::None),
+			FCanExecuteAction(),
+			FIsActionChecked::CreateSP(this, &SBlueprintProfilerToolbar::IsHeatMapDisplayModeSelected, EHeatmapControlId::NodeHeatmap, EBlueprintProfilerHeatMapDisplayMode::None)
+		),
+		NAME_None,
+		EUserInterfaceActionType::Check);
+	HeatMapDisplayModeComboContent.AddMenuEntry(LOCTEXT("HeatMapDisplayMode_Exclusive", "Exclusive"),
+		LOCTEXT("HeatMapDisplayMode_Exclusive_Desc", "Heat Map - Exclusive Timings"),
+		FSlateIcon(),
+		FUIAction(
+			FExecuteAction::CreateSP(this, &SBlueprintProfilerToolbar::OnHeatMapDisplayModeChanged, EHeatmapControlId::NodeHeatmap, EBlueprintProfilerHeatMapDisplayMode::Exclusive),
+			FCanExecuteAction(),
+			FIsActionChecked::CreateSP(this, &SBlueprintProfilerToolbar::IsHeatMapDisplayModeSelected, EHeatmapControlId::NodeHeatmap, EBlueprintProfilerHeatMapDisplayMode::Exclusive)
+		),
+		NAME_None,
+		EUserInterfaceActionType::Check);
+	HeatMapDisplayModeComboContent.AddMenuEntry(LOCTEXT("HeatMapDisplayMode_Inclusive", "Inclusive"),
+		LOCTEXT("HeatMapDisplayMode_Inclusive_Desc", "Heat Map - Inclusive Timings"),
+		FSlateIcon(),
+		FUIAction(
+			FExecuteAction::CreateSP(this, &SBlueprintProfilerToolbar::OnHeatMapDisplayModeChanged, EHeatmapControlId::NodeHeatmap, EBlueprintProfilerHeatMapDisplayMode::Inclusive),
+			FCanExecuteAction(),
+			FIsActionChecked::CreateSP(this, &SBlueprintProfilerToolbar::IsHeatMapDisplayModeSelected, EHeatmapControlId::NodeHeatmap, EBlueprintProfilerHeatMapDisplayMode::Inclusive)
+		),
+		NAME_None,
+		EUserInterfaceActionType::Check);
+	HeatMapDisplayModeComboContent.AddMenuEntry(LOCTEXT("HeatMapDisplayMode_MaxTiming", "Max Time"),
+		LOCTEXT("HeatMapDisplayMode_MaxTiming_Desc", "Heat Map - Max Time"),
+		FSlateIcon(),
+		FUIAction(
+			FExecuteAction::CreateSP(this, &SBlueprintProfilerToolbar::OnHeatMapDisplayModeChanged, EHeatmapControlId::NodeHeatmap, EBlueprintProfilerHeatMapDisplayMode::MaxTiming),
+			FCanExecuteAction(),
+			FIsActionChecked::CreateSP(this, &SBlueprintProfilerToolbar::IsHeatMapDisplayModeSelected, EHeatmapControlId::NodeHeatmap, EBlueprintProfilerHeatMapDisplayMode::MaxTiming)
+		),
+		NAME_None,
+		EUserInterfaceActionType::Check);
+	HeatMapDisplayModeComboContent.AddMenuEntry(LOCTEXT("HeatMapDisplayMode_TotalTiming", "Total Time"),
+		LOCTEXT("HeatMapDisplayMode_TotalTiming_Desc", "Heat Map - Max Time"),
+		FSlateIcon(),
+		FUIAction(
+			FExecuteAction::CreateSP(this, &SBlueprintProfilerToolbar::OnHeatMapDisplayModeChanged, EHeatmapControlId::NodeHeatmap, EBlueprintProfilerHeatMapDisplayMode::Total),
+			FCanExecuteAction(),
+			FIsActionChecked::CreateSP(this, &SBlueprintProfilerToolbar::IsHeatMapDisplayModeSelected, EHeatmapControlId::NodeHeatmap, EBlueprintProfilerHeatMapDisplayMode::Total)
+		),
+		NAME_None,
+		EUserInterfaceActionType::Check);
+
+	FMenuBuilder WireHeatMapDisplayModeComboContent(true, nullptr);
+	WireHeatMapDisplayModeComboContent.AddMenuEntry(LOCTEXT("WireHeatMapDisplayMode_None", "Off"),
+		LOCTEXT("HeatMapDisplayMode_None_Desc", "Normal Display (No Heat Map)"),
+		FSlateIcon(),
+		FUIAction(
+			FExecuteAction::CreateSP(this, &SBlueprintProfilerToolbar::OnHeatMapDisplayModeChanged, EHeatmapControlId::WireHeatmap, EBlueprintProfilerHeatMapDisplayMode::None),
+			FCanExecuteAction(),
+			FIsActionChecked::CreateSP(this, &SBlueprintProfilerToolbar::IsHeatMapDisplayModeSelected, EHeatmapControlId::WireHeatmap, EBlueprintProfilerHeatMapDisplayMode::None)
+		),
+		NAME_None,
+		EUserInterfaceActionType::Check);
+	WireHeatMapDisplayModeComboContent.AddMenuEntry(LOCTEXT("WireHeatMapDisplayMode_HottestPath", "Hottest Path"),
+		LOCTEXT("WireHeatMapDisplayMode_HottestPath_Desc", "Heat Map - Hottest Path"),
+		FSlateIcon(),
+		FUIAction(
+			FExecuteAction::CreateSP(this, &SBlueprintProfilerToolbar::OnHeatMapDisplayModeChanged, EHeatmapControlId::WireHeatmap, EBlueprintProfilerHeatMapDisplayMode::HottestPath),
+			FCanExecuteAction(),
+			FIsActionChecked::CreateSP(this, &SBlueprintProfilerToolbar::IsHeatMapDisplayModeSelected, EHeatmapControlId::WireHeatmap, EBlueprintProfilerHeatMapDisplayMode::HottestPath)
+		),
+		NAME_None,
+		EUserInterfaceActionType::Check);
+	WireHeatMapDisplayModeComboContent.AddMenuEntry(LOCTEXT("WireHeatMapDisplayMode_HottestEndpoint", "Hottest Endpoints"),
+		LOCTEXT("WireHeatMapDisplayMode_HottestEndpoints_Desc", "Heat Map - Hottest Endpoints"),
+		FSlateIcon(),
+		FUIAction(
+			FExecuteAction::CreateSP(this, &SBlueprintProfilerToolbar::OnHeatMapDisplayModeChanged, EHeatmapControlId::WireHeatmap, EBlueprintProfilerHeatMapDisplayMode::HottestEndpoint),
+			FCanExecuteAction(),
+			FIsActionChecked::CreateSP(this, &SBlueprintProfilerToolbar::IsHeatMapDisplayModeSelected, EHeatmapControlId::WireHeatmap, EBlueprintProfilerHeatMapDisplayMode::HottestEndpoint)
+		),
+		NAME_None,
+		EUserInterfaceActionType::Check);
+	WireHeatMapDisplayModeComboContent.AddMenuEntry(LOCTEXT("WireHeatMapDisplayMode_PinToPin", "Pin To Pin Heatmap"),
+		LOCTEXT("WireHeatMapDisplayMode_PinToPin_Desc", "Heat Map - Pin To Pin"),
+		FSlateIcon(),
+		FUIAction(
+			FExecuteAction::CreateSP(this, &SBlueprintProfilerToolbar::OnHeatMapDisplayModeChanged, EHeatmapControlId::WireHeatmap, EBlueprintProfilerHeatMapDisplayMode::PinToPin),
+			FCanExecuteAction(),
+			FIsActionChecked::CreateSP(this, &SBlueprintProfilerToolbar::IsHeatMapDisplayModeSelected, EHeatmapControlId::WireHeatmap, EBlueprintProfilerHeatMapDisplayMode::PinToPin)
+		),
+		NAME_None,
+		EUserInterfaceActionType::Check);
+
+	ChildSlot
+	[
+		SNew(SBorder)
+		.BorderImage(FEditorStyle::GetBrush("BlueprintProfiler.ViewToolBar"))
+		.Padding(0)
+		.VAlign(VAlign_Center)
+		.HAlign(HAlign_Right)
+		[
+			SNew(SHorizontalBox)
+			+SHorizontalBox::Slot()
+			.HAlign(HAlign_Right)
+			.AutoWidth()
+			.Padding(FMargin(5, 0))
+			[
+				SAssignNew(GeneralOptionsComboButton, SComboButton)
+				.ForegroundColor(this, &SBlueprintProfilerToolbar::GetButtonForegroundColor, EHeatmapControlId::GeneralOptions)
+				.ToolTipText(LOCTEXT("GeneralOptions_Tooltip", "General Options"))
+				.ButtonStyle(FEditorStyle::Get(), "ToggleButton")
+				.ContentPadding(2)
+				.OnGetMenuContent(DisplayOptions.Get(), &FBlueprintProfilerStatOptions::CreateComboContent)
+				.ButtonContent()
+				[
+					SNew(STextBlock)
+					.Text(LOCTEXT("GeneralOptions_Label", "General Options"))
+					.ColorAndOpacity(FLinearColor::White)
+				]
+			]
+			+SHorizontalBox::Slot()
+			.Padding(FMargin(5, 0))
+			[
+				SNew(SSeparator)
+				.Orientation(Orient_Vertical)
+			]
+			+SHorizontalBox::Slot()
+			.HAlign(HAlign_Right)
+			.AutoWidth()
+			.Padding(FMargin(5, 0))
+			[
+				SAssignNew(HeatThresholdComboButton, SComboButton)
+				.ForegroundColor(this, &SBlueprintProfilerToolbar::GetButtonForegroundColor, EHeatmapControlId::HeatThresholds)
+				.ToolTipText(LOCTEXT("HeatMapDisplayMode_Tooltip", "Heat Map Display Mode"))
+				.ButtonStyle(FEditorStyle::Get(), "ToggleButton")
+				.ContentPadding(2)
+				.OnGetMenuContent(this, &SBlueprintProfilerToolbar::CreateHeatThresholdsMenuContent)
+				.ButtonContent()
+				[
+					SNew(STextBlock)
+					.Text(LOCTEXT("HeatmapThresholdsCombo_Label", "Heatmap Thresholds"))
+					.ColorAndOpacity(FLinearColor::White)
+				]
+			]
+			+SHorizontalBox::Slot()
+			.Padding(FMargin(5, 0))
+			[
+				SNew(SSeparator)
+				.Orientation(Orient_Vertical)
+			]
+			+SHorizontalBox::Slot()
+			.AutoWidth()
+			.Padding(FMargin(5, 0))
+			[
+				SAssignNew(HeatMapDisplayModeComboButton, SComboButton)
+				.ForegroundColor(this, &SBlueprintProfilerToolbar::GetButtonForegroundColor, EHeatmapControlId::NodeHeatmap)
+				.ToolTipText(LOCTEXT("HeatMapDisplayMode_Tooltip", "Heat Map Display Mode"))
+				.ButtonStyle(FEditorStyle::Get(), "ToggleButton")
+				.ContentPadding(2)
+				.MenuContent()
+				[
+					HeatMapDisplayModeComboContent.MakeWidget()
+				]
+				.ButtonContent()
+				[
+					CreateHeatMapDisplayModeButton(EHeatmapControlId::NodeHeatmap)
+				]
+			]
+			+SHorizontalBox::Slot()
+			.Padding(FMargin(5, 0))
+			[
+				SNew(SSeparator)
+				.Orientation(Orient_Vertical)
+			]
+			+SHorizontalBox::Slot()
+			.AutoWidth()
+			.Padding(FMargin(5, 0))
+			[
+				SAssignNew(WireHeatMapDisplayModeComboButton, SComboButton)
+				.ForegroundColor(this, &SBlueprintProfilerToolbar::GetButtonForegroundColor, EHeatmapControlId::WireHeatmap)
+				.ToolTipText(LOCTEXT("WireHeatMapDisplayMode_Tooltip", "Wire Heat Map Display Mode"))
+				.ButtonStyle(FEditorStyle::Get(), "ToggleButton")
+				.ContentPadding(2)
+				.MenuContent()
+				[
+					WireHeatMapDisplayModeComboContent.MakeWidget()
+				]
+				.ButtonContent()
+				[
+					CreateHeatMapDisplayModeButton(EHeatmapControlId::WireHeatmap)
+				]
+			]
+#if 0
+			+SHorizontalBox::Slot()
+			.Padding(FMargin(5, 0))
+			[
+				SNew(SSeparator)
+				.Orientation(Orient_Vertical)
+			]
+			+SHorizontalBox::Slot()
+			.AutoWidth()
+			.Padding(FMargin(5,0))
+			[
+				SAssignNew(ViewComboButton, SComboButton)
+				.ForegroundColor(this, &SBlueprintProfilerToolbar::GetButtonForegroundColor, EHeatmapControlId::ViewType)
+				.ToolTipText(LOCTEXT("ViewType", "View Type"))
+				.ButtonStyle(FEditorStyle::Get(), "ToggleButton")
+				.ContentPadding(2)
+				.MenuContent()
+				[
+					ViewComboContent.MakeWidget()
+				]
+				.ButtonContent()
+				[
+					CreateViewButton()
+				]
+			]
+#endif
+		]
+	];
+}
+
+void SBlueprintProfilerToolbar::OnViewSelectionChanged(const EBlueprintPerfViewType::Type NewViewType)
+{
+	if (CurrentViewType != NewViewType)
+	{
+		CurrentViewType = NewViewType;
+		ViewChangeDelegate.ExecuteIfBound();
+	}
+}
+
+TSharedRef<SWidget> SBlueprintProfilerToolbar::CreateViewButton() const
+{
+	return SNew(SHorizontalBox)
+		+SHorizontalBox::Slot()
+		.AutoWidth()
+		.VAlign(VAlign_Top)
+		.Padding(0)
+		[
+			SNew(SImage)
+			.Image( FEditorStyle::GetBrush("GenericViewButton") )
+		]
+		+SHorizontalBox::Slot()
+		.AutoWidth()
+		.Padding(2,0)
+		.VAlign(VAlign_Center)
+		[
+			SNew(STextBlock)
+			.Text(LOCTEXT("ProfilerViewLabel", "Current View: "))
+			.ColorAndOpacity(FLinearColor::White)
+		]
+		+SHorizontalBox::Slot()
+		.AutoWidth()
+		.Padding(2,0)
+		.VAlign(VAlign_Center)
+		[
+			SNew(STextBlock)
+			.Text(GetCurrentViewText())
+			.ColorAndOpacity(FLinearColor(1.f, 1.f, 1.f, 0.4f))
+		];
+}
+
+FText SBlueprintProfilerToolbar::GetCurrentViewText() const
+{
+	switch(CurrentViewType)
+	{
+		case EBlueprintPerfViewType::Overview:			return LOCTEXT("OverviewLabel", "Profiler Overview");
+		case EBlueprintPerfViewType::ExecutionGraph:	return LOCTEXT("ExecutionGraphLabel", "Execution Graph");
+		case EBlueprintPerfViewType::LeastPerformant:	return LOCTEXT("LeastPerformantLabel", "Least Performant");
+	}
+	return LOCTEXT("UnknownViewLabel", "Unknown ViewType");
+}
+
+FText SBlueprintProfilerToolbar::GetCurrentHeatMapDisplayModeText(const EHeatmapControlId::Type ControlId) const
+{
+	FText Result = LOCTEXT("HeatMapDisplayModeLabel_Unknown", "<unknown>");
+	IBlueprintProfilerInterface& ProfilerModule = FModuleManager::LoadModuleChecked<IBlueprintProfilerInterface>("BlueprintProfiler");
+	EBlueprintProfilerHeatMapDisplayMode::Type DisplayMode = ControlId == EHeatmapControlId::NodeHeatmap ? ProfilerModule.GetGraphNodeHeatMapDisplayMode() : ProfilerModule.GetWireHeatMapDisplayMode();
+	switch (DisplayMode)
+	{
+		case EBlueprintProfilerHeatMapDisplayMode::None:			Result = LOCTEXT("HeatMapDisplayModeLabel_None", "Off"); break;
+		case EBlueprintProfilerHeatMapDisplayMode::Exclusive:		Result = LOCTEXT("HeatMapDisplayModeLabel_Exclusive", "Exclusive"); break;
+		case EBlueprintProfilerHeatMapDisplayMode::Inclusive:		Result = LOCTEXT("HeatMapDisplayModeLabel_Inclusive", "Inclusive"); break;
+		case EBlueprintProfilerHeatMapDisplayMode::MaxTiming:		Result = LOCTEXT("HeatMapDisplayModeLabel_MaxTiming", "Max Timing"); break;
+		case EBlueprintProfilerHeatMapDisplayMode::Total:			Result = LOCTEXT("HeatMapDisplayModeLabel_TotalTiming", "Total Timing"); break;
+		case EBlueprintProfilerHeatMapDisplayMode::HottestPath:		Result = LOCTEXT("HeatMapDisplayModeLabel_HottestPath", "Hottest Path"); break;
+		case EBlueprintProfilerHeatMapDisplayMode::HottestEndpoint:	Result = LOCTEXT("HeatMapDisplayModeLabel_HottestEndpoint", "Hottest Endpoints"); break;
+		case EBlueprintProfilerHeatMapDisplayMode::PinToPin:		Result = LOCTEXT("HeatMapDisplayModeLabel_PinToPin", "Pin to Pin Heatmap"); break;
+	}
+	return Result;
+}
+
+TSharedRef<SWidget> SBlueprintProfilerToolbar::CreateHeatMapDisplayModeButton(const EHeatmapControlId::Type ControlId) const
+{
+	return SNew(SHorizontalBox)
+		+SHorizontalBox::Slot()
+		.VAlign(VAlign_Center)
+		.HAlign(HAlign_Left)
+		.Padding(0,0)
+		.AutoWidth()
+		[
+			SNew(STextBlock)
+			.Text(ControlId == EHeatmapControlId::NodeHeatmap ? LOCTEXT("NodeHeatMapLabel_None", "Node Heatmap: ") : LOCTEXT("WireHeatMapLabel_None", "Wire Heatmap: "))
+			.ColorAndOpacity(FLinearColor::White)
+		]
+		+SHorizontalBox::Slot()
+		.VAlign(VAlign_Center)
+		.HAlign(HAlign_Left)
+		.Padding(2,0)
+		.AutoWidth()
+		[
+			SNew(STextBlock)
+			.Text(this, &SBlueprintProfilerToolbar::GetCurrentHeatMapDisplayModeText, ControlId)
+			.ColorAndOpacity(FLinearColor(1.f, 1.f, 1.f, 0.4f))
+		];
+}
+
+FSlateColor SBlueprintProfilerToolbar::GetButtonForegroundColor(const EHeatmapControlId::Type ControlId) const
+{
+	static const FName InvertedForegroundName("InvertedForeground");
+	static const FName DefaultForegroundName("DefaultForeground");
+	bool bButtonHovered = false;
+	switch (ControlId)
+	{
+		case EHeatmapControlId::ViewType:			bButtonHovered = ViewComboButton->IsHovered(); break;
+		case EHeatmapControlId::GeneralOptions:		bButtonHovered = GeneralOptionsComboButton->IsHovered(); break;
+		case EHeatmapControlId::NodeHeatmap:		bButtonHovered = HeatMapDisplayModeComboButton->IsHovered(); break;
+		case EHeatmapControlId::WireHeatmap:		bButtonHovered = WireHeatMapDisplayModeComboButton->IsHovered(); break;
+		case EHeatmapControlId::HeatThresholds:		bButtonHovered = HeatThresholdComboButton->IsHovered();	break;
+	}
+	return bButtonHovered ? FEditorStyle::GetSlateColor(InvertedForegroundName) : FEditorStyle::GetSlateColor(DefaultForegroundName);
+}
+
+bool SBlueprintProfilerToolbar::IsHeatMapDisplayModeSelected(const EHeatmapControlId::Type ControlId, const EBlueprintProfilerHeatMapDisplayMode::Type InHeatMapDisplayMode) const
+{
+	IBlueprintProfilerInterface& ProfilerModule = FModuleManager::LoadModuleChecked<IBlueprintProfilerInterface>("BlueprintProfiler");
+	return ControlId == EHeatmapControlId::NodeHeatmap ? (ProfilerModule.GetGraphNodeHeatMapDisplayMode() == InHeatMapDisplayMode) :
+														 (ProfilerModule.GetWireHeatMapDisplayMode() == InHeatMapDisplayMode);
+}
+
+void SBlueprintProfilerToolbar::OnHeatMapDisplayModeChanged(const EHeatmapControlId::Type ControlId, const EBlueprintProfilerHeatMapDisplayMode::Type NewHeatMapDisplayMode)
+{
+	IBlueprintProfilerInterface& ProfilerModule = FModuleManager::LoadModuleChecked<IBlueprintProfilerInterface>("BlueprintProfiler");
+	if (ControlId == EHeatmapControlId::NodeHeatmap)
+	{
+		if (ProfilerModule.GetGraphNodeHeatMapDisplayMode() != NewHeatMapDisplayMode)
+		{
+			ProfilerModule.SetGraphNodeHeatMapDisplayMode(NewHeatMapDisplayMode);
+		}
+	}
+	else 
+	{
+		if (ProfilerModule.GetWireHeatMapDisplayMode() != NewHeatMapDisplayMode)
+		{
+			ProfilerModule.SetWireHeatMapDisplayMode(NewHeatMapDisplayMode);
+		}
+	}
+}
+
+TSharedRef<SWidget> SBlueprintProfilerToolbar::CreateHeatThresholdsMenuContent()
+{
+	return SNew(SVerticalBox)
+		+SVerticalBox::Slot()
+		.AutoHeight()
+		.VAlign(VAlign_Center)
+		.Padding(FMargin(0,2))
+		[
+			SNew(SHorizontalBox)
+			+SHorizontalBox::Slot()
+			.AutoWidth()
+			.VAlign(VAlign_Center)
+			.Padding(FMargin(4,1,4,1))
+			[
+				SNew(STextBlock)
+				.Text(LOCTEXT("EventPerformanceThreshold_Label", "Event Heat Threshold"))
+				.ToolTipText(LOCTEXT("EventPerformanceThreshold_Tooltip", "The Hot Threshold for Overall Event Timings"))
+			]
+			+SHorizontalBox::Slot()
+			.HAlign(HAlign_Right)
+			.Padding(FMargin(4,1,1,1))
+			[
+				SNew(SHorizontalBox)
+				+SHorizontalBox::Slot()
+				.AutoWidth()
+				.HAlign(HAlign_Right)
+				[
+					SNew(SNumericEntryBox<float>)
+					.AllowSpin(true)
+					.MinValue(0.01f)
+					.MaxValue(5.f)
+					.MinSliderValue(0.01f)
+					.MaxSliderValue(5.f)
+					.Value(this, &SBlueprintProfilerToolbar::GetEventThreshold)
+					.OnValueChanged(this, &SBlueprintProfilerToolbar::SetEventThreshold)
+					.MinDesiredValueWidth(100)
+					.ToolTipText(LOCTEXT("EventPerformanceThreshold_Tooltip", "The Hot Threshold for Overall Event Timings"))
+				]
+				+SHorizontalBox::Slot()
+				.AutoWidth()
+				.HAlign(HAlign_Right)
+				[
+					SNew(SButton)
+					.ToolTipText(LOCTEXT("ResetToDefaultToolTip", "Reset to Default"))
+					.ButtonStyle(FEditorStyle::Get(), "NoBorder")
+					.ContentPadding(4)
+					.Visibility(this, &SBlueprintProfilerToolbar::IsEventDefaultButtonVisible)
+					.OnClicked(this, &SBlueprintProfilerToolbar::ResetEventThreshold)
+					.Content()
+					[
+						SNew(SImage)
+						.Image(FEditorStyle::GetBrush("PropertyWindow.DiffersFromDefault"))
+					]
+				]
+			]
+		]
+		+SVerticalBox::Slot()
+		.AutoHeight()
+		.VAlign(VAlign_Center)
+		.Padding(FMargin(0,2))
+		[
+			SNew(SHorizontalBox)
+			+SHorizontalBox::Slot()
+			.AutoWidth()
+			.VAlign(VAlign_Center)
+			.Padding(FMargin(4,1,4,1))
+			[
+				SNew(STextBlock)
+				.Text(LOCTEXT("NodeAveragePerformanceThreshold_Label", "Node Average Heat Threshold"))
+				.ToolTipText(LOCTEXT("NodeAveragePerformanceThreshold_Tooltip", "The Hot Threshold for Node Average Timings"))
+			]
+			+SHorizontalBox::Slot()
+			.HAlign(HAlign_Right)
+			.Padding(FMargin(4,1,1,1))
+			[
+				SNew(SHorizontalBox)
+				+SHorizontalBox::Slot()
+				.AutoWidth()
+				.HAlign(HAlign_Right)
+				[
+					SNew(SNumericEntryBox<float>)
+					.AllowSpin(true)
+					.MinValue(0.001f)
+					.MaxValue(0.5f)
+					.MinSliderValue(0.001f)
+					.MaxSliderValue(0.5f)
+					.Value(this, &SBlueprintProfilerToolbar::GetNodeExclusiveThreshold)
+					.OnValueChanged(this, &SBlueprintProfilerToolbar::SetNodeExclusiveThreshold)
+					.MinDesiredValueWidth(100)
+					.ToolTipText(LOCTEXT("NodeAveragePerformanceThreshold_Tooltip", "The Hot Threshold for Node Average Timings"))
+				]
+				+SHorizontalBox::Slot()
+				.AutoWidth()
+				.HAlign(HAlign_Right)
+				[
+					SNew(SButton)
+					.ToolTipText(LOCTEXT("ResetToDefaultToolTip", "Reset to Default"))
+					.ButtonStyle(FEditorStyle::Get(), "NoBorder")
+					.ContentPadding(4)
+					.Visibility(this, &SBlueprintProfilerToolbar::IsNodeExclDefaultButtonVisible)
+					.OnClicked(this, &SBlueprintProfilerToolbar::ResetNodeExclThreshold)
+					.Content()
+					[
+						SNew(SImage)
+						.Image(FEditorStyle::GetBrush("PropertyWindow.DiffersFromDefault"))
+					]
+				]
+			]
+		]
+		+SVerticalBox::Slot()
+		.AutoHeight()
+		.VAlign(VAlign_Center)
+		.Padding(FMargin(0,2))
+		[
+			SNew(SHorizontalBox)
+			+SHorizontalBox::Slot()
+			.AutoWidth()
+			.VAlign(VAlign_Center)
+			.Padding(FMargin(4,1,4,1))
+			[
+				SNew(STextBlock)
+				.Text(LOCTEXT("NodeInclusivePerformanceThreshold_Label", "Node Inclusive Heat Threshold"))
+				.ToolTipText(LOCTEXT("NodeInclusivePerformanceThreshold_Tooltip", "The Hot Threshold for Node Inclusive Timings"))
+			]
+			+SHorizontalBox::Slot()
+			.HAlign(HAlign_Right)
+			.Padding(1)
+			[
+				SNew(SHorizontalBox)
+				+SHorizontalBox::Slot()
+				.AutoWidth()
+				.HAlign(HAlign_Right)
+				[
+					SNew(SNumericEntryBox<float>)
+					.AllowSpin(true)
+					.MinValue(0.001f)
+					.MaxValue(0.5f)
+					.MinSliderValue(0.001f)
+					.MaxSliderValue(0.5f)
+					.Value(this, &SBlueprintProfilerToolbar::GetNodeInclusiveThreshold)
+					.OnValueChanged(this, &SBlueprintProfilerToolbar::SetNodeInclusiveThreshold)
+					.MinDesiredValueWidth(100)
+					.ToolTipText(LOCTEXT("NodeInclusivePerformanceThreshold_Tooltip", "The Hot Threshold for Node Inclusive Timings"))
+				]
+				+SHorizontalBox::Slot()
+				.AutoWidth()
+				.HAlign(HAlign_Right)
+				[
+					SNew(SButton)
+					.ToolTipText(LOCTEXT("ResetToDefaultToolTip", "Reset to Default"))
+					.ButtonStyle(FEditorStyle::Get(), "NoBorder")
+					.ContentPadding(4)
+					.Visibility(this, &SBlueprintProfilerToolbar::IsNodeInclDefaultButtonVisible)
+					.OnClicked(this, &SBlueprintProfilerToolbar::ResetNodeInclThreshold)
+					.Content()
+					[
+						SNew(SImage)
+						.Image(FEditorStyle::GetBrush("PropertyWindow.DiffersFromDefault"))
+					]
+				]
+			]
+		]
+		+SVerticalBox::Slot()
+		.AutoHeight()
+		.VAlign(VAlign_Center)
+		.Padding(FMargin(0,2))
+		[
+			SNew(SHorizontalBox)
+			+SHorizontalBox::Slot()
+			.AutoWidth()
+			.VAlign(VAlign_Center)
+			.Padding(FMargin(4,1,4,1))
+			[
+				SNew(STextBlock)
+				.Text(LOCTEXT("NodeMaxPerformanceThreshold_Label", "Node Max Heat Threshold"))
+				.ToolTipText(LOCTEXT("NodeMaxPerformanceThreshold_Tooltip", "The Hot Threshold for Max Node Timings"))
+			]
+			+SHorizontalBox::Slot()
+			.HAlign(HAlign_Right)
+			.Padding(FMargin(4,1,1,1))
+			[
+				SNew(SHorizontalBox)
+				+SHorizontalBox::Slot()
+				.AutoWidth()
+				.HAlign(HAlign_Right)
+				[
+					SNew(SNumericEntryBox<float>)
+					.AllowSpin(true)
+					.MinValue(0.01f)
+					.MaxValue(2.f)
+					.MinSliderValue(0.01f)
+					.MaxSliderValue(2.f)
+					.Value(this, &SBlueprintProfilerToolbar::GetNodeMaxThreshold)
+					.OnValueChanged(this, &SBlueprintProfilerToolbar::SetNodeMaxThreshold)
+					.MinDesiredValueWidth(100)
+					.ToolTipText(LOCTEXT("NodeMaxPerformanceThreshold_Tooltip", "The Hot Threshold for Max Node Timings"))
+				]
+				+SHorizontalBox::Slot()
+				.AutoWidth()
+				.HAlign(HAlign_Right)
+				[
+					SNew(SButton)
+					.ToolTipText(LOCTEXT("ResetToDefaultToolTip", "Reset to Default"))
+					.ButtonStyle(FEditorStyle::Get(), "NoBorder")
+					.ContentPadding(4)
+					.Visibility(this, &SBlueprintProfilerToolbar::IsNodeMaxDefaultButtonVisible)
+					.OnClicked(this, &SBlueprintProfilerToolbar::ResetNodeMaxThreshold)
+					.Content()
+					[
+						SNew(SImage)
+						.Image(FEditorStyle::GetBrush("PropertyWindow.DiffersFromDefault"))
+					]
+				]
+			]
+		]
+		+SVerticalBox::Slot()
+		.AutoHeight()
+		.VAlign(VAlign_Center)
+		.Padding(FMargin(0,2))
+		[
+			SNew(SHorizontalBox)
+			+SHorizontalBox::Slot()
+			.AutoWidth()
+			.VAlign(VAlign_Center)
+			.Padding(FMargin(4,1,16,1))
+			[
+				SNew(STextBlock)
+				.Text(LOCTEXT("RecentSampleCheck_Label", "Enable Recent Sample Bias"))
+				.ToolTipText(LOCTEXT("RecentSampleCheck_Tooltip", "Enable a user controlled bias that controls the weighting of new samples against historical."))
+			]
+			+SHorizontalBox::Slot()
+			.AutoWidth()
+			.VAlign(VAlign_Center)
+			.Padding(FMargin(4,1,4,1))
+			[
+				SNew(SCheckBox)
+				.IsChecked(this, &SBlueprintProfilerToolbar::GetRecentSampleBiasChecked)
+				.OnCheckStateChanged(this, &SBlueprintProfilerToolbar::OnRecentSampleBiasChecked)
+			]
+		]
+		+SVerticalBox::Slot()
+		.AutoHeight()
+		.VAlign(VAlign_Center)
+		.Padding(FMargin(0,2))
+		[
+			SNew(SHorizontalBox)
+			.Visibility(this, &SBlueprintProfilerToolbar::GetRecentSampleBiasVisibility)
+			+SHorizontalBox::Slot()
+			.AutoWidth()
+			.VAlign(VAlign_Center)
+			.Padding(FMargin(4,1,4,1))
+			[
+				SNew(STextBlock)
+				.Text(LOCTEXT("RecentSampleValue_Label", "Recent Sample Bias"))
+				.ToolTipText(LOCTEXT("RecentSampleValue_Tooltip", "Shifts the weighting of more recent samples to control the effect of sampling spikes."))
+			]
+			+SHorizontalBox::Slot()
+			.HAlign(HAlign_Right)
+			.Padding(FMargin(4,1,1,1))
+			[
+				SNew(SHorizontalBox)
+				+SHorizontalBox::Slot()
+				.AutoWidth()
+				.HAlign(HAlign_Right)
+				[
+					SNew(SNumericEntryBox<float>)
+					.AllowSpin(true)
+					.MinValue(0.01f)
+					.MaxValue(1.f)
+					.MinSliderValue(0.01f)
+					.MaxSliderValue(1.f)
+					.Value(this, &SBlueprintProfilerToolbar::GetRecentSampleBiasValue)
+					.OnValueChanged(this, &SBlueprintProfilerToolbar::SetRecentSampleBiasValue)
+					.MinDesiredValueWidth(100)
+					.ToolTipText(LOCTEXT("RecentSampleValue_Tooltip", "Shifts the weighting of more recent samples to control the effect of sampling spikes."))
+				]
+				+SHorizontalBox::Slot()
+				.AutoWidth()
+				.HAlign(HAlign_Right)
+				[
+					SNew(SButton)
+					.ToolTipText(LOCTEXT("ResetToDefaultToolTip", "Reset to Default"))
+					.ButtonStyle(FEditorStyle::Get(), "NoBorder")
+					.ContentPadding(4)
+					.Visibility(this, &SBlueprintProfilerToolbar::IsRecentSampleBiasDefaultButtonVisible)
+					.OnClicked(this, &SBlueprintProfilerToolbar::ResetRecentSampleBias)
+					.Content()
+					[
+						SNew(SImage)
+						.Image(FEditorStyle::GetBrush("PropertyWindow.DiffersFromDefault"))
+					]
+				]
+			]
+		];
+}
+
+EVisibility SBlueprintProfilerToolbar::GetRecentSampleBiasVisibility() const
+{
+	EVisibility Result = EVisibility::Collapsed;
+	if (const UEditorExperimentalSettings* EditorSettings = GetDefault<UEditorExperimentalSettings>())
+	{
+		Result = EditorSettings->bEnableBlueprintProfilerRecentSampleBias ? EVisibility::Visible : EVisibility::Collapsed;
+	}
+	return Result;
+}
+
+ECheckBoxState SBlueprintProfilerToolbar::GetRecentSampleBiasChecked() const
+{
+	ECheckBoxState Result = ECheckBoxState::Unchecked;
+	if (const UEditorExperimentalSettings* EditorSettings = GetDefault<UEditorExperimentalSettings>())
+	{
+		Result = EditorSettings->bEnableBlueprintProfilerRecentSampleBias ? ECheckBoxState::Checked : ECheckBoxState::Unchecked;
+	}
+	return Result;
+}
+
+void SBlueprintProfilerToolbar::OnRecentSampleBiasChecked(ECheckBoxState NewState)
+{
+	if (UEditorExperimentalSettings* EditorSettings = GetMutableDefault<UEditorExperimentalSettings>())
+	{
+		EditorSettings->bEnableBlueprintProfilerRecentSampleBias = NewState == ECheckBoxState::Checked;
+	}
+}
+
+TOptional<float> SBlueprintProfilerToolbar::GetRecentSampleBiasValue() const
+{
+	TOptional<float> Result;
+	if (const UEditorExperimentalSettings* EditorSettings = GetDefault<UEditorExperimentalSettings>())
+	{
+		Result = EditorSettings->BlueprintProfilerRecentSampleBias;
+	}
+	return Result;
+}
+
+void SBlueprintProfilerToolbar::SetRecentSampleBiasValue(const float NewValue)
+{
+	if (UEditorExperimentalSettings* EditorSettings = GetMutableDefault<UEditorExperimentalSettings>())
+	{
+		EditorSettings->BlueprintProfilerRecentSampleBias = NewValue;
+		FScriptPerfData::SetRecentSampleBias(NewValue);
+	}
+}
+
+EVisibility SBlueprintProfilerToolbar::IsRecentSampleBiasDefaultButtonVisible() const
+{
+	EVisibility Result = EVisibility::Hidden;
+	if (UEditorExperimentalSettings* EditorSettings = GetMutableDefault<UEditorExperimentalSettings>())
+	{
+		Result = EditorSettings->BlueprintProfilerRecentSampleBias == 0.2f ? EVisibility::Hidden : EVisibility::Visible;
+	}
+	return Result;
+}
+
+FReply SBlueprintProfilerToolbar::ResetRecentSampleBias()
+{
+	if (UEditorExperimentalSettings* EditorSettings = GetMutableDefault<UEditorExperimentalSettings>())
+	{
+		EditorSettings->BlueprintProfilerRecentSampleBias = 0.2f;
+		FScriptPerfData::SetRecentSampleBias(0.2f);
+	}
+	return FReply::Handled();
+}
+
+TOptional<float> SBlueprintProfilerToolbar::GetEventThreshold() const
+{
+	TOptional<float> Result;
+	if (const UEditorExperimentalSettings* EditorSettings = GetDefault<UEditorExperimentalSettings>())
+	{
+		Result = EditorSettings->BlueprintProfilerEventThreshold;
+	}
+	return Result;
+}
+
+void SBlueprintProfilerToolbar::SetEventThreshold(const float NewValue)
+{
+	if (UEditorExperimentalSettings* EditorSettings = GetMutableDefault<UEditorExperimentalSettings>())
+	{
+		EditorSettings->BlueprintProfilerEventThreshold = NewValue;
+		FScriptPerfData::SetEventPerformanceThreshold(NewValue);
+	}
+}
+
+EVisibility SBlueprintProfilerToolbar::IsEventDefaultButtonVisible() const
+{
+	EVisibility Result = EVisibility::Hidden;
+	if (UEditorExperimentalSettings* EditorSettings = GetMutableDefault<UEditorExperimentalSettings>())
+	{
+		Result = EditorSettings->BlueprintProfilerEventThreshold == 1.f ? EVisibility::Hidden : EVisibility::Visible;
+	}
+	return Result;
+}
+
+FReply SBlueprintProfilerToolbar::ResetEventThreshold()
+{
+	if (UEditorExperimentalSettings* EditorSettings = GetMutableDefault<UEditorExperimentalSettings>())
+	{
+		EditorSettings->BlueprintProfilerEventThreshold = 1.f;
+		FScriptPerfData::SetEventPerformanceThreshold(1.f);
+	}
+	return FReply::Handled();
+}
+
+TOptional<float> SBlueprintProfilerToolbar::GetNodeExclusiveThreshold() const
+{
+	TOptional<float> Result;
+	if (const UEditorExperimentalSettings* EditorSettings = GetDefault<UEditorExperimentalSettings>())
+	{
+		Result = EditorSettings->BlueprintProfilerExclNodeThreshold;
+	}
+	return Result;
+}
+
+void SBlueprintProfilerToolbar::SetNodeExclusiveThreshold(const float NewValue)
+{
+	if (UEditorExperimentalSettings* EditorSettings = GetMutableDefault<UEditorExperimentalSettings>())
+	{
+		EditorSettings->BlueprintProfilerExclNodeThreshold = NewValue;
+		FScriptPerfData::SetExclusivePerformanceThreshold(NewValue);
+	}
+}
+
+EVisibility SBlueprintProfilerToolbar::IsNodeExclDefaultButtonVisible() const
+{
+	EVisibility Result = EVisibility::Hidden;
+	if (UEditorExperimentalSettings* EditorSettings = GetMutableDefault<UEditorExperimentalSettings>())
+	{
+		Result = EditorSettings->BlueprintProfilerExclNodeThreshold == 0.2f ? EVisibility::Hidden : EVisibility::Visible;
+	}
+	return Result;
+}
+
+FReply SBlueprintProfilerToolbar::ResetNodeExclThreshold()
+{
+	if (UEditorExperimentalSettings* EditorSettings = GetMutableDefault<UEditorExperimentalSettings>())
+	{
+		EditorSettings->BlueprintProfilerExclNodeThreshold = 0.2f;
+		FScriptPerfData::SetExclusivePerformanceThreshold(0.2f);
+	}
+	return FReply::Handled();
+}
+
+TOptional<float> SBlueprintProfilerToolbar::GetNodeInclusiveThreshold() const
+{
+	TOptional<float> Result;
+	if (const UEditorExperimentalSettings* EditorSettings = GetDefault<UEditorExperimentalSettings>())
+	{
+		Result = EditorSettings->BlueprintProfilerInclNodeThreshold;
+	}
+	return Result;
+}
+
+void SBlueprintProfilerToolbar::SetNodeInclusiveThreshold(const float NewValue)
+{
+	if (UEditorExperimentalSettings* EditorSettings = GetMutableDefault<UEditorExperimentalSettings>())
+	{
+		EditorSettings->BlueprintProfilerInclNodeThreshold = NewValue;
+		FScriptPerfData::SetInclusivePerformanceThreshold(NewValue);
+	}
+}
+
+EVisibility SBlueprintProfilerToolbar::IsNodeInclDefaultButtonVisible() const
+{
+	EVisibility Result = EVisibility::Hidden;
+	if (UEditorExperimentalSettings* EditorSettings = GetMutableDefault<UEditorExperimentalSettings>())
+	{
+		Result = EditorSettings->BlueprintProfilerInclNodeThreshold == 0.25f ? EVisibility::Hidden : EVisibility::Visible;
+	}
+	return Result;
+}
+
+FReply SBlueprintProfilerToolbar::ResetNodeInclThreshold()
+{
+	if (UEditorExperimentalSettings* EditorSettings = GetMutableDefault<UEditorExperimentalSettings>())
+	{
+		EditorSettings->BlueprintProfilerInclNodeThreshold = 0.25f;
+		FScriptPerfData::SetInclusivePerformanceThreshold(0.25f);
+	}
+	return FReply::Handled();
+}
+
+TOptional<float> SBlueprintProfilerToolbar::GetNodeMaxThreshold() const
+{
+	TOptional<float> Result;
+	if (const UEditorExperimentalSettings* EditorSettings = GetDefault<UEditorExperimentalSettings>())
+	{
+		Result = EditorSettings->BlueprintProfilerMaxNodeThreshold;
+	}
+	return Result;
+}
+
+void SBlueprintProfilerToolbar::SetNodeMaxThreshold(const float NewValue)
+{
+	if (UEditorExperimentalSettings* EditorSettings = GetMutableDefault<UEditorExperimentalSettings>())
+	{
+		EditorSettings->BlueprintProfilerMaxNodeThreshold = NewValue;
+		FScriptPerfData::SetMaxPerformanceThreshold(NewValue);
+	}
+}
+
+EVisibility SBlueprintProfilerToolbar::IsNodeMaxDefaultButtonVisible() const
+{
+	EVisibility Result = EVisibility::Hidden;
+	if (UEditorExperimentalSettings* EditorSettings = GetMutableDefault<UEditorExperimentalSettings>())
+	{
+		Result = EditorSettings->BlueprintProfilerMaxNodeThreshold == 0.5f ? EVisibility::Hidden : EVisibility::Visible;
+	}
+	return Result;
+}
+
+FReply SBlueprintProfilerToolbar::ResetNodeMaxThreshold()
+{
+	if (UEditorExperimentalSettings* EditorSettings = GetMutableDefault<UEditorExperimentalSettings>())
+	{
+		EditorSettings->BlueprintProfilerMaxNodeThreshold = 0.5f;
+		FScriptPerfData::SetMaxPerformanceThreshold(0.5f);
+	}
+	return FReply::Handled();
+}
+
+
+#undef LOCTEXT_NAMESPACE
diff --git a/Engine/Source/Editor/Kismet/Private/Profiler/SBlueprintProfilerToolbar.h b/Engine/Source/Editor/Kismet/Private/Profiler/SBlueprintProfilerToolbar.h
new file mode 100644
--- /dev/null
+++ b/Engine/Source/Editor/Kismet/Private/Profiler/SBlueprintProfilerToolbar.h
@@ -0,0 +1,232 @@
+// Copyright 1998-2016 Epic Games, Inc. All Rights Reserved.
+
+#pragma once
+
+#include "BlueprintProfilerModule.h"
+
+/** Blueprint performance view type */
+namespace EBlueprintPerfViewType
+{
+	enum Type
+	{
+		None = 0,
+		Overview,
+		ExecutionGraph,
+		LeastPerformant
+	};
+}
+
+/** Heatmap Control Id */
+namespace EHeatmapControlId
+{
+	enum Type
+	{
+		None = 0,
+		ViewType,
+		GeneralOptions,
+		NodeHeatmap,
+		WireHeatmap,
+		HeatThresholds
+	};
+}
+
+//////////////////////////////////////////////////////////////////////////
+// FBlueprintProfilerStatOptions
+
+struct FBlueprintProfilerStatOptions : public TSharedFromThis<FBlueprintProfilerStatOptions>
+{
+public:
+
+	/** Display flags controlling widget appearance */
+	enum EDisplayFlags
+	{
+		DisplayByInstance			= 0x00000001,	// Show instance stats
+		ScopeToDebugInstance		= 0x00000002,	// Scope stats to current debug instance
+		GraphFilter					= 0x00000004,	// Filter to current graph
+		DisplayPure					= 0x00000008,	// Display pure stats
+		AutoExpand					= 0x00000010,	// Auto expand stats
+		AverageBlueprintStats		= 0x00000020,	// Average blueprint stats ( Rather than sum )
+		Modified					= 0x10000000,	// Display state was changed
+		Default						= DisplayByInstance|DisplayPure|AverageBlueprintStats|Modified
+	};
+
+	FBlueprintProfilerStatOptions()
+		: ActiveInstance(NAME_None)
+		, ActiveGraph(NAME_None)
+		, Flags(Default)
+	{
+	}
+
+	/** Marks display state as modified */
+	void SetStateModified() { Flags |= Modified; }
+
+	/** Returns if the filter state options are modified */
+	bool IsStateModified() const { return (Flags & Modified) != 0U; }
+
+	/** Add flags */
+	void AddFlags(const uint32 FlagsIn) { Flags |= FlagsIn; }
+
+	/** Clear flags */
+	void ClearFlags(const uint32 FlagsIn) { Flags &= ~FlagsIn; }
+
+	/** Returns true if any flags match */
+	bool HasFlags(const uint32 FlagsIn) const { return (Flags & FlagsIn) != 0U; }
+
+	/** Returns true if all flags match */
+	bool HasAllFlags(const uint32 FlagsIn) const { return (Flags & FlagsIn) == FlagsIn; }
+
+	/** Set active instance name */
+	FName GetActiveInstance() const { return HasFlags(DisplayByInstance) ? ActiveInstance : NAME_None; }
+
+	/** Set active instance name */
+	void SetActiveInstance(const FName InstanceName);
+
+	/** Set active graph name */
+	void SetActiveGraph(const FName GraphName);
+
+	/** Create toolbar widget */
+	TSharedRef<SWidget> CreateToolbar();
+
+	/** Create combo content */
+	TSharedRef<SWidget> CreateComboContent();
+
+	/** Returns if the UI should be checked for flags */
+	ECheckBoxState GetChecked(const uint32 FlagsIn) const;
+
+	/** Handles UI check states */
+	void OnChecked(ECheckBoxState NewState, const uint32 FlagsIn);
+
+	/** Returns control visibility */
+	EVisibility GetVisibility(const uint32 FlagsIn) const;
+
+	/** Returns true if the exec node passes current filters */
+	bool IsFiltered(TSharedPtr<class FScriptExecutionNode> Node) const;
+
+private:
+
+	/** The active instance name */
+	FName ActiveInstance;
+	/** The active graph name */
+	FName ActiveGraph;
+	/** The display flags */
+	uint32 Flags;
+};
+
+//////////////////////////////////////////////////////////////////////////
+// SBlueprintProfilerToolbar
+
+class SBlueprintProfilerToolbar : public SCompoundWidget
+{
+public:
+		DECLARE_DELEGATE(FOnProfilerViewChanged);
+
+		SLATE_BEGIN_ARGS(SBlueprintProfilerToolbar)
+		: _ProfileViewType(EBlueprintPerfViewType::ExecutionGraph)
+		{}
+
+		SLATE_ARGUMENT(EBlueprintPerfViewType::Type, ProfileViewType)
+		SLATE_ARGUMENT(TSharedPtr<struct FBlueprintProfilerStatOptions>, DisplayOptions)
+		SLATE_EVENT(FOnProfilerViewChanged, OnViewChanged)
+	SLATE_END_ARGS()
+
+
+public:
+
+	void Construct(const FArguments& InArgs);
+
+	/** Returns the currently active view */
+	EBlueprintPerfViewType::Type GetActiveViewType() const { return CurrentViewType; }
+
+protected:
+
+	/** Constructs and returns the view button widget */
+	TSharedRef<SWidget> CreateViewButton() const;
+
+	/** Returns the current view label */
+	FText GetCurrentViewText() const;
+
+	/** Returns the active view button foreground color */
+	FSlateColor GetButtonForegroundColor(const EHeatmapControlId::Type ControlId) const;
+
+	/** Called when the profiler view type is changed */
+	void OnViewSelectionChanged(const EBlueprintPerfViewType::Type NewViewType);
+
+	/** Returns the current heat map display mode label */
+	FText GetCurrentHeatMapDisplayModeText(const EHeatmapControlId::Type ControlId) const;
+
+	/** Constructs and returns the heat map display mode button widget */
+	TSharedRef<SWidget> CreateHeatMapDisplayModeButton(const EHeatmapControlId::Type ControlId) const;
+
+	/** Returns the active heat map display mode button foreground color */
+	FSlateColor GetHeatMapDisplayModeButtonForegroundColor(const EHeatmapControlId::Type ControlId) const;
+
+	/** Returns whether or not the given heat map display mode is selected */
+	bool IsHeatMapDisplayModeSelected(const EHeatmapControlId::Type ControlId, const EBlueprintProfilerHeatMapDisplayMode::Type InHeatMapDisplayMode) const;
+
+	/** Called when the heat map display mode is changed */
+	void OnHeatMapDisplayModeChanged(const EHeatmapControlId::Type ControlId, const EBlueprintProfilerHeatMapDisplayMode::Type NewHeatMapDisplayMode);
+
+	/** Create the heat threshold combo menu content */
+	TSharedRef<SWidget> CreateHeatThresholdsMenuContent();
+
+private:
+
+	/** Recent sample bias helpers */
+	EVisibility GetRecentSampleBiasVisibility() const;
+	ECheckBoxState GetRecentSampleBiasChecked() const;
+	void OnRecentSampleBiasChecked(ECheckBoxState NewState);
+	TOptional<float> GetRecentSampleBiasValue() const;
+	void SetRecentSampleBiasValue(const float NewValue);
+	EVisibility IsRecentSampleBiasDefaultButtonVisible() const;
+	FReply ResetRecentSampleBias();
+
+	/** Get/Set event performance threshold helpers */
+	TOptional<float> GetEventThreshold() const;
+	void SetEventThreshold(const float NewValue);
+	EVisibility IsEventDefaultButtonVisible() const;
+	FReply ResetEventThreshold();
+
+	/** Get/Set node exclusive performance threshold helpers */
+	TOptional<float> GetNodeExclusiveThreshold() const;
+	void SetNodeExclusiveThreshold(const float NewValue);
+	EVisibility IsNodeExclDefaultButtonVisible() const;
+	FReply ResetNodeExclThreshold();
+
+	/** Get/Set node inclusive performance threshold helpers */
+	TOptional<float> GetNodeInclusiveThreshold() const;
+	void SetNodeInclusiveThreshold(const float NewValue);
+	EVisibility IsNodeInclDefaultButtonVisible() const;
+	FReply ResetNodeInclThreshold();
+
+	/** Get/Set node max performance threshold helpers */
+	TOptional<float> GetNodeMaxThreshold() const;
+	void SetNodeMaxThreshold(const float NewValue);
+	EVisibility IsNodeMaxDefaultButtonVisible() const;
+	FReply ResetNodeMaxThreshold();
+
+protected:
+
+	/** Current profiler view type */
+	EBlueprintPerfViewType::Type CurrentViewType;
+
+	/** Stat display options */
+	TSharedPtr<FBlueprintProfilerStatOptions> DisplayOptions;
+
+	/** Delegate to call when the view changes */
+	FOnProfilerViewChanged ViewChangeDelegate;
+
+	/** General options combo button widget */
+	TSharedPtr<SComboButton> GeneralOptionsComboButton;
+
+	/** View combo button widget */
+	TSharedPtr<SComboButton> ViewComboButton;
+
+	/** Heat threshold combo button widget */
+	TSharedPtr<SComboButton> HeatThresholdComboButton;
+
+	/** Heat map display mode combo button widget */
+	TSharedPtr<SComboButton> HeatMapDisplayModeComboButton;
+
+	/** Heat map display mode combo button widget */
+	TSharedPtr<SComboButton> WireHeatMapDisplayModeComboButton;
+};
