Copying //UE4/Dev-Blueprints to //UE4/Dev-Main (Source: //UE4/Dev-Blueprints @ 3025888)

diff --git a/Engine/Source/Editor/Kismet/Private/Profiler/SBlueprintProfilerView.cpp b/Engine/Source/Editor/Kismet/Private/Profiler/SBlueprintProfilerView.cpp
--- a/Engine/Source/Editor/Kismet/Private/Profiler/SBlueprintProfilerView.cpp
+++ b/Engine/Source/Editor/Kismet/Private/Profiler/SBlueprintProfilerView.cpp
@@ -2,12 +2,18 @@
 
 #include "BlueprintEditorPrivatePCH.h"
 #include "SBlueprintProfilerView.h"
+#include "SBlueprintProfilerToolbar.h"
+#include "BPProfilerStatisticWidgets.h"
+#include "BlueprintEditor.h"
 
 // Profiler View types
 #include "SProfilerOverview.h"
 #include "SLeastPerformantDisplay.h"
 #include "SGraphExecutionStatDisplay.h"
+
 #include "ScriptPerfData.h"
+#include "SNumericEntryBox.h"
+#include "Editor/UnrealEd/Classes/Settings/EditorExperimentalSettings.h"
 
 #define LOCTEXT_NAMESPACE "BlueprintProfilerView"
 
@@ -18,14 +24,21 @@ SBlueprintProfilerView::~SBlueprintProfilerView()
 {
 	// Register for Profiler toggle events
 	FBlueprintCoreDelegates::OnToggleScriptProfiler.RemoveAll(this);
+	// Remove delegate for graph structural changes
+	if (IBlueprintProfilerInterface* Profiler = FModuleManager::GetModulePtr<IBlueprintProfilerInterface>("BlueprintProfiler"))
+	{
+		Profiler->GetGraphLayoutChangedDelegate().RemoveAll(this);
+	}
 }
 
 void SBlueprintProfilerView::Construct(const FArguments& InArgs)
 {
-	CurrentViewType = InArgs._ProfileViewType;
 	BlueprintEditor = InArgs._AssetEditor;
 	// Register for Profiler toggle events
 	FBlueprintCoreDelegates::OnToggleScriptProfiler.AddSP(this, &SBlueprintProfilerView::OnToggleProfiler);
+	// Remove delegate for graph structural changes
+	IBlueprintProfilerInterface& ProfilerModule = FModuleManager::LoadModuleChecked<IBlueprintProfilerInterface>("BlueprintProfiler");
+	ProfilerModule.GetGraphLayoutChangedDelegate().AddSP(this, &SBlueprintProfilerView::OnGraphLayoutChanged);
 	// Initialise the number format.
 	FNumberFormattingOptions StatisticNumberFormat;
 	StatisticNumberFormat.MinimumFractionalDigits = 4;
@@ -44,361 +57,90 @@ void SBlueprintProfilerView::Construct(const FArguments& InArgs)
 	{
 		StatusText = LOCTEXT("DisabledProfilerText", "The Blueprint Profiler is experimental and is currently not enabled in the editor preferences");
 	}
+	// Create Display Options
+	if (!DisplayOptions.IsValid())
+	{
+		DisplayOptions = MakeShareable(new FBlueprintProfilerStatOptions);
+	}
+	// Create the toolbar
+	SAssignNew(ProfilerToolbar, SBlueprintProfilerToolbar)
+		.ProfileViewType(InArgs._ProfileViewType)
+		.DisplayOptions(DisplayOptions)
+		.OnViewChanged(this, &SBlueprintProfilerView::OnViewChanged);
 	// Create the profiler view widgets
 	UpdateActiveProfilerWidget();
 }
 
 void SBlueprintProfilerView::OnToggleProfiler(bool bEnabled)
 {
-	if (GetDefault<UEditorExperimentalSettings>()->bBlueprintPerformanceAnalysisTools)
-	{
-		StatusText = bEnabled ? LOCTEXT("ProfilerNoDataText", "The Blueprint Profiler is active but does not currently have any data to display") :
-								LOCTEXT("ProfilerInactiveText", "The Blueprint Profiler is currently Inactive");
-	}
-	else
-	{
-		StatusText = LOCTEXT("DisabledProfilerText", "The Blueprint Profiler is experimental and is currently not enabled in the editor preferences");
-	}
+	UpdateStatusMessage();
 	UpdateActiveProfilerWidget();
 }
 
-void SBlueprintProfilerView::UpdateActiveProfilerWidget()
-{
-		FMenuBuilder ViewComboContent(true, NULL);
-		ViewComboContent.AddMenuEntry(	LOCTEXT("OverviewViewType", "Profiler Overview"), 
-										LOCTEXT("OverviewViewTypeDesc", "Displays High Level Profiling Information"), 
-										FSlateIcon(), 
-										FExecuteAction::CreateSP(this, &SBlueprintProfilerView::OnViewSelectionChanged, EBlueprintPerfViewType::Overview),
-										NAME_None, 
-										EUserInterfaceActionType::Button);
-		ViewComboContent.AddMenuEntry(	LOCTEXT("ExecutionGraphViewType", "Execution Graph"), 
-										LOCTEXT("ExecutionGraphViewTypeDesc", "Displays the Execution Graph with Statistics"), 
-										FSlateIcon(), 
-										FExecuteAction::CreateSP(this, &SBlueprintProfilerView::OnViewSelectionChanged, EBlueprintPerfViewType::ExecutionGraph),
-										NAME_None, 
-										EUserInterfaceActionType::Button);
-		ViewComboContent.AddMenuEntry(	LOCTEXT("LeastPerformantViewType", "Least Performant List"), 
-										LOCTEXT("LeastPerformantViewTypeDesc", "Displays a list of Least Performant Areas of the Blueprint"), 
-										FSlateIcon(), 
-										FExecuteAction::CreateSP(this, &SBlueprintProfilerView::OnViewSelectionChanged, EBlueprintPerfViewType::LeastPerformant),
-										NAME_None, 
-										EUserInterfaceActionType::Button);
-
-		FMenuBuilder HeatMapDisplayModeComboContent(true, nullptr);
-		HeatMapDisplayModeComboContent.AddMenuEntry(LOCTEXT("HeatMapDisplayMode_None", "Off"),
-			LOCTEXT("HeatMapDisplayMode_None_Desc", "Normal Display (No Heat Map)"),
-			FSlateIcon(),
-			FUIAction(
-				FExecuteAction::CreateSP(this, &SBlueprintProfilerView::OnHeatMapDisplayModeChanged, EBlueprintProfilerHeatMapDisplayMode::None),
-				FCanExecuteAction(),
-				FIsActionChecked::CreateSP(this, &SBlueprintProfilerView::IsHeatMapDisplayModeSelected, EBlueprintProfilerHeatMapDisplayMode::None)
-			),
-			NAME_None,
-			EUserInterfaceActionType::Check);
-		HeatMapDisplayModeComboContent.AddMenuEntry(LOCTEXT("HeatMapDisplayMode_Exclusive", "Exclusive"),
-			LOCTEXT("HeatMapDisplayMode_Exclusive_Desc", "Heat Map - Exclusive Timings"),
-			FSlateIcon(),
-			FUIAction(
-				FExecuteAction::CreateSP(this, &SBlueprintProfilerView::OnHeatMapDisplayModeChanged, EBlueprintProfilerHeatMapDisplayMode::Exclusive),
-				FCanExecuteAction(),
-				FIsActionChecked::CreateSP(this, &SBlueprintProfilerView::IsHeatMapDisplayModeSelected, EBlueprintProfilerHeatMapDisplayMode::Exclusive)
-			),
-			NAME_None,
-			EUserInterfaceActionType::Check);
-		HeatMapDisplayModeComboContent.AddMenuEntry(LOCTEXT("HeatMapDisplayMode_Inclusive", "Inclusive"),
-			LOCTEXT("HeatMapDisplayMode_Inclusive_Desc", "Heat Map - Inclusive Timings"),
-			FSlateIcon(),
-			FUIAction(
-				FExecuteAction::CreateSP(this, &SBlueprintProfilerView::OnHeatMapDisplayModeChanged, EBlueprintProfilerHeatMapDisplayMode::Inclusive),
-				FCanExecuteAction(),
-				FIsActionChecked::CreateSP(this, &SBlueprintProfilerView::IsHeatMapDisplayModeSelected, EBlueprintProfilerHeatMapDisplayMode::Inclusive)
-			),
-			NAME_None,
-			EUserInterfaceActionType::Check);
-		HeatMapDisplayModeComboContent.AddMenuEntry(LOCTEXT("HeatMapDisplayMode_MaxTiming", "Max Time"),
-			LOCTEXT("HeatMapDisplayMode_MaxTiming_Desc", "Heat Map - Max Time"),
-			FSlateIcon(),
-			FUIAction(
-				FExecuteAction::CreateSP(this, &SBlueprintProfilerView::OnHeatMapDisplayModeChanged, EBlueprintProfilerHeatMapDisplayMode::MaxTiming),
-				FCanExecuteAction(),
-				FIsActionChecked::CreateSP(this, &SBlueprintProfilerView::IsHeatMapDisplayModeSelected, EBlueprintProfilerHeatMapDisplayMode::MaxTiming)
-			),
-			NAME_None,
-			EUserInterfaceActionType::Check);
-
-		FMenuBuilder WireHeatMapDisplayModeComboContent(true, nullptr);
-		WireHeatMapDisplayModeComboContent.AddMenuEntry(LOCTEXT("WireHeatMapDisplayMode_None", "Off"),
-			LOCTEXT("HeatMapDisplayMode_None_Desc", "Normal Display (No Heat Map)"),
-			FSlateIcon(),
-			FUIAction(
-				FExecuteAction::CreateSP(this, &SBlueprintProfilerView::OnWireHeatMapDisplayModeChanged, EBlueprintProfilerHeatMapDisplayMode::None),
-				FCanExecuteAction(),
-				FIsActionChecked::CreateSP(this, &SBlueprintProfilerView::IsWireHeatMapDisplayModeSelected, EBlueprintProfilerHeatMapDisplayMode::None)
-			),
-			NAME_None,
-			EUserInterfaceActionType::Check);
-		WireHeatMapDisplayModeComboContent.AddMenuEntry(LOCTEXT("HeatMapDisplayMode_Exclusive", "Exclusive"),
-			LOCTEXT("HeatMapDisplayMode_Exclusive_Desc", "Heat Map - Exclusive Timings"),
-			FSlateIcon(),
-			FUIAction(
-				FExecuteAction::CreateSP(this, &SBlueprintProfilerView::OnWireHeatMapDisplayModeChanged, EBlueprintProfilerHeatMapDisplayMode::Exclusive),
-				FCanExecuteAction(),
-				FIsActionChecked::CreateSP(this, &SBlueprintProfilerView::IsWireHeatMapDisplayModeSelected, EBlueprintProfilerHeatMapDisplayMode::Exclusive)
-			),
-			NAME_None,
-			EUserInterfaceActionType::Check);
-		WireHeatMapDisplayModeComboContent.AddMenuEntry(LOCTEXT("HeatMapDisplayMode_Inclusive", "Inclusive"),
-			LOCTEXT("HeatMapDisplayMode_Inclusive_Desc", "Heat Map - Inclusive Timings"),
-			FSlateIcon(),
-			FUIAction(
-				FExecuteAction::CreateSP(this, &SBlueprintProfilerView::OnWireHeatMapDisplayModeChanged, EBlueprintProfilerHeatMapDisplayMode::Inclusive),
-				FCanExecuteAction(),
-				FIsActionChecked::CreateSP(this, &SBlueprintProfilerView::IsWireHeatMapDisplayModeSelected, EBlueprintProfilerHeatMapDisplayMode::Inclusive)
-			),
-			NAME_None,
-			EUserInterfaceActionType::Check);
-		WireHeatMapDisplayModeComboContent.AddMenuEntry(LOCTEXT("HeatMapDisplayMode_MaxTiming", "Max Time"),
-			LOCTEXT("HeatMapDisplayMode_MaxTiming_Desc", "Heat Map - Max Time"),
-			FSlateIcon(),
-			FUIAction(
-				FExecuteAction::CreateSP(this, &SBlueprintProfilerView::OnWireHeatMapDisplayModeChanged, EBlueprintProfilerHeatMapDisplayMode::MaxTiming),
-				FCanExecuteAction(),
-				FIsActionChecked::CreateSP(this, &SBlueprintProfilerView::IsWireHeatMapDisplayModeSelected, EBlueprintProfilerHeatMapDisplayMode::MaxTiming)
-			),
-			NAME_None,
-			EUserInterfaceActionType::Check);
-
-		ChildSlot
-		[
-			SNew(SVerticalBox)
-			+SVerticalBox::Slot()
-			.AutoHeight()
-			.Padding(0)
-			[
-				SNew(SBorder)
-				.BorderImage(FEditorStyle::GetBrush("BlueprintProfiler.ViewToolBar"))
-				.Padding(0)
-				.VAlign(VAlign_Top)
-				.HAlign(HAlign_Right)
-				[
-					SNew(SHorizontalBox)
-					+SHorizontalBox::Slot()
-					.AutoWidth()
-					.Padding(FMargin(5, 0))
-					[
-						SAssignNew(HeatMapDisplayModeComboButton, SComboButton)
-						.ForegroundColor(this, &SBlueprintProfilerView::GetHeatMapDisplayModeButtonForegroundColor)
-						.ToolTipText(LOCTEXT("BlueprintProfilerHeatMapDisplayMode_Tooltip", "Heat Map Display Mode"))
-						.ButtonStyle(FEditorStyle::Get(), "ToggleButton")
-						.ContentPadding(2)
-						.MenuContent()
-						[
-							HeatMapDisplayModeComboContent.MakeWidget()
-						]
-						.ButtonContent()
-						[
-							CreateHeatMapDisplayModeButton()
-						]
-					]
-					+SHorizontalBox::Slot()
-					.Padding(FMargin(5, 0))
-					[
-						SNew(SSeparator)
-						.Orientation(Orient_Vertical)
-					]
-					+SHorizontalBox::Slot()
-					.AutoWidth()
-					.Padding(FMargin(5, 0))
-					[
-						SAssignNew(WireHeatMapDisplayModeComboButton, SComboButton)
-						.ForegroundColor(this, &SBlueprintProfilerView::GetWireHeatMapDisplayModeButtonForegroundColor)
-						.ToolTipText(LOCTEXT("BlueprintProfilerWireHeatMapDisplayMode_Tooltip", "Wire Heat Map Display Mode"))
-						.ButtonStyle(FEditorStyle::Get(), "ToggleButton")
-						.ContentPadding(2)
-						.MenuContent()
-						[
-							WireHeatMapDisplayModeComboContent.MakeWidget()
-						]
-						.ButtonContent()
-						[
-							CreateWireHeatMapDisplayModeButton()
-						]
-					]
-					+SHorizontalBox::Slot()
-					.Padding(FMargin(5, 0))
-					[
-						SNew(SSeparator)
-						.Orientation(Orient_Vertical)
-					]
-					+SHorizontalBox::Slot()
-					.AutoWidth()
-					.Padding(FMargin(5,0))
-					[
-						SAssignNew(ViewComboButton, SComboButton)
-						.ForegroundColor(this, &SBlueprintProfilerView::GetViewButtonForegroundColor)
-						.ToolTipText(LOCTEXT("BlueprintProfilerViewType", "View Type"))
-						.ButtonStyle(FEditorStyle::Get(), "ToggleButton")
-						.ContentPadding(2)
-						.MenuContent()
-						[
-							ViewComboContent.MakeWidget()
-						]
-						.ButtonContent()
-						[
-							CreateViewButton()
-						]
-					]
-				]
-			]
-			+SVerticalBox::Slot()
-			[
-				SNew(SBorder)
-				.Padding(FMargin(0,2,0,0))
-				.BorderImage(FEditorStyle::GetBrush("NoBorder"))
-				[
-					CreateActiveStatisticWidget()
-				]
-			]
-		];
-	}
-
-FText SBlueprintProfilerView::GetCurrentViewText() const
-{
-	switch(CurrentViewType)
-	{
-		case EBlueprintPerfViewType::Overview:			return LOCTEXT("OverviewLabel", "Profiler Overview");
-		case EBlueprintPerfViewType::ExecutionGraph:	return LOCTEXT("ExecutionGraphLabel", "Execution Graph");
-		case EBlueprintPerfViewType::LeastPerformant:	return LOCTEXT("LeastPerformantLabel", "Least Performant");
-	}
-	return LOCTEXT("UnknownViewLabel", "Unknown ViewType");
-}
-
-TSharedRef<SWidget> SBlueprintProfilerView::CreateViewButton() const
-{
-	return SNew(SHorizontalBox)
-		+SHorizontalBox::Slot()
-		.AutoWidth()
-		.VAlign(VAlign_Center)
-		[
-			SNew(SImage)
-			.Image( FEditorStyle::GetBrush("GenericViewButton") )
-		]
-		+SHorizontalBox::Slot()
-		.AutoWidth()
-		.Padding(2, 0, 0, 0)
-		.VAlign(VAlign_Center)
-		[
-			SNew(STextBlock)
-			.Text(GetCurrentViewText())
-		];
-}
-
-FSlateColor SBlueprintProfilerView::GetViewButtonForegroundColor() const
+void SBlueprintProfilerView::OnGraphLayoutChanged(TWeakObjectPtr<UBlueprint> Blueprint)
 {
-	static const FName InvertedForegroundName("InvertedForeground");
-	static const FName DefaultForegroundName("DefaultForeground");
-	return ViewComboButton->IsHovered() ? FEditorStyle::GetSlateColor(InvertedForegroundName) : FEditorStyle::GetSlateColor(DefaultForegroundName);
-}
-
-void SBlueprintProfilerView::OnViewSelectionChanged(const EBlueprintPerfViewType::Type NewViewType)
-{
-	if (CurrentViewType != NewViewType)
+	UBlueprint* CurrentBP = BlueprintEditor.IsValid() ? BlueprintEditor.Pin()->GetBlueprintObj() : nullptr;
+	if (CurrentBP == Blueprint)
 	{
-		CurrentViewType = NewViewType;
+		UpdateStatusMessage();
 		UpdateActiveProfilerWidget();
 	}
 }
 
-FText SBlueprintProfilerView::GetCurrentHeatMapDisplayModeText() const
-{
-	IBlueprintProfilerInterface& ProfilerModule = FModuleManager::LoadModuleChecked<IBlueprintProfilerInterface>("BlueprintProfiler");
-	switch (ProfilerModule.GetGraphNodeHeatMapDisplayMode())
-	{
-	case EBlueprintProfilerHeatMapDisplayMode::None:		return LOCTEXT("HeatMapDisplayModeLabel_None", "Heat Map: Off");
-	case EBlueprintProfilerHeatMapDisplayMode::Exclusive:	return LOCTEXT("HeatMapDisplayModeLabel_Exclusive", "Heat Map: Exclusive");
-	case EBlueprintProfilerHeatMapDisplayMode::Inclusive:	return LOCTEXT("HeatMapDisplayModeLabel_Inclusive", "Heat Map: Inclusive");
-	case EBlueprintProfilerHeatMapDisplayMode::MaxTiming:	return LOCTEXT("HeatMapDisplayModeLabel_MaxTiming", "Heat Map: Max Timing");
-	}
-	return LOCTEXT("HeatMapDisplayModeLabel_Unknown", "<unknown>");
-}
-
-TSharedRef<SWidget> SBlueprintProfilerView::CreateHeatMapDisplayModeButton() const
-{
-	return SNew(SHorizontalBox)
-		+ SHorizontalBox::Slot()
-		.VAlign(VAlign_Center)
-		[
-			SNew(STextBlock)
-			.Text(this, &SBlueprintProfilerView::GetCurrentHeatMapDisplayModeText)
-		];
-}
-
-FSlateColor SBlueprintProfilerView::GetHeatMapDisplayModeButtonForegroundColor() const
-{
-	static const FName InvertedForegroundName("InvertedForeground");
-	static const FName DefaultForegroundName("DefaultForeground");
-	return HeatMapDisplayModeComboButton->IsHovered() ? FEditorStyle::GetSlateColor(InvertedForegroundName) : FEditorStyle::GetSlateColor(DefaultForegroundName);
-}
-
-bool SBlueprintProfilerView::IsHeatMapDisplayModeSelected(const EBlueprintProfilerHeatMapDisplayMode::Type InHeatMapDisplayMode) const
+void SBlueprintProfilerView::UpdateStatusMessage()
 {
-	IBlueprintProfilerInterface& ProfilerModule = FModuleManager::LoadModuleChecked<IBlueprintProfilerInterface>("BlueprintProfiler");
-	return ProfilerModule.GetGraphNodeHeatMapDisplayMode() == InHeatMapDisplayMode;
-}
-
-void SBlueprintProfilerView::OnHeatMapDisplayModeChanged(const EBlueprintProfilerHeatMapDisplayMode::Type NewHeatMapDisplayMode)
-{
-	IBlueprintProfilerInterface& ProfilerModule = FModuleManager::LoadModuleChecked<IBlueprintProfilerInterface>("BlueprintProfiler");
-	if (ProfilerModule.GetGraphNodeHeatMapDisplayMode() != NewHeatMapDisplayMode)
+	if (GetDefault<UEditorExperimentalSettings>()->bBlueprintPerformanceAnalysisTools)
 	{
-		ProfilerModule.SetGraphNodeHeatMapDisplayMode(NewHeatMapDisplayMode);
+		IBlueprintProfilerInterface* ProfilerInterface = FModuleManager::GetModulePtr<IBlueprintProfilerInterface>("BlueprintProfiler");
+		if (ProfilerInterface && ProfilerInterface->IsProfilerEnabled())
+		{
+			UBlueprint* CurrentBP = BlueprintEditor.IsValid() ? BlueprintEditor.Pin()->GetBlueprintObj() : nullptr;
+			if (CurrentBP && CurrentBP->GeneratedClass->HasInstrumentation())
+			{
+				StatusText = LOCTEXT("ProfilerNoDataText", "The Blueprint Profiler is active but does not currently have any data to display");
+			}
+			else
+			{
+				StatusText = LOCTEXT("ProfilerNoInstrumentationText", "The Blueprint Profiler is active but the current blueprint does not have any instrumentation");
+			}
+		}
+		else
+		{
+			StatusText = LOCTEXT("ProfilerInactiveText", "The Blueprint Profiler is currently Inactive");
+		}								
 	}
-}
-
-FText SBlueprintProfilerView::GetCurrentWireHeatMapDisplayModeText() const
-{
-	IBlueprintProfilerInterface& ProfilerModule = FModuleManager::LoadModuleChecked<IBlueprintProfilerInterface>("BlueprintProfiler");
-	switch (ProfilerModule.GetWireHeatMapDisplayMode())
+	else
 	{
-	case EBlueprintProfilerHeatMapDisplayMode::None:		return LOCTEXT("WireHeatMapDisplayModeLabel_None", "Wire Heat Map: Off");
-	case EBlueprintProfilerHeatMapDisplayMode::Exclusive:	return LOCTEXT("WireHeatMapDisplayModeLabel_Exclusive", "Wire Heat Map: Exclusive");
-	case EBlueprintProfilerHeatMapDisplayMode::Inclusive:	return LOCTEXT("WireHeatMapDisplayModeLabel_Inclusive", "Wire Heat Map: Inclusive");
-	case EBlueprintProfilerHeatMapDisplayMode::MaxTiming:	return LOCTEXT("WireHeatMapDisplayModeLabel_MaxTiming", "Wire Heat Map: Max Timing");
+		StatusText = LOCTEXT("DisabledProfilerText", "The Blueprint Profiler is experimental and is currently not enabled in the editor preferences");
 	}
-	return LOCTEXT("WireHeatMapDisplayModeLabel_Unknown", "<unknown>");
 }
 
-TSharedRef<SWidget> SBlueprintProfilerView::CreateWireHeatMapDisplayModeButton() const
+void SBlueprintProfilerView::UpdateActiveProfilerWidget()
 {
-	return SNew(SHorizontalBox)
-		+ SHorizontalBox::Slot()
-		.VAlign(VAlign_Center)
+	ChildSlot
+	[
+		SNew(SVerticalBox)
+		+SVerticalBox::Slot()
+		.AutoHeight()
+		.Padding(0)
 		[
-			SNew(STextBlock)
-			.Text(this, &SBlueprintProfilerView::GetCurrentWireHeatMapDisplayModeText)
-		];
-}
-
-FSlateColor SBlueprintProfilerView::GetWireHeatMapDisplayModeButtonForegroundColor() const
-{
-	static const FName InvertedForegroundName("InvertedForeground");
-	static const FName DefaultForegroundName("DefaultForeground");
-	return WireHeatMapDisplayModeComboButton->IsHovered() ? FEditorStyle::GetSlateColor(InvertedForegroundName) : FEditorStyle::GetSlateColor(DefaultForegroundName);
-}
-
-bool SBlueprintProfilerView::IsWireHeatMapDisplayModeSelected(const EBlueprintProfilerHeatMapDisplayMode::Type InHeatMapDisplayMode) const
-{
-	IBlueprintProfilerInterface& ProfilerModule = FModuleManager::LoadModuleChecked<IBlueprintProfilerInterface>("BlueprintProfiler");
-	return ProfilerModule.GetWireHeatMapDisplayMode() == InHeatMapDisplayMode;
+			ProfilerToolbar.ToSharedRef()
+		]
+		+SVerticalBox::Slot()
+		[
+			SNew(SBorder)
+			.Padding(FMargin(0,2,0,0))
+			.BorderImage(FEditorStyle::GetBrush("NoBorder"))
+			[
+				CreateActiveStatisticWidget()
+			]
+		]
+	];
 }
 
-void SBlueprintProfilerView::OnWireHeatMapDisplayModeChanged(const EBlueprintProfilerHeatMapDisplayMode::Type NewHeatMapDisplayMode)
+void SBlueprintProfilerView::OnViewChanged()
 {
-	IBlueprintProfilerInterface& ProfilerModule = FModuleManager::LoadModuleChecked<IBlueprintProfilerInterface>("BlueprintProfiler");
-	if (ProfilerModule.GetWireHeatMapDisplayMode() != NewHeatMapDisplayMode)
-	{
-		ProfilerModule.SetWireHeatMapDisplayMode(NewHeatMapDisplayMode);
-	}
+	UpdateActiveProfilerWidget();
 }
 
 TSharedRef<SWidget> SBlueprintProfilerView::CreateActiveStatisticWidget()
@@ -407,10 +149,14 @@ TSharedRef<SWidget> SBlueprintProfilerView::CreateActiveStatisticWidget()
 	EBlueprintPerfViewType::Type NewViewType = EBlueprintPerfViewType::None;
 	if (GetDefault<UEditorExperimentalSettings>()->bBlueprintPerformanceAnalysisTools)
 	{
-		IBlueprintProfilerInterface* ProfilerInterface = FModuleManager::GetModulePtr<IBlueprintProfilerInterface>("BlueprintProfiler");
-		if (ProfilerInterface && ProfilerInterface->IsProfilerEnabled())
+		IBlueprintProfilerInterface& ProfilerModule = FModuleManager::LoadModuleChecked<IBlueprintProfilerInterface>("BlueprintProfiler");
+		if (ProfilerModule.IsProfilerEnabled())
 		{
-			NewViewType = CurrentViewType;
+			UBlueprint* ActiveBlueprint = BlueprintEditor.IsValid() ? BlueprintEditor.Pin()->GetBlueprintObj() : nullptr;
+			if (ActiveBlueprint && ActiveBlueprint->GeneratedClass->HasInstrumentation() && ProfilerToolbar.IsValid())
+			{
+				NewViewType = ProfilerToolbar->GetActiveViewType();
+			}
 		}
 	}
 	switch(NewViewType)
@@ -423,7 +169,8 @@ TSharedRef<SWidget> SBlueprintProfilerView::CreateActiveStatisticWidget()
 		case EBlueprintPerfViewType::ExecutionGraph:
 		{
 			return SNew(SGraphExecutionStatDisplay)
-				.AssetEditor(BlueprintEditor);
+				.AssetEditor(BlueprintEditor)
+				.DisplayOptions(DisplayOptions);
 		}
 		case EBlueprintPerfViewType::LeastPerformant:
 		{
diff --git a/Engine/Source/Editor/Kismet/Private/Profiler/SBlueprintProfilerView.h b/Engine/Source/Editor/Kismet/Private/Profiler/SBlueprintProfilerView.h
--- a/Engine/Source/Editor/Kismet/Private/Profiler/SBlueprintProfilerView.h
+++ b/Engine/Source/Editor/Kismet/Private/Profiler/SBlueprintProfilerView.h
@@ -2,19 +2,7 @@
 
 #pragma once
 
-#include "BlueprintProfilerModule.h"
-
-/** Blueprint performance view type */
-namespace EBlueprintPerfViewType
-{
-	enum Type
-	{
-		None = 0,
-		Overview,
-		ExecutionGraph,
-		LeastPerformant
-	};
-}
+#include "SBlueprintProfilerToolbar.h"
 
 //////////////////////////////////////////////////////////////////////////
 // SBlueprintProfilerView
@@ -42,57 +30,26 @@ protected:
 	/** Called when the profiler state is toggled */
 	void OnToggleProfiler(bool bEnabled);
 
+	/** Called to update the display when the stat layouts change or the blueprint is compiled. */
+	void OnGraphLayoutChanged(TWeakObjectPtr<UBlueprint> Blueprint);
+
 	/** Returns the profiler current status message */
 	FText GetProfilerStatusText() const { return StatusText; }
 
 	/** Called to create the child profiler widgets */
 	void UpdateActiveProfilerWidget();
 
-	/** Returns the current view label */
-	FText GetCurrentViewText() const;
-
-	/** Constructs and returns the view button widget */
-	TSharedRef<SWidget> CreateViewButton() const;
-
-	/** Returns the active view button foreground color */
-	FSlateColor GetViewButtonForegroundColor() const;
-
 	/** Called when the profiler view type is changed */
-	void OnViewSelectionChanged(const EBlueprintPerfViewType::Type NewViewType);
-
-	/** Returns the current heat map display mode label */
-	FText GetCurrentHeatMapDisplayModeText() const;
-
-	/** Constructs and returns the heat map display mode button widget */
-	TSharedRef<SWidget> CreateHeatMapDisplayModeButton() const;
-
-	/** Returns the active heat map display mode button foreground color */
-	FSlateColor GetHeatMapDisplayModeButtonForegroundColor() const;
-
-	/** Returns whether or not the given heat map display mode is selected */
-	bool IsHeatMapDisplayModeSelected(const EBlueprintProfilerHeatMapDisplayMode::Type InHeatMapDisplayMode) const;
-
-	/** Called when the heat map display mode is changed */
-	void OnHeatMapDisplayModeChanged(const EBlueprintProfilerHeatMapDisplayMode::Type NewHeatMapDisplayMode);
-
-	/** Returns the current wire heat map display mode label */
-	FText GetCurrentWireHeatMapDisplayModeText() const;
-
-	/** Constructs and returns the wire heat map display mode button widget */
-	TSharedRef<SWidget> CreateWireHeatMapDisplayModeButton() const;
-
-	/** Returns the active wire heat map display mode button foreground color */
-	FSlateColor GetWireHeatMapDisplayModeButtonForegroundColor() const;
-
-	/** Returns whether or not the given wire heat map display mode is selected */
-	bool IsWireHeatMapDisplayModeSelected(const EBlueprintProfilerHeatMapDisplayMode::Type InHeatMapDisplayMode) const;
-
-	/** Called when the wire heat map display mode is changed */
-	void OnWireHeatMapDisplayModeChanged(const EBlueprintProfilerHeatMapDisplayMode::Type NewHeatMapDisplayMode);
+	void OnViewChanged();
 
 	/** Create active statistic display widget */
 	TSharedRef<SWidget> CreateActiveStatisticWidget();
 
+private:
+
+	/** Update status message */
+	void UpdateStatusMessage();
+
 protected:
 
 	/** Profiler status Text */
@@ -101,15 +58,9 @@ protected:
 	/** Blueprint editor */
 	TWeakPtr<class FBlueprintEditor> BlueprintEditor;
 
-	/** Current profiler view type */
-	EBlueprintPerfViewType::Type CurrentViewType;
-
-	/** View combo button widget */
-	TSharedPtr<SComboButton> ViewComboButton;
-
-	/** Heat map display mode combo button widget */
-	TSharedPtr<SComboButton> HeatMapDisplayModeComboButton;
+	/** Display options */
+	TSharedPtr<FBlueprintProfilerStatOptions> DisplayOptions;
 
-	/** Heat map display mode combo button widget */
-	TSharedPtr<SComboButton> WireHeatMapDisplayModeComboButton;
+	/** Profiler Toolbar */
+	TSharedPtr<class SBlueprintProfilerToolbar> ProfilerToolbar;
 };
diff --git a/Engine/Source/Editor/Kismet/Private/Profiler/SGraphExecutionStatDisplay.cpp b/Engine/Source/Editor/Kismet/Private/Profiler/SGraphExecutionStatDisplay.cpp
--- a/Engine/Source/Editor/Kismet/Private/Profiler/SGraphExecutionStatDisplay.cpp
+++ b/Engine/Source/Editor/Kismet/Private/Profiler/SGraphExecutionStatDisplay.cpp
@@ -2,6 +2,7 @@
 
 #include "BlueprintEditorPrivatePCH.h"
 #include "BPProfilerStatisticWidgets.h"
+#include "SBlueprintProfilerToolbar.h"
 #include "SGraphExecutionStatDisplay.h"
 #include "BlueprintEditor.h"
 #include "Public/Profiler/EventExecution.h"
@@ -27,13 +28,10 @@ SGraphExecutionStatDisplay::~SGraphExecutionStatDisplay()
 void SGraphExecutionStatDisplay::Construct(const FArguments& InArgs)
 {	
 	BlueprintEditor = InArgs._AssetEditor;
-	if (!DisplayOptions.IsValid())
-	{
-		DisplayOptions = MakeShareable(new FBPProfilerStatDiplayOptions);
-	}
+	DisplayOptions = InArgs._DisplayOptions;
 	// Register for profiling toggle events
 	FBlueprintCoreDelegates::OnToggleScriptProfiler.AddSP(this, &SGraphExecutionStatDisplay::OnToggleProfiler);
-	// Remove delegate for graph structural changes
+	// Register delegate for graph structural changes
 	if (IBlueprintProfilerInterface* Profiler = FModuleManager::GetModulePtr<IBlueprintProfilerInterface>("BlueprintProfiler"))
 	{
 		Profiler->GetGraphLayoutChangedDelegate().AddSP(this, &SGraphExecutionStatDisplay::OnGraphLayoutChanged);
@@ -43,19 +41,6 @@ void SGraphExecutionStatDisplay::Construct(const FArguments& InArgs)
 	[
 		SNew(SVerticalBox)
 		+SVerticalBox::Slot()
-		.Padding(FMargin(0, 0, 0, 2))
-		.AutoHeight()
-		.HAlign(HAlign_Fill)
-		.VAlign(VAlign_Top)
-		[
-			SNew(SBorder)
-			.Padding(4)
-			.BorderImage(FEditorStyle::GetBrush("BlueprintProfiler.ViewToolBar"))
-			[
-				DisplayOptions->CreateToolbar()
-			]
-		]
-		+SVerticalBox::Slot()
 		[
 			SNew(SBorder)
 			.Padding(0)
@@ -74,9 +59,9 @@ void SGraphExecutionStatDisplay::Construct(const FArguments& InArgs)
 					+SHeaderRow::Column(SProfilerStatRow::GetStatName(EBlueprintProfilerStat::Name))
 					.DefaultLabel(SProfilerStatRow::GetStatText(EBlueprintProfilerStat::Name))
 					.ManualWidth(450)
-					+SHeaderRow::Column(SProfilerStatRow::GetStatName(EBlueprintProfilerStat::Time))
-					.DefaultLabel(SProfilerStatRow::GetStatText(EBlueprintProfilerStat::Time))
-					.ManualWidth(70)
+					+SHeaderRow::Column(SProfilerStatRow::GetStatName(EBlueprintProfilerStat::AverageTime))
+					.DefaultLabel(SProfilerStatRow::GetStatText(EBlueprintProfilerStat::AverageTime))
+					.ManualWidth(90)
 					+SHeaderRow::Column(SProfilerStatRow::GetStatName(EBlueprintProfilerStat::InclusiveTime))
 					.DefaultLabel(SProfilerStatRow::GetStatText(EBlueprintProfilerStat::InclusiveTime))
 					.ManualWidth(120)
@@ -136,7 +121,7 @@ void SGraphExecutionStatDisplay::OnDoubleClickStatistic(FBPStatWidgetPtr Item)
 
 void SGraphExecutionStatDisplay::OnStatisticExpansionChanged(FBPStatWidgetPtr Item, bool bExpanded)
 {
-	if (Item.IsValid() && !DisplayOptions->HasFlags(FBPProfilerStatDiplayOptions::AutoExpand))
+	if (Item.IsValid() && !DisplayOptions->HasFlags(FBlueprintProfilerStatOptions::AutoExpand))
 	{
 		Item->SetExpansionState(bExpanded);
 	}
@@ -187,13 +172,13 @@ void SGraphExecutionStatDisplay::Tick(const FGeometry& AllottedGeometry, const d
 							if (BlueprintExecNode.IsValid())
 							{
 								// Build Instance widget execution trees
-								DisplayOptions->ClearFlags(FBPProfilerStatDiplayOptions::Modified);
+								DisplayOptions->ClearFlags(FBlueprintProfilerStatOptions::Modified);
 								// Cache Active blueprint and Instance
 								const FName CurrentInstancePath = DisplayOptions->GetActiveInstance();
 
-								if (DisplayOptions->HasFlags(FBPProfilerStatDiplayOptions::DisplayByInstance))
+								if (DisplayOptions->HasFlags(FBlueprintProfilerStatOptions::DisplayByInstance))
 								{
-									if (DisplayOptions->HasFlags(FBPProfilerStatDiplayOptions::ScopeToDebugInstance))
+									if (DisplayOptions->HasFlags(FBlueprintProfilerStatOptions::ScopeToDebugInstance))
 									{
 										if (CurrentInstancePath != NAME_None)
 										{
@@ -233,7 +218,7 @@ void SGraphExecutionStatDisplay::Tick(const FGeometry& AllottedGeometry, const d
 							if (ExecutionStatTree.IsValid())
 							{
 								ExecutionStatTree->RequestTreeRefresh();
-								if (DisplayOptions->HasFlags(FBPProfilerStatDiplayOptions::AutoExpand))
+								if (DisplayOptions->HasFlags(FBlueprintProfilerStatOptions::AutoExpand))
 								{
 									for (auto Iter : RootTreeItems)
 									{
diff --git a/Engine/Source/Editor/Kismet/Private/Profiler/SGraphExecutionStatDisplay.h b/Engine/Source/Editor/Kismet/Private/Profiler/SGraphExecutionStatDisplay.h
--- a/Engine/Source/Editor/Kismet/Private/Profiler/SGraphExecutionStatDisplay.h
+++ b/Engine/Source/Editor/Kismet/Private/Profiler/SGraphExecutionStatDisplay.h
@@ -17,6 +17,7 @@ public:
 		{}
 
 		SLATE_ARGUMENT(TWeakPtr<FBlueprintEditor>, AssetEditor)
+		SLATE_ARGUMENT(TSharedPtr<struct FBlueprintProfilerStatOptions>, DisplayOptions)
 	SLATE_END_ARGS()
 
 
@@ -74,5 +75,5 @@ protected:
 	TWeakObjectPtr<UBlueprint> CurrentBlueprint;
 
 	/** Current display options */
-	TSharedPtr<struct FBPProfilerStatDiplayOptions> DisplayOptions;
+	TSharedPtr<struct FBlueprintProfilerStatOptions> DisplayOptions;
 };
diff --git a/Engine/Source/Editor/Kismet/Private/Profiler/SLeastPerformantDisplay.cpp b/Engine/Source/Editor/Kismet/Private/Profiler/SLeastPerformantDisplay.cpp
--- a/Engine/Source/Editor/Kismet/Private/Profiler/SLeastPerformantDisplay.cpp
+++ b/Engine/Source/Editor/Kismet/Private/Profiler/SLeastPerformantDisplay.cpp
@@ -54,9 +54,9 @@ void SLeastPerformantDisplay::Construct(const FArguments& InArgs)
 					+SHeaderRow::Column(SProfilerStatRow::GetStatName(EBlueprintProfilerStat::Name))
 					.DefaultLabel(SProfilerStatRow::GetStatText(EBlueprintProfilerStat::Name))
 					.ManualWidth(450)
-					+SHeaderRow::Column(SProfilerStatRow::GetStatName(EBlueprintProfilerStat::Time))
-					.DefaultLabel(SProfilerStatRow::GetStatText(EBlueprintProfilerStat::Time))
-					.ManualWidth(70)
+					+SHeaderRow::Column(SProfilerStatRow::GetStatName(EBlueprintProfilerStat::AverageTime))
+					.DefaultLabel(SProfilerStatRow::GetStatText(EBlueprintProfilerStat::AverageTime))
+					.ManualWidth(90)
 					+SHeaderRow::Column(SProfilerStatRow::GetStatName(EBlueprintProfilerStat::InclusiveTime))
 					.DefaultLabel(SProfilerStatRow::GetStatText(EBlueprintProfilerStat::InclusiveTime))
 					.ManualWidth(120)
diff --git a/Engine/Source/Editor/Kismet/Private/Profiler/SProfilerOverview.cpp b/Engine/Source/Editor/Kismet/Private/Profiler/SProfilerOverview.cpp
--- a/Engine/Source/Editor/Kismet/Private/Profiler/SProfilerOverview.cpp
+++ b/Engine/Source/Editor/Kismet/Private/Profiler/SProfilerOverview.cpp
@@ -54,9 +54,9 @@ void SProfilerOverview::Construct(const FArguments& InArgs)
 					+SHeaderRow::Column(SProfilerStatRow::GetStatName(EBlueprintProfilerStat::Name))
 					.DefaultLabel(SProfilerStatRow::GetStatText(EBlueprintProfilerStat::Name))
 					.ManualWidth(450)
-					+SHeaderRow::Column(SProfilerStatRow::GetStatName(EBlueprintProfilerStat::Time))
-					.DefaultLabel(SProfilerStatRow::GetStatText(EBlueprintProfilerStat::Time))
-					.ManualWidth(70)
+					+SHeaderRow::Column(SProfilerStatRow::GetStatName(EBlueprintProfilerStat::AverageTime))
+					.DefaultLabel(SProfilerStatRow::GetStatText(EBlueprintProfilerStat::AverageTime))
+					.ManualWidth(90)
 					+SHeaderRow::Column(SProfilerStatRow::GetStatName(EBlueprintProfilerStat::InclusiveTime))
 					.DefaultLabel(SProfilerStatRow::GetStatText(EBlueprintProfilerStat::InclusiveTime))
 					.ManualWidth(120)
diff --git a/Engine/Source/Editor/Kismet/Private/Profiler/ScriptPerfData.cpp b/Engine/Source/Editor/Kismet/Private/Profiler/ScriptPerfData.cpp
--- a/Engine/Source/Editor/Kismet/Private/Profiler/ScriptPerfData.cpp
+++ b/Engine/Source/Editor/Kismet/Private/Profiler/ScriptPerfData.cpp
@@ -7,37 +7,36 @@
 // FScriptPerfData
 
 FNumberFormattingOptions FScriptPerfData::StatNumberFormat;
+bool FScriptPerfData::bAverageBlueprintStats = true;
+bool FScriptPerfData::bUseRecentSampleBias = false;
 float FScriptPerfData::RecentSampleBias = 0.2f;
 float FScriptPerfData::HistoricalSampleBias = 0.8f;
 float FScriptPerfData::EventPerformanceThreshold = 1.f;
-float FScriptPerfData::NodePerformanceThreshold = 0.2f;
+float FScriptPerfData::ExclusivePerformanceThreshold = 0.2f;
 float FScriptPerfData::InclusivePerformanceThreshold = 0.25f;
 float FScriptPerfData::MaxPerformanceThreshold = 0.5f;
+float FScriptPerfData::NodeTotalTimeWaterMark = 0.f;
 
-void FScriptPerfData::AddEventTiming(const double NodeTimingIn)
+void FScriptPerfData::AddEventTiming(const double ExclusiveTimingIn)
 {
-	const double TimingInMs = NodeTimingIn * 1000;
-	NodeTiming = (TimingInMs * RecentSampleBias) + (NodeTiming * HistoricalSampleBias);
-	MaxTiming = FMath::Max<double>(MaxTiming, NodeTiming);
-	MinTiming = FMath::Min<double>(MinTiming, NodeTiming);
-	TotalTiming += NodeTimingIn;
+	const double TimingInMs = ExclusiveTimingIn * 1000;
+	ExclusiveTiming = (TimingInMs * RecentSampleBias) + (ExclusiveTiming * HistoricalSampleBias);
+	MaxTiming = FMath::Max<double>(MaxTiming, ExclusiveTiming);
+	MinTiming = FMath::Min<double>(MinTiming, ExclusiveTiming);
+	TotalTiming += ExclusiveTimingIn;
 	NumSamples++;
 }
 
-void FScriptPerfData::AddPureChainTiming(const double PureNodeTimingIn)
+void FScriptPerfData::AddInclusiveTiming(const double InclusiveNodeTimingIn, const bool bIncrementSamples, const bool bDiscreteTiming)
 {
-	const double PureTimingInMs = PureNodeTimingIn * 1000;
-	const double NewInclusiveTimingInMs = NodeTiming + PureTimingInMs;
-	InclusiveTiming = (NewInclusiveTimingInMs * RecentSampleBias) + (InclusiveTiming * HistoricalSampleBias);
-	TotalTiming += PureNodeTimingIn;
-}
-
-void FScriptPerfData::AddInclusiveTiming(const double InclusiveNodeTimingIn)
-{
-	const double TimingInMs = InclusiveNodeTimingIn * 1000;
+	// Discrete timings should negate the discrete nodetimings so they sum correctly later.
+	const double TimingInMs = bDiscreteTiming ? ((InclusiveNodeTimingIn * 1000) - ExclusiveTiming) : (InclusiveNodeTimingIn * 1000);
 	InclusiveTiming = (TimingInMs * RecentSampleBias) + (InclusiveTiming * HistoricalSampleBias);
 	TotalTiming += InclusiveNodeTimingIn;
-	NumSamples++;
+	if (bIncrementSamples)
+	{
+		NumSamples++;
+	}
 }
 
 void FScriptPerfData::AddData(const FScriptPerfData& DataIn)
@@ -45,7 +44,7 @@ void FScriptPerfData::AddData(const FScriptPerfData& DataIn)
 	if (DataIn.IsDataValid())
 	{
 		// Accumulate data, find min, max values
-		NodeTiming += DataIn.NodeTiming;
+		ExclusiveTiming += DataIn.ExclusiveTiming;
 		if (DataIn.InclusiveTiming > 0.0)
 		{
 			InclusiveTiming += DataIn.InclusiveTiming;
@@ -61,20 +60,72 @@ void FScriptPerfData::AddBranchData(const FScriptPerfData& DataIn)
 {
 	if (DataIn.IsDataValid())
 	{
-		NumSamples += DataIn.NumSamples;
+		// Accumulate all data with the exception of samples for which we take the first valid value.
+		const bool bFirstSample = MaxTiming == -MAX_dbl;
+		ExclusiveTiming += DataIn.ExclusiveTiming;
+		if (DataIn.InclusiveTiming > 0.0)
+		{
+			InclusiveTiming += DataIn.InclusiveTiming;
+		}
+		TotalTiming += DataIn.TotalTiming;
+		if (bFirstSample)
+		{
+			NumSamples = DataIn.NumSamples;
+			MaxTiming = DataIn.MaxTiming;
+			MinTiming = DataIn.MinTiming;
+		}
+		else
+		{
+			MaxTiming += DataIn.MaxTiming;
+			MinTiming += DataIn.MinTiming;
+		}
+	}
+}
+
+void FScriptPerfData::InitialiseFromDataSet(const TArray<TSharedPtr<FScriptPerfData>>& DataSet)
+{
+	Reset();
+	for (auto DataPoint : DataSet)
+	{
+		if (DataPoint.IsValid())
+		{
+			ExclusiveTiming += DataPoint->ExclusiveTiming;
+			InclusiveTiming += DataPoint->InclusiveTiming;
+			TotalTiming += DataPoint->TotalTiming;
+			// Either sum the samples to get the average or use the highest sample rate encountered in the dataset.
+			NumSamples = bAverageBlueprintStats ? (NumSamples + DataPoint->NumSamples) : FMath::Max(NumSamples, DataPoint->NumSamples);
+			MaxTiming = FMath::Max<double>(MaxTiming, DataPoint->MaxTiming);
+			MinTiming = FMath::Min<double>(MinTiming, DataPoint->MinTiming);
+		}
 	}
 }
 
 void FScriptPerfData::Reset()
 {
-	NodeTiming = 0.0;
+	ExclusiveTiming = 0.0;
 	InclusiveTiming = 0.0;
 	TotalTiming = 0.0;
+	HottestPathHeatValue = 0.f;
 	NumSamples = 0;
 	MaxTiming = -MAX_dbl;
 	MinTiming = MAX_dbl;
 }
 
+double FScriptPerfData::GetExclusiveTiming() const
+{
+	return bUseRecentSampleBias ? ExclusiveTiming : ((TotalTiming*1000.0)/NumSamples);
+}
+
+void FScriptPerfData::EnableBlueprintStatAverage(const bool bEnable)
+{
+	bAverageBlueprintStats = bEnable;
+}
+
+void FScriptPerfData::EnableRecentSampleBias(const bool bEnable)
+{
+	bUseRecentSampleBias = bEnable;
+}
+
 void FScriptPerfData::SetRecentSampleBias(const float RecentSampleBiasIn)
 {
 	RecentSampleBias = RecentSampleBiasIn;
@@ -86,15 +137,15 @@ void FScriptPerfData::SetEventPerformanceThreshold(const float EventPerformanceT
 	EventPerformanceThreshold = 1.f / EventPerformanceThresholdIn;
 }
 
-void FScriptPerfData::SetNodePerformanceThreshold(const float NodePerformanceThresholdIn)
+void FScriptPerfData::SetExclusivePerformanceThreshold(const float ExclusivePerformanceThresholdIn)
 {
-	NodePerformanceThreshold = 1.f / NodePerformanceThresholdIn;
+	ExclusivePerformanceThreshold = 1.f / ExclusivePerformanceThresholdIn;
 }
 
 float FScriptPerfData::GetNodeHeatLevel() const
 {
-	const float PerfMetric = StatType == EScriptPerfDataType::Event ? EventPerformanceThreshold : NodePerformanceThreshold;
-	return FMath::Min<float>(NodeTiming * PerfMetric, 1.f);
+	const float PerfMetric = StatType == EScriptPerfDataType::Event ? EventPerformanceThreshold : ExclusivePerformanceThreshold;
+	return FMath::Min<float>(GetExclusiveTiming() * PerfMetric, 1.f);
 }
 
 void FScriptPerfData::SetInclusivePerformanceThreshold(const float InclusivePerformanceThresholdIn)
@@ -104,7 +155,8 @@ void FScriptPerfData::SetInclusivePerformanceThreshold(const float InclusivePerf
 
 float FScriptPerfData::GetInclusiveHeatLevel() const
 {
-	return FMath::Min<float>(InclusiveTiming * InclusivePerformanceThreshold, 1.f);
+	const float PerfMetric = StatType == EScriptPerfDataType::Event ? EventPerformanceThreshold : InclusivePerformanceThreshold;
+	return FMath::Min<float>(GetInclusiveTiming() * InclusivePerformanceThreshold, 1.f);
 }
 
 void FScriptPerfData::SetMaxPerformanceThreshold(const float MaxPerformanceThresholdIn)
@@ -117,20 +169,28 @@ float FScriptPerfData::GetMaxTimeHeatLevel() const
 	return FMath::Min<float>(MaxTiming * MaxPerformanceThreshold, 1.f);
 }
 
-FText FScriptPerfData::GetNodeTimingText() const
+float FScriptPerfData::GetTotalHeatLevel() const
 {
-	if (NodeTiming > 0.0)
+	NodeTotalTimeWaterMark = FMath::Max<float>(TotalTiming, NodeTotalTimeWaterMark);
+	return FMath::Min<float>(TotalTiming / NodeTotalTimeWaterMark, 1.f);
+}
+
+FText FScriptPerfData::GetExclusiveTimingText() const
+{
+	const double Value = GetExclusiveTiming();
+	if (Value > 0.0)
 	{
-		return FText::AsNumber(NodeTiming, &StatNumberFormat);
+		return FText::AsNumber(Value, &StatNumberFormat);
 	}
 	return FText::GetEmpty();
 }
 
 FText FScriptPerfData::GetInclusiveTimingText() const
 {
-	if (InclusiveTiming > 0.0)
+	if (NumSamples > 0)
 	{
-		return FText::AsNumber(InclusiveTiming, &StatNumberFormat);
+		const double InclusiveTemp = GetInclusiveTiming();
+		return FText::AsNumber(InclusiveTemp, &StatNumberFormat);
 	}
 	return FText::GetEmpty();
 }
diff --git a/Engine/Source/Editor/Kismet/Private/SBlueprintEditorToolbar.cpp b/Engine/Source/Editor/Kismet/Private/SBlueprintEditorToolbar.cpp
--- a/Engine/Source/Editor/Kismet/Private/SBlueprintEditorToolbar.cpp
+++ b/Engine/Source/Editor/Kismet/Private/SBlueprintEditorToolbar.cpp
@@ -671,7 +671,14 @@ FSlateIcon FBlueprintEditorToolbar::GetStatusImage() const
 	case BS_Error:
 		return FSlateIcon(FEditorStyle::GetStyleSetName(), "Kismet.Status.Error");
 	case BS_UpToDate:
-		return FSlateIcon(FEditorStyle::GetStyleSetName(), "Kismet.Status.Good");
+	{
+		const bool bInstrumented = BlueprintObj->GeneratedClass->HasInstrumentation();
+		if (BlueprintEditor.Pin()->IsProfilerActive())
+		{
+			return bInstrumented ? FSlateIcon(FEditorStyle::GetStyleSetName(), "Kismet.Status.Instrumented") : FSlateIcon(FEditorStyle::GetStyleSetName(), "Kismet.Status.NotInstrumented");
+		}
+		return bInstrumented ? FSlateIcon(FEditorStyle::GetStyleSetName(), "Kismet.Status.Instrumented") : FSlateIcon(FEditorStyle::GetStyleSetName(), "Kismet.Status.Good");
+	}
 	case BS_UpToDateWithWarnings:
 		return FSlateIcon(FEditorStyle::GetStyleSetName(), "Kismet.Status.Warning");
 	}
@@ -698,7 +705,14 @@ FText FBlueprintEditorToolbar::GetStatusTooltip() const
 	case BS_Error:
 		return LOCTEXT("CompileError_Status", "There was an error during compilation, see the log for details");
 	case BS_UpToDate:
-		return LOCTEXT("GoodToGo_Status", "Good to go");
+	{
+		const bool bInstrumented = BlueprintObj->GeneratedClass->HasInstrumentation();
+		if (BlueprintEditor.Pin()->IsProfilerActive())
+		{
+			return bInstrumented ? LOCTEXT("GoodToGoInstrumented_Status", "Instrumentation is active") : LOCTEXT("GoodToGoNotInstrumented_Status", "Compile to add instrumentation");
+		}
+		return bInstrumented ? LOCTEXT("GoodToGoInstrumented_Status", "Instrumentation is active") : LOCTEXT("GoodToGo_Status", "Good to go");
+	}
 	case BS_UpToDateWithWarnings:
 		return LOCTEXT("GoodToGoWarning_Status", "There was a warning during compilation, see the log for details");
 	}
@@ -706,7 +720,7 @@ FText FBlueprintEditorToolbar::GetStatusTooltip() const
 
 FText FBlueprintEditorToolbar::GetProfilerStatusTooltip() const
 {
-	return BlueprintEditor.Pin()->IsProfilerActive() ? LOCTEXT("Profiling_Enabled", "Profiling Active") :
+	return BlueprintEditor.Pin()->IsProfilerActive() ?	LOCTEXT("Profiling_Enabled", "Profiling Active") :
 														LOCTEXT("Profiling_Disabled", "Profiling Inactive");
 }
 
diff --git a/Engine/Source/Editor/Kismet/Private/SMyBlueprint.cpp b/Engine/Source/Editor/Kismet/Private/SMyBlueprint.cpp
--- a/Engine/Source/Editor/Kismet/Private/SMyBlueprint.cpp
+++ b/Engine/Source/Editor/Kismet/Private/SMyBlueprint.cpp
@@ -1325,7 +1325,7 @@ void SMyBlueprint::CollectAllActions(FGraphActionListBuilderBase& OutAllActions)
 
 						FText FunctionCategory = Function->GetMetaDataText(FBlueprintMetadata::MD_FunctionCategory, TEXT("UObjectCategory"), Function->GetFullGroupName(false));
 
-						TSharedPtr<FEdGraphSchemaAction_K2Graph> NewFuncAction = MakeShareable(new FEdGraphSchemaAction_K2Graph(EEdGraphSchemaAction_K2Graph::Function, FunctionCategory, FunctionDesc, FunctionTooltip, 1));
+						TSharedPtr<FEdGraphSchemaAction_K2Graph> NewFuncAction = MakeShareable(new FEdGraphSchemaAction_K2Graph(EEdGraphSchemaAction_K2Graph::Function, FunctionCategory, FunctionDesc, FunctionTooltip, 1, NodeSectionID::INTERFACE));
 						NewFuncAction->FuncName = FunctionName;
 						OutAllActions.AddAction(NewFuncAction);
 					}
@@ -2290,7 +2290,10 @@ UEdGraph* SMyBlueprint::GetFocusedGraph() const
 
 void SMyBlueprint::OnObjectPropertyChanged(UObject* InObject, FPropertyChangedEvent& InPropertyChangedEvent)
 {
-	bNeedsRefresh = ( InObject == Blueprint );
+	if (InObject == Blueprint && InPropertyChangedEvent.ChangeType != EPropertyChangeType::ValueSet)
+	{
+		bNeedsRefresh = true;
+	}
 }
 
 bool SMyBlueprint::IsEditingMode() const
diff --git a/Engine/Source/Editor/Kismet/Private/SSCSEditor.cpp b/Engine/Source/Editor/Kismet/Private/SSCSEditor.cpp
--- a/Engine/Source/Editor/Kismet/Private/SSCSEditor.cpp
+++ b/Engine/Source/Editor/Kismet/Private/SSCSEditor.cpp
@@ -2662,6 +2662,10 @@ void SSCS_RowWidget::OnMakeNewRootDropAction(FSCSEditorTreeNodePtrType DroppedNo
 					SCS_Node->AttachToName = NAME_None;
 				}
 
+				// Cache the current relative location and rotation values (for propagation)
+				const FVector OldRelativeLocation = SceneComponentTemplate->RelativeLocation;
+				const FRotator OldRelativeRotation = SceneComponentTemplate->RelativeRotation;
+
 				// Reset the relative transform (location and rotation only; scale is preserved)
 				SceneComponentTemplate->SetRelativeLocation(FVector::ZeroVector);
 				SceneComponentTemplate->SetRelativeRotation(FRotator::ZeroRotator);
@@ -2678,6 +2682,10 @@ void SSCS_RowWidget::OnMakeNewRootDropAction(FSCSEditorTreeNodePtrType DroppedNo
 						// Detach from root (keeping world transform, except for scale)
 						SceneComponentInstance->DetachFromComponent(DetachmentTransformRules);
 
+						// Propagate the default relative location & rotation reset from the template to the instance
+						FComponentEditorUtils::ApplyDefaultValueChange(SceneComponentInstance, SceneComponentInstance->RelativeLocation, OldRelativeLocation, SceneComponentTemplate->RelativeLocation);
+						FComponentEditorUtils::ApplyDefaultValueChange(SceneComponentInstance, SceneComponentInstance->RelativeRotation, OldRelativeRotation, SceneComponentTemplate->RelativeRotation);
+
 						// Must also reset the root component here, so that RerunConstructionScripts() will cache the correct root component instance data
 						AActor* Owner = SceneComponentInstance->GetOwner();
 						if (Owner)
@@ -4316,12 +4324,14 @@ void SSCSEditor::UpdateTree(bool bRegenerateTreeNodes)
 					AddTreeNodeFromComponent(RootComponent);
 				}
 				
-				// Now add the rest of the instanced scene component hierarchy
+				// Now add the rest of the instanced scene component hierarchy (excluding editor-only instances and nested DSOs attached to BP-constructed instances, which are not mutable)
 				for(auto CompIter = Components.CreateIterator(); CompIter; ++CompIter)
 				{
 					USceneComponent* SceneComp = Cast<USceneComponent>(*CompIter);
+					USceneComponent* ParentSceneComp = SceneComp != nullptr ? SceneComp->GetAttachParent() : nullptr;
 					if(SceneComp != nullptr && !SceneComp->IsEditorOnly()
-						&& (SceneComp->CreationMethod != EComponentCreationMethod::UserConstructionScript || !GetDefault<UBlueprintEditorSettings>()->bHideConstructionScriptComponentsInDetailsView))
+						&& (SceneComp->CreationMethod != EComponentCreationMethod::UserConstructionScript || !GetDefault<UBlueprintEditorSettings>()->bHideConstructionScriptComponentsInDetailsView)
+						&& (ParentSceneComp == nullptr || !ParentSceneComp->IsCreatedByConstructionScript() || !SceneComp->HasAnyFlags(RF_DefaultSubObject)))
 					{
 						AddTreeNodeFromComponent(SceneComp);
 					}
diff --git a/Engine/Source/Editor/Kismet/Private/UserDefinedStructureEditor.cpp b/Engine/Source/Editor/Kismet/Private/UserDefinedStructureEditor.cpp
--- a/Engine/Source/Editor/Kismet/Private/UserDefinedStructureEditor.cpp
+++ b/Engine/Source/Editor/Kismet/Private/UserDefinedStructureEditor.cpp
@@ -119,16 +119,25 @@ public:
 
 	virtual void PreChange(const class UUserDefinedStruct* Struct, FStructureEditorUtils::EStructureEditorChangeInfo Info) override
 	{
-		StructData->Destroy();
-		DetailsView->SetObject(nullptr);
-		DetailsView->OnFinishedChangingProperties().Clear();
+		// No need to destroy the struct data if only the default values are changing
+		if (Info != FStructureEditorUtils::DefaultValueChanged)
+		{
+			StructData->Destroy();
+			DetailsView->SetObject(nullptr);
+			DetailsView->OnFinishedChangingProperties().Clear();
+		}
 	}
 
 	virtual void PostChange(const class UUserDefinedStruct* Struct, FStructureEditorUtils::EStructureEditorChangeInfo Info) override
 	{
-		StructData->Initialize(UserDefinedStruct.Get());
+		// If change is due to default value, then struct data was not destroyed (see PreChange) and therefore does not need to be re-initialized
+		if (Info != FStructureEditorUtils::DefaultValueChanged)
+		{
+			StructData->Initialize(UserDefinedStruct.Get());
+			DetailsView->SetObject(UserDefinedStruct.Get());
+		}
+
 		FStructureEditorUtils::Fill_MakeStructureDefaultValue(UserDefinedStruct.Get(), StructData->GetStructMemory());
-		DetailsView->SetObject(UserDefinedStruct.Get());
 	}
 
 	// FNotifyHook interface
@@ -171,29 +180,31 @@ void FDefaultValueDetails::OnFinishedChangingProperties(const FPropertyChangedEv
 
 		check(PropertyChangedEvent.MemberProperty && OwnerStruct);
 
-		ensure(OwnerStruct == UserDefinedStruct.Get());
-		const UProperty* DirectProperty = PropertyChangedEvent.MemberProperty;
-		while (DirectProperty && !Cast<const UUserDefinedStruct>(DirectProperty->GetOuter()))
+		if ( ensure(OwnerStruct == UserDefinedStruct.Get()) )
 		{
-			DirectProperty = Cast<const UProperty>(DirectProperty->GetOuter());
-		}
-		ensure(nullptr != DirectProperty);
+			const UProperty* DirectProperty = PropertyChangedEvent.MemberProperty;
+			while (DirectProperty && !Cast<const UUserDefinedStruct>(DirectProperty->GetOuter()))
+			{
+				DirectProperty = Cast<const UProperty>(DirectProperty->GetOuter());
+			}
+			ensure(nullptr != DirectProperty);
 
-		if (DirectProperty)
-		{
-			FString DefaultValueString;
-			bool bDefaultValueSet = false;
+			if (DirectProperty)
 			{
-				if (StructData.IsValid() && StructData->IsValid())
+				FString DefaultValueString;
+				bool bDefaultValueSet = false;
 				{
-					bDefaultValueSet = FBlueprintEditorUtils::PropertyValueToString(DirectProperty, StructData->GetStructMemory(), DefaultValueString);
+					if (StructData.IsValid() && StructData->IsValid())
+					{
+						bDefaultValueSet = FBlueprintEditorUtils::PropertyValueToString(DirectProperty, StructData->GetStructMemory(), DefaultValueString);
+					}
 				}
-			}
 
-			const FGuid VarGuid = FStructureEditorUtils::GetGuidForProperty(DirectProperty);
-			if (bDefaultValueSet && VarGuid.IsValid())
-			{
-				FStructureEditorUtils::ChangeVariableDefaultValue(UserDefinedStruct.Get(), VarGuid, DefaultValueString);
+				const FGuid VarGuid = FStructureEditorUtils::GetGuidForProperty(DirectProperty);
+				if (bDefaultValueSet && VarGuid.IsValid())
+				{
+					FStructureEditorUtils::ChangeVariableDefaultValue(UserDefinedStruct.Get(), VarGuid, DefaultValueString);
+				}
 			}
 		}
 	}
diff --git a/Engine/Source/Editor/Kismet/Public/BlueprintEditor.h b/Engine/Source/Editor/Kismet/Public/BlueprintEditor.h
--- a/Engine/Source/Editor/Kismet/Public/BlueprintEditor.h
+++ b/Engine/Source/Editor/Kismet/Public/BlueprintEditor.h
@@ -157,6 +157,7 @@ public:
 	//~ Begin IBlueprintEditor Interface
 	virtual void RefreshEditors(ERefreshBlueprintEditorReason::Type Reason = ERefreshBlueprintEditorReason::UnknownReason) override;
 	virtual void JumpToHyperlink(const UObject* ObjectReference, bool bRequestRename = false) override;
+	virtual void JumpToPin(const class UEdGraphPin* Pin) override;
 	virtual void SummonSearchUI(bool bSetFindWithinBlueprint, FString NewSearchTerms = FString(), bool bSelectFirstResult = false) override;
 	virtual void SummonFindAndReplaceUI() override;
 	virtual TArray<TSharedPtr<class FSCSEditorTreeNode> > GetSelectedSCSEditorTreeNodes() const override;
@@ -183,9 +184,6 @@ public:
 	/** Pan the view to center on a particular node */
 	void JumpToNode(const class UEdGraphNode* Node, bool bRequestRename=false);
 
-	/** Pan the view to center on a particular pin */
-	void JumpToPin(const class UEdGraphPin* Pin);
-
 	/** Returns a pointer to the Blueprint object we are currently editing, as long as we are editing exactly one */
 	virtual UBlueprint* GetBlueprintObj() const;
 	
diff --git a/Engine/Source/Editor/Kismet/Public/BlueprintEditorModule.h b/Engine/Source/Editor/Kismet/Public/BlueprintEditorModule.h
--- a/Engine/Source/Editor/Kismet/Public/BlueprintEditorModule.h
+++ b/Engine/Source/Editor/Kismet/Public/BlueprintEditorModule.h
@@ -42,6 +42,7 @@ class KISMET_API IBlueprintEditor : public FWorkflowCentricApplication
 {
 public:
 	virtual void JumpToHyperlink(const UObject* ObjectReference, bool bRequestRename) = 0;
+	virtual void JumpToPin(const UEdGraphPin* PinToFocusOn) = 0;
 
 	/** Invokes the search UI and sets the mode and search terms optionally */
 	virtual void SummonSearchUI(bool bSetFindWithinBlueprint, FString NewSearchTerms = FString(), bool bSelectFirstResult = false) = 0;
diff --git a/Engine/Source/Editor/Kismet/Public/Profiler/EventExecution.h b/Engine/Source/Editor/Kismet/Public/Profiler/EventExecution.h
--- a/Engine/Source/Editor/Kismet/Public/Profiler/EventExecution.h
+++ b/Engine/Source/Editor/Kismet/Public/Profiler/EventExecution.h
@@ -23,6 +23,8 @@ namespace EScriptExecutionNodeFlags
 		Node						= 0x00000800,	// Node timing
 		ExecPin						= 0x00001000,	// Exec pin dummy node
 		PureNode					= 0x00002000,	// Pure node (no exec pins)
+		AsyncTask					= 0x00004000,	// Async task node
+		AsyncTaskDelegate			= 0x00008000,	// Async task delegate pin/event
 		FunctionTunnel				= 0x00010000,	// Tunnel function
 		TunnelEntryPin				= 0x00020000,	// Tunnel entry pin
 		TunnelExitPin				= 0x00040000,	// Tunnel exit pin
@@ -123,12 +125,18 @@ public:
 	/** Test whether or not perf data is available for the given instance/trace path */
 	bool HasPerfDataForInstanceAndTracePath(FName InstanceName, const FTracePath& TracePath) const;
 
-	/** Get/add perf data for instance and tracepath */
+	/** Get or add perf data for instance and tracepath */
+	TSharedPtr<class FScriptPerfData> GetOrAddPerfDataByInstanceAndTracePath(FName InstanceName, const FTracePath& TracePath);
+
+	/** Get perf data for instance and tracepath */
 	TSharedPtr<class FScriptPerfData> GetPerfDataByInstanceAndTracePath(FName InstanceName, const FTracePath& TracePath);
 
 	/** Get all instance perf data for the trace path, excluding the global blueprint data */
 	void GetInstancePerfDataByTracePath(const FTracePath& TracePath, TArray<TSharedPtr<FScriptPerfData>>& ResultsOut);
 
+	/** Get or add global blueprint perf data for the trace path */
+	TSharedPtr<FScriptPerfData> GetOrAddBlueprintPerfDataByTracePath(const FTracePath& TracePath);
+
 	/** Get global blueprint perf data for the trace path */
 	TSharedPtr<FScriptPerfData> GetBlueprintPerfDataByTracePath(const FTracePath& TracePath);
 
@@ -160,6 +168,8 @@ struct KISMET_API FScriptExecNodeParams
 	uint32 NodeFlags;
 	/** Oberved object */
 	TWeakObjectPtr<const UObject> ObservedObject;
+	/** Observed pin */
+	FEdGraphPinReference ObservedPin;
 	/** Display name for widget UI */
 	FText DisplayName;
 	/** Tooltip for widget UI */
@@ -170,6 +180,28 @@ struct KISMET_API FScriptExecNodeParams
 	FSlateBrush* Icon;
 };
 
+//////////////////////////////////////////////////////////////////////////
+// FScriptExecutionHottestPathParams
+
+struct FScriptExecutionHottestPathParams
+{
+	FScriptExecutionHottestPathParams(const FName InstanceNameIn, const float EventTimingIn)
+		: InstanceName(InstanceNameIn)
+		, EventTiming(EventTimingIn)
+		, TimeSoFar(0.0)
+	{
+	}
+
+	/** Instance name */
+	const FName InstanceName;
+	/** Event total timing */
+	const double EventTiming;
+	/** Event time so far */
+	double TimeSoFar;
+	/** Current tracepath */
+	FTracePath TracePath;
+};
+
 //////////////////////////////////////////////////////////////////////////
 // FScriptExecutionNode
 
@@ -198,6 +230,9 @@ public:
 	/** Does the node contain these flags */
 	bool HasFlags(const uint32 Flags) const { return (NodeFlags & Flags) != 0U; }
 
+	/** Returns if this exec node can be safely cyclicly linked */
+	bool IsCyclicNode() const { return (NodeFlags & EScriptExecutionNodeFlags::CallSite) == 0U; }
+
 	/** Returns if this exec event represents a change in class/blueprint */
 	bool IsClass() const { return (NodeFlags & EScriptExecutionNodeFlags::Class) != 0U; }
 
@@ -231,6 +266,9 @@ public:
 	/** Gets the observed object context */
 	const UObject* GetObservedObject() { return ObservedObject.Get(); }
 
+	/** Gets the observed object context */
+	const UEdGraphPin* GetObservedPin() const { return ObservedPin.Get(); }
+
 	/** Navigate to object */
 	virtual void NavigateToObject() const;
 
@@ -276,14 +314,26 @@ public:
 	/** Refresh Stats */
 	virtual void RefreshStats(const FTracePath& TracePath);
 
+	/** Calculate Hottest Path Stats */
+	virtual float CalculateHottestPathStats(FScriptExecutionHottestPathParams HotPathParams);
+
 	/** Get all exec nodes */
 	virtual void GetAllExecNodes(TMap<FName, TSharedPtr<FScriptExecutionNode>>& ExecNodesOut);
 
 	/** Get all pure nodes associated with the given trace path */
 	virtual void GetAllPureNodes(TMap<int32, TSharedPtr<FScriptExecutionNode>>& PureNodesOut);
 
+	/** Creates and returns the slate icon widget */
+	virtual TSharedRef<SWidget> GetIconWidget();
+
+	/** Creates and returns the slate hyperlink widget */
+	virtual TSharedRef<SWidget> GetHyperlinkWidget();
+
 protected:
 
+	/** Returns if the observed object is valid */
+	virtual bool IsObservedObjectValid() const;
+
 	/** Returns Tunnel Linear Execution Trace */
 	void MapTunnelLinearExecution(FTracePath& TraceInOut) const;
 
@@ -302,8 +352,10 @@ protected:
 	FName OwningGraphName;
 	/** Node flags to describe the source graph node type */
 	uint32 NodeFlags;
-	/** Oberved object */
+	/** Observed object */
 	TWeakObjectPtr<const UObject> ObservedObject;
+	/** Observed pin */
+	FEdGraphPinReference ObservedPin;
 	/** Display name for widget UI */
 	FText DisplayName;
 	/** Tooltip for widget UI */
@@ -426,6 +478,7 @@ public:
 
 	// FScriptExecutionNode
 	virtual void GetLinearExecutionPath(TArray<FLinearExecPath>& LinearExecutionNodes, const FTracePath& TracePath, const bool bIncludeChildren = false) override;
+	virtual void RefreshStats(const FTracePath& TracePath) override;
 	virtual void NavigateToObject() const;
 	// ~FScriptExecutionNode
 
@@ -464,7 +517,27 @@ public:
 
 	// FScriptExecutionNode
 	virtual void NavigateToObject() const override;
+	virtual TSharedRef<SWidget> GetIconWidget() override;
+	virtual TSharedRef<SWidget> GetHyperlinkWidget() override;
+	virtual bool IsObservedObjectValid() const override;
 	// ~FScriptExecutionNode
+
+	/** Sets the current PIE Instance */
+	void SetPIEInstance(const UObject* PIEInstanceIn) { PIEInstance = PIEInstanceIn; }
+
+private:
+
+	/** Return the current instance icon color */
+	FSlateColor GetInstanceIconColor() const;
+
+	/** Return the current instance icon color */
+	FText GetInstanceTooltip() const;
+
+private:
+
+	/** PIE instance */
+	TWeakObjectPtr<const UObject> PIEInstance;
+
 };
 
 //////////////////////////////////////////////////////////////////////////
diff --git a/Engine/Source/Editor/Kismet/Public/Profiler/ScriptPerfData.h b/Engine/Source/Editor/Kismet/Public/Profiler/ScriptPerfData.h
--- a/Engine/Source/Editor/Kismet/Public/Profiler/ScriptPerfData.h
+++ b/Engine/Source/Editor/Kismet/Public/Profiler/ScriptPerfData.h
@@ -20,23 +20,21 @@ public:
 
 	FScriptPerfData(const EScriptPerfDataType StatTypeIn)
 		: StatType(StatTypeIn)
+		, ExclusiveTiming(0.0)
 		, InclusiveTiming(0.0)
-		, NodeTiming(0.0)
 		, MaxTiming(-MAX_dbl)
 		, MinTiming(MAX_dbl)
 		, TotalTiming(0.0)
+		, HottestPathHeatValue(0.f)
 		, NumSamples(0)
 	{
 	}
 
 	/** Add a single event node timing into the dataset */
-	void AddEventTiming(const double NodeTimingIn);
-
-	/** Add a single event pure node timing into the dataset */
-	void AddPureChainTiming(const double PureTimingIn);
+	void AddEventTiming(const double ExclusiveTimingIn);
 
 	/** Add a inclusive timing into the dataset */
-	void AddInclusiveTiming(const double InclusiveNodeTimingIn);
+	void AddInclusiveTiming(const double InclusiveNodeTimingIn, const bool bIncrementSamples, const bool bDiscreteTiming);
 
 	/** Add data */
 	void AddData(const FScriptPerfData& DataIn);
@@ -44,6 +42,9 @@ public:
 	/** Add branch data */
 	void AddBranchData(const FScriptPerfData& DataIn);
 
+	/** Initialise from data set */
+	void InitialiseFromDataSet(const TArray<TSharedPtr<FScriptPerfData>>& DataSet);
+
 	/** Increments sample count without affecting data */
 	void TickSamples() { NumSamples++; }
 
@@ -55,15 +56,17 @@ public:
 
 	// Updates the various thresholds that control stats calcs and display.
 	static void SetNumberFormattingForStats(const FNumberFormattingOptions& FormatIn) { StatNumberFormat = FormatIn; }
+	static void EnableBlueprintStatAverage(const bool bEnable);
+	static void EnableRecentSampleBias(const bool bEnable);
 	static void SetRecentSampleBias(const float RecentSampleBiasIn);
 	static void SetEventPerformanceThreshold(const float EventPerformanceThresholdIn);
-	static void SetNodePerformanceThreshold(const float NodePerformanceThresholdIn);
+	static void SetExclusivePerformanceThreshold(const float ExclusivePerformanceThresholdIn);
 	static void SetInclusivePerformanceThreshold(const float InclusivePerformanceThresholdIn);
 	static void SetMaxPerformanceThreshold(const float MaxPerformanceThresholdIn);
 
 	// Returns the various stats this container holds
-	double GetNodeTiming() const { return NodeTiming; }
-	double GetInclusiveTiming() const { return InclusiveTiming; }
+	double GetExclusiveTiming() const;
+	double GetInclusiveTiming() const { return GetExclusiveTiming() + InclusiveTiming; }
 	double GetMaxTiming() const { return MaxTiming; }
 	double GetMinTiming() const { return MinTiming; }
 	double GetTotalTiming() const { return TotalTiming; }
@@ -73,11 +76,20 @@ public:
 	float GetNodeHeatLevel() const;
 	float GetInclusiveHeatLevel() const;
 	float GetMaxTimeHeatLevel() const;
+	float GetTotalHeatLevel() const;
+
+	// Hottest path interface
+	float GetHottestPathHeatLevel() const { return HottestPathHeatValue; }
+	void SetHottestPathHeatLevel(const float HottestPathValue) { HottestPathHeatValue = HottestPathValue; }
+
+	// Hottest endpoint interface
+	float GetHottestEndpointHeatLevel() const { return HottestEndpointHeatValue; }
+	void SetHottestEndpointHeatLevel(const float HottestEndpointValue) { HottestEndpointHeatValue = HottestEndpointValue; }
 
 	// Returns the various stats this container holds in FText format
 	FText GetTotalTimingText() const;
 	FText GetInclusiveTimingText() const;
-	FText GetNodeTimingText() const;
+	FText GetExclusiveTimingText() const;
 	FText GetMaxTimingText() const;
 	FText GetMinTimingText() const;
 	FText GetSamplesText() const;
@@ -86,33 +98,43 @@ private:
 
 	/** The stat type for performance threshold comparisons */
 	EScriptPerfDataType StatType;
+	/** Exclusive timing */
+	double ExclusiveTiming;
 	/** Inclusive timing, pure and node timings */
 	double InclusiveTiming;
-	/** Node timing */
-	double NodeTiming;
 	/** Max exclusive timing */
 	double MaxTiming;
 	/** Min exclusive timing */
 	double MinTiming;
 	/** Total time accrued */
 	double TotalTiming;
+	/** Hottest path value */
+	float HottestPathHeatValue;
+	/** Hottest endpoint value */
+	float HottestEndpointHeatValue;
 	/** The current inclusive sample count */
 	int32 NumSamples;
 
 	/** Number formatting options */
 	static FNumberFormattingOptions StatNumberFormat;
 
+	/** Sets blueprint level stats as averaged */
+	static bool bAverageBlueprintStats;
+	/** Enables the sample biasing */
+	static bool bUseRecentSampleBias;
 	/** Controls the bias between new and older samples */
 	static float RecentSampleBias;
 	/** Historical Sample Bias */
 	static float HistoricalSampleBias;
 	/** Event performance threshold */
 	static float EventPerformanceThreshold;
-	/** Node performance threshold */
-	static float NodePerformanceThreshold;
+	/** Exclusive performance threshold */
+	static float ExclusivePerformanceThreshold;
 	/** Node inclusive performance threshold */
 	static float InclusivePerformanceThreshold;
 	/** Node max performance threshold */
 	static float MaxPerformanceThreshold;
+	/** Total time watermark */
+	static float NodeTotalTimeWaterMark;
 
 };
diff --git a/Engine/Source/Editor/Kismet/Public/Profiler/TracePath.h b/Engine/Source/Editor/Kismet/Public/Profiler/TracePath.h
--- a/Engine/Source/Editor/Kismet/Public/Profiler/TracePath.h
+++ b/Engine/Source/Editor/Kismet/Public/Profiler/TracePath.h
@@ -63,7 +63,7 @@ public:
 		CachedCrc = FCrc::MemCrc32(&ExitPinScriptOffset, 1, CachedCrc);
 	}
 
-	/** Reset */
+	/** Reset Path and Tunnel */
 	void Reset()
 	{
 #if TRACEPATH_DEBUG
@@ -73,6 +73,15 @@ public:
 		Tunnel.Reset();
 	}
 
+	/** Reset Path Only */
+	void ResetPath()
+	{
+#if TRACEPATH_DEBUG
+		TraceStack.SetNum(0);
+#endif
+		CachedCrc = 0;
+	}
+
 	/** Returns the current tunnel */
 	TSharedPtr<const class FScriptExecutionTunnelEntry> GetTunnel() const { return Tunnel; }
 
diff --git a/Engine/Source/Editor/KismetCompiler/Private/AnimBlueprintCompiler.cpp b/Engine/Source/Editor/KismetCompiler/Private/AnimBlueprintCompiler.cpp
--- a/Engine/Source/Editor/KismetCompiler/Private/AnimBlueprintCompiler.cpp
+++ b/Engine/Source/Editor/KismetCompiler/Private/AnimBlueprintCompiler.cpp
@@ -39,14 +39,6 @@
 //
 class UAnimStateNodeBase;
 
-//////////////////////////////////////////////////////////////////////////
-
-template<>
-FString FNetNameMapping::MakeBaseName<UAnimGraphNode_Base>(const UAnimGraphNode_Base* Net)
-{
-	return FString::Printf(TEXT("%s_%s"), *Net->GetDescriptiveCompiledName(), *Net->NodeGuid.ToString());
-}
-
 //////////////////////////////////////////////////////////////////////////
 // FAnimBlueprintCompiler::FEffectiveConstantRecord
 
@@ -219,7 +211,7 @@ void FAnimBlueprintCompiler::CreateEvaluationHandler(UAnimGraphNode_Base* Visual
 						// Wire up the data input
 						UEdGraphPin* TargetItemPin = ArrayNode->FindPinChecked(TEXT("Item"));
 						TargetItemPin->CopyPersistentDataFromOldPin(*SourcePin);
-						MessageLog.NotifyIntermediateObjectCreation(TargetItemPin, SourcePin);
+						MessageLog.NotifyIntermediatePinCreation(TargetItemPin, SourcePin);
 						SourcePin->BreakAllPinLinks();
 					}
 				}
@@ -233,7 +225,7 @@ void FAnimBlueprintCompiler::CreateEvaluationHandler(UAnimGraphNode_Base* Visual
 
 					PropertiesBeingSet.Add(FName(*SourcePin->PinName));
 					TargetPin->CopyPersistentDataFromOldPin(*SourcePin);
-					MessageLog.NotifyIntermediateObjectCreation(TargetPin, SourcePin);
+					MessageLog.NotifyIntermediatePinCreation(TargetPin, SourcePin);
 					SourcePin->BreakAllPinLinks();
 				}
 			}
@@ -367,7 +359,7 @@ void FAnimBlueprintCompiler::ProcessAnimationNode(UAnimGraphNode_Base* VisualAni
 					bConsumed = true;
 				}
 
-				UEdGraphPin* TrueSourcePin = MessageLog.FindSourceObjectTypeChecked<UEdGraphPin>(SourcePin);
+				UEdGraphPin* TrueSourcePin = MessageLog.FindSourcePin(SourcePin);
 				if (TrueSourcePin)
 				{
 					NewAnimBlueprintClass->GetDebugData().RegisterClassPropertyAssociation(TrueSourcePin, SourcePinProperty);
@@ -414,7 +406,7 @@ void FAnimBlueprintCompiler::ProcessAnimationNode(UAnimGraphNode_Base* VisualAni
 		}
 		else
 		{
-			MessageLog.Error(*FString::Printf(TEXT("A property on @@ references a non-existent %s property named %s"), VisualAnimNode), *(AnimGraphDefaultSchema->NAME_OnEvaluate.ToString()), *(EvaluationHandlerName.ToString()));
+			MessageLog.Error(*FString::Printf(TEXT("A property on @@ references a non-existent %s property named %s"), *(AnimGraphDefaultSchema->NAME_OnEvaluate.ToString()), *(EvaluationHandlerName.ToString())), VisualAnimNode);
 		}
 	}
 }
@@ -1851,7 +1843,7 @@ void FAnimBlueprintCompiler::ProcessTransitionGetter(UK2Node_TransitionRuleGette
 		check(GetterHelper->IsNodePure());
 
 		UEdGraphPin* NewReturnPin = GetterHelper->FindPinChecked(TEXT("ReturnValue"));
-		MessageLog.NotifyIntermediateObjectCreation(NewReturnPin, OutputPin);
+		MessageLog.NotifyIntermediatePinCreation(NewReturnPin, OutputPin);
 
 		NewReturnPin->CopyPersistentDataFromOldPin(*OutputPin);
 	}
diff --git a/Engine/Source/Editor/KismetCompiler/Private/KismetCompiler.cpp b/Engine/Source/Editor/KismetCompiler/Private/KismetCompiler.cpp
--- a/Engine/Source/Editor/KismetCompiler/Private/KismetCompiler.cpp
+++ b/Engine/Source/Editor/KismetCompiler/Private/KismetCompiler.cpp
@@ -2179,7 +2179,7 @@ FPinConnectionResponse FKismetCompilerContext::MovePinLinksToIntermediate(UEdGra
 	 FPinConnectionResponse ConnectionResult = K2Schema->MovePinLinks(SourcePin, IntermediatePin, true);
 
 	 CheckConnectionResponse(ConnectionResult, SourcePin.GetOwningNode());
-	 MessageLog.NotifyIntermediateObjectCreation(&IntermediatePin, &SourcePin);
+	 MessageLog.NotifyIntermediatePinCreation(&IntermediatePin, &SourcePin);
 
 	 return ConnectionResult;
 }
@@ -2190,7 +2190,7 @@ FPinConnectionResponse FKismetCompilerContext::CopyPinLinksToIntermediate(UEdGra
 	FPinConnectionResponse ConnectionResult = K2Schema->CopyPinLinks(SourcePin, IntermediatePin, true);
 
 	CheckConnectionResponse(ConnectionResult, SourcePin.GetOwningNode());
-	MessageLog.NotifyIntermediateObjectCreation(&IntermediatePin, &SourcePin);
+	MessageLog.NotifyIntermediatePinCreation(&IntermediatePin, &SourcePin);
 
 	return ConnectionResult;
 }
@@ -2200,12 +2200,8 @@ UK2Node_TemporaryVariable* FKismetCompilerContext::SpawnInternalVariable(UEdGrap
 	UK2Node_TemporaryVariable* Result = SpawnIntermediateNode<UK2Node_TemporaryVariable>(SourceNode);
 
 	Result->VariableType = FEdGraphPinType(Category, SubCategory, SubcategoryObject, bIsArray, false);
-
 	Result->AllocateDefaultPins();
 
-	// Assign the variable source information to the source object as well
-	MessageLog.NotifyIntermediateObjectCreation(Result->GetVariablePin(), SourceNode);
-
 	return Result;
 }
 
@@ -2344,7 +2340,7 @@ void FKismetCompilerContext::CreateFunctionStubForEvent(UK2Node_Event* SrcEventN
 			const FString MemberVariableName = ClassScopeNetNameMap.MakeValidName(UGSourcePin);
 
 			UEdGraphPin* DestPin = AssignmentNode->CreatePin(EGPD_Input, SourcePin->PinType, MemberVariableName);
-			MessageLog.NotifyIntermediateObjectCreation(DestPin, SourcePin);
+			MessageLog.NotifyIntermediatePinCreation(DestPin, SourcePin);
 			DestPin->MakeLinkTo(SourcePin);
 		}
 	}
@@ -2947,7 +2943,7 @@ void FKismetCompilerContext::ExpandTunnelsAndMacros(UEdGraph* SourceGraph)
 
 				for (UEdGraphPin* TunnelLinkedPin : TunnelPin->LinkedTo)
 				{
-					MessageLog.NotifyIntermediateObjectCreation(TunnelLinkedPin, SinkPin);
+					MessageLog.NotifyIntermediatePinCreation(TunnelLinkedPin, SinkPin);
 				}
 			}
 
diff --git a/Engine/Source/Editor/KismetCompiler/Private/KismetCompilerMisc.cpp b/Engine/Source/Editor/KismetCompiler/Private/KismetCompilerMisc.cpp
--- a/Engine/Source/Editor/KismetCompiler/Private/KismetCompilerMisc.cpp
+++ b/Engine/Source/Editor/KismetCompiler/Private/KismetCompilerMisc.cpp
@@ -5,6 +5,8 @@
 =============================================================================*/
 
 #include "KismetCompilerPrivatePCH.h"
+
+#include "AnimGraphNode_Base.h"
 #include "KismetCompilerMisc.h"
 #include "K2Node_EnumLiteral.h"
 #include "Kismet2/KismetReinstanceUtilities.h"
@@ -1403,9 +1405,7 @@ void FNodeHandlingFunctor::RegisterNets(FKismetFunctionContext& Context, UEdGrap
 
 //////////////////////////////////////////////////////////////////////////
 // FNetNameMapping
-
-template<>
-KISMETCOMPILER_API FString FNetNameMapping::MakeBaseName<UEdGraphPin>(const UEdGraphPin* Net)
+FString FNetNameMapping::MakeBaseName(const UEdGraphPin* Net)
 {
 	UEdGraphNode* Owner = Net->GetOwningNode();
 	FString Part1 = Owner->GetDescriptiveCompiledName();
@@ -1413,12 +1413,16 @@ KISMETCOMPILER_API FString FNetNameMapping::MakeBaseName<UEdGraphPin>(const UEdG
 	return FString::Printf(TEXT("%s_%s"), *Part1, *Net->PinName);
 }
 
-template<>
-KISMETCOMPILER_API FString FNetNameMapping::MakeBaseName<UEdGraphNode>(const UEdGraphNode* Net)
+FString FNetNameMapping::MakeBaseName(const UEdGraphNode* Net)
 {
 	return FString::Printf(TEXT("%s"), *Net->GetDescriptiveCompiledName());
 }
 
+FString FNetNameMapping::MakeBaseName(const UAnimGraphNode_Base* Net)
+{
+	return FString::Printf(TEXT("%s_%s"), *Net->GetDescriptiveCompiledName(), *Net->NodeGuid.ToString());
+}
+
 //////////////////////////////////////////////////////////////////////////
 // FKismetFunctionContext
 
@@ -1929,7 +1933,7 @@ FBPTerminal* FKismetFunctionContext::CreateLocalTerminalFromPinAutoChooseScope(U
 void FBPTerminal::CopyFromPin(UEdGraphPin* Net, const FString& NewName)
 {
 	Type = Net->PinType;
-	Source = Net;
+	SourcePin = Net;
 	Name = NewName;
 
 	bPassedByReference = Net->PinType.bIsReference;
diff --git a/Engine/Source/Editor/KismetCompiler/Private/KismetCompilerVMBackend.cpp b/Engine/Source/Editor/KismetCompiler/Private/KismetCompilerVMBackend.cpp
--- a/Engine/Source/Editor/KismetCompiler/Private/KismetCompilerVMBackend.cpp
+++ b/Engine/Source/Editor/KismetCompiler/Private/KismetCompilerVMBackend.cpp
@@ -383,18 +383,146 @@ public:
 		}
 	}
 
-	virtual void EmitTermExpr(FBPTerminal* Term, UProperty* CoerceProperty = NULL, bool bAllowStaticArray = false)
+	struct FLiteralTypeHelper
 	{
-		//@TODO: Must have a coercion type if it's a literal, because the symbol table isn't plumbed in here and the literals don't carry type information either, yay!
-		check((!Term->bIsLiteral) || (CoerceProperty != NULL));
+		static bool IsBoolean(const FEdGraphPinType* Type, const UProperty* Property)
+		{
+			if (Property)
+			{
+				return Property->IsA<UBoolProperty>();
+			}
+			return Type && (Type->PinCategory == UEdGraphSchema_K2::PC_Boolean);
+		}
+
+		static bool IsString(const FEdGraphPinType* Type, const UProperty* Property)
+		{
+			if (Property)
+			{
+				return Property->IsA<UStrProperty>();
+			}
+			return Type && (Type->PinCategory == UEdGraphSchema_K2::PC_String);
+		}
+
+		static bool IsText(const FEdGraphPinType* Type, const UProperty* Property)
+		{
+			if (Property)
+			{
+				return Property->IsA<UTextProperty>();
+			}
+			return Type && (Type->PinCategory == UEdGraphSchema_K2::PC_Text);
+		}
+
+		static bool IsFloat(const FEdGraphPinType* Type, const UProperty* Property)
+		{
+			if (Property)
+			{
+				return Property->IsA<UFloatProperty>();
+			}
+			return Type && (Type->PinCategory == UEdGraphSchema_K2::PC_Float);
+		}
+
+		static bool IsInt(const FEdGraphPinType* Type, const UProperty* Property)
+		{
+			if (Property)
+			{
+				return Property->IsA<UIntProperty>();
+			}
+			return Type && (Type->PinCategory == UEdGraphSchema_K2::PC_Int);
+		}
+
+		static bool IsInt64(const FEdGraphPinType* Type, const UProperty* Property)
+		{
+			if (Property)
+			{
+				return Property->IsA<UInt64Property>();
+			}
+			return false;
+		}
+
+		static bool IsUInt64(const FEdGraphPinType* Type, const UProperty* Property)
+		{
+			if (Property)
+			{
+				return Property->IsA<UUInt64Property>();
+			}
+			return false;
+		}
+
+		static bool IsByte(const FEdGraphPinType* Type, const UProperty* Property)
+		{
+			if (Property)
+			{
+				return Property->IsA<UByteProperty>();
+			}
+			return Type && ((Type->PinCategory == UEdGraphSchema_K2::PC_Byte) || (Type->PinCategory == UEdGraphSchema_K2::PC_Enum));
+		}
+
+		static bool IsName(const FEdGraphPinType* Type, const UProperty* Property)
+		{
+			if (Property)
+			{
+				return Property->IsA<UNameProperty>();
+			}
+			return Type && (Type->PinCategory == UEdGraphSchema_K2::PC_Name);
+		}
+
+		static bool IsStruct(const FEdGraphPinType* Type, const UProperty* Property)
+		{
+			if (Property)
+			{
+				return Property->IsA<UStructProperty>();
+			}
+			return Type && (Type->PinCategory == UEdGraphSchema_K2::PC_Struct);
+		}
+
+		static bool IsDelegate(const FEdGraphPinType* Type, const UProperty* Property)
+		{
+			if (Property)
+			{
+				return Property->IsA<UDelegateProperty>();
+			}
+			return Type && (Type->PinCategory == UEdGraphSchema_K2::PC_Delegate);
+		}
 
+		static bool IsAsset(const FEdGraphPinType* Type, const UProperty* Property)
+		{
+			if (Property)
+			{
+				return Property->IsA<UAssetObjectProperty>();
+			}
+			return Type && (Type->PinCategory == UEdGraphSchema_K2::PC_Asset);
+		}
+
+		static bool IsObject(const FEdGraphPinType* Type, const UProperty* Property)
+		{
+			if (Property)
+			{
+				return Property->IsA<UObjectPropertyBase>();
+			}
+			return Type && (Type->PinCategory == UEdGraphSchema_K2::PC_Object);
+		}
+
+		static bool IsInterface(const FEdGraphPinType* Type, const UProperty* Property)
+		{
+			if (Property)
+			{
+				return Property->IsA<UInterfaceProperty>();
+			}
+			return Type && (Type->PinCategory == UEdGraphSchema_K2::PC_Interface);
+		}
+	};
+
+	virtual void EmitTermExpr(FBPTerminal* Term, UProperty* CoerceProperty = NULL, bool bAllowStaticArray = false)
+	{
 		if (Term->bIsLiteral)
 		{
-			if (CoerceProperty->IsA(UStrProperty::StaticClass()))
+			check(!Term->Type.bIsArray || CoerceProperty);
+
+			if (FLiteralTypeHelper::IsString(&Term->Type, CoerceProperty))
 			{
 				EmitStringLiteral(Term->Name);
 			}
-			else if (CoerceProperty->IsA(UTextProperty::StaticClass()))
+			else if (FLiteralTypeHelper::IsText(&Term->Type, CoerceProperty))
 			{
 				Writer << EX_TextConst;
 				
@@ -436,13 +564,13 @@ public:
 					}
 				}
 			}
-			else if (CoerceProperty->IsA(UFloatProperty::StaticClass()))
+			else if (FLiteralTypeHelper::IsFloat(&Term->Type, CoerceProperty))
 			{
 				float Value = FCString::Atof(*(Term->Name));
 				Writer << EX_FloatConst;
 				Writer << Value;
 			}
-			else if (CoerceProperty->IsA(UIntProperty::StaticClass()))
+			else if (FLiteralTypeHelper::IsInt(&Term->Type, CoerceProperty))
 			{
 				// In certain cases (like UKismetArrayLibrary functions), we have
 				// polymorphic functions that provide their own "custom thunk" 
@@ -457,7 +585,7 @@ public:
 				// circumvented with AutoCreateRefTerm, but when it is a self 
 				// (literal) node we still end up here. So, we try to detect and 
 				// handle that case here.
-				if ((Term->Type.PinSubCategory == Schema->PN_Self) && CoerceProperty->HasAnyPropertyFlags(CPF_ReferenceParm))
+				if ((Term->Type.PinSubCategory == Schema->PN_Self) && CoerceProperty && CoerceProperty->HasAnyPropertyFlags(CPF_ReferenceParm))
 				{
 					Writer << EX_Self;
 				}
@@ -469,35 +597,32 @@ public:
 					Writer << Value;
 				}
 			}
-			else if (CoerceProperty->IsA(UInt64Property::StaticClass()))
+			else if (FLiteralTypeHelper::IsInt64(&Term->Type, CoerceProperty))
 			{
 				int64 Value = 0;
 				LexicalConversion::FromString(Value, *(Term->Name));
 				Writer << EX_Int64Const;
 				Writer << Value;
 			}
-			else if (CoerceProperty->IsA(UUInt64Property::StaticClass()))
+			else if (FLiteralTypeHelper::IsUInt64(&Term->Type, CoerceProperty))
 			{
 				uint64 Value = 0;
 				LexicalConversion::FromString(Value, *(Term->Name));
 				Writer << EX_UInt64Const;
 				Writer << Value;
 			}
-			else if (CoerceProperty->IsA(UByteProperty::StaticClass()))
+			else if (FLiteralTypeHelper::IsByte(&Term->Type, CoerceProperty))
 			{
-				UByteProperty* ByteProp = CastChecked< UByteProperty >( CoerceProperty );
 				uint8 Value = 0;
 
+				UByteProperty* ByteProp = Cast< UByteProperty >(CoerceProperty);
+				//Parameter property can represent a generic byte. we need the actual type to parse the value.
+				UEnum* EnumPtr = (ByteProp && ByteProp->Enum) ? ByteProp->Enum : Cast<UEnum>(Term->Type.PinSubCategoryObject.Get()); 
 				//Check for valid enum object reference
-				if (ByteProp->Enum)
+				if (EnumPtr)
 				{
 					//Get value from enum string
-					Value = ByteProp->Enum->GetValueByName(*(Term->Name));
-				}
-				// Allow enum literals to communicate with byte properties as literals
-				else if (UEnum* EnumPtr = Cast<UEnum>(Term->Type.PinSubCategoryObject.Get()))
-				{
-					Value = EnumPtr->GetValueByName( *(Term->Name) );
+					Value = EnumPtr->GetValueByName(*(Term->Name));
 				}
 				else
 				{
@@ -507,12 +632,12 @@ public:
 				Writer << EX_ByteConst;
 				Writer << Value;
 			}
-			else if (UBoolProperty* BoolProperty = Cast<UBoolProperty>(CoerceProperty))
+			else if (FLiteralTypeHelper::IsBoolean(&Term->Type, CoerceProperty))
 			{
 				bool bValue = Term->Name.ToBool();
 				Writer << (bValue ? EX_True : EX_False);
 			}
-			else if (UNameProperty* NameProperty = Cast<UNameProperty>(CoerceProperty))
+			else if (FLiteralTypeHelper::IsName(&Term->Type, CoerceProperty))
 			{
 				FName LiteralName(*(Term->Name));
 				Writer << EX_NameConst;
@@ -522,58 +647,62 @@ public:
 			//{
 			//	ensureMsgf(false, TEXT("Class property literals are not supported yet!"));
 			//}
-			else if (UStructProperty* StructProperty = Cast<UStructProperty>(CoerceProperty))
+			else if (FLiteralTypeHelper::IsStruct(&Term->Type, CoerceProperty))
 			{
-				if (StructProperty->Struct == VectorStruct)
+				UStructProperty* StructProperty = Cast<UStructProperty>(CoerceProperty);
+				UScriptStruct* Struct = StructProperty ? StructProperty->Struct : Cast<UScriptStruct>(Term->Type.PinSubCategoryObject.Get());
+				check(Struct);
+
+				if (Struct == VectorStruct)
 				{
 					FVector V = FVector::ZeroVector;
 					const bool bParsedUsingCustomFormat = FDefaultValueHelper::ParseVector(Term->Name, /*out*/ V);
 					if (!bParsedUsingCustomFormat)
 					{
-						StructProperty->ImportText(*Term->Name, &V, PPF_None, nullptr);
+						UStructProperty::ImportText_Static(Struct, GetPathNameSafe(StructProperty), *Term->Name, &V, PPF_None, nullptr, (FOutputDevice*)GWarn);
 					}
 					Writer << EX_VectorConst;
 					Writer << V;
 				}
-				else if (StructProperty->Struct == RotatorStruct)
+				else if (Struct == RotatorStruct)
 				{
 					FRotator R = FRotator::ZeroRotator;
 					const bool bParsedUsingCustomFormat = FDefaultValueHelper::ParseRotator(Term->Name, /*out*/ R);
 					if (!bParsedUsingCustomFormat)
 					{
-						StructProperty->ImportText(*Term->Name, &R, PPF_None, nullptr);
+						UStructProperty::ImportText_Static(Struct, GetPathNameSafe(StructProperty), *Term->Name, &R, PPF_None, nullptr, (FOutputDevice*)GWarn);
 					}
 					Writer << EX_RotationConst;
 					Writer << R;
 				}
-				else if (StructProperty->Struct == TransformStruct)
+				else if (Struct == TransformStruct)
 				{
 					FTransform T = FTransform::Identity;
 					const bool bParsedUsingCustomFormat = T.InitFromString(Term->Name);
 					if (!bParsedUsingCustomFormat)
 					{
-						StructProperty->ImportText(*Term->Name, &T, PPF_None, nullptr);
+						UStructProperty::ImportText_Static(Struct, GetPathNameSafe(StructProperty), *Term->Name, &T, PPF_None, nullptr, (FOutputDevice*)GWarn);
 					}
 					Writer << EX_TransformConst;
 					Writer << T;
 				}
 				else
 				{
-					UScriptStruct* Struct = StructProperty->Struct;
-					int32 StructSize = Struct->GetStructureSize() * StructProperty->ArrayDim;
+					const int32 ArrayDim = StructProperty ? StructProperty->ArrayDim : 1; //@TODO: is it safe when StructProperty == nullptr?
+					int32 StructSize = Struct->GetStructureSize() * ArrayDim;
 					uint8* StructData = (uint8*)FMemory_Alloca(StructSize);
-					StructProperty->InitializeValue(StructData);
-					if (!ensure(bAllowStaticArray || 1 == StructProperty->ArrayDim))
+					Struct->InitializeStruct(StructData, ArrayDim);
+					if (!ensure(bAllowStaticArray || 1 == ArrayDim))
 					{
-						UE_LOG(LogK2Compiler, Error, TEXT("Unsupported static array. Property: %s, Struct: %s"), *StructProperty->GetName(), *Struct->GetName());
+						UE_LOG(LogK2Compiler, Error, TEXT("Unsupported static array. Property: %s, Struct: %s"), *GetPathNameSafe(StructProperty), *Struct->GetName());
 					}
 					if(!FStructureEditorUtils::Fill_MakeStructureDefaultValue(Cast<UUserDefinedStruct>(Struct), StructData))
 					{
-						UE_LOG(LogK2Compiler, Warning, TEXT("MakeStructureDefaultValue parsing error. Property: %s, Struct: %s"), *StructProperty->GetName(), *Struct->GetName());
+						UE_LOG(LogK2Compiler, Warning, TEXT("MakeStructureDefaultValue parsing error. Property: %s, Struct: %s"), *GetPathNameSafe(StructProperty), *Struct->GetName());
 					}
 
 					// Assume that any errors on the import of the name string have been caught in the function call generation
-					StructProperty->ImportText(Term->Name.IsEmpty() ? TEXT("()") : *Term->Name, StructData, 0, NULL, GLog);
+					UStructProperty::ImportText_Static(Struct, GetPathNameSafe(StructProperty), Term->Name.IsEmpty() ? TEXT("()") : *Term->Name, StructData, 0, nullptr, GLog);
 
  					Writer << EX_StructConst;
 					Writer << Struct;
@@ -601,7 +730,7 @@ public:
 							EmitTermExpr(&NewTerm, Prop, true);
 						}
 					}
-
+					Struct->DestroyStruct(StructData, ArrayDim);
 					Writer << EX_EndStructConst;
 				}
 			}
@@ -638,7 +767,7 @@ public:
 				}
 				Writer << EX_EndArrayConst;
 			}
-			else if (CoerceProperty->IsA(UDelegateProperty::StaticClass()))
+			else if (FLiteralTypeHelper::IsDelegate(&Term->Type, CoerceProperty))
 			{
 				if (Term->Name == TEXT(""))
 				{
@@ -652,13 +781,13 @@ public:
 					Writer << FunctionName;
 				}
 			}
-			else if (CoerceProperty->IsA(UAssetObjectProperty::StaticClass()))
+			else if (FLiteralTypeHelper::IsAsset(&Term->Type, CoerceProperty))
 			{
 				Writer << EX_AssetConst;
 				FAssetPtr AssetPtr(Term->ObjectLiteral);
 				EmitStringLiteral(AssetPtr.GetUniqueID().ToString());
 			}
-			else if (CoerceProperty->IsA(UObjectPropertyBase::StaticClass()))
+			else if (FLiteralTypeHelper::IsObject(&Term->Type, CoerceProperty))
 			{
 				// Note: This case handles both UObjectProperty and UClassProperty
 				if (Term->Type.PinSubCategory == Schema->PN_Self)
@@ -675,7 +804,7 @@ public:
 					Writer << Term->ObjectLiteral;
 				}
 			}
-			else if (CoerceProperty->IsA(UInterfaceProperty::StaticClass()))
+			else if (FLiteralTypeHelper::IsInterface(&Term->Type, CoerceProperty))
 			{
 				if (Term->Type.PinSubCategory == Schema->PN_Self)
 				{
@@ -697,9 +826,9 @@ public:
 				if (ensure(CurrentCompilerContext))
 				{
 					FFormatNamedArguments Args;
-					Args.Add(TEXT("PropertyType"), CoerceProperty->GetClass()->GetDisplayNameText());
-					Args.Add(TEXT("PropertyName"), CoerceProperty->GetDisplayNameText());
-					CurrentCompilerContext->MessageLog.Error(*FText::Format(LOCTEXT("InvalidProperty", "It is not possible to express this type ({PropertyType}) as a literal value for the property {PropertyName} on pin @@! If it is inside a struct, you can add a Make struct node to resolve this issue!"), Args).ToString(), Term->Source);
+					Args.Add(TEXT("PropertyType"), CoerceProperty ? CoerceProperty->GetClass()->GetDisplayNameText() : FText());
+					Args.Add(TEXT("PropertyName"), CoerceProperty ? CoerceProperty->GetDisplayNameText() : FText());
+					CurrentCompilerContext->MessageLog.Error(*FText::Format(LOCTEXT("InvalidProperty", "It is not possible to express this type ({PropertyType}) as a literal value for the property {PropertyName} on pin @@! If it is inside a struct, you can add a Make struct node to resolve this issue!"), Args).ToString(), Term->SourcePin);
 				}
 			}
 		}
@@ -957,7 +1086,7 @@ public:
 		{
 			ensure(!Term->InlineGeneratedParameter->bIsJumpTarget);
 			auto TermSourceAsNode = Cast<UEdGraphNode>(Term->Source);
-			auto TermSourceAsPin = Cast<UEdGraphPin>(Term->Source);
+			auto TermSourceAsPin = Term->SourcePin;
 			UEdGraphNode* SourceNode = TermSourceAsNode ? TermSourceAsNode
 				: (TermSourceAsPin ? TermSourceAsPin->GetOwningNodeUnchecked() : nullptr);
 			if (ensure(CurrentCompilerContext && CurrentFunctionContext))
@@ -1429,8 +1558,10 @@ public:
 
 			for (auto PinContext : PinContextArray)
 			{
-				UEdGraphPin const* TrueSourcePin = Cast<UEdGraphPin const>(FunctionContext.MessageLog.FindSourceObject(PinContext));
-				if (TrueSourcePin)
+				UEdGraphPin const* TrueSourcePin = FunctionContext.MessageLog.FindSourcePin(PinContext);
+				// Source pin can be marked as pending kill if it was a generated pin that node logic decided to disown, e.g.
+				// logic in UK2Node_CallFunction to handle bWantsEnumToExecExpansion:
+				if (TrueSourcePin && !TrueSourcePin->IsPendingKill())
 				{
 					ClassBeingBuilt->GetDebugData().RegisterPinToCodeAssociation(TrueSourcePin, FunctionContext.Function, Offset);
 				}
diff --git a/Engine/Source/Editor/KismetCompiler/Private/UserDefinedStructureCompilerUtils.cpp b/Engine/Source/Editor/KismetCompiler/Private/UserDefinedStructureCompilerUtils.cpp
--- a/Engine/Source/Editor/KismetCompiler/Private/UserDefinedStructureCompilerUtils.cpp
+++ b/Engine/Source/Editor/KismetCompiler/Private/UserDefinedStructureCompilerUtils.cpp
@@ -109,37 +109,42 @@ struct FUserDefinedStructureCompilerInner
 			EditorData->CleanDefaultInstance();
 		}
 
-		const FString TransientString = FString::Printf(TEXT("TRASHSTRUCT_%s"), *StructToClean->GetName());
-		const FName TransientName = MakeUniqueObjectName(GetTransientPackage(), UUserDefinedStruct::StaticClass(), FName(*TransientString));
-		UUserDefinedStruct* TransientStruct = NewObject<UUserDefinedStruct>(GetTransientPackage(), TransientName, RF_Public | RF_Transient);
-		
-		TArray<UObject*> SubObjects;
-		GetObjectsWithOuter(StructToClean, SubObjects, true);
-		SubObjects.Remove(StructToClean->EditorData);
-		for( auto SubObjIt = SubObjects.CreateIterator(); SubObjIt; ++SubObjIt )
+		UUserDefinedStruct* TransientStruct = nullptr;
+
+		if (FStructureEditorUtils::FStructEditorManager::ActiveChange != FStructureEditorUtils::EStructureEditorChangeInfo::DefaultValueChanged)
 		{
-			UObject* CurrSubObj = *SubObjIt;
-			CurrSubObj->Rename(NULL, TransientStruct, REN_DontCreateRedirectors);
-			if( UProperty* Prop = Cast<UProperty>(CurrSubObj) )
+			const FString TransientString = FString::Printf(TEXT("TRASHSTRUCT_%s"), *StructToClean->GetName());
+			const FName TransientName = MakeUniqueObjectName(GetTransientPackage(), UUserDefinedStruct::StaticClass(), FName(*TransientString));
+			TransientStruct = NewObject<UUserDefinedStruct>(GetTransientPackage(), TransientName, RF_Public | RF_Transient);
+
+			TArray<UObject*> SubObjects;
+			GetObjectsWithOuter(StructToClean, SubObjects, true);
+			SubObjects.Remove(StructToClean->EditorData);
+			for (auto SubObjIt = SubObjects.CreateIterator(); SubObjIt; ++SubObjIt)
 			{
-				FKismetCompilerUtilities::InvalidatePropertyExport(Prop);
-			}
-			else
-			{
-				FLinkerLoad::InvalidateExport(CurrSubObj);
+				UObject* CurrSubObj = *SubObjIt;
+				CurrSubObj->Rename(NULL, TransientStruct, REN_DontCreateRedirectors);
+				if (UProperty* Prop = Cast<UProperty>(CurrSubObj))
+				{
+					FKismetCompilerUtilities::InvalidatePropertyExport(Prop);
+				}
+				else
+				{
+					FLinkerLoad::InvalidateExport(CurrSubObj);
+				}
 			}
-		}
 
-		StructToClean->SetSuperStruct(NULL);
-		StructToClean->Children = NULL;
-		StructToClean->Script.Empty();
-		StructToClean->MinAlignment = 0;
-		StructToClean->RefLink = NULL;
-		StructToClean->PropertyLink = NULL;
-		StructToClean->DestructorLink = NULL;
-		StructToClean->ScriptObjectReferences.Empty();
-		StructToClean->PropertyLink = NULL;
-		StructToClean->ErrorMessage.Empty();
+			StructToClean->SetSuperStruct(NULL);
+			StructToClean->Children = NULL;
+			StructToClean->Script.Empty();
+			StructToClean->MinAlignment = 0;
+			StructToClean->RefLink = NULL;
+			StructToClean->PropertyLink = NULL;
+			StructToClean->DestructorLink = NULL;
+			StructToClean->ScriptObjectReferences.Empty();
+			StructToClean->PropertyLink = NULL;
+			StructToClean->ErrorMessage.Empty();
+		}
 
 		return TransientStruct;
 	}
@@ -172,45 +177,62 @@ struct FUserDefinedStructureCompilerInner
 				continue;
 			}
 
-			UProperty* NewProperty = FKismetCompilerUtilities::CreatePropertyOnScope(Struct, VarDesc.VarName, VarType, NULL, 0, Schema, MessageLog);
-			if (NewProperty != NULL)
+			UProperty* VarProperty = nullptr;
+
+			bool bIsNewVariable = false;
+			if (FStructureEditorUtils::FStructEditorManager::ActiveChange == FStructureEditorUtils::EStructureEditorChangeInfo::DefaultValueChanged)
 			{
-				NewProperty->SetFlags(RF_LoadCompleted);
-				FKismetCompilerUtilities::LinkAddedProperty(Struct, NewProperty);
+				VarProperty = FindField<UProperty>(Struct, VarDesc.VarName);
+				if (!ensureMsgf(VarProperty, TEXT("Could not find the expected property (%s); was the struct (%s) unexpectedly sanitized?"), *VarDesc.VarName.ToString(), *Struct->GetName()))
+				{
+					VarProperty = FKismetCompilerUtilities::CreatePropertyOnScope(Struct, VarDesc.VarName, VarType, NULL, 0, Schema, MessageLog);
+					bIsNewVariable = true;
+				}
 			}
 			else
+			{
+				VarProperty = FKismetCompilerUtilities::CreatePropertyOnScope(Struct, VarDesc.VarName, VarType, NULL, 0, Schema, MessageLog);
+				bIsNewVariable = true;
+			}
+
+			if (VarProperty == nullptr)
 			{
 				LogError(Struct, MessageLog, FString::Printf(*LOCTEXT("VariableInvalidType_Error", "The variable %s declared in %s has an invalid type %s").ToString(),
 					*VarDesc.VarName.ToString(), *Struct->GetName(), *UEdGraphSchema_K2::TypeToText(VarType).ToString()));
 				continue;
 			}
-
-			NewProperty->SetPropertyFlags(CPF_Edit | CPF_BlueprintVisible);
+			else if (bIsNewVariable)
+			{
+				VarProperty->SetFlags(RF_LoadCompleted);
+				FKismetCompilerUtilities::LinkAddedProperty(Struct, VarProperty);				
+			}
+			
+			VarProperty->SetPropertyFlags(CPF_Edit | CPF_BlueprintVisible);
 			if (VarDesc.bDontEditoOnInstance)
 			{
-				NewProperty->SetPropertyFlags(CPF_DisableEditOnInstance);
+				VarProperty->SetPropertyFlags(CPF_DisableEditOnInstance);
 			}
 			if (VarDesc.bEnableMultiLineText)
 			{
-				NewProperty->SetMetaData("MultiLine", TEXT("true"));
+				VarProperty->SetMetaData("MultiLine", TEXT("true"));
 			}
 			if (VarDesc.bEnable3dWidget)
 			{
-				NewProperty->SetMetaData(FEdMode::MD_MakeEditWidget, TEXT("true"));
+				VarProperty->SetMetaData(FEdMode::MD_MakeEditWidget, TEXT("true"));
 			}
-			NewProperty->SetMetaData(TEXT("DisplayName"), *VarDesc.FriendlyName);
-			NewProperty->SetMetaData(FBlueprintMetadata::MD_Tooltip, *VarDesc.ToolTip);
-			NewProperty->RepNotifyFunc = NAME_None;
+			VarProperty->SetMetaData(TEXT("DisplayName"), *VarDesc.FriendlyName);
+			VarProperty->SetMetaData(FBlueprintMetadata::MD_Tooltip, *VarDesc.ToolTip);
+			VarProperty->RepNotifyFunc = NAME_None;
 
 			if (!VarDesc.DefaultValue.IsEmpty())
 			{
-				NewProperty->SetMetaData(TEXT("MakeStructureDefaultValue"), *VarDesc.DefaultValue);
+				VarProperty->SetMetaData(TEXT("MakeStructureDefaultValue"), *VarDesc.DefaultValue);
 			}
 			VarDesc.CurrentDefaultValue = VarDesc.DefaultValue;
 
 			VarDesc.bInvalidMember = false;
 
-			if (NewProperty->HasAnyPropertyFlags(CPF_InstancedReference | CPF_ContainsInstancedReference))
+			if (VarProperty->HasAnyPropertyFlags(CPF_InstancedReference | CPF_ContainsInstancedReference))
 			{
 				Struct->StructFlags = EStructFlags(Struct->StructFlags | STRUCT_HasInstancedReference);
 			}
@@ -223,12 +245,12 @@ struct FUserDefinedStructureCompilerInner
 				{
 					// prevent Actor variables from having default values (because Blueprint templates are library elements that can 
 					// bridge multiple levels and different levels might not have the actor that the default is referencing).
-					NewProperty->PropertyFlags |= CPF_DisableEditOnTemplate;
+					VarProperty->PropertyFlags |= CPF_DisableEditOnTemplate;
 				}
 				else
 				{
 					// clear the disable-default-value flag that might have been present (if this was an AActor variable before)
-					NewProperty->PropertyFlags &= ~(CPF_DisableEditOnTemplate);
+					VarProperty->PropertyFlags &= ~(CPF_DisableEditOnTemplate);
 				}
 			}
 		}
@@ -242,7 +264,7 @@ struct FUserDefinedStructureCompilerInner
 		Struct->SetMetaData(FBlueprintMetadata::MD_Tooltip, *FStructureEditorUtils::GetTooltip(Struct));
 
 		auto EditorData = CastChecked<UUserDefinedStructEditorData>(Struct->EditorData);
-		Struct->SetSuperStruct(EditorData->NativeBase);
+		Struct->SetSuperStruct(EditorData->NativeBase);	
 
 		CreateVariables(Struct, K2Schema, MessageLog);
 
diff --git a/Engine/Source/Editor/KismetCompiler/Public/BPTerminal.h b/Engine/Source/Editor/KismetCompiler/Public/BPTerminal.h
--- a/Engine/Source/Editor/KismetCompiler/Public/BPTerminal.h
+++ b/Engine/Source/Editor/KismetCompiler/Public/BPTerminal.h
@@ -16,9 +16,12 @@ struct FBPTerminal
 	bool bIsSavePersistent;
 	bool bPassedByReference;
 
-	// Source node (e.g., pin)
+	// Source node
 	UObject* Source;
 
+	// Source pin
+	UEdGraphPin* SourcePin;
+
 	// Context->
 	FBPTerminal* Context;
 
@@ -43,6 +46,7 @@ struct FBPTerminal
 		, bIsSavePersistent(false)
 		, bPassedByReference(false)
 		, Source(nullptr)
+		, SourcePin(nullptr)
 		, Context(nullptr)
 		, AssociatedVarProperty(nullptr)
 		, ObjectLiteral(nullptr)
diff --git a/Engine/Source/Editor/KismetCompiler/Public/KismetCompilerMisc.h b/Engine/Source/Editor/KismetCompiler/Public/KismetCompilerMisc.h
--- a/Engine/Source/Editor/KismetCompiler/Public/KismetCompilerMisc.h
+++ b/Engine/Source/Editor/KismetCompiler/Public/KismetCompilerMisc.h
@@ -184,11 +184,13 @@ struct FNetNameMapping
 {
 public:
 	TMap<const UObject*, FString> NetToName;
+	TMap<UEdGraphPin*, FString> NetPinToName;
 	TMap<FString, int32> BaseNameToCount;
 
 public:
-	template<typename NetType>
-	static FString MakeBaseName(const NetType* Net);
+	KISMETCOMPILER_API static FString MakeBaseName(const UEdGraphNode* Net);
+	KISMETCOMPILER_API static FString MakeBaseName(const UEdGraphPin* Net);;
+	KISMETCOMPILER_API static FString MakeBaseName(const UAnimGraphNode_Base* Net);
 
 	// Come up with a valid, unique (within the scope of NetNameMap) name based on an existing Net object.
 	// The resulting name is stable across multiple calls if given the same pointer.
@@ -217,5 +219,29 @@ public:
 			return NetName;
 		}
 	}
+
+	FString MakeValidName(UEdGraphPin* Net)
+	{
+		if (FString* Result = NetPinToName.Find(Net))
+		{
+			return *Result;
+		}
+		else
+		{
+			FString NetName = MakeBaseName(Net);
+			FNodeHandlingFunctor::SanitizeName(NetName);
+
+			int32& ExistingCount = BaseNameToCount.FindOrAdd(NetName);
+			++ExistingCount;
+			if (ExistingCount > 1)
+			{
+				NetName += FString::FromInt(ExistingCount);
+			}
+
+			NetPinToName.Add(Net, NetName);
+
+			return NetName;
+		}
+	}
 };
 
