Copying //UE4/Dev-Blueprints to //UE4/Dev-Main (Source: //UE4/Dev-Blueprints @ 3042808)

diff --git a/Engine/Source/Editor/BlueprintGraph/Classes/EdGraphSchema_K2.h b/Engine/Source/Editor/BlueprintGraph/Classes/EdGraphSchema_K2.h
--- a/Engine/Source/Editor/BlueprintGraph/Classes/EdGraphSchema_K2.h
+++ b/Engine/Source/Editor/BlueprintGraph/Classes/EdGraphSchema_K2.h
@@ -218,6 +218,20 @@ enum class EObjectReferenceType : uint8
 	AllTypes		= 0x0f,
 };
 
+/**
+* Filter flags for GetVariableTypeTree
+*/
+enum class ETypeTreeFilter : uint8
+{
+	None			= 0x00, // No Exec or Wildcards
+	AllowExec		= 0x01, // Include Executable pins
+	AllowWildcard	= 0x02, // Include Wildcard pins
+	IndexTypesOnly	= 0x04, // Exclude all pins that aren't index types
+	RootTypesOnly	= 0x08	// Exclude all pins that aren't root types
+};
+
+ENUM_CLASS_FLAGS(ETypeTreeFilter);
+
 struct FTypesDatabase;
 
 UCLASS(config=Editor)
@@ -834,19 +848,9 @@ public:
 	 * Get the type tree for all of the property types valid for this schema
 	 *
 	 * @param	TypeTree		The array that will contain the type tree hierarchy for this schema upon returning
-	 * @param	bAllowExec		Whether or not to add the exec type to the type tree
-	 * @param	bAllowWildcard	Whether or not to add the wildcard type to the type tree
-	 */
-	void GetVariableTypeTree(TArray< TSharedPtr<FPinTypeTreeInfo> >& TypeTree, bool bAllowExec, bool bAllowWildcard) const;
-
-	/**
-	 * Get the type tree for the index property types valid for this schema
-	 *
-	 * @param	TypeTree	The array that will contain the type tree hierarchy for this schema upon returning
-	 * @param	bAllowExec	Whether or not to add the exec type to the type tree
-	 * @param	bAllowExec	Whether or not to add the exec type to the type tree
+	 * @param	TypeTreeFilter	ETypeTreeFilter flags that determine how the TypeTree is populated.
 	 */
-	void GetVariableIndexTypeTree( TArray< TSharedPtr<FPinTypeTreeInfo> >& TypeTree, bool bAllowExec, bool bAllowWildcard ) const;
+	void GetVariableTypeTree(TArray< TSharedPtr<FPinTypeTreeInfo> >& TypeTree, ETypeTreeFilter TypeTreeFilter = ETypeTreeFilter::None) const;
 
 	/**
 	 * Returns whether or not the specified type has valid subtypes available
@@ -1013,9 +1017,6 @@ public:
 	static void GetReplaceVariableMenu(class FMenuBuilder& MenuBuilder, class UK2Node_Variable* Variable, UBlueprint* OwnerBlueprint, bool bReplaceExistingVariable = false);
 
 private:
-
-	void GetVariableTypeTreeImpl(TArray< TSharedPtr<FPinTypeTreeInfo> >& TypeTree, bool bAllowExec, bool bAllowWildCard, bool bIndexTypesOnly) const;
-
 	/**
 	 * Returns true if the specified function has any out parameters
 	 * @param [in] Function	The function to check for out parameters
diff --git a/Engine/Source/Editor/BlueprintGraph/Classes/K2Node_TunnelBoundary.h b/Engine/Source/Editor/BlueprintGraph/Classes/K2Node_TunnelBoundary.h
new file mode 100644
--- /dev/null
+++ b/Engine/Source/Editor/BlueprintGraph/Classes/K2Node_TunnelBoundary.h
@@ -0,0 +1,50 @@
+// Copyright 1998-2016 Epic Games, Inc. All Rights Reserved.
+
+#pragma once
+#include "K2Node.h"
+#include "K2Node_TunnelBoundary.generated.h"
+
+UCLASS()
+class BLUEPRINTGRAPH_API UK2Node_TunnelBoundary : public UK2Node
+{
+	GENERATED_UCLASS_BODY()
+
+	/** Tunnel Graph Name */
+	UPROPERTY()
+	FName TunnelGraphName;
+
+	/** Final exit site */
+	UPROPERTY()
+	FEdGraphPinReference FinalExitSite;
+
+	/** Wether or not this is an exit point */
+	bool bEntryNode;
+
+public:
+
+	//~ Begin of UK2Node Interface
+	virtual FText GetNodeTitle(ENodeTitleType::Type TitleType) const override;
+	virtual class FNodeHandlingFunctor* CreateNodeHandler(class FKismetCompilerContext& CompilerContext) const override;
+	//~ End UK2Node Interface
+
+	/** Creates Boundary nodes for a tunnel graph */
+	static void CreateBoundaryNodesForGraph(UEdGraph* TunnelGraph, class FCompilerResultsLog& MessageLog);
+
+	/** Creates Boundary nodes for a tunnel instance */
+	static void CreateBoundaryNodesForTunnelInstance(UK2Node_Tunnel* TunnelInstance, UEdGraph* TunnelGraph, class FCompilerResultsLog& MessageLog);
+
+	/** Checks if the tunnel node is a pure tunnel rather than a tunnel instance */
+	static bool IsPureTunnel(UK2Node_Tunnel* Tunnel);
+
+protected:
+
+	/** Wires up the tunnel entry boundries */
+	void WireUpTunnelEntry(UK2Node_Tunnel* TunnelInstance, UEdGraphPin* TunnelPin, FCompilerResultsLog& MessageLog);
+
+	/** Wires up the tunnel exit boundries */
+	void WireUpTunnelExit(UK2Node_Tunnel* TunnelInstance, UEdGraphPin* TunnelPin, FCompilerResultsLog& MessageLog);
+
+	/** Determines the tunnel graph name */
+	void DetermineTunnelGraphName(UK2Node_Tunnel* TunnelInstance);
+
+};
diff --git a/Engine/Source/Editor/BlueprintGraph/Private/BlueprintComponentNodeSpawner.cpp b/Engine/Source/Editor/BlueprintGraph/Private/BlueprintComponentNodeSpawner.cpp
--- a/Engine/Source/Editor/BlueprintGraph/Private/BlueprintComponentNodeSpawner.cpp
+++ b/Engine/Source/Editor/BlueprintGraph/Private/BlueprintComponentNodeSpawner.cpp
@@ -147,6 +147,10 @@ UEdGraphNode* UBlueprintComponentNodeSpawner::Invoke(UEdGraph* ParentGraph, FBin
 	FCustomizeNodeDelegate PostSpawnDelegate = FCustomizeNodeDelegate::CreateStatic(PostSpawnLambda, CustomizeNodeDelegate);
 	// let SpawnNode() allocate default pins (so we can modify them)
 	UK2Node_AddComponent* NewNode = Super::SpawnNode<UK2Node_AddComponent>(NodeClass, ParentGraph, FBindingSet(), Location, PostSpawnDelegate);
+	if (NewNode->Pins.Num() == 0)
+	{
+		NewNode->AllocateDefaultPins();
+	}
 
 	// set the return type to be the type of the template
 	UEdGraphPin* ReturnPin = NewNode->GetReturnValuePin();
diff --git a/Engine/Source/Editor/BlueprintGraph/Private/EdGraphSchema_K2.cpp b/Engine/Source/Editor/BlueprintGraph/Private/EdGraphSchema_K2.cpp
--- a/Engine/Source/Editor/BlueprintGraph/Private/EdGraphSchema_K2.cpp
+++ b/Engine/Source/Editor/BlueprintGraph/Private/EdGraphSchema_K2.cpp
@@ -3437,13 +3437,13 @@ FText UEdGraphSchema_K2::TypeToText(const FEdGraphPinType& Type)
 	return PropertyText;
 }
 
-void UEdGraphSchema_K2::GetVariableTypeTree(TArray< TSharedPtr<FPinTypeTreeInfo> >& TypeTree, bool bAllowExec, bool bAllowWildcard) const
+void UEdGraphSchema_K2::GetVariableTypeTree(TArray< TSharedPtr<FPinTypeTreeInfo> >& TypeTree, ETypeTreeFilter TypeTreeFilter) const
 {
-	GetVariableTypeTreeImpl(TypeTree, bAllowExec, bAllowWildcard, false);
-}
+	bool bAllowExec = (TypeTreeFilter & ETypeTreeFilter::AllowExec) == ETypeTreeFilter::AllowExec;
+	bool bAllowWildCard = (TypeTreeFilter & ETypeTreeFilter::AllowWildcard) == ETypeTreeFilter::AllowWildcard;
+	bool bIndexTypesOnly = (TypeTreeFilter & ETypeTreeFilter::IndexTypesOnly) == ETypeTreeFilter::IndexTypesOnly;
+	bool bRootTypesOnly = (TypeTreeFilter & ETypeTreeFilter::RootTypesOnly) == ETypeTreeFilter::RootTypesOnly;
 
-void UEdGraphSchema_K2::GetVariableTypeTreeImpl(TArray< TSharedPtr<FPinTypeTreeInfo> >& TypeTree, bool bAllowExec, bool bAllowWildCard, bool bIndexTypesOnly) const
-{
 #ifdef SCHEMA_K2_GETVARIABLETYPETREE_LOG_TIME
 	static_assert(false, "Macro redefinition.");
 #endif
@@ -3454,13 +3454,21 @@ void UEdGraphSchema_K2::GetVariableTypeTreeImpl(TArray< TSharedPtr<FPinTypeTreeI
 #endif //SCHEMA_K2_GETVARIABLETYPETREE_LOG_TIME
 
 	FTypesDatabase TypesDatabase;
-	FGatherTypesHelper::FillLoadedTypesDatabase(TypesDatabase, bIndexTypesOnly);
+	FTypesDatabase* TypesDatabasePtr = nullptr;
+	if (!bRootTypesOnly)
+	{
+		TypesDatabasePtr = &TypesDatabase;
+		FGatherTypesHelper::FillLoadedTypesDatabase(TypesDatabase, bIndexTypesOnly);
+	}
 
 #if SCHEMA_K2_GETVARIABLETYPETREE_LOG_TIME
 	const auto DatabaseLoadedTime = FPlatformTime::Seconds();
 #endif //SCHEMA_K2_GETVARIABLETYPETREE_LOG_TIME
 
-	FGatherTypesHelper::FillUnLoadedTypesDatabase(TypesDatabase, bIndexTypesOnly);
+	if (!bRootTypesOnly)
+	{
+		FGatherTypesHelper::FillUnLoadedTypesDatabase(TypesDatabase, bIndexTypesOnly);
+	}
 
 #if SCHEMA_K2_GETVARIABLETYPETREE_LOG_TIME
 	const auto DatabaseUnLoadedTime = FPlatformTime::Seconds();
@@ -3485,9 +3493,12 @@ void UEdGraphSchema_K2::GetVariableTypeTreeImpl(TArray< TSharedPtr<FPinTypeTreeI
 		TypeTree.Add(MakeShareable(new FPinTypeTreeInfo(GetCategoryText(PC_Text, true), PC_Text, this, LOCTEXT("TextType", "A localizable text string"))));
 
 		// Add in special first-class struct types
-		TypeTree.Add(MakeShareable(new FPinTypeTreeInfo(PC_Struct, TBaseStructure<FVector>::Get(), LOCTEXT("VectorType", "A 3D vector"))));
-		TypeTree.Add(MakeShareable(new FPinTypeTreeInfo(PC_Struct, TBaseStructure<FRotator>::Get(), LOCTEXT("RotatorType", "A 3D rotation"))));
-		TypeTree.Add(MakeShareable(new FPinTypeTreeInfo(PC_Struct, TBaseStructure<FTransform>::Get(), LOCTEXT("TransformType", "A 3D transformation, including translation, rotation and 3D scale."))));
+		if (!bRootTypesOnly)
+		{
+			TypeTree.Add(MakeShareable(new FPinTypeTreeInfo(PC_Struct, TBaseStructure<FVector>::Get(), LOCTEXT("VectorType", "A 3D vector"))));
+			TypeTree.Add(MakeShareable(new FPinTypeTreeInfo(PC_Struct, TBaseStructure<FRotator>::Get(), LOCTEXT("RotatorType", "A 3D rotation"))));
+			TypeTree.Add(MakeShareable(new FPinTypeTreeInfo(PC_Struct, TBaseStructure<FTransform>::Get(), LOCTEXT("TransformType", "A 3D transformation, including translation, rotation and 3D scale."))));
+		}
 	}
 	// Add wildcard type
 	if (bAllowWildCard)
@@ -3495,14 +3506,20 @@ void UEdGraphSchema_K2::GetVariableTypeTreeImpl(TArray< TSharedPtr<FPinTypeTreeI
 		TypeTree.Add( MakeShareable( new FPinTypeTreeInfo(GetCategoryText(PC_Wildcard, true), PC_Wildcard, this, LOCTEXT("WildcardType", "Wildcard type (unspecified).")) ) );
 	}
 
-	FTypesDatabase* TypesDatabasePtr = &TypesDatabase;
 	// Add the types that have subtrees
 	if (!bIndexTypesOnly)
 	{
 		TypeTree.Add(MakeShareable(new FPinTypeTreeInfo(GetCategoryText(PC_Struct, true), PC_Struct, this, LOCTEXT("StructType", "Struct (value) types."), true, TypesDatabasePtr)));
 		TypeTree.Add(MakeShareable(new FPinTypeTreeInfo(GetCategoryText(PC_Interface, true), PC_Interface, this, LOCTEXT("InterfaceType", "Interface pointer."), true, TypesDatabasePtr)));
 
-		TypeTree.Add(MakeShareable(new FPinTypeTreeInfo(GetCategoryText(AllObjectTypes, true), AllObjectTypes, this, LOCTEXT("ObjectType", "Object pointer."), true, TypesDatabasePtr)));
+		if (!bRootTypesOnly)
+		{
+			TypeTree.Add(MakeShareable(new FPinTypeTreeInfo(GetCategoryText(AllObjectTypes, true), AllObjectTypes, this, LOCTEXT("ObjectType", "Object pointer."), true, TypesDatabasePtr)));
+		}
+		else
+		{
+			TypeTree.Add(MakeShareable(new FPinTypeTreeInfo(GetCategoryText(PC_Object, true), PC_Object, this, LOCTEXT("ObjectType", "Object pointer."), true, TypesDatabasePtr)));
+		}
 	}
 	TypeTree.Add( MakeShareable( new FPinTypeTreeInfo(GetCategoryText(PC_Enum, true), PC_Enum, this, LOCTEXT("EnumType", "Enumeration types."), true, TypesDatabasePtr) ) );
 
@@ -3513,11 +3530,6 @@ void UEdGraphSchema_K2::GetVariableTypeTreeImpl(TArray< TSharedPtr<FPinTypeTreeI
 #undef SCHEMA_K2_GETVARIABLETYPETREE_LOG_TIME
 }
 
-void UEdGraphSchema_K2::GetVariableIndexTypeTree( TArray< TSharedPtr<FPinTypeTreeInfo> >& TypeTree, bool bAllowExec, bool bAllowWildcard ) const
-{
-	GetVariableTypeTreeImpl(TypeTree, bAllowExec, bAllowWildcard, true);
-}
-
 bool UEdGraphSchema_K2::DoesTypeHaveSubtypes(const FString& Category) const
 {
 	return (Category == PC_Struct) || (Category == PC_Object) || (Category == PC_Asset) || (Category == PC_AssetClass) || (Category == PC_Interface) || (Category == PC_Class) || (Category == PC_Enum) || (Category == AllObjectTypes);
diff --git a/Engine/Source/Editor/BlueprintGraph/Private/K2Node.cpp b/Engine/Source/Editor/BlueprintGraph/Private/K2Node.cpp
--- a/Engine/Source/Editor/BlueprintGraph/Private/K2Node.cpp
+++ b/Engine/Source/Editor/BlueprintGraph/Private/K2Node.cpp
@@ -655,7 +655,7 @@ void UK2Node::ReconstructSinglePin(UEdGraphPin* NewPin, UEdGraphPin* OldPin, ERe
 		}
 	}
 
-	OldPin->InvalidateAndTrash();;
+	OldPin->MarkPendingKill();
 }
 
 void UK2Node::RewireOldPinsToNewPins(TArray<UEdGraphPin*>& InOldPins, TArray<UEdGraphPin*>& InNewPins)
diff --git a/Engine/Source/Editor/BlueprintGraph/Private/K2Node_ActorBoundEvent.cpp b/Engine/Source/Editor/BlueprintGraph/Private/K2Node_ActorBoundEvent.cpp
--- a/Engine/Source/Editor/BlueprintGraph/Private/K2Node_ActorBoundEvent.cpp
+++ b/Engine/Source/Editor/BlueprintGraph/Private/K2Node_ActorBoundEvent.cpp
@@ -88,7 +88,9 @@ bool UK2Node_ActorBoundEvent::CanPasteHere(const UEdGraph* TargetGraph) const
 	bool bDisallowPaste = !Super::CanPasteHere(TargetGraph);
 	if (!bDisallowPaste)
 	{
-		void* LevelBP = GetReferencedLevelActor()->GetLevel()->GetLevelScriptBlueprint(true);
+		const AActor* ReferencedLevelActor = GetReferencedLevelActor();
+		ULevel* Level = ReferencedLevelActor ? ReferencedLevelActor->GetLevel() : nullptr;
+		const UBlueprint* LevelBP = Level ? Level->GetLevelScriptBlueprint(true) : nullptr;
 		if (FBlueprintEditorUtils::FindBlueprintForGraph(TargetGraph) == LevelBP)
 		{
 			if (const UK2Node_Event* PreExistingNode = FKismetEditorUtilities::FindBoundEventForActor(GetReferencedLevelActor(), DelegatePropertyName))
diff --git a/Engine/Source/Editor/BlueprintGraph/Private/K2Node_BaseAsyncTask.cpp b/Engine/Source/Editor/BlueprintGraph/Private/K2Node_BaseAsyncTask.cpp
--- a/Engine/Source/Editor/BlueprintGraph/Private/K2Node_BaseAsyncTask.cpp
+++ b/Engine/Source/Editor/BlueprintGraph/Private/K2Node_BaseAsyncTask.cpp
@@ -217,7 +217,7 @@ bool UK2Node_BaseAsyncTask::FBaseAsyncTaskHelper::HandleDelegateImplementation(
 		return false;
 	}
 
-	UK2Node_CustomEvent* CurrentCENode = CompilerContext.SpawnIntermediateNode<UK2Node_CustomEvent>(CurrentNode, SourceGraph);
+	UK2Node_CustomEvent* CurrentCENode = CompilerContext.SpawnIntermediateEventNode<UK2Node_CustomEvent>(CurrentNode, PinForCurrentDelegateProperty, SourceGraph);
 	{
 	UK2Node_AddDelegate* AddDelegateNode = CompilerContext.SpawnIntermediateNode<UK2Node_AddDelegate>(CurrentNode, SourceGraph);
 	AddDelegateNode->SetFromProperty(CurrentProperty, false);
diff --git a/Engine/Source/Editor/BlueprintGraph/Private/K2Node_CustomEvent.cpp b/Engine/Source/Editor/BlueprintGraph/Private/K2Node_CustomEvent.cpp
--- a/Engine/Source/Editor/BlueprintGraph/Private/K2Node_CustomEvent.cpp
+++ b/Engine/Source/Editor/BlueprintGraph/Private/K2Node_CustomEvent.cpp
@@ -183,7 +183,7 @@ bool UK2Node_CustomEvent::CanCreateUserDefinedPin(const FEdGraphPinType& InPinTy
 	const UEdGraphSchema_K2* Schema = GetDefault<UEdGraphSchema_K2>();
 	if(InDesiredDirection == EGPD_Input)
 	{
-		OutErrorMessage = NSLOCTEXT("K2Node", "AddInputPinError", "Cannot add input pins to function result node!");
+		OutErrorMessage = NSLOCTEXT("K2Node", "AddInputPinError", "Cannot add input pins to custom event node!");
 		return false;
 	}
 	else if (InPinType.PinCategory == Schema->PC_Exec && !CanModifyExecutionWires())
@@ -191,6 +191,29 @@ bool UK2Node_CustomEvent::CanCreateUserDefinedPin(const FEdGraphPinType& InPinTy
 		OutErrorMessage = LOCTEXT("MultipleExecPinError", "Cannot support more exec pins!");
 		return false;
 	}
+	else
+	{
+		TArray<TSharedPtr<UEdGraphSchema_K2::FPinTypeTreeInfo>> TypeTree;
+		Schema->GetVariableTypeTree(TypeTree, ETypeTreeFilter::RootTypesOnly);
+
+		bool bIsValid = false;
+		for (TSharedPtr<UEdGraphSchema_K2::FPinTypeTreeInfo>& TypeInfo : TypeTree)
+		{
+			FEdGraphPinType CurrentType = TypeInfo->GetPinType(false);
+			// only concerned with the list of categories
+			if (CurrentType.PinCategory == InPinType.PinCategory)
+			{
+				bIsValid = true;
+				break;
+			}
+		}
+
+		if (!bIsValid)
+		{
+			OutErrorMessage = LOCTEXT("AddInputPinError", "Cannot add pins of this type to custom event node!");
+			return false;
+		}
+	}
 
 	return true;
 }
@@ -383,6 +406,7 @@ void UK2Node_CustomEvent::ReconstructNode()
 				TSharedPtr<FUserPinInfo> NewPinInfo = MakeShareable( new FUserPinInfo() );
 				NewPinInfo->PinName = NewPinName;
 				NewPinInfo->PinType = PinType;
+				NewPinInfo->DesiredPinDirection = EGPD_Output;
 				UserDefinedPins.Add(NewPinInfo);
 			}
 		}
diff --git a/Engine/Source/Editor/BlueprintGraph/Private/K2Node_DelegateSet.cpp b/Engine/Source/Editor/BlueprintGraph/Private/K2Node_DelegateSet.cpp
--- a/Engine/Source/Editor/BlueprintGraph/Private/K2Node_DelegateSet.cpp
+++ b/Engine/Source/Editor/BlueprintGraph/Private/K2Node_DelegateSet.cpp
@@ -271,7 +271,7 @@ void UK2Node_DelegateSet::ExpandNode(class FKismetCompilerContext& CompilerConte
 			const UEdGraphSchema_K2* Schema = CompilerContext.GetSchema();
 
 			// First, create an event node matching the delegate signature
-			UK2Node_Event* DelegateEvent = CompilerContext.SpawnIntermediateNode<UK2Node_Event>(this, SourceGraph);
+			UK2Node_Event* DelegateEvent = CompilerContext.SpawnIntermediateEventNode<UK2Node_Event>(this, nullptr, SourceGraph);
 			DelegateEvent->EventReference.SetFromField<UFunction>(TargetFunction, false);
 			DelegateEvent->CustomFunctionName = GetDelegateTargetEntryPointName();
 			DelegateEvent->bInternalEvent = true;
diff --git a/Engine/Source/Editor/BlueprintGraph/Private/K2Node_EditablePinBase.cpp b/Engine/Source/Editor/BlueprintGraph/Private/K2Node_EditablePinBase.cpp
--- a/Engine/Source/Editor/BlueprintGraph/Private/K2Node_EditablePinBase.cpp
+++ b/Engine/Source/Editor/BlueprintGraph/Private/K2Node_EditablePinBase.cpp
@@ -15,7 +15,11 @@ void UK2Node_EditablePinBase::AllocateDefaultPins()
 	// Add in pins based on the user defined pins in this node
 	for(int32 i = 0; i < UserDefinedPins.Num(); i++)
 	{
-		CreatePinFromUserDefinition( UserDefinedPins[i] );
+		//FText DummyErrorMsg;
+		//if (!IsEditable() || CanCreateUserDefinedPin(UserDefinedPins[i]->PinType, UserDefinedPins[i]->DesiredPinDirection, DummyErrorMsg))
+		{
+			CreatePinFromUserDefinition(UserDefinedPins[i]);
+		}
 	}
 }
 
@@ -131,6 +135,9 @@ void UK2Node_EditablePinBase::ImportCustomProperties(const TCHAR* SourceText, FF
 	{
 		TSharedPtr<FUserPinInfo> PinInfo = MakeShareable( new FUserPinInfo() );
 
+		// UserDefinedPins don't currently support direction so set them to an unknown value by default for now
+		PinInfo->DesiredPinDirection = EGPD_MAX;
+
 		if (!FParse::Value(SourceText, TEXT("Name="), PinInfo->PinName))
 		{
 			Warn->Logf( *NSLOCTEXT( "Core", "SyntaxError", "Syntax Error" ).ToString() );
diff --git a/Engine/Source/Editor/BlueprintGraph/Private/K2Node_GetClassDefaults.cpp b/Engine/Source/Editor/BlueprintGraph/Private/K2Node_GetClassDefaults.cpp
--- a/Engine/Source/Editor/BlueprintGraph/Private/K2Node_GetClassDefaults.cpp
+++ b/Engine/Source/Editor/BlueprintGraph/Private/K2Node_GetClassDefaults.cpp
@@ -451,9 +451,34 @@ void UK2Node_GetClassDefaults::CreateOutputPins(UClass* InClass)
 	{
 		if (UBlueprint* Blueprint = Cast<UBlueprint>(InClass->ClassGeneratedBy))
 		{
+			UK2Node_GetClassDefaults* ThisNode = this;
+			auto OnBlueprintClassModified = [ThisNode](UBlueprint* TargetBlueprint)
+			{
+				UBlueprint* OwnerBlueprint = ThisNode->GetBlueprint();
+				// The Blueprint that contains this node may have finished 
+				// regenerating (see bHasBeenRegenerated), but we still may be
+				// in the midst of unwinding a cyclic load (dependent Blueprints);
+				// this lambda could be triggered during the targeted 
+				// Blueprint's regeneration - meaning we really haven't completed 
+				// the load process. In this situation, we cannot "reset loaders" 
+				// because it is not likely that all of the package's objects
+				// have been post-loaded (meaning an assert will most likely  
+				// fire from ReconstructNode). To guard against this, we flip this
+				// Blueprint's bIsRegeneratingOnLoad (like in 
+				// UBlueprintGeneratedClass::ConditionalRecompileClass), which
+				// we use throughout Blueprints to keep us from reseting loaders 
+				// on object Rename()
+				const bool bOldIsRegeneratingVal = OwnerBlueprint->bIsRegeneratingOnLoad;
+				OwnerBlueprint->bIsRegeneratingOnLoad = bOldIsRegeneratingVal || TargetBlueprint->bIsRegeneratingOnLoad;
+
+				ThisNode->ReconstructNode();
+
+				OwnerBlueprint->bIsRegeneratingOnLoad = bOldIsRegeneratingVal;
+			};
+
 			BlueprintSubscribedTo = Blueprint;
-			OnBlueprintChangedDelegate = Blueprint->OnChanged().AddLambda([this](UBlueprint* /* InBlueprint */) { ReconstructNode(); });
-			OnBlueprintCompiledDelegate = Blueprint->OnCompiled().AddLambda([this](UBlueprint* /* InBlueprint */) { ReconstructNode(); });
+			OnBlueprintChangedDelegate  = Blueprint->OnChanged().AddLambda(OnBlueprintClassModified);
+			OnBlueprintCompiledDelegate = Blueprint->OnCompiled().AddLambda(OnBlueprintClassModified);
 		}
 	}
 }
diff --git a/Engine/Source/Editor/BlueprintGraph/Private/K2Node_InputAction.cpp b/Engine/Source/Editor/BlueprintGraph/Private/K2Node_InputAction.cpp
--- a/Engine/Source/Editor/BlueprintGraph/Private/K2Node_InputAction.cpp
+++ b/Engine/Source/Editor/BlueprintGraph/Private/K2Node_InputAction.cpp
@@ -162,7 +162,7 @@ void UK2Node_InputAction::ExpandNode(FKismetCompilerContext& CompilerContext, UE
 		{			
 			UEdGraphPin *EachPin = (*PinIt).Pin;
 			// Create the input touch event
-			UK2Node_InputActionEvent* InputActionEvent = CompilerContext.SpawnIntermediateNode<UK2Node_InputActionEvent>(this, SourceGraph);
+			UK2Node_InputActionEvent* InputActionEvent = CompilerContext.SpawnIntermediateEventNode<UK2Node_InputActionEvent>(this, EachPin, SourceGraph);
 			InputActionEvent->CustomFunctionName = FName( *FString::Printf(TEXT("InpActEvt_%s_%s"), *InputActionName.ToString(), *InputActionEvent->GetName()));
 			InputActionEvent->InputActionName = InputActionName;
 			InputActionEvent->bConsumeInput = bConsumeInput;
@@ -195,7 +195,7 @@ void UK2Node_InputAction::ExpandNode(FKismetCompilerContext& CompilerContext, UE
 	
 		if (InputActionPin->LinkedTo.Num() > 0)
 		{
-			UK2Node_InputActionEvent* InputActionEvent = CompilerContext.SpawnIntermediateNode<UK2Node_InputActionEvent>(this, SourceGraph);
+			UK2Node_InputActionEvent* InputActionEvent = CompilerContext.SpawnIntermediateEventNode<UK2Node_InputActionEvent>(this, InputActionPin, SourceGraph);
 			InputActionEvent->CustomFunctionName = FName( *FString::Printf(TEXT("InpActEvt_%s_%s"), *InputActionName.ToString(), *InputActionEvent->GetName()));
 			InputActionEvent->InputActionName = InputActionName;
 			InputActionEvent->bConsumeInput = bConsumeInput;
diff --git a/Engine/Source/Editor/BlueprintGraph/Private/K2Node_InputKey.cpp b/Engine/Source/Editor/BlueprintGraph/Private/K2Node_InputKey.cpp
--- a/Engine/Source/Editor/BlueprintGraph/Private/K2Node_InputKey.cpp
+++ b/Engine/Source/Editor/BlueprintGraph/Private/K2Node_InputKey.cpp
@@ -283,7 +283,7 @@ void UK2Node_InputKey::ExpandNode(FKismetCompilerContext& CompilerContext, UEdGr
 		{			
 			UEdGraphPin *EachPin = (*PinIt).Pin;
 			// Create the input touch event
-			UK2Node_InputKeyEvent* InputKeyEvent = CompilerContext.SpawnIntermediateNode<UK2Node_InputKeyEvent>(this, SourceGraph);
+			UK2Node_InputKeyEvent* InputKeyEvent = CompilerContext.SpawnIntermediateEventNode<UK2Node_InputKeyEvent>(this, EachPin, SourceGraph);
 			const FName ModifierName = GetModifierName();
 			if ( ModifierName != NAME_None )
 			{
@@ -328,7 +328,7 @@ void UK2Node_InputKey::ExpandNode(FKismetCompilerContext& CompilerContext, UEdGr
 	
 		if (InputKeyPin->LinkedTo.Num() > 0)
 		{
-			UK2Node_InputKeyEvent* InputKeyEvent = CompilerContext.SpawnIntermediateNode<UK2Node_InputKeyEvent>(this, SourceGraph);
+			UK2Node_InputKeyEvent* InputKeyEvent = CompilerContext.SpawnIntermediateEventNode<UK2Node_InputKeyEvent>(this, InputKeyPin, SourceGraph);
 			const FName ModifierName = GetModifierName();
 			if ( ModifierName != NAME_None )
 			{
diff --git a/Engine/Source/Editor/BlueprintGraph/Private/K2Node_InputTouch.cpp b/Engine/Source/Editor/BlueprintGraph/Private/K2Node_InputTouch.cpp
--- a/Engine/Source/Editor/BlueprintGraph/Private/K2Node_InputTouch.cpp
+++ b/Engine/Source/Editor/BlueprintGraph/Private/K2Node_InputTouch.cpp
@@ -200,7 +200,7 @@ void UK2Node_InputTouch::ExpandNode(FKismetCompilerContext& CompilerContext, UEd
 		{			
 			UEdGraphPin *EachPin= (*PinIt).Pin;
 			// Create the input touch event
-			UK2Node_InputTouchEvent* InputTouchEvent = CompilerContext.SpawnIntermediateNode<UK2Node_InputTouchEvent>(this, SourceGraph);
+			UK2Node_InputTouchEvent* InputTouchEvent = CompilerContext.SpawnIntermediateEventNode<UK2Node_InputTouchEvent>(this, EachPin, SourceGraph);
 			InputTouchEvent->CustomFunctionName = FName(*FString::Printf(TEXT("InpTchEvt_%s"), *EachPin->GetName()));
 			InputTouchEvent->bConsumeInput = bConsumeInput;
 			InputTouchEvent->bExecuteWhenPaused = bExecuteWhenPaused;
@@ -241,7 +241,7 @@ void UK2Node_InputTouch::ExpandNode(FKismetCompilerContext& CompilerContext, UEd
 		
 		if (InputTouchPin->LinkedTo.Num() > 0)
 		{
-			UK2Node_InputTouchEvent* InputTouchEvent = CompilerContext.SpawnIntermediateNode<UK2Node_InputTouchEvent>(this, SourceGraph);
+			UK2Node_InputTouchEvent* InputTouchEvent = CompilerContext.SpawnIntermediateEventNode<UK2Node_InputTouchEvent>(this, InputTouchPin, SourceGraph);
 			InputTouchEvent->CustomFunctionName = FName( *FString::Printf(TEXT("InpTchEvt_%s"), *InputTouchEvent->GetName()));
 			InputTouchEvent->InputKeyEvent = InputEvent;
 			InputTouchEvent->bConsumeInput = bConsumeInput;
diff --git a/Engine/Source/Editor/BlueprintGraph/Private/K2Node_LoadAsset.cpp b/Engine/Source/Editor/BlueprintGraph/Private/K2Node_LoadAsset.cpp
--- a/Engine/Source/Editor/BlueprintGraph/Private/K2Node_LoadAsset.cpp
+++ b/Engine/Source/Editor/BlueprintGraph/Private/K2Node_LoadAsset.cpp
@@ -84,16 +84,16 @@ void UK2Node_LoadAsset::ExpandNode(class FKismetCompilerContext& CompilerContext
 	}
 
 	// connect to asset
+	auto CallFunctionAssetPin = CallLoadAssetNode->FindPin(GetInputPinName());
 	{
 		auto AssetPin = FindPin(GetInputPinName());
-		auto CallFunctionAssetPin = CallLoadAssetNode->FindPin(GetInputPinName());
 		ensure(CallFunctionAssetPin);
 		bIsErrorFree &= AssetPin && CallFunctionAssetPin && CompilerContext.MovePinLinksToIntermediate(*AssetPin, *CallFunctionAssetPin).CanSafeConnect();
 	}
 
 	// Create OnLoadEvent
 	const FString DelegateOnLoadedParamName(TEXT("OnLoaded"));
-	auto OnLoadEventNode = CompilerContext.SpawnIntermediateNode<UK2Node_CustomEvent>(this, SourceGraph);
+	auto OnLoadEventNode = CompilerContext.SpawnIntermediateEventNode<UK2Node_CustomEvent>(this, CallFunctionAssetPin, SourceGraph);
 	OnLoadEventNode->CustomFunctionName = *FString::Printf(TEXT("OnLoaded_%s"), *CompilerContext.GetGuid(this));
 	OnLoadEventNode->AllocateDefaultPins();
 	{
diff --git a/Engine/Source/Editor/BlueprintGraph/Private/K2Node_MatineeController.cpp b/Engine/Source/Editor/BlueprintGraph/Private/K2Node_MatineeController.cpp
--- a/Engine/Source/Editor/BlueprintGraph/Private/K2Node_MatineeController.cpp
+++ b/Engine/Source/Editor/BlueprintGraph/Private/K2Node_MatineeController.cpp
@@ -116,7 +116,7 @@ void UK2Node_MatineeController::ExpandNode(FKismetCompilerContext& CompilerConte
 			{
 				FName EventFuncName = MatineeActor->GetFunctionNameForEvent( FName(*(MatineePin->PinName)) );
 
-				UK2Node_Event* MatineeEventNode = CompilerContext.SpawnIntermediateNode<UK2Node_Event>(this, SourceGraph);
+				UK2Node_Event* MatineeEventNode = CompilerContext.SpawnIntermediateEventNode<UK2Node_Event>(this, MatineePin, SourceGraph);
 				MatineeEventNode->EventReference.SetFromField<UFunction>(MatineeEventSig, false);
 				MatineeEventNode->CustomFunctionName = EventFuncName;
 				MatineeEventNode->bInternalEvent = true;
diff --git a/Engine/Source/Editor/BlueprintGraph/Private/K2Node_TunnelBoundary.cpp b/Engine/Source/Editor/BlueprintGraph/Private/K2Node_TunnelBoundary.cpp
new file mode 100644
--- /dev/null
+++ b/Engine/Source/Editor/BlueprintGraph/Private/K2Node_TunnelBoundary.cpp
@@ -0,0 +1,171 @@
+// Copyright 1998-2016 Epic Games, Inc. All Rights Reserved.
+
+#include "BlueprintGraphPrivatePCH.h"
+#include "KismetCompiler.h"
+
+#define LOCTEXT_NAMESPACE "FKCHandler_TunnelBoundary"
+
+//////////////////////////////////////////////////////////////////////////
+// FKCHandler_TunnelBoundary
+
+class FKCHandler_TunnelBoundary : public FNodeHandlingFunctor
+{
+public:
+	FKCHandler_TunnelBoundary(FKismetCompilerContext& InCompilerContext)
+		: FNodeHandlingFunctor(InCompilerContext)
+	{
+	}
+
+	virtual void Compile(FKismetFunctionContext& Context, UEdGraphNode* Node) override
+	{
+		GenerateSimpleThenGoto(Context, *Node);
+	}
+
+};
+
+//////////////////////////////////////////////////////////////////////////
+// UK2Node_TunnelBoundary
+
+UK2Node_TunnelBoundary::UK2Node_TunnelBoundary(const FObjectInitializer& ObjectInitializer)
+	: Super(ObjectInitializer)
+	, bEntryNode(false)
+{
+	return;
+}
+
+FText UK2Node_TunnelBoundary::GetNodeTitle(ENodeTitleType::Type TitleType) const
+{
+	FFormatNamedArguments Args;
+	Args.Add(TEXT("GraphName"), FText::FromName(TunnelGraphName));
+	Args.Add(TEXT("EntryType"), bEntryNode ? FText::FromString(TEXT("Tunnel Entry")) : FText::FromString(TEXT("Tunnel Exit")));
+	return FText::Format(LOCTEXT("UK2Node_TunnelBoundary_FullTitle", "{GraphName} {EntryType}"), Args);
+}
+
+FNodeHandlingFunctor* UK2Node_TunnelBoundary::CreateNodeHandler(FKismetCompilerContext& CompilerContext) const
+{
+	return new FKCHandler_TunnelBoundary(CompilerContext);
+}
+
+void UK2Node_TunnelBoundary::CreateBoundaryNodesForGraph(UEdGraph* TunnelGraph, class FCompilerResultsLog& MessageLog)
+{
+	if (UK2Node_Tunnel* OuterTunnel = TunnelGraph->GetTypedOuter<UK2Node_Tunnel>())
+	{
+		UK2Node_Tunnel* SourceTunnel = Cast<UK2Node_Tunnel>(MessageLog.FindSourceObject(OuterTunnel));
+		CreateBoundaryNodesForTunnelInstance(SourceTunnel, TunnelGraph, MessageLog);
+	}
+}
+
+void UK2Node_TunnelBoundary::CreateBoundaryNodesForTunnelInstance(UK2Node_Tunnel* TunnelInstance, UEdGraph* TunnelGraph, class FCompilerResultsLog& MessageLog)
+{
+	TArray<UEdGraphPin*> TunnelEntryPins;
+	TArray<UEdGraphPin*> TunnelExitPins;
+	if (TunnelGraph)
+	{
+		TArray<UK2Node_Tunnel*> Tunnels;
+		TunnelGraph->GetNodesOfClass<UK2Node_Tunnel>(Tunnels);
+		for (auto TunnelNode : Tunnels)
+		{
+			if (IsPureTunnel(TunnelNode))
+			{
+				for (UEdGraphPin* Pin : TunnelNode->Pins)
+				{
+					if (Pin->LinkedTo.Num() && Pin->PinType.PinCategory == UEdGraphSchema_K2::PC_Exec)
+					{
+						if (Pin->Direction == EGPD_Output)
+						{
+							TunnelEntryPins.Add(Pin);
+						}
+						else
+						{
+							TunnelExitPins.Add(Pin);
+						}
+					}
+				}
+			}
+		}
+	}
+	// Create the Boundary nodes for each unique entry/exit site.
+	for (UEdGraphPin* EntryPin : TunnelEntryPins)
+	{
+		UK2Node_TunnelBoundary* TunnelBoundaryNode = TunnelGraph->CreateBlankNode<UK2Node_TunnelBoundary>();
+		MessageLog.NotifyIntermediateObjectCreation(TunnelBoundaryNode, TunnelInstance);
+		TunnelBoundaryNode->CreateNewGuid();
+		TunnelBoundaryNode->WireUpTunnelEntry(TunnelInstance, EntryPin, MessageLog);
+	}
+	for (UEdGraphPin* ExitPin : TunnelExitPins)
+	{
+		UK2Node_TunnelBoundary* TunnelBoundaryNode = TunnelGraph->CreateBlankNode<UK2Node_TunnelBoundary>();
+		MessageLog.NotifyIntermediateObjectCreation(TunnelBoundaryNode, TunnelInstance);
+		TunnelBoundaryNode->CreateNewGuid();
+		TunnelBoundaryNode->WireUpTunnelExit(TunnelInstance, ExitPin, MessageLog);
+	}
+}
+
+bool UK2Node_TunnelBoundary::IsPureTunnel(UK2Node_Tunnel* Tunnel)
+{
+	return !Tunnel->IsA<UK2Node_MacroInstance>() && !Tunnel->IsA<UK2Node_Composite>();
+}
+
+void UK2Node_TunnelBoundary::WireUpTunnelEntry(UK2Node_Tunnel* TunnelInstance, UEdGraphPin* TunnelPin, FCompilerResultsLog& MessageLog)
+{
+	if (TunnelInstance && TunnelPin)
+	{
+		// Mark as entry node
+		bEntryNode = true;
+		// Grab the graph name
+		DetermineTunnelGraphName(TunnelInstance);
+		// Find the true source pin
+		UEdGraphPin* SourcePin = TunnelInstance->FindPin(TunnelPin->PinName);
+		// Wire in the node
+		UEdGraphPin* OutputPin = CreatePin(EGPD_Output, UEdGraphSchema_K2::PC_Exec, TEXT(""), NULL, false, false, TEXT("TunnelEntryExec"));
+		MessageLog.NotifyIntermediatePinCreation(OutputPin, SourcePin);
+		for (auto LinkedPin : TunnelPin->LinkedTo)
+		{
+			LinkedPin->MakeLinkTo(OutputPin);
+		}
+		UEdGraphPin* InputPin = CreatePin(EGPD_Input, UEdGraphSchema_K2::PC_Exec, TEXT(""), NULL, false, false, SourcePin->PinName);
+		MessageLog.NotifyIntermediatePinCreation(InputPin, SourcePin);
+		TunnelPin->BreakAllPinLinks();
+		TunnelPin->MakeLinkTo(InputPin);
+	}
+}
+
+void UK2Node_TunnelBoundary::WireUpTunnelExit(UK2Node_Tunnel* TunnelInstance, UEdGraphPin* TunnelPin, FCompilerResultsLog& MessageLog)
+{
+	if (TunnelInstance && TunnelPin)
+	{
+		// Mark as exit node
+		bEntryNode = false;
+		// Grab the graph name
+		DetermineTunnelGraphName(TunnelInstance);
+		// Find the true source pin
+		UEdGraphPin* SourcePin = TunnelInstance->FindPin(TunnelPin->PinName);
+		// Wire in the node
+		UEdGraphPin* InputPin = CreatePin(EGPD_Input, UEdGraphSchema_K2::PC_Exec, TEXT(""), NULL, false, false, SourcePin->PinName);
+		MessageLog.NotifyIntermediatePinCreation(InputPin, SourcePin);
+		for (auto LinkedPin : TunnelPin->LinkedTo)
+		{
+			LinkedPin->MakeLinkTo(InputPin);
+		}
+		UEdGraphPin* OutputPin = CreatePin(EGPD_Output, UEdGraphSchema_K2::PC_Exec, TEXT(""), NULL, false, false, TEXT("TunnelExitExec"));
+		MessageLog.NotifyIntermediatePinCreation(OutputPin, SourcePin);
+		TunnelPin->BreakAllPinLinks();
+		TunnelPin->MakeLinkTo(OutputPin);
+	}
+}
+
+void UK2Node_TunnelBoundary::DetermineTunnelGraphName(UK2Node_Tunnel* TunnelInstance)
+{
+	UEdGraph* TunnelGraph = nullptr;
+	if (const UK2Node_MacroInstance* MacroInstance = Cast<UK2Node_MacroInstance>(TunnelInstance))
+	{
+		TunnelGraph = MacroInstance->GetMacroGraph();
+	}
+	else if (const UK2Node_Composite* CompositeInstance = Cast<UK2Node_Composite>(TunnelInstance))
+	{
+		TunnelGraph = CompositeInstance->BoundGraph;
+	}
+	TunnelGraphName = TunnelGraph ? TunnelGraph->GetFName() : NAME_None;
+}
+
+#undef LOCTEXT_NAMESPACE
diff --git a/Engine/Source/Editor/BlueprintGraph/Public/BlueprintGraphDefinitions.h b/Engine/Source/Editor/BlueprintGraph/Public/BlueprintGraphDefinitions.h
--- a/Engine/Source/Editor/BlueprintGraph/Public/BlueprintGraphDefinitions.h
+++ b/Engine/Source/Editor/BlueprintGraph/Public/BlueprintGraphDefinitions.h
@@ -66,5 +66,6 @@
 #include "K2Node_VariableGet.h"
 #include "K2Node_VariableSet.h"
 #include "K2Node_SetFieldsInStruct.h"
+#include "K2Node_TunnelBoundary.h"
 
 #include "NodeDependingOnEnumInterface.h"
diff --git a/Engine/Source/Editor/EditorStyle/Private/EditorStyleClasses.cpp b/Engine/Source/Editor/EditorStyle/Private/EditorStyleClasses.cpp
--- a/Engine/Source/Editor/EditorStyle/Private/EditorStyleClasses.cpp
+++ b/Engine/Source/Editor/EditorStyle/Private/EditorStyleClasses.cpp
@@ -14,6 +14,12 @@ UEditorStyleSettings::UEditorStyleSettings( const FObjectInitializer& ObjectInit
 
 	AssetEditorOpenLocation = EAssetEditorOpenLocation::Default;
 
+	RegularColor = FLinearColor(0.035f, 0.035f, 0.035f);
+	RuleColor = FLinearColor(0.008f, 0.008f, 0.008f);
+	CenterColor = FLinearColor::Black;
+
+	bUseGrid = true;
+
 	bShowFriendlyNames = true;
 	LogTimestampMode = ELogTimes::None;
 }
diff --git a/Engine/Source/Editor/EditorStyle/Private/SlateEditorStyle.cpp b/Engine/Source/Editor/EditorStyle/Private/SlateEditorStyle.cpp
--- a/Engine/Source/Editor/EditorStyle/Private/SlateEditorStyle.cpp
+++ b/Engine/Source/Editor/EditorStyle/Private/SlateEditorStyle.cpp
@@ -4839,6 +4839,20 @@ void FSlateEditorStyle::FStyle::SetupLevelEditorStyle()
 		Set( "BlueprintProfiler.BPPinDisconnected", new IMAGE_BRUSH( "Icons/BPProfiler_ExecPin_Disconnected", Icon12x16 ) );
 		Set( "BlueprintProfiler.PureNode", new IMAGE_BRUSH( "/Icons/pill_16x", Icon16x16 ) );
 		Set( "BlueprintProfiler.BPNode", new IMAGE_BRUSH("Icons/Icon_Profiler_BPNode_16x", Icon16x16));
+
+		Button = FButtonStyle()
+			.SetNormal( FSlateNoResource() )
+			.SetHovered( FSlateNoResource() )
+			.SetPressed( FSlateNoResource() )
+			.SetNormalPadding( FMargin( 2,2,2,2 ) )
+			.SetPressedPadding( FMargin( 2,2,2,2 ) );
+		Set( "BlueprintProfiler.TreeArrowButton", Button );
+
+		Set( "BlueprintProfiler.TreeArrowGlyph", FTextBlockStyle(NormalText)
+			.SetFont(TTF_FONT("Fonts/FontAwesome", 10))
+			.SetShadowOffset(FVector2D::ZeroVector)
+		);
+
 	}
 #endif // WITH_EDITOR
 
diff --git a/Engine/Source/Editor/EditorStyle/Public/Classes/EditorStyleSettings.h b/Engine/Source/Editor/EditorStyle/Public/Classes/EditorStyleSettings.h
--- a/Engine/Source/Editor/EditorStyle/Public/Classes/EditorStyleSettings.h
+++ b/Engine/Source/Editor/EditorStyle/Public/Classes/EditorStyleSettings.h
@@ -80,6 +80,22 @@ public:
 	UPROPERTY(EditAnywhere, config, Category=UserInterface)
 	uint32 bUseSmallToolBarIcons:1;
 
+	/** If true the material editor and blueprint editor will show a grid on it's background. */
+	UPROPERTY(EditAnywhere, config, Category = Graphs, meta = (DisplayName = "Use Grids In The Material And Blueprint Editor"))
+	uint32 bUseGrid : 1;
+
+	/** The color used to represent regular grid lines */
+	UPROPERTY(EditAnywhere, config, Category = Graphs, meta = (DisplayName = "Grid regular Color"))
+	FLinearColor RegularColor;
+
+	/** The color used to represent ruler lines in the grid */
+	UPROPERTY(EditAnywhere, config, Category = Graphs, meta = (DisplayName = "Grid ruler Color"))
+	FLinearColor RuleColor;
+
+	/** The color used to represent the center lines in the grid */
+	UPROPERTY(EditAnywhere, config, Category = Graphs, meta = (DisplayName = "Grid center Color"))
+	FLinearColor CenterColor;
+
 	/** Enables animated transitions for certain menus and pop-up windows.  Note that animations may be automatically disabled at low frame rates in order to improve responsiveness. */
 	UPROPERTY(EditAnywhere, config, Category=UserInterface)
 	uint32 bEnableWindowAnimations:1;
diff --git a/Engine/Source/Editor/GraphEditor/Private/KismetPins/SGraphPinIndex.cpp b/Engine/Source/Editor/GraphEditor/Private/KismetPins/SGraphPinIndex.cpp
--- a/Engine/Source/Editor/GraphEditor/Private/KismetPins/SGraphPinIndex.cpp
+++ b/Engine/Source/Editor/GraphEditor/Private/KismetPins/SGraphPinIndex.cpp
@@ -14,12 +14,11 @@ TSharedRef<SWidget>	SGraphPinIndex::GetDefaultValueWidget()
 {
 	const UEdGraphSchema_K2* Schema = GetDefault<UEdGraphSchema_K2>();
 
-	return SNew(SPinTypeSelector, FGetPinTypeTree::CreateUObject(Schema, &UEdGraphSchema_K2::GetVariableIndexTypeTree))
+	return SNew(SPinTypeSelector, FGetPinTypeTree::CreateUObject(Schema, &UEdGraphSchema_K2::GetVariableTypeTree))
 		.TargetPinType(this, &SGraphPinIndex::OnGetPinType)
 		.OnPinTypeChanged(this, &SGraphPinIndex::OnTypeChanged)
 		.Schema(Schema)
-		.bAllowExec(false)
-		.bAllowWildcard(false)
+		.TypeTreeFilter(ETypeTreeFilter::IndexTypesOnly)
 		.IsEnabled(true)
 		.bAllowArrays(false);
 }
diff --git a/Engine/Source/Editor/GraphEditor/Private/SGraphPanel.cpp b/Engine/Source/Editor/GraphEditor/Private/SGraphPanel.cpp
--- a/Engine/Source/Editor/GraphEditor/Private/SGraphPanel.cpp
+++ b/Engine/Source/Editor/GraphEditor/Private/SGraphPanel.cpp
@@ -34,8 +34,8 @@ FGraphPinHandle::FGraphPinHandle(UEdGraphPin* InPin)
 	{
 		if (UEdGraphNode* Node = InPin->GetOwningNodeUnchecked())
 		{
-			PinName = InPin->PinName;
 			NodeGuid = Node->NodeGuid;
+			PinId = InPin->PinId;
 		}
 	}
 }
@@ -50,7 +50,7 @@ TSharedPtr<SGraphPin> FGraphPinHandle::FindInGraphPanel(const SGraphPanel& InPan
 		{
 			UEdGraphNode* Node = GraphNode->GetNodeObj();
 
-			if (UEdGraphPin* Pin = Node->FindPin(PinName))
+			if (UEdGraphPin* Pin = Node->FindPinById(PinId))
 			{
 				return GraphNode->FindWidgetForPin(Pin);
 			}
diff --git a/Engine/Source/Editor/GraphEditor/Private/SGraphPin.cpp b/Engine/Source/Editor/GraphEditor/Private/SGraphPin.cpp
--- a/Engine/Source/Editor/GraphEditor/Private/SGraphPin.cpp
+++ b/Engine/Source/Editor/GraphEditor/Private/SGraphPin.cpp
@@ -392,27 +392,6 @@ FReply SGraphPin::OnPinMouseDown( const FGeometry& SenderGeometry, const FPointe
 				// Note that for some nodes, this can cause reconstruction. In that case, pins we had previously linked to may now be destroyed.
 				const UEdGraphSchema* Schema = GraphPinObj->GetSchema();
 				Schema->BreakPinLinks(*GraphPinObj, true);
-
-				// Check to see if the panel has been invalidated by a graph change notification
-				if (!OwnerPanelPtr->Contains(OwnerNodePinned->GetNodeObj()))
-				{
-					// Force the panel to update. This will cause node & pin widgets to be reinstanced to match any reconstructed node/pin object references.
-					OwnerPanelPtr->Update();
-
-					// Obtain the full set of pins again after the update
-					AllPins.Empty(AllPins.Num());
-					OwnerPanelPtr->GetAllPins(AllPins);
-
-					// Rebuild the UEdGraphPin->SGraphPin mapping for the full pin set
-					PinToPinWidgetMap.Empty(PinToPinWidgetMap.Num());
-					for( TSet< TSharedRef<SWidget> >::TIterator ConnectorIt(AllPins); ConnectorIt; ++ConnectorIt )
-					{
-						const TSharedRef<SWidget>& SomePinWidget = *ConnectorIt;
-						const SGraphPin& PinWidget = static_cast<const SGraphPin&>(SomePinWidget.Get());
-
-						PinToPinWidgetMap.Add(PinWidget.GetPinObj(), StaticCastSharedRef<SGraphPin>(SomePinWidget));
-					}
-				}
 				
 				// Now iterate over our lookup table to find the instances of pin widgets that we had previously linked to
 				TArray<TSharedRef<SGraphPin>> PinArray;
@@ -601,7 +580,7 @@ void SGraphPin::OnDragEnter( const FGeometry& MyGeometry, const FDragDropEvent&
 	if (Operation->IsOfType<FGraphEditorDragDropAction>())
 	{
 		// Ensure that the pin is valid before using it
-		if(GraphPinObj != NULL && GraphPinObj->GetOuter() != NULL && GraphPinObj->GetOuter()->IsA(UEdGraphNode::StaticClass()))
+		if(GraphPinObj != NULL && !GraphPinObj->IsPendingKill() && GraphPinObj->GetOuter() != NULL && GraphPinObj->GetOuter()->IsA(UEdGraphNode::StaticClass()))
 		{
 			// Inform the Drag and Drop operation that we are hovering over this pin.
 			TSharedPtr<FGraphEditorDragDropAction> DragConnectionOp = StaticCastSharedPtr<FGraphEditorDragDropAction>(Operation);
@@ -760,7 +739,7 @@ EVisibility SGraphPin::IsPinVisibleAsAdvanced() const
 		}
 	}
 
-	const bool bIsAdvancedPin = GraphPinObj && GraphPinObj->bAdvancedView;
+	const bool bIsAdvancedPin = GraphPinObj && !GraphPinObj->IsPendingKill() && GraphPinObj->bAdvancedView;
 	const bool bCanBeHidden = !IsConnected();
 	return (bIsAdvancedPin && bHideAdvancedPin && bCanBeHidden) ? EVisibility::Collapsed : EVisibility::Visible;
 }
@@ -989,7 +968,7 @@ FText SGraphPin::GetTooltipText() const
 {
 	FText HoverText = FText::GetEmpty();
 
-	UEdGraphNode* GraphNode = GraphPinObj ? GraphPinObj->GetOwningNodeUnchecked() : nullptr;
+	UEdGraphNode* GraphNode = GraphPinObj && !GraphPinObj->IsPendingKill() ? GraphPinObj->GetOwningNodeUnchecked() : nullptr;
 	if (GraphNode != nullptr)
 	{
 		FString HoverStr;
diff --git a/Engine/Source/Editor/GraphEditor/Private/SNodePanel.cpp b/Engine/Source/Editor/GraphEditor/Private/SNodePanel.cpp
--- a/Engine/Source/Editor/GraphEditor/Private/SNodePanel.cpp
+++ b/Engine/Source/Editor/GraphEditor/Private/SNodePanel.cpp
@@ -436,7 +436,6 @@ EActiveTimerReturnType SNodePanel::HandleZoomToFit(double InCurrentTime, float I
 void SNodePanel::Tick( const FGeometry& AllottedGeometry, const double InCurrentTime, const float InDeltaTime )
 {
 	CachedGeometry = AllottedGeometry;
-	bool bWasActiveTimerRegisteredThisFrame = false;
 	bool bCanMoveToTargetObjectThisFrame = true;
 
 	if(DeferredSelectionTargetObjects.Num() > 0)
@@ -469,14 +468,10 @@ void SNodePanel::Tick( const FGeometry& AllottedGeometry, const double InCurrent
 		if (bCanMoveToTargetObjectThisFrame && GetBoundsForNode(DeferredMovementTargetObject, ZoomTargetTopLeft, ZoomTargetBottomRight, ZoomPadding))
 		{
 			DeferredMovementTargetObject = nullptr;
-			if (!ActiveTimerHandle.IsValid())
-			{
-				ActiveTimerHandle = RegisterActiveTimer(0.f, FWidgetActiveTimerDelegate::CreateSP(this, &SNodePanel::HandleZoomToFit));
-				bWasActiveTimerRegisteredThisFrame = true;
-			}
+			RequestZoomToFit();
 		}
 	}
-	
+
 	// Zoom to node extents
 	if( bDeferredZoomToNodeExtents )
 	{
@@ -485,11 +480,7 @@ void SNodePanel::Tick( const FGeometry& AllottedGeometry, const double InCurrent
 		if( GetBoundsForNodes(bDeferredZoomToSelection, ZoomTargetTopLeft, ZoomTargetBottomRight, ZoomPadding))
 		{
 			bDeferredZoomToSelection = false;
-			if (!ActiveTimerHandle.IsValid())
-			{
-				ActiveTimerHandle = RegisterActiveTimer(0.f, FWidgetActiveTimerDelegate::CreateSP(this, &SNodePanel::HandleZoomToFit));
-				bWasActiveTimerRegisteredThisFrame = true;
-			}
+			RequestZoomToFit();
 		}
 	}
 
@@ -673,6 +664,10 @@ FReply SNodePanel::OnMouseMove(const FGeometry& MyGeometry, const FPointerEvent&
 				this->bIsZoomingWithTrackpad = true;
 				bShowSoftwareCursor = true;
 			}
+
+			// Stop the zoom-to-fit in favor of user control
+			CancelZoomToFit();
+
 			return ReplyState;
 		}
 		else if (bIsRightMouseButtonDown)
@@ -690,6 +685,9 @@ FReply SNodePanel::OnMouseMove(const FGeometry& MyGeometry, const FPointerEvent&
 			this->bIsPanning = true;
 			ViewOffset -= CursorDelta / GetZoomAmount();
 
+			// Stop the zoom-to-fit in favor of user control
+			CancelZoomToFit();
+
 			return ReplyState;
 		}
 		else if (bIsLeftMouseButtonDown)
@@ -779,6 +777,9 @@ FReply SNodePanel::OnMouseMove(const FGeometry& MyGeometry, const FPointerEvent&
 				FindNodesAffectedByMarquee( /*out*/ Marquee.AffectedNodes );
 				return FReply::Handled();
 			}
+
+			// Stop the zoom-to-fit in favor of user control
+			CancelZoomToFit();
 		}
 	}
 
@@ -923,6 +924,9 @@ FReply SNodePanel::OnMouseWheel(const FGeometry& MyGeometry, const FPointerEvent
 	const int32 ZoomLevelDelta = FMath::FloorToInt( MouseEvent.GetWheelDelta() );
 	ChangeZoomLevel(ZoomLevelDelta, WidgetSpaceCursorPos, MouseEvent.IsControlDown());
 
+	// Stop the zoom-to-fit in favor of user control
+	CancelZoomToFit();
+
 	return FReply::Handled();
 }
 
@@ -973,6 +977,10 @@ FReply SNodePanel::OnTouchGesture( const FGeometry& MyGeometry, const FPointerEv
 			ChangeZoomLevel(ZoomLevelDelta, WidgetSpaceCursorPos, GestureEvent.IsControlDown());
 			TotalGestureMagnify = 0.0f;
 		}
+
+		// Stop the zoom-to-fit in favor of user control
+		CancelZoomToFit();
+
 		return FReply::Handled();
 	}
 	else if (GestureType == EGestureEvent::Scroll)
@@ -982,6 +990,10 @@ FReply SNodePanel::OnTouchGesture( const FGeometry& MyGeometry, const FPointerEv
 
 		this->bIsPanning = true;
 		ViewOffset -= (bUseDirectionInvertedFromDevice == GestureEvent.IsDirectionInvertedFromDevice() ? GestureDelta : -GestureDelta) / GetZoomAmount();
+		
+		// Stop the zoom-to-fit in favor of user control
+		CancelZoomToFit();
+
 		return FReply::Handled();
 	}
 	return FReply::Unhandled();
@@ -1058,24 +1070,13 @@ void SNodePanel::SelectAndCenterObject(const UObject* ObjectToSelect, bool bCent
 		DeferredMovementTargetObject = ObjectToSelect;
 	}
 
-	auto PinnedActiveTimerHandle = ActiveTimerHandle.Pin();
-	if (PinnedActiveTimerHandle.IsValid())
-	{
-		// Stop zooming for now and let the a new zoom target be established
-		UnRegisterActiveTimer(PinnedActiveTimerHandle.ToSharedRef());
-	}
+	CancelZoomToFit();
 }
 
 void SNodePanel::CenterObject(const UObject* ObjectToCenter)
 {
 	DeferredMovementTargetObject = ObjectToCenter;
-
-	auto PinnedActiveTimerHandle = ActiveTimerHandle.Pin();
-	if (PinnedActiveTimerHandle.IsValid())
-	{
-		// Stop zooming for now and let the a new zoom target be established
-		UnRegisterActiveTimer(PinnedActiveTimerHandle.ToSharedRef());
-	}
+	CancelZoomToFit();
 }
 
 /** Add a slot to the CanvasPanel dynamically */
@@ -1136,11 +1137,7 @@ void SNodePanel::RestoreViewSettings(const FVector2D& InViewOffset, float InZoom
 		ZoomLevel = ZoomLevels->GetNearestZoomLevel(InZoomAmount);
 		bDeferredZoomToNodeExtents = false;
 
-		auto PinnedActiveTimerHandle = ActiveTimerHandle.Pin();
-		if (PinnedActiveTimerHandle.IsValid())
-		{
-			UnRegisterActiveTimer(PinnedActiveTimerHandle.ToSharedRef());
-		}
+		CancelZoomToFit();
 	}
 
 	PostChangedZoom();
@@ -1163,9 +1160,9 @@ void SNodePanel::PaintBackgroundAsLines(const FSlateBrush* BackgroundImage, cons
 	const int32 RulePeriod = (int32)FEditorStyle::GetFloat("Graph.Panel.GridRulePeriod");
 	check(RulePeriod > 0);
 
-	const FLinearColor RegularColor(FEditorStyle::GetColor("Graph.Panel.GridLineColor"));
-	const FLinearColor RuleColor(FEditorStyle::GetColor("Graph.Panel.GridRuleColor"));
-	const FLinearColor CenterColor(FEditorStyle::GetColor("Graph.Panel.GridCenterColor"));
+	const FLinearColor RegularColor(GetDefault<UEditorStyleSettings>()->RegularColor);
+	const FLinearColor RuleColor(GetDefault<UEditorStyleSettings>()->RuleColor);
+	const FLinearColor CenterColor(GetDefault<UEditorStyleSettings>()->CenterColor);
 	const float GraphSmallestGridSize = 8.0f;
 	const float RawZoomFactor = GetZoomAmount();
 	const float NominalGridSize = GetSnapGridSize();
@@ -1200,64 +1197,67 @@ void SNodePanel::PaintBackgroundAsLines(const FSlateBrush* BackgroundImage, cons
 	new (LinePoints) FVector2D(0.0f, 0.0f);
 	new (LinePoints) FVector2D(0.0f, 0.0f);
 
-	// Horizontal bars
-	for (int32 GridIndex = 0; ImageOffsetY < AllottedGeometry.Size.Y; ImageOffsetY += GridCellSize, ++GridIndex)
-	{
-		if (ImageOffsetY >= 0.0f)
-		{
-			const bool bIsRuleLine = (GridIndex % RulePeriod) == 0;
-			const int32 Layer = bIsRuleLine ? (DrawLayerId+1) : DrawLayerId;
+	//If we want to use grid then show grid, otherwise don't render the grid
+	if (GetDefault<UEditorStyleSettings>()->bUseGrid == true){
 
-			const FLinearColor* Color = bIsRuleLine ? &RuleColor : &RegularColor;
-			if (FMath::IsNearlyEqual(ZeroSpace.Y, ImageOffsetY, 1.0f))
+		// Horizontal bars
+		for (int32 GridIndex = 0; ImageOffsetY < AllottedGeometry.Size.Y; ImageOffsetY += GridCellSize, ++GridIndex)
+		{
+			if (ImageOffsetY >= 0.0f)
 			{
-				Color = &CenterColor;
-			}
+				const bool bIsRuleLine = (GridIndex % RulePeriod) == 0;
+				const int32 Layer = bIsRuleLine ? (DrawLayerId + 1) : DrawLayerId;
 
-			LinePoints[0] = FVector2D(0.0f, ImageOffsetY);
-			LinePoints[1] = FVector2D(AllottedGeometry.Size.X, ImageOffsetY);
-
-			FSlateDrawElement::MakeLines(
-				OutDrawElements,
-				Layer,
-				AllottedGeometry.ToPaintGeometry(),
-				LinePoints,
-				MyClippingRect,
-				ESlateDrawEffect::None,
-				*Color,
-				bAntialias);
+				const FLinearColor* Color = bIsRuleLine ? &RuleColor : &RegularColor;
+				if (FMath::IsNearlyEqual(ZeroSpace.Y, ImageOffsetY, 1.0f))
+				{
+					Color = &CenterColor;
+				}
+
+				LinePoints[0] = FVector2D(0.0f, ImageOffsetY);
+				LinePoints[1] = FVector2D(AllottedGeometry.Size.X, ImageOffsetY);
+
+				FSlateDrawElement::MakeLines(
+					OutDrawElements,
+					Layer,
+					AllottedGeometry.ToPaintGeometry(),
+					LinePoints,
+					MyClippingRect,
+					ESlateDrawEffect::None,
+					*Color,
+					bAntialias);
+			}
 		}
-	}
 
-	// Vertical bars
-	for (int32 GridIndex = 0; ImageOffsetX < AllottedGeometry.Size.X; ImageOffsetX += GridCellSize, ++GridIndex)
-	{
-		if (ImageOffsetX >= 0.0f)
+		// Vertical bars
+		for (int32 GridIndex = 0; ImageOffsetX < AllottedGeometry.Size.X; ImageOffsetX += GridCellSize, ++GridIndex)
 		{
-			const bool bIsRuleLine = (GridIndex % RulePeriod) == 0;
-			const int32 Layer = bIsRuleLine ? (DrawLayerId+1) : DrawLayerId;
-
-			const FLinearColor* Color = bIsRuleLine ? &RuleColor : &RegularColor;
-			if (FMath::IsNearlyEqual(ZeroSpace.X, ImageOffsetX, 1.0f))
+			if (ImageOffsetX >= 0.0f)
 			{
-				Color = &CenterColor;
-			}
+				const bool bIsRuleLine = (GridIndex % RulePeriod) == 0;
+				const int32 Layer = bIsRuleLine ? (DrawLayerId + 1) : DrawLayerId;
+
+				const FLinearColor* Color = bIsRuleLine ? &RuleColor : &RegularColor;
+				if (FMath::IsNearlyEqual(ZeroSpace.X, ImageOffsetX, 1.0f))
+				{
+					Color = &CenterColor;
+				}
 
-			LinePoints[0] = FVector2D(ImageOffsetX, 0.0f);
-			LinePoints[1] = FVector2D(ImageOffsetX, AllottedGeometry.Size.Y);
-
-			FSlateDrawElement::MakeLines(
-				OutDrawElements,
-				Layer,
-				AllottedGeometry.ToPaintGeometry(),
-				LinePoints,
-				MyClippingRect,
-				ESlateDrawEffect::None,
-				*Color,
-				bAntialias);
+				LinePoints[0] = FVector2D(ImageOffsetX, 0.0f);
+				LinePoints[1] = FVector2D(ImageOffsetX, AllottedGeometry.Size.Y);
+
+				FSlateDrawElement::MakeLines(
+					OutDrawElements,
+					Layer,
+					AllottedGeometry.ToPaintGeometry(),
+					LinePoints,
+					MyClippingRect,
+					ESlateDrawEffect::None,
+					*Color,
+					bAntialias);
+			}
 		}
 	}
-
 	DrawLayerId += 2;
 }
 
@@ -1555,10 +1555,7 @@ void SNodePanel::ZoomToTarget(const FVector2D& TopLeft, const FVector2D& BottomR
 	ZoomTargetTopLeft = TopLeft;
 	ZoomTargetBottomRight = BottomRight;
 
-	if (!ActiveTimerHandle.IsValid())
-	{
-		ActiveTimerHandle = RegisterActiveTimer(0.f, FWidgetActiveTimerDelegate::CreateSP(this, &SNodePanel::HandleZoomToFit));
-	}
+	RequestZoomToFit();
 }
 
 void SNodePanel::ChangeZoomLevel(int32 ZoomLevelDelta, const FVector2D& WidgetSpaceZoomOrigin, bool bOverrideZoomLimiting)
@@ -1627,6 +1624,22 @@ void SNodePanel::PostChangedZoom()
 	CurrentLOD = ZoomLevels->GetLOD(ZoomLevel);
 }
 
+void SNodePanel::RequestZoomToFit()
+{
+	if (!ActiveTimerHandle.IsValid())
+	{
+		ActiveTimerHandle = RegisterActiveTimer(0.f, FWidgetActiveTimerDelegate::CreateSP(this, &SNodePanel::HandleZoomToFit));
+	}
+}
+
+void SNodePanel::CancelZoomToFit()
+{
+	if (ActiveTimerHandle.IsValid())
+	{
+		UnRegisterActiveTimer(ActiveTimerHandle.Pin().ToSharedRef());
+	}
+}
+
 bool SNodePanel::HasMoved() const
 {
 	return (!FMath::IsNearlyEqual(GetZoomAmount(), OldZoomAmount) || !ViewOffset.Equals(OldViewOffset, SMALL_NUMBER));
diff --git a/Engine/Source/Editor/GraphEditor/Public/GraphSplineOverlapResult.h b/Engine/Source/Editor/GraphEditor/Public/GraphSplineOverlapResult.h
--- a/Engine/Source/Editor/GraphEditor/Public/GraphSplineOverlapResult.h
+++ b/Engine/Source/Editor/GraphEditor/Public/GraphSplineOverlapResult.h
@@ -13,8 +13,8 @@ struct GRAPHEDITOR_API FGraphPinHandle
 	/** The GUID of the node to which this pin belongs */
 	FGuid NodeGuid;
 
-	/** The name of the pin we are referencing */
-	FString PinName;
+	/** The GUID of the pin we are referencing */
+	FGuid PinId;
 
 	/** Constructor */
 	FGraphPinHandle(UEdGraphPin* InPin);
diff --git a/Engine/Source/Editor/GraphEditor/Public/SNodePanel.h b/Engine/Source/Editor/GraphEditor/Public/SNodePanel.h
--- a/Engine/Source/Editor/GraphEditor/Public/SNodePanel.h
+++ b/Engine/Source/Editor/GraphEditor/Public/SNodePanel.h
@@ -806,6 +806,12 @@ protected:
 
 	// Should be called whenever the zoom level has changed
 	void PostChangedZoom();
+
+	// Fires up a per-tick function to zoom the graph to fit
+	void RequestZoomToFit();
+
+	// Cancels any active zoom-to-fit action
+	void CancelZoomToFit();
 protected:
 	// The interface for mapping ZoomLevel values to actual node scaling values
 	TScopedPointer<FZoomLevelsContainer> ZoomLevels;
diff --git a/Engine/Source/Editor/Kismet/Private/BlueprintDetailsCustomization.cpp b/Engine/Source/Editor/Kismet/Private/BlueprintDetailsCustomization.cpp
--- a/Engine/Source/Editor/Kismet/Private/BlueprintDetailsCustomization.cpp
+++ b/Engine/Source/Editor/Kismet/Private/BlueprintDetailsCustomization.cpp
@@ -256,7 +256,7 @@ void FBlueprintVarActionDetails::CustomizeDetails( IDetailLayoutBuilder& DetailL
 		.OnPinTypeChanged(this, &FBlueprintVarActionDetails::OnVarTypeChanged)
 		.IsEnabled(this, &FBlueprintVarActionDetails::GetVariableTypeChangeEnabled)
 		.Schema(Schema)
-		.bAllowExec(false)
+		.TypeTreeFilter(ETypeTreeFilter::None)
 		.Font( DetailFontInfo )
 		.ToolTip(VarTypeTooltip)
 	];
@@ -2271,6 +2271,17 @@ void FBlueprintGraphArgumentLayout::GenerateHeaderRowContent( FDetailWidgetRow&
 {
 	const UEdGraphSchema_K2* K2Schema = GetDefault<UEdGraphSchema_K2>();
 
+	ETypeTreeFilter TypeTreeFilter = ETypeTreeFilter::None;
+	if (TargetNode->CanModifyExecutionWires())
+	{
+		TypeTreeFilter |= ETypeTreeFilter::AllowExec;
+	}
+
+	if (ShouldAllowWildcard(TargetNode))
+	{
+		TypeTreeFilter |= ETypeTreeFilter::AllowWildcard;
+	}
+
 	NodeRow
 	.NameContent()
 	[
@@ -2300,8 +2311,7 @@ void FBlueprintGraphArgumentLayout::GenerateHeaderRowContent( FDetailWidgetRow&
 				.OnPinTypePreChanged(this, &FBlueprintGraphArgumentLayout::OnPrePinInfoChange)
 				.OnPinTypeChanged(this, &FBlueprintGraphArgumentLayout::PinInfoChanged)
 				.Schema(K2Schema)
-				.bAllowExec(TargetNode->CanModifyExecutionWires())
-				.bAllowWildcard(ShouldAllowWildcard(TargetNode))
+				.TypeTreeFilter(TypeTreeFilter)
 				.bAllowArrays(!ShouldPinBeReadOnly())
 				.IsEnabled(!ShouldPinBeReadOnly(true))
 				.Font( IDetailLayoutBuilder::GetDetailFont() )
diff --git a/Engine/Source/Editor/Kismet/Private/BlueprintEditor.cpp b/Engine/Source/Editor/Kismet/Private/BlueprintEditor.cpp
--- a/Engine/Source/Editor/Kismet/Private/BlueprintEditor.cpp
+++ b/Engine/Source/Editor/Kismet/Private/BlueprintEditor.cpp
@@ -4109,11 +4109,11 @@ void FBlueprintEditor::OnChangePinType()
 				UK2Node_Select* SelectNode = Cast<UK2Node_Select>(SelectedPin->GetOwningNode());
 				if (SelectNode && SelectNode->GetIndexPin() == SelectedPin)
 				{
-					TSharedRef<SCompoundWidget> PinChange = SNew(SPinTypeSelector, FGetPinTypeTree::CreateUObject(Schema, &UEdGraphSchema_K2::GetVariableIndexTypeTree))
+					TSharedRef<SCompoundWidget> PinChange = SNew(SPinTypeSelector, FGetPinTypeTree::CreateUObject(Schema, &UEdGraphSchema_K2::GetVariableTypeTree))
 				.TargetPinType(this, &FBlueprintEditor::OnGetPinType, RootPin)
 						.OnPinTypeChanged(this, &FBlueprintEditor::OnChangePinTypeFinished, SelectedPin)
 						.Schema(Schema)
-						.bAllowExec(false)
+						.TypeTreeFilter(ETypeTreeFilter::IndexTypesOnly)
 						.IsEnabled(true)
 						.bAllowArrays(false);
 
@@ -4131,7 +4131,7 @@ void FBlueprintEditor::OnChangePinType()
 				.TargetPinType(this, &FBlueprintEditor::OnGetPinType, RootPin)
 						.OnPinTypeChanged(this, &FBlueprintEditor::OnChangePinTypeFinished, SelectedPin)
 						.Schema(Schema)
-						.bAllowExec(false)
+						.TypeTreeFilter(ETypeTreeFilter::None)
 						.IsEnabled(true)
 						.bAllowArrays(false);
 
@@ -4533,6 +4533,8 @@ bool FBlueprintEditor::CanCollapseSelectionToFunction(TSet<class UEdGraphNode*>&
 
 	UEdGraphNode* InterfaceTemplateNode = nullptr;
 
+	TArray<UEdGraphPin*> EntryGatewayPins;
+
 	// Runs through every node and fully validates errors with placing selection in a function graph, reporting all errors.
 	for (TSet<UEdGraphNode*>::TConstIterator NodeIt(InSelection); NodeIt; ++NodeIt)
 	{
@@ -4570,37 +4572,45 @@ bool FBlueprintEditor::CanCollapseSelectionToFunction(TSet<class UEdGraphNode*>&
 			{
 				if(Node->Pins[PinIndex]->PinType.PinCategory == K2Schema->PC_Exec)
 				{
-					for (int32 LinkIndex = 0; LinkIndex < Node->Pins[PinIndex]->LinkedTo.Num(); ++LinkIndex)
+					if (Node->Pins[PinIndex]->LinkedTo.Num() == 0 && Node->Pins[PinIndex]->Direction == EGPD_Input)
 					{
-						if(!InSelection.Contains(Node->Pins[PinIndex]->LinkedTo[LinkIndex]->GetOwningNode()))
+						EntryGatewayPins.Add(Node->Pins[PinIndex]);
+					}
+					else
+					{
+						for (int32 LinkIndex = 0; LinkIndex < Node->Pins[PinIndex]->LinkedTo.Num(); ++LinkIndex)
 						{
-							if(Node->Pins[PinIndex]->Direction == EGPD_Input)
-							{
-								// For input pins, there must be a single connection 
-								if(InputConnection == NULL || InputConnection == Node->Pins[PinIndex])
-								{
-									InputConnection = Node->Pins[PinIndex];
-								}
-								else
-								{
-									// Check if the input connection was linked, report what node it is connected to
-									LogResults.Error(*LOCTEXT("TooManyPathsMultipleInput_Error", "Found too many input connections in selection! @@ is connected to @@, previously found @@ connected to @@").ToString(), Node, Node->Pins[PinIndex]->LinkedTo[LinkIndex]->GetOwningNode(), InputConnection->GetOwningNode(), InputConnection->LinkedTo[0]->GetOwningNode());
-									bBadConnection = true;
-								}
-							}
-							else
+							if (!InSelection.Contains(Node->Pins[PinIndex]->LinkedTo[LinkIndex]->GetOwningNode()))
 							{
-								// For output pins, as long as they all connect to the same pin, we consider the selection valid for being made into a function
-								if(OutputConnection == NULL || OutputConnection == Node->Pins[PinIndex]->LinkedTo[LinkIndex])
+								if (Node->Pins[PinIndex]->Direction == EGPD_Input)
 								{
-									OutputConnection = Node->Pins[PinIndex]->LinkedTo[LinkIndex];
+									// For input pins, there must be a single connection 
+									if (InputConnection == NULL || InputConnection == Node->Pins[PinIndex])
+									{
+										EntryGatewayPins.Add(Node->Pins[PinIndex]);
+										InputConnection = Node->Pins[PinIndex];
+									}
+									else
+									{
+										// Check if the input connection was linked, report what node it is connected to
+										LogResults.Error(*LOCTEXT("TooManyPathsMultipleInput_Error", "Found too many input connections in selection! @@ is connected to @@, previously found @@ connected to @@").ToString(), Node, Node->Pins[PinIndex]->LinkedTo[LinkIndex]->GetOwningNode(), InputConnection->GetOwningNode(), InputConnection->LinkedTo[0]->GetOwningNode());
+										bBadConnection = true;
+									}
 								}
 								else
 								{
-									check(OutputConnection->LinkedTo.Num());
-
-									LogResults.Error(*LOCTEXT("TooManyPathsMultipleOutput_Error", "Found too many output connections in selection! @@ is connected to @@, previously found @@ connected to @@").ToString(), Node, Node->Pins[PinIndex]->LinkedTo[LinkIndex]->GetOwningNode(), OutputConnection->GetOwningNode(), OutputConnection->LinkedTo[0]->GetOwningNode());
-									bBadConnection = true;
+									// For output pins, as long as they all connect to the same pin, we consider the selection valid for being made into a function
+									if (OutputConnection == NULL || OutputConnection == Node->Pins[PinIndex]->LinkedTo[LinkIndex])
+									{
+										OutputConnection = Node->Pins[PinIndex]->LinkedTo[LinkIndex];
+									}
+									else
+									{
+										check(OutputConnection->LinkedTo.Num());
+
+										LogResults.Error(*LOCTEXT("TooManyPathsMultipleOutput_Error", "Found too many output connections in selection! @@ is connected to @@, previously found @@ connected to @@").ToString(), Node, Node->Pins[PinIndex]->LinkedTo[LinkIndex]->GetOwningNode(), OutputConnection->GetOwningNode(), OutputConnection->LinkedTo[0]->GetOwningNode());
+										bBadConnection = true;
+									}
 								}
 							}
 						}
@@ -4610,6 +4620,13 @@ bool FBlueprintEditor::CanCollapseSelectionToFunction(TSet<class UEdGraphNode*>&
 		}
 	}
 
+	if (!bBadConnection && InputConnection == nullptr && EntryGatewayPins.Num() > 1)
+	{
+		// Too many input gateway pins with no connections.
+		LogResults.Error(*LOCTEXT("AmbiguousEntryPaths_Error", "Multiple entry pin possibilities. Unable to convert to a function. Make sure that selection either has only 1 entry pin or exactly 1 entry pin has a connection.").ToString());
+		bBadConnection = true;
+	}
+
 	// No need to check for cycling if the selection is invalid anyways.
 	if(!bBadConnection && FBlueprintEditorUtils::CheckIfSelectionIsCycling(InSelection, LogResults))
 	{
@@ -6483,6 +6500,10 @@ void FBlueprintEditor::CollapseNodesIntoGraph(UEdGraphNode* InGatewayNode, UK2No
 
 	UEdGraphNode* InterfaceTemplateNode = nullptr;
 
+	// If our return node only contains an exec pin, then we don't need to add it
+	// This helps to mitigate cases where it is unclear which exec pins should be connected to the return node
+	bool bDiscardReturnNode = true;
+
 	// For collapsing to functions can use a single event as a template for the function. This event MUST be deleted at the end, and the pins pre-generated. 
 	if(InGatewayNode->GetClass() == UK2Node_CallFunction::StaticClass())
 	{
@@ -6599,6 +6620,11 @@ void FBlueprintEditor::CollapseNodesIntoGraph(UEdGraphNode* InGatewayNode, UK2No
 
 						if(!RemotePortPin && !LocalPortPin)
 						{
+							if (LocalPin->Direction == EGPD_Output)
+							{
+								bDiscardReturnNode = false;
+							}
+							
 							RemotePortPin = InGatewayNode->CreatePin(LocalPin->Direction, LocalPin->PinType, UniquePortName);
 							LocalPortPin = LocalPort->CreateUserDefinedPin(UniquePortName, LocalPin->PinType, (LocalPin->Direction == EGPD_Input)? EGPD_Output : EGPD_Input);
 						}
@@ -6622,6 +6648,13 @@ void FBlueprintEditor::CollapseNodesIntoGraph(UEdGraphNode* InGatewayNode, UK2No
 						RemotePin->LinkedTo.Remove(LocalPin);
 						RemotePin->MakeLinkTo(RemotePortPin);
 
+						// The Entry Node only supports a single link, so if we made links above
+						// we need to break them now, to make room for the new link.
+						if (LocalPort == InEntryNode)
+						{
+							LocalPortPin->BreakAllPinLinks();
+						}
+
 						// Fix up the local pin
 						LocalPin->LinkedTo.Remove(RemotePin);
 						--LinkIndex;
@@ -6679,6 +6712,11 @@ void FBlueprintEditor::CollapseNodesIntoGraph(UEdGraphNode* InGatewayNode, UK2No
 		InResultNode->NodePosY = CenterY;
 		InResultNode->SnapToGrid(SNodePanel::GetSnapGridSize());
 	}
+
+	if (bDiscardReturnNode)
+	{
+		InResultNode->DestroyNode();
+	}
 }
 
 void FBlueprintEditor::CollapseNodes(TSet<UEdGraphNode*>& InCollapsableNodes)
@@ -8007,9 +8045,11 @@ void FBlueprintEditor::UpdatePreviewActor(UBlueprint* InBlueprint, bool bInForce
 			SpawnInfo.bNoFail = true;
 			SpawnInfo.ObjectFlags = RF_Transient|RF_Transactional;
 
-			// Temporarily remove the deprecated flag so we can respawn the Blueprint in the viewport
-			bool bIsClassDeprecated = PreviewBlueprint->GeneratedClass->HasAnyClassFlags(CLASS_Deprecated);
+			// Temporarily remove the deprecated and abstract flag so we can respawn the Blueprint in the viewport
+			const bool bIsClassDeprecated = PreviewBlueprint->GeneratedClass->HasAnyClassFlags(CLASS_Deprecated);
+			const bool bIsClassAbstract = PreviewBlueprint->GeneratedClass->HasAnyClassFlags(CLASS_Abstract);
 			PreviewBlueprint->GeneratedClass->ClassFlags &= ~CLASS_Deprecated;
+			PreviewBlueprint->GeneratedClass->ClassFlags &= ~CLASS_Abstract;
 
 			PreviewActorPtr = PreviewActor = PreviewScene.GetWorld()->SpawnActor(PreviewBlueprint->GeneratedClass, &SpawnLocation, &SpawnRotation, SpawnInfo);
 
@@ -8018,6 +8058,12 @@ void FBlueprintEditor::UpdatePreviewActor(UBlueprint* InBlueprint, bool bInForce
 			{
 				PreviewBlueprint->GeneratedClass->ClassFlags |= CLASS_Deprecated;
 			}
+			 
+			// Reassign the abstract flag if it was previously assigned
+			if (bIsClassAbstract)
+			{
+				PreviewBlueprint->GeneratedClass->ClassFlags |= CLASS_Abstract;
+			}
 
 			check(PreviewActor);
 
diff --git a/Engine/Source/Editor/Kismet/Private/Profiler/BPProfilerStatisticWidgets.cpp b/Engine/Source/Editor/Kismet/Private/Profiler/BPProfilerStatisticWidgets.cpp
--- a/Engine/Source/Editor/Kismet/Private/Profiler/BPProfilerStatisticWidgets.cpp
+++ b/Engine/Source/Editor/Kismet/Private/Profiler/BPProfilerStatisticWidgets.cpp
@@ -4,6 +4,7 @@
 #include "SBlueprintProfilerView.h"
 #include "BPProfilerStatisticWidgets.h"
 #include "EventExecution.h"
+#include "SProfilerStatExpander.h"
 #include "Developer/BlueprintProfiler/Public/BlueprintProfilerModule.h"
 
 #define LOCTEXT_NAMESPACE "BlueprintProfilerViewTypesUI"
@@ -43,7 +44,7 @@ TSharedRef<SWidget> SProfilerStatRow::GenerateWidgetForColumn(const FName& Colum
 			.HAlign(HAlign_Left)
 			.VAlign(VAlign_Center)
 			[
-				SNew(SExpanderArrow, SharedThis(this))
+				SNew(SProfilerStatExpander, SharedThis(this))
 				.IndentAmount(15.f)
 			]
 			+SHorizontalBox::Slot()
@@ -104,7 +105,7 @@ FBPProfilerStatWidget::FBPProfilerStatWidget(TSharedPtr<class FScriptExecutionNo
 	: WidgetTracePath(WidgetTracePathIn)
 	, ExecNode(InExecNode)
 {
-	if (ExecNode->HasFlags(EScriptExecutionNodeFlags::FunctionTunnel))
+	if (ExecNode->HasFlags(EScriptExecutionNodeFlags::TunnelInstance))
 	{
 		WidgetTracePath = FTracePath(WidgetTracePathIn, StaticCastSharedPtr<FScriptExecutionTunnelEntry>(InExecNode));
 	}
@@ -120,13 +121,13 @@ TSharedRef<SWidget> FBPProfilerStatWidget::GenerateColumnWidget(FName ColumnName
 				+SHorizontalBox::Slot()
 				.AutoWidth()
 				[
-					ExecNode->GetIconWidget()
+					ExecNode->GetIconWidget(WidgetTracePath)
 				]
 				+SHorizontalBox::Slot()
 				.AutoWidth()
 				.Padding(FMargin(5,0))
 				[
-					ExecNode->GetHyperlinkWidget()
+					ExecNode->GetHyperlinkWidget(WidgetTracePath)
 		#if TRACEPATH_DEBUG
 				]
 				+SHorizontalBox::Slot()
@@ -213,94 +214,215 @@ void FBPProfilerStatWidget::GenerateExecNodeWidgets(const TSharedPtr<FBlueprintP
 {
 	if (ExecNode.IsValid())
 	{
-		// Grab Performance Stats
+		// Create perf stats entry
 		PerformanceStats = ExecNode->GetOrAddPerfDataByInstanceAndTracePath(DisplayOptions->GetActiveInstance(), WidgetTracePath);
 		CachedChildren.Reset(0);
-
+		// Check for specialised variants
 		if (ExecNode->HasFlags(EScriptExecutionNodeFlags::PureStats))
 		{
 			if (ExecNode->IsPureChain())
 			{
-				// Get the full pure node chain associated with this exec node.
-				TMap<int32, TSharedPtr<FScriptExecutionNode>> AllPureNodes;
-				ExecNode->GetAllPureNodes(AllPureNodes);
-
-				// Build trace path, tree view node widget and register perf stats for tracking.
-				FTracePath PureTracePath(WidgetTracePath);
-				for (auto Iter : AllPureNodes)
-				{
-					PureTracePath.AddExitPin(Iter.Key);
-					TSharedPtr<FBPProfilerStatWidget> NewPureChildNode = MakeShareable<FBPProfilerStatWidget>(new FBPProfilerStatWidget(Iter.Value, PureTracePath));
-					NewPureChildNode->GenerateExecNodeWidgets(DisplayOptions);
-
-					// Pure nodes are shown in reverse execution order.
-					CachedChildren.Insert(NewPureChildNode, 0);
-				}
+				// Generate the specialised pure chain widget.
+				GeneratePureNodeWidgets(DisplayOptions);
 			}
 		}
-		else
+		else if (ExecNode->IsTunnelEntry())
+		{
+			TSharedPtr<FScriptExecutionTunnelEntry> TunnelEntryNode = StaticCastSharedPtr<FScriptExecutionTunnelEntry>(ExecNode);
+			// Generate specialised tunnel node widget.
+			if (TunnelEntryNode->IsComplexTunnel())
+			{
+				// Create a widget that houses the exit sites tracepaths as children.
+				GenerateComplexTunnelWidgets(TunnelEntryNode, DisplayOptions);
+			}
+			else
+			{
+				// Create a widget that houses only tunnel exec nodes as children.
+				GenerateSimpleTunnelWidgets(TunnelEntryNode, DisplayOptions);
+			}
+		}
+		else	
+		{
+			// Create a standard set of widgets that serves for most exec nodes.
+			GenerateStandardNodeWidgets(DisplayOptions);
+		}
+	}
+}
+
+void FBPProfilerStatWidget::GenerateStandardNodeWidgets(const TSharedPtr<FBlueprintProfilerStatOptions> DisplayOptions)
+{
+	// Add any child nodes that are not filtered out.
+	if (ExecNode->HasFlags(EScriptExecutionNodeFlags::Container))
+	{
+		for (auto Iter : ExecNode->GetChildNodes())
 		{
-			TArray<FScriptNodeExecLinkage::FLinearExecPath> Children;
-			ExecNode->GetFilteredChildNodes(Children, WidgetTracePath);
-			for (auto Iter : Children)
+			// Filter out events based on graph
+			if (!DisplayOptions->IsFiltered(Iter))
 			{
-				// Filter out events based on graph
-				if (!DisplayOptions->IsFiltered(Iter.LinkedNode))
+				TArray<FScriptNodeExecLinkage::FLinearExecPath> LinearExecNodes;
+				FTracePath ChildTracePath(WidgetTracePath);
+				Iter->GetLinearExecutionPath(LinearExecNodes, ChildTracePath, false);
+				if (LinearExecNodes.Num() > 1)
 				{
-					TArray<FScriptNodeExecLinkage::FLinearExecPath> LinearExecNodes;
-					FTracePath ChildTracePath(Iter.TracePath);
-					Iter.LinkedNode->GetLinearExecutionPath(LinearExecNodes, ChildTracePath, Iter.bIncludeChildren);
-					if (LinearExecNodes.Num() > 1)
+					TSharedPtr<FBPProfilerStatWidget> ChildContainer = AsShared();
+					for (auto LinearPathIter : LinearExecNodes)
 					{
-						TSharedPtr<FBPProfilerStatWidget> ChildContainer = AsShared();
-						for (auto LinearPathIter : LinearExecNodes)
+						if (!DisplayOptions->IsFiltered(LinearPathIter.LinkedNode))
 						{
-							if (!DisplayOptions->IsFiltered(LinearPathIter.LinkedNode))
+							TSharedPtr<FBPProfilerStatWidget> NewLinkedNode = MakeShareable<FBPProfilerStatWidget>(new FBPProfilerStatWidget(LinearPathIter.LinkedNode, LinearPathIter.TracePath));
+							NewLinkedNode->GenerateExecNodeWidgets(DisplayOptions);
+							ChildContainer->CachedChildren.Add(NewLinkedNode);
+							if (LinearPathIter.LinkedNode->HasFlags(EScriptExecutionNodeFlags::Container))
 							{
-								TSharedPtr<FBPProfilerStatWidget> NewLinkedNode = MakeShareable<FBPProfilerStatWidget>(new FBPProfilerStatWidget(LinearPathIter.LinkedNode, LinearPathIter.TracePath));
-								NewLinkedNode->GenerateExecNodeWidgets(DisplayOptions);
-								ChildContainer->CachedChildren.Add(NewLinkedNode);
-								if (LinearPathIter.LinkedNode->HasFlags(EScriptExecutionNodeFlags::Container))
-								{
-									ChildContainer = NewLinkedNode;
-								}
+								ChildContainer = NewLinkedNode;
 							}
 						}
 					}
-					else
-					{
-						TSharedPtr<FBPProfilerStatWidget> NewChildNode = MakeShareable<FBPProfilerStatWidget>(new FBPProfilerStatWidget(Iter.LinkedNode, ChildTracePath));
-						NewChildNode->GenerateExecNodeWidgets(DisplayOptions);
-						CachedChildren.Add(NewChildNode);
-					}
+				}
+				else
+				{
+					TSharedPtr<FBPProfilerStatWidget> NewChildNode = MakeShareable<FBPProfilerStatWidget>(new FBPProfilerStatWidget(Iter, ChildTracePath));
+					NewChildNode->GenerateExecNodeWidgets(DisplayOptions);
+					CachedChildren.Add(NewChildNode);
 				}
 			}
-			if (ExecNode->IsBranch())
+		}
+	}
+	// Follow execution links and generate node widgets
+	const bool bShouldBuildLinks = ExecNode->IsBranch()||ExecNode->HasFlags(EScriptExecutionNodeFlags::TunnelEntryPin);
+	if (bShouldBuildLinks)
+	{
+		for (auto LinkIter : ExecNode->GetLinkedNodes())
+		{
+			if (!DisplayOptions->IsFiltered(LinkIter.Value))
 			{
-				for (auto LinkIter : ExecNode->GetLinkedNodes())
+				TArray<FScriptNodeExecLinkage::FLinearExecPath> LinearExecNodes;
+				FTracePath LinkPath(WidgetTracePath);
+				if (!LinkIter.Value->HasFlags(EScriptExecutionNodeFlags::InvalidTrace))
+				{
+					LinkPath.AddExitPin(LinkIter.Key);
+				}
+				LinkIter.Value->GetLinearExecutionPath(LinearExecNodes, LinkPath);
+				for (auto LinearPathIter : LinearExecNodes)
 				{
-					if (!DisplayOptions->IsFiltered(LinkIter.Value))
+					TSharedPtr<FBPProfilerStatWidget> NewLinkedNode = MakeShareable<FBPProfilerStatWidget>(new FBPProfilerStatWidget(LinearPathIter.LinkedNode, LinearPathIter.TracePath));
+					NewLinkedNode->GenerateExecNodeWidgets(DisplayOptions);
+					CachedChildren.Add(NewLinkedNode);
+				}
+			}
+		}
+	}
+}
+
+void FBPProfilerStatWidget::GeneratePureNodeWidgets(const TSharedPtr<FBlueprintProfilerStatOptions> DisplayOptions)
+{
+	// Get the full pure node chain associated with this exec node.
+	TMap<int32, TSharedPtr<FScriptExecutionNode>> AllPureNodes;
+	ExecNode->GetAllPureNodes(AllPureNodes);
+
+	// Build trace path, tree view node widget and register perf stats for tracking.
+	FTracePath PureTracePath(WidgetTracePath);
+	for (auto Iter : AllPureNodes)
+	{
+		PureTracePath.AddExitPin(Iter.Key);
+		TSharedPtr<FBPProfilerStatWidget> NewPureChildNode = MakeShareable<FBPProfilerStatWidget>(new FBPProfilerStatWidget(Iter.Value, PureTracePath));
+		NewPureChildNode->GenerateExecNodeWidgets(DisplayOptions);
+
+		// Pure nodes are shown in reverse execution order.
+		CachedChildren.Insert(NewPureChildNode, 0);
+	}
+}
+
+void FBPProfilerStatWidget::GenerateSimpleTunnelWidgets(TSharedPtr<FScriptExecutionTunnelEntry> TunnelEntryNode, const TSharedPtr<FBlueprintProfilerStatOptions> DisplayOptions)
+{
+	// Map entry points
+	for (auto EntryPoint : TunnelEntryNode->GetChildNodes())
+	{
+		// Filter out events based on graph
+		if (!DisplayOptions->IsFiltered(EntryPoint))
+		{
+			TArray<FScriptNodeExecLinkage::FLinearExecPath> LinearExecNodes;
+			EntryPoint->GetLinearExecutionPath(LinearExecNodes, WidgetTracePath, true);
+			if (LinearExecNodes.Num() > 1)
+			{
+				TSharedPtr<FBPProfilerStatWidget> ChildContainer = AsShared();
+				for (auto LinearPathIter : LinearExecNodes)
+				{
+					const bool bInternalTunnelBoundary = TunnelEntryNode->IsInternalBoundary(LinearPathIter.LinkedNode);
+					if (!DisplayOptions->IsFiltered(LinearPathIter.LinkedNode) && !bInternalTunnelBoundary)
 					{
-						TArray<FScriptNodeExecLinkage::FLinearExecPath> LinearExecNodes;
-						FTracePath LinkPath(WidgetTracePath);
-						if (LinkIter.Value->HasFlags(EScriptExecutionNodeFlags::AsyncTaskDelegate))
-						{
-							LinkPath.ResetPath();
-						}
-						if (!LinkIter.Value->HasFlags(EScriptExecutionNodeFlags::InvalidTrace))
-						{
-							LinkPath.AddExitPin(LinkIter.Key);
-						}
-						LinkIter.Value->GetLinearExecutionPath(LinearExecNodes, LinkPath);
-						for (auto LinearPathIter : LinearExecNodes)
+						TSharedPtr<FBPProfilerStatWidget> NewLinkedNode = MakeShareable<FBPProfilerStatWidget>(new FBPProfilerStatWidget(LinearPathIter.LinkedNode, LinearPathIter.TracePath));
+						NewLinkedNode->GenerateExecNodeWidgets(DisplayOptions);
+						CachedChildren.Add(NewLinkedNode);
+						if (LinearPathIter.LinkedNode->HasFlags(EScriptExecutionNodeFlags::Container))
 						{
-							TSharedPtr<FBPProfilerStatWidget> NewLinkedNode = MakeShareable<FBPProfilerStatWidget>(new FBPProfilerStatWidget(LinearPathIter.LinkedNode, LinearPathIter.TracePath));
-							NewLinkedNode->GenerateExecNodeWidgets(DisplayOptions);
-							CachedChildren.Add(NewLinkedNode);
+							ChildContainer = NewLinkedNode;
 						}
 					}
 				}
 			}
+			else
+			{
+				TSharedPtr<FBPProfilerStatWidget> NewChildNode = MakeShareable<FBPProfilerStatWidget>(new FBPProfilerStatWidget(EntryPoint, WidgetTracePath));
+				NewChildNode->GenerateExecNodeWidgets(DisplayOptions);
+				CachedChildren.Add(NewChildNode);
+			}
+		}
+	}
+}
+
+void FBPProfilerStatWidget::GenerateComplexTunnelWidgets(TSharedPtr<FScriptExecutionTunnelEntry> TunnelEntryNode, const TSharedPtr<FBlueprintProfilerStatOptions> DisplayOptions)
+{
+	// Map entry points
+	for (auto EntryPoint : TunnelEntryNode->GetChildNodes())
+	{
+		// Filter out events based on graph
+		if (!DisplayOptions->IsFiltered(EntryPoint))
+		{
+			TSharedPtr<FBPProfilerStatWidget> EntryPointNode = MakeShareable<FBPProfilerStatWidget>(new FBPProfilerStatWidget(EntryPoint, WidgetTracePath));
+			EntryPointNode->GenerateTunnelLinkWidgets(DisplayOptions);
+			CachedChildren.Add(EntryPointNode);
+		}
+	}
+	// Create exit links as child nodes
+	for (auto ExitSite : TunnelEntryNode->GetExitSites())
+	{
+		TSharedPtr<FBPProfilerStatWidget> ExitSiteNode = MakeShareable<FBPProfilerStatWidget>(new FBPProfilerStatWidget(ExitSite.Value, WidgetTracePath));
+		ExitSiteNode->GenerateTunnelLinkWidgets(DisplayOptions, ExitSite.Key);
+		CachedChildren.Add(ExitSiteNode);
+		if (ExitSite.Value->IsTunnelExit())
+		{
+			TSharedPtr<FScriptExecutionTunnelExit> ExitSiteTypedNode = StaticCastSharedPtr<FScriptExecutionTunnelExit>(ExitSite.Value);
+			// Allow exit sites to jump to external pins using the correct widget path.
+			ExitSiteTypedNode->AddInstanceExitSite(WidgetTracePath, ExitSiteTypedNode->GetExternalPin());
+		}
+	}
+}
+
+void FBPProfilerStatWidget::GenerateTunnelLinkWidgets(const TSharedPtr<FBlueprintProfilerStatOptions> DisplayOptions, const int32 ScriptOffset)
+{
+	// Create perf stats entry
+	PerformanceStats = ExecNode->GetOrAddPerfDataByInstanceAndTracePath(DisplayOptions->GetActiveInstance(), WidgetTracePath);
+	CachedChildren.Reset(0);
+	// Create child tracepath
+	FTracePath ChildTrace(WidgetTracePath);
+	if (ScriptOffset != INDEX_NONE)
+	{
+		ChildTrace.AddExitPin(ScriptOffset);
+	}
+	// Build linked widgets
+	for (auto LinkIter : ExecNode->GetLinkedNodes())
+	{
+		if (!DisplayOptions->IsFiltered(LinkIter.Value))
+		{
+			TArray<FScriptNodeExecLinkage::FLinearExecPath> LinearExecNodes;
+			LinkIter.Value->GetLinearExecutionPath(LinearExecNodes, ChildTrace);
+			for (auto LinearPathIter : LinearExecNodes)
+			{
+				TSharedPtr<FBPProfilerStatWidget> NewLinkedNode = MakeShareable<FBPProfilerStatWidget>(new FBPProfilerStatWidget(LinearPathIter.LinkedNode, LinearPathIter.TracePath));
+				NewLinkedNode->GenerateExecNodeWidgets(DisplayOptions);
+				CachedChildren.Add(NewLinkedNode);
+			}
 		}
 	}
 }
diff --git a/Engine/Source/Editor/Kismet/Private/Profiler/BPProfilerStatisticWidgets.h b/Engine/Source/Editor/Kismet/Private/Profiler/BPProfilerStatisticWidgets.h
--- a/Engine/Source/Editor/Kismet/Private/Profiler/BPProfilerStatisticWidgets.h
+++ b/Engine/Source/Editor/Kismet/Private/Profiler/BPProfilerStatisticWidgets.h
@@ -65,7 +65,7 @@ public:
 	virtual ~FBPProfilerStatWidget() {}
 
 	/** Generate exec node widgets */
-	virtual void GenerateExecNodeWidgets(const TSharedPtr<struct FBlueprintProfilerStatOptions> DisplayOptions);
+	void GenerateExecNodeWidgets(const TSharedPtr<struct FBlueprintProfilerStatOptions> DisplayOptions);
 
 	/** Gather children for list/tree widget */
 	void GatherChildren(TArray<TSharedPtr<FBPProfilerStatWidget>>& OutChildren);
@@ -96,6 +96,21 @@ public:
 
 protected:
 
+	/** Generate standard node widgets */
+	void GenerateStandardNodeWidgets(const TSharedPtr<struct FBlueprintProfilerStatOptions> DisplayOptions);
+
+	/** Generate pure node widgets */
+	void GeneratePureNodeWidgets(const TSharedPtr<FBlueprintProfilerStatOptions> DisplayOptions);
+
+	/** Generate simple tunnel widgets */
+	void GenerateSimpleTunnelWidgets(TSharedPtr<FScriptExecutionTunnelEntry> TunnelEntryNode, const TSharedPtr<FBlueprintProfilerStatOptions> DisplayOptions);
+
+	/** Generate complex tunnel widgets */
+	void GenerateComplexTunnelWidgets(TSharedPtr<FScriptExecutionTunnelEntry> TunnelEntryNode, const TSharedPtr<FBlueprintProfilerStatOptions> DisplayOptions);
+
+	/** Generate tunnel exit site linked widgets */
+	void GenerateTunnelLinkWidgets(const TSharedPtr<FBlueprintProfilerStatOptions> DisplayOptions, const int32 ScriptOffset = INDEX_NONE);
+
 	/** Navigate to referenced node */
 	void NavigateTo() const;
 
