Copying //UE4/Dev-Blueprints to //UE4/Dev-Main (Source: //UE4/Dev-Blueprints @ 3042808)

diff --git a/Engine/Source/Editor/Kismet/Private/Profiler/EventExecution.cpp b/Engine/Source/Editor/Kismet/Private/Profiler/EventExecution.cpp
--- a/Engine/Source/Editor/Kismet/Private/Profiler/EventExecution.cpp
+++ b/Engine/Source/Editor/Kismet/Private/Profiler/EventExecution.cpp
@@ -29,14 +29,6 @@ TSharedPtr<FScriptExecutionNode> FScriptNodeExecLinkage::GetLinkedNodeByScriptOf
 	return Result;
 }
 
-void FScriptNodeExecLinkage::GetFilteredChildNodes(TArray<FLinearExecPath>& ChildArrayInOut, const FTracePath& TracePath)
-{
-	for (auto Child : ChildNodes)
-	{
-		ChildArrayInOut.Add(FLinearExecPath(Child, TracePath));
-	}
-}
-
 //////////////////////////////////////////////////////////////////////////
 // FScriptNodePerfData
 
@@ -177,14 +169,14 @@ bool FScriptExecutionNode::operator == (const FScriptExecutionNode& NodeIn) cons
 	return NodeName == NodeIn.NodeName;
 }
 
-TSharedRef<SWidget> FScriptExecutionNode::GetIconWidget()
+TSharedRef<SWidget> FScriptExecutionNode::GetIconWidget(const uint32 /*TracePath*/)
 {
 	return SNew(SImage)
 			.Image(Icon)
 			.ColorAndOpacity(IconColor);
 }
 
-TSharedRef<SWidget> FScriptExecutionNode::GetHyperlinkWidget()
+TSharedRef<SWidget> FScriptExecutionNode::GetHyperlinkWidget(const uint32 /*TracePath*/)
 {
 	return SNew(SHyperlink)
 		#if TRACEPATH_DEBUG
@@ -219,7 +211,7 @@ void FScriptExecutionNode::GetLinearExecutionPath(TArray<FLinearExecPath>& Linea
 			else
 			{
 				FTracePath NewTracePath(TracePath);
-				if (NodeIter.Value->HasFlags(EScriptExecutionNodeFlags::AsyncTaskDelegate))
+				if (NodeIter.Value->HasFlags(EScriptExecutionNodeFlags::EventPin))
 				{
 					NewTracePath.ResetPath();
 				}
@@ -235,7 +227,7 @@ void FScriptExecutionNode::GetLinearExecutionPath(TArray<FLinearExecPath>& Linea
 
 void FScriptExecutionNode::MapTunnelLinearExecution(FTracePath& Trace) const
 {
-	if (HasFlags(EScriptExecutionNodeFlags::FunctionTunnel))
+	if (HasFlags(EScriptExecutionNodeFlags::TunnelInstance))
 	{
 		for (auto ChildIter : ChildNodes)
 		{
@@ -282,13 +274,13 @@ void FScriptExecutionNode::RefreshStats(const FTracePath& TracePath)
 	// Update stats based on type
 	switch(GetStatisticContainerType())
 	{
+		case EScriptStatContainerType::Standard:
 		case EScriptStatContainerType::Container:
 		{
 			GetInstancePerfDataByTracePath(TracePath, BlueprintPooledStats);
 			bRefreshChildStats = true;
 			break;
 		}
-		case EScriptStatContainerType::Standard:
 		case EScriptStatContainerType::SequentialBranch:
 		{
 			GetInstancePerfDataByTracePath(TracePath, BlueprintPooledStats);
@@ -347,7 +339,7 @@ void FScriptExecutionNode::RefreshStats(const FTracePath& TracePath)
 		for (auto LinkIter : LinkedNodes)
 		{
 			FTracePath LinkTracePath(TracePath);
-			if (LinkIter.Value->HasFlags(EScriptExecutionNodeFlags::AsyncTaskDelegate))
+			if (LinkIter.Value->HasFlags(EScriptExecutionNodeFlags::EventPin))
 			{
 				LinkTracePath.ResetPath();
 			}
@@ -484,140 +476,106 @@ void FScriptExecutionTunnelInstance::GetBlueprintPerfDataForAllTracePaths(FScrip
 {
 	OutPerfData.Reset();
 
-	for (auto CustomEntryPoint : CustomEntryPoints)
+	for (auto EntrySite : EntrySites)
 	{
-		CustomEntryPoint->GetBlueprintPerfDataForAllTracePaths(OutPerfData);
+		EntrySite.Value->GetBlueprintPerfDataForAllTracePaths(OutPerfData);
 	}
 }
 
-//////////////////////////////////////////////////////////////////////////
-// FScriptExecutionTunnelEntry
+void FScriptExecutionTunnelInstance::AddEntrySite(const int32 ScriptOffset, TSharedPtr<FScriptExecutionTunnelEntry> EntrySite)
+{
+	EntrySites.FindOrAdd(ScriptOffset) = EntrySite;
+	EntrySite->SetTunnelInstance(SharedThis<FScriptExecutionTunnelInstance>(this));
+}
 
-void FScriptExecutionTunnelEntry::GetLinearExecutionPath(TArray<FLinearExecPath>& LinearExecutionNodes, const FTracePath& TracePath, const bool bIncludeChildren)
+void FScriptExecutionTunnelInstance::CopyBoundarySites(TSharedPtr<FScriptExecutionTunnelInstance> OtherInstance)
 {
-	LinearExecutionNodes.Add(FLinearExecPath(AsShared(), TracePath));
-	if (TunnelType != MultiIOTunnel)
+	for (auto EntrySite : EntrySites)
 	{
-		FTracePath TunnelTracePath(TracePath, SharedThis<FScriptExecutionTunnelEntry>(this));
-		for (auto TunnelExit : LinkedNodes)
-		{
-			if (IsPathValidForTunnel(TunnelExit.Key))
-			{
-				TunnelTracePath.AddExitPin(TunnelExit.Key);
-				TunnelExit.Value->GetLinearExecutionPath(LinearExecutionNodes, TunnelTracePath);
-			}
-		}
+		OtherInstance->EntrySites.Add(EntrySite.Key) = EntrySite.Value;
+	}
+	for (auto ExitSite : ExitSites)
+	{
+		OtherInstance->ExitSites.Add(ExitSite.Key) = ExitSite.Value;
 	}
 }
 
-void FScriptExecutionTunnelEntry::GetFilteredChildNodes(TArray<FLinearExecPath>& ChildArrayInOut, const FTracePath& TracePath)
+TSharedPtr<FScriptExecutionNode> FScriptExecutionTunnelInstance::FindBoundarySite(const int32 ScriptOffset)
 {
-	ChildArrayInOut.Reset();
-	switch(TunnelType)
+	TSharedPtr<FScriptExecutionNode> Result;
+	if (TSharedPtr<FScriptExecutionTunnelEntry>* EntryResult = EntrySites.Find(ScriptOffset))
 	{
-		case SimpleTunnel:
-		{
-			FTracePath ChildTrace(TracePath, SharedThis<FScriptExecutionTunnelEntry>(this));
-			for (auto Child : ChildNodes)
-			{
-				if (Child->HasFlags(EScriptExecutionNodeFlags::TunnelEntryPin))
-				{
-					for (auto NodeToAdd : Child->GetChildNodes())
-					{
-						FTracePath NewTrace(ChildTrace);
-						ChildArrayInOut.Add(FLinearExecPath(NodeToAdd, NewTrace));
-					}
-				}
-				else if (!Child->HasFlags(EScriptExecutionNodeFlags::TunnelPin))
-				{
-					ChildArrayInOut.Add(FLinearExecPath(Child, ChildTrace));
-				}
-			}
-			break;
-		}
-		case SinglePathTunnel:
-		{
-			FTracePath ChildTrace(TracePath, SharedThis<FScriptExecutionTunnelEntry>(this));
-			for (auto Child : ChildNodes)
-			{
-				ChildArrayInOut.Add(FLinearExecPath(Child, ChildTrace));
-			}
-			break;
-		}
-		case MultiIOTunnel:
+		Result = *EntryResult;
+	}
+	else if (TSharedPtr<FScriptExecutionTunnelExit>* ExitResult = ExitSites.Find(ScriptOffset))
+	{
+		Result = *ExitResult;
+	}
+	return Result;
+}
+
+//////////////////////////////////////////////////////////////////////////
+// FScriptExecutionTunnelEntry
+
+void FScriptExecutionTunnelEntry::GetLinearExecutionPath(TArray<FLinearExecPath>& LinearExecutionNodes, const FTracePath& TracePath, const bool bIncludeChildren)
+{
+	LinearExecutionNodes.Add(FLinearExecPath(AsShared(), TracePath));
+	if (!IsComplexTunnel())
+	{
+		for (auto TunnelExit : LinkedNodes)
 		{
-			TArray<FLinearExecPath> Results;
-			FTracePath ChildTrace(TracePath, SharedThis<FScriptExecutionTunnelEntry>(this));
-			for (auto Child : ChildNodes)
+			for (auto LinkedToExit : TunnelExit.Value->GetLinkedNodes())
 			{
-				const bool bIgnoreChildren = Child->HasFlags(EScriptExecutionNodeFlags::TunnelExitPin);
-				ChildArrayInOut.Add(FLinearExecPath(Child, ChildTrace, bIgnoreChildren));
-				Child->GetLinearExecutionPath(Results, ChildTrace, true);
-			}
-			for (auto Node : Results)
-			{
-				if (Node.LinkedNode->HasFlags(EScriptExecutionNodeFlags::TunnelExitPin))
+				FTracePath ExitLinkTracePath(TracePath);
+				if (LinkedToExit.Key != INDEX_NONE && !LinkedToExit.Value->HasFlags(EScriptExecutionNodeFlags::InvalidTrace))
 				{
-					ChildArrayInOut.Add(FLinearExecPath(Node.LinkedNode, TracePath, true));
+					ExitLinkTracePath.AddExitPin(TunnelExit.Key);
 				}
+				LinkedToExit.Value->GetLinearExecutionPath(LinearExecutionNodes, ExitLinkTracePath);
 			}
-			break;
 		}
 	}
 }
 
-void FScriptExecutionTunnelEntry::AddTunnelTiming(const FName InstanceName, const FTracePath& TracePath, const double Time)
+void FScriptExecutionTunnelEntry::AddTunnelTiming(const FName InstanceName, const FTracePath& TracePath, const FTracePath& InternalTracePath, const int32 ExitScriptOffset, const double Time)
 {
 	TSharedPtr<FScriptPerfData> PerfData = GetPerfDataByInstanceAndTracePath(InstanceName, TracePath);
 	PerfData->AddEventTiming(Time);
+	// Tick Entry points
 	for (auto EntrySite : ChildNodes)
 	{
 		PerfData = EntrySite->GetPerfDataByInstanceAndTracePath(InstanceName, TracePath);
 		PerfData->TickSamples();
 	}
-}
-
-void FScriptExecutionTunnelEntry::UpdateTunnelExit(const FName InstanceName, const FTracePath& TracePath, const int32 ExitScriptOffset)
-{
-	// Tick exit site
+	// Tick exit sites
 	TSharedPtr<FScriptExecutionNode> TunnelExitNode = GetExitSite(ExitScriptOffset);
 	if (TunnelExitNode.IsValid())
 	{
-		TSharedPtr<FScriptPerfData> PerfData = TunnelExitNode->GetPerfDataByInstanceAndTracePath(InstanceName, TracePath);
+		PerfData = TunnelExitNode->GetPerfDataByInstanceAndTracePath(InstanceName, TracePath);
+		PerfData->TickSamples();
+		PerfData = TunnelExitNode->GetPerfDataByInstanceAndTracePath(InstanceName, InternalTracePath);
 		PerfData->TickSamples();
 	}
 }
 
-bool FScriptExecutionTunnelEntry::IsPathValidForTunnel(const int32 ScriptOffset) const
+void FScriptExecutionTunnelEntry::AddExitSite(const int32 ExitScriptOffset, TSharedPtr<FScriptExecutionTunnelExit> ExitSite)
 {
-	return ValidExitPoints.Contains(ScriptOffset);
+	LinkedNodes.Add(ExitScriptOffset) = ExitSite;
 }
 
-void FScriptExecutionTunnelEntry::AddExitSite(TSharedPtr<FScriptExecutionNode> ExitSite, const int32 ExitScriptOffset)
+TSharedPtr<FScriptExecutionTunnelExit> FScriptExecutionTunnelEntry::GetExitSite(const int32 ScriptOffset) const
 {
-	ValidExitPoints.FindOrAdd(ExitScriptOffset) = ExitSite;
-}
-
-TSharedPtr<FScriptExecutionNode> FScriptExecutionTunnelEntry::GetExitSite(const int32 ScriptOffset) const
-{
-	TSharedPtr<FScriptExecutionNode> ExitSite;
-	if (const TSharedPtr<FScriptExecutionNode>* SearchResult = ValidExitPoints.Find(ScriptOffset))
+	TSharedPtr<FScriptExecutionTunnelExit> ExitSite;
+	if (const TSharedPtr<FScriptExecutionNode>* SearchResult = LinkedNodes.Find(ScriptOffset))
 	{
-		ExitSite = *SearchResult;
+		if ((*SearchResult)->IsTunnelExit())
+		{
+			ExitSite = StaticCastSharedPtr<FScriptExecutionTunnelExit>(*SearchResult);
+		}
 	}
 	return ExitSite;
 }
 
-bool FScriptExecutionTunnelEntry::IsFinalExitSite(const int32 ScriptOffset) const
-{
-	bool bResult = false;
-	if (const TSharedPtr<FScriptExecutionNode>* SearchResult = ValidExitPoints.Find(ScriptOffset))
-	{
-		bResult = (*SearchResult)->HasFlags(EScriptExecutionNodeFlags::TunnelFinalExitPin);
-	}
-	return bResult;
-}
-
 void FScriptExecutionTunnelEntry::RefreshStats(const FTracePath& TracePath)
 {
 	// Process stat update
@@ -631,13 +589,6 @@ void FScriptExecutionTunnelEntry::RefreshStats(const FTracePath& TracePath)
 	{
 		ChildIter->RefreshStats(TunnelTracePath);
 	}
-	// Refresh All links
-	for (auto LinkIter : LinkedNodes)
-	{
-		FTracePath LinkTracePath(TracePath);
-		LinkTracePath.AddExitPin(LinkIter.Key);
-		LinkIter.Value->RefreshStats(LinkTracePath);
-	}
 	// Update the owning blueprint stats
 	if (BlueprintPooledStats.Num() > 0)
 	{
@@ -646,75 +597,51 @@ void FScriptExecutionTunnelEntry::RefreshStats(const FTracePath& TracePath)
 	}
 }
 
-void FScriptExecutionTunnelEntry::BuildExitBranches()
+bool FScriptExecutionTunnelEntry::IsInternalBoundary(TSharedPtr<FScriptExecutionNode> PotentialBoundaryNode) const
 {
-	if (TunnelType == ETunnelType::MultiIOTunnel)
+	bool bIsInternalBoundary = false;
+	if (PotentialBoundaryNode->IsTunnelBoundary())
 	{
-		BranchedExitSites.Reset();
-		TArray<FLinearExecPath> Results;
-		FTracePath TracePath(SharedThis<FScriptExecutionTunnelEntry>(this));
-		for (auto Child : ChildNodes)
+		bIsInternalBoundary = PotentialBoundaryNode->HasFlags(EScriptExecutionNodeFlags::TunnelEntryPin);
+		if (!PotentialBoundaryNode->IsTunnelEntry())
 		{
-			Child->GetLinearExecutionPath(Results, TracePath, true);
-		}
-		for (auto Node : Results)
-		{
-			if (Node.LinkedNode->HasFlags(EScriptExecutionNodeFlags::TunnelExitPin))
+			if (const int32* SearchResult = LinkedNodes.FindKey(PotentialBoundaryNode))
 			{
-				BranchedExitSites.Add(Node);
+				bIsInternalBoundary = true;
 			}
 		}
 	}
+	return bIsInternalBoundary;
 }
 
 //////////////////////////////////////////////////////////////////////////
 // FScriptExecutionTunnelExit
 
-void FScriptExecutionTunnelExit::GetLinearExecutionPath(TArray<FLinearExecPath>& LinearExecutionNodes, const FTracePath& TracePath, const bool bIncludeExitLinks)
+void FScriptExecutionTunnelExit::GetLinearExecutionPath(TArray<FLinearExecPath>& LinearExecutionNodes, const FTracePath& TracePath, const bool bIncludeChildren)
 {
-	TSharedPtr<const FScriptExecutionTunnelEntry> ActiveTunnel = TracePath.GetTunnel();
-	if (ActiveTunnel.IsValid())
-	{
-		if (ActiveTunnel->GetTunnelType() != FScriptExecutionTunnelEntry::SimpleTunnel)
-		{
-			LinearExecutionNodes.Add(FLinearExecPath(AsShared(), TracePath));
-			FTracePath NewTracePath(TracePath);
-			for (auto TunnelExit : LinkedNodes)
-			{
-				if (bIncludeExitLinks && ActiveTunnel->IsPathValidForTunnel(TunnelExit.Key))
-				{
-					NewTracePath.AddExitPin(TunnelExit.Key);
-					TunnelExit.Value->GetLinearExecutionPath(LinearExecutionNodes, NewTracePath);
-				}
-			}
-		}
-	}
+	LinearExecutionNodes.Add(FLinearExecPath(AsShared(), TracePath));
 }
 
 void FScriptExecutionTunnelExit::RefreshStats(const FTracePath& TracePath)
 {
 	// Process stat update
 	TArray<TSharedPtr<FScriptPerfData>> BlueprintPooledStats;
-	TSharedPtr<const FScriptExecutionTunnelEntry> TunnelNode = TracePath.GetTunnel();
-	check (TunnelNode.IsValid());
 	// Update stats based on type
 	GetInstancePerfDataByTracePath(TracePath, BlueprintPooledStats);
 	// Refresh all valid tunnel links
 	for (auto LinkIter : LinkedNodes)
 	{
-		if (TunnelNode->IsPathValidForTunnel(LinkIter.Key))
+		FTracePath LinkTracePath(TracePath);
+		LinkTracePath.ExitTunnel();
+		if (LinkIter.Value->HasFlags(EScriptExecutionNodeFlags::EventPin))
 		{
-			FTracePath LinkTracePath(TracePath);
-			if (LinkIter.Value->HasFlags(EScriptExecutionNodeFlags::AsyncTaskDelegate))
-			{
-				LinkTracePath.ResetPath();
-			}
-			if (!LinkIter.Value->HasFlags(EScriptExecutionNodeFlags::InvalidTrace))
-			{
-				LinkTracePath.AddExitPin(LinkIter.Key);
-			}
-			LinkIter.Value->RefreshStats(LinkTracePath);
+			LinkTracePath.ResetPath();
 		}
+		if (!LinkIter.Value->HasFlags(EScriptExecutionNodeFlags::InvalidTrace))
+		{
+			LinkTracePath.AddExitPin(LinkIter.Key);
+		}
+		LinkIter.Value->RefreshStats(LinkTracePath);
 	}
 	// Update the owning blueprint stats
 	if (BlueprintPooledStats.Num() > 0)
@@ -724,11 +651,52 @@ void FScriptExecutionTunnelExit::RefreshStats(const FTracePath& TracePath)
 	}
 }
 
-void FScriptExecutionTunnelExit::NavigateToObject() const
+void FScriptExecutionTunnelExit::AddInstanceExitSite(const uint32 TracePath, UEdGraphPin* InstanceExitPin)
 {
-	if (IsObservedObjectValid())
+	const FSlateBrush* InstanceIcon = InstanceExitPin->LinkedTo.Num() ?	FEditorStyle::GetBrush(TEXT("BlueprintProfiler.BPPinConnected")) : 
+																		FEditorStyle::GetBrush(TEXT("BlueprintProfiler.BPPinDisconnected"));
+	InstanceExitSiteInfo.FindOrAdd(TracePath) = FTunnelExitInstanceData(InstanceIcon, InstanceExitPin);
+}
+
+TSharedRef<SWidget> FScriptExecutionTunnelExit::GetIconWidget(const uint32 TracePath)
+{
+	if (FTunnelExitInstanceData* SearchResult = InstanceExitSiteInfo.Find(TracePath))
 	{
-		FKismetEditorUtilities::BringKismetToFocusAttentionOnObject(ObservedObject.Get());
+		return SNew(SImage)
+				.Image(SearchResult->InstanceIcon)
+				.ColorAndOpacity(IconColor);
+	}
+	else
+	{
+		return SNew(SImage)
+				.Image(Icon)
+				.ColorAndOpacity(IconColor);
+	}
+}
+
+TSharedRef<SWidget> FScriptExecutionTunnelExit::GetHyperlinkWidget(const uint32 TracePath)
+{
+	return SNew(SHyperlink)
+		#if TRACEPATH_DEBUG
+		.Text(FText::FromName(NodeName))
+		#else
+		.Text(DisplayName)
+		#endif
+		.Style(FEditorStyle::Get(), "HoverOnlyHyperlink")
+		.ToolTipText(Tooltip)
+		.OnNavigate(this, &FScriptExecutionTunnelExit::NavigateToExitSite, TracePath);
+}
+
+void FScriptExecutionTunnelExit::NavigateToExitSite(const uint32 TracePath) const
+{
+	UEdGraphPin* Pin = ObservedPin.Get();
+	if (const FTunnelExitInstanceData* SearchResult = InstanceExitSiteInfo.Find(TracePath))
+	{
+		Pin = SearchResult->InstancePin.Get();
+	}
+	if (Pin)
+	{
+		FKismetEditorUtilities::BringKismetToFocusAttentionOnObject(Pin->GetOwningNode());
 	}
 }
 
@@ -918,7 +886,11 @@ bool FScriptExecutionInstance::IsObservedObjectValid() const
 {
 	if (const FWorldContext* PIEWorldContext = GEditor->GetPIEWorldContext())
 	{
-		return PIEInstance.IsValid() && !PIEInstance.Get()->IsPendingKill();
+		if (PIEInstance.IsValid() && PIEInstance.Get()->IsPendingKill())
+		{
+			bPIEInstanceDestroyed = true;
+		}
+		return PIEInstance.IsValid();
 	}
 	else
 	{
@@ -926,14 +898,46 @@ bool FScriptExecutionInstance::IsObservedObjectValid() const
 	}
 }
 
-TSharedRef<SWidget> FScriptExecutionInstance::GetIconWidget()
+UObject* FScriptExecutionInstance::GetActiveObject() const
+{
+	UObject* Result = nullptr;
+	if (const FWorldContext* PIEWorldContext = GEditor->GetPIEWorldContext())
+	{
+		if (!PIEInstance.IsValid() && !bPIEInstanceDestroyed)
+		{
+			RefreshPIEInstance();
+		}
+		Result = const_cast<UObject*>(PIEInstance.Get());
+	}
+	else
+	{
+		Result = const_cast<UObject*>(ObservedObject.Get());
+	}
+	return Result;
+}
+
+void FScriptExecutionInstance::RefreshPIEInstance() const
+{
+	if (!bPIEInstanceDestroyed)
+	{
+		const FWorldContext* PIEWorldContext = GEditor->GetPIEWorldContext();
+		UWorld* PIEWorld = PIEWorldContext ? PIEWorldContext->World() : nullptr;
+		ULevel* CurrentLevel = PIEWorld ? PIEWorld->GetCurrentLevel() : nullptr;
+		if (CurrentLevel && ObservedObject.IsValid())
+		{
+			PIEInstance = FindObject<UObject>(CurrentLevel, *ObservedObject.Get()->GetName());
+		}
+	}
+}
+
+TSharedRef<SWidget> FScriptExecutionInstance::GetIconWidget(const uint32 /*TracePath*/)
 {
 	return SNew(SImage)
 			.Image(Icon)
 			.ColorAndOpacity(this, &FScriptExecutionInstance::GetInstanceIconColor);
 }
 
-TSharedRef<SWidget> FScriptExecutionInstance::GetHyperlinkWidget()
+TSharedRef<SWidget> FScriptExecutionInstance::GetHyperlinkWidget(const uint32 /*TracePath*/)
 {
 	return SNew(SHyperlink)
 		#if TRACEPATH_DEBUG
diff --git a/Engine/Source/Editor/Kismet/Private/Profiler/SProfilerStatExpander.cpp b/Engine/Source/Editor/Kismet/Private/Profiler/SProfilerStatExpander.cpp
new file mode 100644
--- /dev/null
+++ b/Engine/Source/Editor/Kismet/Private/Profiler/SProfilerStatExpander.cpp
@@ -0,0 +1,50 @@
+// Copyright 1998-2016 Epic Games, Inc. All Rights Reserved.
+
+#include "BlueprintEditorPrivatePCH.h"
+#include "SProfilerStatExpander.h"
+
+void SProfilerStatExpander::Construct( const FArguments& InArgs, const TSharedPtr<class ITableRow>& TableRow  )
+{
+	OwnerRowPtr = TableRow;
+	StyleSet = InArgs._StyleSet;
+	IndentAmount = InArgs._IndentAmount;
+	BaseIndentLevel = InArgs._BaseIndentLevel;
+
+	this->ChildSlot
+	.Padding( TAttribute<FMargin>( this, &SProfilerStatExpander::GetExpanderPadding ) )
+	[
+		SAssignNew(ExpanderArrow, SButton)
+		.ButtonStyle( FEditorStyle::Get(), "BlueprintProfiler.TreeArrowButton" )
+		.VAlign(VAlign_Center)
+		.HAlign(HAlign_Center)
+		.Visibility( this, &SProfilerStatExpander::GetExpanderVisibility )
+		.ClickMethod( EButtonClickMethod::MouseDown )
+		.OnClicked( this, &SProfilerStatExpander::OnArrowClicked )
+		.ContentPadding(FMargin(2.f,0.f,2.f,1.f))
+		.ForegroundColor( FSlateColor::UseForeground() )
+		.IsFocusable( false )
+		[
+			SNew(STextBlock)
+			.Text( this, &SProfilerStatExpander::GetStatExpanderTextGlyph )
+			.TextStyle(FEditorStyle::Get(), "BlueprintProfiler.TreeArrowGlyph")
+			.ColorAndOpacity( this, &SProfilerStatExpander::GetStatExpanderColorAndOpacity )
+		]
+	];
+}
+
+FText SProfilerStatExpander::GetStatExpanderTextGlyph() const
+{
+	static const FText NormalGlyph = FText::FromString(FString(TEXT("\xf138"))); /*fa-chevron-circle-right*/
+	static const FText ExpandedGlyph = FText::FromString(FString(TEXT("\xf13a"))); /*fa-chevron-circle-down*/
+
+	return OwnerRowPtr.Pin()->IsItemExpanded() ? ExpandedGlyph : NormalGlyph;
+}
+
+FSlateColor SProfilerStatExpander::GetStatExpanderColorAndOpacity() const
+{
+	static const FLinearColor NormalColor(1.f, 1.f, 1.f, 0.6f);
+	static const FLinearColor HoveredColor(1.f, 1.f, 1.f, 1.f);
+
+	return ExpanderArrow->IsHovered() ? HoveredColor : NormalColor;
+}
+
diff --git a/Engine/Source/Editor/Kismet/Private/Profiler/SProfilerStatExpander.h b/Engine/Source/Editor/Kismet/Private/Profiler/SProfilerStatExpander.h
new file mode 100644
--- /dev/null
+++ b/Engine/Source/Editor/Kismet/Private/Profiler/SProfilerStatExpander.h
@@ -0,0 +1,34 @@
+// Copyright 1998-2016 Epic Games, Inc. All Rights Reserved.
+
+#pragma once
+#include "SExpanderArrow.h"
+
+
+/**
+ *
+ *
+ */
+class SProfilerStatExpander : public SExpanderArrow
+{
+public:
+
+	SLATE_BEGIN_ARGS( SProfilerStatExpander )
+		: _StyleSet(&FCoreStyle::Get())
+		, _IndentAmount(10)
+		, _BaseIndentLevel(0)
+	{ }
+		SLATE_ARGUMENT(const ISlateStyle*, StyleSet)
+		SLATE_ATTRIBUTE(float, IndentAmount)
+		SLATE_ATTRIBUTE(int32, BaseIndentLevel)
+	SLATE_END_ARGS()
+
+	void Construct( const FArguments& InArgs, const TSharedPtr<class ITableRow>& TableRow );
+
+protected:
+
+	/** Returns the statistic expander glyph */
+	FText GetStatExpanderTextGlyph() const;
+
+	/** Returns the statistic expander glyph color */
+	FSlateColor GetStatExpanderColorAndOpacity() const;
+};
diff --git a/Engine/Source/Editor/Kismet/Private/SBlueprintPalette.cpp b/Engine/Source/Editor/Kismet/Private/SBlueprintPalette.cpp
--- a/Engine/Source/Editor/Kismet/Private/SBlueprintPalette.cpp
+++ b/Engine/Source/Editor/Kismet/Private/SBlueprintPalette.cpp
@@ -701,7 +701,7 @@ public:
 			.Schema(Schema)
 			.TargetPinType(this, &SPinTypeSelectorHelper::OnGetVarType)
 			.OnPinTypeChanged(this, &SPinTypeSelectorHelper::OnVarTypeChanged)
-			.bAllowExec(false)
+			.TypeTreeFilter(ETypeTreeFilter::None)
 			.bCompactSelector(true)
 		];
 	}
diff --git a/Engine/Source/Editor/Kismet/Private/SMyBlueprint.cpp b/Engine/Source/Editor/Kismet/Private/SMyBlueprint.cpp
--- a/Engine/Source/Editor/Kismet/Private/SMyBlueprint.cpp
+++ b/Engine/Source/Editor/Kismet/Private/SMyBlueprint.cpp
@@ -1604,9 +1604,8 @@ void SMyBlueprint::OnActionSelected( const TArray< TSharedPtr<FEdGraphSchemaActi
 
 		CurrentBlueprint = BlueprintEditor->GetBlueprintObj();
 		CurrentInspector = BlueprintEditor->GetInspector();
-
-		OnActionSelectedHelper(InAction, BlueprintEditorPtr, Blueprint, CurrentInspector.ToSharedRef());
 	}
+	OnActionSelectedHelper(InAction, BlueprintEditorPtr, Blueprint, CurrentInspector.ToSharedRef());
 }
 
 void SMyBlueprint::OnActionSelectedHelper(TSharedPtr<FEdGraphSchemaAction> InAction, TWeakPtr< FBlueprintEditor > InBlueprintEditor, UBlueprint* Blueprint, TSharedRef<SKismetInspector> Inspector)
@@ -1640,7 +1639,10 @@ void SMyBlueprint::OnActionSelectedHelper(TSharedPtr<FEdGraphSchemaAction> InAct
 			Options.bForceRefresh = true;
 
 			Inspector->ShowDetailsForSingleObject(VarAction->GetProperty(), Options);
-			InBlueprintEditor.Pin()->GetReplaceReferencesWidget()->SetSourceVariable(VarAction->GetProperty());
+			if (InBlueprintEditor.IsValid())
+			{
+				InBlueprintEditor.Pin()->GetReplaceReferencesWidget()->SetSourceVariable(VarAction->GetProperty());
+			}
 		}
 		else if (InAction->GetTypeId() == FEdGraphSchemaAction_K2LocalVar::StaticGetTypeId())
 		{
diff --git a/Engine/Source/Editor/Kismet/Private/SSCSEditor.cpp b/Engine/Source/Editor/Kismet/Private/SSCSEditor.cpp
--- a/Engine/Source/Editor/Kismet/Private/SSCSEditor.cpp
+++ b/Engine/Source/Editor/Kismet/Private/SSCSEditor.cpp
@@ -5643,6 +5643,83 @@ void SSCSEditor::OnOpenBlueprintEditor(bool bForceCodeEditing) const
 	}
 }
 
+/** 
+This struct saves and deselects all selected instanced components (from given actor), then finds them (in recreated actor instance, after compilation) and selects them again.
+*/
+struct FRestoreSelectedInstanceComponent
+{
+	TWeakObjectPtr<UClass> ActorClass;
+	FName ActorName;
+	TWeakObjectPtr<UObject> ActorOuter;
+
+	struct FComponentKey
+	{
+		FName Name;
+		TWeakObjectPtr<UClass> Class;
+
+		FComponentKey(FName InName, UClass* InClass) : Name(InName), Class(InClass) {}
+	};
+	TArray<FComponentKey> ComponentKeys;
+
+	FRestoreSelectedInstanceComponent()
+		: ActorClass(nullptr)
+		, ActorOuter(nullptr)
+	{ }
+
+	void Save(AActor* InActor)
+	{
+		check(InActor);
+		ActorClass = InActor->GetClass();
+		ActorName = InActor->GetFName();
+		ActorOuter = InActor->GetOuter();
+
+		check(GEditor);
+		TArray<UActorComponent*> ComponentsToSaveAndDelesect;
+		for (auto Iter = GEditor->GetSelectedComponentIterator(); Iter; ++Iter)
+		{
+			UActorComponent* Component = CastChecked<UActorComponent>(*Iter, ECastCheckedType::NullAllowed);
+			if (Component && InActor->GetInstanceComponents().Contains(Component))
+			{
+				ComponentsToSaveAndDelesect.Add(Component);
+			}
+		}
+
+		for (UActorComponent* Component : ComponentsToSaveAndDelesect)
+		{
+			USelection* SelectedComponents = GEditor->GetSelectedComponents();
+			if (ensure(SelectedComponents))
+			{
+				ComponentKeys.Add(FComponentKey(Component->GetFName(), Component->GetClass()));
+				SelectedComponents->Deselect(Component);
+			}
+		}
+	}
+
+	void Restore()
+	{
+		AActor* Actor = (ActorClass.IsValid() && ActorOuter.IsValid()) 
+			? Cast<AActor>((UObject*)FindObjectWithOuter(ActorOuter.Get(), ActorClass.Get(), ActorName)) 
+			: nullptr;
+		if (Actor)
+		{
+			for (auto IterKey : ComponentKeys)
+			{
+				TArray<UActorComponent*> OwnedComponents;
+				Actor->GetComponents(OwnedComponents);
+				UActorComponent** ComponentPtr = OwnedComponents.FindByPredicate([&](UActorComponent* InComp)
+				{
+					return InComp && (InComp->GetFName() == IterKey.Name) && (InComp->GetClass() == IterKey.Class.Get());
+				});
+				if (ComponentPtr && *ComponentPtr)
+				{
+					check(GEditor);
+					GEditor->SelectComponent(*ComponentPtr, true, false);
+				}
+			}
+		}
+	}
+};
+
 void SSCSEditor::OnApplyChangesToBlueprint() const
 {
 	int32 NumChangedProperties = 0;
@@ -5654,6 +5731,7 @@ void SSCSEditor::OnApplyChangesToBlueprint() const
 	{
 		// Cache the actor label as by the time we need it, it may be invalid
 		const FString ActorLabel = Actor->GetActorLabel();
+		FRestoreSelectedInstanceComponent RestoreSelectedInstanceComponent;
 		{
 			const FScopedTransaction Transaction(LOCTEXT("PushToBlueprintDefaults_Transaction", "Apply Changes to Blueprint"));
 
@@ -5680,6 +5758,7 @@ void SSCSEditor::OnApplyChangesToBlueprint() const
 					NumChangedProperties = EditorUtilities::CopyActorProperties(Actor, BlueprintCDO, CopyOptions);
 					if (Actor->GetInstanceComponents().Num() > 0)
 					{
+						RestoreSelectedInstanceComponent.Save(Actor);
 						FKismetEditorUtilities::AddComponentsToBlueprint(Blueprint, Actor->GetInstanceComponents());
 						NumChangedProperties += Actor->GetInstanceComponents().Num();
 						Actor->ClearInstanceComponents(true);
@@ -5697,6 +5776,7 @@ void SSCSEditor::OnApplyChangesToBlueprint() const
  		{
 			FBlueprintEditorUtils::MarkBlueprintAsModified(Blueprint);
 			FKismetEditorUtilities::CompileBlueprint(Blueprint);
+			RestoreSelectedInstanceComponent.Restore();
  		}
 
 		// Set up a notification record to indicate success/failure
diff --git a/Engine/Source/Editor/Kismet/Private/UserDefinedStructureEditor.cpp b/Engine/Source/Editor/Kismet/Private/UserDefinedStructureEditor.cpp
--- a/Engine/Source/Editor/Kismet/Private/UserDefinedStructureEditor.cpp
+++ b/Engine/Source/Editor/Kismet/Private/UserDefinedStructureEditor.cpp
@@ -782,13 +782,13 @@ public:
 		}
 	}
 
-	void GetFilteredVariableTypeTree( TArray< TSharedPtr<UEdGraphSchema_K2::FPinTypeTreeInfo> >& TypeTree, bool bAllowExec, bool bAllowWildcard ) const
+	void GetFilteredVariableTypeTree( TArray< TSharedPtr<UEdGraphSchema_K2::FPinTypeTreeInfo> >& TypeTree, ETypeTreeFilter TypeTreeFilter) const
 	{
 		auto K2Schema = GetDefault<UEdGraphSchema_K2>();
 		auto StructureDetailsSP = StructureDetails.Pin();
 		if(StructureDetailsSP.IsValid() && K2Schema)
 		{
-			K2Schema->GetVariableTypeTree(TypeTree, bAllowExec, bAllowWildcard);
+			K2Schema->GetVariableTypeTree(TypeTree, TypeTreeFilter);
 			const auto Parent = StructureDetailsSP->GetUserDefinedStruct();
 			// THE TREE HAS ONLY 2 LEVELS
 			for (auto PinTypePtr : TypeTree)
@@ -864,8 +864,7 @@ public:
 				.OnPinTypePreChanged(this, &FUserDefinedStructureFieldLayout::OnPrePinInfoChange)
 				.OnPinTypeChanged(this, &FUserDefinedStructureFieldLayout::PinInfoChanged)
 				.Schema(K2Schema)
-				.bAllowExec(false)
-				.bAllowWildcard(false)
+				.TypeTreeFilter(ETypeTreeFilter::None)
 				.Font( IDetailLayoutBuilder::GetDetailFont() )
 			]
 			+ SHorizontalBox::Slot()
diff --git a/Engine/Source/Editor/Kismet/Public/Profiler/EventExecution.h b/Engine/Source/Editor/Kismet/Public/Profiler/EventExecution.h
--- a/Engine/Source/Editor/Kismet/Public/Profiler/EventExecution.h
+++ b/Engine/Source/Editor/Kismet/Public/Profiler/EventExecution.h
@@ -23,22 +23,25 @@ namespace EScriptExecutionNodeFlags
 		Node						= 0x00000800,	// Node timing
 		ExecPin						= 0x00001000,	// Exec pin dummy node
 		PureNode					= 0x00002000,	// Pure node (no exec pins)
-		AsyncTask					= 0x00004000,	// Async task node
-		AsyncTaskDelegate			= 0x00008000,	// Async task delegate pin/event
-		FunctionTunnel				= 0x00010000,	// Tunnel function
-		TunnelEntryPin				= 0x00020000,	// Tunnel entry pin
-		TunnelExitPin				= 0x00040000,	// Tunnel exit pin
-		ReEntrantTunnelPin			= 0x00080000,	// Re-Entrant tunnel pin
-		TunnelFinalExitPin			= 0x00100000,	// Tunnel final exit or exec then pin
+		EventPin					= 0x00004000,	// Delegate pin entry point
+		TunnelInstance				= 0x00008000,	// Tunnel instance
+		TunnelEntryPin				= 0x00010000,	// Internal tunnel entry pin
+		TunnelExitPin				= 0x00020000,	// Internal tunnel exit pin
+		TunnelEntryPinInstance		= 0x00040000,	// Tunnel entry pin instance
+		TunnelExitPinInstance		= 0x00080000,	// Tunnel exit pin instance
 		PureChain					= 0x00200000,	// Pure node call chain
 		CyclicLinkage				= 0x00400000,	// Marks execution path as cyclic.
 		InvalidTrace				= 0x00800000,	// Indicates that node doesn't contain a valid script trace.
-									// Groups
-		CallSite					= FunctionCall|MacroCall|FunctionTunnel,
+		// Groups
+		CallSite					= FunctionCall|MacroCall|TunnelInstance,
 		BranchNode					= ConditionalBranch|SequentialBranch,
+		TunnelInstancePin			= TunnelEntryPinInstance|TunnelExitPinInstance,
 		TunnelPin					= TunnelEntryPin|TunnelExitPin,
+		TunnelBoundary				= TunnelInstancePin|TunnelPin,
+		TunnelEntry					= TunnelEntryPin|TunnelEntryPinInstance,
+		TunnelExit					= TunnelExitPin|TunnelExitPinInstance,
 		PureStats					= PureNode|PureChain,
-		Container					= Class|Instance|Event|CustomEvent|TunnelPin|ExecPin|PureChain
+		Container					= Class|Instance|Event|CustomEvent|ExecPin|PureChain
 	};
 }
 
@@ -92,9 +95,6 @@ public:
 	/** Returns the child nodes map */
 	TArray<TSharedPtr<class FScriptExecutionNode>>& GetChildNodes() { return ChildNodes; }
 
-	/** Returns the filtered child nodes */
-	virtual void GetFilteredChildNodes(TArray<FLinearExecPath>& ChildArrayInOut, const FTracePath& TracePath);
-
 	/** Returns the number of children */
 	int32 GetNumChildren() const { return ChildNodes.Num(); }
 
@@ -230,8 +230,8 @@ public:
 	/** Does the node contain these flags */
 	bool HasFlags(const uint32 Flags) const { return (NodeFlags & Flags) != 0U; }
 
-	/** Returns if this exec node can be safely cyclicly linked */
-	bool IsCyclicNode() const { return (NodeFlags & EScriptExecutionNodeFlags::CallSite) == 0U; }
+	/** Returns the current node flags */
+	uint32 GetFlags() const { return NodeFlags; }
 
 	/** Returns if this exec event represents a change in class/blueprint */
 	bool IsClass() const { return (NodeFlags & EScriptExecutionNodeFlags::Class) != 0U; }
@@ -260,9 +260,27 @@ public:
 	/** Returns if this node is a pure chain node */
 	bool IsPureChain() const { return (NodeFlags & EScriptExecutionNodeFlags::PureChain) != 0U; }
 
-	/** Returns if this event potentially multiple exit sites */
+	/** Returns if this event potentially has multiple exit sites */
 	bool IsBranch() const { return (NodeFlags & EScriptExecutionNodeFlags::BranchNode) != 0U; }
 
+	/** Returns if this event represents a tunnel Boundary */
+	bool IsTunnelBoundary() const { return (NodeFlags & EScriptExecutionNodeFlags::TunnelBoundary) != 0U; }
+	
+	/** Returns if this event represents a tunnel instance */
+	bool IsTunnelInstance() const { return (NodeFlags & EScriptExecutionNodeFlags::TunnelInstance) != 0U; }
+
+	/** Returns if this event represents a tunnel entry */
+	bool IsTunnelEntry() const { return (NodeFlags & EScriptExecutionNodeFlags::TunnelEntryPinInstance) != 0U; }
+
+	/** Returns if this event represents a tunnel exit */
+	bool IsTunnelExit() const { return (NodeFlags & EScriptExecutionNodeFlags::TunnelExit) != 0U; }
+
+	/** Returns if this exec node can be safely cyclicly linked */
+	bool CanLinkAsCyclicNode() const { return (NodeFlags & EScriptExecutionNodeFlags::CallSite) != 0U; }
+
+	/** Returns true if this node is cyclicly linked */
+	bool HasCyclicLinkage() const { return (NodeFlags & EScriptExecutionNodeFlags::CyclicLinkage) != 0U; }
+
 	/** Gets the observed object context */
 	const UObject* GetObservedObject() { return ObservedObject.Get(); }
 
@@ -324,10 +342,10 @@ public:
 	virtual void GetAllPureNodes(TMap<int32, TSharedPtr<FScriptExecutionNode>>& PureNodesOut);
 
 	/** Creates and returns the slate icon widget */
-	virtual TSharedRef<SWidget> GetIconWidget();
+	virtual TSharedRef<SWidget> GetIconWidget(const uint32 TracePath = 0U);
 
 	/** Creates and returns the slate hyperlink widget */
-	virtual TSharedRef<SWidget> GetHyperlinkWidget();
+	virtual TSharedRef<SWidget> GetHyperlinkWidget(const uint32 TracePath = 0U);
 
 protected:
 
@@ -386,13 +404,24 @@ public:
 	virtual void GetBlueprintPerfDataForAllTracePaths(FScriptPerfData& OutPerfData) override;
 	// ~FScriptNodePerfData
 
-	/** Adds a custom child entry point */
-	void AddCustomEntryPoint(TSharedPtr<class FScriptExecutionTunnelEntry> CustomEntryPoint) { CustomEntryPoints.Add(CustomEntryPoint); }
+	/** Adds an entry site at the specified script offset */
+	void AddEntrySite(const int32 ScriptOffset, TSharedPtr<class FScriptExecutionTunnelEntry> EntrySite);
+
+	/** Adds an entry site at the specified script offset */
+	void AddExitSite(const int32 ScriptOffset, TSharedPtr<class FScriptExecutionTunnelExit> ExitSite) { ExitSites.FindOrAdd(ScriptOffset) = ExitSite; }
+
+	/** Copies tunnel Boundary sites  */
+	void CopyBoundarySites(TSharedPtr<FScriptExecutionTunnelInstance> OtherInstance);
+
+	/** Find a tunnel Boundary node by script offset */
+	TSharedPtr<FScriptExecutionNode> FindBoundarySite(const int32 ScriptOffset);
 
 private:
 
-	/** List of custom entry points */
-	TArray<TSharedPtr<class FScriptExecutionTunnelEntry>> CustomEntryPoints;
+	/** List of entry sites by script offset */
+	TMap<int32, TSharedPtr<FScriptExecutionTunnelEntry>> EntrySites;
+	/** List of exit sites by script offset */
+	TMap<int32, TSharedPtr<FScriptExecutionTunnelExit>> ExitSites;
 
 };
 
@@ -403,70 +432,69 @@ class KISMET_API FScriptExecutionTunnelEntry : public FScriptExecutionNode
 {
 public:
 
-	/** Tunnel execution type */
-	enum ETunnelType
-	{
-		SimpleTunnel = 0,
-		SinglePathTunnel,
-		MultiIOTunnel
-	};
-
-	FScriptExecutionTunnelEntry(const FScriptExecNodeParams& InitParams, const class UK2Node_Tunnel* TunnelInstanceIn, const ETunnelType TypeIn)
+	FScriptExecutionTunnelEntry(const FScriptExecNodeParams& InitParams)
 		: FScriptExecutionNode(InitParams)
-		, TunnelType(TypeIn)
-		, TunnelInstance(TunnelInstanceIn)
 	{
 	}
 
 	// FScriptExecutionNode
 	virtual void GetLinearExecutionPath(TArray<FLinearExecPath>& LinearExecutionNodes, const FTracePath& TracePath, const bool bIncludeChildren = false) override;
-	virtual void GetFilteredChildNodes(TArray<FLinearExecPath>& ChildArrayInOut, const FTracePath& TracePath) override;
 	virtual void RefreshStats(const FTracePath& TracePath) override;
 	// ~FScriptExecutionNode
 
 	/** Add tunnel timing */
-	void AddTunnelTiming(const FName InstanceName, const FTracePath& TracePath, const double Time);
+	void AddTunnelTiming(const FName InstanceName, const FTracePath& TracePath, const FTracePath& InternalTracePath, const int32 ExitScriptOffset, const double Time);
 
-	/** Update tunnel exit site */
-	void UpdateTunnelExit(const FName InstanceName, const FTracePath& TracePath, const int32 ExitScriptOffset);
+	/** Returns the tunnel instance type */
+	bool IsComplexTunnel() const { return LinkedNodes.Num() > 1; }
 
-	/** Returns true if the script offset is valid for this tunnel */
-	bool IsPathValidForTunnel(const int32 ScriptOffset) const;
+	/** Returns the tunnel instance exec node */
+	TSharedPtr<FScriptExecutionTunnelInstance> GetTunnelInstance() const { return TunnelInstance; }
 
-	/** Returns the tunnel instance graph node */
-	const ETunnelType GetTunnelType() const { return TunnelType; }
-
-	/** Returns the tunnel instance graph node */
-	const UK2Node_Tunnel* GetTunnelInstance() const { return TunnelInstance.Get(); }
+	/** Sets the tunnel instance exec node */
+	void SetTunnelInstance(TSharedPtr<FScriptExecutionTunnelInstance> TunnelInstanceIn) { TunnelInstance = TunnelInstanceIn; }
 
 	/** Adds a valid exit script offset */
-	void AddExitSite(TSharedPtr<FScriptExecutionNode> ExitSite, const int32 ExitScriptOffset);
+	void AddExitSite(const int32 ExitScriptOffset, TSharedPtr<class FScriptExecutionTunnelExit> ExitSite);
 
-	/** Returns the exit point associated with the script offset */
-	TSharedPtr<FScriptExecutionNode> GetExitSite(const int32 ScriptOffset) const;
+	/** Returns the exit point associated with the exit script offset */
+	TSharedPtr<FScriptExecutionTunnelExit> GetExitSite(const int32 ScriptOffset) const;
 
-	/** Returns if the exit point associated with the script offset is the final exit site */
-	bool IsFinalExitSite(const int32 ScriptOffset) const;
+	/** Returns the exit point associated with the entry script offset */
+	TMap<int32, TSharedPtr<FScriptExecutionNode>> GetExitSites() { return LinkedNodes; }
 
-	/** Builds exit sites for instance */
-	void BuildExitBranches();
+	/** Returns true if the supplied node is inside this tunnel */
+	bool IsInternalBoundary(TSharedPtr<FScriptExecutionNode> PotentialBoundaryNode) const;
 
 private:
 
-	/** The tunnel type */
-	ETunnelType TunnelType;
 	/** The tunnel instance this node represents */
-	TWeakObjectPtr<const UK2Node_Tunnel> TunnelInstance;
-	/** List of valid exit scripts offsets */
-	TMap<int32, TSharedPtr<FScriptExecutionNode>> ValidExitPoints;
-	/** Branched exit sites */
-	TArray<FLinearExecPath> BranchedExitSites;
+	TSharedPtr<FScriptExecutionTunnelInstance> TunnelInstance;
 
 };
 
 //////////////////////////////////////////////////////////////////////////
 // FScriptExecutionTunnelExit
 
+struct FTunnelExitInstanceData
+{
+	FTunnelExitInstanceData()
+		: InstanceIcon(nullptr)
+	{
+	}
+
+	FTunnelExitInstanceData(const FSlateBrush* PinIconIn, UEdGraphPin* InstancePinIn)
+		: InstanceIcon(PinIconIn)
+		, InstancePin(InstancePinIn)
+	{
+	}
+
+	/** Instance specific icon */
+	const FSlateBrush* InstanceIcon;
+	/** Instance specific pin */
+	FEdGraphPinReference InstancePin;
+};
+
 class KISMET_API FScriptExecutionTunnelExit : public FScriptExecutionNode
 {
 public:
@@ -479,9 +507,29 @@ public:
 	// FScriptExecutionNode
 	virtual void GetLinearExecutionPath(TArray<FLinearExecPath>& LinearExecutionNodes, const FTracePath& TracePath, const bool bIncludeChildren = false) override;
 	virtual void RefreshStats(const FTracePath& TracePath) override;
-	virtual void NavigateToObject() const;
+	virtual TSharedRef<SWidget> GetIconWidget(const uint32 TracePath = 0U) override;
+	virtual TSharedRef<SWidget> GetHyperlinkWidget(const uint32 TracePath = 0U) override;
 	// ~FScriptExecutionNode
 
+	/** Navigate to exit site using tracepath */
+	virtual void NavigateToExitSite(const uint32 TracePath) const;
+
+	/** Add Instance exit pin */
+	void AddInstanceExitSite(const uint32 TracePath, UEdGraphPin* InstanceExitPin);
+
+	/** Returns the external tunnel pin */
+	UEdGraphPin* GetExternalPin() const { return ExternalPin.Get(); }
+
+	/** Sets the external tunnel pin */
+	void SetExternalPin(UEdGraphPin* ExternalPinIn) { ExternalPin = ExternalPinIn; }
+
+private:
+
+	/** The internal tunnel pin this node represents */
+	FEdGraphPinReference ExternalPin;
+	/** Instance specific exit sites */
+	TMap<const uint32, FTunnelExitInstanceData> InstanceExitSiteInfo;
+
 };
 
 //////////////////////////////////////////////////////////////////////////
@@ -512,18 +560,22 @@ public:
 
 	FScriptExecutionInstance(const FScriptExecNodeParams& InitParams)
 		: FScriptExecutionNode(InitParams)
+		, bPIEInstanceDestroyed(false)
 	{
 	}
 
 	// FScriptExecutionNode
 	virtual void NavigateToObject() const override;
-	virtual TSharedRef<SWidget> GetIconWidget() override;
-	virtual TSharedRef<SWidget> GetHyperlinkWidget() override;
+	virtual TSharedRef<SWidget> GetIconWidget(const uint32 TracePath = 0U) override;
+	virtual TSharedRef<SWidget> GetHyperlinkWidget(const uint32 TracePath = 0U) override;
 	virtual bool IsObservedObjectValid() const override;
 	// ~FScriptExecutionNode
 
-	/** Sets the current PIE Instance */
-	void SetPIEInstance(const UObject* PIEInstanceIn) { PIEInstance = PIEInstanceIn; }
+	/** Refresh PIE Instance */
+	void RefreshPIEInstance() const;
+
+	/** Returns the active object instance - PIE Instance if in PIE, otherwise editor instance */
+	UObject* GetActiveObject() const;
 
 private:
 
@@ -536,7 +588,9 @@ private:
 private:
 
 	/** PIE instance */
-	TWeakObjectPtr<const UObject> PIEInstance;
+	mutable TWeakObjectPtr<const UObject> PIEInstance;
+	/** Indicates the PIE instance has been detected as destroyed, so we don't keep refreshing */
+	mutable bool bPIEInstanceDestroyed;
 
 };
 
diff --git a/Engine/Source/Editor/Kismet/Public/Profiler/TracePath.h b/Engine/Source/Editor/Kismet/Public/Profiler/TracePath.h
--- a/Engine/Source/Editor/Kismet/Public/Profiler/TracePath.h
+++ b/Engine/Source/Editor/Kismet/Public/Profiler/TracePath.h
@@ -18,24 +18,29 @@ public:
 
 	FTracePath(const FTracePath& TracePathIn)
 		: CachedCrc(TracePathIn.CachedCrc)
+		, TunnelEntryCrc(TracePathIn.TunnelEntryCrc)
 		, Tunnel(TracePathIn.Tunnel)
 #if TRACEPATH_DEBUG
 		, TraceStack(TracePathIn.TraceStack)
+		, TunnelTraceStack(TracePathIn.TunnelTraceStack)
 #endif
 	{
 	}
 
-	FTracePath(TSharedPtr<const class FScriptExecutionTunnelEntry> TunnelNode)
+	FTracePath(TSharedPtr<class FScriptExecutionTunnelEntry> TunnelNode)
 		: CachedCrc(0U)
+		, TunnelEntryCrc(0U)
 		, Tunnel(TunnelNode)
 	{
 	}
 
-	FTracePath(const FTracePath& TracePathIn, TSharedPtr<const class FScriptExecutionTunnelEntry> TunnelNode)
+	FTracePath(const FTracePath& TracePathIn, TSharedPtr<class FScriptExecutionTunnelEntry> TunnelNode)
 		: CachedCrc(TracePathIn.CachedCrc)
+		, TunnelEntryCrc(TracePathIn.CachedCrc)
 		, Tunnel(TunnelNode)
 #if TRACEPATH_DEBUG
 		, TraceStack(TracePathIn.TraceStack)
+		, TunnelTraceStack(TracePathIn.TraceStack)
 #endif
 	{
 	}
@@ -45,7 +50,9 @@ public:
 	FTracePath& operator = (const FTracePath& TracePathIn)
 	{
 		TraceStack = TracePathIn.TraceStack;
+		TunnelTraceStack = TracePathIn.TunnelTraceStack;
 		CachedCrc = TracePathIn.CachedCrc;
+		TunnelEntryCrc = TracePathIn.TunnelEntryCrc;
 		Tunnel = TracePathIn.Tunnel;
 		return *this;
 	}
@@ -68,6 +75,7 @@ public:
 	{
 #if TRACEPATH_DEBUG
 		TraceStack.SetNum(0);
+		TunnelTraceStack.SetNum(0);
 #endif
 		CachedCrc = 0;
 		Tunnel.Reset();
@@ -83,7 +91,32 @@ public:
 	}
 
 	/** Returns the current tunnel */
-	TSharedPtr<const class FScriptExecutionTunnelEntry> GetTunnel() const { return Tunnel; }
+	TSharedPtr<class FScriptExecutionTunnelEntry> GetTunnel() const { return Tunnel; }
+
+	/** Returns tunnel entry tracepath */
+	void GetTunnelTracePath(FTracePath& TracePathInOut) const
+	{
+		TracePathInOut.CachedCrc = TunnelEntryCrc;
+		TracePathInOut.TunnelEntryCrc = 0U;
+		TracePathInOut.Tunnel.Reset();
+#if TRACEPATH_DEBUG
+		TracePathInOut.TraceStack = TunnelTraceStack;
+		TracePathInOut.TunnelTraceStack.SetNum(0);
+#endif
+	}
+
+	/** Exits the current tunnel, restoring the entry tracepath */
+	void ExitTunnel()
+	{
+		Tunnel.Reset();
+		CachedCrc = TunnelEntryCrc;
+		TunnelEntryCrc = 0U;
+#if TRACEPATH_DEBUG
+		TraceStack = TunnelTraceStack;
+		TunnelTraceStack.SetNum(0);
+#endif
+	}
+
 
 #if TRACEPATH_DEBUG
 	const FString GetPathString() const
@@ -102,11 +135,15 @@ private:
 
 	/** Cached crc */
 	uint32 CachedCrc;
+	/** Tunnel Entry Crc */
+	uint32 TunnelEntryCrc;
 	/** The current tunnel entry point */
-	TSharedPtr<const class FScriptExecutionTunnelEntry> Tunnel;
+	TSharedPtr<class FScriptExecutionTunnelEntry> Tunnel;
 
 #if TRACEPATH_DEBUG
 	/** Contributing pin list - for debugging only */
 	TArray<int32> TraceStack;
+	/** Contributing pin list for tunnel entry - for debugging only */
+	TArray<int32> TunnelTraceStack;
 #endif
 };
diff --git a/Engine/Source/Editor/KismetCompiler/Private/AnimBlueprintCompiler.cpp b/Engine/Source/Editor/KismetCompiler/Private/AnimBlueprintCompiler.cpp
--- a/Engine/Source/Editor/KismetCompiler/Private/AnimBlueprintCompiler.cpp
+++ b/Engine/Source/Editor/KismetCompiler/Private/AnimBlueprintCompiler.cpp
@@ -151,7 +151,7 @@ void FAnimBlueprintCompiler::CreateEvaluationHandlerStruct(UAnimGraphNode_Base*
 	HandlerFunctionNames.Add(Record.HandlerFunctionName);
 	
 	// Add a custom event in the graph
-	UK2Node_CustomEvent* EntryNode = SpawnIntermediateNode<UK2Node_CustomEvent>(VisualAnimNode, ConsolidatedEventGraph);
+	UK2Node_CustomEvent* EntryNode = SpawnIntermediateEventNode<UK2Node_CustomEvent>(VisualAnimNode, nullptr, ConsolidatedEventGraph);
 	EntryNode->bInternalEvent = true;
 	EntryNode->CustomFunctionName = Record.HandlerFunctionName;
 	EntryNode->AllocateDefaultPins();
diff --git a/Engine/Source/Editor/KismetCompiler/Private/KismetCompiler.cpp b/Engine/Source/Editor/KismetCompiler/Private/KismetCompiler.cpp
--- a/Engine/Source/Editor/KismetCompiler/Private/KismetCompiler.cpp
+++ b/Engine/Source/Editor/KismetCompiler/Private/KismetCompiler.cpp
@@ -744,8 +744,16 @@ void FKismetCompilerContext::CreatePropertiesFromList(UStruct* Scope, UField**&
 			// Record in the debugging information
 			//@TODO: Rename RegisterClassPropertyAssociation, etc..., to better match that indicate it works with locals
 			{
-				UObject* TrueSourceObject = MessageLog.FindSourceObject(Term.Source);
-				NewClass->GetDebugData().RegisterClassPropertyAssociation(TrueSourceObject, NewProperty);
+				if (Term.SourcePin)
+				{
+					UEdGraphPin* TrueSourcePin = MessageLog.FindSourcePin(Term.SourcePin);
+					NewClass->GetDebugData().RegisterClassPropertyAssociation(TrueSourcePin, NewProperty);
+				}
+				else
+				{
+					UObject* TrueSourceObject = MessageLog.FindSourceObject(Term.Source);
+					NewClass->GetDebugData().RegisterClassPropertyAssociation(TrueSourceObject, NewProperty);
+				}
 			}
 
 			// Record the desired default value for this, if specified by the term
@@ -818,8 +826,16 @@ void FKismetCompilerContext::CreateLocalVariablesForFunction(FKismetFunctionCont
 
 				if (Term.AssociatedVarProperty != NULL)
 				{
-					UObject* TrueSourceObject = MessageLog.FindSourceObject(Term.Source);
-					NewClass->GetDebugData().RegisterClassPropertyAssociation(TrueSourceObject, Term.AssociatedVarProperty);
+					if (Term.SourcePin)
+					{
+						UEdGraphPin* TrueSourcePin = MessageLog.FindSourcePin(Term.SourcePin);
+						NewClass->GetDebugData().RegisterClassPropertyAssociation(TrueSourcePin, Term.AssociatedVarProperty);
+					}
+					else
+					{
+						UObject* TrueSourceObject = MessageLog.FindSourceObject(Term.Source);
+						NewClass->GetDebugData().RegisterClassPropertyAssociation(TrueSourceObject, Term.AssociatedVarProperty);
+					}
 				}
 			}
 		}
@@ -922,7 +938,15 @@ void FKismetCompilerContext::CopyTermDefaultsToDefaultObject(UObject* DefaultObj
 				{
 					const FString ErrorMessage = *FString::Printf(*LOCTEXT("ParseDefaultValueError", "Can't parse default value '%s' for @@. Property: %s.").ToString(), *Value, *Property->GetName());
 					UObject* InstigatorObject = NewClass->GetDebugData().FindObjectThatCreatedProperty(Property);
-					MessageLog.Warning(*ErrorMessage, InstigatorObject);
+					if(InstigatorObject)
+					{
+						MessageLog.Warning(*ErrorMessage, InstigatorObject);
+					}
+					else
+					{
+						UEdGraphPin* InstigatorPin = NewClass->GetDebugData().FindPinThatCreatedProperty(Property);
+						MessageLog.Warning(*ErrorMessage, InstigatorPin);
+					}
 				}
 
 				break;
@@ -1821,6 +1845,13 @@ void FKismetCompilerContext::FinishCompilingClass(UClass* Class)
 
 		// If the flag is inherited, this will keep the bool up-to-date
 		Blueprint->bDeprecate = (Class->ClassFlags & CLASS_Deprecated) == CLASS_Deprecated;
+		
+		// If the Blueprint was marked as abstract, then flag the class as abstract.
+		if (Blueprint->bGenerateAbstractClass)
+		{
+			NewClass->ClassFlags |= CLASS_Abstract;
+		}
+		Blueprint->bGenerateAbstractClass = (Class->ClassFlags & CLASS_Abstract) == CLASS_Abstract;	
 
 		// Add the description to the tooltip
 		if (!Blueprint->BlueprintDescription.IsEmpty())
@@ -1957,7 +1988,12 @@ void FKismetCompilerContext::BuildDynamicBindingObjects(UBlueprintGeneratedClass
  */ 
 void FKismetCompilerContext::CreatePinEventNodeForTimelineFunction(UK2Node_Timeline* TimelineNode, UEdGraph* SourceGraph, FName FunctionName, const FString& PinName, FName ExecFuncName)
 {
-	UK2Node_Event* TimelineEventNode = SpawnIntermediateNode<UK2Node_Event>(TimelineNode, SourceGraph);
+	UEdGraphPin* SourcePin = nullptr;
+	if (UK2Node_Timeline* SourceNode = Cast<UK2Node_Timeline>(MessageLog.FindSourceObject(TimelineNode)))
+	{
+		SourcePin = SourceNode->FindPin(PinName);
+	}
+	UK2Node_Event* TimelineEventNode = SpawnIntermediateEventNode<UK2Node_Event>(TimelineNode, SourcePin, SourceGraph);
 	TimelineEventNode->EventReference.SetExternalMember(FunctionName, UTimelineComponent::StaticClass());
 	TimelineEventNode->CustomFunctionName = FunctionName; // Make sure we name this function the thing we are expecting
 	TimelineEventNode->bInternalEvent = true;
@@ -2397,11 +2433,12 @@ void FKismetCompilerContext::MergeUbergraphPagesIn(UEdGraph* Ubergraph)
 	for (TArray<UEdGraph*>::TIterator It(Blueprint->UbergraphPages); It; ++It)
 	{
 		UEdGraph* SourceGraph = *It;
+		const bool bCreateTunnelBoundries = CompileOptions.IsInstrumentationActive();
 
 		if (CompileOptions.bSaveIntermediateProducts)
 		{
 			TArray<UEdGraphNode*> ClonedNodeList;
-			FEdGraphUtilities::CloneAndMergeGraphIn(Ubergraph, SourceGraph, MessageLog, /*bRequireSchemaMatch=*/ true, /*bIsCompiling*/ true, &ClonedNodeList);
+			FEdGraphUtilities::CloneAndMergeGraphIn(Ubergraph, SourceGraph, MessageLog, /*bRequireSchemaMatch=*/ true, /*bIsCompiling*/ true, bCreateTunnelBoundries/*bCreateTunnelBoundries*/, &ClonedNodeList);
 
 			// Create a comment block around the ubergrapgh contents before anything else got started
 			int32 OffsetX = 0;
@@ -2419,7 +2456,7 @@ void FKismetCompilerContext::MergeUbergraphPagesIn(UEdGraph* Ubergraph)
 		}
 		else
 		{
-			FEdGraphUtilities::CloneAndMergeGraphIn(Ubergraph, SourceGraph, MessageLog, /*bRequireSchemaMatch=*/ true, /*bIsCompiling*/ true);
+			FEdGraphUtilities::CloneAndMergeGraphIn(Ubergraph, SourceGraph, MessageLog, /*bRequireSchemaMatch=*/ true, /*bIsCompiling*/ true, bCreateTunnelBoundries/*bCreateTunnelBoundries*/);
 		}
 	}
 }
@@ -2591,7 +2628,7 @@ void FKismetCompilerContext::CreateAndProcessUbergraph()
 				if (!bFoundEntry)
 				{
 					// Create an entry node stub, so that we have a entry point for interfaces to call to
-					UK2Node_Event* EventNode = SpawnIntermediateNode<UK2Node_Event>(NULL, ConsolidatedEventGraph);
+					UK2Node_Event* EventNode = SpawnIntermediateEventNode<UK2Node_Event>(nullptr, nullptr, ConsolidatedEventGraph);
 					EventNode->EventReference.SetExternalMember(FunctionName, InterfaceDesc.Interface);
 					EventNode->bOverrideFunction = true;
 					EventNode->AllocateDefaultPins();
@@ -2758,6 +2795,12 @@ void FKismetCompilerContext::ExpandTunnelsAndMacros(UEdGraph* SourceGraph)
 			// Clone the macro graph, then move all of its children, keeping a list of nodes from the macro
 			UEdGraph* ClonedGraph = FEdGraphUtilities::CloneGraph(MacroGraph, NULL, &MessageLog, true);
 
+			if (CompileOptions.IsInstrumentationActive())
+			{
+				// Create tunnel Boundary signals
+				UK2Node_TunnelBoundary::CreateBoundaryNodesForTunnelInstance(MacroInstanceNode, ClonedGraph, MessageLog);
+			}
+
 			for (int32 I = 0; I < ClonedGraph->Nodes.Num(); ++I)
 			{
 				MacroGeneratedNodes.Add(ClonedGraph->Nodes[I], CurrentNode);
@@ -3352,6 +3395,16 @@ void FKismetCompilerContext::Compile()
 			}
 		}
 
+		// Create a mapping between compile time generated events and what created them.
+		if (TargetClass->bHasInstrumentation)
+		{
+			FBlueprintDebugData& DebugData = TargetClass->GetDebugData();
+			for (auto EventInfo : SourcePinToExpansionEvent)
+			{
+				DebugData.RegisterPinToEventName(EventInfo.Key, EventInfo.Value->GetFunctionName());
+			}
+		}
+
 		// Save off intermediate build products if requested
 		if (CompileOptions.bSaveIntermediateProducts && !Blueprint->bIsRegeneratingOnLoad)
 		{
diff --git a/Engine/Source/Editor/KismetCompiler/Private/KismetCompilerVMBackend.cpp b/Engine/Source/Editor/KismetCompiler/Private/KismetCompilerVMBackend.cpp
--- a/Engine/Source/Editor/KismetCompiler/Private/KismetCompilerVMBackend.cpp
+++ b/Engine/Source/Editor/KismetCompiler/Private/KismetCompilerVMBackend.cpp
@@ -1548,83 +1548,80 @@ public:
 			Writer << EventType;
 		}
 
-		if (Statement.Type != KCST_InstrumentedStateRestore)
+		TArray<UEdGraphPin*> PinContextArray(Statement.PureOutputContextArray);
+		if (Statement.ExecContext != nullptr)
 		{
-			TArray<UEdGraphPin*> PinContextArray(Statement.PureOutputContextArray);
-			if (Statement.ExecContext != nullptr)
-			{
-				PinContextArray.Add(Statement.ExecContext);
-			}
+			PinContextArray.Add(Statement.ExecContext);
+		}
 
-			for (auto PinContext : PinContextArray)
+		for (auto PinContext : PinContextArray)
+		{
+			UEdGraphPin const* TrueSourcePin = FunctionContext.MessageLog.FindSourcePin(PinContext);
+			// Source pin can be marked as pending kill if it was a generated pin that node logic decided to disown, e.g.
+			// logic in UK2Node_CallFunction to handle bWantsEnumToExecExpansion:
+			if (TrueSourcePin && !TrueSourcePin->IsPendingKill())
 			{
-				UEdGraphPin const* TrueSourcePin = FunctionContext.MessageLog.FindSourcePin(PinContext);
-				// Source pin can be marked as pending kill if it was a generated pin that node logic decided to disown, e.g.
-				// logic in UK2Node_CallFunction to handle bWantsEnumToExecExpansion:
-				if (TrueSourcePin && !TrueSourcePin->IsPendingKill())
-				{
-					ClassBeingBuilt->GetDebugData().RegisterPinToCodeAssociation(TrueSourcePin, FunctionContext.Function, Offset);
-				}
+				ClassBeingBuilt->GetDebugData().RegisterPinToCodeAssociation(TrueSourcePin, FunctionContext.Function, Offset);
 			}
+		}
 
-			if (SourceNode != NULL)
+		if (SourceNode != NULL)
+		{
+			// Record where this NOP is
+			UEdGraphNode* TrueSourceNode = Cast<UEdGraphNode>(FunctionContext.MessageLog.FindSourceObject(SourceNode));
+			if (TrueSourceNode)
 			{
-				// Record where this NOP is
-				UEdGraphNode* TrueSourceNode = Cast<UEdGraphNode>(FunctionContext.MessageLog.FindSourceObject(SourceNode));
-				if (TrueSourceNode)
+				// If this is a debug site for an expanded macro instruction, there should also be a macro source node associated with it
+				UEdGraphNode* MacroSourceNode = Cast<UEdGraphNode>(CompilerContext.MessageLog.FinalNodeBackToMacroSourceMap.FindSourceObject(SourceNode));
+				if (MacroSourceNode == SourceNode)
 				{
-					// If this is a debug site for an expanded macro instruction, there should also be a macro source node associated with it
-					UEdGraphNode* MacroSourceNode = Cast<UEdGraphNode>(CompilerContext.MessageLog.FinalNodeBackToMacroSourceMap.FindSourceObject(SourceNode));
-					if (MacroSourceNode == SourceNode)
-					{
-						// The function above will return the given node if not found in the map. In that case there is no associated source macro node, so we clear it.
-						MacroSourceNode = NULL;
-					}
+					// The function above will return the given node if not found in the map. In that case there is no associated source macro node, so we clear it.
+					MacroSourceNode = NULL;
+				}
 
-					TArray<TWeakObjectPtr<UEdGraphNode>> MacroInstanceNodes;
-					const bool bInstrumentedBreakpoint = Statement.Type == KCST_InstrumentedWireEntry;
-					bool bBreakpointSite = Statement.Type == KCST_DebugSite || bInstrumentedBreakpoint;
+				TArray<TWeakObjectPtr<UEdGraphNode>> MacroInstanceNodes;
+				const bool bInstrumentedBreakpoint = Statement.Type == KCST_InstrumentedWireEntry;
+				bool bBreakpointSite = Statement.Type == KCST_DebugSite || bInstrumentedBreakpoint;
 
-					if (MacroSourceNode)
+				if (MacroSourceNode)
+				{
+					// Only associate macro instance node breakpoints with source nodes that are linked to the entry node in an impure macro graph
+					if (bBreakpointSite)
 					{
-						// Only associate macro instance node breakpoints with source nodes that are linked to the entry node in an impure macro graph
-						if (bBreakpointSite)
+						const UK2Node_MacroInstance* MacroInstanceNode = Cast<const UK2Node_MacroInstance>(TrueSourceNode);
+						if (MacroInstanceNode)
 						{
-							const UK2Node_MacroInstance* MacroInstanceNode = Cast<const UK2Node_MacroInstance>(TrueSourceNode);
-							if (MacroInstanceNode)
-							{
-								TArray<const UEdGraphNode*> ValidBreakpointLocations;
-								FKismetDebugUtilities::GetValidBreakpointLocations(MacroInstanceNode, ValidBreakpointLocations);
-								bBreakpointSite = ValidBreakpointLocations.Contains(MacroSourceNode);
-							}
+							TArray<const UEdGraphNode*> ValidBreakpointLocations;
+							FKismetDebugUtilities::GetValidBreakpointLocations(MacroInstanceNode, ValidBreakpointLocations);
+							bBreakpointSite = ValidBreakpointLocations.Contains(MacroSourceNode);
 						}
-
-						// Gather up all the macro instance nodes that lead to this macro source node
-						CompilerContext.MessageLog.MacroSourceToMacroInstanceNodeMap.MultiFind(MacroSourceNode, MacroInstanceNodes);
 					}
 
-					ClassBeingBuilt->GetDebugData().RegisterNodeToCodeAssociation(TrueSourceNode, MacroSourceNode, MacroInstanceNodes, FunctionContext.Function, Offset, bBreakpointSite);
+					// Gather up all the macro instance nodes that lead to this macro source node
+					CompilerContext.MessageLog.MacroSourceToMacroInstanceNodeMap.MultiFind(MacroSourceNode, MacroInstanceNodes);
+				}
 
-					// Track pure node script code range for the current impure (exec) node
-					if (Statement.Type == KCST_InstrumentedPureNodeEntry)
-					{
-						if (PureNodeEntryCount == 0)
-						{
-							// Indicates the starting offset for this pure node call chain.
-							PureNodeEntryStart = Offset;
-						}
+				ClassBeingBuilt->GetDebugData().RegisterNodeToCodeAssociation(TrueSourceNode, MacroSourceNode, MacroInstanceNodes, FunctionContext.Function, Offset, bBreakpointSite);
 
-						++PureNodeEntryCount;
-					}
-					else if (Statement.Type == KCST_InstrumentedWireEntry && PureNodeEntryCount > 0)
+				// Track pure node script code range for the current impure (exec) node
+				if (Statement.Type == KCST_InstrumentedPureNodeEntry)
+				{
+					if (PureNodeEntryCount == 0)
 					{
-						// Map script code range for the full set of pure node inputs feeding in to the current impure (exec) node at the current offset
-						ClassBeingBuilt->GetDebugData().RegisterPureNodeScriptCodeRange(TrueSourceNode, FunctionContext.Function, FInt32Range(PureNodeEntryStart, Offset));
-
-						// Reset pure node code range tracking.
-						PureNodeEntryCount = 0;
-						PureNodeEntryStart = 0;
+						// Indicates the starting offset for this pure node call chain.
+						PureNodeEntryStart = Offset;
 					}
+
+					++PureNodeEntryCount;
+				}
+				else if (Statement.Type == KCST_InstrumentedWireEntry && PureNodeEntryCount > 0)
+				{
+					// Map script code range for the full set of pure node inputs feeding in to the current impure (exec) node at the current offset
+					ClassBeingBuilt->GetDebugData().RegisterPureNodeScriptCodeRange(TrueSourceNode, FunctionContext.Function, FInt32Range(PureNodeEntryStart, Offset));
+
+					// Reset pure node code range tracking.
+					PureNodeEntryCount = 0;
+					PureNodeEntryStart = 0;
 				}
 			}
 		}
diff --git a/Engine/Source/Editor/KismetCompiler/Public/KismetCompiler.h b/Engine/Source/Editor/KismetCompiler/Public/KismetCompiler.h
--- a/Engine/Source/Editor/KismetCompiler/Public/KismetCompiler.h
+++ b/Engine/Source/Editor/KismetCompiler/Public/KismetCompiler.h
@@ -81,6 +81,10 @@ public:
 	TMap<UEdGraphNode*, UEdGraphNode*> CallsIntoUbergraph;
 	int32 bIsFullCompile:1;
 
+	// Source pin to latent/delegate entry point
+	TMap<UEdGraphNode*, UK2Node_Event*> SourceNodeToExpansionEvent;
+	TMap<UEdGraphPin*, UK2Node_Event*> SourcePinToExpansionEvent;
+
 	// Map that can be used to find the macro node that spawned a provided node, 
 	// if any. Macro instances can have more macros inside of them, so entries 
 	// in this map may chain (i.e. values may also need to be used as keys to find
@@ -109,7 +113,8 @@ public:
 		}
 
 		NodeType* Result = ParentGraph->CreateBlankNode<NodeType>();
-		MessageLog.NotifyIntermediateObjectCreation(Result, SourceNode);
+		check (Cast<UK2Node_Event>(Result) == nullptr);
+		MessageLog.NotifyIntermediateObjectCreation(Result, SourceNode); // this might be useful to track back function entry nodes to events.
 		Result->CreateNewGuid();
 
 		AutoAssignNodePosition(Result);
@@ -117,6 +122,39 @@ public:
 		return Result;
 	}
 
+	// Spawns an intermediate event node associated with the source node (for error purposes)
+	template <typename NodeType>
+	NodeType* SpawnIntermediateEventNode(UEdGraphNode* SourceNode, UEdGraphPin* SourcePin = nullptr, UEdGraph* ParentGraph = nullptr)
+	{
+		if (ParentGraph == nullptr)
+		{
+			ParentGraph = SourceNode->GetGraph();
+		}
+
+		NodeType* Result = ParentGraph->CreateBlankNode<NodeType>();
+		check (Cast<UK2Node_Event>(Result) != nullptr);
+		MessageLog.NotifyIntermediateObjectCreation(Result, SourceNode); // this might be useful to track back function entry nodes to events.
+		Result->CreateNewGuid();
+		// Track all expansion step generated events.
+		if (bIsFullCompile)
+		{
+			if (SourcePin)
+			{
+				UEdGraphNode* TrueSourceNode = Cast<UEdGraphNode>(MessageLog.FindSourceObject(SourcePin->GetOwningNode()));
+				UEdGraphPin* TrueSourcePin = TrueSourceNode->FindPin(SourcePin->GetName());
+				SourcePinToExpansionEvent.Add(TrueSourcePin) = Cast<UK2Node_Event>(Result);
+			}
+			else if (SourceNode)
+			{
+				UEdGraphNode* TrueSourceNode = Cast<UEdGraphNode>(MessageLog.FindSourceObject(SourceNode));
+				SourceNodeToExpansionEvent.Add(TrueSourceNode) = Cast<UK2Node_Event>(Result);
+			}
+		}
+		AutoAssignNodePosition(Result);
+
+		return Result;
+	}
+
 	/**
 	 * Moves pin links over from the source-pin to the specified intermediate, 
 	 * and validates the result (additionally logs a redirect from the 
diff --git a/Engine/Source/Editor/KismetWidgets/Private/SPinTypeSelector.cpp b/Engine/Source/Editor/KismetWidgets/Private/SPinTypeSelector.cpp
--- a/Engine/Source/Editor/KismetWidgets/Private/SPinTypeSelector.cpp
+++ b/Engine/Source/Editor/KismetWidgets/Private/SPinTypeSelector.cpp
@@ -97,8 +97,7 @@ void SPinTypeSelector::Construct(const FArguments& InArgs, FGetPinTypeTree GetPi
 	GetPinTypeTree = GetPinTypeTreeFunc;
 
 	Schema = (UEdGraphSchema_K2*)(InArgs._Schema);
-	bAllowExec = InArgs._bAllowExec;
-	bAllowWildcard = InArgs._bAllowWildcard;
+	TypeTreeFilter = InArgs._TypeTreeFilter;
 	TreeViewWidth = InArgs._TreeViewWidth;
 	TreeViewHeight = InArgs._TreeViewHeight;
 
@@ -535,7 +534,7 @@ void SPinTypeSelector::GetTypeChildren(FPinTypeTreeItem InItem, TArray<FPinTypeT
 
 TSharedRef<SWidget>	SPinTypeSelector::GetMenuContent()
 {
-	GetPinTypeTree.Execute(TypeTreeRoot, bAllowExec, bAllowWildcard);
+	GetPinTypeTree.Execute(TypeTreeRoot, TypeTreeFilter);
 
 	// Remove read-only root items if they have no children; there will be no subtree to select non read-only items from in that case
 	int32 RootItemIndex = 0;
diff --git a/Engine/Source/Editor/KismetWidgets/Public/SPinTypeSelector.h b/Engine/Source/Editor/KismetWidgets/Public/SPinTypeSelector.h
--- a/Engine/Source/Editor/KismetWidgets/Public/SPinTypeSelector.h
+++ b/Engine/Source/Editor/KismetWidgets/Public/SPinTypeSelector.h
@@ -13,7 +13,7 @@ DECLARE_DELEGATE_OneParam(FOnPinTypeChanged, const FEdGraphPinType&)
 typedef TSharedPtr<class UEdGraphSchema_K2::FPinTypeTreeInfo> FPinTypeTreeItem;
 typedef STreeView<FPinTypeTreeItem> SPinTypeTreeView;
 
-DECLARE_DELEGATE_ThreeParams(FGetPinTypeTree, TArray<FPinTypeTreeItem >&, bool, bool);
+DECLARE_DELEGATE_TwoParams(FGetPinTypeTree, TArray<FPinTypeTreeItem >&, ETypeTreeFilter);
 
 struct FObjectReferenceType;
 typedef TSharedPtr<struct FObjectReferenceType> FObjectReferenceListItem;
@@ -25,8 +25,7 @@ public:
 	SLATE_BEGIN_ARGS( SPinTypeSelector )
 		: _TargetPinType()
 		, _Schema(NULL)
-		, _bAllowExec(false)
-		, _bAllowWildcard(false)
+		, _TypeTreeFilter(ETypeTreeFilter::None)
 		, _bAllowArrays(true)
 		, _TreeViewWidth(300.f)
 		, _TreeViewHeight(400.f)
@@ -35,8 +34,7 @@ public:
 		{}
 		SLATE_ATTRIBUTE( FEdGraphPinType, TargetPinType )
 		SLATE_ARGUMENT( const UEdGraphSchema_K2*, Schema )
-		SLATE_ARGUMENT( bool, bAllowExec )
-		SLATE_ARGUMENT( bool, bAllowWildcard )
+		SLATE_ARGUMENT( ETypeTreeFilter, TypeTreeFilter )
 		SLATE_ARGUMENT( bool, bAllowArrays )
 		SLATE_ATTRIBUTE( FOptionalSize, TreeViewWidth )
 		SLATE_ATTRIBUTE( FOptionalSize, TreeViewHeight )
@@ -81,11 +79,8 @@ protected:
 	/** Schema in charge of determining available types for this pin */
 	UEdGraphSchema_K2*			Schema;
 
-	/** Whether or not to allow exec type wires in the type selection */
-	bool						bAllowExec;
-
-	/** Whether or not to allow wildcard as an option in the type selection */
-	bool						bAllowWildcard;
+	/** UEdgraphSchema::ETypeTreeFilter flags for filtering available types*/
+	ETypeTreeFilter				TypeTreeFilter;
 
 	/** Desired width of the tree view widget */
 	TAttribute<FOptionalSize>	TreeViewWidth;
diff --git a/Engine/Source/Editor/UnrealEd/Classes/Editor/TransBuffer.h b/Engine/Source/Editor/UnrealEd/Classes/Editor/TransBuffer.h
--- a/Engine/Source/Editor/UnrealEd/Classes/Editor/TransBuffer.h
+++ b/Engine/Source/Editor/UnrealEd/Classes/Editor/TransBuffer.h
@@ -115,6 +115,7 @@ public:
 	virtual bool IsObjectSerializationEnabled() override { return DisallowObjectSerialization == 0; }
 	virtual void SetPrimaryUndoObject( UObject* Object ) override;
 	virtual bool IsObjectInTransationBuffer( const UObject* Object ) const override;
+	virtual bool IsObjectTransacting(const UObject* Object) const override;
 	virtual bool ContainsPieObject() const override;
 	virtual bool IsActive() override
 	{
@@ -179,4 +180,7 @@ private:
 
 	// Holds an event delegate that is executed when a undo operation is being attempted.
 	FOnTransactorUndo UndoDelegate;
+
+	// Reference to the current transaction, nullptr when not transacting:
+	FTransaction* CurrentTransaction;
 };
diff --git a/Engine/Source/Editor/UnrealEd/Classes/Editor/Transactor.h b/Engine/Source/Editor/UnrealEd/Classes/Editor/Transactor.h
--- a/Engine/Source/Editor/UnrealEd/Classes/Editor/Transactor.h
+++ b/Engine/Source/Editor/UnrealEd/Classes/Editor/Transactor.h
@@ -346,6 +346,8 @@ protected:
 	/** Count of the number of UModels modified since the last call to FTransaction::Apply */
 	int32					NumModelsModified;
 
+	/** Objects that will be changed directly by the transaction, empty when not transacting */
+	TMap<UObject*, TSharedPtr<ITransactionObjectAnnotation>> ChangedObjects;
 public:
 	// Constructor.
 	FTransaction(  const TCHAR* InContext=NULL, const FText& InTitle=FText(), bool InFlip=0 )
@@ -415,6 +417,9 @@ public:
 	/** @return True if this record contains a reference to a pie object */
 	bool ContainsPieObject() const;
 
+	/** Checks if a specific object is in the transaction currently underway */
+	bool IsObjectTransacting(const UObject* Object) const;
+
 	/**
 	 * Outputs the contents of the ObjectMap to the specified output device.
 	 */
@@ -594,11 +599,14 @@ class UTransactor : public UObject
 	/** 
 	 * Set passed object as the primary context object for transactions
 	 */
-	virtual void SetPrimaryUndoObject( UObject* Object ) PURE_VIRTUAL(UTransactor::MakePrimaryUndoObject,);
+	virtual void SetPrimaryUndoObject( UObject* Object ) PURE_VIRTUAL(UTransactor::SetPrimaryUndoObject,);
 
 	/** Checks if a specific object is referenced by the transaction buffer */
 	virtual bool IsObjectInTransationBuffer( const UObject* Object ) const { return false; }
 
+	/** Checks if a specific object is in the transaction currently underway */
+	virtual bool IsObjectTransacting(const UObject* Object) const PURE_VIRTUAL(UTransactor::IsObjectTransacting, return false;);
+
 	/** @return True if this record contains a reference to a pie object */
 	virtual bool ContainsPieObject() const { return false; }
 };
diff --git a/Engine/Source/Editor/UnrealEd/Private/Commandlets/DumpBlueprintsInfoCommandlet.cpp b/Engine/Source/Editor/UnrealEd/Private/Commandlets/DumpBlueprintsInfoCommandlet.cpp
--- a/Engine/Source/Editor/UnrealEd/Private/Commandlets/DumpBlueprintsInfoCommandlet.cpp
+++ b/Engine/Source/Editor/UnrealEd/Private/Commandlets/DumpBlueprintsInfoCommandlet.cpp
@@ -1930,7 +1930,7 @@ static bool DumpBlueprintInfoUtils::DumpPinContextActions(uint32 Indent, UEdGrap
 	{
 		TArray< TSharedPtr<UEdGraphSchema_K2::FPinTypeTreeInfo> > TypeTree;
 		UEdGraphSchema_K2 const* GraphSchema = GetDefault<UEdGraphSchema_K2>(Graph->Schema);
-		GraphSchema->GetVariableTypeTree(TypeTree, /*bAllowExec =*/true, /*bAllowWildcard =*/true);
+		GraphSchema->GetVariableTypeTree(TypeTree, ETypeTreeFilter::AllowExec | ETypeTreeFilter::AllowWildcard);
 
 		for (TSharedPtr<UEdGraphSchema_K2::FPinTypeTreeInfo> TypeInfo : TypeTree)
 		{
diff --git a/Engine/Source/Editor/UnrealEd/Private/EdGraphUtilities.cpp b/Engine/Source/Editor/UnrealEd/Private/EdGraphUtilities.cpp
--- a/Engine/Source/Editor/UnrealEd/Private/EdGraphUtilities.cpp
+++ b/Engine/Source/Editor/UnrealEd/Private/EdGraphUtilities.cpp
@@ -8,6 +8,7 @@
 #include "UnrealExporter.h"
 #include "SNotificationList.h"
 #include "NotificationManager.h"
+#include "K2Node_TunnelBoundary.h"
 
 /////////////////////////////////////////////////////
 // FGraphObjectTextFactory
@@ -235,10 +236,20 @@ UEdGraph* FEdGraphUtilities::CloneGraph(UEdGraph* InSource, UObject* NewOuter, F
 }
 
 // Clones the content from SourceGraph and merges it into MergeTarget; including merging/flattening all of the children from the SourceGraph into MergeTarget
-void FEdGraphUtilities::CloneAndMergeGraphIn(UEdGraph* MergeTarget, UEdGraph* SourceGraph, FCompilerResultsLog& MessageLog, bool bRequireSchemaMatch, bool bInIsCompiling/* = false*/, TArray<UEdGraphNode*>* OutClonedNodes)
+void FEdGraphUtilities::CloneAndMergeGraphIn(UEdGraph* MergeTarget, UEdGraph* SourceGraph, FCompilerResultsLog& MessageLog, bool bRequireSchemaMatch, bool bInIsCompiling/* = false*/, bool bCreateBoundaryNodes/* = false*/, TArray<UEdGraphNode*>* OutClonedNodes)
 {
 	// Clone the graph, then move all of it's children
 	UEdGraph* ClonedGraph = CloneGraph(SourceGraph, NULL, &MessageLog, true);
+
+	// Create any Boundary nodes around child graphs if specified.
+	if (bCreateBoundaryNodes)
+	{
+		for (UEdGraph* SubGraph : ClonedGraph->SubGraphs)
+		{
+			UK2Node_TunnelBoundary::CreateBoundaryNodesForGraph(SubGraph, MessageLog);
+		}
+	}
+
 	MergeChildrenGraphsIn(ClonedGraph, ClonedGraph, bRequireSchemaMatch);
 
 	// Duplicate the list of cloned nodes
diff --git a/Engine/Source/Editor/UnrealEd/Private/EditorTransaction.cpp b/Engine/Source/Editor/UnrealEd/Private/EditorTransaction.cpp
--- a/Engine/Source/Editor/UnrealEd/Private/EditorTransaction.cpp
+++ b/Engine/Source/Editor/UnrealEd/Private/EditorTransaction.cpp
@@ -201,6 +201,15 @@ bool FTransaction::ContainsPieObject() const
 	return false;
 }
 
+bool FTransaction::IsObjectTransacting(const UObject* Object) const
+{
+	// This function is meaningless when called outside of a transaction context. Without this
+	// ensure clients will commonly introduced bugs by having some logic that runs during
+	// the transacting and some logic that does not, yielding assymetrical results.
+	ensure(GIsTransacting);
+	ensure(ChangedObjects.Num() != 0);
+	return ChangedObjects.Contains(Object);
+}
 
 void FTransaction::RemoveRecords( int32 Count /* = 1  */ )
 {
@@ -408,7 +417,6 @@ void FTransaction::Apply()
 	const int32 End   = Inc==1 ? Records.Num() :              -1;
 
 	// Init objects.
-	TMap<UObject*, TSharedPtr<ITransactionObjectAnnotation>> ChangedObjects;
 	for( int32 i=Start; i!=End; i+=Inc )
 	{
 		FObjectRecord& Record = Records[i];
@@ -499,6 +507,8 @@ void FTransaction::Apply()
 		UObject* ChangedObject = ChangedObjectIt.Key;
 		ChangedObject->CheckDefaultSubobjects();
 	}
+
+	ChangedObjects.Empty();
 }
 
 SIZE_T FTransaction::DataSize() const
@@ -838,6 +848,7 @@ bool UTransBuffer::Undo(bool bCanRedo)
 	{
 		FTransaction& Transaction = UndoBuffer[ UndoBuffer.Num() - ++UndoCount ];
 		UE_LOG(LogEditorTransaction, Log,  TEXT("Undo %s"), *Transaction.GetTitle().ToString() );
+		CurrentTransaction = &Transaction;
 
 		BeforeRedoUndoDelegate.Broadcast(Transaction.GetContext());
 		Transaction.Apply();
@@ -848,6 +859,8 @@ bool UTransBuffer::Undo(bool bCanRedo)
 			UndoBuffer.RemoveAt(UndoBuffer.Num() - UndoCount, UndoCount);
 			UndoCount = 0;
 		}
+
+		CurrentTransaction = nullptr;
 	}
 	GIsTransacting = false;
 
@@ -872,10 +885,13 @@ bool UTransBuffer::Redo()
 	{
 		FTransaction& Transaction = UndoBuffer[ UndoBuffer.Num() - UndoCount-- ];
 		UE_LOG(LogEditorTransaction, Log,  TEXT("Redo %s"), *Transaction.GetTitle().ToString() );
+		CurrentTransaction = &Transaction;
 
 		BeforeRedoUndoDelegate.Broadcast(Transaction.GetContext());
 		Transaction.Apply();
 		RedoDelegate.Broadcast(Transaction.GetContext(), true);
+
+		CurrentTransaction = nullptr;
 	}
 	GIsTransacting = false;
 
@@ -950,6 +966,17 @@ bool UTransBuffer::IsObjectInTransationBuffer( const UObject* Object ) const
 	return false;
 }
 
+bool UTransBuffer::IsObjectTransacting(const UObject* Object) const
+{
+	// We can't provide a truly meaningful answer to this question when not transacting:
+	if (ensure(CurrentTransaction))
+	{
+		return CurrentTransaction->IsObjectTransacting(Object);
+	}
+
+	return false;
+}
+
 bool UTransBuffer::ContainsPieObject() const
 {
 	for( const FTransaction& Transaction : UndoBuffer )
diff --git a/Engine/Source/Editor/UnrealEd/Private/Kismet2/BlueprintEditorUtils.cpp b/Engine/Source/Editor/UnrealEd/Private/Kismet2/BlueprintEditorUtils.cpp
--- a/Engine/Source/Editor/UnrealEd/Private/Kismet2/BlueprintEditorUtils.cpp
+++ b/Engine/Source/Editor/UnrealEd/Private/Kismet2/BlueprintEditorUtils.cpp
@@ -1048,7 +1048,7 @@ struct FRegenerationHelper
 				ForcedLoadMembers(InterfaceBlueprint);
 				if (InterfaceBlueprint->HasAnyFlags(RF_BeingRegenerated))
 				{
-					InterfaceBlueprint->RegenerateClass(InterfaceClass, nullptr, ObjLoaded);
+					InterfaceBlueprint->RegenerateClass(InterfaceClass, InterfaceClass->ClassDefaultObject, ObjLoaded);
 				}
 			}
 		}
@@ -1466,9 +1466,11 @@ UClass* FBlueprintEditorUtils::RegenerateBlueprintClass(UBlueprint* Blueprint, U
 	// member further up the stack).
 	if (!Blueprint->bHasBeenRegenerated)
 	{
-		check(PreviousCDO != nullptr);
 		FRegenerationHelper::ForceLoadMetaData(Blueprint);
-		FRegenerationHelper::ForcedLoadMembers(PreviousCDO);
+		if (ensure(PreviousCDO))
+		{
+			FRegenerationHelper::ForcedLoadMembers(PreviousCDO);
+		}
 		FRegenerationHelper::ForcedLoadMembers(Blueprint);
 	}
 
@@ -1730,6 +1732,16 @@ UClass* FBlueprintEditorUtils::RegenerateBlueprintClass(UBlueprint* Blueprint, U
 			}
 		}
 	}
+	else
+	{
+		if (Blueprint->GeneratedClass && !Blueprint->bHasBeenRegenerated && !Blueprint->bIsRegeneratingOnLoad)
+		{
+			FObjectDuplicationParameters Params(Blueprint->GeneratedClass, Blueprint->GeneratedClass->GetOuter());
+			Params.ApplyFlags = RF_Transient;
+			Params.DestName = *(FString("SKEL_COPY_") + Blueprint->GeneratedClass->GetName());
+			Blueprint->SkeletonGeneratedClass = (UClass*)StaticDuplicateObjectEx(Params);
+		}
+	}
 
 	if ( bRegenerated )
 	{		
diff --git a/Engine/Source/Editor/UnrealEd/Private/Kismet2/CompilerResultsLog.cpp b/Engine/Source/Editor/UnrealEd/Private/Kismet2/CompilerResultsLog.cpp
--- a/Engine/Source/Editor/UnrealEd/Private/Kismet2/CompilerResultsLog.cpp
+++ b/Engine/Source/Editor/UnrealEd/Private/Kismet2/CompilerResultsLog.cpp
@@ -87,15 +87,18 @@ UEdGraphPin* FBacktrackMap::FindSourcePin(UEdGraphPin* PossiblyDuplicatedPin)
 	{
 		// Not in the map, maybe its owning node was duplicated - and then maybe he GUID matches
 		// some node on the original node:
-		if (UObject* OriginalOwner = FindSourceObject(PossiblyDuplicatedPin->GetOwningNode()))
+		if (PossiblyDuplicatedPin)
 		{
-			if (UEdGraphNode* AsNode = Cast<UEdGraphNode>(OriginalOwner))
+			if (UObject* OriginalOwner = FindSourceObject(PossiblyDuplicatedPin->GetOwningNode()))
 			{
-				FGuid TargetGuid = PossiblyDuplicatedPin->PinId;
-				UEdGraphPin** ExistingPin = AsNode->Pins.FindByPredicate([TargetGuid](const UEdGraphPin* TargetPin) { return TargetPin && TargetPin->PinId == TargetGuid; });
-				if (ExistingPin != nullptr)
+				if (UEdGraphNode* AsNode = Cast<UEdGraphNode>(OriginalOwner))
 				{
-					return *ExistingPin;
+					FGuid TargetGuid = PossiblyDuplicatedPin->PinId;
+					UEdGraphPin** ExistingPin = AsNode->Pins.FindByPredicate([TargetGuid](const UEdGraphPin* TargetPin) { return TargetPin && TargetPin->PinId == TargetGuid; });
+					if (ExistingPin != nullptr)
+					{
+						return *ExistingPin;
+					}
 				}
 			}
 		}
diff --git a/Engine/Source/Editor/UnrealEd/Private/Kismet2/Kismet2.cpp b/Engine/Source/Editor/UnrealEd/Private/Kismet2/Kismet2.cpp
--- a/Engine/Source/Editor/UnrealEd/Private/Kismet2/Kismet2.cpp
+++ b/Engine/Source/Editor/UnrealEd/Private/Kismet2/Kismet2.cpp
@@ -2081,9 +2081,12 @@ void FKismetEditorUtilities::StripExternalComponents(class UBlueprint* Blueprint
 	FArchiveInvalidateTransientRefs InvalidateRefsAr;
 	
 	UClass* SkeletonGeneratedClass = Blueprint->SkeletonGeneratedClass;
-	UObject* SkeletonCDO = SkeletonGeneratedClass->GetDefaultObject();
+	if (SkeletonGeneratedClass)
+	{
+		UObject* SkeletonCDO = SkeletonGeneratedClass->GetDefaultObject();
 
-	SkeletonCDO->Serialize(InvalidateRefsAr);
+		SkeletonCDO->Serialize(InvalidateRefsAr);
+	}
 
 	UClass* GeneratedClass = Blueprint->GeneratedClass;
 	UObject* GeneratedCDO = GeneratedClass->GetDefaultObject();
diff --git a/Engine/Source/Editor/UnrealEd/Public/EdGraphUtilities.h b/Engine/Source/Editor/UnrealEd/Public/EdGraphUtilities.h
--- a/Engine/Source/Editor/UnrealEd/Public/EdGraphUtilities.h
+++ b/Engine/Source/Editor/UnrealEd/Public/EdGraphUtilities.h
@@ -70,7 +70,7 @@ public:
 	  * @param OutClonedNodes			Will populate with a full list of cloned nodes if provided
 	  * @param bInIsCompiling			TRUE if the function is being called during compilation, this will eliminate some nodes that will not be compiled
 	  */
-	static void CloneAndMergeGraphIn(UEdGraph* MergeTarget, UEdGraph* SourceGraph, FCompilerResultsLog& MessageLog, bool bRequireSchemaMatch, bool bInIsCompiling = false, TArray<UEdGraphNode*>* OutClonedNodes = NULL);
+	static void CloneAndMergeGraphIn(UEdGraph* MergeTarget, UEdGraph* SourceGraph, FCompilerResultsLog& MessageLog, bool bRequireSchemaMatch, bool bInIsCompiling = false, bool bCreateBoundaryNodes = false, TArray<UEdGraphNode*>* OutClonedNodes = NULL);
 
 	/**
 	 * Moves the contents of all of the children graphs of ParentGraph (recursively) into the MergeTarget graph.
diff --git a/Engine/Source/Programs/AutomationTool/AutomationUtils/ProjectParams.cs b/Engine/Source/Programs/AutomationTool/AutomationUtils/ProjectParams.cs
--- a/Engine/Source/Programs/AutomationTool/AutomationUtils/ProjectParams.cs
+++ b/Engine/Source/Programs/AutomationTool/AutomationUtils/ProjectParams.cs
@@ -1032,19 +1032,24 @@ namespace AutomationTool
         [Help("nativizeAssets", "Runs a \"nativization\" pass on Blueprint assets, converting then into C++ (replacing the assets with the generated source).")]
         public bool RunAssetNativization;
 
-		/// <summary>
-		/// Shared: Ref to an auto-generated plugin file that should be incorporated into the project's build
-		/// </summary>
-        public List<FileReference> BlueprintPluginPaths = new List<FileReference>();
+        public struct BlueprintPluginKey
+        {
+            public bool Client;
+            public UnrealTargetPlatform TargetPlatform;
+        }
+        /// <summary>
+        /// Shared: Ref to an auto-generated plugin file that should be incorporated into the project's build
+        /// </summary>
+        public Dictionary<BlueprintPluginKey, FileReference> BlueprintPluginPaths = new Dictionary<BlueprintPluginKey, FileReference>();
 
-		#endregion
+        #endregion
 
-		#region Build
+        #region Build
 
-		/// <summary>
-		/// Build: True if build step should be executed, command: -build
-		/// </summary>
-		[Help("build", "True if build step should be executed")]
+        /// <summary>
+        /// Build: True if build step should be executed, command: -build
+        /// </summary>
+        [Help("build", "True if build step should be executed")]
 		public bool Build { private set; get; }
 
 		/// <summary>
diff --git a/Engine/Source/Programs/AutomationTool/Scripts/BuildProjectCommand.Automation.cs b/Engine/Source/Programs/AutomationTool/Scripts/BuildProjectCommand.Automation.cs
--- a/Engine/Source/Programs/AutomationTool/Scripts/BuildProjectCommand.Automation.cs
+++ b/Engine/Source/Programs/AutomationTool/Scripts/BuildProjectCommand.Automation.cs
@@ -37,6 +37,32 @@ public partial class Project : CommandUtils
 {
 	#region Build Command
 
+    static string GetBlueprintPluginPathArgument(ProjectParams Params, bool Client, UnrealTargetPlatform TargetPlatform)
+    {
+        string ScriptPluginArgs = "";
+
+        // if we're utilizing an auto-generated code plugin/module (a product of 
+        // the cook process), make sure to compile it along with the targets here
+
+        if (Params.RunAssetNativization)
+        {
+            ProjectParams.BlueprintPluginKey PluginKey = new ProjectParams.BlueprintPluginKey();
+            PluginKey.Client = Client;
+            PluginKey.TargetPlatform = TargetPlatform;
+            FileReference CodePlugin = null;
+            if(Params.BlueprintPluginPaths.TryGetValue(PluginKey, out CodePlugin))
+            {
+                ScriptPluginArgs += "-PLUGIN \"" + CodePlugin + "\" ";
+            }
+            else
+            {
+                LogWarning("BlueprintPluginPath for " + TargetPlatform + " " + (Client ? "client" : "server") + " was not found");
+            }
+        }
+
+        return ScriptPluginArgs;
+    }
+
     public static void Build(BuildCommand Command, ProjectParams Params, int WorkingCL = -1, ProjectBuildTargets TargetMask = ProjectBuildTargets.All)
 	{
 		Params.ValidateAndLog();
@@ -79,18 +105,6 @@ public partial class Project : CommandUtils
 			}
 		}
 
-        string ScriptPluginArgs = "";
-        // if we're utilizing an auto-generated code plugin/module (a product of 
-        // the cook process), make sure to compile it along with the targets here
-        if (Params.RunAssetNativization)
-        {
-            // Add every plugin:
-            foreach( var CodePlugin in  Params.BlueprintPluginPaths)
-            {
-                ScriptPluginArgs += "-PLUGIN \"" + CodePlugin + "\" "; 
-            }
-        }
-
 		// Setup cooked targets
 		if (Params.HasClientCookedTargets && (TargetMask & ProjectBuildTargets.ClientCooked) == ProjectBuildTargets.ClientCooked)
 		{
@@ -98,7 +112,8 @@ public partial class Project : CommandUtils
 			{
 				foreach (var ClientPlatform in Params.ClientTargetPlatforms)
 				{
-					CrashReportPlatforms.Add(ClientPlatform);
+                    string ScriptPluginArgs = GetBlueprintPluginPathArgument(Params, true, ClientPlatform);
+                    CrashReportPlatforms.Add(ClientPlatform);
 					Agenda.AddTargets(Params.ClientCookedTargets.ToArray(), ClientPlatform, BuildConfig, Params.CodeBasedUprojectPath, InAddArgs: ScriptPluginArgs + " -remoteini=\"" + Params.RawProjectPath.Directory.FullName + "\"");
 				}
 			}
@@ -109,7 +124,8 @@ public partial class Project : CommandUtils
 			{
 				foreach (var ServerPlatform in Params.ServerTargetPlatforms)
 				{
-					CrashReportPlatforms.Add(ServerPlatform);
+                    string ScriptPluginArgs = GetBlueprintPluginPathArgument(Params, false, ServerPlatform);
+                    CrashReportPlatforms.Add(ServerPlatform);
 					Agenda.AddTargets(Params.ServerCookedTargets.ToArray(), ServerPlatform, BuildConfig, Params.CodeBasedUprojectPath, InAddArgs: ScriptPluginArgs + " -remoteini=\"" + Params.RawProjectPath.Directory.FullName + "\"");
 				}
 			}
diff --git a/Engine/Source/Programs/AutomationTool/Scripts/CookCommand.Automation.cs b/Engine/Source/Programs/AutomationTool/Scripts/CookCommand.Automation.cs
--- a/Engine/Source/Programs/AutomationTool/Scripts/CookCommand.Automation.cs
+++ b/Engine/Source/Programs/AutomationTool/Scripts/CookCommand.Automation.cs
@@ -17,9 +17,25 @@ using UnrealBuildTool;
 /// </remarks>
 public partial class Project : CommandUtils
 {
-	#region Cook Command
+    #region Cook Command
 
-	public static void Cook(ProjectParams Params)
+    static void AddBlueprintPluginPathArgument(ProjectParams Params, bool Client, UnrealTargetPlatform TargetPlatform, string PlatformToCook)
+    {
+        if (Params.RunAssetNativization)
+        {
+            ProjectParams.BlueprintPluginKey PluginKey = new ProjectParams.BlueprintPluginKey();
+            PluginKey.Client = Client;
+            PluginKey.TargetPlatform = TargetPlatform;
+
+            string ProjectDir = Params.RawProjectPath.Directory.ToString();
+            // If you change this target path you must also update logic in CookOnTheFlyServer.cpp. Passing a single directory around is cumbersome for testing, so I have hard coded it.
+            // Similarly if you change the .uplugin name you must update DefaultPluginName in BlueprintNativeCodeGenModule.cpp
+            string GeneratedPluginPath = CombinePaths(ProjectDir, "Intermediate", PlatformToCook, "NativizedAssets/NativizedAssets.uplugin");
+            Params.BlueprintPluginPaths.Add(PluginKey, new FileReference(GeneratedPluginPath));
+        }
+    }
+
+    public static void Cook(ProjectParams Params)
 	{
 		if ((!Params.Cook && !(Params.CookOnTheFly && !Params.SkipServer)) || Params.SkipCook)
 		{
@@ -80,15 +96,16 @@ public partial class Project : CommandUtils
 		else
 		{
 			var PlatformsToCook = new HashSet<string>();
-
-			if (!Params.NoClient)
+            if (!Params.NoClient)
 			{
 				foreach (var ClientPlatform in Params.ClientTargetPlatforms)
 				{
 					// Use the data platform, sometimes we will copy another platform's data
 					var DataPlatform = Params.GetCookedDataPlatformForClientTarget(ClientPlatform);
-					PlatformsToCook.Add(Params.GetTargetPlatformInstance(DataPlatform).GetCookPlatform(false, Params.Client, Params.CookFlavor));
-				}
+                    string PlatformToCook = Params.GetTargetPlatformInstance(DataPlatform).GetCookPlatform(false, Params.Client, Params.CookFlavor);
+                    PlatformsToCook.Add(PlatformToCook);
+                    AddBlueprintPluginPathArgument(Params, true, DataPlatform, PlatformToCook);
+                }
 			}
 			if (Params.DedicatedServer)
 			{
@@ -96,8 +113,10 @@ public partial class Project : CommandUtils
 				{
 					// Use the data platform, sometimes we will copy another platform's data
 					var DataPlatform = Params.GetCookedDataPlatformForServerTarget(ServerPlatform);
-					PlatformsToCook.Add(Params.GetTargetPlatformInstance(DataPlatform).GetCookPlatform(true, false, Params.CookFlavor));
-				}
+                    string PlatformToCook = Params.GetTargetPlatformInstance(DataPlatform).GetCookPlatform(true, false, Params.CookFlavor);
+                    PlatformsToCook.Add(PlatformToCook);
+                    AddBlueprintPluginPathArgument(Params, false, DataPlatform, PlatformToCook);
+                }
 			}
 
 			if (Params.Clean.HasValue && Params.Clean.Value && !Params.IterativeCooking)
@@ -216,16 +235,6 @@ public partial class Project : CommandUtils
                 if (Params.RunAssetNativization)
                 {
                     CommandletParams += " -NativizeAssets";
-
-                    // Store plugin paths now, it's easiest to do so while PlatformsToCook is still available:
-                    string ProjectDir = Params.RawProjectPath.Directory.ToString();
-                    foreach (var Platform in PlatformsToCook)
-                    {
-                        // If you change this target path you must also update logic in CookOnTheFlyServer.cpp. Passing a single directory around is cumbersome for testing, so I have hard coded it.
-                        // Similarly if you change the .uplugin name you must update DefaultPluginName in BlueprintNativeCodeGenModule.cpp
-                        string GeneratedPluginPath = CombinePaths(ProjectDir, "Intermediate", Platform, "NativizedAssets/NativizedAssets.uplugin");
-                        Params.BlueprintPluginPaths.Add( new FileReference(GeneratedPluginPath) );
-                    }
                 }
                 if (Params.HasAdditionalCookerOptions)
                 {
diff --git a/Engine/Source/Programs/UnrealHeaderTool/Private/HeaderParser.cpp b/Engine/Source/Programs/UnrealHeaderTool/Private/HeaderParser.cpp
--- a/Engine/Source/Programs/UnrealHeaderTool/Private/HeaderParser.cpp
+++ b/Engine/Source/Programs/UnrealHeaderTool/Private/HeaderParser.cpp
@@ -926,8 +926,10 @@ FScriptLocation::FScriptLocation()
 FString FHeaderParser::GetContext()
 {
 	auto* FileScope = GetCurrentFileScope();
-
-	FString ScopeFilename = IFileManager::Get().ConvertToAbsolutePathForExternalAppForRead(*FileScope->GetSourceFile()->GetFilename());
+	FUnrealSourceFile* SourceFile = FileScope ? FileScope->GetSourceFile() : GetCurrentSourceFile();
+	FString ScopeFilename = SourceFile
+		? IFileManager::Get().ConvertToAbsolutePathForExternalAppForRead(*SourceFile->GetFilename())
+		: TEXT("UNKNOWN");
 
 	return FString::Printf(TEXT("%s(%i)"), *ScopeFilename, InputLine);
 }
@@ -2309,7 +2311,16 @@ void FHeaderParser::PopNest(ENestType NestType, const TCHAR* Descr)
 	// Pop the nesting level.
 	NestType = TopNest->NestType;
 	NestLevel--;
-	TopNest--;
+	if (NestLevel == 0)
+	{
+		TopNest = nullptr;
+	}
+	else
+	{
+		TopNest--;
+		check(TopNest >= Nest);
+
+	}
 }
 
 void FHeaderParser::FixupDelegateProperties( FClasses& AllClasses, UStruct* Struct, FScope& Scope, TMap<FName, UFunction*>& DelegateCache )
@@ -2354,7 +2365,10 @@ void FHeaderParser::FixupDelegateProperties( FClasses& AllClasses, UStruct* Stru
 						if (SourceDelegateFunction == nullptr)
 						{
 							// Try to find in other packages.
-							SourceDelegateFunction = Cast<UFunction>(StaticFindObject(UFunction::StaticClass(), ANY_PACKAGE, *NameOfDelegateFunction));
+							UObject* DelegateSignatureOuter = DelegatePropertyToken->Token.DelegateSignatureOwnerClass 
+								? ((UObject*)DelegatePropertyToken->Token.DelegateSignatureOwnerClass) 
+								: ((UObject*)ANY_PACKAGE);
+							SourceDelegateFunction = Cast<UFunction>(StaticFindObject(UFunction::StaticClass(), DelegateSignatureOuter, *NameOfDelegateFunction));
 
 							if (SourceDelegateFunction == nullptr)
 							{
@@ -3589,6 +3603,7 @@ FIndexRange*                    ParsedVarIndexRange
 					if (UFunction* DelegateFunc = Cast<UFunction>(LocScope->FindTypeByName(*(DelegateIdentifierStripped + HEADER_GENERATED_DELEGATE_SIGNATURE_SUFFIX))))
 					{
 						SetDelegateType(DelegateFunc, DelegateIdentifierStripped);
+						VarProperty.DelegateSignatureOwnerClass = LocalOwnerClass;
 					}
 				}
 				else
@@ -7075,7 +7090,8 @@ FHeaderParser::FHeaderParser(FFeedbackContext* InWarn, FUHTMakefile& InUHTMakefi
 , Warn                              (InWarn)
 , UHTMakefile(InUHTMakefile)
 , bSpottedAutogeneratedHeaderInclude(false)
-, TopNest                           (NULL)
+, NestLevel							(0)
+, TopNest                           (nullptr)
 {
 	FScriptLocation::Compiler = this;
 
@@ -8408,5 +8424,8 @@ UFunction* FHeaderParser::CreateFunction(const FFuncInfo &FuncInfo) const
 
 UDelegateFunction* FHeaderParser::CreateDelegateFunction(const FFuncInfo &FuncInfo) const
 {
-	return CreateFunctionImpl<UDelegateFunction>(FuncInfo, IsInAClass() ? (UObject*)GetCurrentClass() : (UObject*)GetCurrentFileScope()->GetSourceFile()->GetPackage(), GetCurrentScope());
+	FFileScope* CurrentFileScope = GetCurrentFileScope();
+	FUnrealSourceFile* LocSourceFile = CurrentFileScope ? CurrentFileScope->GetSourceFile() : nullptr;
+	UObject* CurrentPackage = LocSourceFile ? LocSourceFile->GetPackage() : nullptr;
+	return CreateFunctionImpl<UDelegateFunction>(FuncInfo, IsInAClass() ? (UObject*)GetCurrentClass() : CurrentPackage, GetCurrentScope());
 }
diff --git a/Engine/Source/Programs/UnrealHeaderTool/Private/HeaderParser.h b/Engine/Source/Programs/UnrealHeaderTool/Private/HeaderParser.h
--- a/Engine/Source/Programs/UnrealHeaderTool/Private/HeaderParser.h
+++ b/Engine/Source/Programs/UnrealHeaderTool/Private/HeaderParser.h
@@ -247,6 +247,11 @@ protected:
 	FFileScope* GetCurrentFileScope() const
 	{
 		int32 Index = 0;
+		if (!TopNest)
+		{
+			check(!NestLevel);
+			return nullptr;
+		}
 		while (TopNest[Index].NestType != ENestType::GlobalScope)
 		{
 			--Index;
diff --git a/Engine/Source/Programs/UnrealHeaderTool/Private/ParserHelper.h b/Engine/Source/Programs/UnrealHeaderTool/Private/ParserHelper.h
--- a/Engine/Source/Programs/UnrealHeaderTool/Private/ParserHelper.h
+++ b/Engine/Source/Programs/UnrealHeaderTool/Private/ParserHelper.h
@@ -105,6 +105,7 @@ public:
 
 	UClass* MetaClass;
 	FName   DelegateName;
+	UClass*	DelegateSignatureOwnerClass;
 	FName   RepNotifyName;
 
 	/** Raw string (not type-checked) used for specifying special text when exporting a property to the *Classes.h file */
@@ -127,8 +128,9 @@ public:
 	, RefQualifier        (ERefQualifier::None)
 	, PropertyExportFlags (PROPEXPORT_Public)
 	, StringSize          (0)
-	, MetaClass           (NULL)
+	, MetaClass           (nullptr)
 	, DelegateName        (NAME_None)
+	, DelegateSignatureOwnerClass(nullptr)
 	, RepNotifyName       (NAME_None)
 	, PointerType         (EPointerType::None)
 	, IntType             (GetSizedIntTypeFromPropertyType(InType))
@@ -143,8 +145,9 @@ public:
 	, RefQualifier        (ERefQualifier::None)
 	, PropertyExportFlags (PROPEXPORT_Public)
 	, StringSize          (0)
-	, MetaClass           (NULL)
+	, MetaClass           (nullptr)
 	, DelegateName        (NAME_None)
+	, DelegateSignatureOwnerClass(nullptr)
 	, RepNotifyName       (NAME_None)
 	, PointerType         (EPointerType::None)
 	, IntType             (InIntType)
@@ -159,8 +162,9 @@ public:
 	, RefQualifier        (ERefQualifier::None)
 	, PropertyExportFlags (PROPEXPORT_Public)
 	, Enum                (InEnum)
-	, MetaClass           (NULL)
+	, MetaClass           (nullptr)
 	, DelegateName        (NAME_None)
+	, DelegateSignatureOwnerClass(nullptr)
 	, RepNotifyName       (NAME_None)
 	, PointerType         (EPointerType::None)
 	, IntType             (GetSizedIntTypeFromPropertyType(InType))
@@ -177,6 +181,7 @@ public:
 	, PropertyClass       (InClass)
 	, MetaClass           (InMetaClass)
 	, DelegateName        (NAME_None)
+	, DelegateSignatureOwnerClass(nullptr)
 	, RepNotifyName       (NAME_None)
 	, PointerType         (EPointerType::None)
 	, IntType             (EIntType::None)
@@ -236,6 +241,7 @@ public:
 	, Struct              (InStruct)
 	, MetaClass           (NULL)
 	, DelegateName        (NAME_None)
+	, DelegateSignatureOwnerClass(nullptr)
 	, RepNotifyName       (NAME_None)
 	, PointerType         (EPointerType::None)
 	, IntType             (EIntType::None)
@@ -245,6 +251,7 @@ public:
 	explicit FPropertyBase(UProperty* Property)
 	: PropertyExportFlags(PROPEXPORT_Public)
 	, DelegateName       (NAME_None)
+	, DelegateSignatureOwnerClass(nullptr)
 	, RepNotifyName      (NAME_None)
 	, IntType            (EIntType::None)
 	{
diff --git a/Engine/Source/Runtime/Core/Private/Misc/CoreMisc.cpp b/Engine/Source/Runtime/Core/Private/Misc/CoreMisc.cpp
--- a/Engine/Source/Runtime/Core/Private/Misc/CoreMisc.cpp
+++ b/Engine/Source/Runtime/Core/Private/Misc/CoreMisc.cpp
@@ -1178,6 +1178,11 @@ void FBlueprintExceptionTracker::ResetRunaway()
 	Recurse = 0;
 	bRanaway = false;
 }
+
+FBlueprintExceptionTracker& FBlueprintExceptionTracker::Get()
+{
+	return TThreadSingleton<FBlueprintExceptionTracker>::Get();
+}
 #endif // DO_BLUEPRINT_GUARD
 
 #if WITH_HOT_RELOAD_CTORS
diff --git a/Engine/Source/Runtime/Core/Public/Misc/CoreMisc.h b/Engine/Source/Runtime/Core/Public/Misc/CoreMisc.h
--- a/Engine/Source/Runtime/Core/Public/Misc/CoreMisc.h
+++ b/Engine/Source/Runtime/Core/Public/Misc/CoreMisc.h
@@ -436,6 +436,7 @@ struct CORE_API FBlueprintExceptionTracker : TThreadSingleton<FBlueprintExceptio
 
 	void ResetRunaway();
 
+	static FBlueprintExceptionTracker& Get();
 public:
 	// map of currently displayed warnings in exception handler
 	TMap<FName, int32> DisplayedWarningsMap;
