Changing User Defined Struct Editor to vertical layout, tab collapsed for new users

diff --git a/Engine/Source/Editor/Kismet/Private/UserDefinedStructureEditor.cpp b/Engine/Source/Editor/Kismet/Private/UserDefinedStructureEditor.cpp
--- a/Engine/Source/Editor/Kismet/Private/UserDefinedStructureEditor.cpp
+++ b/Engine/Source/Editor/Kismet/Private/UserDefinedStructureEditor.cpp
@@ -283,7 +283,8 @@ void FUserDefinedStructureEditor::InitEditor(const EToolkitMode::Type Mode, cons
 	const TSharedRef<FTabManager::FLayout> StandaloneDefaultLayout = FTabManager::NewLayout( "Standalone_UserDefinedStructureEditor_Layout_v1" )
 	->AddArea
 	(
-		FTabManager::NewPrimaryArea() ->SetOrientation(Orient_Vertical)
+		FTabManager::NewPrimaryArea() 
+		->SetOrientation(Orient_Vertical)
 		->Split
 		(
 			FTabManager::NewStack()
@@ -297,6 +298,7 @@ void FUserDefinedStructureEditor::InitEditor(const EToolkitMode::Type Mode, cons
 			->Split
 			(
 				FTabManager::NewStack()
+				->SetHideTabWell( true )
 				->AddTab( MemberVariablesTabId, ETabState::OpenedTab )
 			)
 		)
@@ -360,7 +362,7 @@ TSharedRef<SDockTab> FUserDefinedStructureEditor::SpawnStructureTab(const FSpawn
 		EditedStruct = Cast<UUserDefinedStruct>(EditingObjs[ 0 ]);
 	}
 
-	auto Box = SNew(SHorizontalBox);
+	auto Box = SNew(SVerticalBox);
 
 	{
 		// Create a property view
@@ -373,6 +375,7 @@ TSharedRef<SDockTab> FUserDefinedStructureEditor::SpawnStructureTab(const FSpawn
 		PropertyView->SetObject(EditedStruct);
 		Box->AddSlot()
 		.VAlign(EVerticalAlignment::VAlign_Top)
+		.AutoHeight()
 		[
 			PropertyView.ToSharedRef()
 		];
@@ -391,6 +394,7 @@ TSharedRef<SDockTab> FUserDefinedStructureEditor::SpawnStructureTab(const FSpawn
 			Box->AddSlot()
 			.VAlign(EVerticalAlignment::VAlign_Top)
 			.Padding(2.0f, 0.0f, 0.0f, 0.0f)
+			.AutoHeight()
 			[
 				DefaultValueWidget.ToSharedRef()
 			];
