Copying //UE4/Dev-Blueprints to //UE4/Dev-Main (Source: //UE4/Dev-Blueprints @ 3057841)

diff --git a/Engine/Source/Developer/BlueprintProfiler/Public/BlueprintProfiler.h b/Engine/Source/Developer/BlueprintProfiler/Public/BlueprintProfiler.h
--- a/Engine/Source/Developer/BlueprintProfiler/Public/BlueprintProfiler.h
+++ b/Engine/Source/Developer/BlueprintProfiler/Public/BlueprintProfiler.h
@@ -28,18 +28,16 @@ public:
 	virtual bool IsProfilerEnabled() const override { return bProfilerActive; }
 	virtual bool IsProfilingCaptureActive() const { return bProfilingCaptureActive; }
 	virtual void ToggleProfilingCapture() override;
-	virtual void InstrumentEvent(const EScriptInstrumentationEvent& Event) override;
+	virtual void InstrumentEvent(const FScriptInstrumentationSignal& Event) override;
 
 #if WITH_EDITOR
+	virtual void AddInstrumentedBlueprint(UBlueprint* InstrumentedBlueprint) override;
 	virtual FOnBPStatGraphLayoutChanged& GetGraphLayoutChangedDelegate() { return GraphLayoutChangedDelegate; }
 	virtual TSharedPtr<FBlueprintExecutionContext> GetBlueprintContext(const FString& BlueprintClassPath) override;
 	virtual TSharedPtr<FScriptExecutionNode> GetProfilerDataForNode(const UEdGraphNode* GraphNode) override;
+	virtual TSharedPtr<FScriptExecutionBlueprint> GetProfilerDataForBlueprint(const UBlueprint* Blueprint) override;
 	virtual bool HasDataForInstance(const UObject* Instance) const override;
 	virtual void ProcessEventProfilingData() override;
-	virtual EBlueprintProfilerHeatMapDisplayMode::Type GetGraphNodeHeatMapDisplayMode() const { return GraphNodeHeatMapDisplayMode; }
-	virtual void SetGraphNodeHeatMapDisplayMode(EBlueprintProfilerHeatMapDisplayMode::Type InHeatMapDisplayMode) { GraphNodeHeatMapDisplayMode = InHeatMapDisplayMode; }
-	virtual EBlueprintProfilerHeatMapDisplayMode::Type GetWireHeatMapDisplayMode() const { return WireHeatMapDisplayMode; }
-	virtual void SetWireHeatMapDisplayMode(EBlueprintProfilerHeatMapDisplayMode::Type InHeatMapDisplayMode) { WireHeatMapDisplayMode = InHeatMapDisplayMode; }
 #endif
 	// End IBlueprintProfilerModule
 
@@ -84,10 +82,6 @@ protected:
 #if WITH_EDITOR
 	/** PIE Active */
 	bool bPIEActive;
-	/** Current graph node heat map display mode */
-	EBlueprintProfilerHeatMapDisplayMode::Type GraphNodeHeatMapDisplayMode;
-	/** Current wire heat map display mode */
-	EBlueprintProfilerHeatMapDisplayMode::Type WireHeatMapDisplayMode;
 	/** Suspended script events, the key is the latent LinkId */
 	TMap<FName, TMap<int32, TSharedPtr<FScriptEventPlayback>>> SuspendedEvents;
 	/** Cached object path and code offset lookup to UObjects */
diff --git a/Engine/Source/Developer/BlueprintProfiler/Public/BlueprintProfilerModule.h b/Engine/Source/Developer/BlueprintProfiler/Public/BlueprintProfilerModule.h
--- a/Engine/Source/Developer/BlueprintProfiler/Public/BlueprintProfilerModule.h
+++ b/Engine/Source/Developer/BlueprintProfiler/Public/BlueprintProfilerModule.h
@@ -4,26 +4,6 @@
 
 #include "ModuleManager.h"
 
-#if WITH_EDITOR
-
-/** Profiler heat map display mode */
-namespace EBlueprintProfilerHeatMapDisplayMode
-{
-	enum Type
-	{
-		None = 0,
-		Inclusive,
-		Exclusive,
-		MaxTiming,
-		Total,
-		HottestPath,
-		HottestEndpoint,
-		PinToPin
-	};
-}
-
-#endif
-
 /** Delegate to broadcast structural stats changes */
 DECLARE_MULTICAST_DELEGATE_OneParam(FOnBPStatGraphLayoutChanged, TWeakObjectPtr<UBlueprint>);
 
@@ -44,10 +24,13 @@ public:
 	virtual void ToggleProfilingCapture() {}
 
 	/** Instruments and add a profiling event to the data */
-	virtual void InstrumentEvent(const EScriptInstrumentationEvent& Event) {}
+	virtual void InstrumentEvent(const FScriptInstrumentationSignal& Event) {}
 
 #if WITH_EDITOR
 
+	/** Adds an instrumented blueprint directly */
+	virtual void AddInstrumentedBlueprint(UBlueprint* InstrumentedBlueprint) = 0;
+
 	/** Returns the multicast delegate to signal blueprint layout changes to stats */
 	virtual FOnBPStatGraphLayoutChanged& GetGraphLayoutChangedDelegate() = 0;
 
@@ -57,24 +40,15 @@ public:
 	/** Returns the current profiling event data for node */
 	virtual TSharedPtr<class FScriptExecutionNode> GetProfilerDataForNode(const UEdGraphNode* GraphNode) = 0;
 
+	/** Returns the current profiling event data for the given Blueprint */
+	virtual TSharedPtr<class FScriptExecutionBlueprint> GetProfilerDataForBlueprint(const UBlueprint* Blueprint) = 0;
+
 	/** Returns if the profiler currently retains any data for the specified instance */
 	virtual bool HasDataForInstance(const UObject* Instance) const { return false; }
 
 	/** Process profiling event data */
 	virtual void ProcessEventProfilingData() {}
 
-	/** Returns the current graph node heat map display mode setting */
-	virtual EBlueprintProfilerHeatMapDisplayMode::Type GetGraphNodeHeatMapDisplayMode() const = 0;
-
-	/** Sets the current graph node heat map display mode setting */
-	virtual void SetGraphNodeHeatMapDisplayMode(EBlueprintProfilerHeatMapDisplayMode::Type InHeatMapDisplayMode) {}
-
-	/** Returns the current wire heat map display mode setting */
-	virtual EBlueprintProfilerHeatMapDisplayMode::Type GetWireHeatMapDisplayMode() const = 0;
-
-	/** Sets the current wire heat map display mode setting */
-	virtual void SetWireHeatMapDisplayMode(EBlueprintProfilerHeatMapDisplayMode::Type InHeatMapDisplayMode) {}
-
 #endif // WITH_EDITOR
 
 };
\ No newline at end of file
diff --git a/Engine/Source/Developer/BlueprintProfiler/Public/ScriptInstrumentationCapture.h b/Engine/Source/Developer/BlueprintProfiler/Public/ScriptInstrumentationCapture.h
--- a/Engine/Source/Developer/BlueprintProfiler/Public/ScriptInstrumentationCapture.h
+++ b/Engine/Source/Developer/BlueprintProfiler/Public/ScriptInstrumentationCapture.h
@@ -13,18 +13,18 @@ public:
 
 	FScriptInstrumentedEvent() {}
 
-	bool operator == (const FScriptInstrumentedEvent& OtherIn) const;
-
-	FScriptInstrumentedEvent(EScriptInstrumentation::Type InEventType, const FString& InPathData)
+	FScriptInstrumentedEvent(EScriptInstrumentation::Type InEventType, const FName InFunctionClassScopeName, const FString& InPathData)
 		: EventType(InEventType)
+		, FunctionClassScopeName(InFunctionClassScopeName)
 		, PathData(InPathData)
 		, CodeOffset(-1)
 		, Time(FPlatformTime::Seconds())
 	{
 	}
 
-	FScriptInstrumentedEvent(EScriptInstrumentation::Type InEventType, const FName InFunctionName, const int32 InCodeOffset)
+	FScriptInstrumentedEvent(EScriptInstrumentation::Type InEventType, const FName InFunctionClassScopeName, const FName InFunctionName, const int32 InCodeOffset = INDEX_NONE)
 		: EventType(InEventType)
+		, FunctionClassScopeName(InFunctionClassScopeName)
 		, FunctionName(InFunctionName)
 		, CodeOffset(InCodeOffset)
 		, Time(FPlatformTime::Seconds())
@@ -41,11 +41,20 @@ public:
 	void OverrideType(const EScriptInstrumentation::Type NewType) { EventType = NewType; }
 
 	/** Returns the object path for the object that triggered this event */
-	const FString& GetObjectPath() const { return PathData; }
+	const FString& GetBlueprintClassPath() const { return PathData; }
+
+	/** Returns the object path for the object that triggered this event */
+	const FString& GetInstancePath() const { return PathData; }
 
-	/** Returns the function name the event occurred in */
+	/** Returns the function name the signal was emitted from */
 	FName GetFunctionName() const { return FunctionName; }
 
+	/** Returns the scoped function name signal was emitted from */
+	FName GetScopedFunctionName() const;
+
+	/** Returns the function scope name the signal was emitted from */
+	FName GetFunctionClassScopeName() const { return FunctionClassScopeName; }
+
 	/** Returns the script offset at the time this event was generated */
 	int32 GetScriptCodeOffset() const { return CodeOffset; }
 
@@ -80,9 +89,11 @@ protected:
 	
 	/** The event type */
 	EScriptInstrumentation::Type EventType;
+	/** The function class scope name, this indicates the owning class of the current function */
+	FName FunctionClassScopeName;
 	/** The owning function name */
 	FName FunctionName;
-	/** The object path for the object emitting the event */
+	/** The object / blueprint class path for the object emitting the event */
 	FString PathData;
 	/** The script code offset when this event was emitted */
 	int32 CodeOffset;
@@ -98,7 +109,7 @@ class FInstrumentationCaptureContext
 public:
 
 	/** Update the current execution context, and emit events to signal changes */
-	void UpdateContext(const UObject* InContextObject, TArray<FScriptInstrumentedEvent>& InstrumentationQueue);
+	void UpdateContext(const FScriptInstrumentationSignal& InstrumentSignal, TArray<FScriptInstrumentedEvent>& InstrumentationQueue);
 
 	/** Reset the current event context */
 	void ResetContext();
@@ -107,6 +118,10 @@ private:
 
 	/** The current class/blueprint context */
 	TWeakObjectPtr<const UClass> ContextClass;
+	/** The current function class scope */
+	TWeakObjectPtr<const UClass> ContextFunctionClassScope;
 	/** The current instance context */
 	TWeakObjectPtr<const UObject> ContextObject;
+	/** The current event context */
+	FName ActiveEvent;
 };
diff --git a/Engine/Source/Developer/BlueprintProfiler/Public/ScriptInstrumentationPlayback.h b/Engine/Source/Developer/BlueprintProfiler/Public/ScriptInstrumentationPlayback.h
--- a/Engine/Source/Developer/BlueprintProfiler/Public/ScriptInstrumentationPlayback.h
+++ b/Engine/Source/Developer/BlueprintProfiler/Public/ScriptInstrumentationPlayback.h
@@ -13,7 +13,7 @@ struct FBlueprintExecutionTrace
 	FTracePath TracePath;
 	FName InstanceName;
 	FName FunctionName;
-	FName NodeName;
+	TWeakPtr<FScriptExecutionNode> ProfilerNode;
 	int32 Offset;
 	double ObservationTime;
 };
@@ -79,15 +79,15 @@ public:
 	/** Remaps PIE actor instance paths to editor actor instances */
 	BLUEPRINTPROFILER_API FName RemapInstancePath(const FName InstanceName) const;
 
-	/** Returns the function name containing the event */
-	FName GetEventFunctionName(const FName EventName) const;
-
 	/** Returns the function context containing the event */
-	TSharedPtr<class FBlueprintFunctionContext> GetFunctionContextForEvent(const FName EventName) const;
+	TSharedPtr<class FBlueprintFunctionContext> GetFunctionContextForEventChecked(const FName ScopedEventName) const;
 
 	/** Returns the function context matching the supplied name providing it exists */
 	TSharedPtr<class FBlueprintFunctionContext> GetFunctionContext(const FName FunctionNameIn) const;
 
+	/** Returns the function context matching the supplied graph providing it exists */
+	TSharedPtr<FBlueprintFunctionContext> GetFunctionContextFromGraph(const UEdGraph* Graph) const;
+
 	/** Adds a manually created function context */
 	void AddFunctionContext(const FName FunctionName, TSharedPtr<class FBlueprintFunctionContext> FunctionContext)
 	{
@@ -109,6 +109,12 @@ public:
 	/** Updates cascading stats for the blueprint, this merges instance data upwards to the blueprint */
 	void UpdateConnectedStats();
 
+	/** Scoped function name extraction that results in ClassName::FunctionName for lookups */
+	FName GetScopedFunctionNameFromGraph(const UEdGraph* Graph) const;
+
+	/** Basic function name creation that allows for UFunction lookups in the correct class */
+	FName GetFunctionNameFromGraph(const UEdGraph* Graph) const;
+
 private:
 
 	/** Creates and returns a new function context, returns existing if present  */
@@ -128,8 +134,6 @@ private:
 	TWeakObjectPtr<UBlueprint> Blueprint;
 	/** Blueprint Generated class for the context */
 	TWeakObjectPtr<UBlueprintGeneratedClass> BlueprintClass;
-	/** The ubergraph function name for the blueprint generated class */
-	FName UbergraphFunctionName;
 	/** Root level execution node for this blueprint */
 	TSharedPtr<FScriptExecutionBlueprint> BlueprintNode;
 	/** UFunction contexts for this blueprint */
@@ -167,11 +171,20 @@ public:
 	/** Returns if the function context is fully formed */
 	bool IsContextValid() const { return Function.IsValid() && BlueprintClass.IsValid(); }
 
+	/** Returns if the function context is inherited */
+	bool IsInheritedContext() const { return bIsInheritedContext; }
+
 	/** Returns the graph node at the specified script offset */
 	const UEdGraphNode* GetNodeFromCodeLocation(const int32 ScriptOffset);
 
-	/** Returns the tunnel source node at the specified script offset */
-	const UEdGraphNode* FindTunnelSourceNodeFromCodeLocation(const int32 ScriptOffset);
+	/** Returns the tunnel node at the specified script offset if present */
+	const UEdGraphNode* GetTunnelNodeFromCodeLocation(const int32 ScriptOffset);
+
+	/** Returns the tunnel node at associated with the graph node */
+	const UEdGraphNode* GetTunnelNodeFromGraphNode(const UEdGraphNode* TrueSourceNode) const;
+
+	/** Returns if the specified node is inside a tunnel */
+	bool IsNodeFromTunnelGraph(const UEdGraphNode* Node) const;
 
 	/** Returns the pin at the specified script offset */
 	const UEdGraphPin* GetPinFromCodeLocation(const int32 ScriptOffset);
@@ -185,17 +198,23 @@ public:
 	/** Returns the pure node script code range for the specified node */
 	FInt32Range GetPureNodeScriptCodeRange(const UEdGraphNode* Node) const;
 
-	/** Returns true if this function context contains an execution node matching the node name */
+	/** Returns true if this function context contains an execution node matching the node name, please note this fails in tunnel contexts */
 	bool HasProfilerDataForNode(const FName NodeName) const;
 
+	/** Returns true if this or any child function context contains an entry for the graphnode */
+	bool HasProfilerDataForGraphNode(const UEdGraphNode* GraphNode) const;
+
 	/** Returns the execution node representing the node name */
-	TSharedPtr<FScriptExecutionNode> GetProfilerDataForNode(const FName NodeName);
+	TSharedPtr<FScriptExecutionNode> GetProfilerDataForNode(const FName NodeName) const;
 
 	/** Returns the execution node representing the node name, asserting on fail */
-	TSharedPtr<FScriptExecutionNode> GetProfilerDataForNodeChecked(const FName NodeName);
+	TSharedPtr<FScriptExecutionNode> GetProfilerDataForNodeChecked(const FName NodeName) const;
 
 	/** Returns the execution node representing the graph node */
-	TSharedPtr<FScriptExecutionNode> GetProfilerDataForGraphNode(const UEdGraphNode* Node);
+	TSharedPtr<FScriptExecutionNode> GetProfilerDataForGraphPin(const UEdGraphPin* Pin) const;
+
+	/** Returns the execution node representing the graph node */
+	TSharedPtr<FScriptExecutionNode> GetProfilerDataForGraphNode(const UEdGraphNode* Node) const;
 
 	/** Returns the execution node representing the node name */
 	template<typename ExecNodeType> TSharedPtr<ExecNodeType> GetTypedProfilerDataForNode(const FName NodeName);
@@ -203,12 +222,21 @@ public:
 	/** Returns the execution node representing the object */
 	template<typename ExecNodeType> TSharedPtr<ExecNodeType> GetTypedProfilerDataForGraphNode(const UEdGraphNode* Node);
 
+	/** Returns true if the execution node and function context can be located and initialised from the script offset */
+	virtual bool GetProfilerContextFromScriptOffset(const int32 ScriptOffset, TSharedPtr<FScriptExecutionNode>& ExecNode, TSharedPtr<FBlueprintFunctionContext>& FunctionContext);
+
 	/** Adds all function entry points to node as children */
 	void AddCallSiteEntryPointsToNode(TSharedPtr<FScriptExecutionNode> CallingNode) const;
 
 	/** Returns the function name for this context */
 	FName GetFunctionName() const { return FunctionName; }
 
+	/** Returns the graph name for this context */
+	FName GetGraphName() const { return GraphName; }
+
+	/** Returns the scoped event name for this context */
+	FName GetScopedEventName(const FName EventName) const;
+
 	/** Returns the UFunction for this context */
 	TWeakObjectPtr<UFunction> GetUFunction() { return Function; }
 
@@ -227,11 +255,17 @@ public:
 	/** Looks for matching pin in the supplied node, handles pins with no name. */
 	static UEdGraphPin* FindMatchingPin(const UEdGraphNode* NodeToSearch, const UEdGraphPin* PinToFind, const bool bIgnoreDirection = false);
 
-	/** Utility function that returns the tunnel graph from the supplied node */
-	static UEdGraph* GetTunnelGraphFromNode(const UEdGraphNode* TunnelNode);
+	/** Utility function that returns the graph from the supplied node, handling tunnel boundaries to return the inner graph */
+	static UEdGraph* GetGraphFromNode(const UEdGraphNode* GraphNode, const bool bAllowNonTunnel = true);
+
+	/** Utility function that returns the tunnel instance function name */
+	static FName GetTunnelInstanceFunctionName(const UEdGraphNode* GraphNode);
 
-	/** Utility function that returns if the tunnel node is internal to a tunnel graph */
-	static bool IsTunnelInternal(const UEdGraphNode* TunnelNode);
+	/** Utility to get the scoped function name from the graph node */
+	FName GetScopedFunctionNameFromNode(const UEdGraphNode* GraphNode) const;
+
+	/** Utility function that returns customizations for a node (e.g. Node Color) */
+	void GetNodeCustomizations(FScriptExecNodeParams& ParamsInOut) const;
 
 protected:
 
@@ -258,11 +292,17 @@ protected:
 
 	// --Execution mapping functionality
 
+	/** Examines graphs nodes and determines some characterists to aid in creation of the matching executable node */
+	void DetermineGraphNodeCharacteristics(const UEdGraphNode* GraphNode, TArray<UEdGraphPin*>& InputPins, TArray<UEdGraphPin*>& ExecPins, FScriptExecNodeParams& NodeParams);
+
 	/** Maps each blueprint node following execution wires */
 	virtual TSharedPtr<FScriptExecutionNode> MapNodeExecution(UEdGraphNode* NodeToMap);
 
+	/** Returns the pure chain root node for the argument execution node, creating one if not existing */
+	TSharedPtr<FScriptExecutionNode> FindOrCreatePureChainRoot(TSharedPtr<FScriptExecutionNode> ExecNode);
+
 	/** Maps input pin execution */
-	void MapInputPins(TSharedPtr<FScriptExecutionNode> ExecNode, const TArray<UEdGraphPin*>& Pins);
+	virtual void MapInputPins(TSharedPtr<FScriptExecutionNode> ExecNode, const TArray<UEdGraphPin*>& Pins);
 
 	/** Maps pin execution */
 	void MapExecPins(TSharedPtr<FScriptExecutionNode> ExecNode, const TArray<UEdGraphPin*>& Pins);
@@ -272,6 +312,9 @@ protected:
 	/** Maps tunnel boundries */
 	virtual TSharedPtr<FScriptExecutionNode> MapTunnelBoundary(const UEdGraphPin* TunnelPin);
 
+	/** Get tunnel boundary node */
+	TSharedPtr<FScriptExecutionNode> GetTunnelBoundaryNodeChecked(const UEdGraphPin* TunnelPin);
+
 	/** Maps the tunnel point into the instanced graph, creating the instanced graph if not already existing. */
 	void MapTunnelInstance(UK2Node_Tunnel* TunnelInstance);
 
@@ -284,6 +327,8 @@ protected:
 	FName GraphName;
 	/** The function name */
 	FName FunctionName;
+	/** Indicates this is an inherited function context */
+	bool bIsInheritedContext;
 	/** The Owning Blueprint context */
 	TWeakPtr<FBlueprintExecutionContext> BlueprintContext;
 	/** The UFunction this context represents */
@@ -294,6 +339,8 @@ protected:
 	TWeakObjectPtr<UBlueprintGeneratedClass> BlueprintClass;
 	/** ScriptCode offset to node cache */
 	TMap<int32, TWeakObjectPtr<const UEdGraphNode>> ScriptOffsetToNodes;
+	/** ScriptCode offset to tunnel node cache */
+	TMap<TWeakObjectPtr<const UEdGraphNode>, TWeakObjectPtr<const UEdGraphNode>> NodeToTunnelNode;
 	/** ScriptCode offset to pin cache */
 	TMap<int32, FEdGraphPinReference> ScriptOffsetToPins;
 	/** Execution nodes containing profiling data */
@@ -316,25 +363,32 @@ public:
 	// ~FBlueprintFunctionContext Begin
 	virtual void InitialiseContextFromGraph(TSharedPtr<FBlueprintExecutionContext> BlueprintContext, const FName TunnelInstanceName, UEdGraph* TunnelGraph);
 	virtual void MapFunction() override {}
-	virtual TSharedPtr<FScriptExecutionNode> CreateExecutionNode(FScriptExecNodeParams& InitParams) override;
+	virtual bool GetProfilerContextFromScriptOffset(const int32 ScriptOffset, TSharedPtr<FScriptExecutionNode>& ExecNode, TSharedPtr<FBlueprintFunctionContext>& FunctionContext) override;
 	// ~FBlueprintFunctionContext End
 
-	/** Maps a tunnel function context. */
+	/** Maps a tunnel function context */
 	void MapTunnelContext(TSharedPtr<FBlueprintFunctionContext> CallingFunctionContext, UK2Node_Tunnel* TunnelInstance);
 
+	/** Get pure node from the tunnel pin */
+	TSharedPtr<FScriptExecutionNode> GetPureNodeFromTunnelPin(const UEdGraphPin* TunnelPin) const;
+
 private:
 
 	friend FBlueprintExecutionContext;
 
 	// ~FBlueprintFunctionContext Begin
 	virtual TSharedPtr<FScriptExecutionNode> MapNodeExecution(UEdGraphNode* NodeToMap) override;
+	virtual void MapInputPins(TSharedPtr<FScriptExecutionNode> ExecNode, const TArray<UEdGraphPin*>& Pins) override;
 	virtual TSharedPtr<FScriptExecutionNode> MapTunnelBoundary(const UEdGraphPin* TunnelPin) override;
 	// ~FBlueprintFunctionContext End
 	
 	/** Maps input and output pins in a tunnel graph */
 	void MapTunnelIO();
 
-	/** Returns true if the pin is part of this tunnel */
+	/** Returns true if the node is an internal part of this tunnel */
+	bool IsTunnelNodeInternal(const UEdGraphNode* TunnelNode);
+
+	/** Returns true if the pin is part of this tunnel, unless its a tunnel site to another tunnel */
 	bool IsPinFromThisTunnel(const UEdGraphPin* TunnelPin) const;
 
 	/** Discovers already mapped exit sites */
@@ -344,10 +398,8 @@ private:
 
 	/** The tunnel instance this context represents */
 	TWeakObjectPtr<UK2Node_Tunnel> TunnelInstanceNode;
-	/** The non-instance context of the tunnel graph */
-	TSharedPtr<FBlueprintFunctionContext> TunnelFunctionContext;
-	/** External tunnel nodes */
-	TMap<FName, TSharedPtr<FScriptExecutionNode>> ExternalNodes;
+	/** The tunnel pure pin link nodes */
+	TMap<FEdGraphPinReference, TSharedPtr<FScriptExecutionNode>> PureLinkNodes;
 	/** Staging entry point */
 	TSharedPtr<FScriptExecutionTunnelEntry> StagingEntryPoint;
 
@@ -362,12 +414,15 @@ public:
 
 	struct NodeSignalHelper
 	{
+		/** Easy validation */
+		bool IsValid() const { return ImpureNode.IsValid() && FunctionContext.IsValid(); }
+
+		/** Node signal properties */
 		TSharedPtr<FScriptExecutionNode> ImpureNode;
-		TSharedPtr<FBlueprintExecutionContext> BlueprintContext;
 		TSharedPtr<FBlueprintFunctionContext> FunctionContext;
-		TArray<class FScriptInstrumentedEvent> ExclusiveEvents;
+		TArray<class FScriptInstrumentedEvent> AverageEvents;
 		TArray<class FScriptInstrumentedEvent> InclusiveEvents;
-		TArray<FTracePath> ExclusiveTracePaths;
+		TArray<FTracePath> AverageTracePaths;
 		TArray<FTracePath> InputTracePaths;
 		TArray<FTracePath> InclusiveTracePaths;
 	};
@@ -420,7 +475,7 @@ private:
 	void ProcessExecutionSequence(NodeSignalHelper& CurrentNodeData, const FScriptInstrumentedEvent& CurrSignal);
 
 	/** Add to trace history */
-	void AddToTraceHistory(const FName NodeName, const FScriptInstrumentedEvent& TraceSignal);
+	void AddToTraceHistory(const TSharedPtr<FScriptExecutionNode> ProfilerNode, const FScriptInstrumentedEvent& TraceSignal);
 
 private:
 
diff --git a/Engine/Source/Editor/AnimGraph/Private/AnimGraphNode_BlendListByBool.cpp b/Engine/Source/Editor/AnimGraph/Private/AnimGraphNode_BlendListByBool.cpp
--- a/Engine/Source/Editor/AnimGraph/Private/AnimGraphNode_BlendListByBool.cpp
+++ b/Engine/Source/Editor/AnimGraph/Private/AnimGraphNode_BlendListByBool.cpp
@@ -14,6 +14,8 @@
 UAnimGraphNode_BlendListByBool::UAnimGraphNode_BlendListByBool(const FObjectInitializer& ObjectInitializer)
 	: Super(ObjectInitializer)
 {
+	Node.AddPose();
+	Node.AddPose();
 }
 
 FText UAnimGraphNode_BlendListByBool::GetNodeTitle(ENodeTitleType::Type TitleType) const
diff --git a/Engine/Source/Editor/BlueprintGraph/Classes/K2Node_TunnelBoundary.h b/Engine/Source/Editor/BlueprintGraph/Classes/K2Node_TunnelBoundary.h
--- a/Engine/Source/Editor/BlueprintGraph/Classes/K2Node_TunnelBoundary.h
+++ b/Engine/Source/Editor/BlueprintGraph/Classes/K2Node_TunnelBoundary.h
@@ -47,4 +47,9 @@ protected:
 	/** Determines the tunnel graph name */
 	void DetermineTunnelGraphName(UK2Node_Tunnel* TunnelInstance);
 
+	/** Build a guid map for nodes from the original graph */
+	static void BuildSourceNodeMap(UEdGraphNode* Tunnel, TMap<FGuid, const UEdGraphNode*>& SourceNodeMap);
+
+	/** Determines the true source tunnel instance */
+	static UEdGraphNode* FindTrueSourceTunnelInstance(UEdGraphNode* Tunnel, UEdGraphNode* SourceTunnelInstance);
 };
diff --git a/Engine/Source/Editor/BlueprintGraph/Private/EdGraphSchema_K2.cpp b/Engine/Source/Editor/BlueprintGraph/Private/EdGraphSchema_K2.cpp
--- a/Engine/Source/Editor/BlueprintGraph/Private/EdGraphSchema_K2.cpp
+++ b/Engine/Source/Editor/BlueprintGraph/Private/EdGraphSchema_K2.cpp
@@ -3520,6 +3520,9 @@ void UEdGraphSchema_K2::GetVariableTypeTree(TArray< TSharedPtr<FPinTypeTreeInfo>
 		else
 		{
 			TypeTree.Add(MakeShareable(new FPinTypeTreeInfo(GetCategoryText(PC_Object, true), PC_Object, this, LOCTEXT("ObjectType", "Object pointer."), true, TypesDatabasePtr)));
+			TypeTree.Add(MakeShareable(new FPinTypeTreeInfo(GetCategoryText(PC_Class, true), PC_Class, this, LOCTEXT("ClassType", "Class pointer."), true, TypesDatabasePtr)));
+			TypeTree.Add(MakeShareable(new FPinTypeTreeInfo(GetCategoryText(PC_AssetClass, true), PC_AssetClass, this, LOCTEXT("AssetClassType", "Class ID."), true, TypesDatabasePtr)));
+			TypeTree.Add(MakeShareable(new FPinTypeTreeInfo(GetCategoryText(PC_Asset, true), PC_Asset, this, LOCTEXT("AssetType", "Asset ID."), true, TypesDatabasePtr)));
 		}
 	}
 	TypeTree.Add( MakeShareable( new FPinTypeTreeInfo(GetCategoryText(PC_Enum, true), PC_Enum, this, LOCTEXT("EnumType", "Enumeration types."), true, TypesDatabasePtr) ) );
diff --git a/Engine/Source/Editor/BlueprintGraph/Private/K2Node_CustomEvent.cpp b/Engine/Source/Editor/BlueprintGraph/Private/K2Node_CustomEvent.cpp
--- a/Engine/Source/Editor/BlueprintGraph/Private/K2Node_CustomEvent.cpp
+++ b/Engine/Source/Editor/BlueprintGraph/Private/K2Node_CustomEvent.cpp
@@ -133,15 +133,21 @@ void UK2Node_CustomEvent::Serialize(FArchive& Ar)
 			if (Pin)
 			{
 				if (Pin->Direction == EGPD_Output
-					&& !Pin->PinType.bIsConst
 					&& !K2Schema->IsExecPin(*Pin)
 					&& !K2Schema->IsDelegateCategory(Pin->PinType.PinCategory))
 				{
-					for (auto PinInfo : UserDefinedPins)
+					for (TSharedPtr<FUserPinInfo>& PinInfo : UserDefinedPins)
 					{
 						if (PinInfo->PinName == Pin->PinName)
 						{
-							Pin->PinType.bIsConst = PinInfo->PinType.bIsConst = PinInfo->PinType.bIsArray || PinInfo->PinType.bIsReference;
+							if (!Pin->PinType.bIsConst)
+							{
+								Pin->PinType.bIsConst = PinInfo->PinType.bIsConst = PinInfo->PinType.bIsArray || PinInfo->PinType.bIsReference;
+							}
+							
+							// Some legacy blueprints have incorrectly serialized the pin direction to EGPD_Input
+							// even though CustomEvent only supports Output pins. This works around that issue.
+							PinInfo->DesiredPinDirection = Pin->Direction;
 							break;
 						}
 					}
diff --git a/Engine/Source/Editor/BlueprintGraph/Private/K2Node_EditablePinBase.cpp b/Engine/Source/Editor/BlueprintGraph/Private/K2Node_EditablePinBase.cpp
--- a/Engine/Source/Editor/BlueprintGraph/Private/K2Node_EditablePinBase.cpp
+++ b/Engine/Source/Editor/BlueprintGraph/Private/K2Node_EditablePinBase.cpp
@@ -15,8 +15,8 @@ void UK2Node_EditablePinBase::AllocateDefaultPins()
 	// Add in pins based on the user defined pins in this node
 	for(int32 i = 0; i < UserDefinedPins.Num(); i++)
 	{
-		//FText DummyErrorMsg;
-		//if (!IsEditable() || CanCreateUserDefinedPin(UserDefinedPins[i]->PinType, UserDefinedPins[i]->DesiredPinDirection, DummyErrorMsg))
+		FText DummyErrorMsg;
+		if (!IsEditable() || CanCreateUserDefinedPin(UserDefinedPins[i]->PinType, UserDefinedPins[i]->DesiredPinDirection, DummyErrorMsg))
 		{
 			CreatePinFromUserDefinition(UserDefinedPins[i]);
 		}
diff --git a/Engine/Source/Editor/BlueprintGraph/Private/K2Node_TunnelBoundary.cpp b/Engine/Source/Editor/BlueprintGraph/Private/K2Node_TunnelBoundary.cpp
--- a/Engine/Source/Editor/BlueprintGraph/Private/K2Node_TunnelBoundary.cpp
+++ b/Engine/Source/Editor/BlueprintGraph/Private/K2Node_TunnelBoundary.cpp
@@ -84,12 +84,16 @@ void UK2Node_TunnelBoundary::CreateBoundaryNodesForTunnelInstance(UK2Node_Tunnel
 			}
 		}
 	}
+	// Find the blueprint tunnel instance.
+	UEdGraphNode* SourceTunnelInstance = Cast<UEdGraphNode>(MessageLog.FindSourceObject(TunnelInstance));
+	SourceTunnelInstance = FindTrueSourceTunnelInstance(TunnelInstance, SourceTunnelInstance);
 	// Create the Boundary nodes for each unique entry/exit site.
 	for (UEdGraphPin* EntryPin : TunnelEntryPins)
 	{
 		UK2Node_TunnelBoundary* TunnelBoundaryNode = TunnelGraph->CreateBlankNode<UK2Node_TunnelBoundary>();
 		MessageLog.NotifyIntermediateObjectCreation(TunnelBoundaryNode, TunnelInstance);
 		TunnelBoundaryNode->CreateNewGuid();
+		MessageLog.IntermediateTunnelNodeToSourceNodeMap.Add(TunnelBoundaryNode, SourceTunnelInstance);
 		TunnelBoundaryNode->WireUpTunnelEntry(TunnelInstance, EntryPin, MessageLog);
 	}
 	for (UEdGraphPin* ExitPin : TunnelExitPins)
@@ -97,8 +101,26 @@ void UK2Node_TunnelBoundary::CreateBoundaryNodesForTunnelInstance(UK2Node_Tunnel
 		UK2Node_TunnelBoundary* TunnelBoundaryNode = TunnelGraph->CreateBlankNode<UK2Node_TunnelBoundary>();
 		MessageLog.NotifyIntermediateObjectCreation(TunnelBoundaryNode, TunnelInstance);
 		TunnelBoundaryNode->CreateNewGuid();
+		MessageLog.IntermediateTunnelNodeToSourceNodeMap.Add(TunnelBoundaryNode, SourceTunnelInstance);
 		TunnelBoundaryNode->WireUpTunnelExit(TunnelInstance, ExitPin, MessageLog);
 	}
+	// Build node guid map to locate true source node.
+	TMap<FGuid, const UEdGraphNode*> TrueSourceNodeMap;
+	BuildSourceNodeMap(SourceTunnelInstance, TrueSourceNodeMap);
+	// Register all the nodes to the tunnel instance map
+	if (SourceTunnelInstance)
+	{
+		for (auto Node : TunnelGraph->Nodes)
+		{
+			MessageLog.SourceNodeToTunnelInstanceNodeMap.Add(Node, SourceTunnelInstance);
+			if (const UEdGraphNode** SearchNode = TrueSourceNodeMap.Find(Node->NodeGuid))
+			{
+				const UEdGraphNode* SourceGraphNode = *SearchNode;
+				MessageLog.IntermediateTunnelNodeToSourceNodeMap.Add(Node, SourceGraphNode);
+			}
+		}
+		MessageLog.SourceNodeToTunnelInstanceNodeMap.Add(TunnelInstance, SourceTunnelInstance);
+	}
 }
 
 bool UK2Node_TunnelBoundary::IsPureTunnel(UK2Node_Tunnel* Tunnel)
@@ -118,7 +140,6 @@ void UK2Node_TunnelBoundary::WireUpTunnelEntry(UK2Node_Tunnel* TunnelInstance, U
 		UEdGraphPin* SourcePin = TunnelInstance->FindPin(TunnelPin->PinName);
 		// Wire in the node
 		UEdGraphPin* OutputPin = CreatePin(EGPD_Output, UEdGraphSchema_K2::PC_Exec, TEXT(""), NULL, false, false, TEXT("TunnelEntryExec"));
-		MessageLog.NotifyIntermediatePinCreation(OutputPin, SourcePin);
 		for (auto LinkedPin : TunnelPin->LinkedTo)
 		{
 			LinkedPin->MakeLinkTo(OutputPin);
@@ -148,7 +169,6 @@ void UK2Node_TunnelBoundary::WireUpTunnelExit(UK2Node_Tunnel* TunnelInstance, UE
 			LinkedPin->MakeLinkTo(InputPin);
 		}
 		UEdGraphPin* OutputPin = CreatePin(EGPD_Output, UEdGraphSchema_K2::PC_Exec, TEXT(""), NULL, false, false, TEXT("TunnelExitExec"));
-		MessageLog.NotifyIntermediatePinCreation(OutputPin, SourcePin);
 		TunnelPin->BreakAllPinLinks();
 		TunnelPin->MakeLinkTo(OutputPin);
 	}
@@ -168,4 +188,68 @@ void UK2Node_TunnelBoundary::DetermineTunnelGraphName(UK2Node_Tunnel* TunnelInst
 	TunnelGraphName = TunnelGraph ? TunnelGraph->GetFName() : NAME_None;
 }
 
+void UK2Node_TunnelBoundary::BuildSourceNodeMap(UEdGraphNode* Tunnel, TMap<FGuid, const UEdGraphNode*>& SourceNodeMap)
+{
+	// Crawl the graph and locate source node guids
+	UEdGraph* TunnelGraph = nullptr;
+	if (UK2Node_MacroInstance* SourceMacroInstance = Cast<UK2Node_MacroInstance>(Tunnel))
+	{
+		TunnelGraph = SourceMacroInstance->GetMacroGraph();
+	}
+	else if (UK2Node_Composite* SourceCompositeInstance = Cast<UK2Node_Composite>(Tunnel))
+	{
+		TunnelGraph = SourceCompositeInstance->BoundGraph;
+	}
+	if (TunnelGraph)
+	{
+		for (auto GraphNode : TunnelGraph->Nodes)
+		{
+			SourceNodeMap.Add(GraphNode->NodeGuid) = GraphNode;
+		}
+	}
+}
+
+UEdGraphNode* UK2Node_TunnelBoundary::FindTrueSourceTunnelInstance(UEdGraphNode* Tunnel, UEdGraphNode* SourceTunnelInstance)
+{
+	UEdGraphNode* SourceNode = nullptr;
+	if (Tunnel && SourceTunnelInstance)
+	{
+		if (Tunnel->NodeGuid == SourceTunnelInstance->NodeGuid)
+		{
+			SourceNode = SourceTunnelInstance;
+		}
+		else
+		{
+			UEdGraph* TunnelGraph = nullptr;
+			if (UK2Node_MacroInstance* SourceMacroInstance = Cast<UK2Node_MacroInstance>(SourceTunnelInstance))
+			{
+				TunnelGraph = SourceMacroInstance->GetMacroGraph();
+			}
+			else if (UK2Node_Composite* SourceCompositeInstance = Cast<UK2Node_Composite>(SourceTunnelInstance))
+			{
+				TunnelGraph = SourceCompositeInstance->BoundGraph;
+			}
+			if (TunnelGraph)
+			{
+				for (auto GraphNode : TunnelGraph->Nodes)
+				{
+					if (GraphNode->NodeGuid == Tunnel->NodeGuid)
+					{
+						SourceNode = GraphNode;
+					}
+					else if (GraphNode->IsA<UK2Node_Composite>() || GraphNode->IsA<UK2Node_MacroInstance>())
+					{
+						SourceNode = FindTrueSourceTunnelInstance(Tunnel, GraphNode);
+					}
+					if (SourceNode)
+					{
+						break;
+					}
+				}
+			}
+		}
+	}
+	return SourceNode;
+}
+
 #undef LOCTEXT_NAMESPACE
diff --git a/Engine/Source/Editor/EditorStyle/Private/SlateEditorStyle.cpp b/Engine/Source/Editor/EditorStyle/Private/SlateEditorStyle.cpp
--- a/Engine/Source/Editor/EditorStyle/Private/SlateEditorStyle.cpp
+++ b/Engine/Source/Editor/EditorStyle/Private/SlateEditorStyle.cpp
@@ -4835,6 +4835,7 @@ void FSlateEditorStyle::FStyle::SetupLevelEditorStyle()
 	{
 		// Profiler stat window
 		Set( "BlueprintProfiler.ViewToolBar", new BOX_BRUSH("/Docking/Tab_Foreground", 4/16.0f));
+		Set( "BlueprintProfiler.ViewContent", new BOX_BRUSH("/Docking/AppTabContentArea", 4/16.0f));
 		Set( "BlueprintProfiler.BPIcon_Normal", new IMAGE_BRUSH("Icons/AssetIcons/Blueprint_16x", Icon16x16, FLinearColor( 0.46f, 0.54f, 0.81f )));
 		Set( "BlueprintProfiler.BPIcon_Macro", new IMAGE_BRUSH("Icons/AssetIcons/BlueprintMacroLibrary_16x", Icon16x16, FLinearColor( 0.46f, 0.54f, 0.81f )));
 		Set( "BlueprintProfiler.BPIcon_Interface", new IMAGE_BRUSH("Icons/AssetIcons/BlueprintInterface_16x", Icon16x16, FLinearColor( 0.46f, 0.54f, 0.81f )));
@@ -4862,6 +4863,46 @@ void FSlateEditorStyle::FStyle::SetupLevelEditorStyle()
 			.SetShadowOffset(FVector2D::ZeroVector)
 		);
 
+		Set("BlueprintProfiler.ToggleButton.Start", FCheckBoxStyle()
+			.SetCheckBoxType(ESlateCheckBoxType::ToggleButton)
+			.SetUncheckedImage(BOX_BRUSH("Common/SmallRoundedToggleLeft", FMargin(7.f / 16.f), FLinearColor(0.09f, 0.09f, 0.09f)))
+			.SetUncheckedPressedImage(BOX_BRUSH("Common/SmallRoundedToggleLeft", FMargin(7.f / 16.f), SelectionColor_Pressed))
+			.SetUncheckedHoveredImage(BOX_BRUSH("Common/SmallRoundedToggleLeft", FMargin(7.f / 16.f), SelectionColor_Pressed))
+			.SetCheckedHoveredImage(BOX_BRUSH("Common/SmallRoundedToggleLeft", FMargin(7.f / 16.f), SelectionColor))
+			.SetCheckedPressedImage(BOX_BRUSH("Common/SmallRoundedToggleLeft", FMargin(7.f / 16.f), SelectionColor))
+			.SetCheckedImage(BOX_BRUSH("Common/SmallRoundedToggleLeft", FMargin(7.f / 16.f), SelectionColor)));
+
+		Set("BlueprintProfiler.ToggleButton.Middle", FCheckBoxStyle()
+			.SetCheckBoxType(ESlateCheckBoxType::ToggleButton)
+			.SetUncheckedImage(BOX_BRUSH("Common/SmallRoundedToggleCenter", FMargin(7.f / 16.f), FLinearColor(0.09f, 0.09f, 0.09f)))
+			.SetUncheckedPressedImage(BOX_BRUSH("Common/SmallRoundedToggleCenter", FMargin(7.f / 16.f), SelectionColor_Pressed))
+			.SetUncheckedHoveredImage(BOX_BRUSH("Common/SmallRoundedToggleCenter", FMargin(7.f / 16.f), SelectionColor_Pressed))
+			.SetCheckedHoveredImage(BOX_BRUSH("Common/SmallRoundedToggleCenter", FMargin(7.f / 16.f), SelectionColor))
+			.SetCheckedPressedImage(BOX_BRUSH("Common/SmallRoundedToggleCenter", FMargin(7.f / 16.f), SelectionColor))
+			.SetCheckedImage(BOX_BRUSH("Common/SmallRoundedToggleCenter", FMargin(7.f / 16.f), SelectionColor)));
+
+		Set("BlueprintProfiler.ToggleButton.End", FCheckBoxStyle()
+			.SetCheckBoxType(ESlateCheckBoxType::ToggleButton)
+			.SetUncheckedImage(BOX_BRUSH("Common/SmallRoundedToggleRight", FMargin(7.f / 16.f), FLinearColor(0.09f, 0.09f, 0.09f)))
+			.SetUncheckedPressedImage(BOX_BRUSH("Common/SmallRoundedToggleRight", FMargin(7.f / 16.f), SelectionColor_Pressed))
+			.SetUncheckedHoveredImage(BOX_BRUSH("Common/SmallRoundedToggleRight", FMargin(7.f / 16.f), SelectionColor_Pressed))
+			.SetCheckedHoveredImage(BOX_BRUSH("Common/SmallRoundedToggleRight", FMargin(7.f / 16.f), SelectionColor))
+			.SetCheckedPressedImage(BOX_BRUSH("Common/SmallRoundedToggleRight", FMargin(7.f / 16.f), SelectionColor))
+			.SetCheckedImage(BOX_BRUSH("Common/SmallRoundedToggleRight", FMargin(7.f / 16.f), SelectionColor)));
+
+		Set("BlueprintProfiler.TableView.Row", FTableRowStyle(NormalTableRowStyle)
+			.SetActiveBrush(IMAGE_BRUSH("Common/Selection", Icon8x8, FLinearColor(0.728f, 0.364f, 0.003f, 0.15f)))
+			.SetActiveHoveredBrush(IMAGE_BRUSH("Common/Selection", Icon8x8, FLinearColor(0.728f, 0.364f, 0.003f, 0.2f)))
+			.SetInactiveBrush(IMAGE_BRUSH("Common/Selection", Icon8x8, FLinearColor(0.25f, 0.25f, 0.25f, 0.15f)))
+			.SetInactiveHoveredBrush(IMAGE_BRUSH("Common/Selection", Icon8x8, FLinearColor(0.25f, 0.25f, 0.25f, 0.2f))));
+
+		FTextBlockStyle NodeHyperlinkText = FTextBlockStyle(NormalText)
+			.SetColorAndOpacity(FLinearColor(1.0f, 1.0f, 1.0f))
+			.SetHighlightColor(FLinearColor(1.0f, 1.0f, 1.0f));
+
+		FHyperlinkStyle NodeHyperlinkStyle = GetWidgetStyle<FHyperlinkStyle>("HoverOnlyHyperlink");
+		Set("BlueprintProfiler.Node.Hyperlink", NodeHyperlinkStyle
+			.SetTextStyle(NodeHyperlinkText));
 	}
 #endif // WITH_EDITOR
 
diff --git a/Engine/Source/Editor/GraphEditor/Private/KismetNodes/SGraphNodeK2Base.cpp b/Engine/Source/Editor/GraphEditor/Private/KismetNodes/SGraphNodeK2Base.cpp
--- a/Engine/Source/Editor/GraphEditor/Private/KismetNodes/SGraphNodeK2Base.cpp
+++ b/Engine/Source/Editor/GraphEditor/Private/KismetNodes/SGraphNodeK2Base.cpp
@@ -12,6 +12,7 @@
 
 // Blueprint Profiler
 #include "Developer/BlueprintProfiler/Public/BlueprintProfilerModule.h"
+#include "Editor/Kismet/Public/Profiler/BlueprintProfilerSettings.h"
 #include "Editor/Kismet/Public/Profiler/EventExecution.h"
 #include "Editor/Kismet/Public/Profiler/ScriptPerfData.h"
 
@@ -289,31 +290,41 @@ FSlateColor SGraphNodeK2Base::GetNodeIndicatorOverlayColor() const
 	const float IntensityScale = 0.8f;
 
 	IBlueprintProfilerInterface& ProfilerModule = FModuleManager::LoadModuleChecked<IBlueprintProfilerInterface>("BlueprintProfiler");
-	if (ProfilerModule.IsProfilerEnabled())
+	if (ProfilerModule.IsProfilerEnabled() && GraphNode)
 	{
-		TSharedPtr<class FScriptExecutionNode> ExecNode = ProfilerModule.GetProfilerDataForNode(GetNodeObj());
-		if (ExecNode.IsValid())
+		UBlueprint* Blueprint = FBlueprintEditorUtils::FindBlueprintForNode(GraphNode);
+		if (Blueprint && Blueprint->GeneratedClass && Blueprint->GeneratedClass->HasInstrumentation())
 		{
-			FScriptPerfData NodePerfData(EScriptPerfDataType::Node);
-			ExecNode->GetBlueprintPerfDataForAllTracePaths(NodePerfData);
-
-			switch (ProfilerModule.GetGraphNodeHeatMapDisplayMode())
+			TSharedPtr<FScriptExecutionBlueprint> BlueprintExecNode = ProfilerModule.GetProfilerDataForBlueprint(Blueprint);
+			if (BlueprintExecNode.IsValid())
 			{
-			case EBlueprintProfilerHeatMapDisplayMode::Exclusive:
-				IntensityValue = NodePerfData.GetNodeHeatLevel();
-				break;
+				TSharedPtr<FScriptExecutionNode> ExecNode = ProfilerModule.GetProfilerDataForNode(GraphNode);
+				if (ExecNode.IsValid())
+				{
+					FScriptPerfData NodePerfData(EScriptPerfDataType::Node);
+					ExecNode->GetBlueprintPerfDataForAllTracePaths(NodePerfData);
+
+					NodePerfData.SetHeatLevels(BlueprintExecNode->GetHeatLevelMetrics());
 
-			case EBlueprintProfilerHeatMapDisplayMode::Inclusive:
-				IntensityValue = NodePerfData.GetInclusiveHeatLevel();
-				break;
+					switch (GetDefault<UBlueprintProfilerSettings>()->GraphNodeHeatMapDisplayMode)
+					{
+					case EBlueprintProfilerHeatMapDisplayMode::Average:
+						IntensityValue = NodePerfData.GetAverageHeatLevel();
+						break;
 
-			case EBlueprintProfilerHeatMapDisplayMode::MaxTiming:
-				IntensityValue = NodePerfData.GetMaxTimeHeatLevel();
-				break;
+					case EBlueprintProfilerHeatMapDisplayMode::Inclusive:
+						IntensityValue = NodePerfData.GetInclusiveHeatLevel();
+						break;
 
-			case EBlueprintProfilerHeatMapDisplayMode::Total:
-				IntensityValue = NodePerfData.GetTotalHeatLevel();
-				break;
+					case EBlueprintProfilerHeatMapDisplayMode::MaxTiming:
+						IntensityValue = NodePerfData.GetMaxTimeHeatLevel();
+						break;
+
+					case EBlueprintProfilerHeatMapDisplayMode::Total:
+						IntensityValue = NodePerfData.GetTotalHeatLevel();
+						break;
+					}
+				}
 			}
 		}
 	}
@@ -324,7 +335,7 @@ FSlateColor SGraphNodeK2Base::GetNodeIndicatorOverlayColor() const
 EVisibility SGraphNodeK2Base::GetNodeIndicatorOverlayVisibility() const
 {
 	IBlueprintProfilerInterface& ProfilerModule = FModuleManager::LoadModuleChecked<IBlueprintProfilerInterface>("BlueprintProfiler");
-	if (ProfilerModule.IsProfilerEnabled() && ProfilerModule.GetGraphNodeHeatMapDisplayMode() != EBlueprintProfilerHeatMapDisplayMode::None)
+	if (ProfilerModule.IsProfilerEnabled() && GetDefault<UBlueprintProfilerSettings>()->GraphNodeHeatMapDisplayMode != EBlueprintProfilerHeatMapDisplayMode::None)
 	{
 		return EVisibility::Visible;
 	}
diff --git a/Engine/Source/Editor/Kismet/Private/BlueprintEditor.cpp b/Engine/Source/Editor/Kismet/Private/BlueprintEditor.cpp
--- a/Engine/Source/Editor/Kismet/Private/BlueprintEditor.cpp
+++ b/Engine/Source/Editor/Kismet/Private/BlueprintEditor.cpp
@@ -1371,6 +1371,7 @@ FBlueprintEditor::FBlueprintEditor()
 	, bCodeBasedProject(false)
 	, HasOpenActionMenu(nullptr)
 	, InstructionsFadeCountdown(0.f)
+	, bBlueprintHasInstrumentation(false)
 {
 	AnalyticsStats.GraphActionMenusNonCtxtSensitiveExecCount = 0;
 	AnalyticsStats.GraphActionMenusCtxtSensitiveExecCount = 0;
@@ -1689,16 +1690,12 @@ void FBlueprintEditor::InitBlueprintEditor(
 	// Post-layout initialization
 	PostLayoutBlueprintEditorInitialization();
 
-	// Ensure the profiler UI respects the current state if it had previously been enabled.
+	// Ensure the profiler UI respects the current blueprint state if the blueprint has instrumentation.
 	if (IsProfilerAvailable())
 	{
-		bool bOpenProfilerTab = false;
-		if (GetDefault<UEditorExperimentalSettings>()->bBlueprintPerformanceAnalysisTools)
-		{
-			IBlueprintProfilerInterface& ProfilerModule = FModuleManager::LoadModuleChecked<IBlueprintProfilerInterface>("BlueprintProfiler");
-			bOpenProfilerTab = ProfilerModule.IsProfilerEnabled();
-		}
-		if (bOpenProfilerTab)
+		UBlueprint* MyBlueprint = GetBlueprintObj();
+		bBlueprintHasInstrumentation = MyBlueprint && MyBlueprint->GeneratedClass ? MyBlueprint->GeneratedClass->HasInstrumentation() : false;
+		if (bBlueprintHasInstrumentation)
 		{
 			TabManager->InvokeTab(FBlueprintEditorTabs::BlueprintProfilerID);
 		}
@@ -2233,8 +2230,6 @@ void FBlueprintEditor::CreateDefaultTabContents(const TArray<UBlueprint*>& InBlu
 				.IsEnabled(this, &FBlueprintEditor::IsFocusedGraphEditable);
 	}
 
-	FIsPropertyEditingEnabled IsPropertyEditingEnabledDelegate = FIsPropertyEditingEnabled::CreateSP(this, &FBlueprintEditor::IsFocusedGraphEditable);
-
 	if (IsEditingSingleBlueprint())
 	{
 		this->MyBlueprintWidget = SNew(SMyBlueprint, SharedThis(this));
@@ -2253,7 +2248,6 @@ void FBlueprintEditor::CreateDefaultTabContents(const TArray<UBlueprint*>& InBlu
 		. HideNameArea(true)
 		. ViewIdentifier(FName("BlueprintInspector"))
 		. Kismet2(SharedThis(this))
-		. IsPropertyEditingEnabledDelegate( IsPropertyEditingEnabledDelegate )
 		. OnFinishedChangingProperties( FOnFinishedChangingProperties::FDelegate::CreateSP(this, &FBlueprintEditor::OnFinishedChangingProperties) );
 
 	if ( InBlueprints.Num() > 0 )
@@ -2269,7 +2263,6 @@ void FBlueprintEditor::CreateDefaultTabContents(const TArray<UBlueprint*>& InBlu
 			. ShowPublicViewControl(bShowPublicView)
 			. ShowTitleArea(false)
 			. HideNameArea(bHideNameArea)
-			. IsPropertyEditingEnabledDelegate( IsPropertyEditingEnabledDelegate )
 			. OnFinishedChangingProperties( FOnFinishedChangingProperties::FDelegate::CreateSP( this, &FBlueprintEditor::OnFinishedChangingProperties ) );
 	}
 
@@ -3059,12 +3052,11 @@ void FBlueprintEditor::OnGraphEditorDropStreamingLevel(const TArray< TWeakObject
 		{
 			const FVector2D NodeLocation = DropLocation + (i * FVector2D(0,80));
 				
-			UK2Node_CallFunction* NodeTemplate = NewObject<UK2Node_CallFunction>();
+			UK2Node_CallFunction* NodeTemplate = NewObject<UK2Node_CallFunction>(Graph);
+			NodeTemplate->SetFromFunction(TargetFunc);
 			UK2Node_CallFunction* Node = FEdGraphSchemaAction_K2NewNode::SpawnNodeFromTemplate<UK2Node_CallFunction>(Graph, NodeTemplate, NodeLocation);
-			Node->SetFromFunction(TargetFunc);
 						
 			// Set dropped level package name
-			Node->AllocateDefaultPins();
 			UEdGraphPin* PackageNameInputPin = Node->FindPinChecked(TEXT("PackageName"));
 			PackageNameInputPin->DefaultValue = DroppedLevel->GetWorldAssetPackageName();
 		}
@@ -3198,11 +3190,23 @@ void FBlueprintEditor::Compile()
 		Arguments.Add(TEXT("BlueprintName"), FText::FromString(BlueprintObj->GetName()));
 		BlueprintLog.NewPage(FText::Format(LOCTEXT("CompilationPageLabel", "Compile {BlueprintName}"), Arguments));
 
+		// Decide if we want an instrumented compile.
+		const bool bProfilerAvailable = IsProfilerAvailable();
+		bool bAddInstrumentation = false;
+	    if (bProfilerAvailable)
+	    {
+		    if (GetDefault<UEditorExperimentalSettings>()->bBlueprintPerformanceAnalysisTools)
+		    {
+			    IBlueprintProfilerInterface& ProfilerModule = FModuleManager::LoadModuleChecked<IBlueprintProfilerInterface>("BlueprintProfiler");
+			    bAddInstrumentation = ProfilerModule.IsProfilerEnabled();
+		    }
+		}
+
 		FCompilerResultsLog LogResults;
 		LogResults.BeginEvent(TEXT("Compile"));
 		LogResults.bLogDetailedResults = GetDefault<UBlueprintEditorSettings>()->bShowDetailedCompileResults;
 		LogResults.EventDisplayThresholdMs = GetDefault<UBlueprintEditorSettings>()->CompileEventDisplayThresholdMs;
-		FKismetEditorUtilities::CompileBlueprint(BlueprintObj, false, false, bSaveIntermediateBuildProducts, &LogResults);
+		FKismetEditorUtilities::CompileBlueprint(BlueprintObj, false, false, bSaveIntermediateBuildProducts, &LogResults, false, false, bAddInstrumentation);
 
 		bool bForceMessageDisplay = ((LogResults.NumWarnings > 0) || (LogResults.NumErrors > 0)) && !BlueprintObj->bIsRegeneratingOnLoad;
 		DumpMessagesToCompilerLog(LogResults.Messages, bForceMessageDisplay);
@@ -3225,6 +3229,25 @@ void FBlueprintEditor::Compile()
 
 		LogResults.EndEvent();
 
+	    // Update the blueprint instrumenation state and show the profiler window if required
+	    if (bProfilerAvailable)
+	    {
+		    bBlueprintHasInstrumentation = BlueprintObj && BlueprintObj->GeneratedClass ? BlueprintObj->GeneratedClass->HasInstrumentation() : false;
+		    if (bBlueprintHasInstrumentation)
+		    {
+				if (GetDefault<UEditorExperimentalSettings>()->bBlueprintPerformanceAnalysisTools)
+				{
+					IBlueprintProfilerInterface& ProfilerModule = FModuleManager::LoadModuleChecked<IBlueprintProfilerInterface>("BlueprintProfiler");
+					ProfilerModule.AddInstrumentedBlueprint(BlueprintObj);
+				    TabManager->InvokeTab(FBlueprintEditorTabs::BlueprintProfilerID);
+				}
+		    }
+		    else
+		    {
+			    TabManager->InvokeTab(FBlueprintEditorTabs::BlueprintProfilerID)->RequestCloseTab();
+		    }
+		}
+
 		// send record when player clicks compile and send the result
 		// this will make sure how the users activity is
 		AnalyticsTrackCompileEvent(BlueprintObj, LogResults.NumErrors, LogResults.NumWarnings);
@@ -6486,7 +6509,7 @@ void FBlueprintEditor::ExtractEventTemplateForFunction(class UK2Node_CustomEvent
 	}
 }
 
-void FBlueprintEditor::CollapseNodesIntoGraph(UEdGraphNode* InGatewayNode, UK2Node_EditablePinBase* InEntryNode, UK2Node_EditablePinBase* InResultNode, UEdGraph* InSourceGraph, UEdGraph* InDestinationGraph, TSet<UEdGraphNode*>& InCollapsableNodes)
+void FBlueprintEditor::CollapseNodesIntoGraph(UEdGraphNode* InGatewayNode, UK2Node_EditablePinBase* InEntryNode, UK2Node_EditablePinBase* InResultNode, UEdGraph* InSourceGraph, UEdGraph* InDestinationGraph, TSet<UEdGraphNode*>& InCollapsableNodes, bool bCanDiscardEmptyReturnNode)
 {
 	const UEdGraphSchema_K2* K2Schema = GetDefault<UEdGraphSchema_K2>();
 
@@ -6713,7 +6736,7 @@ void FBlueprintEditor::CollapseNodesIntoGraph(UEdGraphNode* InGatewayNode, UK2No
 		InResultNode->SnapToGrid(SNodePanel::GetSnapGridSize());
 	}
 
-	if (bDiscardReturnNode)
+	if (bCanDiscardEmptyReturnNode && bDiscardReturnNode)
 	{
 		InResultNode->DestroyNode();
 	}
@@ -6799,7 +6822,7 @@ UEdGraph* FBlueprintEditor::CollapseSelectionToFunction(TSharedPtr<SGraphEditor>
 
 	check(OutFunctionNode);
 
-	CollapseNodesIntoGraph(OutFunctionNode, EntryNode, ResultNode, SourceGraph, NewGraph, InCollapsableNodes);
+	CollapseNodesIntoGraph(OutFunctionNode, EntryNode, ResultNode, SourceGraph, NewGraph, InCollapsableNodes, true);
 
 	FBlueprintEditorUtils::MarkBlueprintAsStructurallyModified(GetBlueprintObj());
 	OutFunctionNode->ReconstructNode();
diff --git a/Engine/Source/Editor/Kismet/Private/Profiler/BPProfilerStatisticWidgets.cpp b/Engine/Source/Editor/Kismet/Private/Profiler/BPProfilerStatisticWidgets.cpp
--- a/Engine/Source/Editor/Kismet/Private/Profiler/BPProfilerStatisticWidgets.cpp
+++ b/Engine/Source/Editor/Kismet/Private/Profiler/BPProfilerStatisticWidgets.cpp
@@ -28,7 +28,7 @@ namespace BlueprintProfilerStatText
 	const FText ColumnText_MaxTime(LOCTEXT("MaxTime", "Max Time (ms)"));
 	const FText ColumnText_MinTime(LOCTEXT("MinTime", "Min Time (ms)"));
 	const FText ColumnText_Samples(LOCTEXT("Samples", "Samples"));
-	const FText ColumnText_TotalTime(LOCTEXT("TotalTime", "Total Time (s)"));
+	const FText ColumnText_TotalTime(LOCTEXT("TotalTime", "Total Time (ms)"));
 };
 
 //////////////////////////////////////////////////////////////////////////
@@ -65,7 +65,7 @@ void SProfilerStatRow::Construct(const FArguments& InArgs, TSharedRef<STableView
 {
 	check(InItemToEdit.IsValid());
 	ItemToEdit = InItemToEdit;
-	SMultiColumnTableRow<FBPStatWidgetPtr>::Construct( FSuperRowType::FArguments(), OwnerTableView );	
+	SMultiColumnTableRow<FBPStatWidgetPtr>::Construct(FSuperRowType::FArguments().Style(FEditorStyle::Get(), "BlueprintProfiler.TableView.Row"), OwnerTableView);
 }
 
 const FName SProfilerStatRow::GetStatName(const EBlueprintProfilerStat::Type StatId)
@@ -157,13 +157,14 @@ TSharedRef<SWidget> FBPProfilerStatWidget::GenerateColumnWidget(FName ColumnName
 			}
 			else if (ColumnName == BlueprintProfilerStatText::ColumnId_AverageTime)
 			{
-				TextAttr = TAttribute<FText>(PerformanceStats.Get(), &FScriptPerfData::GetExclusiveTimingText);
-				ColorAttr = TAttribute<FSlateColor>(this, &FBPProfilerStatWidget::GetNodeHeatColor);
+				TextAttr = TAttribute<FText>(PerformanceStats.Get(), &FScriptPerfData::GetAverageTimingText);
+				ColorAttr = TAttribute<FSlateColor>(this, &FBPProfilerStatWidget::GetAverageHeatColor);
 			}
 			else if (ColumnName == BlueprintProfilerStatText::ColumnId_MaxTime)
 			{
 				TextAttr = TAttribute<FText>(PerformanceStats.Get(), &FScriptPerfData::GetMaxTimingText);
-				ColorAttr = TAttribute<FSlateColor>(this, &FBPProfilerStatWidget::GetMaxTimeHeatColor);
+				//ColorAttr = TAttribute<FSlateColor>(this, &FBPProfilerStatWidget::GetMaxTimeHeatColor);
+				ColorAttr = TAttribute<FSlateColor>(FLinearColor::White);
 			}
 			else if (ColumnName == BlueprintProfilerStatText::ColumnId_MinTime)
 			{
@@ -192,9 +193,9 @@ void FBPProfilerStatWidget::NavigateTo() const
 	}
 }
 
-FSlateColor FBPProfilerStatWidget::GetNodeHeatColor() const
+FSlateColor FBPProfilerStatWidget::GetAverageHeatColor() const
 {
-	const float Value = 1.f - PerformanceStats->GetNodeHeatLevel();
+	const float Value = 1.f - PerformanceStats->GetAverageHeatLevel();
 	return FLinearColor(1.f, Value, Value);
 }
 
@@ -252,39 +253,36 @@ void FBPProfilerStatWidget::GenerateExecNodeWidgets(const TSharedPtr<FBlueprintP
 void FBPProfilerStatWidget::GenerateStandardNodeWidgets(const TSharedPtr<FBlueprintProfilerStatOptions> DisplayOptions)
 {
 	// Add any child nodes that are not filtered out.
-	if (ExecNode->HasFlags(EScriptExecutionNodeFlags::Container))
+	for (auto Iter : ExecNode->GetChildNodes())
 	{
-		for (auto Iter : ExecNode->GetChildNodes())
+		// Filter out events based on graph
+		if (!DisplayOptions->IsFiltered(Iter))
 		{
-			// Filter out events based on graph
-			if (!DisplayOptions->IsFiltered(Iter))
+			TArray<FScriptNodeExecLinkage::FLinearExecPath> LinearExecNodes;
+			FTracePath ChildTracePath(WidgetTracePath);
+			Iter->GetLinearExecutionPath(LinearExecNodes, ChildTracePath, false);
+			if (LinearExecNodes.Num() > 1)
 			{
-				TArray<FScriptNodeExecLinkage::FLinearExecPath> LinearExecNodes;
-				FTracePath ChildTracePath(WidgetTracePath);
-				Iter->GetLinearExecutionPath(LinearExecNodes, ChildTracePath, false);
-				if (LinearExecNodes.Num() > 1)
+				TSharedPtr<FBPProfilerStatWidget> ChildContainer = AsShared();
+				for (auto LinearPathIter : LinearExecNodes)
 				{
-					TSharedPtr<FBPProfilerStatWidget> ChildContainer = AsShared();
-					for (auto LinearPathIter : LinearExecNodes)
+					if (!DisplayOptions->IsFiltered(LinearPathIter.LinkedNode))
 					{
-						if (!DisplayOptions->IsFiltered(LinearPathIter.LinkedNode))
+						TSharedPtr<FBPProfilerStatWidget> NewLinkedNode = MakeShareable<FBPProfilerStatWidget>(new FBPProfilerStatWidget(LinearPathIter.LinkedNode, LinearPathIter.TracePath));
+						NewLinkedNode->GenerateExecNodeWidgets(DisplayOptions);
+						ChildContainer->CachedChildren.Add(NewLinkedNode);
+						if (LinearPathIter.LinkedNode->HasFlags(EScriptExecutionNodeFlags::Container))
 						{
-							TSharedPtr<FBPProfilerStatWidget> NewLinkedNode = MakeShareable<FBPProfilerStatWidget>(new FBPProfilerStatWidget(LinearPathIter.LinkedNode, LinearPathIter.TracePath));
-							NewLinkedNode->GenerateExecNodeWidgets(DisplayOptions);
-							ChildContainer->CachedChildren.Add(NewLinkedNode);
-							if (LinearPathIter.LinkedNode->HasFlags(EScriptExecutionNodeFlags::Container))
-							{
-								ChildContainer = NewLinkedNode;
-							}
+							ChildContainer = NewLinkedNode;
 						}
 					}
 				}
-				else
-				{
-					TSharedPtr<FBPProfilerStatWidget> NewChildNode = MakeShareable<FBPProfilerStatWidget>(new FBPProfilerStatWidget(Iter, ChildTracePath));
-					NewChildNode->GenerateExecNodeWidgets(DisplayOptions);
-					CachedChildren.Add(NewChildNode);
-				}
+			}
+			else
+			{
+				TSharedPtr<FBPProfilerStatWidget> NewChildNode = MakeShareable<FBPProfilerStatWidget>(new FBPProfilerStatWidget(Iter, ChildTracePath));
+				NewChildNode->GenerateExecNodeWidgets(DisplayOptions);
+				CachedChildren.Add(NewChildNode);
 			}
 		}
 	}
@@ -339,34 +337,25 @@ void FBPProfilerStatWidget::GenerateSimpleTunnelWidgets(TSharedPtr<FScriptExecut
 	for (auto EntryPoint : TunnelEntryNode->GetChildNodes())
 	{
 		// Filter out events based on graph
-		if (!DisplayOptions->IsFiltered(EntryPoint))
+		if (!DisplayOptions->IsFiltered(EntryPoint) && !EntryPoint->IsPureChain())
 		{
 			TArray<FScriptNodeExecLinkage::FLinearExecPath> LinearExecNodes;
 			EntryPoint->GetLinearExecutionPath(LinearExecNodes, WidgetTracePath, true);
-			if (LinearExecNodes.Num() > 1)
+			TSharedPtr<FBPProfilerStatWidget> ChildContainer = AsShared();
+			for (auto LinearPathIter : LinearExecNodes)
 			{
-				TSharedPtr<FBPProfilerStatWidget> ChildContainer = AsShared();
-				for (auto LinearPathIter : LinearExecNodes)
+				const bool bInternalTunnelBoundary = TunnelEntryNode->IsInternalBoundary(LinearPathIter.LinkedNode);
+				if (!DisplayOptions->IsFiltered(LinearPathIter.LinkedNode) && !bInternalTunnelBoundary  && !EntryPoint->IsPureChain())
 				{
-					const bool bInternalTunnelBoundary = TunnelEntryNode->IsInternalBoundary(LinearPathIter.LinkedNode);
-					if (!DisplayOptions->IsFiltered(LinearPathIter.LinkedNode) && !bInternalTunnelBoundary)
+					TSharedPtr<FBPProfilerStatWidget> NewLinkedNode = MakeShareable<FBPProfilerStatWidget>(new FBPProfilerStatWidget(LinearPathIter.LinkedNode, LinearPathIter.TracePath));
+					NewLinkedNode->GenerateExecNodeWidgets(DisplayOptions);
+					CachedChildren.Add(NewLinkedNode);
+					if (LinearPathIter.LinkedNode->HasFlags(EScriptExecutionNodeFlags::Container))
 					{
-						TSharedPtr<FBPProfilerStatWidget> NewLinkedNode = MakeShareable<FBPProfilerStatWidget>(new FBPProfilerStatWidget(LinearPathIter.LinkedNode, LinearPathIter.TracePath));
-						NewLinkedNode->GenerateExecNodeWidgets(DisplayOptions);
-						CachedChildren.Add(NewLinkedNode);
-						if (LinearPathIter.LinkedNode->HasFlags(EScriptExecutionNodeFlags::Container))
-						{
-							ChildContainer = NewLinkedNode;
-						}
+						ChildContainer = NewLinkedNode;
 					}
 				}
 			}
-			else
-			{
-				TSharedPtr<FBPProfilerStatWidget> NewChildNode = MakeShareable<FBPProfilerStatWidget>(new FBPProfilerStatWidget(EntryPoint, WidgetTracePath));
-				NewChildNode->GenerateExecNodeWidgets(DisplayOptions);
-				CachedChildren.Add(NewChildNode);
-			}
 		}
 	}
 }
@@ -377,15 +366,15 @@ void FBPProfilerStatWidget::GenerateComplexTunnelWidgets(TSharedPtr<FScriptExecu
 	for (auto EntryPoint : TunnelEntryNode->GetChildNodes())
 	{
 		// Filter out events based on graph
-		if (!DisplayOptions->IsFiltered(EntryPoint))
+		if (!DisplayOptions->IsFiltered(EntryPoint) && !EntryPoint->IsPureChain())
 		{
 			TSharedPtr<FBPProfilerStatWidget> EntryPointNode = MakeShareable<FBPProfilerStatWidget>(new FBPProfilerStatWidget(EntryPoint, WidgetTracePath));
-			EntryPointNode->GenerateTunnelLinkWidgets(DisplayOptions);
+			EntryPointNode->GenerateExecNodeWidgets(DisplayOptions);
 			CachedChildren.Add(EntryPointNode);
 		}
 	}
 	// Create exit links as child nodes
-	for (auto ExitSite : TunnelEntryNode->GetExitSites())
+	for (auto ExitSite : TunnelEntryNode->GetLinkedNodes())
 	{
 		TSharedPtr<FBPProfilerStatWidget> ExitSiteNode = MakeShareable<FBPProfilerStatWidget>(new FBPProfilerStatWidget(ExitSite.Value, WidgetTracePath));
 		ExitSiteNode->GenerateTunnelLinkWidgets(DisplayOptions, ExitSite.Key);
diff --git a/Engine/Source/Editor/Kismet/Private/Profiler/BPProfilerStatisticWidgets.h b/Engine/Source/Editor/Kismet/Private/Profiler/BPProfilerStatisticWidgets.h
--- a/Engine/Source/Editor/Kismet/Private/Profiler/BPProfilerStatisticWidgets.h
+++ b/Engine/Source/Editor/Kismet/Private/Profiler/BPProfilerStatisticWidgets.h
@@ -114,8 +114,8 @@ protected:
 	/** Navigate to referenced node */
 	void NavigateTo() const;
 
-	/** Calculate and return the heat color to use for exclusive timing stats. */
-	FSlateColor GetNodeHeatColor() const;
+	/** Calculate and return the heat color to use for average timing stats. */
+	FSlateColor GetAverageHeatColor() const;
 
 	/** Calculate and return the heat color to use for inclusive timing stats. */
 	FSlateColor GetInclusiveHeatColor() const;
@@ -133,5 +133,4 @@ protected:
 	TSharedPtr<FScriptPerfData> PerformanceStats;
 	/** Child statistics */
 	TArray<TSharedPtr<FBPProfilerStatWidget>> CachedChildren;
-
 };
diff --git a/Engine/Source/Editor/Kismet/Private/Profiler/BlueprintProfilerSettings.cpp b/Engine/Source/Editor/Kismet/Private/Profiler/BlueprintProfilerSettings.cpp
new file mode 100644
--- /dev/null
+++ b/Engine/Source/Editor/Kismet/Private/Profiler/BlueprintProfilerSettings.cpp
@@ -0,0 +1,29 @@
+// Copyright 1998-2016 Epic Games, Inc. All Rights Reserved.
+
+#include "BlueprintEditorPrivatePCH.h"
+#include "BlueprintProfilerSettings.h"
+#include "ScriptPerfData.h"
+
+float UBlueprintProfilerSettings::CustomEventPerformanceThresholdDefaultValue = 1.f;
+float UBlueprintProfilerSettings::CustomAveragePerformanceThresholdDefaultValue = 0.2f;
+float UBlueprintProfilerSettings::CustomInclusivePerformanceThresholdDefaultValue = 0.25f;
+float UBlueprintProfilerSettings::CustomMaxPerformanceThresholdDefaultValue = 0.5f;
+
+UBlueprintProfilerSettings::UBlueprintProfilerSettings(const FObjectInitializer& ObjectInitializer)
+	: Super(ObjectInitializer)
+	, bDisplayByInstance(true)
+	, bScopeToDebugInstance(false)
+	, bGraphFilter(false)
+	, bDisplayPure(true)
+	, bDisplayInheritedEvents(true)
+	, bAverageBlueprintStats(true)
+	, GraphNodeHeatMapDisplayMode(EBlueprintProfilerHeatMapDisplayMode::None)
+	, WireHeatMapDisplayMode(EBlueprintProfilerHeatMapDisplayMode::None)
+	, HeatLevelMetricsType(EBlueprintProfilerHeatLevelMetricsType::ClassRelative)
+	, CustomEventPerformanceThreshold(CustomEventPerformanceThresholdDefaultValue)
+	, CustomAveragePerformanceThreshold(CustomAveragePerformanceThresholdDefaultValue)
+	, CustomInclusivePerformanceThreshold(CustomInclusivePerformanceThresholdDefaultValue)
+	, CustomMaxPerformanceThreshold(CustomMaxPerformanceThresholdDefaultValue)
+{
+	FScriptPerfData::EnableBlueprintStatAverage(bAverageBlueprintStats);
+}
\ No newline at end of file
diff --git a/Engine/Source/Editor/Kismet/Private/Profiler/EventExecution.cpp b/Engine/Source/Editor/Kismet/Private/Profiler/EventExecution.cpp
--- a/Engine/Source/Editor/Kismet/Private/Profiler/EventExecution.cpp
+++ b/Engine/Source/Editor/Kismet/Private/Profiler/EventExecution.cpp
@@ -1,12 +1,14 @@
 // Copyright 1998-2016 Epic Games, Inc. All Rights Reserved.
 
 #include "BlueprintEditorPrivatePCH.h"
+#include "BlueprintProfilerSettings.h"
 #include "EventExecution.h"
-#include "BlueprintProfilerModule.h"
 #include "SHyperlink.h"
 
-// This forces checking against correct tracepaths when querying data
-#define STRICT_PERFDATA_CREATION 0
+// Debugging defines
+#define STRICT_PERFDATA_CREATION 0		// Asserts when script perf data is created from an incorrect tracepath.
+#define DISPLAY_NODENAMES_IN_TITLES 0 	// Displays full node names in the stat widget titles for debugging.
+#define DISPLAY_NODENAMES_ON_TOOLTIPS 0 // Displays full node names in the stat widget tooltips for debugging.
 
 #define LOCTEXT_NAMESPACE "BlueprintEventUI"
 
@@ -60,7 +62,7 @@
 	TSharedPtr<FScriptPerfData>& Result = InstanceMap.FindOrAdd(TracePath);
 	if (!Result.IsValid())
 	{
-		Result = MakeShareable<FScriptPerfData>(new FScriptPerfData(GetPerfDataType()));
+		Result = MakeShareable<FScriptPerfData>(new FScriptPerfData(GetPerfDataType(), SampleFrequency));
 	}
 	return Result;
 }
@@ -76,7 +78,7 @@
 	TSharedPtr<FScriptPerfData>& Result = InstanceMap.FindOrAdd(TracePath);
 	if (!Result.IsValid())
 	{
-		Result = MakeShareable<FScriptPerfData>(new FScriptPerfData(GetPerfDataType()));
+		Result = MakeShareable<FScriptPerfData>(new FScriptPerfData(GetPerfDataType(), SampleFrequency));
 	}
 	return Result;
 	#endif
@@ -103,7 +105,7 @@
 	TSharedPtr<FScriptPerfData>& Result = BlueprintMap.FindOrAdd(TracePath);
 	if (!Result.IsValid())
 	{
-		Result = MakeShareable<FScriptPerfData>(new FScriptPerfData(GetPerfDataType()));
+		Result = MakeShareable<FScriptPerfData>(new FScriptPerfData(GetPerfDataType(), SampleFrequency));
 	}
 	return Result;
 }
@@ -120,7 +122,7 @@
 	TSharedPtr<FScriptPerfData>& Result = BlueprintMap.FindOrAdd(TracePath);
 	if (!Result.IsValid())
 	{
-		Result = MakeShareable<FScriptPerfData>(new FScriptPerfData(GetPerfDataType()));
+		Result = MakeShareable<FScriptPerfData>(new FScriptPerfData(GetPerfDataType(), SampleFrequency));
 	}
 	return Result;
 	#endif
@@ -145,19 +147,25 @@
 // FScriptExecutionNode
 
 FScriptExecutionNode::FScriptExecutionNode()
-	: NodeFlags(EScriptExecutionNodeFlags::None)
+	: FScriptNodePerfData(1)
+	, NodeFlags(EScriptExecutionNodeFlags::None)
 	, bExpansionState(false)
 {
 }
 
 FScriptExecutionNode::FScriptExecutionNode(const FScriptExecNodeParams& InitParams)
-	: NodeName(InitParams.NodeName)
+	: FScriptNodePerfData(InitParams.SampleFrequency)
+	, NodeName(InitParams.NodeName)
 	, OwningGraphName(InitParams.OwningGraphName)
 	, NodeFlags(InitParams.NodeFlags)
 	, ObservedObject(InitParams.ObservedObject)
 	, ObservedPin(InitParams.ObservedPin)
 	, DisplayName(InitParams.DisplayName)
+#if DISPLAY_NODENAMES_ON_TOOLTIPS
+	, Tooltip(FText::FromName(InitParams.NodeName))
+#else
 	, Tooltip(InitParams.Tooltip)
+#endif
 	, IconColor(InitParams.IconColor)
 	, Icon(InitParams.Icon)
 	, bExpansionState(false)
@@ -179,12 +187,12 @@
 TSharedRef<SWidget> FScriptExecutionNode::GetHyperlinkWidget(const uint32 /*TracePath*/)
 {
 	return SNew(SHyperlink)
-		#if TRACEPATH_DEBUG
+		#if DISPLAY_NODENAMES_IN_TITLES
 		.Text(FText::FromName(NodeName))
 		#else
 		.Text(DisplayName)
 		#endif
-		.Style(FEditorStyle::Get(), "HoverOnlyHyperlink")
+		.Style(FEditorStyle::Get(), "BlueprintProfiler.Node.Hyperlink")
 		.ToolTipText(Tooltip)
 		.OnNavigate(this, &FScriptExecutionNode::NavigateToObject);
 }
@@ -304,22 +312,22 @@
 					TSharedPtr<FScriptPerfData> InstancePerfData = GetPerfDataByInstanceAndTracePath(InstanceIter, TracePath);
 					InstancePerfData->Reset();
 					BlueprintPooledStats.Add(InstancePerfData);
+					TArray<TSharedPtr<FScriptPerfData>> ChildDataSet;
 
 					for (auto ChildIter : ChildExecPaths)
 					{
 						TSharedPtr<FScriptPerfData> ChildPerfData = ChildIter.LinkedNode->GetPerfDataByInstanceAndTracePath(InstanceIter, ChildIter.TracePath);
-						InstancePerfData->AddBranchData(*ChildPerfData.Get());
+						if (ChildPerfData.IsValid())
+						{
+							ChildDataSet.Add(ChildPerfData);
+						}
 					}
+					InstancePerfData->InitialiseFromDataSet(ChildDataSet);
 				}
 			}
 			// Update Blueprint stats as branches
 			TSharedPtr<FScriptPerfData> BlueprintData = GetOrAddBlueprintPerfDataByTracePath(TracePath);
-			BlueprintData->Reset();
-
-			for (auto BlueprintChildDataIter : BlueprintPooledStats)
-			{
-				BlueprintData->AddBranchData(*BlueprintChildDataIter.Get());
-			}
+			BlueprintData->InitialiseFromDataSet(BlueprintPooledStats);
 			BlueprintPooledStats.Reset(0);
 			bRefreshLinkStats = false;
 			break;
@@ -358,6 +366,31 @@
 	}
 }
 
+void FScriptExecutionNode::CalculateHeatLevelStats(TSharedPtr<FScriptHeatLevelMetrics> HeatLevelMetrics)
+{
+	for (auto InstancePerfDataMapIter : InstanceInputPinToPerfDataMap)
+	{
+		for (auto TracePathPerfDataMapIter : InstancePerfDataMapIter.Value)
+		{
+			TSharedPtr<FScriptPerfData> PerfData = TracePathPerfDataMapIter.Value;
+			if (PerfData.IsValid())
+			{
+				PerfData->SetHeatLevels(HeatLevelMetrics);
+			}
+		}
+	}
+
+	for (auto ChildIter : ChildNodes)
+	{
+		ChildIter->CalculateHeatLevelStats(HeatLevelMetrics);
+	}
+
+	for (auto LinkIter : LinkedNodes)
+	{
+		LinkIter.Value->CalculateHeatLevelStats(HeatLevelMetrics);
+	}
+}
+
 float FScriptExecutionNode::CalculateHottestPathStats(FScriptExecutionHottestPathParams HotPathParams)
 {
 	// Grab local perf data
@@ -367,11 +400,9 @@
 	if (!IsEvent())
 	{
 		// Subtract local inclusive cost from parent inclusive cost
-		const double NodeTime = IsBranch() ? LocalPerfData->GetExclusiveTiming() : LocalPerfData->GetInclusiveTiming();
+		const double NodeTime = IsBranch() ? LocalPerfData->GetAverageTiming() : LocalPerfData->GetInclusiveTiming();
 		HotPathParams.TimeSoFar += NodeTime;
 		AccumulatedTime += NodeTime;
-		const float HottestEndpointValue = static_cast<float>(HotPathParams.TimeSoFar / HotPathParams.EventTiming);
-		LocalPerfData->SetHottestEndpointHeatLevel(HottestEndpointValue);
 	}
 	// Update children
 	for (auto ChildIter : ChildNodes)
@@ -540,29 +571,28 @@
 void FScriptExecutionTunnelEntry::AddTunnelTiming(const FName InstanceName, const FTracePath& TracePath, const FTracePath& InternalTracePath, const int32 ExitScriptOffset, const double Time)
 {
 	TSharedPtr<FScriptPerfData> PerfData = GetPerfDataByInstanceAndTracePath(InstanceName, TracePath);
-	PerfData->AddEventTiming(Time);
+	PerfData->AddEventTiming(0.0);
+	PerfData->AddInclusiveTiming(Time, false);
+	PerfData->OverrideSampleCount(TunnelEntryCount);
 	// Tick Entry points
 	for (auto EntrySite : ChildNodes)
 	{
 		PerfData = EntrySite->GetPerfDataByInstanceAndTracePath(InstanceName, TracePath);
-		PerfData->TickSamples();
+		PerfData->AddEventTiming(0.0);
+		PerfData->OverrideSampleCount(TunnelEntryCount);
 	}
 	// Tick exit sites
 	TSharedPtr<FScriptExecutionNode> TunnelExitNode = GetExitSite(ExitScriptOffset);
 	if (TunnelExitNode.IsValid())
 	{
 		PerfData = TunnelExitNode->GetPerfDataByInstanceAndTracePath(InstanceName, TracePath);
-		PerfData->TickSamples();
+		PerfData->AddEventTiming(0.0);
+		PerfData->OverrideSampleCount(TunnelEntryCount);
 		PerfData = TunnelExitNode->GetPerfDataByInstanceAndTracePath(InstanceName, InternalTracePath);
-		PerfData->TickSamples();
+		PerfData->AddEventTiming(0.0);
 	}
 }
 
-void FScriptExecutionTunnelEntry::AddExitSite(const int32 ExitScriptOffset, TSharedPtr<FScriptExecutionTunnelExit> ExitSite)
-{
-	LinkedNodes.Add(ExitScriptOffset) = ExitSite;
-}
-
 TSharedPtr<FScriptExecutionTunnelExit> FScriptExecutionTunnelEntry::GetExitSite(const int32 ScriptOffset) const
 {
 	TSharedPtr<FScriptExecutionTunnelExit> ExitSite;
@@ -584,11 +614,16 @@
 	GetInstancePerfDataByTracePath(TracePath, BlueprintPooledStats);
 	// Create new trace
 	FTracePath TunnelTracePath(TracePath, SharedThis<FScriptExecutionTunnelEntry>(this));
-	// Refresh Child Links
+	// Refresh child entry points
 	for (auto ChildIter : ChildNodes)
 	{
 		ChildIter->RefreshStats(TunnelTracePath);
 	}
+	// Refresh exit points
+	for (auto ExitPoint : LinkedNodes)
+	{
+		ExitPoint.Value->RefreshStats(TunnelTracePath);
+	}
 	// Update the owning blueprint stats
 	if (BlueprintPooledStats.Num() > 0)
 	{
@@ -631,17 +666,21 @@
 	// Refresh all valid tunnel links
 	for (auto LinkIter : LinkedNodes)
 	{
-		FTracePath LinkTracePath(TracePath);
-		LinkTracePath.ExitTunnel();
-		if (LinkIter.Value->HasFlags(EScriptExecutionNodeFlags::EventPin))
+		// Don't refresh any linked tunnel exits, the owning entry point will do that.
+		if (!LinkIter.Value->HasFlags(EScriptExecutionNodeFlags::TunnelExitPin))
 		{
-			LinkTracePath.ResetPath();
+			FTracePath LinkTracePath(TracePath);
+			LinkTracePath.ExitTunnel();
+			if (LinkIter.Value->HasFlags(EScriptExecutionNodeFlags::EventPin))
+			{
+				LinkTracePath.ResetPath();
+			}
+			if (!LinkIter.Value->HasFlags(EScriptExecutionNodeFlags::InvalidTrace))
+			{
+				LinkTracePath.AddExitPin(LinkIter.Key);
+			}
+			LinkIter.Value->RefreshStats(LinkTracePath);
 		}
-		if (!LinkIter.Value->HasFlags(EScriptExecutionNodeFlags::InvalidTrace))
-		{
-			LinkTracePath.AddExitPin(LinkIter.Key);
-		}
-		LinkIter.Value->RefreshStats(LinkTracePath);
 	}
 	// Update the owning blueprint stats
 	if (BlueprintPooledStats.Num() > 0)
@@ -677,7 +716,7 @@
 TSharedRef<SWidget> FScriptExecutionTunnelExit::GetHyperlinkWidget(const uint32 TracePath)
 {
 	return SNew(SHyperlink)
-		#if TRACEPATH_DEBUG
+		#if DISPLAY_NODENAMES_IN_TITLES
 		.Text(FText::FromName(NodeName))
 		#else
 		.Text(DisplayName)
@@ -734,7 +773,7 @@
 	// Refresh all pure nodes and accumulate for blueprint stat.
 	if (HasFlags(EScriptExecutionNodeFlags::PureChain))
 	{
-		//GetInstancePerfDataByTracePath
+		// Update linked pure nodes and find instances
 		TSet<FName> ValidInstances;
 		FTracePath PureTracePath(TracePath);
 		TMap<int32, TSharedPtr<FScriptExecutionNode>> AllPureNodes;
@@ -745,18 +784,6 @@
 			PureTracePath.AddExitPin(PureIter.Key);
 			PureIter.Value->RefreshStats(PureTracePath);
 		}
-		for (auto InstanceName : ValidInstances)
-		{
-			TSharedPtr<FScriptPerfData> InstancePerfData = GetPerfDataByInstanceAndTracePath(InstanceName, TracePath);
-			InstancePerfData->Reset();
-			BlueprintPooledStats.Add(InstancePerfData);
-
-			for (auto ChildIter : ChildNodes)
-			{
-				TSharedPtr<FScriptPerfData> ChildPerfData = ChildIter->GetPerfDataByInstanceAndTracePath(InstanceName, TracePath);
-				InstancePerfData->AddBranchData(*ChildPerfData.Get());
-			}
-		}
 	}
 	// Grab all instance stats to accumulate into the blueprint stat.
 	GetInstancePerfDataByTracePath(TracePath, BlueprintPooledStats);
@@ -769,8 +796,115 @@
 }
 
 //////////////////////////////////////////////////////////////////////////
+// FScriptExecutionInstance
+
+void FScriptExecutionInstance::NavigateToObject() const
+{
+	if (const FWorldContext* PIEWorldContext = GEditor->GetPIEWorldContext())
+	{
+		if (const AActor* PIEActor = Cast<AActor>(PIEInstance.Get()))
+		{
+			GEditor->SelectNone(false, false);
+			GEditor->SelectActor(const_cast<AActor*>(PIEActor), true, true, true);
+		}
+	}
+	else
+	{
+		if (const AActor* EditorActor = Cast<AActor>(ObservedObject.Get()))
+		{
+			GEditor->SelectNone(false, false);
+			GEditor->SelectActor(const_cast<AActor*>(EditorActor), true, true, true);
+		}
+	}
+}
+
+bool FScriptExecutionInstance::IsObservedObjectValid() const
+{
+	if (const FWorldContext* PIEWorldContext = GEditor->GetPIEWorldContext())
+	{
+		if (PIEInstance.IsValid() && PIEInstance.Get()->IsPendingKill())
+		{
+			bPIEInstanceDestroyed = true;
+		}
+		return PIEInstance.IsValid();
+	}
+	else
+	{
+		return ObservedObject.IsValid() && !ObservedObject.Get()->IsPendingKill();
+	}
+}
+
+UObject* FScriptExecutionInstance::GetActiveObject() const
+{
+	UObject* Result = nullptr;
+	if (const FWorldContext* PIEWorldContext = GEditor->GetPIEWorldContext())
+	{
+		if (!PIEInstance.IsValid() && !bPIEInstanceDestroyed)
+		{
+			RefreshPIEInstance();
+		}
+		Result = const_cast<UObject*>(PIEInstance.Get());
+	}
+	else
+	{
+		Result = const_cast<UObject*>(ObservedObject.Get());
+	}
+	return Result;
+}
+
+void FScriptExecutionInstance::RefreshPIEInstance() const
+{
+	if (!bPIEInstanceDestroyed)
+	{
+		const FWorldContext* PIEWorldContext = GEditor->GetPIEWorldContext();
+		UWorld* PIEWorld = PIEWorldContext ? PIEWorldContext->World() : nullptr;
+		ULevel* CurrentLevel = PIEWorld ? PIEWorld->GetCurrentLevel() : nullptr;
+		if (CurrentLevel && ObservedObject.IsValid())
+		{
+			PIEInstance = FindObject<UObject>(CurrentLevel, *ObservedObject.Get()->GetName());
+		}
+	}
+}
+
+TSharedRef<SWidget> FScriptExecutionInstance::GetIconWidget(const uint32 /*TracePath*/)
+{
+	return SNew(SImage)
+			.Image(Icon)
+			.ColorAndOpacity(this, &FScriptExecutionInstance::GetInstanceIconColor);
+}
+
+TSharedRef<SWidget> FScriptExecutionInstance::GetHyperlinkWidget(const uint32 /*TracePath*/)
+{
+	return SNew(SHyperlink)
+		#if DISPLAY_NODENAMES_IN_TITLES
+		.Text(FText::FromName(NodeName))
+		#else
+		.Text(DisplayName)
+		#endif
+		.Style(FEditorStyle::Get(), "HoverOnlyHyperlink")
+		.ToolTipText(this, &FScriptExecutionInstance::GetInstanceTooltip)
+		.OnNavigate(this, &FScriptExecutionInstance::NavigateToObject);
+}
+
+FSlateColor FScriptExecutionInstance::GetInstanceIconColor() const
+{
+	return FSlateColor(IsObservedObjectValid() ? IconColor : FLinearColor(1.f, 1.f, 1.f, 0.3f));
+}
+
+FText FScriptExecutionInstance::GetInstanceTooltip() const
+{
+	return IsObservedObjectValid() ? Tooltip : LOCTEXT("ActorInstanceInvalid_Tooltip", "Actor instance no longer exists");
+}
+
+//////////////////////////////////////////////////////////////////////////
 // FScriptExecutionBlueprint
 
+FScriptExecutionBlueprint::FScriptExecutionBlueprint(const FScriptExecNodeParams& InitParams)
+	: FScriptExecutionNode(InitParams)
+{
+	HeatLevelMetrics = MakeShareable(new FScriptHeatLevelMetrics());
+}
+
 TSharedPtr<FScriptExecutionNode> FScriptExecutionBlueprint::GetInstanceByName(FName InstanceName)
 {
 	TSharedPtr<FScriptExecutionNode> Result;
@@ -796,8 +930,7 @@
 		BlueprintEventIter->RefreshStats(TracePath);
 	}
 	// Determine if we require the hottest path's updating
-	IBlueprintProfilerInterface& ProfilerModule = FModuleManager::LoadModuleChecked<IBlueprintProfilerInterface>("BlueprintProfiler");
-	const bool bUpdateHottestPaths = ProfilerModule.GetWireHeatMapDisplayMode() != EBlueprintProfilerHeatMapDisplayMode::None;
+	const bool bUpdateHottestPaths = GetDefault<UBlueprintProfilerSettings>()->WireHeatMapDisplayMode != EBlueprintProfilerHeatMapDisplayMode::None;
 	// Update instance stats
 	for (auto InstanceIter : Instances)
 	{
@@ -811,7 +944,7 @@
 			// Update the hottest path stats
 			if (bUpdateHottestPaths)
 			{
-				FScriptExecutionHottestPathParams HotPathParams(InstanceIter->GetName(), InstanceEventPerfData->GetExclusiveTiming());
+				FScriptExecutionHottestPathParams HotPathParams(InstanceIter->GetName(), InstanceEventPerfData->GetAverageTiming());
 				InstanceEventIter->CalculateHottestPathStats(HotPathParams);
 			}
 		}
@@ -828,6 +961,16 @@
 		{
 			BlueprintData->AddData(*InstanceIter.Get());
 		}
+		
+		// Now that all stats have been updated/accumulated, update heat level metrics data.
+		UpdateHeatLevelMetrics(BlueprintData);
+
+		// Finally, calculate updated heat levels based on the current heat level metrics data.
+		CalculateHeatLevelStats(HeatLevelMetrics);
+		for (auto InstanceIter : Instances)
+		{
+			InstanceIter->CalculateHeatLevelStats(HeatLevelMetrics);
+		}
 	}
 }
 
@@ -859,105 +1002,85 @@
 	}
 }
 
-//////////////////////////////////////////////////////////////////////////
-// FScriptExecutionInstance
+void FScriptExecutionBlueprint::SortEvents()
+{
+	// Sort the events
+	const auto SortEvents = [](const TSharedPtr<FScriptExecutionNode>& A, const TSharedPtr<FScriptExecutionNode>& B)
+	{
+		if (A->HasFlags(EScriptExecutionNodeFlags::InheritedEvent))
+		{
+			if (!B->HasFlags(EScriptExecutionNodeFlags::InheritedEvent))
+			{
+				return false;
+			}
+		}
+		else if (B->HasFlags(EScriptExecutionNodeFlags::InheritedEvent))
+		{
+			return true;
+		}
+		return true;
+	};
+	ChildNodes.Sort(SortEvents);
+}
 
-void FScriptExecutionInstance::NavigateToObject() const
+void FScriptExecutionBlueprint::UpdateHeatLevelMetrics(TSharedPtr<FScriptPerfData> BlueprintData)
 {
-	if (const FWorldContext* PIEWorldContext = GEditor->GetPIEWorldContext())
+	const UBlueprintProfilerSettings* ProfilerSettings = GetDefault<UBlueprintProfilerSettings>();
+	check(ProfilerSettings != nullptr);
+
+	switch (ProfilerSettings->HeatLevelMetricsType)
 	{
-		if (const AActor* PIEActor = Cast<AActor>(PIEInstance.Get()))
+	case EBlueprintProfilerHeatLevelMetricsType::ClassRelative:
+	{
+		if (BlueprintData.IsValid())
 		{
-			GEditor->SelectNone(false, false);
-			GEditor->SelectActor(const_cast<AActor*>(PIEActor), true, true, true);
+			const double AverageTiming = BlueprintData->GetAverageTiming();
+			const double InclusiveTiming = BlueprintData->GetInclusiveTiming();
+			const double MaxTiming = BlueprintData->GetMaxTiming();
+			const double TotalTiming = BlueprintData->GetTotalTiming();
+
+			HeatLevelMetrics->AveragePerformanceThreshold = AverageTiming > 0.0f ? 1.0f / AverageTiming : 0.0f;
+			HeatLevelMetrics->InclusivePerformanceThreshold = InclusiveTiming > 0.0f ? 1.0f / InclusiveTiming : 0.0f;
+			HeatLevelMetrics->MaxPerformanceThreshold = MaxTiming > 0.0f ? 1.0f / MaxTiming : 0.0f;
+			HeatLevelMetrics->TotalTimePerformanceThreshold = TotalTiming > 0.0f ? 1.0f / TotalTiming : 0.0f;
 		}
+		else
+		{
+			HeatLevelMetrics->AveragePerformanceThreshold = 0.0f;
+			HeatLevelMetrics->InclusivePerformanceThreshold = 0.0f;
+			HeatLevelMetrics->MaxPerformanceThreshold = 0.0f;
+			HeatLevelMetrics->TotalTimePerformanceThreshold = 0.0f;
+		}
+
+		HeatLevelMetrics->bUseTotalTimeWaterMark = false;
+		HeatLevelMetrics->EventPerformanceThreshold = HeatLevelMetrics->AveragePerformanceThreshold;
 	}
-	else
+		break;
+
+	case EBlueprintProfilerHeatLevelMetricsType::FrameRelative:
 	{
-		if (const AActor* EditorActor = Cast<AActor>(ObservedObject.Get()))
-		{
-			GEditor->SelectNone(false, false);
-			GEditor->SelectActor(const_cast<AActor*>(EditorActor), true, true, true);
-		}
+		// @TODO (this type is not needed for MVP)
+		HeatLevelMetrics->EventPerformanceThreshold = 0.0f;
+		HeatLevelMetrics->AveragePerformanceThreshold = 0.0f;
+		HeatLevelMetrics->InclusivePerformanceThreshold = 0.0f;
+		HeatLevelMetrics->MaxPerformanceThreshold = 0.0f;
+		HeatLevelMetrics->bUseTotalTimeWaterMark = false;
+		HeatLevelMetrics->TotalTimePerformanceThreshold = 0.0f;
+	}
+		break;
+
+	case EBlueprintProfilerHeatLevelMetricsType::CustomThresholds:
+	default:
+	{
+		HeatLevelMetrics->EventPerformanceThreshold = ProfilerSettings->CustomEventPerformanceThreshold;
+		HeatLevelMetrics->AveragePerformanceThreshold = ProfilerSettings->CustomAveragePerformanceThreshold;
+		HeatLevelMetrics->InclusivePerformanceThreshold = ProfilerSettings->CustomInclusivePerformanceThreshold;
+		HeatLevelMetrics->MaxPerformanceThreshold = ProfilerSettings->CustomMaxPerformanceThreshold;
+		HeatLevelMetrics->bUseTotalTimeWaterMark = true;
+		HeatLevelMetrics->TotalTimePerformanceThreshold = 0.0f;	// not used for this type
+	}
+		break;
 	}
 }
 
-bool FScriptExecutionInstance::IsObservedObjectValid() const
-{
-	if (const FWorldContext* PIEWorldContext = GEditor->GetPIEWorldContext())
-	{
-		if (PIEInstance.IsValid() && PIEInstance.Get()->IsPendingKill())
-		{
-			bPIEInstanceDestroyed = true;
-		}
-		return PIEInstance.IsValid();
-	}
-	else
-	{
-		return ObservedObject.IsValid() && !ObservedObject.Get()->IsPendingKill();
-	}
-}
-
-UObject* FScriptExecutionInstance::GetActiveObject() const
-{
-	UObject* Result = nullptr;
-	if (const FWorldContext* PIEWorldContext = GEditor->GetPIEWorldContext())
-	{
-		if (!PIEInstance.IsValid() && !bPIEInstanceDestroyed)
-		{
-			RefreshPIEInstance();
-		}
-		Result = const_cast<UObject*>(PIEInstance.Get());
-	}
-	else
-	{
-		Result = const_cast<UObject*>(ObservedObject.Get());
-	}
-	return Result;
-}
-
-void FScriptExecutionInstance::RefreshPIEInstance() const
-{
-	if (!bPIEInstanceDestroyed)
-	{
-		const FWorldContext* PIEWorldContext = GEditor->GetPIEWorldContext();
-		UWorld* PIEWorld = PIEWorldContext ? PIEWorldContext->World() : nullptr;
-		ULevel* CurrentLevel = PIEWorld ? PIEWorld->GetCurrentLevel() : nullptr;
-		if (CurrentLevel && ObservedObject.IsValid())
-		{
-			PIEInstance = FindObject<UObject>(CurrentLevel, *ObservedObject.Get()->GetName());
-		}
-	}
-}
-
-TSharedRef<SWidget> FScriptExecutionInstance::GetIconWidget(const uint32 /*TracePath*/)
-{
-	return SNew(SImage)
-			.Image(Icon)
-			.ColorAndOpacity(this, &FScriptExecutionInstance::GetInstanceIconColor);
-}
-
-TSharedRef<SWidget> FScriptExecutionInstance::GetHyperlinkWidget(const uint32 /*TracePath*/)
-{
-	return SNew(SHyperlink)
-		#if TRACEPATH_DEBUG
-		.Text(FText::FromName(NodeName))
-		#else
-		.Text(DisplayName)
-		#endif
-		.Style(FEditorStyle::Get(), "HoverOnlyHyperlink")
-		.ToolTipText(this, &FScriptExecutionInstance::GetInstanceTooltip)
-		.OnNavigate(this, &FScriptExecutionInstance::NavigateToObject);
-}
-
-FSlateColor FScriptExecutionInstance::GetInstanceIconColor() const
-{
-	return FSlateColor(IsObservedObjectValid() ? IconColor : FLinearColor(1.f, 1.f, 1.f, 0.3f));
-}
-
-FText FScriptExecutionInstance::GetInstanceTooltip() const
-{
-	return IsObservedObjectValid() ? Tooltip : LOCTEXT("ActorInstanceInvalid_Tooltip", "Actor instance no longer exists");
-}
-
 #undef LOCTEXT_NAMESPACE
