Copying //UE4/Dev-Blueprints to //UE4/Dev-Main (Source: //UE4/Dev-Blueprints @ 3057841)

diff --git a/Engine/Source/Editor/Kismet/Private/Profiler/SBlueprintProfilerToolbar.cpp b/Engine/Source/Editor/Kismet/Private/Profiler/SBlueprintProfilerToolbar.cpp
--- a/Engine/Source/Editor/Kismet/Private/Profiler/SBlueprintProfilerToolbar.cpp
+++ b/Engine/Source/Editor/Kismet/Private/Profiler/SBlueprintProfilerToolbar.cpp
@@ -1,17 +1,54 @@
 // Copyright 1998-2016 Epic Games, Inc. All Rights Reserved.
 
 #include "BlueprintEditorPrivatePCH.h"
+#include "BlueprintProfilerModule.h"
 #include "SBlueprintProfilerToolbar.h"
-//#include "BPProfilerStatisticWidgets.h"
 #include "EventExecution.h"
 #include "SNumericEntryBox.h"
-#include "Editor/UnrealEd/Classes/Settings/EditorExperimentalSettings.h"
 
 #define LOCTEXT_NAMESPACE "BlueprintProfilerToolbar"
 
 //////////////////////////////////////////////////////////////////////////
 // FBPProfilerStatOptions
 
+FBlueprintProfilerStatOptions::FBlueprintProfilerStatOptions()
+	: ActiveInstance(NAME_None)
+	, ActiveGraph(NAME_None)
+	, Flags(Modified)
+{
+	const UBlueprintProfilerSettings* ProfilerSettings = GetDefault<UBlueprintProfilerSettings>();
+
+	if (ProfilerSettings->bDisplayByInstance)
+	{
+		Flags |= DisplayByInstance;
+	}
+
+	if (ProfilerSettings->bScopeToDebugInstance)
+	{
+		Flags |= ScopeToDebugInstance;
+	}
+
+	if (ProfilerSettings->bGraphFilter)
+	{
+		Flags |= GraphFilter;
+	}
+
+	if (ProfilerSettings->bDisplayPure)
+	{
+		Flags |= DisplayPure;
+	}
+
+	if (ProfilerSettings->bDisplayInheritedEvents)
+	{
+		Flags |= DisplayInheritedEvents;
+	}
+
+	if (ProfilerSettings->bAverageBlueprintStats)
+	{
+		Flags |= AverageBlueprintStats;
+	}
+}
+
 void FBlueprintProfilerStatOptions::SetActiveInstance(const FName InstanceName)
 {
 	if (ActiveInstance != InstanceName && HasFlags(ScopeToDebugInstance))
@@ -30,83 +67,6 @@ void FBlueprintProfilerStatOptions::SetActiveGraph(const FName GraphName)
 	ActiveGraph = GraphName;
 }
 
-TSharedRef<SWidget> FBlueprintProfilerStatOptions::CreateToolbar()
-{
-	return	
-		SNew(SVerticalBox)
-		+SVerticalBox::Slot()
-		.HAlign(HAlign_Right)
-		.AutoHeight()
-		[
-			SNew(SHorizontalBox)
-			+SHorizontalBox::Slot()
-			.AutoWidth()
-			.Padding(FMargin(5,0))
-			[
-				SNew(SCheckBox)
-				.Content()
-				[
-					SNew(STextBlock)
-					.Text(LOCTEXT("FilterToGraph", "Filter to Graph"))
-				]
-				.IsChecked<FBlueprintProfilerStatOptions, uint32>(this, &FBlueprintProfilerStatOptions::GetChecked, GraphFilter)
-				.OnCheckStateChanged<FBlueprintProfilerStatOptions, uint32>(this, &FBlueprintProfilerStatOptions::OnChecked, GraphFilter)
-			]
-			+SHorizontalBox::Slot()
-			.AutoWidth()
-			.Padding(FMargin(5,0))
-			[
-				SNew(SCheckBox)
-				.Content()
-				[
-					SNew(STextBlock)
-					.Text(LOCTEXT("DisplayPureStats", "Pure Timings"))
-				]
-				.IsChecked<FBlueprintProfilerStatOptions, uint32>(this, &FBlueprintProfilerStatOptions::GetChecked, DisplayPure)
-				.OnCheckStateChanged<FBlueprintProfilerStatOptions, uint32>(this, &FBlueprintProfilerStatOptions::OnChecked, DisplayPure)
-			]
-			+SHorizontalBox::Slot()
-			.AutoWidth()
-			.Padding(FMargin(5,0))
-			[
-				SNew(SCheckBox)
-				.Content()
-				[
-					SNew(STextBlock)
-					.Text(LOCTEXT("ShowInstancesCheck", "Show Instances"))
-				]
-				.IsChecked<FBlueprintProfilerStatOptions, uint32>(this, &FBlueprintProfilerStatOptions::GetChecked, DisplayByInstance)
-				.OnCheckStateChanged<FBlueprintProfilerStatOptions, uint32>(this, &FBlueprintProfilerStatOptions::OnChecked, DisplayByInstance)
-			]
-			+SHorizontalBox::Slot()
-			.AutoWidth()
-			.Padding(FMargin(5,0))
-			[
-				SNew(SCheckBox)
-				.Content()
-				[
-					SNew(STextBlock)
-					.Text(LOCTEXT("InstanceFilterCheck", "Debug Filter Scope"))
-				]
-				.IsChecked<FBlueprintProfilerStatOptions, uint32>(this, &FBlueprintProfilerStatOptions::GetChecked, ScopeToDebugInstance)
-				.OnCheckStateChanged<FBlueprintProfilerStatOptions, uint32>(this, &FBlueprintProfilerStatOptions::OnChecked, ScopeToDebugInstance)
-			]
-			//+SHorizontalBox::Slot()
-			//.AutoWidth()
-			//.Padding(FMargin(5,0))
-			//[
-			//	SNew(SCheckBox)
-			//	.Content()
-			//	[
-			//		SNew(STextBlock)
-			//		.Text(LOCTEXT("AutoItemExpansion", "Auto Expand Statistics"))
-			//	]
-			//	.IsChecked<FBlueprintProfilerStatOptions, uint32>(this, &FBlueprintProfilerStatOptions::GetChecked, AutoExpand)
-			//	.OnCheckStateChanged<FBlueprintProfilerStatOptions, uint32>(this, &FBlueprintProfilerStatOptions::OnChecked, AutoExpand)
-			//]
-		];
-}
-
 TSharedRef<SWidget> FBlueprintProfilerStatOptions::CreateComboContent()
 {
 	return	
@@ -151,6 +111,7 @@ TSharedRef<SWidget> FBlueprintProfilerStatOptions::CreateComboContent()
 				.OnCheckStateChanged<FBlueprintProfilerStatOptions, uint32>(this, &FBlueprintProfilerStatOptions::OnChecked, AverageBlueprintStats)
 			]
 		]
+#if 0 // Removing pure timings until we have an advanced display options.
 		+SVerticalBox::Slot()
 		.HAlign(HAlign_Left)
 		.AutoHeight()
@@ -171,6 +132,7 @@ TSharedRef<SWidget> FBlueprintProfilerStatOptions::CreateComboContent()
 				.OnCheckStateChanged<FBlueprintProfilerStatOptions, uint32>(this, &FBlueprintProfilerStatOptions::OnChecked, DisplayPure)
 			]
 		]
+#endif
 		+SVerticalBox::Slot()
 		.HAlign(HAlign_Left)
 		.AutoHeight()
@@ -184,11 +146,11 @@ TSharedRef<SWidget> FBlueprintProfilerStatOptions::CreateComboContent()
 				.Content()
 				[
 					SNew(STextBlock)
-					.Text(LOCTEXT("ShowInstancesCheck_Label", "Show Instances"))
-					.ToolTipText(LOCTEXT("ShowInstancesCheck_Tooltip", "Display statistics by instance rather than by blueprint."))
+					.Text(LOCTEXT("DisplayInheritedEvents_Label", "Show Inherited Blueprint Events"))
+					.ToolTipText(LOCTEXT("DisplayInheritedEvents_Tooltip", "Display events for parent blueprints."))
 				]
-				.IsChecked<FBlueprintProfilerStatOptions, uint32>(this, &FBlueprintProfilerStatOptions::GetChecked, DisplayByInstance)
-				.OnCheckStateChanged<FBlueprintProfilerStatOptions, uint32>(this, &FBlueprintProfilerStatOptions::OnChecked, DisplayByInstance)
+				.IsChecked<FBlueprintProfilerStatOptions, uint32>(this, &FBlueprintProfilerStatOptions::GetChecked, DisplayInheritedEvents)
+				.OnCheckStateChanged<FBlueprintProfilerStatOptions, uint32>(this, &FBlueprintProfilerStatOptions::OnChecked, DisplayInheritedEvents)
 			]
 		]
 		+SVerticalBox::Slot()
@@ -196,7 +158,6 @@ TSharedRef<SWidget> FBlueprintProfilerStatOptions::CreateComboContent()
 		.AutoHeight()
 		[
 			SNew(SHorizontalBox)
-			.Visibility<FBlueprintProfilerStatOptions, uint32>(this, &FBlueprintProfilerStatOptions::GetVisibility, ScopeToDebugInstance)
 			+SHorizontalBox::Slot()
 			.AutoWidth()
 			.Padding(FMargin(5,2))
@@ -205,19 +166,19 @@ TSharedRef<SWidget> FBlueprintProfilerStatOptions::CreateComboContent()
 				.Content()
 				[
 					SNew(STextBlock)
-					.Text(LOCTEXT("InstanceFilterCheck_Label", "Debug Filter Scope"))
-					.ToolTipText(LOCTEXT("InstanceFilterCheck_Tooltip", "Restrict the displayed instances to match the blueprint debug filter."))
+					.Text(LOCTEXT("ShowInstancesCheck_Label", "Show Instances"))
+					.ToolTipText(LOCTEXT("ShowInstancesCheck_Tooltip", "Display statistics by instance rather than by blueprint."))
 				]
-				.IsChecked<FBlueprintProfilerStatOptions, uint32>(this, &FBlueprintProfilerStatOptions::GetChecked, ScopeToDebugInstance)
-				.OnCheckStateChanged<FBlueprintProfilerStatOptions, uint32>(this, &FBlueprintProfilerStatOptions::OnChecked, ScopeToDebugInstance)
+				.IsChecked<FBlueprintProfilerStatOptions, uint32>(this, &FBlueprintProfilerStatOptions::GetChecked, DisplayByInstance)
+				.OnCheckStateChanged<FBlueprintProfilerStatOptions, uint32>(this, &FBlueprintProfilerStatOptions::OnChecked, DisplayByInstance)
 			]
-#if 0
 		]
 		+SVerticalBox::Slot()
 		.HAlign(HAlign_Left)
 		.AutoHeight()
 		[
 			SNew(SHorizontalBox)
+			.Visibility<FBlueprintProfilerStatOptions, uint32>(this, &FBlueprintProfilerStatOptions::GetVisibility, ScopeToDebugInstance)
 			+SHorizontalBox::Slot()
 			.AutoWidth()
 			.Padding(FMargin(5,2))
@@ -226,13 +187,12 @@ TSharedRef<SWidget> FBlueprintProfilerStatOptions::CreateComboContent()
 				.Content()
 				[
 					SNew(STextBlock)
-					.Text(LOCTEXT("AutoItemExpansion_Label", "Auto Expand Statistics"))
-					.ToolTipText(LOCTEXT("AutoItemExpansion_Tooltip", "Auto expand the statistic trees."))
+					.Text(LOCTEXT("InstanceFilterCheck_Label", "Debug Filter Scope"))
+					.ToolTipText(LOCTEXT("InstanceFilterCheck_Tooltip", "Restrict the displayed instances to match the blueprint debug filter."))
 				]
-				.IsChecked<FBlueprintProfilerStatOptions, uint32>(this, &FBlueprintProfilerStatOptions::GetChecked, AutoExpand)
-				.OnCheckStateChanged<FBlueprintProfilerStatOptions, uint32>(this, &FBlueprintProfilerStatOptions::OnChecked, AutoExpand)
+				.IsChecked<FBlueprintProfilerStatOptions, uint32>(this, &FBlueprintProfilerStatOptions::GetChecked, ScopeToDebugInstance)
+				.OnCheckStateChanged<FBlueprintProfilerStatOptions, uint32>(this, &FBlueprintProfilerStatOptions::OnChecked, ScopeToDebugInstance)
 			]
-#endif
 		];
 }
 
@@ -267,11 +227,47 @@ void FBlueprintProfilerStatOptions::OnChecked(ECheckBoxState NewState, const uin
 	{
 		Flags &= ~FlagsIn;
 	}
+
 	if (FlagsIn & AverageBlueprintStats)
 	{
 		FScriptPerfData::EnableBlueprintStatAverage(NewState == ECheckBoxState::Checked);
 	}
+
 	Flags |= Modified;
+
+	UBlueprintProfilerSettings* ProfilerSettings = GetMutableDefault<UBlueprintProfilerSettings>();
+
+	if (FlagsIn & DisplayByInstance)
+	{
+		ProfilerSettings->bDisplayByInstance = !!(Flags & DisplayByInstance);
+	}
+
+	if (FlagsIn & ScopeToDebugInstance)
+	{
+		ProfilerSettings->bScopeToDebugInstance = !!(Flags & ScopeToDebugInstance);
+	}
+
+	if (FlagsIn & GraphFilter)
+	{
+		ProfilerSettings->bGraphFilter = !!(Flags & GraphFilter);
+	}
+
+	if (FlagsIn & DisplayPure)
+	{
+		ProfilerSettings->bDisplayPure = !!(Flags & DisplayPure);
+	}
+
+	if (FlagsIn & DisplayInheritedEvents)
+	{
+		ProfilerSettings->bDisplayInheritedEvents = !!(Flags & DisplayInheritedEvents);
+	}
+
+	if (FlagsIn & AverageBlueprintStats)
+	{
+		ProfilerSettings->bAverageBlueprintStats = !!(Flags & AverageBlueprintStats);
+	}
+
+	ProfilerSettings->SaveConfig();
 }
 
 EVisibility FBlueprintProfilerStatOptions::GetVisibility(const uint32 FlagsIn) const
@@ -287,15 +283,22 @@ EVisibility FBlueprintProfilerStatOptions::GetVisibility(const uint32 FlagsIn) c
 bool FBlueprintProfilerStatOptions::IsFiltered(TSharedPtr<FScriptExecutionNode> Node) const
 {
 	bool bFilteredOut = !HasFlags(EDisplayFlags::DisplayPure) && Node->HasFlags(EScriptExecutionNodeFlags::PureStats);
-	if (Node->IsEvent() && HasFlags(EDisplayFlags::GraphFilter))
+	if (Node->IsEvent())
 	{
-		if (Node->GetGraphName() == UEdGraphSchema_K2::FN_UserConstructionScript)
+		if (HasFlags(EDisplayFlags::GraphFilter))
 		{
-			bFilteredOut = ActiveGraph != UEdGraphSchema_K2::FN_UserConstructionScript;
+			if (Node->GetGraphName() == UEdGraphSchema_K2::FN_UserConstructionScript)
+			{
+				bFilteredOut = ActiveGraph != UEdGraphSchema_K2::FN_UserConstructionScript;
+			}
+			else
+			{
+				bFilteredOut = ActiveGraph == UEdGraphSchema_K2::FN_UserConstructionScript;
+			}
 		}
-		else
+		if (!HasFlags(EDisplayFlags::DisplayInheritedEvents) && Node->HasFlags(EScriptExecutionNodeFlags::InheritedEvent))
 		{
-			bFilteredOut = ActiveGraph == UEdGraphSchema_K2::FN_UserConstructionScript;
+			bFilteredOut = true;
 		}
 	}
 	return bFilteredOut;
@@ -308,6 +311,7 @@ void SBlueprintProfilerToolbar::Construct(const FArguments& InArgs)
 {
 	CurrentViewType = InArgs._ProfileViewType;
 	DisplayOptions = InArgs._DisplayOptions;
+	
 	ViewChangeDelegate = InArgs._OnViewChanged;
 
 	FMenuBuilder ViewComboContent(true, NULL);
@@ -341,13 +345,13 @@ void SBlueprintProfilerToolbar::Construct(const FArguments& InArgs)
 		),
 		NAME_None,
 		EUserInterfaceActionType::Check);
-	HeatMapDisplayModeComboContent.AddMenuEntry(LOCTEXT("HeatMapDisplayMode_Exclusive", "Exclusive"),
-		LOCTEXT("HeatMapDisplayMode_Exclusive_Desc", "Heat Map - Exclusive Timings"),
+	HeatMapDisplayModeComboContent.AddMenuEntry(LOCTEXT("HeatMapDisplayMode_Average", "Average"),
+		LOCTEXT("HeatMapDisplayMode_Average_Desc", "Heat Map - Average Timings"),
 		FSlateIcon(),
 		FUIAction(
-			FExecuteAction::CreateSP(this, &SBlueprintProfilerToolbar::OnHeatMapDisplayModeChanged, EHeatmapControlId::NodeHeatmap, EBlueprintProfilerHeatMapDisplayMode::Exclusive),
+			FExecuteAction::CreateSP(this, &SBlueprintProfilerToolbar::OnHeatMapDisplayModeChanged, EHeatmapControlId::NodeHeatmap, EBlueprintProfilerHeatMapDisplayMode::Average),
 			FCanExecuteAction(),
-			FIsActionChecked::CreateSP(this, &SBlueprintProfilerToolbar::IsHeatMapDisplayModeSelected, EHeatmapControlId::NodeHeatmap, EBlueprintProfilerHeatMapDisplayMode::Exclusive)
+			FIsActionChecked::CreateSP(this, &SBlueprintProfilerToolbar::IsHeatMapDisplayModeSelected, EHeatmapControlId::NodeHeatmap, EBlueprintProfilerHeatMapDisplayMode::Average)
 		),
 		NAME_None,
 		EUserInterfaceActionType::Check);
@@ -403,16 +407,6 @@ void SBlueprintProfilerToolbar::Construct(const FArguments& InArgs)
 		),
 		NAME_None,
 		EUserInterfaceActionType::Check);
-	WireHeatMapDisplayModeComboContent.AddMenuEntry(LOCTEXT("WireHeatMapDisplayMode_HottestEndpoint", "Hottest Endpoints"),
-		LOCTEXT("WireHeatMapDisplayMode_HottestEndpoints_Desc", "Heat Map - Hottest Endpoints"),
-		FSlateIcon(),
-		FUIAction(
-			FExecuteAction::CreateSP(this, &SBlueprintProfilerToolbar::OnHeatMapDisplayModeChanged, EHeatmapControlId::WireHeatmap, EBlueprintProfilerHeatMapDisplayMode::HottestEndpoint),
-			FCanExecuteAction(),
-			FIsActionChecked::CreateSP(this, &SBlueprintProfilerToolbar::IsHeatMapDisplayModeSelected, EHeatmapControlId::WireHeatmap, EBlueprintProfilerHeatMapDisplayMode::HottestEndpoint)
-		),
-		NAME_None,
-		EUserInterfaceActionType::Check);
 	WireHeatMapDisplayModeComboContent.AddMenuEntry(LOCTEXT("WireHeatMapDisplayMode_PinToPin", "Pin To Pin Heatmap"),
 		LOCTEXT("WireHeatMapDisplayMode_PinToPin_Desc", "Heat Map - Pin To Pin"),
 		FSlateIcon(),
@@ -458,21 +452,21 @@ void SBlueprintProfilerToolbar::Construct(const FArguments& InArgs)
 				.Orientation(Orient_Vertical)
 			]
 			+SHorizontalBox::Slot()
-			.HAlign(HAlign_Right)
 			.AutoWidth()
 			.Padding(FMargin(5, 0))
 			[
-				SAssignNew(HeatThresholdComboButton, SComboButton)
-				.ForegroundColor(this, &SBlueprintProfilerToolbar::GetButtonForegroundColor, EHeatmapControlId::HeatThresholds)
+				SAssignNew(HeatMapDisplayModeComboButton, SComboButton)
+				.ForegroundColor(this, &SBlueprintProfilerToolbar::GetButtonForegroundColor, EHeatmapControlId::NodeHeatmap)
 				.ToolTipText(LOCTEXT("HeatMapDisplayMode_Tooltip", "Heat Map Display Mode"))
 				.ButtonStyle(FEditorStyle::Get(), "ToggleButton")
 				.ContentPadding(2)
-				.OnGetMenuContent(this, &SBlueprintProfilerToolbar::CreateHeatThresholdsMenuContent)
+				.MenuContent()
+				[
+					HeatMapDisplayModeComboContent.MakeWidget()
+				]
 				.ButtonContent()
 				[
-					SNew(STextBlock)
-					.Text(LOCTEXT("HeatmapThresholdsCombo_Label", "Heatmap Thresholds"))
-					.ColorAndOpacity(FLinearColor::White)
+					CreateHeatMapDisplayModeButton(EHeatmapControlId::NodeHeatmap)
 				]
 			]
 			+SHorizontalBox::Slot()
@@ -485,18 +479,18 @@ void SBlueprintProfilerToolbar::Construct(const FArguments& InArgs)
 			.AutoWidth()
 			.Padding(FMargin(5, 0))
 			[
-				SAssignNew(HeatMapDisplayModeComboButton, SComboButton)
-				.ForegroundColor(this, &SBlueprintProfilerToolbar::GetButtonForegroundColor, EHeatmapControlId::NodeHeatmap)
-				.ToolTipText(LOCTEXT("HeatMapDisplayMode_Tooltip", "Heat Map Display Mode"))
+				SAssignNew(WireHeatMapDisplayModeComboButton, SComboButton)
+				.ForegroundColor(this, &SBlueprintProfilerToolbar::GetButtonForegroundColor, EHeatmapControlId::WireHeatmap)
+				.ToolTipText(LOCTEXT("WireHeatMapDisplayMode_Tooltip", "Wire Heat Map Display Mode"))
 				.ButtonStyle(FEditorStyle::Get(), "ToggleButton")
 				.ContentPadding(2)
 				.MenuContent()
 				[
-					HeatMapDisplayModeComboContent.MakeWidget()
+					WireHeatMapDisplayModeComboContent.MakeWidget()
 				]
 				.ButtonContent()
 				[
-					CreateHeatMapDisplayModeButton(EHeatmapControlId::NodeHeatmap)
+					CreateHeatMapDisplayModeButton(EHeatmapControlId::WireHeatmap)
 				]
 			]
 			+SHorizontalBox::Slot()
@@ -506,23 +500,77 @@ void SBlueprintProfilerToolbar::Construct(const FArguments& InArgs)
 				.Orientation(Orient_Vertical)
 			]
 			+SHorizontalBox::Slot()
+			.VAlign(VAlign_Center)
+			.HAlign(HAlign_Right)
 			.AutoWidth()
 			.Padding(FMargin(5, 0))
 			[
-				SAssignNew(WireHeatMapDisplayModeComboButton, SComboButton)
-				.ForegroundColor(this, &SBlueprintProfilerToolbar::GetButtonForegroundColor, EHeatmapControlId::WireHeatmap)
-				.ToolTipText(LOCTEXT("WireHeatMapDisplayMode_Tooltip", "Wire Heat Map Display Mode"))
-				.ButtonStyle(FEditorStyle::Get(), "ToggleButton")
-				.ContentPadding(2)
-				.MenuContent()
+				SNew(STextBlock)
+				.Text(LOCTEXT("HeatmapThresholds_Label", "Heatmap Thresholds:"))
+				.ColorAndOpacity(FLinearColor::White)
+			]
+			+SHorizontalBox::Slot()
+			.HAlign(HAlign_Right)
+			.AutoWidth()
+			.Padding(FMargin(5, 2))
+			[
+				SNew(SUniformGridPanel)
+				+SUniformGridPanel::Slot(0, 0)
 				[
-					WireHeatMapDisplayModeComboContent.MakeWidget()
+					SNew(SCheckBox)
+					.Style(FEditorStyle::Get(), "BlueprintProfiler.ToggleButton.Start")
+					.IsChecked(this, &SBlueprintProfilerToolbar::IsHeatLevelMetricsTypeSelected, EBlueprintProfilerHeatLevelMetricsType::ClassRelative)
+					.OnCheckStateChanged(this, &SBlueprintProfilerToolbar::OnHeatLevelMetricsTypeChanged, EBlueprintProfilerHeatLevelMetricsType::ClassRelative)
+					.ToolTipText(LOCTEXT("HeatLevelMetricsType_Default_Tooltip", "Heat map thresholds are automatically calculated relative to overall Blueprint execution time."))
+					[
+						SNew(SHorizontalBox)
+						+SHorizontalBox::Slot()
+						.VAlign(VAlign_Center)
+						.HAlign(HAlign_Center)
+						.Padding(6, 2)
+						[
+							SNew(STextBlock)
+							.Text(LOCTEXT("Default", "Default"))
+							.Font(FEditorStyle::GetFontStyle(TEXT("PropertyWindow.NormalFont")))
+							.ColorAndOpacity(this, &SBlueprintProfilerToolbar::GetHeatLevelMetricsTypeTextColor, EBlueprintProfilerHeatLevelMetricsType::ClassRelative)
+						]
+					]
 				]
-				.ButtonContent()
+				+SUniformGridPanel::Slot(1, 0)
 				[
-					CreateHeatMapDisplayModeButton(EHeatmapControlId::WireHeatmap)
+					SNew(SCheckBox)
+					.Style(FEditorStyle::Get(), "BlueprintProfiler.ToggleButton.End")
+					.IsChecked(this, &SBlueprintProfilerToolbar::IsHeatLevelMetricsTypeSelected, EBlueprintProfilerHeatLevelMetricsType::CustomThresholds)
+					.OnCheckStateChanged(this, &SBlueprintProfilerToolbar::OnHeatLevelMetricsTypeChanged, EBlueprintProfilerHeatLevelMetricsType::CustomThresholds)
+					.ToolTipText(LOCTEXT("HeatLevelMetricsType_Custom_Tooltip", "Heat map thresholds are explicitly defined."))
+					[
+						SNew(SHorizontalBox)
+						+SHorizontalBox::Slot()
+						.VAlign(VAlign_Center)
+						.HAlign(HAlign_Center)
+						.Padding(6, 2)
+						[
+							SNew(STextBlock)
+							.Text(LOCTEXT("Custom", "Custom"))
+							.Font(FEditorStyle::GetFontStyle(TEXT("PropertyWindow.NormalFont")))
+							.ColorAndOpacity(this, &SBlueprintProfilerToolbar::GetHeatLevelMetricsTypeTextColor, EBlueprintProfilerHeatLevelMetricsType::CustomThresholds)
+						]
+					]
 				]
 			]
+			+SHorizontalBox::Slot()
+			.HAlign(HAlign_Right)
+			.AutoWidth()
+			.Padding(FMargin(1, 0))
+			[
+				SAssignNew(CustomHeatThresholdComboButton, SComboButton)
+				.IsEnabled(this, &SBlueprintProfilerToolbar::IsCustomHeatThresholdsMenuEnabled)
+				.ForegroundColor(this, &SBlueprintProfilerToolbar::GetButtonForegroundColor, EHeatmapControlId::HeatThresholds)
+				.ToolTipText(LOCTEXT("HeatMapCustomThresholdsButton_Tooltip", "Custom Thresholds"))
+				.ButtonStyle(FEditorStyle::Get(), "ToggleButton")
+				.ContentPadding(2)
+				.OnGetMenuContent(this, &SBlueprintProfilerToolbar::CreateCustomHeatThresholdsMenuContent)
+			]
 #if 0
 			+SHorizontalBox::Slot()
 			.Padding(FMargin(5, 0))
@@ -608,16 +656,15 @@ FText SBlueprintProfilerToolbar::GetCurrentHeatMapDisplayModeText(const EHeatmap
 {
 	FText Result = LOCTEXT("HeatMapDisplayModeLabel_Unknown", "<unknown>");
 	IBlueprintProfilerInterface& ProfilerModule = FModuleManager::LoadModuleChecked<IBlueprintProfilerInterface>("BlueprintProfiler");
-	EBlueprintProfilerHeatMapDisplayMode::Type DisplayMode = ControlId == EHeatmapControlId::NodeHeatmap ? ProfilerModule.GetGraphNodeHeatMapDisplayMode() : ProfilerModule.GetWireHeatMapDisplayMode();
+	EBlueprintProfilerHeatMapDisplayMode DisplayMode = ControlId == EHeatmapControlId::NodeHeatmap ? GetDefault<UBlueprintProfilerSettings>()->GraphNodeHeatMapDisplayMode : GetDefault<UBlueprintProfilerSettings>()->WireHeatMapDisplayMode;
 	switch (DisplayMode)
 	{
 		case EBlueprintProfilerHeatMapDisplayMode::None:			Result = LOCTEXT("HeatMapDisplayModeLabel_None", "Off"); break;
-		case EBlueprintProfilerHeatMapDisplayMode::Exclusive:		Result = LOCTEXT("HeatMapDisplayModeLabel_Exclusive", "Exclusive"); break;
+		case EBlueprintProfilerHeatMapDisplayMode::Average:			Result = LOCTEXT("HeatMapDisplayModeLabel_Average", "Average"); break;
 		case EBlueprintProfilerHeatMapDisplayMode::Inclusive:		Result = LOCTEXT("HeatMapDisplayModeLabel_Inclusive", "Inclusive"); break;
 		case EBlueprintProfilerHeatMapDisplayMode::MaxTiming:		Result = LOCTEXT("HeatMapDisplayModeLabel_MaxTiming", "Max Timing"); break;
 		case EBlueprintProfilerHeatMapDisplayMode::Total:			Result = LOCTEXT("HeatMapDisplayModeLabel_TotalTiming", "Total Timing"); break;
 		case EBlueprintProfilerHeatMapDisplayMode::HottestPath:		Result = LOCTEXT("HeatMapDisplayModeLabel_HottestPath", "Hottest Path"); break;
-		case EBlueprintProfilerHeatMapDisplayMode::HottestEndpoint:	Result = LOCTEXT("HeatMapDisplayModeLabel_HottestEndpoint", "Hottest Endpoints"); break;
 		case EBlueprintProfilerHeatMapDisplayMode::PinToPin:		Result = LOCTEXT("HeatMapDisplayModeLabel_PinToPin", "Pin to Pin Heatmap"); break;
 	}
 	return Result;
@@ -659,38 +706,78 @@ FSlateColor SBlueprintProfilerToolbar::GetButtonForegroundColor(const EHeatmapCo
 		case EHeatmapControlId::GeneralOptions:		bButtonHovered = GeneralOptionsComboButton->IsHovered(); break;
 		case EHeatmapControlId::NodeHeatmap:		bButtonHovered = HeatMapDisplayModeComboButton->IsHovered(); break;
 		case EHeatmapControlId::WireHeatmap:		bButtonHovered = WireHeatMapDisplayModeComboButton->IsHovered(); break;
-		case EHeatmapControlId::HeatThresholds:		bButtonHovered = HeatThresholdComboButton->IsHovered();	break;
+		case EHeatmapControlId::HeatThresholds:		bButtonHovered = CustomHeatThresholdComboButton->IsHovered();	break;
 	}
 	return bButtonHovered ? FEditorStyle::GetSlateColor(InvertedForegroundName) : FEditorStyle::GetSlateColor(DefaultForegroundName);
 }
 
-bool SBlueprintProfilerToolbar::IsHeatMapDisplayModeSelected(const EHeatmapControlId::Type ControlId, const EBlueprintProfilerHeatMapDisplayMode::Type InHeatMapDisplayMode) const
+ECheckBoxState SBlueprintProfilerToolbar::IsHeatLevelMetricsTypeSelected(EBlueprintProfilerHeatLevelMetricsType InHeatLevelMetricsType) const
 {
-	IBlueprintProfilerInterface& ProfilerModule = FModuleManager::LoadModuleChecked<IBlueprintProfilerInterface>("BlueprintProfiler");
-	return ControlId == EHeatmapControlId::NodeHeatmap ? (ProfilerModule.GetGraphNodeHeatMapDisplayMode() == InHeatMapDisplayMode) :
-														 (ProfilerModule.GetWireHeatMapDisplayMode() == InHeatMapDisplayMode);
+	return GetDefault<UBlueprintProfilerSettings>()->HeatLevelMetricsType == InHeatLevelMetricsType ? ECheckBoxState::Checked : ECheckBoxState::Unchecked;
 }
 
-void SBlueprintProfilerToolbar::OnHeatMapDisplayModeChanged(const EHeatmapControlId::Type ControlId, const EBlueprintProfilerHeatMapDisplayMode::Type NewHeatMapDisplayMode)
+void SBlueprintProfilerToolbar::OnHeatLevelMetricsTypeChanged(ECheckBoxState InCheckedState, const EBlueprintProfilerHeatLevelMetricsType NewHeatLevelMetricsType)
 {
-	IBlueprintProfilerInterface& ProfilerModule = FModuleManager::LoadModuleChecked<IBlueprintProfilerInterface>("BlueprintProfiler");
+	if (InCheckedState == ECheckBoxState::Checked)
+	{
+		UBlueprintProfilerSettings* Settings = GetMutableDefault<UBlueprintProfilerSettings>();
+		Settings->HeatLevelMetricsType = NewHeatLevelMetricsType;
+
+		Settings->SaveConfig();
+	}
+}
+
+FSlateColor SBlueprintProfilerToolbar::GetHeatLevelMetricsTypeTextColor(EBlueprintProfilerHeatLevelMetricsType InHeatLevelMetricsType) const
+{
+	const UBlueprintProfilerSettings* Settings = GetDefault<UBlueprintProfilerSettings>();
+	if (Settings->HeatLevelMetricsType == InHeatLevelMetricsType)
+	{
+		return FSlateColor(FLinearColor::Black);
+	}
+	else
+	{
+		return GetForegroundColor();
+	}
+}
+
+bool SBlueprintProfilerToolbar::IsHeatMapDisplayModeSelected(const EHeatmapControlId::Type ControlId, const EBlueprintProfilerHeatMapDisplayMode InHeatMapDisplayMode) const
+{
+	const UBlueprintProfilerSettings* Settings = GetDefault<UBlueprintProfilerSettings>();
+	return ControlId == EHeatmapControlId::NodeHeatmap ? (Settings->GraphNodeHeatMapDisplayMode == InHeatMapDisplayMode) :
+														 (Settings->WireHeatMapDisplayMode == InHeatMapDisplayMode);
+}
+
+void SBlueprintProfilerToolbar::OnHeatMapDisplayModeChanged(const EHeatmapControlId::Type ControlId, const EBlueprintProfilerHeatMapDisplayMode NewHeatMapDisplayMode)
+{
+	UBlueprintProfilerSettings* Settings = GetMutableDefault<UBlueprintProfilerSettings>();
 	if (ControlId == EHeatmapControlId::NodeHeatmap)
 	{
-		if (ProfilerModule.GetGraphNodeHeatMapDisplayMode() != NewHeatMapDisplayMode)
+		if (Settings->GraphNodeHeatMapDisplayMode != NewHeatMapDisplayMode)
 		{
-			ProfilerModule.SetGraphNodeHeatMapDisplayMode(NewHeatMapDisplayMode);
+			Settings->GraphNodeHeatMapDisplayMode = NewHeatMapDisplayMode;
+
+			// @TODO - should we save this setting?
+			//Settings->SaveConfig();
 		}
 	}
 	else 
 	{
-		if (ProfilerModule.GetWireHeatMapDisplayMode() != NewHeatMapDisplayMode)
+		if (Settings->WireHeatMapDisplayMode != NewHeatMapDisplayMode)
 		{
-			ProfilerModule.SetWireHeatMapDisplayMode(NewHeatMapDisplayMode);
+			Settings->WireHeatMapDisplayMode = NewHeatMapDisplayMode;
+
+			// @TODO - should we save this setting?
+			//Settings->SaveConfig();
 		}
 	}
 }
 
-TSharedRef<SWidget> SBlueprintProfilerToolbar::CreateHeatThresholdsMenuContent()
+bool SBlueprintProfilerToolbar::IsCustomHeatThresholdsMenuEnabled() const
+{
+	return GetDefault<UBlueprintProfilerSettings>()->HeatLevelMetricsType == EBlueprintProfilerHeatLevelMetricsType::CustomThresholds;
+}
+
+TSharedRef<SWidget> SBlueprintProfilerToolbar::CreateCustomHeatThresholdsMenuContent()
 {
 	return SNew(SVerticalBox)
 		+SVerticalBox::Slot()
@@ -723,8 +810,8 @@ TSharedRef<SWidget> SBlueprintProfilerToolbar::CreateHeatThresholdsMenuContent()
 					.MaxValue(5.f)
 					.MinSliderValue(0.01f)
 					.MaxSliderValue(5.f)
-					.Value(this, &SBlueprintProfilerToolbar::GetEventThreshold)
-					.OnValueChanged(this, &SBlueprintProfilerToolbar::SetEventThreshold)
+					.Value(this, &SBlueprintProfilerToolbar::GetCustomPerformanceThreshold, ECustomPerformanceThreshold::Event)
+					.OnValueChanged(this, &SBlueprintProfilerToolbar::SetCustomPerformanceThreshold, ECustomPerformanceThreshold::Event)
 					.MinDesiredValueWidth(100)
 					.ToolTipText(LOCTEXT("EventPerformanceThreshold_Tooltip", "The Hot Threshold for Overall Event Timings"))
 				]
@@ -736,8 +823,8 @@ TSharedRef<SWidget> SBlueprintProfilerToolbar::CreateHeatThresholdsMenuContent()
 					.ToolTipText(LOCTEXT("ResetToDefaultToolTip", "Reset to Default"))
 					.ButtonStyle(FEditorStyle::Get(), "NoBorder")
 					.ContentPadding(4)
-					.Visibility(this, &SBlueprintProfilerToolbar::IsEventDefaultButtonVisible)
-					.OnClicked(this, &SBlueprintProfilerToolbar::ResetEventThreshold)
+					.Visibility(this, &SBlueprintProfilerToolbar::IsCustomPerformanceThresholdDefaultButtonVisible, ECustomPerformanceThreshold::Event)
+					.OnClicked(this, &SBlueprintProfilerToolbar::ResetCustomPerformanceThreshold, ECustomPerformanceThreshold::Event)
 					.Content()
 					[
 						SNew(SImage)
@@ -776,8 +863,8 @@ TSharedRef<SWidget> SBlueprintProfilerToolbar::CreateHeatThresholdsMenuContent()
 					.MaxValue(0.5f)
 					.MinSliderValue(0.001f)
 					.MaxSliderValue(0.5f)
-					.Value(this, &SBlueprintProfilerToolbar::GetNodeExclusiveThreshold)
-					.OnValueChanged(this, &SBlueprintProfilerToolbar::SetNodeExclusiveThreshold)
+					.Value(this, &SBlueprintProfilerToolbar::GetCustomPerformanceThreshold, ECustomPerformanceThreshold::Average)
+					.OnValueChanged(this, &SBlueprintProfilerToolbar::SetCustomPerformanceThreshold, ECustomPerformanceThreshold::Average)
 					.MinDesiredValueWidth(100)
 					.ToolTipText(LOCTEXT("NodeAveragePerformanceThreshold_Tooltip", "The Hot Threshold for Node Average Timings"))
 				]
@@ -789,8 +876,8 @@ TSharedRef<SWidget> SBlueprintProfilerToolbar::CreateHeatThresholdsMenuContent()
 					.ToolTipText(LOCTEXT("ResetToDefaultToolTip", "Reset to Default"))
 					.ButtonStyle(FEditorStyle::Get(), "NoBorder")
 					.ContentPadding(4)
-					.Visibility(this, &SBlueprintProfilerToolbar::IsNodeExclDefaultButtonVisible)
-					.OnClicked(this, &SBlueprintProfilerToolbar::ResetNodeExclThreshold)
+					.Visibility(this, &SBlueprintProfilerToolbar::IsCustomPerformanceThresholdDefaultButtonVisible, ECustomPerformanceThreshold::Average)
+					.OnClicked(this, &SBlueprintProfilerToolbar::ResetCustomPerformanceThreshold, ECustomPerformanceThreshold::Average)
 					.Content()
 					[
 						SNew(SImage)
@@ -829,8 +916,8 @@ TSharedRef<SWidget> SBlueprintProfilerToolbar::CreateHeatThresholdsMenuContent()
 					.MaxValue(0.5f)
 					.MinSliderValue(0.001f)
 					.MaxSliderValue(0.5f)
-					.Value(this, &SBlueprintProfilerToolbar::GetNodeInclusiveThreshold)
-					.OnValueChanged(this, &SBlueprintProfilerToolbar::SetNodeInclusiveThreshold)
+					.Value(this, &SBlueprintProfilerToolbar::GetCustomPerformanceThreshold, ECustomPerformanceThreshold::Inclusive)
+					.OnValueChanged(this, &SBlueprintProfilerToolbar::SetCustomPerformanceThreshold, ECustomPerformanceThreshold::Inclusive)
 					.MinDesiredValueWidth(100)
 					.ToolTipText(LOCTEXT("NodeInclusivePerformanceThreshold_Tooltip", "The Hot Threshold for Node Inclusive Timings"))
 				]
@@ -842,8 +929,8 @@ TSharedRef<SWidget> SBlueprintProfilerToolbar::CreateHeatThresholdsMenuContent()
 					.ToolTipText(LOCTEXT("ResetToDefaultToolTip", "Reset to Default"))
 					.ButtonStyle(FEditorStyle::Get(), "NoBorder")
 					.ContentPadding(4)
-					.Visibility(this, &SBlueprintProfilerToolbar::IsNodeInclDefaultButtonVisible)
-					.OnClicked(this, &SBlueprintProfilerToolbar::ResetNodeInclThreshold)
+					.Visibility(this, &SBlueprintProfilerToolbar::IsCustomPerformanceThresholdDefaultButtonVisible, ECustomPerformanceThreshold::Inclusive)
+					.OnClicked(this, &SBlueprintProfilerToolbar::ResetCustomPerformanceThreshold, ECustomPerformanceThreshold::Inclusive)
 					.Content()
 					[
 						SNew(SImage)
@@ -882,8 +969,8 @@ TSharedRef<SWidget> SBlueprintProfilerToolbar::CreateHeatThresholdsMenuContent()
 					.MaxValue(2.f)
 					.MinSliderValue(0.01f)
 					.MaxSliderValue(2.f)
-					.Value(this, &SBlueprintProfilerToolbar::GetNodeMaxThreshold)
-					.OnValueChanged(this, &SBlueprintProfilerToolbar::SetNodeMaxThreshold)
+					.Value(this, &SBlueprintProfilerToolbar::GetCustomPerformanceThreshold, ECustomPerformanceThreshold::Max)
+					.OnValueChanged(this, &SBlueprintProfilerToolbar::SetCustomPerformanceThreshold, ECustomPerformanceThreshold::Max)
 					.MinDesiredValueWidth(100)
 					.ToolTipText(LOCTEXT("NodeMaxPerformanceThreshold_Tooltip", "The Hot Threshold for Max Node Timings"))
 				]
@@ -895,87 +982,8 @@ TSharedRef<SWidget> SBlueprintProfilerToolbar::CreateHeatThresholdsMenuContent()
 					.ToolTipText(LOCTEXT("ResetToDefaultToolTip", "Reset to Default"))
 					.ButtonStyle(FEditorStyle::Get(), "NoBorder")
 					.ContentPadding(4)
-					.Visibility(this, &SBlueprintProfilerToolbar::IsNodeMaxDefaultButtonVisible)
-					.OnClicked(this, &SBlueprintProfilerToolbar::ResetNodeMaxThreshold)
-					.Content()
-					[
-						SNew(SImage)
-						.Image(FEditorStyle::GetBrush("PropertyWindow.DiffersFromDefault"))
-					]
-				]
-			]
-		]
-		+SVerticalBox::Slot()
-		.AutoHeight()
-		.VAlign(VAlign_Center)
-		.Padding(FMargin(0,2))
-		[
-			SNew(SHorizontalBox)
-			+SHorizontalBox::Slot()
-			.AutoWidth()
-			.VAlign(VAlign_Center)
-			.Padding(FMargin(4,1,16,1))
-			[
-				SNew(STextBlock)
-				.Text(LOCTEXT("RecentSampleCheck_Label", "Enable Recent Sample Bias"))
-				.ToolTipText(LOCTEXT("RecentSampleCheck_Tooltip", "Enable a user controlled bias that controls the weighting of new samples against historical."))
-			]
-			+SHorizontalBox::Slot()
-			.AutoWidth()
-			.VAlign(VAlign_Center)
-			.Padding(FMargin(4,1,4,1))
-			[
-				SNew(SCheckBox)
-				.IsChecked(this, &SBlueprintProfilerToolbar::GetRecentSampleBiasChecked)
-				.OnCheckStateChanged(this, &SBlueprintProfilerToolbar::OnRecentSampleBiasChecked)
-			]
-		]
-		+SVerticalBox::Slot()
-		.AutoHeight()
-		.VAlign(VAlign_Center)
-		.Padding(FMargin(0,2))
-		[
-			SNew(SHorizontalBox)
-			.Visibility(this, &SBlueprintProfilerToolbar::GetRecentSampleBiasVisibility)
-			+SHorizontalBox::Slot()
-			.AutoWidth()
-			.VAlign(VAlign_Center)
-			.Padding(FMargin(4,1,4,1))
-			[
-				SNew(STextBlock)
-				.Text(LOCTEXT("RecentSampleValue_Label", "Recent Sample Bias"))
-				.ToolTipText(LOCTEXT("RecentSampleValue_Tooltip", "Shifts the weighting of more recent samples to control the effect of sampling spikes."))
-			]
-			+SHorizontalBox::Slot()
-			.HAlign(HAlign_Right)
-			.Padding(FMargin(4,1,1,1))
-			[
-				SNew(SHorizontalBox)
-				+SHorizontalBox::Slot()
-				.AutoWidth()
-				.HAlign(HAlign_Right)
-				[
-					SNew(SNumericEntryBox<float>)
-					.AllowSpin(true)
-					.MinValue(0.01f)
-					.MaxValue(1.f)
-					.MinSliderValue(0.01f)
-					.MaxSliderValue(1.f)
-					.Value(this, &SBlueprintProfilerToolbar::GetRecentSampleBiasValue)
-					.OnValueChanged(this, &SBlueprintProfilerToolbar::SetRecentSampleBiasValue)
-					.MinDesiredValueWidth(100)
-					.ToolTipText(LOCTEXT("RecentSampleValue_Tooltip", "Shifts the weighting of more recent samples to control the effect of sampling spikes."))
-				]
-				+SHorizontalBox::Slot()
-				.AutoWidth()
-				.HAlign(HAlign_Right)
-				[
-					SNew(SButton)
-					.ToolTipText(LOCTEXT("ResetToDefaultToolTip", "Reset to Default"))
-					.ButtonStyle(FEditorStyle::Get(), "NoBorder")
-					.ContentPadding(4)
-					.Visibility(this, &SBlueprintProfilerToolbar::IsRecentSampleBiasDefaultButtonVisible)
-					.OnClicked(this, &SBlueprintProfilerToolbar::ResetRecentSampleBias)
+					.Visibility(this, &SBlueprintProfilerToolbar::IsCustomPerformanceThresholdDefaultButtonVisible, ECustomPerformanceThreshold::Max)
+					.OnClicked(this, &SBlueprintProfilerToolbar::ResetCustomPerformanceThreshold, ECustomPerformanceThreshold::Max)
 					.Content()
 					[
 						SNew(SImage)
@@ -986,228 +994,130 @@ TSharedRef<SWidget> SBlueprintProfilerToolbar::CreateHeatThresholdsMenuContent()
 		];
 }
 
-EVisibility SBlueprintProfilerToolbar::GetRecentSampleBiasVisibility() const
+TOptional<float> SBlueprintProfilerToolbar::GetCustomPerformanceThreshold(ECustomPerformanceThreshold InType) const
 {
-	EVisibility Result = EVisibility::Collapsed;
-	if (const UEditorExperimentalSettings* EditorSettings = GetDefault<UEditorExperimentalSettings>())
-	{
-		Result = EditorSettings->bEnableBlueprintProfilerRecentSampleBias ? EVisibility::Visible : EVisibility::Collapsed;
-	}
-	return Result;
-}
+	TOptional<float> Result;
+	const UBlueprintProfilerSettings* ProfilerSettings = GetDefault<UBlueprintProfilerSettings>();
 
-ECheckBoxState SBlueprintProfilerToolbar::GetRecentSampleBiasChecked() const
-{
-	ECheckBoxState Result = ECheckBoxState::Unchecked;
-	if (const UEditorExperimentalSettings* EditorSettings = GetDefault<UEditorExperimentalSettings>())
+	switch (InType)
 	{
-		Result = EditorSettings->bEnableBlueprintProfilerRecentSampleBias ? ECheckBoxState::Checked : ECheckBoxState::Unchecked;
-	}
-	return Result;
-}
+	case ECustomPerformanceThreshold::Event:
+		Result = ProfilerSettings->CustomEventPerformanceThreshold;
+		break;
 
-void SBlueprintProfilerToolbar::OnRecentSampleBiasChecked(ECheckBoxState NewState)
-{
-	if (UEditorExperimentalSettings* EditorSettings = GetMutableDefault<UEditorExperimentalSettings>())
-	{
-		EditorSettings->bEnableBlueprintProfilerRecentSampleBias = NewState == ECheckBoxState::Checked;
-	}
-}
+	case ECustomPerformanceThreshold::Average:
+		Result = ProfilerSettings->CustomAveragePerformanceThreshold;
+		break;
 
-TOptional<float> SBlueprintProfilerToolbar::GetRecentSampleBiasValue() const
-{
-	TOptional<float> Result;
-	if (const UEditorExperimentalSettings* EditorSettings = GetDefault<UEditorExperimentalSettings>())
-	{
-		Result = EditorSettings->BlueprintProfilerRecentSampleBias;
-	}
-	return Result;
-}
+	case ECustomPerformanceThreshold::Inclusive:
+		Result = ProfilerSettings->CustomInclusivePerformanceThreshold;
+		break;
 
-void SBlueprintProfilerToolbar::SetRecentSampleBiasValue(const float NewValue)
-{
-	if (UEditorExperimentalSettings* EditorSettings = GetMutableDefault<UEditorExperimentalSettings>())
-	{
-		EditorSettings->BlueprintProfilerRecentSampleBias = NewValue;
-		FScriptPerfData::SetRecentSampleBias(NewValue);
+	case ECustomPerformanceThreshold::Max:
+		Result = ProfilerSettings->CustomMaxPerformanceThreshold;
+		break;
 	}
-}
 
-EVisibility SBlueprintProfilerToolbar::IsRecentSampleBiasDefaultButtonVisible() const
-{
-	EVisibility Result = EVisibility::Hidden;
-	if (UEditorExperimentalSettings* EditorSettings = GetMutableDefault<UEditorExperimentalSettings>())
-	{
-		Result = EditorSettings->BlueprintProfilerRecentSampleBias == 0.2f ? EVisibility::Hidden : EVisibility::Visible;
-	}
 	return Result;
 }
 
-FReply SBlueprintProfilerToolbar::ResetRecentSampleBias()
+void SBlueprintProfilerToolbar::SetCustomPerformanceThreshold(float NewValue, ECustomPerformanceThreshold InType)
 {
-	if (UEditorExperimentalSettings* EditorSettings = GetMutableDefault<UEditorExperimentalSettings>())
+	float OldValue = 0.0f;
+	UBlueprintProfilerSettings* ProfilerSettings = GetMutableDefault<UBlueprintProfilerSettings>();
+	
+	switch (InType)
 	{
-		EditorSettings->BlueprintProfilerRecentSampleBias = 0.2f;
-		FScriptPerfData::SetRecentSampleBias(0.2f);
-	}
-	return FReply::Handled();
-}
+	case ECustomPerformanceThreshold::Event:
+		OldValue = ProfilerSettings->CustomEventPerformanceThreshold;
+		ProfilerSettings->CustomEventPerformanceThreshold = NewValue;
+		break;
 
-TOptional<float> SBlueprintProfilerToolbar::GetEventThreshold() const
-{
-	TOptional<float> Result;
-	if (const UEditorExperimentalSettings* EditorSettings = GetDefault<UEditorExperimentalSettings>())
-	{
-		Result = EditorSettings->BlueprintProfilerEventThreshold;
-	}
-	return Result;
-}
+	case ECustomPerformanceThreshold::Average:
+		OldValue = ProfilerSettings->CustomAveragePerformanceThreshold;
+		ProfilerSettings->CustomAveragePerformanceThreshold = NewValue;
+		break;
 
-void SBlueprintProfilerToolbar::SetEventThreshold(const float NewValue)
-{
-	if (UEditorExperimentalSettings* EditorSettings = GetMutableDefault<UEditorExperimentalSettings>())
-	{
-		EditorSettings->BlueprintProfilerEventThreshold = NewValue;
-		FScriptPerfData::SetEventPerformanceThreshold(NewValue);
-	}
-}
+	case ECustomPerformanceThreshold::Inclusive:
+		OldValue = ProfilerSettings->CustomInclusivePerformanceThreshold;
+		ProfilerSettings->CustomInclusivePerformanceThreshold = NewValue;
+		break;
 
-EVisibility SBlueprintProfilerToolbar::IsEventDefaultButtonVisible() const
-{
-	EVisibility Result = EVisibility::Hidden;
-	if (UEditorExperimentalSettings* EditorSettings = GetMutableDefault<UEditorExperimentalSettings>())
-	{
-		Result = EditorSettings->BlueprintProfilerEventThreshold == 1.f ? EVisibility::Hidden : EVisibility::Visible;
+	case ECustomPerformanceThreshold::Max:
+		OldValue = ProfilerSettings->CustomMaxPerformanceThreshold;
+		ProfilerSettings->CustomMaxPerformanceThreshold = NewValue;
+		break;
 	}
-	return Result;
-}
 
-FReply SBlueprintProfilerToolbar::ResetEventThreshold()
-{
-	if (UEditorExperimentalSettings* EditorSettings = GetMutableDefault<UEditorExperimentalSettings>())
+	if (OldValue != NewValue)
 	{
-		EditorSettings->BlueprintProfilerEventThreshold = 1.f;
-		FScriptPerfData::SetEventPerformanceThreshold(1.f);
+		ProfilerSettings->SaveConfig();
 	}
-	return FReply::Handled();
 }
 
-TOptional<float> SBlueprintProfilerToolbar::GetNodeExclusiveThreshold() const
+EVisibility SBlueprintProfilerToolbar::IsCustomPerformanceThresholdDefaultButtonVisible(ECustomPerformanceThreshold InType) const
 {
-	TOptional<float> Result;
-	if (const UEditorExperimentalSettings* EditorSettings = GetDefault<UEditorExperimentalSettings>())
-	{
-		Result = EditorSettings->BlueprintProfilerExclNodeThreshold;
-	}
-	return Result;
-}
+	float CurValue = 0.f;
+	float DefValue = 0.f;
+	const UBlueprintProfilerSettings* ProfilerSettings = GetDefault<UBlueprintProfilerSettings>();
 
-void SBlueprintProfilerToolbar::SetNodeExclusiveThreshold(const float NewValue)
-{
-	if (UEditorExperimentalSettings* EditorSettings = GetMutableDefault<UEditorExperimentalSettings>())
+	switch (InType)
 	{
-		EditorSettings->BlueprintProfilerExclNodeThreshold = NewValue;
-		FScriptPerfData::SetExclusivePerformanceThreshold(NewValue);
-	}
-}
+	case ECustomPerformanceThreshold::Event:
+		CurValue = ProfilerSettings->CustomEventPerformanceThreshold;
+		DefValue = UBlueprintProfilerSettings::CustomEventPerformanceThresholdDefaultValue;
+		break;
 
-EVisibility SBlueprintProfilerToolbar::IsNodeExclDefaultButtonVisible() const
-{
-	EVisibility Result = EVisibility::Hidden;
-	if (UEditorExperimentalSettings* EditorSettings = GetMutableDefault<UEditorExperimentalSettings>())
-	{
-		Result = EditorSettings->BlueprintProfilerExclNodeThreshold == 0.2f ? EVisibility::Hidden : EVisibility::Visible;
-	}
-	return Result;
-}
+	case ECustomPerformanceThreshold::Average:
+		CurValue = ProfilerSettings->CustomAveragePerformanceThreshold;
+		DefValue = UBlueprintProfilerSettings::CustomAveragePerformanceThresholdDefaultValue;
+		break;
 
-FReply SBlueprintProfilerToolbar::ResetNodeExclThreshold()
-{
-	if (UEditorExperimentalSettings* EditorSettings = GetMutableDefault<UEditorExperimentalSettings>())
-	{
-		EditorSettings->BlueprintProfilerExclNodeThreshold = 0.2f;
-		FScriptPerfData::SetExclusivePerformanceThreshold(0.2f);
-	}
-	return FReply::Handled();
-}
+	case ECustomPerformanceThreshold::Inclusive:
+		CurValue = ProfilerSettings->CustomInclusivePerformanceThreshold;
+		DefValue = UBlueprintProfilerSettings::CustomInclusivePerformanceThresholdDefaultValue;
+		break;
 
-TOptional<float> SBlueprintProfilerToolbar::GetNodeInclusiveThreshold() const
-{
-	TOptional<float> Result;
-	if (const UEditorExperimentalSettings* EditorSettings = GetDefault<UEditorExperimentalSettings>())
-	{
-		Result = EditorSettings->BlueprintProfilerInclNodeThreshold;
+	case ECustomPerformanceThreshold::Max:
+		CurValue = ProfilerSettings->CustomMaxPerformanceThreshold;
+		DefValue = UBlueprintProfilerSettings::CustomMaxPerformanceThresholdDefaultValue;
+		break;
 	}
-	return Result;
-}
 
-void SBlueprintProfilerToolbar::SetNodeInclusiveThreshold(const float NewValue)
-{
-	if (UEditorExperimentalSettings* EditorSettings = GetMutableDefault<UEditorExperimentalSettings>())
-	{
-		EditorSettings->BlueprintProfilerInclNodeThreshold = NewValue;
-		FScriptPerfData::SetInclusivePerformanceThreshold(NewValue);
-	}
+	return CurValue == DefValue ? EVisibility::Hidden : EVisibility::Visible;
 }
 
-EVisibility SBlueprintProfilerToolbar::IsNodeInclDefaultButtonVisible() const
+FReply SBlueprintProfilerToolbar::ResetCustomPerformanceThreshold(ECustomPerformanceThreshold InType)
 {
-	EVisibility Result = EVisibility::Hidden;
-	if (UEditorExperimentalSettings* EditorSettings = GetMutableDefault<UEditorExperimentalSettings>())
+	float OldValue = 0.0f;
+	UBlueprintProfilerSettings* ProfilerSettings = GetMutableDefault<UBlueprintProfilerSettings>();
+	
+	switch (InType)
 	{
-		Result = EditorSettings->BlueprintProfilerInclNodeThreshold == 0.25f ? EVisibility::Hidden : EVisibility::Visible;
-	}
-	return Result;
-}
+	case ECustomPerformanceThreshold::Event:
+		OldValue = ProfilerSettings->CustomEventPerformanceThreshold;
+		ProfilerSettings->CustomEventPerformanceThreshold = UBlueprintProfilerSettings::CustomEventPerformanceThresholdDefaultValue;
+		break;
 
-FReply SBlueprintProfilerToolbar::ResetNodeInclThreshold()
-{
-	if (UEditorExperimentalSettings* EditorSettings = GetMutableDefault<UEditorExperimentalSettings>())
-	{
-		EditorSettings->BlueprintProfilerInclNodeThreshold = 0.25f;
-		FScriptPerfData::SetInclusivePerformanceThreshold(0.25f);
-	}
-	return FReply::Handled();
-}
+	case ECustomPerformanceThreshold::Average:
+		OldValue = ProfilerSettings->CustomAveragePerformanceThreshold;
+		ProfilerSettings->CustomAveragePerformanceThreshold = UBlueprintProfilerSettings::CustomAveragePerformanceThresholdDefaultValue;
+		break;
 
-TOptional<float> SBlueprintProfilerToolbar::GetNodeMaxThreshold() const
-{
-	TOptional<float> Result;
-	if (const UEditorExperimentalSettings* EditorSettings = GetDefault<UEditorExperimentalSettings>())
-	{
-		Result = EditorSettings->BlueprintProfilerMaxNodeThreshold;
-	}
-	return Result;
-}
+	case ECustomPerformanceThreshold::Inclusive:
+		OldValue = ProfilerSettings->CustomInclusivePerformanceThreshold;
+		ProfilerSettings->CustomInclusivePerformanceThreshold = UBlueprintProfilerSettings::CustomInclusivePerformanceThresholdDefaultValue;
+		break;
 
-void SBlueprintProfilerToolbar::SetNodeMaxThreshold(const float NewValue)
-{
-	if (UEditorExperimentalSettings* EditorSettings = GetMutableDefault<UEditorExperimentalSettings>())
-	{
-		EditorSettings->BlueprintProfilerMaxNodeThreshold = NewValue;
-		FScriptPerfData::SetMaxPerformanceThreshold(NewValue);
+	case ECustomPerformanceThreshold::Max:
+		OldValue = ProfilerSettings->CustomMaxPerformanceThreshold;
+		ProfilerSettings->CustomMaxPerformanceThreshold = UBlueprintProfilerSettings::CustomMaxPerformanceThresholdDefaultValue;
+		break;
 	}
-}
 
-EVisibility SBlueprintProfilerToolbar::IsNodeMaxDefaultButtonVisible() const
-{
-	EVisibility Result = EVisibility::Hidden;
-	if (UEditorExperimentalSettings* EditorSettings = GetMutableDefault<UEditorExperimentalSettings>())
-	{
-		Result = EditorSettings->BlueprintProfilerMaxNodeThreshold == 0.5f ? EVisibility::Hidden : EVisibility::Visible;
-	}
-	return Result;
-}
+	ProfilerSettings->SaveConfig();
 
-FReply SBlueprintProfilerToolbar::ResetNodeMaxThreshold()
-{
-	if (UEditorExperimentalSettings* EditorSettings = GetMutableDefault<UEditorExperimentalSettings>())
-	{
-		EditorSettings->BlueprintProfilerMaxNodeThreshold = 0.5f;
-		FScriptPerfData::SetMaxPerformanceThreshold(0.5f);
-	}
 	return FReply::Handled();
 }
 
-
 #undef LOCTEXT_NAMESPACE
diff --git a/Engine/Source/Editor/Kismet/Private/Profiler/SBlueprintProfilerToolbar.h b/Engine/Source/Editor/Kismet/Private/Profiler/SBlueprintProfilerToolbar.h
--- a/Engine/Source/Editor/Kismet/Private/Profiler/SBlueprintProfilerToolbar.h
+++ b/Engine/Source/Editor/Kismet/Private/Profiler/SBlueprintProfilerToolbar.h
@@ -2,7 +2,7 @@
 
 #pragma once
 
-#include "BlueprintProfilerModule.h"
+#include "BlueprintProfilerSettings.h"
 
 /** Blueprint performance view type */
 namespace EBlueprintPerfViewType
@@ -44,18 +44,13 @@ public:
 		ScopeToDebugInstance		= 0x00000002,	// Scope stats to current debug instance
 		GraphFilter					= 0x00000004,	// Filter to current graph
 		DisplayPure					= 0x00000008,	// Display pure stats
-		AutoExpand					= 0x00000010,	// Auto expand stats
+		DisplayInheritedEvents		= 0x00000010,	// Display inherited events
 		AverageBlueprintStats		= 0x00000020,	// Average blueprint stats ( Rather than sum )
 		Modified					= 0x10000000,	// Display state was changed
-		Default						= DisplayByInstance|DisplayPure|AverageBlueprintStats|Modified
 	};
 
-	FBlueprintProfilerStatOptions()
-		: ActiveInstance(NAME_None)
-		, ActiveGraph(NAME_None)
-		, Flags(Default)
-	{
-	}
+	/** Default constructor */
+	FBlueprintProfilerStatOptions();
 
 	/** Marks display state as modified */
 	void SetStateModified() { Flags |= Modified; }
@@ -69,6 +64,9 @@ public:
 	/** Clear flags */
 	void ClearFlags(const uint32 FlagsIn) { Flags &= ~FlagsIn; }
 
+	/** Set flags */
+	void SetFlags(const uint32 FlagsIn) { Flags = FlagsIn; }
+
 	/** Returns true if any flags match */
 	bool HasFlags(const uint32 FlagsIn) const { return (Flags & FlagsIn) != 0U; }
 
@@ -84,9 +82,6 @@ public:
 	/** Set active graph name */
 	void SetActiveGraph(const FName GraphName);
 
-	/** Create toolbar widget */
-	TSharedRef<SWidget> CreateToolbar();
-
 	/** Create combo content */
 	TSharedRef<SWidget> CreateComboContent();
 
@@ -160,49 +155,42 @@ protected:
 	/** Returns the active heat map display mode button foreground color */
 	FSlateColor GetHeatMapDisplayModeButtonForegroundColor(const EHeatmapControlId::Type ControlId) const;
 
+	/** Returns whether or not the given heat level metrics type is selected */
+	ECheckBoxState IsHeatLevelMetricsTypeSelected(const EBlueprintProfilerHeatLevelMetricsType InHeatLevelMetricsType) const;
+
+	/** Called when the heat level metrics type checkbox state is changed */
+	void OnHeatLevelMetricsTypeChanged(ECheckBoxState InCheckedState, const EBlueprintProfilerHeatLevelMetricsType NewHeatLevelMetricsType);
+
+	/** Returns text label color for heat level metrics type toggle buttons */
+	FSlateColor GetHeatLevelMetricsTypeTextColor(const EBlueprintProfilerHeatLevelMetricsType InHeatLevelMetricsType) const;
+
 	/** Returns whether or not the given heat map display mode is selected */
-	bool IsHeatMapDisplayModeSelected(const EHeatmapControlId::Type ControlId, const EBlueprintProfilerHeatMapDisplayMode::Type InHeatMapDisplayMode) const;
+	bool IsHeatMapDisplayModeSelected(const EHeatmapControlId::Type ControlId, const EBlueprintProfilerHeatMapDisplayMode InHeatMapDisplayMode) const;
 
 	/** Called when the heat map display mode is changed */
-	void OnHeatMapDisplayModeChanged(const EHeatmapControlId::Type ControlId, const EBlueprintProfilerHeatMapDisplayMode::Type NewHeatMapDisplayMode);
+	void OnHeatMapDisplayModeChanged(const EHeatmapControlId::Type ControlId, const EBlueprintProfilerHeatMapDisplayMode NewHeatMapDisplayMode);
 
-	/** Create the heat threshold combo menu content */
-	TSharedRef<SWidget> CreateHeatThresholdsMenuContent();
+	/** Returns whether or not the custom heat thresholds combo button should be enabled */
+	bool IsCustomHeatThresholdsMenuEnabled() const;
+
+	/** Create the custom heat thresholds combo menu content */
+	TSharedRef<SWidget> CreateCustomHeatThresholdsMenuContent();
 
 private:
+	/** Custom performance thresholds */
+	enum ECustomPerformanceThreshold
+	{
+		Event,
+		Average,
+		Inclusive,
+		Max
+	};
 
-	/** Recent sample bias helpers */
-	EVisibility GetRecentSampleBiasVisibility() const;
-	ECheckBoxState GetRecentSampleBiasChecked() const;
-	void OnRecentSampleBiasChecked(ECheckBoxState NewState);
-	TOptional<float> GetRecentSampleBiasValue() const;
-	void SetRecentSampleBiasValue(const float NewValue);
-	EVisibility IsRecentSampleBiasDefaultButtonVisible() const;
-	FReply ResetRecentSampleBias();
-
-	/** Get/Set event performance threshold helpers */
-	TOptional<float> GetEventThreshold() const;
-	void SetEventThreshold(const float NewValue);
-	EVisibility IsEventDefaultButtonVisible() const;
-	FReply ResetEventThreshold();
-
-	/** Get/Set node exclusive performance threshold helpers */
-	TOptional<float> GetNodeExclusiveThreshold() const;
-	void SetNodeExclusiveThreshold(const float NewValue);
-	EVisibility IsNodeExclDefaultButtonVisible() const;
-	FReply ResetNodeExclThreshold();
-
-	/** Get/Set node inclusive performance threshold helpers */
-	TOptional<float> GetNodeInclusiveThreshold() const;
-	void SetNodeInclusiveThreshold(const float NewValue);
-	EVisibility IsNodeInclDefaultButtonVisible() const;
-	FReply ResetNodeInclThreshold();
-
-	/** Get/Set node max performance threshold helpers */
-	TOptional<float> GetNodeMaxThreshold() const;
-	void SetNodeMaxThreshold(const float NewValue);
-	EVisibility IsNodeMaxDefaultButtonVisible() const;
-	FReply ResetNodeMaxThreshold();
+	/** Get/Set custom performance threshold helpers */
+	TOptional<float> GetCustomPerformanceThreshold(ECustomPerformanceThreshold InType) const;
+	void SetCustomPerformanceThreshold(const float NewValue, ECustomPerformanceThreshold InType);
+	EVisibility IsCustomPerformanceThresholdDefaultButtonVisible(ECustomPerformanceThreshold InType) const;
+	FReply ResetCustomPerformanceThreshold(ECustomPerformanceThreshold InType);
 
 protected:
 
@@ -221,12 +209,12 @@ protected:
 	/** View combo button widget */
 	TSharedPtr<SComboButton> ViewComboButton;
 
-	/** Heat threshold combo button widget */
-	TSharedPtr<SComboButton> HeatThresholdComboButton;
-
 	/** Heat map display mode combo button widget */
 	TSharedPtr<SComboButton> HeatMapDisplayModeComboButton;
 
 	/** Heat map display mode combo button widget */
 	TSharedPtr<SComboButton> WireHeatMapDisplayModeComboButton;
+
+	/** Custom heat threshold combo button widget */
+	TSharedPtr<SComboButton> CustomHeatThresholdComboButton;
 };
diff --git a/Engine/Source/Editor/Kismet/Private/Profiler/SBlueprintProfilerView.cpp b/Engine/Source/Editor/Kismet/Private/Profiler/SBlueprintProfilerView.cpp
--- a/Engine/Source/Editor/Kismet/Private/Profiler/SBlueprintProfilerView.cpp
+++ b/Engine/Source/Editor/Kismet/Private/Profiler/SBlueprintProfilerView.cpp
@@ -44,7 +44,11 @@ void SBlueprintProfilerView::Construct(const FArguments& InArgs)
 	StatisticNumberFormat.MinimumFractionalDigits = 4;
 	StatisticNumberFormat.MaximumFractionalDigits = 4;
 	StatisticNumberFormat.UseGrouping = false;
-	FScriptPerfData::SetNumberFormattingForStats(StatisticNumberFormat);
+	FNumberFormattingOptions TimeNumberFormat;
+	TimeNumberFormat.MinimumFractionalDigits = 1;
+	TimeNumberFormat.MaximumFractionalDigits = 1;
+	TimeNumberFormat.UseGrouping = false;
+	FScriptPerfData::SetNumberFormattingForStats(StatisticNumberFormat, TimeNumberFormat);
 	// Create the display text for the user
 	if (GetDefault<UEditorExperimentalSettings>()->bBlueprintPerformanceAnalysisTools)
 	{
@@ -130,7 +134,7 @@ void SBlueprintProfilerView::UpdateActiveProfilerWidget()
 		[
 			SNew(SBorder)
 			.Padding(FMargin(0,2,0,0))
-			.BorderImage(FEditorStyle::GetBrush("NoBorder"))
+			.BorderImage(FEditorStyle::GetBrush("BlueprintProfiler.ViewContent"))
 			[
 				CreateActiveStatisticWidget()
 			]
@@ -191,7 +195,7 @@ TSharedRef<SWidget> SBlueprintProfilerView::CreateActiveStatisticWidget()
 					[
 						SNew(SBorder)
 						.Padding(0)
-						.BorderImage(FEditorStyle::GetBrush("NoBorder"))
+						.BorderImage(FEditorStyle::GetBrush("BlueprintProfiler.ViewContent"))
 						.HAlign(HAlign_Center)
 						.VAlign(VAlign_Center)
 						[
diff --git a/Engine/Source/Editor/Kismet/Private/Profiler/SGraphExecutionStatDisplay.cpp b/Engine/Source/Editor/Kismet/Private/Profiler/SGraphExecutionStatDisplay.cpp
--- a/Engine/Source/Editor/Kismet/Private/Profiler/SGraphExecutionStatDisplay.cpp
+++ b/Engine/Source/Editor/Kismet/Private/Profiler/SGraphExecutionStatDisplay.cpp
@@ -39,45 +39,92 @@ void SGraphExecutionStatDisplay::Construct(const FArguments& InArgs)
 
 	ChildSlot
 	[
-		SNew(SVerticalBox)
-		+SVerticalBox::Slot()
+		SNew(SOverlay)
+		+SOverlay::Slot()
 		[
-			SNew(SBorder)
+			SNew(SVerticalBox)
+			.Visibility(this, &SGraphExecutionStatDisplay::GetWidgetVisibility)
+			+SVerticalBox::Slot()
+			[
+				SNew(SBorder)
+				.Padding(0)
+				.BorderImage(FEditorStyle::GetBrush("BlueprintProfiler.ViewContent"))
+				[
+					SAssignNew(ExecutionStatTree, STreeView<FBPStatWidgetPtr>)
+					.TreeItemsSource(&RootTreeItems)
+					.SelectionMode(ESelectionMode::Single)
+					.OnGetChildren(this, &SGraphExecutionStatDisplay::OnGetChildren)
+					.OnGenerateRow(this, &SGraphExecutionStatDisplay::OnGenerateRow)
+					.OnMouseButtonDoubleClick(this, &SGraphExecutionStatDisplay::OnDoubleClickStatistic)
+					.OnExpansionChanged(this, &SGraphExecutionStatDisplay::OnStatisticExpansionChanged)
+					.HeaderRow
+					(
+						SNew(SHeaderRow)
+						+SHeaderRow::Column(SProfilerStatRow::GetStatName(EBlueprintProfilerStat::Name))
+						.DefaultLabel(SProfilerStatRow::GetStatText(EBlueprintProfilerStat::Name))
+						.FillWidth(0.5f)
+						+SHeaderRow::Column(SProfilerStatRow::GetStatName(EBlueprintProfilerStat::AverageTime))
+						.DefaultLabel(SProfilerStatRow::GetStatText(EBlueprintProfilerStat::AverageTime))
+						.FixedWidth(90)
+						.HAlignHeader(HAlign_Right)
+						.HAlignCell(HAlign_Right)
+						.VAlignCell(VAlign_Center)
+						+SHeaderRow::Column(SProfilerStatRow::GetStatName(EBlueprintProfilerStat::InclusiveTime))
+						.DefaultLabel(SProfilerStatRow::GetStatText(EBlueprintProfilerStat::InclusiveTime))
+						.FixedWidth(120)
+						.HAlignHeader(HAlign_Right)
+						.HAlignCell(HAlign_Right)
+						.VAlignCell(VAlign_Center)
+						+SHeaderRow::Column(SProfilerStatRow::GetStatName(EBlueprintProfilerStat::MaxTime))
+						.DefaultLabel(SProfilerStatRow::GetStatText(EBlueprintProfilerStat::MaxTime))
+						.FixedWidth(90)
+						.HAlignHeader(HAlign_Right)
+						.HAlignCell(HAlign_Right)
+						.VAlignCell(VAlign_Center)
+						+SHeaderRow::Column(SProfilerStatRow::GetStatName(EBlueprintProfilerStat::MinTime))
+						.DefaultLabel(SProfilerStatRow::GetStatText(EBlueprintProfilerStat::MinTime))
+						.FixedWidth(90)
+						.HAlignHeader(HAlign_Right)
+						.HAlignCell(HAlign_Right)
+						.VAlignCell(VAlign_Center)
+						+SHeaderRow::Column(SProfilerStatRow::GetStatName(EBlueprintProfilerStat::TotalTime))
+						.DefaultLabel(SProfilerStatRow::GetStatText(EBlueprintProfilerStat::TotalTime))
+						.FixedWidth(100)
+						.HAlignHeader(HAlign_Right)
+						.HAlignCell(HAlign_Right)
+						.VAlignCell(VAlign_Center)
+						+SHeaderRow::Column(SProfilerStatRow::GetStatName(EBlueprintProfilerStat::Samples))
+						.DefaultLabel(SProfilerStatRow::GetStatText(EBlueprintProfilerStat::Samples))
+						.FixedWidth(100)
+						.HAlignHeader(HAlign_Center)
+						.HAlignCell(HAlign_Center)
+						.VAlignCell(VAlign_Center)
+					)
+				]
+			]
+		]
+		+SOverlay::Slot()
+		[
+			SNew(SVerticalBox)
+			.Visibility(this, &SGraphExecutionStatDisplay::GetNoDataWidgetVisibility)
+			+SVerticalBox::Slot()
+			.HAlign(HAlign_Center)
+			.VAlign(VAlign_Center)
 			.Padding(0)
-			.BorderImage(FEditorStyle::GetBrush("NoBorder"))
 			[
-				SAssignNew(ExecutionStatTree, STreeView<FBPStatWidgetPtr>)
-				.TreeItemsSource(&RootTreeItems)
-				.SelectionMode(ESelectionMode::Single)
-				.OnGetChildren(this, &SGraphExecutionStatDisplay::OnGetChildren)
-				.OnGenerateRow(this, &SGraphExecutionStatDisplay::OnGenerateRow)
-				.OnMouseButtonDoubleClick(this, &SGraphExecutionStatDisplay::OnDoubleClickStatistic)
-				.OnExpansionChanged(this, &SGraphExecutionStatDisplay::OnStatisticExpansionChanged)
-				.HeaderRow
-				(
-					SNew(SHeaderRow)
-					+SHeaderRow::Column(SProfilerStatRow::GetStatName(EBlueprintProfilerStat::Name))
-					.DefaultLabel(SProfilerStatRow::GetStatText(EBlueprintProfilerStat::Name))
-					.ManualWidth(450)
-					+SHeaderRow::Column(SProfilerStatRow::GetStatName(EBlueprintProfilerStat::AverageTime))
-					.DefaultLabel(SProfilerStatRow::GetStatText(EBlueprintProfilerStat::AverageTime))
-					.ManualWidth(90)
-					+SHeaderRow::Column(SProfilerStatRow::GetStatName(EBlueprintProfilerStat::InclusiveTime))
-					.DefaultLabel(SProfilerStatRow::GetStatText(EBlueprintProfilerStat::InclusiveTime))
-					.ManualWidth(120)
-					+SHeaderRow::Column(SProfilerStatRow::GetStatName(EBlueprintProfilerStat::MaxTime))
-					.DefaultLabel(SProfilerStatRow::GetStatText(EBlueprintProfilerStat::MaxTime))
-					.ManualWidth(90)
-					+SHeaderRow::Column(SProfilerStatRow::GetStatName(EBlueprintProfilerStat::MinTime))
-					.DefaultLabel(SProfilerStatRow::GetStatText(EBlueprintProfilerStat::MinTime))
-					.ManualWidth(90)
-					+SHeaderRow::Column(SProfilerStatRow::GetStatName(EBlueprintProfilerStat::TotalTime))
-					.DefaultLabel(SProfilerStatRow::GetStatText(EBlueprintProfilerStat::TotalTime))
-					.ManualWidth(80)
-					+SHeaderRow::Column(SProfilerStatRow::GetStatName(EBlueprintProfilerStat::Samples))
-					.DefaultLabel(SProfilerStatRow::GetStatText(EBlueprintProfilerStat::Samples))
-					.ManualWidth(60)
-				)
+				SNew(SHorizontalBox)
+				+SHorizontalBox::Slot()
+				[
+					SNew(SBorder)
+					.Padding(0)
+					.BorderImage(FEditorStyle::GetBrush("BlueprintProfiler.ViewContent"))
+					.HAlign(HAlign_Center)
+					.VAlign(VAlign_Center)
+					[
+						SNew(STextBlock)
+						.Text(LOCTEXT("NoInstancesMessage", "There are currently no valid instances of this blueprint in the current level to gather profiling data on."))
+					]
+				]
 			]
 		]
 	];
@@ -121,7 +168,7 @@ void SGraphExecutionStatDisplay::OnDoubleClickStatistic(FBPStatWidgetPtr Item)
 
 void SGraphExecutionStatDisplay::OnStatisticExpansionChanged(FBPStatWidgetPtr Item, bool bExpanded)
 {
-	if (Item.IsValid() && !DisplayOptions->HasFlags(FBlueprintProfilerStatOptions::AutoExpand))
+	if (Item.IsValid())
 	{
 		Item->SetExpansionState(bExpanded);
 	}
@@ -218,20 +265,10 @@ void SGraphExecutionStatDisplay::Tick(const FGeometry& AllottedGeometry, const d
 							if (ExecutionStatTree.IsValid())
 							{
 								ExecutionStatTree->RequestTreeRefresh();
-								if (DisplayOptions->HasFlags(FBlueprintProfilerStatOptions::AutoExpand))
-								{
-									for (auto Iter : RootTreeItems)
-									{
-										Iter->ExpandWidgetState(ExecutionStatTree, true);
-									}
-								}
-								else
+								for (auto Iter : RootTreeItems)
 								{
-									for (auto Iter : RootTreeItems)
-									{
-										Iter->ProbeChildWidgetExpansionStates();
-										Iter->RestoreWidgetExpansionState(ExecutionStatTree);
-									}
+									Iter->ProbeChildWidgetExpansionStates();
+									Iter->RestoreWidgetExpansionState(ExecutionStatTree);
 								}
 							}
 						}
diff --git a/Engine/Source/Editor/Kismet/Private/Profiler/SLeastPerformantDisplay.cpp b/Engine/Source/Editor/Kismet/Private/Profiler/SLeastPerformantDisplay.cpp
--- a/Engine/Source/Editor/Kismet/Private/Profiler/SLeastPerformantDisplay.cpp
+++ b/Engine/Source/Editor/Kismet/Private/Profiler/SLeastPerformantDisplay.cpp
@@ -41,7 +41,7 @@ void SLeastPerformantDisplay::Construct(const FArguments& InArgs)
 		[
 			SNew(SBorder)
 			.Padding(0)
-			.BorderImage(FEditorStyle::GetBrush("NoBorder"))
+			.BorderImage(FEditorStyle::GetBrush("BlueprintProfiler.ViewContent"))
 			[
 				SAssignNew(BlueprintStatList, SListView<FBPStatWidgetPtr>)
 				.ListItemsSource(&RootTreeItems)
@@ -53,25 +53,43 @@ void SLeastPerformantDisplay::Construct(const FArguments& InArgs)
 					SNew(SHeaderRow)
 					+SHeaderRow::Column(SProfilerStatRow::GetStatName(EBlueprintProfilerStat::Name))
 					.DefaultLabel(SProfilerStatRow::GetStatText(EBlueprintProfilerStat::Name))
-					.ManualWidth(450)
+					.FillWidth(0.5f)
 					+SHeaderRow::Column(SProfilerStatRow::GetStatName(EBlueprintProfilerStat::AverageTime))
 					.DefaultLabel(SProfilerStatRow::GetStatText(EBlueprintProfilerStat::AverageTime))
-					.ManualWidth(90)
+					.FixedWidth(90)
+					.HAlignHeader(HAlign_Right)
+					.HAlignCell(HAlign_Right)
+					.VAlignCell(VAlign_Center)
 					+SHeaderRow::Column(SProfilerStatRow::GetStatName(EBlueprintProfilerStat::InclusiveTime))
 					.DefaultLabel(SProfilerStatRow::GetStatText(EBlueprintProfilerStat::InclusiveTime))
-					.ManualWidth(120)
+					.FixedWidth(120)
+					.HAlignHeader(HAlign_Right)
+					.HAlignCell(HAlign_Right)
+					.VAlignCell(VAlign_Center)
 					+SHeaderRow::Column(SProfilerStatRow::GetStatName(EBlueprintProfilerStat::MaxTime))
 					.DefaultLabel(SProfilerStatRow::GetStatText(EBlueprintProfilerStat::MaxTime))
-					.ManualWidth(90)
+					.FixedWidth(90)
+					.HAlignHeader(HAlign_Right)
+					.HAlignCell(HAlign_Right)
+					.VAlignCell(VAlign_Center)
 					+SHeaderRow::Column(SProfilerStatRow::GetStatName(EBlueprintProfilerStat::MinTime))
 					.DefaultLabel(SProfilerStatRow::GetStatText(EBlueprintProfilerStat::MinTime))
-					.ManualWidth(90)
+					.FixedWidth(90)
+					.HAlignHeader(HAlign_Right)
+					.HAlignCell(HAlign_Right)
+					.VAlignCell(VAlign_Center)
 					+SHeaderRow::Column(SProfilerStatRow::GetStatName(EBlueprintProfilerStat::TotalTime))
 					.DefaultLabel(SProfilerStatRow::GetStatText(EBlueprintProfilerStat::TotalTime))
-					.ManualWidth(80)
+					.FixedWidth(100)
+					.HAlignHeader(HAlign_Right)
+					.HAlignCell(HAlign_Right)
+					.VAlignCell(VAlign_Center)
 					+SHeaderRow::Column(SProfilerStatRow::GetStatName(EBlueprintProfilerStat::Samples))
 					.DefaultLabel(SProfilerStatRow::GetStatText(EBlueprintProfilerStat::Samples))
-					.ManualWidth(60)
+					.FixedWidth(100)
+					.HAlignHeader(HAlign_Center)
+					.HAlignCell(HAlign_Center)
+					.VAlignCell(VAlign_Center)
 				)
 			]
 		]
diff --git a/Engine/Source/Editor/Kismet/Private/Profiler/SProfilerOverview.cpp b/Engine/Source/Editor/Kismet/Private/Profiler/SProfilerOverview.cpp
--- a/Engine/Source/Editor/Kismet/Private/Profiler/SProfilerOverview.cpp
+++ b/Engine/Source/Editor/Kismet/Private/Profiler/SProfilerOverview.cpp
@@ -41,7 +41,7 @@ void SProfilerOverview::Construct(const FArguments& InArgs)
 		[
 			SNew(SBorder)
 			.Padding(0)
-			.BorderImage(FEditorStyle::GetBrush("NoBorder"))
+			.BorderImage(FEditorStyle::GetBrush("BlueprintProfiler.ViewContent"))
 			[
 				SAssignNew(BlueprintStatList, SListView<FBPStatWidgetPtr>)
 				.ListItemsSource(&RootTreeItems)
@@ -53,25 +53,43 @@ void SProfilerOverview::Construct(const FArguments& InArgs)
 					SNew(SHeaderRow)
 					+SHeaderRow::Column(SProfilerStatRow::GetStatName(EBlueprintProfilerStat::Name))
 					.DefaultLabel(SProfilerStatRow::GetStatText(EBlueprintProfilerStat::Name))
-					.ManualWidth(450)
+					.FillWidth(0.5f)
 					+SHeaderRow::Column(SProfilerStatRow::GetStatName(EBlueprintProfilerStat::AverageTime))
 					.DefaultLabel(SProfilerStatRow::GetStatText(EBlueprintProfilerStat::AverageTime))
-					.ManualWidth(90)
+					.FixedWidth(90)
+					.HAlignHeader(HAlign_Right)
+					.HAlignCell(HAlign_Right)
+					.VAlignCell(VAlign_Center)
 					+SHeaderRow::Column(SProfilerStatRow::GetStatName(EBlueprintProfilerStat::InclusiveTime))
 					.DefaultLabel(SProfilerStatRow::GetStatText(EBlueprintProfilerStat::InclusiveTime))
-					.ManualWidth(120)
+					.FixedWidth(120)
+					.HAlignHeader(HAlign_Right)
+					.HAlignCell(HAlign_Right)
+					.VAlignCell(VAlign_Center)
 					+SHeaderRow::Column(SProfilerStatRow::GetStatName(EBlueprintProfilerStat::MaxTime))
 					.DefaultLabel(SProfilerStatRow::GetStatText(EBlueprintProfilerStat::MaxTime))
-					.ManualWidth(90)
+					.FixedWidth(90)
+					.HAlignHeader(HAlign_Right)
+					.HAlignCell(HAlign_Right)
+					.VAlignCell(VAlign_Center)
 					+SHeaderRow::Column(SProfilerStatRow::GetStatName(EBlueprintProfilerStat::MinTime))
 					.DefaultLabel(SProfilerStatRow::GetStatText(EBlueprintProfilerStat::MinTime))
-					.ManualWidth(90)
+					.FixedWidth(90)
+					.HAlignHeader(HAlign_Right)
+					.HAlignCell(HAlign_Right)
+					.VAlignCell(VAlign_Center)
 					+SHeaderRow::Column(SProfilerStatRow::GetStatName(EBlueprintProfilerStat::TotalTime))
 					.DefaultLabel(SProfilerStatRow::GetStatText(EBlueprintProfilerStat::TotalTime))
-					.ManualWidth(80)
+					.FixedWidth(100)
+					.HAlignHeader(HAlign_Right)
+					.HAlignCell(HAlign_Right)
+					.VAlignCell(VAlign_Center)
 					+SHeaderRow::Column(SProfilerStatRow::GetStatName(EBlueprintProfilerStat::Samples))
 					.DefaultLabel(SProfilerStatRow::GetStatText(EBlueprintProfilerStat::Samples))
-					.ManualWidth(60)
+					.FixedWidth(100)
+					.HAlignHeader(HAlign_Center)
+					.HAlignCell(HAlign_Center)
+					.VAlignCell(VAlign_Center)
 				)
 			]
 		]
diff --git a/Engine/Source/Editor/Kismet/Private/Profiler/ScriptPerfData.cpp b/Engine/Source/Editor/Kismet/Private/Profiler/ScriptPerfData.cpp
--- a/Engine/Source/Editor/Kismet/Private/Profiler/ScriptPerfData.cpp
+++ b/Engine/Source/Editor/Kismet/Private/Profiler/ScriptPerfData.cpp
@@ -3,39 +3,40 @@
 #include "BlueprintEditorPrivatePCH.h"
 #include "ScriptPerfData.h"
 
+//////////////////////////////////////////////////////////////////////////
+// FScriptHeatLevelMetrics
+
+float FScriptHeatLevelMetrics::NodeTotalTimeWaterMark = 0.f;
+
 //////////////////////////////////////////////////////////////////////////
 // FScriptPerfData
 
 FNumberFormattingOptions FScriptPerfData::StatNumberFormat;
+FNumberFormattingOptions FScriptPerfData::TimeNumberFormat;
+
 bool FScriptPerfData::bAverageBlueprintStats = true;
-bool FScriptPerfData::bUseRecentSampleBias = false;
-float FScriptPerfData::RecentSampleBias = 0.2f;
-float FScriptPerfData::HistoricalSampleBias = 0.8f;
-float FScriptPerfData::EventPerformanceThreshold = 1.f;
-float FScriptPerfData::ExclusivePerformanceThreshold = 0.2f;
-float FScriptPerfData::InclusivePerformanceThreshold = 0.25f;
-float FScriptPerfData::MaxPerformanceThreshold = 0.5f;
-float FScriptPerfData::NodeTotalTimeWaterMark = 0.f;
-
-void FScriptPerfData::AddEventTiming(const double ExclusiveTimingIn)
-{
-	const double TimingInMs = ExclusiveTimingIn * 1000;
-	ExclusiveTiming = (TimingInMs * RecentSampleBias) + (ExclusiveTiming * HistoricalSampleBias);
-	MaxTiming = FMath::Max<double>(MaxTiming, ExclusiveTiming);
-	MinTiming = FMath::Min<double>(MinTiming, ExclusiveTiming);
-	TotalTiming += ExclusiveTimingIn;
-	NumSamples++;
+
+void FScriptPerfData::AddEventTiming(const double AverageTimingIn)
+{
+	const double TimingInMs = AverageTimingIn * 1000;
+	AverageTiming += TimingInMs;
+	RawSamples++;
+	const int32 SampleCount = GetSampleCount();
+	if (SampleCount > 0)
+	{
+		const double NormalisedAverageTiming = GetAverageTiming();
+		MaxTiming = FMath::Max<double>(MaxTiming, NormalisedAverageTiming);
+		MinTiming = FMath::Min<double>(MinTiming, NormalisedAverageTiming);
+	}
 }
 
-void FScriptPerfData::AddInclusiveTiming(const double InclusiveNodeTimingIn, const bool bIncrementSamples, const bool bDiscreteTiming)
+void FScriptPerfData::AddInclusiveTiming(const double InclusiveNodeTimingIn, const bool bIncrementSamples)
 {
 	// Discrete timings should negate the discrete nodetimings so they sum correctly later.
-	const double TimingInMs = bDiscreteTiming ? ((InclusiveNodeTimingIn * 1000) - ExclusiveTiming) : (InclusiveNodeTimingIn * 1000);
-	InclusiveTiming = (TimingInMs * RecentSampleBias) + (InclusiveTiming * HistoricalSampleBias);
-	TotalTiming += InclusiveNodeTimingIn;
+	InclusiveTiming += InclusiveNodeTimingIn * 1000;
 	if (bIncrementSamples)
 	{
-		NumSamples++;
+		RawSamples++;
 	}
 }
 
@@ -44,142 +45,79 @@ void FScriptPerfData::AddData(const FScriptPerfData& DataIn)
 	if (DataIn.IsDataValid())
 	{
 		// Accumulate data, find min, max values
-		ExclusiveTiming += DataIn.ExclusiveTiming;
-		if (DataIn.InclusiveTiming > 0.0)
-		{
-			InclusiveTiming += DataIn.InclusiveTiming;
-		}
-		TotalTiming += DataIn.TotalTiming;
-		NumSamples += DataIn.NumSamples;
+		InclusiveTiming += DataIn.InclusiveTiming;
+		AverageTiming += DataIn.AverageTiming;
+		RawSamples += DataIn.RawSamples;
 		MaxTiming = FMath::Max<double>(MaxTiming, DataIn.MaxTiming);
 		MinTiming = FMath::Min<double>(MinTiming, DataIn.MinTiming);
 	}
 }
 
-void FScriptPerfData::AddBranchData(const FScriptPerfData& DataIn)
-{
-	if (DataIn.IsDataValid())
-	{
-		// Accumulate all data with the exception of samples for which we take the first valid value.
-		const bool bFirstSample = MaxTiming == -MAX_dbl;
-		ExclusiveTiming += DataIn.ExclusiveTiming;
-		if (DataIn.InclusiveTiming > 0.0)
-		{
-			InclusiveTiming += DataIn.InclusiveTiming;
-		}
-		TotalTiming += DataIn.TotalTiming;
-		if (bFirstSample)
-		{
-			NumSamples = DataIn.NumSamples;
-			MaxTiming = DataIn.MaxTiming;
-			MinTiming = DataIn.MinTiming;
-		}
-		else
-		{
-			MaxTiming += DataIn.MaxTiming;
-			MinTiming += DataIn.MinTiming;
-		}
-	}
-}
-
 void FScriptPerfData::InitialiseFromDataSet(const TArray<TSharedPtr<FScriptPerfData>>& DataSet)
 {
 	Reset();
+	const double Denom = 1.0 / DataSet.Num();
 	for (auto DataPoint : DataSet)
 	{
 		if (DataPoint.IsValid())
 		{
-			ExclusiveTiming += DataPoint->ExclusiveTiming;
-			InclusiveTiming += DataPoint->InclusiveTiming;
-			TotalTiming += DataPoint->TotalTiming;
+			AverageTiming += DataPoint->AverageTiming * Denom;
+			InclusiveTiming += DataPoint->InclusiveTiming * Denom;
 			// Either sum the samples to get the average or use the highest sample rate encountered in the dataset.
-			NumSamples = bAverageBlueprintStats ? (NumSamples + DataPoint->NumSamples) : FMath::Max(NumSamples, DataPoint->NumSamples);
-			MaxTiming = FMath::Max<double>(MaxTiming, DataPoint->MaxTiming);
-			MinTiming = FMath::Min<double>(MinTiming, DataPoint->MinTiming);
+			RawSamples = bAverageBlueprintStats ? (RawSamples + DataPoint->RawSamples) : FMath::Max(RawSamples, DataPoint->RawSamples);
 		}
 	}
+	if (IsDataValid())
+	{
+		const double NormalisedAverageTiming = GetAverageTiming();
+		MaxTiming = FMath::Max<double>(MaxTiming, NormalisedAverageTiming);
+		MinTiming = FMath::Min<double>(MinTiming, NormalisedAverageTiming);
+	}
 }
 
 void FScriptPerfData::Reset()
 {
-	ExclusiveTiming = 0.0;
+	AverageTiming = 0.0;
 	InclusiveTiming = 0.0;
-	TotalTiming = 0.0;
+	AverageTimingHeatLevel = 0.f;
+	InclusiveTimingHeatLevel = 0.f;
+	MaxTimingHeatLevel = 0.f;
+	TotalTimingHeatLevel = 0.f;
 	HottestPathHeatValue = 0.f;
-	NumSamples = 0;
+	RawSamples = 0;
 	MaxTiming = -MAX_dbl;
 	MinTiming = MAX_dbl;
 }
 
-double FScriptPerfData::GetExclusiveTiming() const
-{
-	return bUseRecentSampleBias ? ExclusiveTiming : ((TotalTiming*1000.0)/NumSamples);
-}
-
-void FScriptPerfData::EnableBlueprintStatAverage(const bool bEnable)
-{
-	bAverageBlueprintStats = bEnable;
-}
-
-void FScriptPerfData::EnableRecentSampleBias(const bool bEnable)
+void FScriptPerfData::SetNumberFormattingForStats(const FNumberFormattingOptions& StatFormatIn, const FNumberFormattingOptions& TimeFormatIn)
 {
-	bUseRecentSampleBias = bEnable;
+	StatNumberFormat = StatFormatIn;
+	TimeNumberFormat = TimeFormatIn;
 }
 
-void FScriptPerfData::SetRecentSampleBias(const float RecentSampleBiasIn)
+double FScriptPerfData::GetAverageTiming() const
 {
-	RecentSampleBias = RecentSampleBiasIn;
-	HistoricalSampleBias = 1.f - RecentSampleBias;
+	// We need to use the GetSampleCount() call so we can take account of both the 
+	// SampleBase (re-entrant handling) and the RawSamples.
+	return AverageTiming / GetSampleCount();
 }
 
-void FScriptPerfData::SetEventPerformanceThreshold(const float EventPerformanceThresholdIn)
+double FScriptPerfData::GetInclusiveTiming() const
 {
-	EventPerformanceThreshold = 1.f / EventPerformanceThresholdIn;
+	// For inclusive timings we just use RawSamples as there is no re-entrant handling.
+	return (AverageTiming / GetSampleCount()) + (InclusiveTiming / RawSamples);
 }
 
-void FScriptPerfData::SetExclusivePerformanceThreshold(const float ExclusivePerformanceThresholdIn)
-{
-	ExclusivePerformanceThreshold = 1.f / ExclusivePerformanceThresholdIn;
-}
-
-float FScriptPerfData::GetNodeHeatLevel() const
-{
-	const float PerfMetric = StatType == EScriptPerfDataType::Event ? EventPerformanceThreshold : ExclusivePerformanceThreshold;
-	return FMath::Min<float>(GetExclusiveTiming() * PerfMetric, 1.f);
-}
-
-void FScriptPerfData::SetInclusivePerformanceThreshold(const float InclusivePerformanceThresholdIn)
-{
-	InclusivePerformanceThreshold = 1.f / InclusivePerformanceThresholdIn;
-}
-
-float FScriptPerfData::GetInclusiveHeatLevel() const
-{
-	const float PerfMetric = StatType == EScriptPerfDataType::Event ? EventPerformanceThreshold : InclusivePerformanceThreshold;
-	return FMath::Min<float>(GetInclusiveTiming() * InclusivePerformanceThreshold, 1.f);
-}
-
-void FScriptPerfData::SetMaxPerformanceThreshold(const float MaxPerformanceThresholdIn)
-{
-	MaxPerformanceThreshold = 1.f / MaxPerformanceThresholdIn;
-}
-
-float FScriptPerfData::GetMaxTimeHeatLevel() const
-{
-	return FMath::Min<float>(MaxTiming * MaxPerformanceThreshold, 1.f);
-}
-
-float FScriptPerfData::GetTotalHeatLevel() const
+void FScriptPerfData::EnableBlueprintStatAverage(const bool bEnable)
 {
-	NodeTotalTimeWaterMark = FMath::Max<float>(TotalTiming, NodeTotalTimeWaterMark);
-	return FMath::Min<float>(TotalTiming / NodeTotalTimeWaterMark, 1.f);
+	bAverageBlueprintStats = bEnable;
 }
 
-FText FScriptPerfData::GetExclusiveTimingText() const
+FText FScriptPerfData::GetAverageTimingText() const
 {
-	const double Value = GetExclusiveTiming();
-	if (Value > 0.0)
+	if (GetSampleCount() > 0)
 	{
+		const double Value = GetAverageTiming();
 		return FText::AsNumber(Value, &StatNumberFormat);
 	}
 	return FText::GetEmpty();
@@ -187,7 +125,7 @@ FText FScriptPerfData::GetExclusiveTimingText() const
 
 FText FScriptPerfData::GetInclusiveTimingText() const
 {
-	if (NumSamples > 0)
+	if (GetSampleCount() > 0)
 	{
 		const double InclusiveTemp = GetInclusiveTiming();
 		return FText::AsNumber(InclusiveTemp, &StatNumberFormat);
@@ -215,18 +153,46 @@ FText FScriptPerfData::GetMinTimingText() const
 
 FText FScriptPerfData::GetTotalTimingText() const
 {
-	if (TotalTiming > 0.0)
+	if (GetSampleCount() > 0)
 	{
-		return FText::AsNumber(TotalTiming, &StatNumberFormat);
+		return FText::AsNumber(AverageTiming, &TimeNumberFormat);
 	}
 	return FText::GetEmpty();
 }
 
 FText FScriptPerfData::GetSamplesText() const
 {
-	if (NumSamples > 0)
+	if (GetSampleCount() > 0)
 	{
-		return FText::AsNumber(NumSamples);
+		return FText::AsNumber(GetSampleCount());
 	}
 	return FText::GetEmpty();
+}
+
+void FScriptPerfData::SetHeatLevels(TSharedPtr<FScriptHeatLevelMetrics> HeatLevelMetrics)
+{
+	if (ensure(HeatLevelMetrics.IsValid()))
+	{
+		// Note: We're intentionally not accessing AverageTiming/InclusiveTiming members directly here - the accessors make some internal adjustments.
+		AverageTimingHeatLevel = FMath::Min<float>(GetAverageTiming() * HeatLevelMetrics->AveragePerformanceThreshold, 1.f);
+		InclusiveTimingHeatLevel = FMath::Min<float>(GetInclusiveTiming() * HeatLevelMetrics->InclusivePerformanceThreshold, 1.f);
+		MaxTimingHeatLevel = FMath::Min<float>(MaxTiming * HeatLevelMetrics->MaxPerformanceThreshold, 1.f);
+
+		if (HeatLevelMetrics->bUseTotalTimeWaterMark)
+		{
+			HeatLevelMetrics->NodeTotalTimeWaterMark = FMath::Max<float>(AverageTiming, HeatLevelMetrics->NodeTotalTimeWaterMark);
+			TotalTimingHeatLevel = FMath::Min<float>(AverageTiming / HeatLevelMetrics->NodeTotalTimeWaterMark, 1.f);
+		}
+		else
+		{
+			TotalTimingHeatLevel = FMath::Min<float>(AverageTiming * HeatLevelMetrics->TotalTimePerformanceThreshold, 1.f);
+		}
+	}
+	else
+	{
+		AverageTimingHeatLevel = 0.0f;
+		InclusiveTimingHeatLevel = 0.0f;
+		MaxTimingHeatLevel = 0.0f;
+		TotalTimingHeatLevel = 0.0f;
+	}
 }
\ No newline at end of file
diff --git a/Engine/Source/Editor/Kismet/Private/SKismetInspector.cpp b/Engine/Source/Editor/Kismet/Private/SKismetInspector.cpp
--- a/Engine/Source/Editor/Kismet/Private/SKismetInspector.cpp
+++ b/Engine/Source/Editor/Kismet/Private/SKismetInspector.cpp
@@ -135,9 +135,6 @@ TSharedRef<SWidget> SKismetInspector::MakeContextualEditingWidget(struct FKismet
 	PropertyView->HideFilterArea(Options.bHideFilterArea);
 	PropertyView->SetObjects(SelectionInfo.ObjectsForPropertyEditing, Options.bForceRefresh);
 
-	PropertyView->SetIsPropertyEditingEnabledDelegate(FIsPropertyEditingEnabled::CreateSP(this, &SKismetInspector::IsPropertyEditingEnabled));
-
-
 	if (SelectionInfo.ObjectsForPropertyEditing.Num())
 	{
 		ContextualEditingWidget->AddSlot()
@@ -316,7 +313,9 @@ void SKismetInspector::Construct(const FArguments& InArgs)
 		
 	//@TODO: .IsEnabled( FSlateApplication::Get().GetNormalExecutionAttribute() );
 	PropertyView->SetIsPropertyVisibleDelegate( FIsPropertyVisible::CreateSP(this, &SKismetInspector::IsPropertyVisible) );
-	PropertyView->SetIsPropertyEditingEnabledDelegate(InArgs._IsPropertyEditingEnabledDelegate);
+	PropertyView->SetIsPropertyEditingEnabledDelegate(FIsPropertyEditingEnabled::CreateSP(this, &SKismetInspector::IsPropertyEditingEnabled));
+
+	IsPropertyEditingEnabledDelegate = InArgs._IsPropertyEditingEnabledDelegate;
 	UserOnFinishedChangingProperties = InArgs._OnFinishedChangingProperties;
 
 	TWeakPtr<SMyBlueprint> MyBlueprint = Kismet2.IsValid() ? Kismet2->GetMyBlueprintWidget() : InArgs._MyBlueprintWidget;
@@ -776,7 +775,7 @@ bool SKismetInspector::IsPropertyEditingEnabled() const
 			break;
 		}
 	}
-	return bIsEditable;
+	return bIsEditable && (!IsPropertyEditingEnabledDelegate.IsBound() || IsPropertyEditingEnabledDelegate.Execute());
 }
 
 EVisibility SKismetInspector::GetInheritedBlueprintComponentWarningVisibility() const
diff --git a/Engine/Source/Editor/Kismet/Public/BlueprintEditor.h b/Engine/Source/Editor/Kismet/Public/BlueprintEditor.h
--- a/Engine/Source/Editor/Kismet/Public/BlueprintEditor.h
+++ b/Engine/Source/Editor/Kismet/Public/BlueprintEditor.h
@@ -836,7 +836,7 @@ protected:
 	 * @param InDestinationGraph		The destination graph to move the selected nodes to
 	 * @param InCollapsableNodes		The selection of nodes being collapsed
 	 */
-	void CollapseNodesIntoGraph(UEdGraphNode* InGatewayNode, class UK2Node_EditablePinBase* InEntryNode, class UK2Node_EditablePinBase* InResultNode, UEdGraph* InSourceGraph, UEdGraph* InDestinationGraph, TSet<UEdGraphNode*>& InCollapsableNodes);
+	void CollapseNodesIntoGraph(UEdGraphNode* InGatewayNode, class UK2Node_EditablePinBase* InEntryNode, class UK2Node_EditablePinBase* InResultNode, UEdGraph* InSourceGraph, UEdGraph* InDestinationGraph, TSet<UEdGraphNode*>& InCollapsableNodes, bool bCanDiscardEmptyReturnNode = false);
 
 	/** Called when a selection of nodes are being collapsed into a sub-graph */
 	void CollapseNodes(TSet<class UEdGraphNode*>& InCollapsableNodes);
@@ -1172,6 +1172,9 @@ private:
 
 	/** Handle to the registered OnActiveTabChanged delegate */
 	FDelegateHandle OnActiveTabChangedDelegateHandle;
+
+	/** Blueprint generated class instrumentation state */
+	bool bBlueprintHasInstrumentation;
 };
 
 #undef LOCTEXT_NAMESPACE
diff --git a/Engine/Source/Editor/Kismet/Public/Profiler/BlueprintProfilerSettings.h b/Engine/Source/Editor/Kismet/Public/Profiler/BlueprintProfilerSettings.h
new file mode 100644
--- /dev/null
+++ b/Engine/Source/Editor/Kismet/Public/Profiler/BlueprintProfilerSettings.h
@@ -0,0 +1,94 @@
+// Copyright 1998-2016 Epic Games, Inc. All Rights Reserved.
+
+#pragma once
+
+#include "BlueprintProfilerSettings.generated.h"
+
+/** Heat map display mode */
+UENUM()
+enum class EBlueprintProfilerHeatMapDisplayMode : uint8
+{
+	None = 0,
+	Inclusive,
+	Average,
+	MaxTiming,
+	Total,
+	HottestPath,
+	PinToPin
+};
+
+/** Heat map level metrics type */
+UENUM()
+enum class EBlueprintProfilerHeatLevelMetricsType : uint8
+{
+	ClassRelative,
+	FrameRelative,
+	CustomThresholds
+};
+
+UCLASS(config=EditorPerProjectUserSettings)
+class KISMET_API UBlueprintProfilerSettings
+	: public UObject
+{
+	GENERATED_UCLASS_BODY()
+
+public:
+	/** Display statistics by instance rather than by blueprint */
+	UPROPERTY(config)
+	bool bDisplayByInstance;
+
+	/** Restrict the displayed instances to match the blueprint debug filter */
+	UPROPERTY(config)
+	bool bScopeToDebugInstance;
+
+	/** Only Display events that are applicable to the current graph */
+	UPROPERTY(config)
+	bool bGraphFilter;
+
+	/** Display timings from pure nodes */
+	UPROPERTY(config)
+	bool bDisplayPure;
+
+	/** Display events for parent blueprints */
+	UPROPERTY(config)
+	bool bDisplayInheritedEvents;
+
+	/** Blueprint Statistics are Averaged or Summed By Instance */
+	UPROPERTY(config)
+	bool bAverageBlueprintStats;
+
+	/** Current graph node heat map display mode */
+	UPROPERTY(config)
+	EBlueprintProfilerHeatMapDisplayMode GraphNodeHeatMapDisplayMode;
+
+	/** Current wire heat map display mode */
+	UPROPERTY(config)
+	EBlueprintProfilerHeatMapDisplayMode WireHeatMapDisplayMode;
+
+	/** Current heat level metrics type */
+	UPROPERTY(config)
+	EBlueprintProfilerHeatLevelMetricsType HeatLevelMetricsType;
+
+	/** Custom performance threshold for event timings in ms. */
+	UPROPERTY(config)
+	float CustomEventPerformanceThreshold;
+
+	/** Custom performance threshold for average timings in ms. */
+	UPROPERTY(config)
+	float CustomAveragePerformanceThreshold;
+
+	/** Custom performance threshold for inclusive timings in ms. */
+	UPROPERTY(config)
+	float CustomInclusivePerformanceThreshold;
+
+	/** Custom performance threshold for max timings in ms. */
+	UPROPERTY(config)
+	float CustomMaxPerformanceThreshold;
+
+	/** Default custom threshold values */
+	/** @TODO - maybe we don't have a need for these anymore? they're currently supporting the "reset to default" feature in the UI - not sure we have that in the "normal" editor settings UI */
+	static float CustomEventPerformanceThresholdDefaultValue;
+	static float CustomAveragePerformanceThresholdDefaultValue;
+	static float CustomInclusivePerformanceThresholdDefaultValue;
+	static float CustomMaxPerformanceThresholdDefaultValue;
+};
diff --git a/Engine/Source/Editor/Kismet/Public/Profiler/EventExecution.h b/Engine/Source/Editor/Kismet/Public/Profiler/EventExecution.h
--- a/Engine/Source/Editor/Kismet/Public/Profiler/EventExecution.h
+++ b/Engine/Source/Editor/Kismet/Public/Profiler/EventExecution.h
@@ -15,8 +15,10 @@ namespace EScriptExecutionNodeFlags
 		Instance					= 0x00000002,	// Instance
 		Event						= 0x00000004,	// Event
 		CustomEvent					= 0x00000008,	// Custom Event
-		FunctionCall				= 0x00000010,	// Function Call
-		MacroCall					= 0x00000020,	// Macro Call
+		InheritedEvent				= 0x00000010,	// Inherited Event
+		FunctionCall				= 0x00000020,	// Function Call
+		ParentFunctionCall			= 0x00000040,	// Parent Function Call
+		MacroCall					= 0x00000080,	// Macro Call
 		MacroNode					= 0x00000100,	// Macro Node
 		ConditionalBranch			= 0x00000200,	// Node has multiple exit pins using a jump
 		SequentialBranch			= 0x00000400,	// Node has multiple exit pins ran in sequence
@@ -29,11 +31,11 @@ namespace EScriptExecutionNodeFlags
 		TunnelExitPin				= 0x00020000,	// Internal tunnel exit pin
 		TunnelEntryPinInstance		= 0x00040000,	// Tunnel entry pin instance
 		TunnelExitPinInstance		= 0x00080000,	// Tunnel exit pin instance
-		PureChain					= 0x00200000,	// Pure node call chain
-		CyclicLinkage				= 0x00400000,	// Marks execution path as cyclic.
-		InvalidTrace				= 0x00800000,	// Indicates that node doesn't contain a valid script trace.
+		PureChain					= 0x00100000,	// Pure node call chain
+		CyclicLinkage				= 0x00200000,	// Marks execution path as cyclic.
+		InvalidTrace				= 0x00400000,	// Indicates that node doesn't contain a valid script trace.
 		// Groups
-		CallSite					= FunctionCall|MacroCall|TunnelInstance,
+		CallSite					= FunctionCall|ParentFunctionCall|MacroCall|TunnelInstance,
 		BranchNode					= ConditionalBranch|SequentialBranch,
 		TunnelInstancePin			= TunnelEntryPinInstance|TunnelExitPinInstance,
 		TunnelPin					= TunnelEntryPin|TunnelExitPin,
@@ -119,6 +121,11 @@ class KISMET_API FScriptNodePerfData
 {
 public:
 
+	FScriptNodePerfData(const int32 SampleFreqIn = 1)
+		: SampleFrequency(SampleFreqIn)
+	{
+	}
+
 	/** Returns a TSet containing all valid instance names */
 	void GetValidInstanceNames(TSet<FName>& ValidInstances) const;
 
@@ -143,6 +150,9 @@ public:
 	/** Get global blueprint perf data for all trace paths */
 	virtual void GetBlueprintPerfDataForAllTracePaths(FScriptPerfData& OutPerfData);
 
+	/** Set the sample frequency */
+	void SetSampleFrequency(const int32 NewSampleFrequency) { SampleFrequency = NewSampleFrequency; }
+
 protected:
 
 	/** Returns performance data type */
@@ -150,6 +160,8 @@ protected:
 
 protected:
 	
+	/** Expected sample frequency for this perf data container */
+	int32 SampleFrequency;
 	/** FScriptExeutionPath hash to perf data */
 	TMap<FName, TMap<const uint32, TSharedPtr<FScriptPerfData>>> InstanceInputPinToPerfDataMap;
 
@@ -166,6 +178,8 @@ struct KISMET_API FScriptExecNodeParams
 	FName OwningGraphName;
 	/** Node flags to describe the source graph node type */
 	uint32 NodeFlags;
+	/** The expected sample frequency that constitutes a singles execution of this node */
+	int32 SampleFrequency;
 	/** Oberved object */
 	TWeakObjectPtr<const UObject> ObservedObject;
 	/** Observed pin */
@@ -248,6 +262,9 @@ public:
 	/** Returns if this event is a function callsite event */
 	bool IsFunctionCallSite() const { return (NodeFlags & EScriptExecutionNodeFlags::FunctionCall) != 0U; }
 
+	/** Returns if this event is a parent function callsite event */
+	bool IsParentFunctionCallSite() const { return (NodeFlags & EScriptExecutionNodeFlags::ParentFunctionCall) != 0U; }
+
 	/** Returns if this event is a macro callsite event */
 	bool IsMacroCallSite() const { return (NodeFlags & EScriptExecutionNodeFlags::MacroCall) != 0U; }
 
@@ -284,6 +301,9 @@ public:
 	/** Gets the observed object context */
 	const UObject* GetObservedObject() { return ObservedObject.Get(); }
 
+	/** Gets a typed observed object context */
+	template<typename CastType> const CastType* GetTypedObservedObject() { return Cast<CastType>(ObservedObject.Get()); }
+
 	/** Gets the observed object context */
 	const UEdGraphPin* GetObservedPin() const { return ObservedPin.Get(); }
 
@@ -332,6 +352,9 @@ public:
 	/** Refresh Stats */
 	virtual void RefreshStats(const FTracePath& TracePath);
 
+	/** Calculate heat level stats */
+	virtual void CalculateHeatLevelStats(TSharedPtr<FScriptHeatLevelMetrics> HeatLevelMetrics);
+
 	/** Calculate Hottest Path Stats */
 	virtual float CalculateHottestPathStats(FScriptExecutionHottestPathParams HotPathParams);
 
@@ -339,7 +362,7 @@ public:
 	virtual void GetAllExecNodes(TMap<FName, TSharedPtr<FScriptExecutionNode>>& ExecNodesOut);
 
 	/** Get all pure nodes associated with the given trace path */
-	virtual void GetAllPureNodes(TMap<int32, TSharedPtr<FScriptExecutionNode>>& PureNodesOut);
+	void GetAllPureNodes(TMap<int32, TSharedPtr<FScriptExecutionNode>>& PureNodesOut);
 
 	/** Creates and returns the slate icon widget */
 	virtual TSharedRef<SWidget> GetIconWidget(const uint32 TracePath = 0U);
@@ -434,6 +457,7 @@ public:
 
 	FScriptExecutionTunnelEntry(const FScriptExecNodeParams& InitParams)
 		: FScriptExecutionNode(InitParams)
+		, TunnelEntryCount(0)
 	{
 	}
 
@@ -454,20 +478,19 @@ public:
 	/** Sets the tunnel instance exec node */
 	void SetTunnelInstance(TSharedPtr<FScriptExecutionTunnelInstance> TunnelInstanceIn) { TunnelInstance = TunnelInstanceIn; }
 
-	/** Adds a valid exit script offset */
-	void AddExitSite(const int32 ExitScriptOffset, TSharedPtr<class FScriptExecutionTunnelExit> ExitSite);
-
 	/** Returns the exit point associated with the exit script offset */
 	TSharedPtr<FScriptExecutionTunnelExit> GetExitSite(const int32 ScriptOffset) const;
 
-	/** Returns the exit point associated with the entry script offset */
-	TMap<int32, TSharedPtr<FScriptExecutionNode>> GetExitSites() { return LinkedNodes; }
-
 	/** Returns true if the supplied node is inside this tunnel */
 	bool IsInternalBoundary(TSharedPtr<FScriptExecutionNode> PotentialBoundaryNode) const;
 
+	/** Increment tunnel entry count */
+	void IncrementTunnelEntryCount() { TunnelEntryCount++; }
+
 private:
 
+	/** The tunnel entry count */
+	int32 TunnelEntryCount;
 	/** The tunnel instance this node represents */
 	TSharedPtr<FScriptExecutionTunnelInstance> TunnelInstance;
 
@@ -601,10 +624,7 @@ class KISMET_API FScriptExecutionBlueprint : public FScriptExecutionNode
 {
 public:
 
-	FScriptExecutionBlueprint(const FScriptExecNodeParams& InitParams)
-		: FScriptExecutionNode(InitParams)
-	{
-	}
+	FScriptExecutionBlueprint(const FScriptExecNodeParams& InitParams);
 
 	/** Adds new blueprint instance node */
 	void AddInstance(TSharedPtr<FScriptExecutionNode> Instance) { Instances.Add(Instance); }
@@ -618,6 +638,12 @@ public:
 	/** Returns the instance that matches the supplied name if present */
 	TSharedPtr<FScriptExecutionNode> GetInstanceByName(FName InstanceName);
 
+	/** Sorts all contained events */
+	void SortEvents();
+	
+	/** Returns the heat level metrics object that's used for updating heat levels */
+	TSharedPtr<FScriptHeatLevelMetrics> GetHeatLevelMetrics() const { return HeatLevelMetrics; }
+
 	// FScriptExecutionNode
 	virtual EScriptStatContainerType::Type GetStatisticContainerType() const override { return EScriptStatContainerType::Container; }
 	virtual void RefreshStats(const FTracePath& TracePath) override;
@@ -625,10 +651,15 @@ public:
 	virtual void NavigateToObject() const override;
 	// ~FScriptExecutionNode
 
+protected:
+	/** Updates internal heat level metrics data */
+	void UpdateHeatLevelMetrics(TSharedPtr<FScriptPerfData> BlueprintData);
+
 private:
 
 	/** Exec nodes representing all instances based on this blueprint */
 	TArray<TSharedPtr<FScriptExecutionNode>> Instances;
-
+	/** Internal metrics data used for calculating perf data heat levels */
+	TSharedPtr<FScriptHeatLevelMetrics> HeatLevelMetrics;
 };
 
