Copying //UE4/Dev-Blueprints to //UE4/Dev-Main (Source: //UE4/Dev-Blueprints @ 3057841)

diff --git a/Engine/Source/Editor/Kismet/Public/Profiler/ScriptPerfData.h b/Engine/Source/Editor/Kismet/Public/Profiler/ScriptPerfData.h
--- a/Engine/Source/Editor/Kismet/Public/Profiler/ScriptPerfData.h
+++ b/Engine/Source/Editor/Kismet/Public/Profiler/ScriptPerfData.h
@@ -2,6 +2,8 @@
 
 #pragma once
 
+#include "BlueprintProfilerSettings.h"
+
 //////////////////////////////////////////////////////////////////////////
 // EScriptPerfDataType
 
@@ -11,6 +13,39 @@ enum EScriptPerfDataType
 	Node
 };
 
+//////////////////////////////////////////////////////////////////////////
+// FScriptHeatLevelMetrics
+
+struct KISMET_API FScriptHeatLevelMetrics
+{
+	/** Event performance threshold */
+	float EventPerformanceThreshold;
+	/** Average performance threshold */
+	float AveragePerformanceThreshold;
+	/** Inclusive performance threshold */
+	float InclusivePerformanceThreshold;
+	/** Node max performance threshold */
+	float MaxPerformanceThreshold;
+	/** Total time performance threshold */
+	float TotalTimePerformanceThreshold;
+	/** Whether or not to use total time watermark when calculating total time heat level */
+	bool bUseTotalTimeWaterMark;
+	/** Total time watermark (note: this is intentionally global) */
+	static float NodeTotalTimeWaterMark;
+
+	/** Default constructor */
+	FScriptHeatLevelMetrics()
+		: EventPerformanceThreshold(0.0f)
+		, AveragePerformanceThreshold(0.0f)
+		, InclusivePerformanceThreshold(0.0f)
+		, MaxPerformanceThreshold(0.0f)
+		, TotalTimePerformanceThreshold(0.0f)
+		, bUseTotalTimeWaterMark(false)
+	{
+
+	}
+};
+
 //////////////////////////////////////////////////////////////////////////
 // FScriptPerfData
 
@@ -18,78 +53,73 @@ class KISMET_API FScriptPerfData : public TSharedFromThis<FScriptPerfData>
 {
 public:
 
-	FScriptPerfData(const EScriptPerfDataType StatTypeIn)
+	FScriptPerfData(const EScriptPerfDataType StatTypeIn, const int32 SampleFrequencyIn = 1)
 		: StatType(StatTypeIn)
-		, ExclusiveTiming(0.0)
+		, SampleFrequency(SampleFrequencyIn)
+		, AverageTiming(0.0)
 		, InclusiveTiming(0.0)
 		, MaxTiming(-MAX_dbl)
 		, MinTiming(MAX_dbl)
-		, TotalTiming(0.0)
+		, AverageTimingHeatLevel(0.f)
+		, InclusiveTimingHeatLevel(0.f)
+		, MaxTimingHeatLevel(0.f)
+		, TotalTimingHeatLevel(0.f)
 		, HottestPathHeatValue(0.f)
-		, NumSamples(0)
+		, RawSamples(0)
 	{
 	}
 
 	/** Add a single event node timing into the dataset */
-	void AddEventTiming(const double ExclusiveTimingIn);
+	void AddEventTiming(const double AverageTimingIn);
 
 	/** Add a inclusive timing into the dataset */
-	void AddInclusiveTiming(const double InclusiveNodeTimingIn, const bool bIncrementSamples, const bool bDiscreteTiming);
+	void AddInclusiveTiming(const double InclusiveNodeTimingIn, const bool bIncrementSamples);
 
 	/** Add data */
 	void AddData(const FScriptPerfData& DataIn);
 
-	/** Add branch data */
-	void AddBranchData(const FScriptPerfData& DataIn);
-
 	/** Initialise from data set */
 	void InitialiseFromDataSet(const TArray<TSharedPtr<FScriptPerfData>>& DataSet);
 
 	/** Increments sample count without affecting data */
-	void TickSamples() { NumSamples++; }
+	void TickSamples() { RawSamples++; }
 
 	/** Reset the current data buffer and all derived data stats */
 	void Reset();
 
 	// Returns if this data structure has any valid data
-	bool IsDataValid() const { return NumSamples > 0; }
+	bool IsDataValid() const { return RawSamples > 0; }
 
 	// Updates the various thresholds that control stats calcs and display.
-	static void SetNumberFormattingForStats(const FNumberFormattingOptions& FormatIn) { StatNumberFormat = FormatIn; }
+	static void SetNumberFormattingForStats(const FNumberFormattingOptions& StatFormatIn, const FNumberFormattingOptions& TimeFormatIn);
 	static void EnableBlueprintStatAverage(const bool bEnable);
-	static void EnableRecentSampleBias(const bool bEnable);
-	static void SetRecentSampleBias(const float RecentSampleBiasIn);
-	static void SetEventPerformanceThreshold(const float EventPerformanceThresholdIn);
-	static void SetExclusivePerformanceThreshold(const float ExclusivePerformanceThresholdIn);
-	static void SetInclusivePerformanceThreshold(const float InclusivePerformanceThresholdIn);
-	static void SetMaxPerformanceThreshold(const float MaxPerformanceThresholdIn);
 
 	// Returns the various stats this container holds
-	double GetExclusiveTiming() const;
-	double GetInclusiveTiming() const { return GetExclusiveTiming() + InclusiveTiming; }
+	double GetAverageTiming() const;
+	double GetInclusiveTiming() const;
 	double GetMaxTiming() const { return MaxTiming; }
 	double GetMinTiming() const { return MinTiming; }
-	double GetTotalTiming() const { return TotalTiming; }
-	int32 GetSampleCount() const { return NumSamples; }
+	double GetTotalTiming() const { return AverageTiming; }
+	int32 GetSampleCount() const { return RawSamples / SampleFrequency; }
+	void OverrideSampleCount(const int32 NewSampleCount) { RawSamples = NewSampleCount; }
 
-	// Returns various performance heat levels for visual display
-	float GetNodeHeatLevel() const;
-	float GetInclusiveHeatLevel() const;
-	float GetMaxTimeHeatLevel() const;
-	float GetTotalHeatLevel() const;
+	// Returns performance heat levels for visual display
+	float GetAverageHeatLevel() const { return AverageTimingHeatLevel; }
+	float GetInclusiveHeatLevel() const { return InclusiveTimingHeatLevel; }
+	float GetMaxTimeHeatLevel() const { return MaxTimingHeatLevel; }
+	float GetTotalHeatLevel() const { return TotalTimingHeatLevel; }
+
+	// Calculate performance heat levels for visual display
+	void SetHeatLevels(TSharedPtr<FScriptHeatLevelMetrics> HeatLevelMetrics);
 
 	// Hottest path interface
 	float GetHottestPathHeatLevel() const { return HottestPathHeatValue; }
 	void SetHottestPathHeatLevel(const float HottestPathValue) { HottestPathHeatValue = HottestPathValue; }
 
-	// Hottest endpoint interface
-	float GetHottestEndpointHeatLevel() const { return HottestEndpointHeatValue; }
-	void SetHottestEndpointHeatLevel(const float HottestEndpointValue) { HottestEndpointHeatValue = HottestEndpointValue; }
-
 	// Returns the various stats this container holds in FText format
 	FText GetTotalTimingText() const;
 	FText GetInclusiveTimingText() const;
-	FText GetExclusiveTimingText() const;
+	FText GetAverageTimingText() const;
 	FText GetMaxTimingText() const;
 	FText GetMinTimingText() const;
 	FText GetSamplesText() const;
@@ -97,44 +127,36 @@ public:
 private:
 
 	/** The stat type for performance threshold comparisons */
-	EScriptPerfDataType StatType;
-	/** Exclusive timing */
-	double ExclusiveTiming;
-	/** Inclusive timing, pure and node timings */
+	const EScriptPerfDataType StatType;
+	/** The sample frequency */
+	const int32 SampleFrequency;
+
+	/** Average timing */
+	double AverageTiming;
+	/** Inclusive timing */
 	double InclusiveTiming;
-	/** Max exclusive timing */
+	/** Max average timing */
 	double MaxTiming;
-	/** Min exclusive timing */
+	/** Min average timing */
 	double MinTiming;
-	/** Total time accrued */
-	double TotalTiming;
+	/** Average timing heat level */
+	float AverageTimingHeatLevel;
+	/** Inclusive timing heat level */
+	float InclusiveTimingHeatLevel;
+	/** Max timing heat level */
+	float MaxTimingHeatLevel;
+	/** Total timing heat level */
+	float TotalTimingHeatLevel;
 	/** Hottest path value */
 	float HottestPathHeatValue;
-	/** Hottest endpoint value */
-	float HottestEndpointHeatValue;
-	/** The current inclusive sample count */
-	int32 NumSamples;
+	/** The current raw average sample count */
+	int32 RawSamples;
 
-	/** Number formatting options */
+	/** Stat number formatting options */
 	static FNumberFormattingOptions StatNumberFormat;
+	/** Time number formatting options */
+	static FNumberFormattingOptions TimeNumberFormat;
 
 	/** Sets blueprint level stats as averaged */
 	static bool bAverageBlueprintStats;
-	/** Enables the sample biasing */
-	static bool bUseRecentSampleBias;
-	/** Controls the bias between new and older samples */
-	static float RecentSampleBias;
-	/** Historical Sample Bias */
-	static float HistoricalSampleBias;
-	/** Event performance threshold */
-	static float EventPerformanceThreshold;
-	/** Exclusive performance threshold */
-	static float ExclusivePerformanceThreshold;
-	/** Node inclusive performance threshold */
-	static float InclusivePerformanceThreshold;
-	/** Node max performance threshold */
-	static float MaxPerformanceThreshold;
-	/** Total time watermark */
-	static float NodeTotalTimeWaterMark;
-
 };
diff --git a/Engine/Source/Editor/Kismet/Public/SKismetInspector.h b/Engine/Source/Editor/Kismet/Public/SKismetInspector.h
--- a/Engine/Source/Editor/Kismet/Public/SKismetInspector.h
+++ b/Engine/Source/Editor/Kismet/Public/SKismetInspector.h
@@ -131,6 +131,9 @@ protected:
 	/** Set of object properties that should be visible */
 	TSet<TWeakObjectPtr<UProperty> > SelectedObjectProperties;
 	
+	/** User defined delegate for IsPropertyEditingEnabled: */
+	FIsPropertyEditingEnabled IsPropertyEditingEnabledDelegate;
+
 	/** User defined delegate for OnFinishedChangingProperties */
 	FOnFinishedChangingProperties::FDelegate UserOnFinishedChangingProperties;
 
diff --git a/Engine/Source/Editor/KismetCompiler/Private/KismetCompilerModule.cpp b/Engine/Source/Editor/KismetCompiler/Private/KismetCompilerModule.cpp
--- a/Engine/Source/Editor/KismetCompiler/Private/KismetCompilerModule.cpp
+++ b/Engine/Source/Editor/KismetCompiler/Private/KismetCompilerModule.cpp
@@ -101,6 +101,20 @@ void FKismet2CompilerModule::CompileBlueprintInner(class UBlueprint* Blueprint,
 			}
 		}
 
+		// we need to add instrumentation to any inherited blueprints.
+		if (Reinstancer.IsValid() && CompileOptions.IsInstrumentationActive())
+		{
+			TArray<UBlueprint*> InheritedBlueprints;
+			Blueprint->GetBlueprintHierarchyFromClass(Cast<UClass>(Blueprint->GeneratedClass), InheritedBlueprints);
+			for (auto CurrentBP : InheritedBlueprints)
+			{
+				if (CurrentBP != Blueprint)
+				{
+					Reinstancer->EnlistDependentBlueprintToRecompile(CurrentBP, true);
+				}
+			}
+		}
+
 		if (bRecompileDependencies)
 		{
 			Reinstancer->BlueprintWasRecompiled(Blueprint, CompileOptions.CompileType == EKismetCompileType::BytecodeOnly);
diff --git a/Engine/Source/Editor/KismetCompiler/Private/KismetCompilerVMBackend.cpp b/Engine/Source/Editor/KismetCompiler/Private/KismetCompilerVMBackend.cpp
--- a/Engine/Source/Editor/KismetCompiler/Private/KismetCompilerVMBackend.cpp
+++ b/Engine/Source/Editor/KismetCompiler/Private/KismetCompilerVMBackend.cpp
@@ -1029,6 +1029,24 @@ public:
 			Writer << EX_VirtualFunction;
 			Writer << FunctionName;
 		}
+		
+		TArray<FName> WildcardParams;
+		const bool bIsCustomThunk = FunctionToCall->HasMetaData(TEXT("CustomThunk"));
+		if (bIsCustomThunk)
+		{
+			auto CollectWildcards = [&](FName MetaDataName)
+			{
+				const FString DependentPinMetaData = FunctionToCall->GetMetaData(MetaDataName);
+				TArray<FString> TypeDependentPinNames;
+				DependentPinMetaData.ParseIntoArray(TypeDependentPinNames, TEXT(","), true);
+				for (FString& Iter : TypeDependentPinNames)
+				{
+					WildcardParams.Add(FName(*Iter));
+				}
+			};
+			CollectWildcards(FBlueprintMetadata::MD_ArrayDependentParam);
+			CollectWildcards(FName(TEXT("CustomStructureParam")));
+		}
 
 		// Emit function parameters
 		int32 NumParams = 0;
@@ -1042,18 +1060,15 @@ public:
 				check(Term != NULL);
 
 				// Latent function handling:  Need to emit a fixup request into the FLatentInfo struct
-				static const FName NAME_LatentInfo = TEXT("LatentInfo");
  				if (bIsUbergraph && FuncParamProperty->GetName() == FunctionToCall->GetMetaData("LatentInfo"))
  				{
 					EmitLatentInfoTerm(Term, FuncParamProperty, Statement.TargetLabel);
  				}
 				else
 				{
-					// Emit parameter term normally
-					EmitTerm(Term, FuncParamProperty);
+					const bool bWildcard = WildcardParams.Contains(FuncParamProperty->GetFName());
+					EmitTerm(Term, bWildcard ? nullptr : FuncParamProperty);
 				}
-
-
 				NumParams++;
 			}
 		}
@@ -1614,7 +1629,33 @@ public:
 					CompilerContext.MessageLog.MacroSourceToMacroInstanceNodeMap.MultiFind(MacroSourceNode, MacroInstanceNodes);
 				}
 
-				ClassBeingBuilt->GetDebugData().RegisterNodeToCodeAssociation(TrueSourceNode, MacroSourceNode, MacroInstanceNodes, FunctionContext.Function, Offset, bBreakpointSite);
+				// Register source tunnels for instrumentation.
+				if (FunctionContext.IsInstrumentationRequired())
+				{
+					// After blueprint profiler MVP these changes will be refactored and rolled back into normal compilation.
+
+					// Locate the source tunnel instance
+					if (TWeakObjectPtr<UEdGraphNode>* SourceTunnelInstance = CompilerContext.MessageLog.SourceNodeToTunnelInstanceNodeMap.Find(SourceNode))
+					{
+						if (SourceTunnelInstance->IsValid())
+						{
+							TrueSourceNode = SourceTunnelInstance->Get();
+						}
+					}
+					// Locate the real source node, as the code above fails with nested tunnels
+					if (TWeakObjectPtr<UEdGraphNode>* NewSourceNode = CompilerContext.MessageLog.IntermediateTunnelNodeToSourceNodeMap.Find(SourceNode))
+					{
+						if (NewSourceNode->IsValid())
+						{
+							MacroSourceNode = NewSourceNode->Get();
+						}
+					}
+					ClassBeingBuilt->GetDebugData().RegisterNodeToCodeAssociation(TrueSourceNode, MacroSourceNode, MacroInstanceNodes, FunctionContext.Function, Offset, bBreakpointSite);
+				}
+				else
+				{
+					ClassBeingBuilt->GetDebugData().RegisterNodeToCodeAssociation(TrueSourceNode, MacroSourceNode, MacroInstanceNodes, FunctionContext.Function, Offset, bBreakpointSite);
+				}
 
 				// Track pure node script code range for the current impure (exec) node
 				if (Statement.Type == KCST_InstrumentedPureNodeEntry)
@@ -1630,7 +1671,7 @@ public:
 				else if (Statement.Type == KCST_InstrumentedWireEntry && PureNodeEntryCount > 0)
 				{
 					// Map script code range for the full set of pure node inputs feeding in to the current impure (exec) node at the current offset
-					ClassBeingBuilt->GetDebugData().RegisterPureNodeScriptCodeRange(TrueSourceNode, FunctionContext.Function, FInt32Range(PureNodeEntryStart, Offset));
+					ClassBeingBuilt->GetDebugData().RegisterPureNodeScriptCodeRange(MacroSourceNode ? MacroSourceNode : TrueSourceNode, FunctionContext.Function, FInt32Range(PureNodeEntryStart, Offset));
 
 					// Reset pure node code range tracking.
 					PureNodeEntryCount = 0;
diff --git a/Engine/Source/Editor/KismetCompiler/Public/KismetCompiler.h b/Engine/Source/Editor/KismetCompiler/Public/KismetCompiler.h
--- a/Engine/Source/Editor/KismetCompiler/Public/KismetCompiler.h
+++ b/Engine/Source/Editor/KismetCompiler/Public/KismetCompiler.h
@@ -113,7 +113,7 @@ public:
 		}
 
 		NodeType* Result = ParentGraph->CreateBlankNode<NodeType>();
-		check (Cast<UK2Node_Event>(Result) == nullptr);
+		//check (Cast<UK2Node_Event>(Result) == nullptr); -- Removed to avoid any fallout, will replace with care later
 		MessageLog.NotifyIntermediateObjectCreation(Result, SourceNode); // this might be useful to track back function entry nodes to events.
 		Result->CreateNewGuid();
 
@@ -132,7 +132,7 @@ public:
 		}
 
 		NodeType* Result = ParentGraph->CreateBlankNode<NodeType>();
-		check (Cast<UK2Node_Event>(Result) != nullptr);
+		//check (Cast<UK2Node_Event>(Result) != nullptr); -- Removed to avoid any fallout, will replace with care later
 		MessageLog.NotifyIntermediateObjectCreation(Result, SourceNode); // this might be useful to track back function entry nodes to events.
 		Result->CreateNewGuid();
 		// Track all expansion step generated events.
diff --git a/Engine/Source/Editor/UnrealEd/Classes/Settings/EditorExperimentalSettings.h b/Engine/Source/Editor/UnrealEd/Classes/Settings/EditorExperimentalSettings.h
--- a/Engine/Source/Editor/UnrealEd/Classes/Settings/EditorExperimentalSettings.h
+++ b/Engine/Source/Editor/UnrealEd/Classes/Settings/EditorExperimentalSettings.h
@@ -67,30 +67,6 @@ public:
 	UPROPERTY(EditAnywhere, config, Category=Blueprints, meta=(DisplayName="Blueprint Performance Analysis Tools"))
 	bool bBlueprintPerformanceAnalysisTools;
 
-	/** Recent sample bias state. */
-	UPROPERTY(EditAnywhere, config, Category=Blueprints, meta=(DisplayName="Enable Blueprint Profiler Sample Bias"))
-	bool bEnableBlueprintProfilerRecentSampleBias;
-
-	/** Bias to weight prominance of newer samples against hositorical samples. */
-	UPROPERTY(EditAnywhere, config, Category=Blueprints, meta=(DisplayName="Blueprint Profiler Sample Bias", ClampMin=0.01, ClampMax=1.0))
-	float BlueprintProfilerRecentSampleBias;
-
-	/** Blueprint profiler performance threshold for events in ms. */
-	UPROPERTY(EditAnywhere, config, Category=Blueprints, meta=(DisplayName="Blueprint Profiler Event Timing Performance Threshold", ClampMin=0.0001, ClampMax=20.0))
-	float BlueprintProfilerEventThreshold;
-
-	/** Blueprint profiler performance threshold for exclusive node timings in ms. */
-	UPROPERTY(EditAnywhere, config, Category=Blueprints, meta=(DisplayName="Blueprint Profiler Exclusive Node Performance Threshold", ClampMin=0.0001, ClampMax=20.0))
-	float BlueprintProfilerExclNodeThreshold;
-
-	/** Blueprint profiler performance threshold for inclusive node timings in ms. */
-	UPROPERTY(EditAnywhere, config, Category=Blueprints, meta=(DisplayName="Blueprint Profiler Inclusive Node Performance Threshold", ClampMin=0.0001, ClampMax=20.0))
-	float BlueprintProfilerInclNodeThreshold;
-
-	/** Blueprint profiler performance threshold for max node timings in ms. */
-	UPROPERTY(EditAnywhere, config, Category=Blueprints, meta=(DisplayName="Blueprint Profiler Max Node Performance Threshold", ClampMin=0.0001, ClampMax=20.0))
-	float BlueprintProfilerMaxNodeThreshold;
-
 	/** Enables the visual diff tool for widget blueprints. WARNING: changes to the widget hierarchy will not be detected */
 	UPROPERTY(EditAnywhere, config, Category=Blueprints, meta=(DisplayName="Use the Diff Tool for Widget Blueprints"))
 	bool bEnableWidgetVisualDiff;
diff --git a/Engine/Source/Editor/UnrealEd/Private/EdGraphUtilities.cpp b/Engine/Source/Editor/UnrealEd/Private/EdGraphUtilities.cpp
--- a/Engine/Source/Editor/UnrealEd/Private/EdGraphUtilities.cpp
+++ b/Engine/Source/Editor/UnrealEd/Private/EdGraphUtilities.cpp
@@ -240,17 +240,7 @@ void FEdGraphUtilities::CloneAndMergeGraphIn(UEdGraph* MergeTarget, UEdGraph* So
 {
 	// Clone the graph, then move all of it's children
 	UEdGraph* ClonedGraph = CloneGraph(SourceGraph, NULL, &MessageLog, true);
-
-	// Create any Boundary nodes around child graphs if specified.
-	if (bCreateBoundaryNodes)
-	{
-		for (UEdGraph* SubGraph : ClonedGraph->SubGraphs)
-		{
-			UK2Node_TunnelBoundary::CreateBoundaryNodesForGraph(SubGraph, MessageLog);
-		}
-	}
-
-	MergeChildrenGraphsIn(ClonedGraph, ClonedGraph, bRequireSchemaMatch);
+	MergeChildrenGraphsIn(ClonedGraph, ClonedGraph, bRequireSchemaMatch, false, &MessageLog, bCreateBoundaryNodes);
 
 	// Duplicate the list of cloned nodes
 	if (OutClonedNodes != NULL)
@@ -267,7 +257,7 @@ void FEdGraphUtilities::CloneAndMergeGraphIn(UEdGraph* MergeTarget, UEdGraph* So
 }
 
 // Moves the contents of all of the children graphs (recursively) into the target graph.  This does not clone, it's destructive to the source
-void FEdGraphUtilities::MergeChildrenGraphsIn(UEdGraph* MergeTarget, UEdGraph* ParentGraph, bool bRequireSchemaMatch, bool bInIsCompiling/* = false*/)
+void FEdGraphUtilities::MergeChildrenGraphsIn(UEdGraph* MergeTarget, UEdGraph* ParentGraph, bool bRequireSchemaMatch, bool bInIsCompiling/* = false*/, FCompilerResultsLog* MessageLog/* = nullptr*/, bool bWantBoundaryNodes/* = false*/)
 {
 	// Determine if we are regenerating a blueprint on load
 	UBlueprint* Blueprint = FBlueprintEditorUtils::FindBlueprintForGraph(MergeTarget);
@@ -278,6 +268,12 @@ void FEdGraphUtilities::MergeChildrenGraphsIn(UEdGraph* MergeTarget, UEdGraph* P
 	{
 		UEdGraph* ChildGraph = ParentGraph->SubGraphs[Index];
 
+		if (bWantBoundaryNodes)
+		{
+			// Create boundary nodes around tunnels for debugging/profiling if requested.
+			UK2Node_TunnelBoundary::CreateBoundaryNodesForGraph(ChildGraph, *MessageLog);
+		}
+
 		auto NodeOwner = Cast<const UEdGraphNode>(ChildGraph ? ChildGraph->GetOuter() : nullptr);
 		const bool bNonVirtualGraph = NodeOwner ? NodeOwner->ShouldMergeChildGraphs() : true;
 
@@ -294,7 +290,7 @@ void FEdGraphUtilities::MergeChildrenGraphsIn(UEdGraph* MergeTarget, UEdGraph* P
 				ChildGraph->MoveNodesToAnotherGraph(MergeTarget, IsAsyncLoading() || bIsLoading, bInIsCompiling);
 			}
 
-			MergeChildrenGraphsIn(MergeTarget, ChildGraph, bRequireSchemaMatch, bInIsCompiling);
+			MergeChildrenGraphsIn(MergeTarget, ChildGraph, bRequireSchemaMatch, bInIsCompiling, MessageLog, bWantBoundaryNodes);
 		}
 	}
 }
diff --git a/Engine/Source/Editor/UnrealEd/Private/Kismet2/BlueprintEditorUtils.cpp b/Engine/Source/Editor/UnrealEd/Private/Kismet2/BlueprintEditorUtils.cpp
--- a/Engine/Source/Editor/UnrealEd/Private/Kismet2/BlueprintEditorUtils.cpp
+++ b/Engine/Source/Editor/UnrealEd/Private/Kismet2/BlueprintEditorUtils.cpp
@@ -3004,7 +3004,7 @@ void FBlueprintEditorUtils::EnsureCachedDependenciesUpToDate(UBlueprint* Bluepri
 	}
 }
 
-void FBlueprintEditorUtils::GetDependentBlueprints(UBlueprint* Blueprint, TArray<UBlueprint*>& DependentBlueprints)
+void FBlueprintEditorUtils::GetDependentBlueprints(UBlueprint* Blueprint, TArray<UBlueprint*>& DependentBlueprints, bool bRemoveSelf)
 {
 	TArray<UObject*> AllBlueprints;
 	bool const bIncludeDerivedClasses = true;
@@ -3028,12 +3028,17 @@ void FBlueprintEditorUtils::GetDependentBlueprints(UBlueprint* Blueprint, TArray
 					// depends on it must also depend on this Blueprint for re-compiling (bytecode, skeleton, full) purposes
 					if (TestBP->BlueprintType == BPTYPE_MacroLibrary)
 					{
-						GetDependentBlueprints(TestBP, DependentBlueprints);
+						GetDependentBlueprints(TestBP, DependentBlueprints, false);
 					}
 				}
 			}
 		}
 	}
+
+	if (bRemoveSelf)
+	{
+		DependentBlueprints.RemoveSwap(Blueprint);
+	}
 }
 
 
@@ -3284,6 +3289,27 @@ UEdGraph* FBlueprintEditorUtils::FindEventGraph(const UBlueprint* Blueprint)
 	return NULL;
 }
 
+bool FBlueprintEditorUtils::IsEventGraph(const UEdGraph* InGraph)
+{
+	if (InGraph)
+	{
+		if (const UBlueprint* Blueprint = FindBlueprintForGraph(InGraph))
+		{
+			return (nullptr != Blueprint->UbergraphPages.FindByKey(InGraph));
+		}
+	}
+	return false;
+}
+
+bool FBlueprintEditorUtils::IsTunnelInstanceNode(const UEdGraphNode* InGraphNode)
+{
+	if (InGraphNode)
+	{
+		return InGraphNode->IsA<UK2Node_MacroInstance>() || InGraphNode->IsA<UK2Node_Composite>();
+	}
+	return false;
+}
+
 bool FBlueprintEditorUtils::DoesBlueprintDeriveFrom(const UBlueprint* Blueprint, UClass* TestClass)
 {
 	check(Blueprint->SkeletonGeneratedClass != NULL);
diff --git a/Engine/Source/Editor/UnrealEd/Private/Kismet2/Kismet2.cpp b/Engine/Source/Editor/UnrealEd/Private/Kismet2/Kismet2.cpp
--- a/Engine/Source/Editor/UnrealEd/Private/Kismet2/Kismet2.cpp
+++ b/Engine/Source/Editor/UnrealEd/Private/Kismet2/Kismet2.cpp
@@ -679,7 +679,7 @@ bool FKismetEditorUtilities::IsReferencedByUndoBuffer(UBlueprint* Blueprint)
 	return (TotalReferenceCount > NonUndoReferenceCount);
 }
 
-void FKismetEditorUtilities::CompileBlueprint(UBlueprint* BlueprintObj, bool bIsRegeneratingOnLoad, bool bSkipGarbageCollection, bool bSaveIntermediateProducts, FCompilerResultsLog* pResults, bool bSkeletonUpToDate, bool bBatchCompile)
+void FKismetEditorUtilities::CompileBlueprint(UBlueprint* BlueprintObj, bool bIsRegeneratingOnLoad, bool bSkipGarbageCollection, bool bSaveIntermediateProducts, FCompilerResultsLog* pResults, bool bSkeletonUpToDate, bool bBatchCompile, bool bAddInstrumentation)
 {
 	FSecondsCounterScope Timer(BlueprintCompileAndLoadTimerData); 
 	BP_SCOPED_COMPILER_EVENT_STAT(EKismetCompilerStats_CompileBlueprint);
@@ -743,14 +743,10 @@ void FKismetEditorUtilities::CompileBlueprint(UBlueprint* BlueprintObj, bool bIs
 	static const FBoolConfigValueHelper IgnoreCompileOnLoadErrorsOnBuildMachine(TEXT("Kismet"), TEXT("bIgnoreCompileOnLoadErrorsOnBuildMachine"), GEngineIni);
 	Results.bLogInfoOnly = BlueprintObj->bIsRegeneratingOnLoad && GIsBuildMachine && IgnoreCompileOnLoadErrorsOnBuildMachine;
 
-	// Determine if we want profiling data
-	IBlueprintProfilerInterface* ProfilerInterface = FModuleManager::GetModulePtr<IBlueprintProfilerInterface>("BlueprintProfiler");
-	const bool bProfilerActive = ProfilerInterface && ProfilerInterface->IsProfilerEnabled() && GetDefault<UEditorExperimentalSettings>()->bBlueprintPerformanceAnalysisTools;
-
 	FKismetCompilerOptions CompileOptions;
 	CompileOptions.bSaveIntermediateProducts = bSaveIntermediateProducts;
 	CompileOptions.bRegenerateSkelton = !bSkeletonUpToDate;
-	CompileOptions.bAddInstrumentation = bProfilerActive;
+	CompileOptions.bAddInstrumentation = bAddInstrumentation;
 	Compiler.CompileBlueprint(BlueprintObj, CompileOptions, Results, ReinstanceHelper);
 
 	FBlueprintEditorUtils::UpdateDelegatesInBlueprint(BlueprintObj);
diff --git a/Engine/Source/Editor/UnrealEd/Private/Kismet2/KismetReinstanceUtilities.cpp b/Engine/Source/Editor/UnrealEd/Private/Kismet2/KismetReinstanceUtilities.cpp
--- a/Engine/Source/Editor/UnrealEd/Private/Kismet2/KismetReinstanceUtilities.cpp
+++ b/Engine/Source/Editor/UnrealEd/Private/Kismet2/KismetReinstanceUtilities.cpp
@@ -1351,8 +1351,9 @@ void FBlueprintCompileReinstancer::ReplaceInstancesOfClass_Inner(TMap<UClass*, U
 				GetObjectsOfClass(OldClass, ObjectsToReplace, bIncludeDerivedClasses);
 
 				// Then fix 'real' (non archetype) instances of the class
-				for (UObject* OldObject : ObjectsToReplace)
+				for (int32 OldObjIndex = 0; OldObjIndex < ObjectsToReplace.Num(); ++OldObjIndex)
 				{
+					UObject* OldObject = ObjectsToReplace[OldObjIndex];
 					// Skip non-archetype instances, EXCEPT for component templates
 					const bool bIsComponent = NewClass->IsChildOf(UActorComponent::StaticClass());
 					if ((!bIsComponent && OldObject->IsTemplate()) || OldObject->IsPendingKill())
@@ -1361,13 +1362,27 @@ void FBlueprintCompileReinstancer::ReplaceInstancesOfClass_Inner(TMap<UClass*, U
 						continue;
 					}
 
-					UBlueprint* CorrespondingBlueprint = Cast<UBlueprint>(OldObject->GetClass()->ClassGeneratedBy);
-					UObject*    OldBlueprintDebugObject = nullptr;
-					// If this object is being debugged, cache it off so we can preserve the 'object being debugged' association
-					if ((CorrespondingBlueprint != nullptr) && (CorrespondingBlueprint->GetObjectBeingDebugged() == OldObject))
+					UBlueprint* DebuggingBlueprint = nullptr;
+					// ObjectsToReplace is not in any garunteed order; so below there may be a 
+					// scenario where one object in ObjectsToReplace needs to come before another; 
+					// in that scenario we swap the two in ObjectsToReplace and finish the loop 
+					// with a new OldObject - to support that we need to re-cache info about the 
+					// (new) old object before it is destroyed (hence why this is all in a lambda,  
+					// for reuse).
+					auto CacheOldObjectState = [&DebuggingBlueprint](UObject* Replacee)
 					{
-						OldBlueprintDebugObject = OldObject;
-					}
+						// clear in case it has already been set, and we're re-executing this
+						DebuggingBlueprint = nullptr;
+
+						if (UBlueprint* OldObjBlueprint = Cast<UBlueprint>(Replacee->GetClass()->ClassGeneratedBy))
+						{
+							if (OldObjBlueprint->GetObjectBeingDebugged() == Replacee)
+							{
+								DebuggingBlueprint = OldObjBlueprint;
+							}
+						}
+					};
+					CacheOldObjectState(OldObject);
 
 					AActor*  OldActor = Cast<AActor>(OldObject);
 					UObject* NewUObject = nullptr;
@@ -1475,14 +1490,47 @@ void FBlueprintCompileReinstancer::ReplaceInstancesOfClass_Inner(TMap<UClass*, U
 					}
 					else
 					{
-						auto OldFlags = OldObject->GetFlags();
-
 						// If the old object was spawned from an archetype (i.e. not the CDO), we must use the new version of that archetype as the template object when constructing the new instance.
 						UObject* OldArchetype = OldObject->GetArchetype();
 						UObject* NewArchetype = OldToNewInstanceMap.FindRef(OldArchetype);
 
+						bool bArchetypeReinstanced = (OldArchetype == OldClass->GetDefaultObject()) || (NewArchetype != nullptr);
+						// if we don't have a updated archetype to spawn from, we need to update/reinstance it
+						while (!bArchetypeReinstanced)
+						{
+							int32 ArchetypeIndex = ObjectsToReplace.Find(OldArchetype);
+							if (ArchetypeIndex != INDEX_NONE)
+							{
+								if (ensure(ArchetypeIndex > OldObjIndex))
+								{
+									// if this object has an archetype, but it hasn't been 
+									// reinstanced yet (but is queued to) then we need to swap out 
+									// the two, and reinstance the archetype first
+									ObjectsToReplace.Swap(ArchetypeIndex, OldObjIndex);
+									OldObject = ObjectsToReplace[OldObjIndex];
+									check(OldObject == OldArchetype);
+
+									// update any info that we stored regarding the previous OldObject
+									CacheOldObjectState(OldObject);
+									
+									OldArchetype = OldObject->GetArchetype();
+									NewArchetype = OldToNewInstanceMap.FindRef(OldArchetype);
+									bArchetypeReinstanced = (OldArchetype == OldClass->GetDefaultObject()) || (NewArchetype != nullptr);
+								}
+								else
+								{
+									break;
+								}
+							}
+							else
+							{
+								break;
+							}
+						}
 						// Check that either this was an instance of the class directly, or we found a new archetype for it
-						check(OldArchetype == OldClass->GetDefaultObject() || NewArchetype);
+						ensureMsgf(bArchetypeReinstanced, TEXT("Reinstancing non-actor (%s); failed to resolve archetype object - property values may be lost."), *OldObject->GetPathName());
+
+						EObjectFlags OldFlags = OldObject->GetFlags();
 
 						FName OldName(OldObject->GetFName());
 						OldObject->Rename(NULL, OldObject->GetOuter(), REN_DoNotDirty | REN_DontCreateRedirectors);
@@ -1579,9 +1627,9 @@ void FBlueprintCompileReinstancer::ReplaceInstancesOfClass_Inner(TMap<UClass*, U
 					}
 
 					// If this original object came from a blueprint and it was in the selected debug set, change the debugging to the new object.
-					if ((CorrespondingBlueprint) && (OldBlueprintDebugObject) && (NewUObject))
+					if (DebuggingBlueprint && NewUObject)
 					{
-						CorrespondingBlueprint->SetObjectBeingDebugged(NewUObject);
+						DebuggingBlueprint->SetObjectBeingDebugged(NewUObject);
 					}
 
 					if (bLogConversions)
diff --git a/Engine/Source/Editor/UnrealEd/Private/Settings/SettingsClasses.cpp b/Engine/Source/Editor/UnrealEd/Private/Settings/SettingsClasses.cpp
--- a/Engine/Source/Editor/UnrealEd/Private/Settings/SettingsClasses.cpp
+++ b/Engine/Source/Editor/UnrealEd/Private/Settings/SettingsClasses.cpp
@@ -70,14 +70,7 @@ UEditorExperimentalSettings::UEditorExperimentalSettings( const FObjectInitializ
 	, bUnifiedBlueprintEditor(true)
 	, bBlueprintableComponents(true)
 	, bBlueprintPerformanceAnalysisTools(false)
-	, bEnableBlueprintProfilerRecentSampleBias(false)
-	, BlueprintProfilerRecentSampleBias(0.2f)
-	, BlueprintProfilerEventThreshold(1.f)
-	, BlueprintProfilerExclNodeThreshold(0.2f)
-	, BlueprintProfilerInclNodeThreshold(0.25f)
-	, BlueprintProfilerMaxNodeThreshold(0.5f)
 	, bUseOpenCLForConvexHullDecomp(false)
-	, bEnableAutoVREditMode(true)
 	, bAllowPotentiallyUnsafePropertyEditing(false)
 {
 }
@@ -113,30 +106,6 @@ void UEditorExperimentalSettings::PostEditChangeProperty( struct FPropertyChange
 			}
 		}
 	}
-	else if (Name == FName(TEXT("bEnableBlueprintProfilerRecentSampleBias")))
-	{
-		FScriptPerfData::EnableRecentSampleBias(bEnableBlueprintProfilerRecentSampleBias);
-	}
-	else if (Name == FName(TEXT("BlueprintProfilerRecentSampleBias")))
-	{
-		FScriptPerfData::SetRecentSampleBias(BlueprintProfilerRecentSampleBias);
-	}
-	else if (Name == FName(TEXT("BlueprintProfilerEventThreshold")))
-	{
-		FScriptPerfData::SetEventPerformanceThreshold(BlueprintProfilerEventThreshold);
-	}
-	else if (Name == FName(TEXT("BlueprintProfilerExclNodeThreshold")))
-	{
-		FScriptPerfData::SetExclusivePerformanceThreshold(BlueprintProfilerExclNodeThreshold);
-	}
-	else if (Name == FName(TEXT("BlueprintProfilerInclNodeThreshold")))
-	{
-		FScriptPerfData::SetInclusivePerformanceThreshold(BlueprintProfilerInclNodeThreshold);
-	}
-	else if (Name == FName(TEXT("BlueprintProfilerMaxNodeThreshold")))
-	{
-		FScriptPerfData::SetMaxPerformanceThreshold(BlueprintProfilerMaxNodeThreshold);
-	}
 
 	if (!FUnrealEdMisc::Get().IsDeletePreferences())
 	{
diff --git a/Engine/Source/Editor/UnrealEd/Public/EdGraphUtilities.h b/Engine/Source/Editor/UnrealEd/Public/EdGraphUtilities.h
--- a/Engine/Source/Editor/UnrealEd/Public/EdGraphUtilities.h
+++ b/Engine/Source/Editor/UnrealEd/Public/EdGraphUtilities.h
@@ -76,7 +76,7 @@ public:
 	 * Moves the contents of all of the children graphs of ParentGraph (recursively) into the MergeTarget graph.
 	 * This does not clone, it's destructive to the ParentGraph
 	 */
-	static void MergeChildrenGraphsIn(UEdGraph* MergeTarget, UEdGraph* ParentGraph, bool bRequireSchemaMatch, bool bInIsCompiling = false);
+	static void MergeChildrenGraphsIn(UEdGraph* MergeTarget, UEdGraph* ParentGraph, bool bRequireSchemaMatch, bool bInIsCompiling = false, FCompilerResultsLog* MessageLog = nullptr, bool bWantBoundaryNodes = false);
 
 	/** Tries to rename the graph to have a name similar to BaseName */
 	static void RenameGraphCloseToName(UEdGraph* Graph, const FString& BaseName, int32 StartIndex = 1);
diff --git a/Engine/Source/Editor/UnrealEd/Public/Kismet2/BlueprintEditorUtils.h b/Engine/Source/Editor/UnrealEd/Public/Kismet2/BlueprintEditorUtils.h
--- a/Engine/Source/Editor/UnrealEd/Public/Kismet2/BlueprintEditorUtils.h
+++ b/Engine/Source/Editor/UnrealEd/Public/Kismet2/BlueprintEditorUtils.h
@@ -438,7 +438,7 @@ public:
 	static void GatherDependencies(const UBlueprint* Blueprint, TSet<TWeakObjectPtr<UBlueprint>>& OutDependencies, TSet<TWeakObjectPtr<UStruct>>& OutUDSDependencies);
 
 	/** Returns a list of loaded Blueprints that are dependent on the given Blueprint. */
-	static void GetDependentBlueprints(UBlueprint* Blueprint, TArray<UBlueprint*>& DependentBlueprints);
+	static void GetDependentBlueprints(UBlueprint* Blueprint, TArray<UBlueprint*>& DependentBlueprints, bool bRemoveSelf = true);
 
 	/** Ensures, that CachedDependencies in BP are up to date */
 	static void EnsureCachedDependenciesUpToDate(UBlueprint* Blueprint);
@@ -505,6 +505,12 @@ public:
 	/** Returns the event graph, if any */
 	static UEdGraph* FindEventGraph(const UBlueprint* Blueprint);
 
+	/** Checks if given graph is an event graph */
+	static bool IsEventGraph(const UEdGraph* InGraph);
+
+	/** Checks if given node is a tunnel instance node */
+	static bool IsTunnelInstanceNode(const UEdGraphNode* InGraphNode);
+
 	/** See if a class is the one generated by this blueprint */
 	static bool DoesBlueprintDeriveFrom(const UBlueprint* Blueprint, UClass* TestClass);
 
diff --git a/Engine/Source/Editor/UnrealEd/Public/Kismet2/CompilerResultsLog.h b/Engine/Source/Editor/UnrealEd/Public/Kismet2/CompilerResultsLog.h
--- a/Engine/Source/Editor/UnrealEd/Public/Kismet2/CompilerResultsLog.h
+++ b/Engine/Source/Editor/UnrealEd/Public/Kismet2/CompilerResultsLog.h
@@ -94,6 +94,11 @@ public:
 	FBacktrackMap FinalNodeBackToMacroSourceMap;
 	TMultiMap<TWeakObjectPtr<UEdGraphNode>, TWeakObjectPtr<UEdGraphNode>> MacroSourceToMacroInstanceNodeMap;
 
+	// Used to track node generatation from tunnels, added in addition to existing code avoid causing any fallout.
+	// This will be refactored after blueprint profiler MVP.
+	TMap<TWeakObjectPtr<UEdGraphNode>, TWeakObjectPtr<UEdGraphNode>> SourceNodeToTunnelInstanceNodeMap;
+	TMap<TWeakObjectPtr<UEdGraphNode>, TWeakObjectPtr<UEdGraphNode>> IntermediateTunnelNodeToSourceNodeMap;
+
 	// Minimum event time (ms) for inclusion into the final summary log
 	int EventDisplayThresholdMs;
 
diff --git a/Engine/Source/Editor/UnrealEd/Public/Kismet2/KismetEditorUtilities.h b/Engine/Source/Editor/UnrealEd/Public/Kismet2/KismetEditorUtilities.h
--- a/Engine/Source/Editor/UnrealEd/Public/Kismet2/KismetEditorUtilities.h
+++ b/Engine/Source/Editor/UnrealEd/Public/Kismet2/KismetEditorUtilities.h
@@ -97,7 +97,7 @@ public:
 	static void CreateDefaultEventGraphs(UBlueprint* Blueprint);
 
 	/** Tries to compile a blueprint, updating any actors in the editor who are using the old class, etc... */
-	static void CompileBlueprint(UBlueprint* BlueprintObj, bool bIsRegeneratingOnLoad = false, bool bSkipGarbageCollection = false, bool bSaveIntermediateProducts = false, class FCompilerResultsLog* pResults = nullptr, bool bSkeletonUpToDate = false, bool bBatchCompile = false);
+	static void CompileBlueprint(UBlueprint* BlueprintObj, bool bIsRegeneratingOnLoad = false, bool bSkipGarbageCollection = false, bool bSaveIntermediateProducts = false, class FCompilerResultsLog* pResults = nullptr, bool bSkeletonUpToDate = false, bool bBatchCompile = false, bool bAddInstrumentation = false);
 
 	/** Generates a blueprint skeleton only.  Minimal compile, no notifications will be sent, no GC, etc.  Only successful if there isn't already a skeleton generated */
 	static bool GenerateBlueprintSkeleton(UBlueprint* BlueprintObj, bool bForceRegeneration = false);
diff --git a/Engine/Source/Runtime/AnimGraphRuntime/Public/AnimNodes/AnimNode_BlendListByBool.h b/Engine/Source/Runtime/AnimGraphRuntime/Public/AnimNodes/AnimNode_BlendListByBool.h
--- a/Engine/Source/Runtime/AnimGraphRuntime/Public/AnimNodes/AnimNode_BlendListByBool.h
+++ b/Engine/Source/Runtime/AnimGraphRuntime/Public/AnimNodes/AnimNode_BlendListByBool.h
@@ -20,10 +20,6 @@ public:
 	FAnimNode_BlendListByBool()
 		: FAnimNode_BlendListBase()
 	{
-#if	WITH_EDITOR
-		AddPose();
-		AddPose();
-#endif
 	}
 
 protected:
diff --git a/Engine/Source/Runtime/CoreUObject/Private/Blueprint/BlueprintSupport.cpp b/Engine/Source/Runtime/CoreUObject/Private/Blueprint/BlueprintSupport.cpp
--- a/Engine/Source/Runtime/CoreUObject/Private/Blueprint/BlueprintSupport.cpp
+++ b/Engine/Source/Runtime/CoreUObject/Private/Blueprint/BlueprintSupport.cpp
@@ -2051,6 +2051,14 @@ void FDeferredObjInitializerTracker::ResolveDeferredSubClassObjects(UClass* Supe
 // don't want other files ending up with this internal define
 #undef DEFERRED_DEPENDENCY_CHECK
 
+FBlueprintDependencyData::FBlueprintDependencyData(const TCHAR* InPackageName
+	, const TCHAR* InObjectName, const TCHAR* InClassPackageName, const TCHAR* InClassName) 
+	: PackageName(InPackageName)
+	, ObjectName(InObjectName)
+	, ClassPackageName(InClassPackageName)
+	, ClassName(InClassName)
+{}
+
 FConvertedBlueprintsDependencies& FConvertedBlueprintsDependencies::Get()
 {
 	static FConvertedBlueprintsDependencies ConvertedBlueprintsDependencies;
@@ -2075,6 +2083,23 @@ void FConvertedBlueprintsDependencies::GetAssets(FName PackageName, TArray<FBlue
 	}
 }
 
+void FConvertedBlueprintsDependencies::FillUsedAssetsInDynamicClass(UDynamicClass* DynamicClass, GetDependenciesNamesFunc GetUsedAssets)
+{
+	check(DynamicClass && GetUsedAssets);
+	ensure(DynamicClass->UsedAssets.Num() == 0);
+
+	TArray<FBlueprintDependencyData> UsedAssetdData;
+	GetUsedAssets(UsedAssetdData);
+
+	for (FBlueprintDependencyData& ItData : UsedAssetdData)
+	{
+		const FString PathToObj = FString::Printf(TEXT("%s.%s"), *ItData.PackageName.ToString(), *ItData.ObjectName.ToString());
+		UObject* TheAsset = LoadObject<UObject>(nullptr, *PathToObj);
+		DynamicClass->UsedAssets.Add(TheAsset);
+	}
+}
+
+
 #if WITH_EDITOR
 
 /*******************************************************************************
diff --git a/Engine/Source/Runtime/CoreUObject/Private/UObject/ScriptCore.cpp b/Engine/Source/Runtime/CoreUObject/Private/UObject/ScriptCore.cpp
--- a/Engine/Source/Runtime/CoreUObject/Private/UObject/ScriptCore.cpp
+++ b/Engine/Source/Runtime/CoreUObject/Private/UObject/ScriptCore.cpp
@@ -144,7 +144,7 @@ void FBlueprintCoreDelegates::ThrowScriptException(const UObject* ActiveObject,
 	}
 }
 
-void FBlueprintCoreDelegates::InstrumentScriptEvent(const EScriptInstrumentationEvent& Info)
+void FBlueprintCoreDelegates::InstrumentScriptEvent(const FScriptInstrumentationSignal& Info)
 {
 	OnScriptProfilingEvent.Broadcast(Info);
 }
@@ -449,6 +449,47 @@ FString FFrame::GetStackTrace() const
 	return Result;
 }
 
+//////////////////////////////////////////////////////////////////////////
+// FScriptInstrumentationSignal
+
+FScriptInstrumentationSignal::FScriptInstrumentationSignal(EScriptInstrumentation::Type InEventType, const UObject* InContextObject, const struct FFrame& InStackFrame)
+	: EventType(InEventType)
+	, ContextObject(InContextObject)
+	, Function(InStackFrame.Node)
+	, StackFramePtr(&InStackFrame)
+	, LatentLinkId(INDEX_NONE)
+{
+}
+
+const UClass* FScriptInstrumentationSignal::GetClass() const
+{
+	return ContextObject ? ContextObject->GetClass() : nullptr;
+}
+
+const UClass* FScriptInstrumentationSignal::GetFunctionClassScope() const
+{
+	return Function->GetOuterUClass();
+}
+
+FName FScriptInstrumentationSignal::GetFunctionName() const
+{
+	return Function->GetFName();
+}
+
+int32 FScriptInstrumentationSignal::GetScriptCodeOffset() const
+{
+	int32 CodeOffset = INDEX_NONE;
+	if (EventType == EScriptInstrumentation::ResumeEvent)
+	{
+		// Resume events require the link id rather than script code offset
+		CodeOffset = LatentLinkId;
+	}
+	else if (StackFramePtr != nullptr)
+	{
+		CodeOffset = StackFramePtr->Code - StackFramePtr->Node->Script.GetData() - 1;
+	}
+	return CodeOffset;
+}
 
 /*-----------------------------------------------------------------------------
 	Global script execution functions.
@@ -1150,7 +1191,7 @@ void UObject::ProcessEvent( UFunction* Function, void* Parms )
 			if (Function->HasAnyFunctionFlags(FUNC_Event|FUNC_BlueprintEvent))
 			{
 				// Don't handle latent actions here, let the latent action manager handle them.
-				EScriptInstrumentationEvent EventInstrumentationInfo(EScriptInstrumentation::Event, this, Function->GetFName());
+				FScriptInstrumentationSignal EventInstrumentationInfo(EScriptInstrumentation::Event, this, Function);
 				FBlueprintCoreDelegates::InstrumentScriptEvent(EventInstrumentationInfo);
 				bInstrumentScriptEvent = true;
 			}
@@ -1287,7 +1328,7 @@ void UObject::ProcessEvent( UFunction* Function, void* Parms )
 #if !(UE_BUILD_SHIPPING || UE_BUILD_TEST)
 	if (bInstrumentScriptEvent)
 	{
-		EScriptInstrumentationEvent EventInstrumentationInfo(EScriptInstrumentation::Stop, this, NAME_None);
+		FScriptInstrumentationSignal EventInstrumentationInfo(EScriptInstrumentation::Stop, this, Function);
 		FBlueprintCoreDelegates::InstrumentScriptEvent(EventInstrumentationInfo);
 	}
 #if WITH_EDITORONLY_DATA
@@ -1579,7 +1620,7 @@ void UObject::execInstrumentation( FFrame& Stack, RESULT_DECL )
 		}
 	}
 #endif
-	EScriptInstrumentationEvent InstrumentationEventInfo(EventType, this, Stack);
+	FScriptInstrumentationSignal InstrumentationEventInfo(EventType, this, Stack);
 	FBlueprintCoreDelegates::InstrumentScriptEvent(InstrumentationEventInfo);
 	Stack.SkipCode(1);
 #endif
diff --git a/Engine/Source/Runtime/CoreUObject/Public/Blueprint/BlueprintSupport.h b/Engine/Source/Runtime/CoreUObject/Public/Blueprint/BlueprintSupport.h
--- a/Engine/Source/Runtime/CoreUObject/Public/Blueprint/BlueprintSupport.h
+++ b/Engine/Source/Runtime/CoreUObject/Public/Blueprint/BlueprintSupport.h
@@ -217,6 +217,13 @@ struct COREUOBJECT_API FBlueprintDependencyData
 	FName ObjectName;
 	FName ClassPackageName;
 	FName ClassName;
+
+	FBlueprintDependencyData() {}
+
+	FORCENOINLINE FBlueprintDependencyData(const TCHAR* InPackageName
+		, const TCHAR* InObjectName
+		, const TCHAR* InClassPackageName
+		, const TCHAR* InClassName);
 };
 
 /**
@@ -237,4 +244,6 @@ public:
 
 	/** Get all assets paths necessary for the class with the given class name and all converted classes that dependencies. */
 	void GetAssets(FName PackageName, TArray<FBlueprintDependencyData>& OutDependencies) const;
+
+	static void FillUsedAssetsInDynamicClass(UDynamicClass* DynamicClass, GetDependenciesNamesFunc GetUsedAssets);
 };
\ No newline at end of file
diff --git a/Engine/Source/Runtime/CoreUObject/Public/Serialization/ArchiveUObject.h b/Engine/Source/Runtime/CoreUObject/Public/Serialization/ArchiveUObject.h
--- a/Engine/Source/Runtime/CoreUObject/Public/Serialization/ArchiveUObject.h
+++ b/Engine/Source/Runtime/CoreUObject/Public/Serialization/ArchiveUObject.h
@@ -1147,6 +1147,11 @@ public:
 			// start the initial serialization
 			SerializedObjects.Add(SearchObject);
 			SerializeObject(SearchObject);
+			for (int32 Iter = 0; Iter < PendingSerializationObjects.Num(); Iter++)
+			{
+				SerializeObject(PendingSerializationObjects[Iter]);
+			}
+			PendingSerializationObjects.Reset();
 		}
 	}
 
@@ -1198,8 +1203,8 @@ public:
 				SerializedObjects.Add(Obj, &bAlreadyAdded);
 				if (!bAlreadyAdded)
 				{
-					// otherwise recurse down into the object if it is contained within the initial search object
-					SerializeObject(Obj);
+					// No recursion
+					PendingSerializationObjects.Add(Obj);
 				}
 			}
 			else if ( bNullPrivateReferences && !Obj->HasAnyFlags(RF_Public) )
@@ -1228,6 +1233,9 @@ protected:
 	/** List of objects that have already been serialized */
 	TSet<UObject*> SerializedObjects;
 
+	/** Object that will be serialized */
+	TArray<UObject*> PendingSerializationObjects;
+
 	/** Map of referencing objects to referencing properties */
 	TMap<UObject*, TArray<UProperty*>> ReplacedReferences;
 
diff --git a/Engine/Source/Runtime/CoreUObject/Public/UObject/Script.h b/Engine/Source/Runtime/CoreUObject/Public/UObject/Script.h
--- a/Engine/Source/Runtime/CoreUObject/Public/UObject/Script.h
+++ b/Engine/Source/Runtime/CoreUObject/Public/UObject/Script.h
@@ -302,6 +302,7 @@ namespace EScriptInstrumentation
 	enum Type
 	{
 		Class = 0,
+		ClassScope,
 		Instance,
 		Event,
 		ResumeEvent,
@@ -347,51 +348,64 @@ protected:
 	FText Description;
 };
 
-// Information about a blueprint instrumentation event
-struct EScriptInstrumentationEvent
+// Information about a blueprint instrumentation signal
+struct COREUOBJECT_API FScriptInstrumentationSignal
 {
 public:
 
-	EScriptInstrumentationEvent(EScriptInstrumentation::Type InEventType, const UObject* InContextObject, const struct FFrame& InStackFrame)
-		: EventType(InEventType)
-		, ContextObject(InContextObject)
-		, EventName(NAME_None)
-		, StackFramePtr(&InStackFrame)
-		, LatentLinkId(INDEX_NONE)
-	{
-	}
-
-	EScriptInstrumentationEvent(EScriptInstrumentation::Type InEventType, const UObject* InContextObject, FName InEventName)
-		: EventType(InEventType)
-		, ContextObject(InContextObject)
-		, EventName(InEventName)
-		, StackFramePtr(nullptr)
-		, LatentLinkId(INDEX_NONE)
-	{
-	}
+	FScriptInstrumentationSignal(EScriptInstrumentation::Type InEventType, const UObject* InContextObject, const struct FFrame& InStackFrame);
 
-	EScriptInstrumentationEvent(EScriptInstrumentation::Type InEventType, const UObject* InContextObject, FName InEventName, const int32 LinkId)
+	FScriptInstrumentationSignal(EScriptInstrumentation::Type InEventType, const UObject* InContextObject, UFunction* InFunction, const int32 LinkId = INDEX_NONE)
 		: EventType(InEventType)
 		, ContextObject(InContextObject)
-		, EventName(InEventName)
+		, Function(InFunction)
 		, StackFramePtr(nullptr)
 		, LatentLinkId(LinkId)
 	{
 	}
 
+	/** Access to the event type */
 	EScriptInstrumentation::Type GetType() const { return EventType; }
+
+	/** Designates the event type */
 	void SetType(EScriptInstrumentation::Type InType) { EventType = InType;	}
+
+	/** Returns true if the context object is valid */
+	bool IsContextObjectValid() const { return ContextObject != nullptr; }
+
+	/** Returns the context object */
 	const UObject* GetContextObject() const { return ContextObject; }
+
+	/** Returns true if the stackframe is valid */
 	bool IsStackFrameValid() const { return StackFramePtr != nullptr; }
-	FName GetEventName() const { return EventName; }
+
+	/** Returns the stackframe */
 	const FFrame& GetStackFrame() const { return *StackFramePtr; }
+
+	/** Returns the owner class name of the active instance */
+	const UClass* GetClass() const;
+
+	/** Returns the function scope class */
+	const UClass* GetFunctionClassScope() const;
+
+	/** Returns the name of the active function */
+	FName GetFunctionName() const;
+
+	/** Returns the script code offset */
+	int32 GetScriptCodeOffset() const;
+
+	/** Returns the latent link id for latent events */
 	int32 GetLatentLinkId() const { return LatentLinkId; }
 
 protected:
 
+	/** The event signal type */
 	EScriptInstrumentation::Type EventType;
+	/** The context object the event is from */
 	const UObject* ContextObject;
-	const FName EventName;
+	/** The function that emitted this event */
+	const UFunction* Function;
+	/** The stack frame for the  */
 	const struct FFrame* StackFramePtr;
 	const int32 LatentLinkId;
 };
@@ -404,8 +418,8 @@ public:
 	DECLARE_MULTICAST_DELEGATE_ThreeParams(FOnScriptDebuggingEvent, const UObject*, const struct FFrame&, const FBlueprintExceptionInfo&);
 	// Callback for when script execution terminates.
 	DECLARE_MULTICAST_DELEGATE(FOnScriptExecutionEnd);
-	// Callback for blueprint profiling events
-	DECLARE_MULTICAST_DELEGATE_OneParam(FOnScriptInstrumentEvent, const EScriptInstrumentationEvent& );
+	// Callback for blueprint profiling signals
+	DECLARE_MULTICAST_DELEGATE_OneParam(FOnScriptInstrumentEvent, const FScriptInstrumentationSignal& );
 	// Callback for blueprint instrumentation enable/disable events
 	DECLARE_MULTICAST_DELEGATE_OneParam(FOnToggleScriptProfiler, bool );
 
@@ -421,7 +435,7 @@ public:
 
 public:
 	static void ThrowScriptException(const UObject* ActiveObject, const struct FFrame& StackFrame, const FBlueprintExceptionInfo& Info);
-	static void InstrumentScriptEvent(const EScriptInstrumentationEvent& Info);
+	static void InstrumentScriptEvent(const FScriptInstrumentationSignal& Info);
 	static void SetScriptMaximumLoopIterations( const int32 MaximumLoopIterations );
 };
 
diff --git a/Engine/Source/Runtime/Engine/Classes/Engine/LatentActionManager.h b/Engine/Source/Runtime/Engine/Classes/Engine/LatentActionManager.h
--- a/Engine/Source/Runtime/Engine/Classes/Engine/LatentActionManager.h
+++ b/Engine/Source/Runtime/Engine/Classes/Engine/LatentActionManager.h
@@ -35,6 +35,14 @@ struct ENGINE_API FLatentActionInfo
 		, CallbackTarget(NULL)
 	{
 	}
+
+	FORCENOINLINE FLatentActionInfo(int32 InLinkage, int32 InUUID, const TCHAR* InFunctionName, UObject* InCallbackTarget)
+		: Linkage(InLinkage)
+		, UUID(InUUID)
+		, ExecutionFunction(InFunctionName)
+		, CallbackTarget(InCallbackTarget)
+	{
+	}
 };
 
 
diff --git a/Engine/Source/Runtime/Engine/Private/EdGraph/EdGraphSchema.cpp b/Engine/Source/Runtime/Engine/Private/EdGraph/EdGraphSchema.cpp
--- a/Engine/Source/Runtime/Engine/Private/EdGraph/EdGraphSchema.cpp
+++ b/Engine/Source/Runtime/Engine/Private/EdGraph/EdGraphSchema.cpp
@@ -484,10 +484,10 @@ void UEdGraphSchema::BreakNodeLinks(UEdGraphNode& TargetNode) const
 	for (TArray<UEdGraphPin*>::TIterator PinIt(TargetNode.Pins); PinIt; ++PinIt)
 	{
 		UEdGraphPin* TargetPin = *PinIt;
-		if (TargetPin)
+		if (TargetPin != nullptr && TargetPin->SubPins.Num() == 0)
 		{
 			// Keep track of which node(s) the pin's connected to
-			for (auto OtherPin : TargetPin->LinkedTo)
+			for (UEdGraphPin*& OtherPin : TargetPin->LinkedTo)
 			{
 				if (OtherPin)
 				{
@@ -535,18 +535,26 @@ void UEdGraphSchema::BreakPinLinks(UEdGraphPin& TargetPin, bool bSendsNodeNotifc
 	TargetPin.BreakAllPinLinks();
 
 #if WITH_EDITOR
+	UEdGraphNode* OwningNode = TargetPin.GetOwningNode();
 	TSet<UEdGraphNode*> NodeList;
+
 	// Notify this node
-	TargetPin.GetOwningNode()->PinConnectionListChanged(&TargetPin);
-	NodeList.Add(TargetPin.GetOwningNode());
+	if (OwningNode != nullptr)
+	{
+		OwningNode->PinConnectionListChanged(&TargetPin);
+		NodeList.Add(OwningNode);
+	}
 
 	// As well as all other nodes that were connected
 	for (TArray<UEdGraphPin*>::TIterator PinIt(OldLinkedTo); PinIt; ++PinIt)
 	{
 		UEdGraphPin* OtherPin = *PinIt;
 		UEdGraphNode* OtherNode = OtherPin->GetOwningNode();
-		OtherNode->PinConnectionListChanged(OtherPin);
-		NodeList.Add(OtherNode);
+		if (OtherNode != nullptr)
+		{
+			OtherNode->PinConnectionListChanged(OtherPin);
+			NodeList.Add(OtherNode);
+		}
 	}
 
 	if (bSendsNodeNotifcation)
@@ -558,6 +566,7 @@ void UEdGraphSchema::BreakPinLinks(UEdGraphPin& TargetPin, bool bSendsNodeNotifc
 			Node->NodeConnectionListChanged();
 		}
 	}
+	
 #endif	//#if WITH_EDITOR
 }
 
diff --git a/Engine/Source/Runtime/Engine/Private/LatentActionManager.cpp b/Engine/Source/Runtime/Engine/Private/LatentActionManager.cpp
--- a/Engine/Source/Runtime/Engine/Private/LatentActionManager.cpp
+++ b/Engine/Source/Runtime/Engine/Private/LatentActionManager.cpp
@@ -173,10 +173,10 @@ void FLatentActionManager::TickLatentActionForObject(float DeltaTime, FActionLis
 					if (CallbackTarget->GetClass()->HasInstrumentation())
 					{
 						// Place event markers around the latent blueprint call so we can match up the entry points.
-						EScriptInstrumentationEvent LatentEventStartInfo(EScriptInstrumentation::ResumeEvent, CallbackTarget, ExecutionFunction->GetFName(), LinkInfo.LinkID);
+						FScriptInstrumentationSignal LatentEventStartInfo(EScriptInstrumentation::ResumeEvent, CallbackTarget, ExecutionFunction, LinkInfo.LinkID);
 						FBlueprintCoreDelegates::InstrumentScriptEvent(LatentEventStartInfo);
 						CallbackTarget->ProcessEvent(ExecutionFunction, &(LinkInfo.LinkID));
-						EScriptInstrumentationEvent LatentEventStopInfo(EScriptInstrumentation::Stop, CallbackTarget, NAME_None);
+						FScriptInstrumentationSignal LatentEventStopInfo(EScriptInstrumentation::Stop, CallbackTarget, ExecutionFunction);
 						FBlueprintCoreDelegates::InstrumentScriptEvent(LatentEventStopInfo);
 					}
 					else
diff --git a/Engine/Source/Runtime/Slate/Private/Widgets/Input/SButton.cpp b/Engine/Source/Runtime/Slate/Private/Widgets/Input/SButton.cpp
--- a/Engine/Source/Runtime/Slate/Private/Widgets/Input/SButton.cpp
+++ b/Engine/Source/Runtime/Slate/Private/Widgets/Input/SButton.cpp
@@ -432,11 +432,4 @@ void SButton::SetButtonStyle(const FButtonStyle* ButtonStyle)
 
 	HoveredSound = Style->HoveredSlateSound;
 	PressedSound = Style->PressedSlateSound;
-}
-
-// @HSL_BEGIN - [REDACTED_EMAIL] - 5/25/2016 - Need to release the button when we invalidate this button from being a cursor captor, otherwise we won't be able to press it again until we release it
-void SButton::OnMouseCaptureLost()
-{
-	bIsPressed = false;
-}
-// @HSL_END - [REDACTED_EMAIL] - 5/25/2016
\ No newline at end of file
+}
\ No newline at end of file
diff --git a/Engine/Source/Runtime/Slate/Public/Widgets/Input/SButton.h b/Engine/Source/Runtime/Slate/Public/Widgets/Input/SButton.h
--- a/Engine/Source/Runtime/Slate/Public/Widgets/Input/SButton.h
+++ b/Engine/Source/Runtime/Slate/Public/Widgets/Input/SButton.h
@@ -160,9 +160,6 @@ public:
 	virtual void OnMouseLeave( const FPointerEvent& MouseEvent ) override;
 
 	virtual bool IsInteractable() const override;
-// @HSL_BEGIN - [REDACTED_EMAIL] - 5/25/2016 - Need to release the button when we invalidate this button from being a cursor captor, otherwise we won't be able to press it again until we release it
-	virtual void OnMouseCaptureLost() override;
-// @HSL_END - [REDACTED_EMAIL] - 5/25/2016
 
 protected:
 
