Copying //UE4/Dev-Blueprints to //UE4/Dev-Main (Source: //UE4/Dev-Blueprints @ 3130440)

diff --git a/Engine/Source/Editor/BlueprintGraph/Private/BlueprintActionFilter.cpp b/Engine/Source/Editor/BlueprintGraph/Private/BlueprintActionFilter.cpp
--- a/Engine/Source/Editor/BlueprintGraph/Private/BlueprintActionFilter.cpp
+++ b/Engine/Source/Editor/BlueprintGraph/Private/BlueprintActionFilter.cpp
@@ -1288,8 +1288,9 @@ static bool BlueprintActionFilterImpl::IsFunctionMissingPinParam(FBlueprintActio
 				// want to flip the connotation here
 				bool const bWantsOutputConnection = (PinDir == EGPD_Input) ^ bIsEventSpawner;
 				
+				// we don't support direct 'containers of containers, hence the !IsContainer() check here:
 				if ( K2Schema->FunctionHasParamOfType(AssociatedFunc, K2Node->GetGraph(), PinType, bWantsOutputConnection) || 
-					(bIsArrayFunction && ArrayFunctionHasParamOfType(AssociatedFunc, K2Node->GetGraph(), PinType, bWantsOutputConnection)) )
+					(bIsArrayFunction && ArrayFunctionHasParamOfType(AssociatedFunc, K2Node->GetGraph(), PinType, bWantsOutputConnection) && !PinType.IsContainer()) )
 				{
 					bIsFilteredOut = false;
 				}
@@ -1613,12 +1614,18 @@ static bool BlueprintActionFilterImpl::IsIncompatibleMacroInstance(FBlueprintAct
 	{
 		if(const UBlueprint* MacroBP = Cast<const UBlueprint>(BlueprintAction.GetActionOwner()))
 		{
-			check(MacroBP->ParentClass != nullptr);
+			if (!ensure(MacroBP->ParentClass != nullptr))
+			{
+				return true;
+			}
 
 			for(auto BlueprintIt = Filter.Context.Blueprints.CreateConstIterator(); BlueprintIt && !bIsFilteredOut; ++BlueprintIt)
 			{
 				const UBlueprint* Blueprint = *BlueprintIt;
-				check(Blueprint != nullptr && Blueprint->ParentClass != nullptr)
+				if (!ensure(Blueprint != nullptr && Blueprint->ParentClass != nullptr))
+				{
+					return true;
+				}
 
 				bIsFilteredOut = (Blueprint != MacroBP) && (MacroBP->BlueprintType != BPTYPE_MacroLibrary || !Blueprint->ParentClass->IsChildOf(MacroBP->ParentClass));
 			}
diff --git a/Engine/Source/Editor/BlueprintGraph/Private/BlueprintEditorSettings.cpp b/Engine/Source/Editor/BlueprintGraph/Private/BlueprintEditorSettings.cpp
--- a/Engine/Source/Editor/BlueprintGraph/Private/BlueprintEditorSettings.cpp
+++ b/Engine/Source/Editor/BlueprintGraph/Private/BlueprintEditorSettings.cpp
@@ -23,6 +23,7 @@ UBlueprintEditorSettings::UBlueprintEditorSettings(const FObjectInitializer& Obj
 	, bShowEmptySections(true)
 	, bSpawnDefaultBlueprintNodes(true)
 	, bHideConstructionScriptComponentsInDetailsView(true)
+	, bEnableAdvancedContainers(false)
 	// Compiler Settings
 	, SaveOnCompile(SoC_Never)
 	, bJumpToNodeErrors(false)
diff --git a/Engine/Source/Editor/BlueprintGraph/Private/CallFunctionHandler.cpp b/Engine/Source/Editor/BlueprintGraph/Private/CallFunctionHandler.cpp
--- a/Engine/Source/Editor/BlueprintGraph/Private/CallFunctionHandler.cpp
+++ b/Engine/Source/Editor/BlueprintGraph/Private/CallFunctionHandler.cpp
@@ -134,6 +134,18 @@ void FKCHandler_CallFunction::CreateFunctionCallStatement(FKismetFunctionContext
 							if (AssociatedNode && AssociatedNode->NodeGuid.IsValid())
 							{
 								LatentUUID = GetTypeHash(AssociatedNode->NodeGuid);
+
+								UEdGraphNode* ResultNode = AssociatedNode;
+								UEdGraphNode* SourceNode = Cast<UEdGraphNode>(Context.MessageLog.FindSourceObject(AssociatedNode));
+								while (SourceNode && SourceNode != ResultNode)
+								{
+									if (SourceNode->NodeGuid.IsValid())
+									{
+										LatentUUID = HashCombine(LatentUUID, GetTypeHash(SourceNode->NodeGuid));
+									}
+									ResultNode = SourceNode;
+									SourceNode = Cast<UEdGraphNode>(Context.MessageLog.FindSourceObject(ResultNode));
+								}
 							}
 							else
 							{
diff --git a/Engine/Source/Editor/BlueprintGraph/Private/EdGraphSchema_K2.cpp b/Engine/Source/Editor/BlueprintGraph/Private/EdGraphSchema_K2.cpp
--- a/Engine/Source/Editor/BlueprintGraph/Private/EdGraphSchema_K2.cpp
+++ b/Engine/Source/Editor/BlueprintGraph/Private/EdGraphSchema_K2.cpp
@@ -106,6 +106,8 @@ const FName FBlueprintMetadata::MD_DynamicOutputParam(TEXT("DynamicOutputParam")
 const FName FBlueprintMetadata::MD_ArrayParam(TEXT("ArrayParm"));
 const FName FBlueprintMetadata::MD_ArrayDependentParam(TEXT("ArrayTypeDependentParams"));
 
+const FName FBlueprintMetadata::MD_SetParam(TEXT("SetParam"));
+
 const FName FBlueprintMetadata::MD_Bitmask(TEXT("Bitmask"));
 const FName FBlueprintMetadata::MD_BitmaskEnum(TEXT("BitmaskEnum"));
 const FName FBlueprintMetadata::MD_Bitflags(TEXT("Bitflags"));
@@ -1220,7 +1222,7 @@ static bool IsUsingNonExistantVariable(const UEdGraphNode* InGraphNode, UBluepri
 bool UEdGraphSchema_K2::PinHasSplittableStructType(const UEdGraphPin* InGraphPin) const
 {
 	const FEdGraphPinType& PinType = InGraphPin->PinType;
-	bool bCanSplit = (!PinType.bIsArray && PinType.PinCategory == PC_Struct);
+	bool bCanSplit = (!PinType.IsContainer() && PinType.PinCategory == PC_Struct);
 
 	if (bCanSplit)
 	{
@@ -1253,7 +1255,7 @@ bool UEdGraphSchema_K2::PinHasSplittableStructType(const UEdGraphPin* InGraphPin
 bool UEdGraphSchema_K2::PinDefaultValueIsEditable(const UEdGraphPin& InGraphPin) const
 {
 	// Array types are not currently assignable without a 'make array' node:
-	if( InGraphPin.PinType.bIsArray )
+	if( InGraphPin.PinType.IsContainer() )
 	{
 		return false;
 	}
@@ -2023,23 +2025,23 @@ const FPinConnectionResponse UEdGraphSchema_K2::DetermineConnectionResponseOfCom
 		IsSelfPin(*InputPin) && 
 		OwningNode &&
 		OwningNode->AllowMultipleSelfs(false) &&
-		!InputPin->PinType.bIsArray &&
-		!OutputPin->PinType.bIsArray)
+		!InputPin->PinType.IsContainer() &&
+		!OutputPin->PinType.IsContainer() )
 	{
 		//check if the node wont be expanded as foreach call, if there is a link to an array
-		bool bAnyArrayInput = false;
+		bool bAnyContainerInput = false;
 		for(int InputLinkIndex = 0; InputLinkIndex < InputPin->LinkedTo.Num(); InputLinkIndex++)
 		{
 			if(const UEdGraphPin* Pin = InputPin->LinkedTo[InputLinkIndex])
 			{
-				if(Pin->PinType.bIsArray)
+				if(Pin->PinType.IsContainer())
 				{
-					bAnyArrayInput = true;
+					bAnyContainerInput = true;
 					break;
 				}
 			}
 		}
-		bMultipleSelfException = !bAnyArrayInput;
+		bMultipleSelfException = !bAnyContainerInput;
 	}
 
 	if (bBreakExistingDueToExecOutput)
@@ -2076,6 +2078,12 @@ static FText GetPinIncompatibilityReason(const UEdGraphPin* PinA, const UEdGraph
 
 	FText MessageFormat = LOCTEXT("DefaultPinIncompatibilityMessage", "{PinAType} is not compatible with {PinBType}.");
 
+	if (bIsFatalOut != nullptr)
+	{
+		// the incompatible pins should generate an error by default
+		*bIsFatalOut = true;
+	}
+
 	if (OutputType.PinCategory == UEdGraphSchema_K2::PC_Struct)
 	{
 		if (InputType.PinCategory == UEdGraphSchema_K2::PC_Struct)
@@ -2088,11 +2096,6 @@ static FText GetPinIncompatibilityReason(const UEdGraphPin* PinA, const UEdGraph
 			{
 				MessageFormat = LOCTEXT("ChildStructIncompatible", "Only exactly matching structures are considered compatible. Derived structures are disallowed.");
 			}
-
-			if (bIsFatalOut != nullptr)
-			{
-				*bIsFatalOut = true;
-			}
 		}
 	}
 	else if (OutputType.PinCategory == UEdGraphSchema_K2::PC_Class)
@@ -2103,6 +2106,12 @@ static FText GetPinIncompatibilityReason(const UEdGraphPin* PinA, const UEdGraph
 			MessageArgs.Add(TEXT("OutputName"), OutputPin->GetDisplayName());
 			MessageArgs.Add(TEXT("InputName"),  InputPin->GetDisplayName());
 			MessageFormat = LOCTEXT("ClassObjectIncompatible", "'{PinAName}' and '{PinBName}' are incompatible ('{OutputName}' is an object type, and '{InputName}' is a reference to an object instance).");
+
+			if ((InputType.PinCategory == UEdGraphSchema_K2::PC_Object) && (bIsFatalOut != nullptr))
+			{
+				// under the hood class is an object, so it's not fatal
+				*bIsFatalOut = false;
+			}
 		}
 	}
 	else if ((OutputType.PinCategory == UEdGraphSchema_K2::PC_Object) )//|| (OutputType.PinCategory == UEdGraphSchema_K2::PC_Interface))
@@ -2163,8 +2172,10 @@ const FPinConnectionResponse UEdGraphSchema_K2::CanCreateConnection(const UEdGra
 	{
 		const bool bAllowMultipleSelfs = OwningNode->AllowMultipleSelfs(true); // it applies also to ForEachCall
 		const bool bNotAnArrayFunction = !InputPin->PinType.bIsArray;
+		const bool bNotAMapFunction = !InputPin->PinType.bIsMap;
+		const bool bNotASetFunction = !InputPin->PinType.bIsSet;
 		const bool bSelfPin = IsSelfPin(*InputPin);
-		bIgnoreArray = bAllowMultipleSelfs && bNotAnArrayFunction && bSelfPin;
+		bIgnoreArray = bAllowMultipleSelfs && bNotAnArrayFunction && bNotAMapFunction && bNotASetFunction  && bSelfPin;
 	}
 
 	// Find the calling context in case one of the pins is of type object and has a value of Self
@@ -2198,7 +2209,7 @@ const FPinConnectionResponse UEdGraphSchema_K2::CanCreateConnection(const UEdGra
 		}
 		else
 		{
-			bool bIsFatal = false;
+			bool bIsFatal = true;
 			FText IncompatibilityReasonText = GetPinIncompatibilityReason(PinA, PinB, &bIsFatal);
 
 			FPinConnectionResponse ConnectionResponse(CONNECT_RESPONSE_DISALLOW, IncompatibilityReasonText.ToString());
@@ -2392,7 +2403,9 @@ bool UEdGraphSchema_K2::SearchForAutocastFunction(const UEdGraphPin* OutputPin,
 	TargetFunction = NAME_None;
 	FunctionOwner = nullptr;
 
-	if (OutputPin->PinType.bIsArray != InputPin->PinType.bIsArray)
+	if (OutputPin->PinType.bIsArray != InputPin->PinType.bIsArray ||
+		OutputPin->PinType.bIsMap != InputPin->PinType.bIsMap ||
+		OutputPin->PinType.bIsSet != InputPin->PinType.bIsSet)
 	{
 		return false;
 	}
@@ -2509,8 +2522,8 @@ bool UEdGraphSchema_K2::FindSpecializedConversionNode(const UEdGraphPin* OutputP
 	if(!bCanConvert)
 	{
 		// CHECK ENUM TO NAME CAST
-		const bool bInoputMatch = InputPin && !InputPin->PinType.bIsArray && ((PC_Name == InputPin->PinType.PinCategory) || (PC_String == InputPin->PinType.PinCategory));
-		const bool bOutputMatch = OutputPin && !OutputPin->PinType.bIsArray && (PC_Byte == OutputPin->PinType.PinCategory) && (NULL != Cast<UEnum>(OutputPin->PinType.PinSubCategoryObject.Get()));
+		const bool bInoputMatch = InputPin && !InputPin->PinType.IsContainer() && ((PC_Name == InputPin->PinType.PinCategory) || (PC_String == InputPin->PinType.PinCategory));
+		const bool bOutputMatch = OutputPin && !OutputPin->PinType.IsContainer() && (PC_Byte == OutputPin->PinType.PinCategory) && (NULL != Cast<UEnum>(OutputPin->PinType.PinSubCategoryObject.Get()));
 		if(bOutputMatch && bInoputMatch)
 		{
 			bCanConvert = true;
@@ -2695,6 +2708,8 @@ FString UEdGraphSchema_K2::IsPinDefaultValid(const UEdGraphPin* Pin, const FStri
 	}
 
 	const bool bIsArray = Pin->PinType.bIsArray;
+	const bool bIsSet = Pin->PinType.bIsSet;
+	const bool bIsMap = Pin->PinType.bIsMap;
 	const bool bIsReference = Pin->PinType.bIsReference;
 	const bool bIsAutoCreateRefTerm = IsAutoCreateRefTerm(Pin);
 
@@ -2707,6 +2722,16 @@ FString UEdGraphSchema_K2::IsPinDefaultValid(const UEdGraphPin* Pin, const FStri
 				FText MsgFormat = LOCTEXT("BadArrayDefaultVal", "Array inputs (like '{PinName}') must have an input wired into them (try connecting a MakeArray node).");
 				return FText::Format(MsgFormat, MessageArgs).ToString();
 			}
+			else if( bIsSet )
+			{
+				FText MsgFormat = LOCTEXT("BadSetDefaultVal", "Array inputs (like '{PinName}') must have an input wired into them (try connecting a MakeSet node).");
+				return FText::Format(MsgFormat, MessageArgs).ToString();
+			}
+			else if ( bIsMap )
+			{
+				FText MsgFormat = LOCTEXT("BadMapDefaultVal", "Array inputs (like '{PinName}') must have an input wired into them (try connecting a MakeMap node).");
+				return FText::Format(MsgFormat, MessageArgs).ToString();
+			}
 			else if( bIsReference )
 			{
 				FText MsgFormat = LOCTEXT("BadRefDefaultVal", "'{PinName}' must have an input wired into it (\"by ref\" params expect a valid input to operate on).");
@@ -2837,6 +2862,25 @@ FLinearColor UEdGraphSchema_K2::GetPinTypeColor(const FEdGraphPinType& PinType)
 	return Settings->DefaultPinTypeColor;
 }
 
+FLinearColor UEdGraphSchema_K2::GetSecondaryPinTypeColor(const FEdGraphPinType& PinType) const
+{
+	if (PinType.bIsMap)
+	{
+		FEdGraphPinType FakePrimary = PinType;
+		FakePrimary.PinCategory = FakePrimary.PinValueType.TerminalCategory;
+		FakePrimary.PinSubCategory = FakePrimary.PinValueType.TerminalSubCategory;
+		FakePrimary.PinSubCategoryObject = FakePrimary.PinValueType.TerminalSubCategoryObject;
+
+		return GetPinTypeColor(FakePrimary);
+	}
+	else
+	{
+		const FString& TypeString = PinType.PinCategory;
+		const UGraphEditorSettings* Settings = GetDefault<UGraphEditorSettings>();
+		return Settings->WildcardPinTypeColor;
+	}
+}
+
 FText UEdGraphSchema_K2::GetPinDisplayName(const UEdGraphPin* Pin) const 
 {
 	FText DisplayName = FText::GetEmpty();
@@ -3133,103 +3177,142 @@ void UEdGraphSchema_K2::CreateFunctionGraphTerminators(UEdGraph& Graph, UFunctio
 	}
 }
 
-bool UEdGraphSchema_K2::ConvertPropertyToPinType(const UProperty* Property, /*out*/ FEdGraphPinType& TypeOut) const
+bool UEdGraphSchema_K2::GetPropertyCategoryInfo(const UProperty* TestProperty, FString& OutCategory, FString& OutSubCategory, UObject*& OutSubCategoryObject, bool& bOutIsWeakPointer)
 {
-	if (Property == NULL)
+	if (const UInterfaceProperty* InterfaceProperty = Cast<const UInterfaceProperty>(TestProperty))
 	{
-		TypeOut.PinCategory = TEXT("bad_type");
-		return false;
-	}
-
-	TypeOut.PinSubCategory = TEXT("");
-	
-	// Handle whether or not this is an array property
-	const UArrayProperty* ArrayProperty = Cast<const UArrayProperty>(Property);
-	const UProperty* TestProperty = ArrayProperty ? ArrayProperty->Inner : Property;
-	TypeOut.bIsArray = (ArrayProperty != NULL);
-	TypeOut.bIsReference = Property->HasAllPropertyFlags(CPF_OutParm|CPF_ReferenceParm);
-	TypeOut.bIsConst     = Property->HasAllPropertyFlags(CPF_ConstParm);
-
-	// Check to see if this is the wildcard property for the target array type
-	UFunction* Function = Cast<UFunction>(Property->GetOuter());
-
-	if( UK2Node_CallArrayFunction::IsWildcardProperty(Function, Property)
-		|| UK2Node_CallFunction::IsStructureWildcardProperty(Function, Property->GetName()))
-	{
-		TypeOut.PinCategory = PC_Wildcard;
-	}
-	else if (const UInterfaceProperty* InterfaceProperty = Cast<const UInterfaceProperty>(TestProperty))
-	{
-		TypeOut.PinCategory = PC_Interface;
-		TypeOut.PinSubCategoryObject = InterfaceProperty->InterfaceClass;
+		OutCategory = PC_Interface;
+		OutSubCategoryObject = InterfaceProperty->InterfaceClass;
 	}
 	else if (const UClassProperty* ClassProperty = Cast<const UClassProperty>(TestProperty))
 	{
-		TypeOut.PinCategory = PC_Class;
-		TypeOut.PinSubCategoryObject = ClassProperty->MetaClass;
+		OutCategory = PC_Class;
+		OutSubCategoryObject = ClassProperty->MetaClass;
 	}
 	else if (const UAssetClassProperty* AssetClassProperty = Cast<const UAssetClassProperty>(TestProperty))
 	{
-		TypeOut.PinCategory = PC_AssetClass;
-		TypeOut.PinSubCategoryObject = AssetClassProperty->MetaClass;
+		OutCategory = PC_AssetClass;
+		OutSubCategoryObject = AssetClassProperty->MetaClass;
 	}
 	else if (const UAssetObjectProperty* AssetObjectProperty = Cast<const UAssetObjectProperty>(TestProperty))
 	{
-		TypeOut.PinCategory = PC_Asset;
-		TypeOut.PinSubCategoryObject = AssetObjectProperty->PropertyClass;
+		OutCategory = PC_Asset;
+		OutSubCategoryObject = AssetObjectProperty->PropertyClass;
 	}
 	else if (const UObjectPropertyBase* ObjectProperty = Cast<const UObjectPropertyBase>(TestProperty))
 	{
-		TypeOut.PinCategory = PC_Object;
-		TypeOut.PinSubCategoryObject = ObjectProperty->PropertyClass;
-		TypeOut.bIsWeakPointer = TestProperty->IsA(UWeakObjectProperty::StaticClass());
+		OutCategory = PC_Object;
+		OutSubCategoryObject = ObjectProperty->PropertyClass;
+		bOutIsWeakPointer = TestProperty->IsA(UWeakObjectProperty::StaticClass());
 	}
 	else if (const UStructProperty* StructProperty = Cast<const UStructProperty>(TestProperty))
 	{
-		TypeOut.PinCategory = PC_Struct;
-		TypeOut.PinSubCategoryObject = StructProperty->Struct;
+		OutCategory = PC_Struct;
+		OutSubCategoryObject = StructProperty->Struct;
 	}
 	else if (Cast<const UFloatProperty>(TestProperty) != NULL)
 	{
-		TypeOut.PinCategory = PC_Float;
+		OutCategory = PC_Float;
 	}
 	else if (Cast<const UIntProperty>(TestProperty) != NULL)
 	{
-		TypeOut.PinCategory = PC_Int;
+		OutCategory = PC_Int;
 
 		if (TestProperty->HasMetaData(FBlueprintMetadata::MD_Bitmask))
 		{
-			TypeOut.PinSubCategory = PSC_Bitmask;
+			OutSubCategory = PSC_Bitmask;
 		}
 	}
 	else if (const UByteProperty* ByteProperty = Cast<const UByteProperty>(TestProperty))
 	{
-		TypeOut.PinCategory = PC_Byte;
+		OutCategory = PC_Byte;
 
 		if (TestProperty->HasMetaData(FBlueprintMetadata::MD_Bitmask))
 		{
-			TypeOut.PinSubCategory = PSC_Bitmask;
+			OutSubCategory = PSC_Bitmask;
 		}
 		else
 		{
-			TypeOut.PinSubCategoryObject = ByteProperty->Enum;
+			OutSubCategoryObject = ByteProperty->Enum;
 		}
 	}
 	else if (Cast<const UNameProperty>(TestProperty) != NULL)
 	{
-		TypeOut.PinCategory = PC_Name;
+		OutCategory = PC_Name;
 	}
 	else if (Cast<const UBoolProperty>(TestProperty) != NULL)
 	{
-		TypeOut.PinCategory = PC_Boolean;
+		OutCategory = PC_Boolean;
 	}
 	else if (Cast<const UStrProperty>(TestProperty) != NULL)
 	{
-		TypeOut.PinCategory = PC_String;
+		OutCategory = PC_String;
 	}
 	else if (Cast<const UTextProperty>(TestProperty) != NULL)
 	{
-		TypeOut.PinCategory = PC_Text;
+		OutCategory = PC_Text;
+	}
+	else
+	{
+		OutCategory = TEXT("bad_type");
+		return false;
+	}
+
+	return true;
+}
+
+bool UEdGraphSchema_K2::ConvertPropertyToPinType(const UProperty* Property, /*out*/ FEdGraphPinType& TypeOut) const
+{
+	if (Property == NULL)
+	{
+		TypeOut.PinCategory = TEXT("bad_type");
+		return false;
+	}
+
+	TypeOut.PinSubCategory = TEXT("");
+	
+	// Handle whether or not this is an array property
+	const UMapProperty* MapProperty = Cast<const UMapProperty>(Property);
+	const USetProperty* SetProperty = Cast<const USetProperty>(Property);
+	const UArrayProperty* ArrayProperty = Cast<const UArrayProperty>(Property);
+	const UProperty* TestProperty = Property;
+	if (MapProperty)
+	{
+		TestProperty = MapProperty->KeyProp;
+
+		// set up value property:
+		UObject* SubCategoryObject = nullptr;
+		bool bIsWeakPtr = false;
+		bool bResult = GetPropertyCategoryInfo(MapProperty->ValueProp, TypeOut.PinValueType.TerminalCategory, TypeOut.PinValueType.TerminalSubCategory, SubCategoryObject, bIsWeakPtr);
+		TypeOut.PinValueType.TerminalSubCategoryObject = SubCategoryObject;
+		ensure(!bIsWeakPtr);
+		if (!bResult)
+		{
+			return false;
+		}
+	}
+	else if (SetProperty)
+	{
+		TestProperty = SetProperty->ElementProp;
+	}
+	else if (ArrayProperty)
+	{
+		TestProperty = ArrayProperty->Inner;
+	}
+	TypeOut.bIsMap = (MapProperty != NULL);
+	TypeOut.bIsSet = (SetProperty != NULL);
+	TypeOut.bIsArray = (ArrayProperty != NULL);
+	TypeOut.bIsReference = Property->HasAllPropertyFlags(CPF_OutParm|CPF_ReferenceParm);
+	TypeOut.bIsConst     = Property->HasAllPropertyFlags(CPF_ConstParm);
+
+	// Check to see if this is the wildcard property for the target array type
+	UFunction* Function = Cast<UFunction>(Property->GetOuter());
+
+	if( UK2Node_CallArrayFunction::IsWildcardProperty(Function, Property)
+		|| UK2Node_CallFunction::IsStructureWildcardProperty(Function, Property->GetName())
+		|| UK2Node_CallFunction::IsWildcardProperty(Function, Property))
+	{
+		TypeOut.PinCategory = PC_Wildcard;
 	}
 	else if (const UMulticastDelegateProperty* MulticastDelegateProperty = Cast<const UMulticastDelegateProperty>(TestProperty))
 	{
@@ -3243,8 +3326,13 @@ bool UEdGraphSchema_K2::ConvertPropertyToPinType(const UProperty* Property, /*ou
 	}
 	else
 	{
-		TypeOut.PinCategory = TEXT("bad_type");
-		return false;
+		UObject* SubCategoryObject = nullptr;
+		bool bResult = GetPropertyCategoryInfo(TestProperty, TypeOut.PinCategory, TypeOut.PinSubCategory, SubCategoryObject, TypeOut.bIsWeakPointer);
+		TypeOut.PinSubCategoryObject = SubCategoryObject;
+		if (!bResult)
+		{
+			return false;
+		}
 	}
 
 	if (TypeOut.PinSubCategory == PSC_Bitmask)
@@ -3377,42 +3465,41 @@ FText UEdGraphSchema_K2::GetCategoryText(const FString& Category, const bool bFo
 	}
 }
 
-FText UEdGraphSchema_K2::TypeToText(const FEdGraphPinType& Type)
+FText UEdGraphSchema_K2::TerminalTypeToText(const FString& Category, const FString& SubCategory, UObject* SubCategoryObject, bool bIsWeakPtr)
 {
 	FText PropertyText;
 
 	const UEdGraphSchema_K2* Schema = GetDefault<UEdGraphSchema_K2>();
-	if (Type.PinSubCategory != Schema->PSC_Bitmask && Type.PinSubCategoryObject != NULL)
+	if (SubCategory != Schema->PSC_Bitmask && SubCategoryObject != NULL)
 	{
-		if (Type.PinCategory == Schema->PC_Byte)
+		if (Category == Schema->PC_Byte)
 		{
 			FFormatNamedArguments Args;
-			Args.Add(TEXT("EnumName"), FText::FromString(Type.PinSubCategoryObject->GetName()));
+			Args.Add(TEXT("EnumName"), FText::FromString(SubCategoryObject->GetName()));
 			PropertyText = FText::Format(LOCTEXT("EnumAsText", "{EnumName} Enum"), Args);
 		}
 		else
 		{
-			UObject* SubCategoryObj = Type.PinSubCategoryObject.Get();
-			FString SubCategoryObjName = SubCategoryObj->GetName();
-			if (UField* SubCategoryField = Cast<UField>(SubCategoryObj))
+			FString SubCategoryObjName = SubCategoryObject->GetName();
+			if (UField* SubCategoryField = Cast<UField>(SubCategoryObject))
 			{
 				SubCategoryObjName = SubCategoryField->GetDisplayNameText().ToString();
 			}
 
-			if( !Type.bIsWeakPointer )
+			if (!bIsWeakPtr)
 			{
-				UClass* PSCOAsClass = Cast<UClass>(Type.PinSubCategoryObject.Get());
+				UClass* PSCOAsClass = Cast<UClass>(SubCategoryObject);
 				const bool bIsInterface = PSCOAsClass && PSCOAsClass->HasAnyClassFlags(CLASS_Interface);
 
 				FFormatNamedArguments Args;
 				// Don't display the category for "well-known" struct types
-				if (Type.PinCategory == PC_Struct && (Type.PinSubCategoryObject == VectorStruct || Type.PinSubCategoryObject == RotatorStruct || Type.PinSubCategoryObject == TransformStruct))
+				if (Category == UEdGraphSchema_K2::PC_Struct && (SubCategoryObject == UEdGraphSchema_K2::VectorStruct || SubCategoryObject == UEdGraphSchema_K2::RotatorStruct || SubCategoryObject == UEdGraphSchema_K2::TransformStruct))
 				{
 					Args.Add(TEXT("Category"), FText::GetEmpty());
 				}
 				else
 				{
-					Args.Add(TEXT("Category"), (!bIsInterface ? GetCategoryText(Type.PinCategory) : GetCategoryText(PC_Interface)));
+					Args.Add(TEXT("Category"), (!bIsInterface ? UEdGraphSchema_K2::GetCategoryText(Category) : UEdGraphSchema_K2::GetCategoryText(PC_Interface)));
 				}
 
 				Args.Add(TEXT("ObjectName"), FText::FromString(FName::NameToDisplayString(SubCategoryObjName, /*bIsBool =*/false)));
@@ -3421,25 +3508,46 @@ FText UEdGraphSchema_K2::TypeToText(const FEdGraphPinType& Type)
 			else
 			{
 				FFormatNamedArguments Args;
-				Args.Add(TEXT("Category"), FText::FromString(Type.PinCategory));
+				Args.Add(TEXT("Category"), FText::FromString(Category));
 				Args.Add(TEXT("ObjectName"), FText::FromString(SubCategoryObjName));
 				PropertyText = FText::Format(LOCTEXT("WeakPtrAsText", "{ObjectName} Weak {Category}"), Args);
 			}
 		}
 	}
-	else if (Type.PinSubCategory != TEXT(""))
+	else if (SubCategory != TEXT(""))
 	{
 		FFormatNamedArguments Args;
-		Args.Add(TEXT("Category"), GetCategoryText(Type.PinCategory));
-		Args.Add(TEXT("ObjectName"), FText::FromString(FName::NameToDisplayString(Type.PinSubCategory, false)));
+		Args.Add(TEXT("Category"), UEdGraphSchema_K2::GetCategoryText(Category));
+		Args.Add(TEXT("ObjectName"), FText::FromString(FName::NameToDisplayString(SubCategory, false)));
 		PropertyText = FText::Format(LOCTEXT("ObjectAsText", "{ObjectName} {Category}"), Args);
 	}
 	else
 	{
-		PropertyText = GetCategoryText(Type.PinCategory);
+		PropertyText = UEdGraphSchema_K2::GetCategoryText(Category);
 	}
 
-	if (Type.bIsArray)
+	return PropertyText;
+}
+
+FText UEdGraphSchema_K2::TypeToText(const FEdGraphPinType& Type)
+{
+	FText PropertyText = TerminalTypeToText(Type.PinCategory, Type.PinSubCategory, Type.PinSubCategoryObject.Get(), Type.bIsWeakPointer);
+
+	if ( Type.bIsMap )
+	{
+		FFormatNamedArguments Args;
+		Args.Add(TEXT("KeyTitle"), PropertyText);
+		FText ValueText = TerminalTypeToText(Type.PinValueType.TerminalCategory, Type.PinValueType.TerminalSubCategory, Type.PinValueType.TerminalSubCategoryObject.Get(), Type.PinValueType.bTerminalIsWeakPointer);
+		Args.Add(TEXT("ValueTitle"), ValueText);
+		PropertyText = FText::Format(LOCTEXT("MapAsText", "Map of {PropertyTitle}s to {ValueTitle}s"), Args);
+	}
+	else if ( Type.bIsSet )
+	{
+		FFormatNamedArguments Args;
+		Args.Add(TEXT("PropertyTitle"), PropertyText);
+		PropertyText = FText::Format(LOCTEXT("SetAsText", "Set of {PropertyTitle}s"), Args);
+	}
+	else if (Type.bIsArray)
 	{
 		FFormatNamedArguments Args;
 		Args.Add(TEXT("PropertyTitle"), PropertyText);
@@ -3695,9 +3803,9 @@ bool UEdGraphSchema_K2::DefaultValueSimpleValidation(const FEdGraphPinType& PinT
 	const FString& PinSubCategory = PinType.PinSubCategory;
 	const UObject* PinSubCategoryObject = PinType.PinSubCategoryObject.Get();
 
-	if (PinType.bIsArray)
+	if (PinType.IsContainer())
 	{
-		// arrays are validated separately
+		// containers are validated separately
 	}
 	//@TODO: FCString::Atoi, FCString::Atof, and appStringToBool will 'accept' any input, but we should probably catch and warn
 	// about invalid input (non numeric for int/byte/float, and non 0/1 or yes/no/true/false for bool)
@@ -3907,7 +4015,7 @@ bool UEdGraphSchema_K2::DefaultValueSimpleValidation(const FEdGraphPinType& PinT
 
 bool UEdGraphSchema_K2::ArePinTypesCompatible(const FEdGraphPinType& Output, const FEdGraphPinType& Input, const UClass* CallingContext, bool bIgnoreArray /*= false*/) const
 {
-	if( !bIgnoreArray && (Output.bIsArray != Input.bIsArray) && (Input.PinCategory != PC_Wildcard || Input.bIsArray) )
+	if( !bIgnoreArray && ( Output.bIsMap != Input.bIsMap || Output.bIsSet != Input.bIsSet || Output.bIsArray != Input.bIsArray ) && (Input.PinCategory != PC_Wildcard || Input.IsContainer()) )
 	{
 		return false;
 	}
@@ -4310,7 +4418,7 @@ bool UEdGraphSchema_K2::ShouldHidePinDefaultValue(UEdGraphPin* Pin) const
 {
 	check(Pin != NULL);
 
-	if (Pin->bDefaultValueIsIgnored || Pin->PinType.bIsArray || (Pin->PinName == PN_Self && Pin->LinkedTo.Num() > 0) || (Pin->PinType.PinCategory == PC_Exec) || (Pin->PinType.bIsReference && !IsAutoCreateRefTerm(Pin)))
+	if (Pin->bDefaultValueIsIgnored || Pin->PinType.IsContainer() || (Pin->PinName == PN_Self && Pin->LinkedTo.Num() > 0) || (Pin->PinType.PinCategory == PC_Exec) || (Pin->PinType.bIsReference && !IsAutoCreateRefTerm(Pin)))
 	{
 		return true;
 	}
@@ -4413,7 +4521,7 @@ void UEdGraphSchema_K2::SetPinDefaultValue(UEdGraphPin* Pin, const UFunction* Fu
 void UEdGraphSchema_K2::SetPinDefaultValueBasedOnType(UEdGraphPin* Pin) const
 {
 	// Create a useful default value based on the pin type
-	if(Pin->PinType.bIsArray)
+	if(Pin->PinType.IsContainer() )
 	{
 		Pin->AutogeneratedDefaultValue = TEXT("");
 	}
diff --git a/Engine/Source/Editor/BlueprintGraph/Private/K2Node_CallArrayFunction.cpp b/Engine/Source/Editor/BlueprintGraph/Private/K2Node_CallArrayFunction.cpp
--- a/Engine/Source/Editor/BlueprintGraph/Private/K2Node_CallArrayFunction.cpp
+++ b/Engine/Source/Editor/BlueprintGraph/Private/K2Node_CallArrayFunction.cpp
@@ -15,12 +15,14 @@ void UK2Node_CallArrayFunction::AllocateDefaultPins()
 	const UEdGraphSchema_K2* Schema = GetDefault<UEdGraphSchema_K2>();
 
 	UEdGraphPin* TargetArrayPin = GetTargetArrayPin();
-	check(TargetArrayPin);
-	TargetArrayPin->PinType.bIsArray = true;
-	TargetArrayPin->PinType.bIsReference = true;
-	TargetArrayPin->PinType.PinCategory = Schema->PC_Wildcard;
-	TargetArrayPin->PinType.PinSubCategory = TEXT("");
-	TargetArrayPin->PinType.PinSubCategoryObject = NULL;
+	if (ensure(TargetArrayPin))
+	{
+		TargetArrayPin->PinType.bIsArray = true;
+		TargetArrayPin->PinType.bIsReference = true;
+		TargetArrayPin->PinType.PinCategory = Schema->PC_Wildcard;
+		TargetArrayPin->PinType.PinSubCategory = TEXT("");
+		TargetArrayPin->PinType.PinSubCategoryObject = NULL;
+	}
 
 	TArray< FArrayPropertyPinCombo > ArrayPins;
 	GetArrayPins(ArrayPins);
@@ -137,26 +139,28 @@ void UK2Node_CallArrayFunction::GetArrayPins(TArray< FArrayPropertyPinCombo >& O
 	OutArrayPinInfo.Empty();
 
 	UFunction* TargetFunction = GetTargetFunction();
-	check(TargetFunction);
-	FString ArrayPointerMetaData = TargetFunction->GetMetaData(FBlueprintMetadata::MD_ArrayParam);
-	TArray<FString> ArrayPinComboNames;
-	ArrayPointerMetaData.ParseIntoArray(ArrayPinComboNames, TEXT(","), true);
-
-	for(auto Iter = ArrayPinComboNames.CreateConstIterator(); Iter; ++Iter)
+	if (ensure(TargetFunction))
 	{
-		TArray<FString> ArrayPinNames;
-		Iter->ParseIntoArray(ArrayPinNames, TEXT("|"), true);
+		FString ArrayPointerMetaData = TargetFunction->GetMetaData(FBlueprintMetadata::MD_ArrayParam);
+		TArray<FString> ArrayPinComboNames;
+		ArrayPointerMetaData.ParseIntoArray(ArrayPinComboNames, TEXT(","), true);
 
-		FArrayPropertyPinCombo ArrayInfo;
-		ArrayInfo.ArrayPin = FindPin(ArrayPinNames[0]);
-		if(ArrayPinNames.Num() > 1)
+		for (auto Iter = ArrayPinComboNames.CreateConstIterator(); Iter; ++Iter)
 		{
-			ArrayInfo.ArrayPropPin = FindPin(ArrayPinNames[1]);
-		}
+			TArray<FString> ArrayPinNames;
+			Iter->ParseIntoArray(ArrayPinNames, TEXT("|"), true);
 
-		if(ArrayInfo.ArrayPin)
-		{
-			OutArrayPinInfo.Add(ArrayInfo);
+			FArrayPropertyPinCombo ArrayInfo;
+			ArrayInfo.ArrayPin = FindPin(ArrayPinNames[0]);
+			if (ArrayPinNames.Num() > 1)
+			{
+				ArrayInfo.ArrayPropPin = FindPin(ArrayPinNames[1]);
+			}
+
+			if (ArrayInfo.ArrayPin)
+			{
+				OutArrayPinInfo.Add(ArrayInfo);
+			}
 		}
 	}
 }
@@ -165,18 +169,21 @@ bool UK2Node_CallArrayFunction::IsWildcardProperty(UFunction* InArrayFunction, c
 {
 	if(InArrayFunction && InProperty)
 	{
-		FString ArrayPointerMetaData = InArrayFunction->GetMetaData(FBlueprintMetadata::MD_ArrayParam);
-		TArray<FString> ArrayPinComboNames;
-		ArrayPointerMetaData.ParseIntoArray(ArrayPinComboNames, TEXT(","), true);
-
-		for(auto Iter = ArrayPinComboNames.CreateConstIterator(); Iter; ++Iter)
+		const FString& ArrayPointerMetaData = InArrayFunction->GetMetaData(FBlueprintMetadata::MD_ArrayParam);
+		if (!ArrayPointerMetaData.IsEmpty())
 		{
-			TArray<FString> ArrayPinNames;
-			Iter->ParseIntoArray(ArrayPinNames, TEXT("|"), true);
+			TArray<FString> ArrayPinComboNames;
+			ArrayPointerMetaData.ParseIntoArray(ArrayPinComboNames, TEXT(","), true);
 
-			if(ArrayPinNames[0] == InProperty->GetName())
+			for (auto Iter = ArrayPinComboNames.CreateConstIterator(); Iter; ++Iter)
 			{
-				return true;
+				TArray<FString> ArrayPinNames;
+				Iter->ParseIntoArray(ArrayPinNames, TEXT("|"), true);
+
+				if (ArrayPinNames[0] == InProperty->GetName())
+				{
+					return true;
+				}
 			}
 		}
 	}
@@ -188,19 +195,20 @@ void UK2Node_CallArrayFunction::GetArrayTypeDependentPins(TArray<UEdGraphPin*>&
 	OutPins.Empty();
 
 	UFunction* TargetFunction = GetTargetFunction();
-	check(TargetFunction);
-
-	const FString DependentPinMetaData = TargetFunction->GetMetaData(FBlueprintMetadata::MD_ArrayDependentParam);
-	TArray<FString> TypeDependentPinNames;
-	DependentPinMetaData.ParseIntoArray(TypeDependentPinNames, TEXT(","), true);
-
-	for(TArray<UEdGraphPin*>::TConstIterator it(Pins); it; ++it)
+	if (ensure(TargetFunction))
 	{
-		UEdGraphPin* CurrentPin = *it;
-		int32 ItemIndex = 0;
-		if( CurrentPin && TypeDependentPinNames.Find(CurrentPin->PinName, ItemIndex) )
+		const FString DependentPinMetaData = TargetFunction->GetMetaData(FBlueprintMetadata::MD_ArrayDependentParam);
+		TArray<FString> TypeDependentPinNames;
+		DependentPinMetaData.ParseIntoArray(TypeDependentPinNames, TEXT(","), true);
+
+		for (TArray<UEdGraphPin*>::TConstIterator it(Pins); it; ++it)
 		{
-			OutPins.Add(CurrentPin);
+			UEdGraphPin* CurrentPin = *it;
+			int32 ItemIndex = 0;
+			if (CurrentPin && TypeDependentPinNames.Find(CurrentPin->PinName, ItemIndex))
+			{
+				OutPins.Add(CurrentPin);
+			}
 		}
 	}
 }
diff --git a/Engine/Source/Editor/BlueprintGraph/Private/K2Node_CallFunction.cpp b/Engine/Source/Editor/BlueprintGraph/Private/K2Node_CallFunction.cpp
--- a/Engine/Source/Editor/BlueprintGraph/Private/K2Node_CallFunction.cpp
+++ b/Engine/Source/Editor/BlueprintGraph/Private/K2Node_CallFunction.cpp
@@ -941,6 +941,9 @@ void UK2Node_CallFunction::PostReconstructNode()
 	Super::PostReconstructNode();
 	InvalidatePinTooltips();
 
+	// conform pins that are marked as SetParam:
+	ConformContainerPins();
+
 	FCustomStructureParamHelper::UpdateCustomStructurePins(GetTargetFunction(), this);
 
 	const UEdGraphSchema_K2* K2Schema = GetDefault<UEdGraphSchema_K2>();
@@ -1011,6 +1014,9 @@ void UK2Node_CallFunction::NotifyPinConnectionListChanged(UEdGraphPin* Pin)
 {
 	Super::NotifyPinConnectionListChanged(Pin);
 
+	// conform pins that are marked as SetParam:
+	ConformContainerPins();
+
 	if (!ensure(Pin))
 	{
 		return;
@@ -2294,6 +2300,87 @@ void UK2Node_CallFunction::InvalidatePinTooltips()
 	bPinTooltipsValid = false;
 }
 
+void UK2Node_CallFunction::ConformContainerPins()
+{
+	const UFunction* TargetFunction = GetTargetFunction();
+	if (TargetFunction == nullptr)
+	{
+		return;
+	}
+
+	// find any pins marked as SetParam
+	const FString& DependentPinMetaData = TargetFunction->GetMetaData(FBlueprintMetadata::MD_SetParam);
+	if (DependentPinMetaData.IsEmpty())
+	{
+		// early out, no SetParam metadata
+		return;
+	}
+
+	// useless copies/allocates in this code, could be an optimization target...
+	TArray<FString> SetParamPinGroups;
+	{
+		DependentPinMetaData.ParseIntoArray(SetParamPinGroups, TEXT(","), true);
+	}
+
+	for (FString& Entry : SetParamPinGroups)
+	{
+		// split the group:
+		TArray<FString> GroupEntries;
+		Entry.ParseIntoArray(GroupEntries, TEXT("|"), true);
+		// resolve pins
+		TArray<UEdGraphPin*> ResolvedPins;
+		for(UEdGraphPin* Pin : Pins)
+		{
+			if (GroupEntries.Contains(Pin->GetName()))
+			{
+				ResolvedPins.Add(Pin);
+			}
+		}
+
+		// if nothing is connected (or non-default), reset to wildcard
+		// else, find the first type and propagate to everyone else::
+		bool bAllPinsAreDefaultOrEmpty = true;
+		FEdGraphTerminalType TypeToPropagate;
+		for (UEdGraphPin* Pin : ResolvedPins)
+		{
+			if (Pin->LinkedTo.Num() != 0 || Pin->DefaultValue != Pin->AutogeneratedDefaultValue)
+			{
+				bAllPinsAreDefaultOrEmpty = false;
+				if (Pin->LinkedTo.Num() != 0)
+				{
+					TypeToPropagate = Pin->LinkedTo[0]->GetPrimaryTerminalType();
+				}
+				else
+				{
+					TypeToPropagate = Pin->GetPrimaryTerminalType();
+				}
+				break;
+			}
+		}
+
+		if (!bAllPinsAreDefaultOrEmpty)
+		{
+			for (UEdGraphPin* Pin : ResolvedPins)
+			{
+				Pin->PinType.PinCategory = TypeToPropagate.TerminalCategory;
+				Pin->PinType.PinSubCategory = TypeToPropagate.TerminalSubCategory;
+				Pin->PinType.PinSubCategoryObject = TypeToPropagate.TerminalSubCategoryObject;
+			}
+		}
+		else
+		{
+			// reset everyone to wildcard:
+			for (UEdGraphPin* Pin : ResolvedPins)
+			{
+				Pin->PinType.PinCategory = UEdGraphSchema_K2::PC_Wildcard;
+				Pin->PinType.PinSubCategory.Empty();
+				Pin->PinType.PinSubCategoryObject = nullptr;
+				Pin->DefaultValue = TEXT("");
+			}
+		}
+	}
+}
+
 FText UK2Node_CallFunction::GetToolTipHeading() const
 {
 	FText Heading = Super::GetToolTipHeading();
@@ -2464,6 +2551,15 @@ bool UK2Node_CallFunction::IsStructureWildcardProperty(const UFunction* Function
 	return false;
 }
 
+bool UK2Node_CallFunction::IsWildcardProperty(const UFunction* InFunction, const UProperty* InProperty)
+{
+	if (InProperty)
+	{
+		return FEdGraphUtilities::IsSetParam(InFunction, InProperty->GetName());
+	}
+	return false;
+}
+
 void UK2Node_CallFunction::AddSearchMetaDataInfo(TArray<struct FSearchTagDataPair>& OutTaggedMetaData) const
 {
 	Super::AddSearchMetaDataInfo(OutTaggedMetaData);
diff --git a/Engine/Source/Editor/BlueprintGraph/Private/K2Node_Copy.cpp b/Engine/Source/Editor/BlueprintGraph/Private/K2Node_Copy.cpp
--- a/Engine/Source/Editor/BlueprintGraph/Private/K2Node_Copy.cpp
+++ b/Engine/Source/Editor/BlueprintGraph/Private/K2Node_Copy.cpp
@@ -192,9 +192,9 @@ bool UK2Node_Copy::IsConnectionDisallowed(const UEdGraphPin* MyPin, const UEdGra
 			OutReason = FText::Format(LOCTEXT("ObjectConnectionDisallowed", "Cannot connect with {0} pin."), FText::FromString(OtherPin->PinType.PinCategory)).ToString();
 			return true;
 		}
-		else if (OtherPin->PinType.bIsArray)
+		else if (OtherPin->PinType.IsContainer())
 		{
-			OutReason = LOCTEXT("ArrayConnectionDisallowed", "Cannot connect with array pin.").ToString();
+			OutReason = LOCTEXT("ArrayConnectionDisallowed", "Cannot connect with container pin.").ToString();
 			return true;
 		}
 	}
diff --git a/Engine/Source/Editor/BlueprintGraph/Private/K2Node_DynamicCast.cpp b/Engine/Source/Editor/BlueprintGraph/Private/K2Node_DynamicCast.cpp
--- a/Engine/Source/Editor/BlueprintGraph/Private/K2Node_DynamicCast.cpp
+++ b/Engine/Source/Editor/BlueprintGraph/Private/K2Node_DynamicCast.cpp
@@ -326,10 +326,10 @@ bool UK2Node_DynamicCast::IsConnectionDisallowed(const UEdGraphPin* MyPin, const
 		const FEdGraphPinType& OtherPinType = OtherPin->PinType;
 		const FText OtherPinName = OtherPin->PinFriendlyName.IsEmpty() ? FText::FromString(OtherPin->PinName) : OtherPin->PinFriendlyName;
 
-		if (OtherPinType.bIsArray)
+		if (OtherPinType.IsContainer())
 		{
 			bIsDisallowed = true;
-			OutReason = LOCTEXT("CannotArrayCast", "You cannot cast arrays of objects.").ToString();
+			OutReason = LOCTEXT("CannotContainerCast", "You cannot cast containers of objects.").ToString();
 		}
 		else if (TargetType == nullptr)
 		{
diff --git a/Engine/Source/Editor/BlueprintGraph/Private/K2Node_EaseFunction.cpp b/Engine/Source/Editor/BlueprintGraph/Private/K2Node_EaseFunction.cpp
--- a/Engine/Source/Editor/BlueprintGraph/Private/K2Node_EaseFunction.cpp
+++ b/Engine/Source/Editor/BlueprintGraph/Private/K2Node_EaseFunction.cpp
@@ -214,12 +214,18 @@ bool UK2Node_EaseFunction::IsConnectionDisallowed(const UEdGraphPin* MyPin, cons
 		MyPin->PinName == FEaseFunctionNodeHelper::GetBPinName() || 
 		MyPin->PinName == FEaseFunctionNodeHelper::GetResultPinName())
 	{
-		const bool bConnectionOk = (OtherPin->PinType.PinCategory == K2Schema->PC_Float ||
-			(OtherPin->PinType.PinCategory == K2Schema->PC_Struct &&
-			OtherPin->PinType.PinSubCategoryObject.IsValid() &&
-			(OtherPin->PinType.PinSubCategoryObject.Get()->GetName() == TEXT("Vector") ||
-			OtherPin->PinType.PinSubCategoryObject.Get()->GetName() == TEXT("Rotator") ||
-			OtherPin->PinType.PinSubCategoryObject.Get()->GetName() == TEXT("Transform"))));
+		const bool bConnectionOk = (
+			OtherPin->PinType.PinCategory == K2Schema->PC_Float ||
+			(
+				OtherPin->PinType.PinCategory == K2Schema->PC_Struct &&
+				OtherPin->PinType.PinSubCategoryObject.IsValid() &&
+				(
+					OtherPin->PinType.PinSubCategoryObject.Get()->GetName() == TEXT("Vector") ||
+					OtherPin->PinType.PinSubCategoryObject.Get()->GetName() == TEXT("Rotator") ||
+					OtherPin->PinType.PinSubCategoryObject.Get()->GetName() == TEXT("Transform")
+				)
+			)
+		) && !OtherPin->PinType.IsContainer();
 		if (!bConnectionOk)
 		{
 			OutReason = LOCTEXT("PinConnectionDisallowed", "Pin type is not supported by function.").ToString();
diff --git a/Engine/Source/Editor/BlueprintGraph/Private/K2Node_MakeArray.cpp b/Engine/Source/Editor/BlueprintGraph/Private/K2Node_MakeArray.cpp
--- a/Engine/Source/Editor/BlueprintGraph/Private/K2Node_MakeArray.cpp
+++ b/Engine/Source/Editor/BlueprintGraph/Private/K2Node_MakeArray.cpp
@@ -520,9 +520,9 @@ void UK2Node_MakeArray::GetContextMenuActions(const FGraphNodeContextMenuBuilder
 bool UK2Node_MakeArray::IsConnectionDisallowed(const UEdGraphPin* MyPin, const UEdGraphPin* OtherPin, FString& OutReason) const
 {
 	// if MyPin has a ParentPin then we are dealing with a split pin and we should evaluate it with default behavior
-	if(MyPin->ParentPin == nullptr && OtherPin->PinType.bIsArray == true && MyPin->Direction == EGPD_Input)
+	if(MyPin->ParentPin == nullptr && OtherPin->PinType.IsContainer() == true && MyPin->Direction == EGPD_Input)
 	{
-		OutReason = NSLOCTEXT("K2Node", "MakeArray_InputIsArray", "Cannot make an array with an input of an array!").ToString();
+		OutReason = NSLOCTEXT("K2Node", "MakeArray_InputIsContainer", "Cannot make an array with an input of a container!").ToString();
 		return true;
 	}
 
diff --git a/Engine/Source/Editor/BlueprintGraph/Public/BlueprintEditorSettings.h b/Engine/Source/Editor/BlueprintGraph/Public/BlueprintEditorSettings.h
--- a/Engine/Source/Editor/BlueprintGraph/Public/BlueprintEditorSettings.h
+++ b/Engine/Source/Editor/BlueprintGraph/Public/BlueprintEditorSettings.h
@@ -78,6 +78,10 @@ public:
 	/** If set will exclude components added in a Blueprint class Construction Script from the component details view */
 	UPROPERTY(EditAnywhere, config, Category=Workflow)
 	bool bHideConstructionScriptComponentsInDetailsView;
+
+	UPROPERTY(EditAnywhere, config, Category = Experimental, AdvancedDisplay, meta = (DisplayName = "Enable (incomplete, non-functional) support for TSet and TMap"))
+	bool bEnableAdvancedContainers;
+
 // Compiler Settings
 public:
 	/** Determines when to save Blueprints post-compile */
@@ -98,6 +102,10 @@ public:
 	UPROPERTY(EditAnywhere, config, Category=DeveloperTools)
 	bool bShowActionMenuItemSignatures;
 
+	/** If enabled, blueprint nodes in the event graph will display with unique names rather than their display name. */
+	UPROPERTY(EditAnywhere, config, Category = DeveloperTools, meta = (DisplayName = "Display Unique Names for Blueprint Nodes"))
+	bool bBlueprintNodeUniqueNames;
+
 // Perf Settings
 public:
 	/** If enabled, additional details will be displayed in the Compiler Results tab after compiling a blueprint. */
diff --git a/Engine/Source/Editor/EditorStyle/Private/SlateEditorStyle.cpp b/Engine/Source/Editor/EditorStyle/Private/SlateEditorStyle.cpp
--- a/Engine/Source/Editor/EditorStyle/Private/SlateEditorStyle.cpp
+++ b/Engine/Source/Editor/EditorStyle/Private/SlateEditorStyle.cpp
@@ -5310,6 +5310,9 @@ void FSlateEditorStyle::FStyle::SetupPersonaStyle()
 
 		Set( "Kismet.VariableList.TypeIcon", new IMAGE_BRUSH( "/Icons/pill_16x", Icon16x16 ) );
 		Set( "Kismet.VariableList.ArrayTypeIcon", new IMAGE_BRUSH( "/Icons/pillarray_16x", Icon16x16 ) );
+		Set( "Kismet.VariableList.SetTypeIcon", new IMAGE_BRUSH( "/Icons/pillset_16x", Icon16x16 ) );
+		Set( "Kismet.VariableList.MapValueTypeIcon", new IMAGE_BRUSH( "/Icons/pillmapvalue_16x", Icon16x16 ) );
+		Set( "Kismet.VariableList.MapKeyTypeIcon", new IMAGE_BRUSH( "/Icons/pillmapkey_16x", Icon16x16 ) );
 		Set( "Kismet.VariableList.ExposeForInstance", new IMAGE_BRUSH( "/Icons/icon_layer_visible", Icon16x16 ) );
 		Set( "Kismet.VariableList.HideForInstance", new IMAGE_BRUSH( "/Icons/icon_layer_not_visible", Icon16x16 ) );
 		Set( "Kismet.VariableList.VariableIsUsed", new IMAGE_BRUSH( "/Icons/icon_variable_used_16x", Icon16x16 ) );
diff --git a/Engine/Source/Editor/GraphEditor/Private/GraphEditorDragDropAction.cpp b/Engine/Source/Editor/GraphEditor/Private/GraphEditorDragDropAction.cpp
--- a/Engine/Source/Editor/GraphEditor/Private/GraphEditorDragDropAction.cpp
+++ b/Engine/Source/Editor/GraphEditor/Private/GraphEditorDragDropAction.cpp
@@ -1,6 +1,7 @@
 // Copyright 1998-2016 Epic Games, Inc. All Rights Reserved.
 
 #include "GraphEditorCommon.h"
+#include "SPinTypeSelector.h"
 #include "SScaleBox.h"
 
 UEdGraphPin* FGraphEditorDragDropAction::GetHoveredPin() const
@@ -115,7 +116,7 @@ void FGraphEditorDragDropAction::SetFeedbackMessage(const TSharedPtr<SWidget>& M
 	}
 }
 
-void FGraphEditorDragDropAction::SetSimpleFeedbackMessage(const FSlateBrush* Icon, const FSlateColor& IconColor, const FText& Message)
+void FGraphEditorDragDropAction::SetSimpleFeedbackMessage(const FSlateBrush* Icon, const FSlateColor& IconColor, const FText& Message, const FSlateBrush* SecondaryIcon /*= nullptr*/, const FSlateColor SecondaryColor /*= FSlateColor()*/)
 {
 	// Let the user know the status of making this connection.
 
@@ -124,7 +125,10 @@ void FGraphEditorDragDropAction::SetSimpleFeedbackMessage(const FSlateBrush* Ico
 	TAttribute<EVisibility> ErrorIconVisibility = TAttribute<EVisibility>::Create(TAttribute<EVisibility>::FGetter::CreateRaw(this, &FGraphEditorDragDropAction::GetErrorIconVisible));
 	TAttribute<EVisibility> IconVisibility = TAttribute<EVisibility>::Create(TAttribute<EVisibility>::FGetter::CreateRaw(this, &FGraphEditorDragDropAction::GetIconVisible));
 	
-		SetFeedbackMessage(
+	TSharedRef<SWidget> TypeImage = SPinTypeSelector::ConstructPinTypeImage(Icon, IconColor, SecondaryIcon, SecondaryColor, TSharedPtr<SToolTip>());
+	TypeImage->SetVisibility(IconVisibility);
+
+	SetFeedbackMessage(
 		SNew(SHorizontalBox)
 		+SHorizontalBox::Slot()
 		.AutoWidth()
@@ -146,10 +150,7 @@ void FGraphEditorDragDropAction::SetSimpleFeedbackMessage(const FSlateBrush* Ico
 			SNew(SScaleBox)
 			.Stretch(EStretch::ScaleToFit)
 			[
-				SNew(SImage)
-				.Visibility(IconVisibility)
-				.Image( Icon )
-				.ColorAndOpacity( IconColor )
+				TypeImage
 			]
 		]
 		+SHorizontalBox::Slot()
diff --git a/Engine/Source/Editor/GraphEditor/Private/SGraphNode.cpp b/Engine/Source/Editor/GraphEditor/Private/SGraphNode.cpp
--- a/Engine/Source/Editor/GraphEditor/Private/SGraphNode.cpp
+++ b/Engine/Source/Editor/GraphEditor/Private/SGraphNode.cpp
@@ -12,6 +12,7 @@
 #include "TutorialMetaData.h"
 #include "SGraphPanel.h"
 #include "SInlineEditableTextBlock.h"
+#include "BlueprintEditorSettings.h"
 
 /////////////////////////////////////////////////////
 // SNodeTitle
@@ -51,9 +52,16 @@ void SNodeTitle::Tick( const FGeometry& AllottedGeometry, const double InCurrent
 
 FText SNodeTitle::GetNodeTitle() const
 {
-	return (GraphNode != NULL)
-		? GraphNode->GetNodeTitle(ENodeTitleType::FullTitle)
-		: NSLOCTEXT("GraphEditor", "NullNode", "Null Node");
+	if (GetDefault<UBlueprintEditorSettings>()->bBlueprintNodeUniqueNames && GraphNode)
+	{
+		return FText::FromName(GraphNode->GetFName());
+	}
+	else
+	{
+		return (GraphNode != NULL)
+			? GraphNode->GetNodeTitle(ENodeTitleType::FullTitle)
+			: NSLOCTEXT("GraphEditor", "NullNode", "Null Node");
+	}
 }
 
 FText SNodeTitle::GetHeadTitle() const
diff --git a/Engine/Source/Editor/GraphEditor/Private/SGraphPalette.cpp b/Engine/Source/Editor/GraphEditor/Private/SGraphPalette.cpp
--- a/Engine/Source/Editor/GraphEditor/Private/SGraphPalette.cpp
+++ b/Engine/Source/Editor/GraphEditor/Private/SGraphPalette.cpp
@@ -11,6 +11,7 @@
 
 #include "IDocumentation.h"
 #include "SInlineEditableTextBlock.h"
+#include "SPinTypeSelector.h"
 
 void SGraphPaletteItem::Construct(const FArguments& InArgs, FCreateWidgetForActionData* const InCreateData)
 {
@@ -100,14 +101,14 @@ TSharedRef<SWidget> SGraphPaletteItem::CreateIconWidget(const FText& IconToolTip
 		.ColorAndOpacity(IconColor);
 }
 
-TSharedRef<SWidget> SGraphPaletteItem::CreateIconWidget(const FText& IconToolTip, const FSlateBrush* IconBrush, const FSlateColor& IconColor, const FString& DocLink, const FString& DocExcerpt)
+TSharedRef<SWidget> SGraphPaletteItem::CreateIconWidget(const FText& IconToolTip, const FSlateBrush* IconBrush, const FSlateColor& IconColor, const FString& DocLink, const FString& DocExcerpt, const FSlateBrush* SecondaryIconBrush, const FSlateColor& SecondaryColor)
 {
-	TSharedPtr<SToolTip> ToolTipWidget = IDocumentation::Get()->CreateToolTip(IconToolTip, NULL, DocLink, DocExcerpt);
-
-	return SNew(SImage)
-		.ToolTip(ToolTipWidget)
-		.Image( IconBrush )
-		.ColorAndOpacity(IconColor);
+	return SPinTypeSelector::ConstructPinTypeImage(
+		IconBrush, 
+		IconColor, 
+		SecondaryIconBrush, 
+		SecondaryColor, 
+		IDocumentation::Get()->CreateToolTip(IconToolTip, NULL, DocLink, DocExcerpt));
 }
 
 TSharedRef<SWidget> SGraphPaletteItem::CreateTextSlotWidget( const FSlateFontInfo& NameFont, FCreateWidgetForActionData* const InCreateData, bool bIsReadOnly )
diff --git a/Engine/Source/Editor/GraphEditor/Public/GraphEditorDragDropAction.h b/Engine/Source/Editor/GraphEditor/Public/GraphEditorDragDropAction.h
--- a/Engine/Source/Editor/GraphEditor/Public/GraphEditorDragDropAction.h
+++ b/Engine/Source/Editor/GraphEditor/Public/GraphEditorDragDropAction.h
@@ -31,7 +31,7 @@ public:
 
 	bool HasFeedbackMessage();
 	void SetFeedbackMessage(const TSharedPtr<SWidget>& Message);
-	void SetSimpleFeedbackMessage(const FSlateBrush* Icon, const FSlateColor& IconColor, const FText& Message);
+	void SetSimpleFeedbackMessage(const FSlateBrush* Icon, const FSlateColor& IconColor, const FText& Message, const FSlateBrush* SecondaryIcon = nullptr, const FSlateColor SecondaryColor = FSlateColor());
 
 protected:
 	UEdGraphPin* GetHoveredPin() const;
diff --git a/Engine/Source/Editor/GraphEditor/Public/SGraphPalette.h b/Engine/Source/Editor/GraphEditor/Public/SGraphPalette.h
--- a/Engine/Source/Editor/GraphEditor/Public/SGraphPalette.h
+++ b/Engine/Source/Editor/GraphEditor/Public/SGraphPalette.h
@@ -33,11 +33,10 @@ public:
 
 protected:
 	/** Creates an icon type widget, including big tooltip link */
-	TSharedRef<SWidget> CreateIconWidget( const FText& IconToolTip, const FSlateBrush* IconBrush, const FSlateColor& IconColor, const FString& DocLink, const FString& DocExcerpt );
+	TSharedRef<SWidget> CreateIconWidget( const FText& IconToolTip, const FSlateBrush* IconBrush, const FSlateColor& IconColor, const FString& DocLink, const FString& DocExcerpt, const FSlateBrush* SecondaryIconBrush, const FSlateColor& SecondaryColor );
 	/** Create an icon type widget */
 	TSharedRef<SWidget> CreateIconWidget(const FText& IconToolTip, const FSlateBrush* IconBrush, const FSlateColor& IconColor);
 
-
 	/* Create the text widget */
 	virtual TSharedRef<SWidget> CreateTextSlotWidget( const FSlateFontInfo& NameFont,  FCreateWidgetForActionData* const InCreateData, bool bIsReadOnly );
 
diff --git a/Engine/Source/Editor/Kismet/Private/BPVariableDragDropAction.cpp b/Engine/Source/Editor/Kismet/Private/BPVariableDragDropAction.cpp
--- a/Engine/Source/Editor/Kismet/Private/BPVariableDragDropAction.cpp
+++ b/Engine/Source/Editor/Kismet/Private/BPVariableDragDropAction.cpp
@@ -55,6 +55,8 @@ void FKismetVariableDragDropAction::HoverTargetChanged()
 	// Icon/text to draw on tooltip
 	FSlateColor IconColor = FLinearColor::White;
 	const FSlateBrush* StatusSymbol = FEditorStyle::GetBrush(TEXT("Graph.ConnectorFeedback.Error"));
+	FSlateBrush const* SecondaryIcon = nullptr;
+	FSlateColor SecondaryColor = FLinearColor::Transparent;
 	FText Message = LOCTEXT("InvalidDropTarget", "Invalid drop target!");
 
 	UEdGraphPin* PinUnderCursor = GetHoveredPin();
@@ -306,11 +308,11 @@ void FKismetVariableDragDropAction::HoverTargetChanged()
 	// Draw variable icon
 	else
 	{
-		StatusSymbol = FBlueprintEditor::GetVarIconAndColor(VariableSource.Get(), VariableName, IconColor);
+		StatusSymbol = FBlueprintEditor::GetVarIconAndColor(VariableSource.Get(), VariableName, IconColor, SecondaryIcon, SecondaryColor);
 		Message = FText::FromString(VariableString);
 	}
 
-	SetSimpleFeedbackMessage(StatusSymbol, IconColor, Message);
+	SetSimpleFeedbackMessage(StatusSymbol, IconColor, Message, SecondaryIcon, SecondaryColor);
 }
 
 FReply FKismetVariableDragDropAction::DroppedOnPin(FVector2D ScreenPosition, FVector2D GraphPosition)
diff --git a/Engine/Source/Editor/Kismet/Private/BlueprintDragDropMenuItem.cpp b/Engine/Source/Editor/Kismet/Private/BlueprintDragDropMenuItem.cpp
--- a/Engine/Source/Editor/Kismet/Private/BlueprintDragDropMenuItem.cpp
+++ b/Engine/Source/Editor/Kismet/Private/BlueprintDragDropMenuItem.cpp
@@ -87,7 +87,11 @@ FSlateBrush const* FBlueprintDragDropMenuItem::GetMenuIcon(FSlateColor& ColorOut
 		if (UProperty const* Property = VariableSpawner->GetVarProperty())
 		{
 			UStruct* const PropertyOwner = CastChecked<UStruct>(Property->GetOuterUField());
-			IconBrush = FBlueprintEditor::GetVarIconAndColor(PropertyOwner, Property->GetFName(), ColorOut);
+			// @note(DanO): not sure wht relies on this logic, but discarding extra icon info util
+			// I can confirm that containers are a possility here:
+			FSlateBrush const* Discarded = nullptr;
+			FSlateColor DiscardedColor;
+			IconBrush = FBlueprintEditor::GetVarIconAndColor(PropertyOwner, Property->GetFName(), ColorOut, Discarded, DiscardedColor);
 		}
 	}
 	return IconBrush;
diff --git a/Engine/Source/Editor/Kismet/Private/BlueprintEditor.cpp b/Engine/Source/Editor/Kismet/Private/BlueprintEditor.cpp
--- a/Engine/Source/Editor/Kismet/Private/BlueprintEditor.cpp
+++ b/Engine/Source/Editor/Kismet/Private/BlueprintEditor.cpp
@@ -495,12 +495,26 @@ const FSlateBrush* FBlueprintEditor::GetGlyphForGraph(const UEdGraph* Graph, boo
 	return ReturnValue;
 }
 
-FSlateBrush const* FBlueprintEditor::GetVarIconAndColor(const UStruct* VarScope, FName VarName, FSlateColor& IconColorOut)
+FSlateBrush const* FBlueprintEditor::GetVarIconAndColor(const UStruct* VarScope, FName VarName, FSlateColor& IconColorOut, FSlateBrush const*& SecondaryBrushOut, FSlateColor& SecondaryColorOut)
 {
-	FLinearColor ColorOut;
-	const FSlateBrush* IconBrush = UK2Node_Variable::GetVariableIconAndColor(VarScope, VarName, ColorOut).GetOptionalIcon();
-	IconColorOut = ColorOut;
-	return IconBrush;
+	if (VarScope != NULL)
+	{
+		UProperty* Property = FindField<UProperty>(VarScope, VarName);
+		if (Property != NULL)
+		{
+			const UEdGraphSchema_K2* K2Schema = GetDefault<UEdGraphSchema_K2>();
+
+			FEdGraphPinType PinType;
+			if (K2Schema->ConvertPropertyToPinType(Property, PinType)) // use schema to get the color
+			{
+				IconColorOut = K2Schema->GetPinTypeColor(PinType);
+				SecondaryBrushOut = FBlueprintEditorUtils::GetSecondaryIconFromPin(PinType);
+				SecondaryColorOut = K2Schema->GetSecondaryPinTypeColor(PinType);
+				return FBlueprintEditorUtils::GetIconFromPin(PinType);
+			}
+		}
+	}
+	return FEditorStyle::GetBrush(TEXT("Kismet.AllClasses.VariableIcon"));
 }
 
 bool FBlueprintEditor::IsInAScriptingMode() const
diff --git a/Engine/Source/Editor/Kismet/Private/Profiler/BPProfilerStatisticWidgets.cpp b/Engine/Source/Editor/Kismet/Private/Profiler/BPProfilerStatisticWidgets.cpp
--- a/Engine/Source/Editor/Kismet/Private/Profiler/BPProfilerStatisticWidgets.cpp
+++ b/Engine/Source/Editor/Kismet/Private/Profiler/BPProfilerStatisticWidgets.cpp
@@ -29,6 +29,16 @@ namespace BlueprintProfilerStatText
 	const FText ColumnText_MinTime(LOCTEXT("MinTime", "Min Time (ms)"));
 	const FText ColumnText_Samples(LOCTEXT("Samples", "Samples"));
 	const FText ColumnText_TotalTime(LOCTEXT("TotalTime", "Total Time (ms)"));
+	// Debugging values
+	const FName ColumnId_HottestPath("HottestPath");
+	const FText ColumnText_HottestPath(LOCTEXT("Hottest Path", "Path Heat"));
+	const FName ColumnId_HeatLevel("HeatLevel");
+	const FText ColumnText_HeatLevel(LOCTEXT("Heat Levels", "Heat Level"));
+
+	// Text Heat display colors
+	const FLinearColor HotColor(1.f, 0.f, 0.f, 1.f);
+	const FLinearColor ColdColor(1.f, 1.f, 1.f, 1.f);
+	const FLinearColor InactiveColor(1.f, 1.f, 1.f, 0.6f);
 };
 
 //////////////////////////////////////////////////////////////////////////
@@ -79,6 +89,9 @@ const FName SProfilerStatRow::GetStatName(const EBlueprintProfilerStat::Type Sta
 		case EBlueprintProfilerStat::MaxTime:		return BlueprintProfilerStatText::ColumnId_MaxTime;
 		case EBlueprintProfilerStat::MinTime:		return BlueprintProfilerStatText::ColumnId_MinTime;
 		case EBlueprintProfilerStat::Samples:		return BlueprintProfilerStatText::ColumnId_Samples;
+		// Debugging Values
+		case EBlueprintProfilerStat::HottestPath:	return BlueprintProfilerStatText::ColumnId_HottestPath;
+		case EBlueprintProfilerStat::HeatLevel:		return BlueprintProfilerStatText::ColumnId_HeatLevel;
 		default:									return NAME_None;
 	}
 }
@@ -94,6 +107,9 @@ const FText SProfilerStatRow::GetStatText(const EBlueprintProfilerStat::Type Sta
 		case EBlueprintProfilerStat::MaxTime:		return BlueprintProfilerStatText::ColumnText_MaxTime;
 		case EBlueprintProfilerStat::MinTime:		return BlueprintProfilerStatText::ColumnText_MinTime;
 		case EBlueprintProfilerStat::Samples:		return BlueprintProfilerStatText::ColumnText_Samples;
+		// Debugging Values
+		case EBlueprintProfilerStat::HottestPath:	return BlueprintProfilerStatText::ColumnText_HottestPath;
+		case EBlueprintProfilerStat::HeatLevel:		return BlueprintProfilerStatText::ColumnText_HeatLevel;
 		default:									return FText::GetEmpty();
 	}
 }
@@ -142,13 +158,13 @@ TSharedRef<SWidget> FBPProfilerStatWidget::GenerateColumnWidget(FName ColumnName
 		else
 		{
 			TAttribute<FText> TextAttr(LOCTEXT("NonApplicableStat", ""));
-			TAttribute<FSlateColor> ColorAttr;
+			TAttribute<FSlateColor> ColorAttr = TAttribute<FSlateColor>(BlueprintProfilerStatText::InactiveColor);
 			const uint32 NonNodeStats = EScriptExecutionNodeFlags::Container|EScriptExecutionNodeFlags::CallSite|EScriptExecutionNodeFlags::BranchNode|EScriptExecutionNodeFlags::ExecPin;
 
 			if (ColumnName == BlueprintProfilerStatText::ColumnId_TotalTime)
 			{
 				TextAttr = TAttribute<FText>(PerformanceStats.Get(), &FScriptPerfData::GetTotalTimingText);
-				ColorAttr = TAttribute<FSlateColor>(FLinearColor::White);
+				ColorAttr = TAttribute<FSlateColor>(this, &FBPProfilerStatWidget::GetTotalTimeHeatColor);
 			}
 			else if (ColumnName == BlueprintProfilerStatText::ColumnId_InclusiveTime)
 			{
@@ -163,18 +179,24 @@ TSharedRef<SWidget> FBPProfilerStatWidget::GenerateColumnWidget(FName ColumnName
 			else if (ColumnName == BlueprintProfilerStatText::ColumnId_MaxTime)
 			{
 				TextAttr = TAttribute<FText>(PerformanceStats.Get(), &FScriptPerfData::GetMaxTimingText);
-				//ColorAttr = TAttribute<FSlateColor>(this, &FBPProfilerStatWidget::GetMaxTimeHeatColor);
-				ColorAttr = TAttribute<FSlateColor>(FLinearColor::White);
+				ColorAttr = TAttribute<FSlateColor>(this, &FBPProfilerStatWidget::GetMaxTimeHeatColor);
 			}
 			else if (ColumnName == BlueprintProfilerStatText::ColumnId_MinTime)
 			{
 				TextAttr = TAttribute<FText>(PerformanceStats.Get(), &FScriptPerfData::GetMinTimingText);
-				ColorAttr = TAttribute<FSlateColor>(FLinearColor::White);
 			}
 			else if (ColumnName == BlueprintProfilerStatText::ColumnId_Samples)
 			{
 				TextAttr = TAttribute<FText>(PerformanceStats.Get(), &FScriptPerfData::GetSamplesText);
-				ColorAttr = TAttribute<FSlateColor>(FLinearColor::White);
+			}
+			// Debugging Values
+			else if (ColumnName == BlueprintProfilerStatText::ColumnId_HottestPath)
+			{
+				TextAttr = TAttribute<FText>(PerformanceStats.Get(), &FScriptPerfData::GetHottestPathText);
+			}
+			else if (ColumnName == BlueprintProfilerStatText::ColumnId_HeatLevel)
+			{
+				TextAttr = TAttribute<FText>(PerformanceStats.Get(), &FScriptPerfData::GetHeatLevelText);
 			}
 			// Create the actual widget
 			return SNew(STextBlock)
@@ -195,20 +217,50 @@ void FBPProfilerStatWidget::NavigateTo() const
 
 FSlateColor FBPProfilerStatWidget::GetAverageHeatColor() const
 {
-	const float Value = 1.f - PerformanceStats->GetAverageHeatLevel();
-	return FLinearColor(1.f, Value, Value);
+	const UBlueprintProfilerSettings* ProfilerSettings = GetDefault<UBlueprintProfilerSettings>();
+	if (ProfilerSettings->GraphNodeHeatMapDisplayMode == EBlueprintProfilerHeatMapDisplayMode::Average)
+	{
+		const float Value = 1.f - PerformanceStats->GetAverageHeatLevel();
+		return FLinearColor(1.f, Value, Value);
+//		return FLinearColor::LerpUsingHSV(BlueprintProfilerStatText::ColdColor, BlueprintProfilerStatText::HotColor, PerformanceStats->GetAverageHeatLevel());
+	}
+	return BlueprintProfilerStatText::InactiveColor;
 }
 
 FSlateColor FBPProfilerStatWidget::GetInclusiveHeatColor() const
 {
-	const float Value = 1.f - PerformanceStats->GetInclusiveHeatLevel();
-	return FLinearColor(1.f, Value, Value);
+	const UBlueprintProfilerSettings* ProfilerSettings = GetDefault<UBlueprintProfilerSettings>();
+	if (ProfilerSettings->GraphNodeHeatMapDisplayMode == EBlueprintProfilerHeatMapDisplayMode::Inclusive)
+	{
+		const float Value = 1.f - PerformanceStats->GetInclusiveHeatLevel();
+		return FLinearColor(1.f, Value, Value);
+//		return FLinearColor::LerpUsingHSV(BlueprintProfilerStatText::ColdColor, BlueprintProfilerStatText::HotColor, PerformanceStats->GetInclusiveHeatLevel());
+	}
+	return BlueprintProfilerStatText::InactiveColor;
 }
 
 FSlateColor FBPProfilerStatWidget::GetMaxTimeHeatColor() const
 {
-	const float Value = 1.f - PerformanceStats->GetMaxTimeHeatLevel();
-	return FLinearColor(1.f, Value, Value);
+	const UBlueprintProfilerSettings* ProfilerSettings = GetDefault<UBlueprintProfilerSettings>();
+	if (ProfilerSettings->GraphNodeHeatMapDisplayMode == EBlueprintProfilerHeatMapDisplayMode::MaxTiming)
+	{
+		const float Value = 1.f - PerformanceStats->GetMaxTimeHeatLevel();
+		return FLinearColor(1.f, Value, Value);
+//		return FLinearColor::LerpUsingHSV(BlueprintProfilerStatText::ColdColor, BlueprintProfilerStatText::HotColor, PerformanceStats->GetMaxTimeHeatLevel());
+	}
+	return BlueprintProfilerStatText::InactiveColor;
+}
+
+FSlateColor FBPProfilerStatWidget::GetTotalTimeHeatColor() const
+{
+	const UBlueprintProfilerSettings* ProfilerSettings = GetDefault<UBlueprintProfilerSettings>();
+	if (ProfilerSettings->GraphNodeHeatMapDisplayMode == EBlueprintProfilerHeatMapDisplayMode::Total)
+	{
+		const float Value = 1.f - PerformanceStats->GetTotalHeatLevel();
+		return FLinearColor(1.f, Value, Value);
+//		return FLinearColor::LerpUsingHSV(BlueprintProfilerStatText::ColdColor, BlueprintProfilerStatText::HotColor, PerformanceStats->GetTotalHeatLevel());
+	}
+	return BlueprintProfilerStatText::InactiveColor;
 }
 
 void FBPProfilerStatWidget::GenerateExecNodeWidgets(const TSharedPtr<FBlueprintProfilerStatOptions> DisplayOptions)
diff --git a/Engine/Source/Editor/Kismet/Private/Profiler/BPProfilerStatisticWidgets.h b/Engine/Source/Editor/Kismet/Private/Profiler/BPProfilerStatisticWidgets.h
--- a/Engine/Source/Editor/Kismet/Private/Profiler/BPProfilerStatisticWidgets.h
+++ b/Engine/Source/Editor/Kismet/Private/Profiler/BPProfilerStatisticWidgets.h
@@ -21,6 +21,9 @@ namespace EBlueprintProfilerStat
 		MaxTime,
 		MinTime,
 		Samples
+		// Debugging Values
+		, HottestPath
+		, HeatLevel
 	};
 };
 
@@ -123,6 +126,9 @@ protected:
 	/** Calculate and return the heat color to use for max time stats. */
 	FSlateColor GetMaxTimeHeatColor() const;
 
+	/** Calculate and return the heat color to use for total time stats. */
+	FSlateColor GetTotalTimeHeatColor() const;
+
 protected:
 
 	/** Widget script tracepath */
diff --git a/Engine/Source/Editor/Kismet/Private/Profiler/BlueprintProfilerSettings.cpp b/Engine/Source/Editor/Kismet/Private/Profiler/BlueprintProfilerSettings.cpp
--- a/Engine/Source/Editor/Kismet/Private/Profiler/BlueprintProfilerSettings.cpp
+++ b/Engine/Source/Editor/Kismet/Private/Profiler/BlueprintProfilerSettings.cpp
@@ -4,11 +4,6 @@
 #include "BlueprintProfilerSettings.h"
 #include "ScriptPerfData.h"
 
-float UBlueprintProfilerSettings::CustomEventPerformanceThresholdDefaultValue = 1.f;
-float UBlueprintProfilerSettings::CustomAveragePerformanceThresholdDefaultValue = 0.2f;
-float UBlueprintProfilerSettings::CustomInclusivePerformanceThresholdDefaultValue = 0.25f;
-float UBlueprintProfilerSettings::CustomMaxPerformanceThresholdDefaultValue = 0.5f;
-
 UBlueprintProfilerSettings::UBlueprintProfilerSettings(const FObjectInitializer& ObjectInitializer)
 	: Super(ObjectInitializer)
 	, bDisplayByInstance(true)
@@ -16,15 +11,14 @@ UBlueprintProfilerSettings::UBlueprintProfilerSettings(const FObjectInitializer&
 	, bGraphFilter(false)
 	, bDisplayPure(true)
 	, bDisplayInheritedEvents(true)
-	, bAverageBlueprintStats(true)
 	, GraphNodeHeatMapDisplayMode(EBlueprintProfilerHeatMapDisplayMode::Average)
 	, WireHeatMapDisplayMode(EBlueprintProfilerHeatMapDisplayMode::None)
 	, HeatLevelMetricsType(EBlueprintProfilerHeatLevelMetricsType::ClassRelative)
-	, CustomEventPerformanceThreshold(CustomEventPerformanceThresholdDefaultValue)
-	, CustomAveragePerformanceThreshold(CustomAveragePerformanceThresholdDefaultValue)
-	, CustomInclusivePerformanceThreshold(CustomInclusivePerformanceThresholdDefaultValue)
-	, CustomMaxPerformanceThreshold(CustomMaxPerformanceThresholdDefaultValue)
+	// Copied out from CustomStatisticThresholdValues.
+	, CustomEventPerformanceThreshold(1.f)
+	, CustomAveragePerformanceThreshold(0.2f)
+	, CustomInclusivePerformanceThreshold(0.25f)
+	, CustomMaxPerformanceThreshold(0.5f)
 	, bThresholdsModified(false)
 {
-	FScriptPerfData::EnableBlueprintStatAverage(bAverageBlueprintStats);
 }
\ No newline at end of file
diff --git a/Engine/Source/Editor/Kismet/Private/Profiler/EventExecution.cpp b/Engine/Source/Editor/Kismet/Private/Profiler/EventExecution.cpp
--- a/Engine/Source/Editor/Kismet/Private/Profiler/EventExecution.cpp
+++ b/Engine/Source/Editor/Kismet/Private/Profiler/EventExecution.cpp
@@ -12,8 +12,7 @@
 
 // Editor profiler defines
 DECLARE_STATS_GROUP(TEXT("BlueprintProfiler"), STATGROUP_BlueprintProfiler, STATCAT_Advanced);
-DECLARE_CYCLE_STAT(TEXT("Calculate Hottest Path"), STAT_CalculateHottestPathCost, STATGROUP_BlueprintProfiler);
-DECLARE_CYCLE_STAT(TEXT("Calculate Heat Levels"), STAT_CalculateHeatLevelsCost, STATGROUP_BlueprintProfiler);
+DECLARE_CYCLE_STAT(TEXT("Heat Display Update"), STAT_HeatDisplayUpdateCost, STATGROUP_BlueprintProfiler);
 
 #define LOCTEXT_NAMESPACE "BlueprintEventUI"
 
@@ -43,7 +42,7 @@ void FScriptNodePerfData::GetValidInstanceNames(TSet<FName>& ValidInstances) con
 {
 	for (auto InstanceIter : InstanceInputPinToPerfDataMap)
 	{
-		if (InstanceIter.Key != NAME_None)
+		if (IsValidInstanceName(InstanceIter.Key))
 		{
 			ValidInstances.Add(InstanceIter.Key);
 		}
@@ -93,7 +92,7 @@ void FScriptNodePerfData::GetInstancePerfDataByTracePath(const FTracePath& Trace
 {
 	for (auto InstanceIter : InstanceInputPinToPerfDataMap)
 	{
-		if (InstanceIter.Key != NAME_None)
+		if (IsValidInstanceName(InstanceIter.Key))
 		{
 			TSharedPtr<FScriptPerfData> Result = InstanceIter.Value.FindOrAdd(TracePath);
 			if (Result.IsValid())
@@ -106,7 +105,7 @@ void FScriptNodePerfData::GetInstancePerfDataByTracePath(const FTracePath& Trace
 
 TSharedPtr<FScriptPerfData> FScriptNodePerfData::GetOrAddBlueprintPerfDataByTracePath(const FTracePath& TracePath)
 {
-	TMap<const uint32, TSharedPtr<FScriptPerfData>>& BlueprintMap = InstanceInputPinToPerfDataMap.FindOrAdd(NAME_None);
+	TMap<const uint32, TSharedPtr<FScriptPerfData>>& BlueprintMap = InstanceInputPinToPerfDataMap.FindOrAdd(SPDN_Blueprint);
 	TSharedPtr<FScriptPerfData>& Result = BlueprintMap.FindOrAdd(TracePath);
 	if (!Result.IsValid())
 	{
@@ -118,12 +117,12 @@ TSharedPtr<FScriptPerfData> FScriptNodePerfData::GetOrAddBlueprintPerfDataByTrac
 TSharedPtr<FScriptPerfData> FScriptNodePerfData::GetBlueprintPerfDataByTracePath(const FTracePath& TracePath)
 {
 	#if STRICT_PERFDATA_CREATION
-	TMap<const uint32, TSharedPtr<FScriptPerfData>>& BlueprintMap = InstanceInputPinToPerfDataMap.FindOrAdd(NAME_None);
+	TMap<const uint32, TSharedPtr<FScriptPerfData>>& BlueprintMap = InstanceInputPinToPerfDataMap.FindOrAdd(SPDN_Blueprint);
 	TSharedPtr<FScriptPerfData>* Result = BlueprintMap.Find(TracePath);
 	check (Result != nullptr);
 	return *Result;
 	#else
-	TMap<const uint32, TSharedPtr<FScriptPerfData>>& BlueprintMap = InstanceInputPinToPerfDataMap.FindOrAdd(NAME_None);
+	TMap<const uint32, TSharedPtr<FScriptPerfData>>& BlueprintMap = InstanceInputPinToPerfDataMap.FindOrAdd(SPDN_Blueprint);
 	TSharedPtr<FScriptPerfData>& Result = BlueprintMap.FindOrAdd(TracePath);
 	if (!Result.IsValid())
 	{
@@ -133,17 +132,16 @@ TSharedPtr<FScriptPerfData> FScriptNodePerfData::GetBlueprintPerfDataByTracePath
 	#endif
 }
 
-void FScriptNodePerfData::GetBlueprintPerfDataForAllTracePaths(FScriptPerfData& OutPerfData)
+void FScriptNodePerfData::UpdatePerfDataForNode()
 {
-	OutPerfData.Reset();
-
-	TMap<const uint32, TSharedPtr<FScriptPerfData>>& BlueprintMap = InstanceInputPinToPerfDataMap.FindOrAdd(NAME_None);
+	NodePerfData.Reset();
+	TMap<const uint32, TSharedPtr<FScriptPerfData>>& BlueprintMap = InstanceInputPinToPerfDataMap.FindOrAdd(SPDN_Blueprint);
 	for (auto MapIt = BlueprintMap.CreateIterator(); MapIt; ++MapIt)
 	{
 		TSharedPtr<FScriptPerfData> CurDataPtr = MapIt.Value();
 		if (CurDataPtr.IsValid())
 		{
-			OutPerfData.AddData(*CurDataPtr);
+			NodePerfData.AddData(*CurDataPtr);
 		}
 	}
 }
@@ -338,12 +336,12 @@ void FScriptExecutionNode::RefreshStats(const FTracePath& TracePath)
 							ChildDataSet.Add(ChildPerfData);
 						}
 					}
-					InstancePerfData->AccumulateDataSet(ChildDataSet);
+					InstancePerfData->AccumulateDataSet(ChildDataSet, true);
 				}
 			}
 			// Update Blueprint stats as branches
 			TSharedPtr<FScriptPerfData> BlueprintData = GetOrAddBlueprintPerfDataByTracePath(TracePath);
-			BlueprintData->InitialiseFromDataSet(BlueprintPooledStats);
+			BlueprintData->CreateBlueprintStats(BlueprintPooledStats);
 			BlueprintPooledStats.Reset(0);
 			bRefreshLinkStats = false;
 			break;
@@ -378,74 +376,61 @@ void FScriptExecutionNode::RefreshStats(const FTracePath& TracePath)
 	if (BlueprintPooledStats.Num() > 0)
 	{
 		TSharedPtr<FScriptPerfData> BlueprintData = GetOrAddBlueprintPerfDataByTracePath(TracePath);
-		BlueprintData->InitialiseFromDataSet(BlueprintPooledStats);
+		BlueprintData->CreateBlueprintStats(BlueprintPooledStats);
 	}
+	// Update the node stats
+	UpdatePerfDataForNode();
 }
 
-void FScriptExecutionNode::CalculateHeatLevelStats(TSharedPtr<FScriptHeatLevelMetrics> HeatLevelMetrics)
+void FScriptExecutionNode::UpdateHeatDisplayStats(FScriptExecutionHottestPathParams& HotPathParams)
 {
-	SCOPE_CYCLE_COUNTER(STAT_CalculateHeatLevelsCost);
-	for (auto InstancePerfDataMapIter : InstanceInputPinToPerfDataMap)
-	{
-		for (auto TracePathPerfDataMapIter : InstancePerfDataMapIter.Value)
-		{
-			TSharedPtr<FScriptPerfData> PerfData = TracePathPerfDataMapIter.Value;
-			if (PerfData.IsValid())
-			{
-				PerfData->SetHeatLevels(HeatLevelMetrics);
-			}
-		}
-	}
-
+	// Grab local perf data by tracepath.
+	FTracePath& TracePath = HotPathParams.GetTracePath();
+	TSharedPtr<FScriptPerfData> LocalPerfData = GetOrAddPerfDataByInstanceAndTracePath(HotPathParams.GetInstanceName(), TracePath);
+	// Calculate hottest path value
+	const float PathHeatLevel = HotPathParams.CalculateHeatLevel();
+	LocalPerfData->SetHottestPathHeatLevel(PathHeatLevel);
+	// Update the heat thresholds
+	LocalPerfData->SetHeatLevels(HotPathParams.GetHeatThresholds());
+	NodePerfData.SetHeatLevels(HotPathParams.GetHeatThresholds());
+	// Update the time taken so far
+	if (!IsEvent() && !HasFlags(EScriptExecutionNodeFlags::ExecPin))
+	{
+		HotPathParams.AddPathTiming(LocalPerfData->GetAverageTiming(), LocalPerfData->GetSampleCount());
+	}
+	// Update children and linked nodes
 	if (!IsPureNode())
 	{
 		for (auto ChildIter : ChildNodes)
 		{
-			ChildIter->CalculateHeatLevelStats(HeatLevelMetrics);
+			ChildIter->UpdateHeatDisplayStats(HotPathParams);
 		}
-	
-		for (auto LinkIter : LinkedNodes)
+		if (HasFlags(EScriptExecutionNodeFlags::ConditionalBranch))
 		{
-			LinkIter.Value->CalculateHeatLevelStats(HeatLevelMetrics);
-		}
-	}
-}
-
-float FScriptExecutionNode::CalculateHottestPathStats(FScriptExecutionHottestPathParams HotPathParams)
-{
-	// Grab local perf data
-	TSharedPtr<FScriptPerfData> LocalPerfData = GetOrAddPerfDataByInstanceAndTracePath(HotPathParams.InstanceName, HotPathParams.TracePath);
-	float AccumulatedTime = 0.f;
-
-	if (!IsEvent())
-	{
-		// Subtract local inclusive cost from parent inclusive cost
-		const double NodeTime = IsBranch() ? LocalPerfData->GetAverageTiming() : LocalPerfData->GetInclusiveTiming();
-		HotPathParams.TimeSoFar += NodeTime;
-		AccumulatedTime += NodeTime;
-	}
-	if (!IsPureNode())
-	{
-		// Update children
-		for (auto ChildIter : ChildNodes)
-		{
-			AccumulatedTime += ChildIter->CalculateHottestPathStats(HotPathParams);
+			for (auto LinkIter : LinkedNodes)
+			{
+				FScriptExecutionHottestPathParams LinkedHotPathParams(HotPathParams);
+				if (!LinkIter.Value->HasFlags(EScriptExecutionNodeFlags::InvalidTrace))
+				{
+					LinkedHotPathParams.GetTracePath().AddExitPin(LinkIter.Key);
+				}
+				LinkIter.Value->UpdateHeatDisplayStats(LinkedHotPathParams);
+			}
 		}
-		// Update Links
-		for (auto LinkIter : LinkedNodes)
+		else
 		{
-			FScriptExecutionHottestPathParams LinkedHotPathParams(HotPathParams);
-			if (!LinkIter.Value->HasFlags(EScriptExecutionNodeFlags::InvalidTrace))
+			FTracePath RootTracePath(TracePath);
+			for (auto LinkIter : LinkedNodes)
 			{
-				LinkedHotPathParams.TracePath.AddExitPin(LinkIter.Key);
+				TracePath = RootTracePath;
+				if (!LinkIter.Value->HasFlags(EScriptExecutionNodeFlags::InvalidTrace))
+				{
+					TracePath.AddExitPin(LinkIter.Key);
+				}
+				LinkIter.Value->UpdateHeatDisplayStats(HotPathParams);
 			}
-			AccumulatedTime += LinkIter.Value->CalculateHottestPathStats(LinkedHotPathParams);
 		}
 	}
-	const float HottestPathValue = static_cast<float>(AccumulatedTime / HotPathParams.EventTiming);
-	LocalPerfData->SetHottestPathHeatLevel(HottestPathValue);
-
-	return AccumulatedTime;
 }
 
 void FScriptExecutionNode::GetAllExecNodes(TMap<FName, TSharedPtr<FScriptExecutionNode>>& ExecNodesOut)
@@ -523,13 +508,13 @@ const EScriptPerfDataType FScriptExecutionNode::GetPerfDataType() const
 //////////////////////////////////////////////////////////////////////////
 // FScriptExecutionTunnelInstance
 
-void FScriptExecutionTunnelInstance::GetBlueprintPerfDataForAllTracePaths(FScriptPerfData& OutPerfData)
+void FScriptExecutionTunnelInstance::UpdatePerfDataForNode()
 {
-	OutPerfData.Reset();
-
+	NodePerfData.Reset();
 	for (auto EntrySite : EntrySites)
 	{
-		EntrySite.Value->GetBlueprintPerfDataForAllTracePaths(OutPerfData);
+		EntrySite.Value->UpdatePerfDataForNode();
+		NodePerfData.AddData(EntrySite.Value->GetNodePerfData());
 	}
 }
 
@@ -565,6 +550,37 @@ TSharedPtr<FScriptExecutionNode> FScriptExecutionTunnelInstance::FindBoundarySit
 	return Result;
 }
 
+void FScriptExecutionTunnelInstance::RefreshStats(const FTracePath& TracePath)
+{
+	// Process stat update
+	TArray<TSharedPtr<FScriptPerfData>> BlueprintPooledStats;
+	// Find all the entry site stats
+	for (auto EntrySite : EntrySites)
+	{
+		EntrySite.Value->GetInstancePerfDataByTracePath(TracePath, BlueprintPooledStats);
+	}
+	// Update the owning blueprint stats
+	if (BlueprintPooledStats.Num() > 0)
+	{
+		TSharedPtr<FScriptPerfData> BlueprintData = GetBlueprintPerfDataByTracePath(TracePath);
+		BlueprintData->AccumulateDataSet(BlueprintPooledStats);
+	}
+	// Update the node stats
+	UpdatePerfDataForNode();
+}
+
+void FScriptExecutionTunnelInstance::UpdateHeatDisplayStats(FScriptExecutionHottestPathParams& HotPathParams)
+{
+	// Grab local perf data by tracepath.
+	TSharedPtr<FScriptPerfData> LocalPerfData = GetOrAddPerfDataByInstanceAndTracePath(HotPathParams.GetInstanceName(), HotPathParams.GetTracePath());
+	// Update the heat thresholds
+	LocalPerfData->SetHeatLevels(HotPathParams.GetHeatThresholds());
+	// Update blueprint level node
+	TSharedPtr<FScriptPerfData> BlueprintPerfData = GetOrAddBlueprintPerfDataByTracePath(HotPathParams.GetTracePath());
+	BlueprintPerfData->SetHeatLevels(HotPathParams.GetHeatThresholds());
+	NodePerfData.SetHeatLevels(HotPathParams.GetHeatThresholds());
+}
+
 //////////////////////////////////////////////////////////////////////////
 // FScriptExecutionTunnelEntry
 
@@ -591,15 +607,15 @@ void FScriptExecutionTunnelEntry::GetLinearExecutionPath(TArray<FLinearExecPath>
 void FScriptExecutionTunnelEntry::AddTunnelTiming(const FName InstanceName, const FTracePath& TracePath, const FTracePath& InternalTracePath, const int32 ExitScriptOffset, const double Time)
 {
 	TSharedPtr<FScriptPerfData> PerfData = GetPerfDataByInstanceAndTracePath(InstanceName, TracePath);
-	PerfData->AddEventTiming(0.0);
-	PerfData->AddInclusiveTiming(Time, false);
-	PerfData->OverrideSampleCount(TunnelEntryCount);
+	const float SampleCount = PerfData->GetRawSampleCount();
+	PerfData->AddEventTiming(Time);
+	PerfData->OverrideSampleCount(SampleCount);
 	// Tick Entry points
 	for (auto EntrySite : ChildNodes)
 	{
 		PerfData = EntrySite->GetPerfDataByInstanceAndTracePath(InstanceName, TracePath);
 		PerfData->AddEventTiming(0.0);
-		PerfData->OverrideSampleCount(TunnelEntryCount);
+		PerfData->OverrideSampleCount(SampleCount);
 	}
 	// Tick exit sites
 	TSharedPtr<FScriptExecutionNode> TunnelExitNode = GetExitSite(ExitScriptOffset);
@@ -607,7 +623,7 @@ void FScriptExecutionTunnelEntry::AddTunnelTiming(const FName InstanceName, cons
 	{
 		PerfData = TunnelExitNode->GetPerfDataByInstanceAndTracePath(InstanceName, TracePath);
 		PerfData->AddEventTiming(0.0);
-		PerfData->OverrideSampleCount(TunnelEntryCount);
+		PerfData->OverrideSampleCount(SampleCount);
 		PerfData = TunnelExitNode->GetPerfDataByInstanceAndTracePath(InstanceName, InternalTracePath);
 		PerfData->AddEventTiming(0.0);
 	}
@@ -632,6 +648,8 @@ void FScriptExecutionTunnelEntry::RefreshStats(const FTracePath& TracePath)
 	TArray<TSharedPtr<FScriptPerfData>> BlueprintPooledStats;
 	// Update stats based on type
 	GetInstancePerfDataByTracePath(TracePath, BlueprintPooledStats);
+	// Update tunnel instance stat
+	TunnelInstance->RefreshStats(TracePath);
 	// Create new trace
 	FTracePath TunnelTracePath(TracePath, SharedThis<FScriptExecutionTunnelEntry>(this));
 	// Refresh child entry points
@@ -648,7 +666,34 @@ void FScriptExecutionTunnelEntry::RefreshStats(const FTracePath& TracePath)
 	if (BlueprintPooledStats.Num() > 0)
 	{
 		TSharedPtr<FScriptPerfData> BlueprintData = GetBlueprintPerfDataByTracePath(TracePath);
-		BlueprintData->InitialiseFromDataSet(BlueprintPooledStats);
+		BlueprintData->CreateBlueprintStats(BlueprintPooledStats);
+
+	}
+	// Update the node stats
+	UpdatePerfDataForNode();
+}
+
+void FScriptExecutionTunnelEntry::UpdateHeatDisplayStats(FScriptExecutionHottestPathParams& HotPathParams)
+{
+	// Grab local perf data by tracepath.
+	FTracePath& TracePath = HotPathParams.GetTracePath();
+	TSharedPtr<FScriptPerfData> LocalPerfData = GetOrAddPerfDataByInstanceAndTracePath(HotPathParams.GetInstanceName(), TracePath);
+	// Calculate hottest path value
+	const float PathHeatLevel = HotPathParams.CalculateHeatLevel();
+	LocalPerfData->SetHottestPathHeatLevel(PathHeatLevel);
+	// Update the heat thresholds
+	LocalPerfData->SetHeatLevels(HotPathParams.GetHeatThresholds());
+	NodePerfData.SetHeatLevels(HotPathParams.GetHeatThresholds());
+	// Update Instance
+	TunnelInstance->UpdateHeatDisplayStats(HotPathParams);
+	// Create new trace
+	FTracePath TunnelTracePath(TracePath, SharedThis<FScriptExecutionTunnelEntry>(this));
+	TracePath = TunnelTracePath;
+	// Update child nodes (entry points)
+	for (auto Child : ChildNodes)
+	{
+		TracePath = TunnelTracePath;
+		Child->UpdateHeatDisplayStats(HotPathParams);
 	}
 }
 
@@ -669,6 +714,15 @@ bool FScriptExecutionTunnelEntry::IsInternalBoundary(TSharedPtr<FScriptExecution
 	return bIsInternalBoundary;
 }
 
+void FScriptExecutionTunnelEntry::IncrementTunnelEntryCount(const FName InstanceName, const FTracePath& TracePath)
+{
+	// Grab local perf data by instance and tracepath.
+	TSharedPtr<FScriptPerfData> PerfData = GetPerfDataByInstanceAndTracePath(InstanceName, TracePath);
+	// Increment raw samples
+	const float RawSamples = PerfData->GetRawSampleCount() + 1.f;
+	PerfData->OverrideSampleCount(RawSamples);
+}
+
 //////////////////////////////////////////////////////////////////////////
 // FScriptExecutionTunnelExit
 
@@ -706,7 +760,39 @@ void FScriptExecutionTunnelExit::RefreshStats(const FTracePath& TracePath)
 	if (BlueprintPooledStats.Num() > 0)
 	{
 		TSharedPtr<FScriptPerfData> BlueprintData = GetOrAddBlueprintPerfDataByTracePath(TracePath);
-		BlueprintData->InitialiseFromDataSet(BlueprintPooledStats);
+		BlueprintData->CreateBlueprintStats(BlueprintPooledStats);
+	}
+	// Update the node stats
+	UpdatePerfDataForNode();
+}
+
+void FScriptExecutionTunnelExit::UpdateHeatDisplayStats(FScriptExecutionHottestPathParams& HotPathParams)
+{
+	// Cache the tracepath and exit the tunnel
+	FTracePath& TracePath = HotPathParams.GetTracePath();
+	const FTracePath InternalPath = TracePath;
+	TracePath.ExitTunnel();
+	// Grab internal and external exit pin perf data by tracepath.
+	TSharedPtr<FScriptPerfData> InternalPinPerfData = GetOrAddPerfDataByInstanceAndTracePath(HotPathParams.GetInstanceName(), InternalPath);
+	TSharedPtr<FScriptPerfData> ExternalPinPerfData = GetOrAddPerfDataByInstanceAndTracePath(HotPathParams.GetInstanceName(), TracePath);
+	// Calculate hottest path value for both sites.
+	const float PathHeatLevel = HotPathParams.CalculateHeatLevel();
+	InternalPinPerfData->SetHottestPathHeatLevel(PathHeatLevel);
+	ExternalPinPerfData->SetHottestPathHeatLevel(PathHeatLevel);
+	// Update the heat thresholds
+	InternalPinPerfData->SetHeatLevels(HotPathParams.GetHeatThresholds());
+	ExternalPinPerfData->SetHeatLevels(HotPathParams.GetHeatThresholds());
+	NodePerfData.SetHeatLevels(HotPathParams.GetHeatThresholds());
+	// Update Links
+	FTracePath BaseTracePath = TracePath;
+	for (auto LinkIter : LinkedNodes)
+	{
+		TracePath = BaseTracePath;
+		if (!LinkIter.Value->HasFlags(EScriptExecutionNodeFlags::InvalidTrace))
+		{
+			TracePath.AddExitPin(LinkIter.Key);
+		}
+		LinkIter.Value->UpdateHeatDisplayStats(HotPathParams);
 	}
 }
 
@@ -811,37 +897,39 @@ void FScriptExecutionPureChainNode::RefreshStats(const FTracePath& TracePath)
 	if (BlueprintPooledStats.Num() > 0)
 	{
 		TSharedPtr<FScriptPerfData> BlueprintData = GetBlueprintPerfDataByTracePath(TracePath);
-		BlueprintData->InitialiseFromDataSet(BlueprintPooledStats);
+		BlueprintData->CreateBlueprintStats(BlueprintPooledStats);
 	}
+	// Update the node stats
+	UpdatePerfDataForNode();
 }
 
-float FScriptExecutionPureChainNode::CalculateHottestPathStats(FScriptExecutionHottestPathParams HotPathParams)
+void FScriptExecutionPureChainNode::UpdateHeatDisplayStats(FScriptExecutionHottestPathParams& HotPathParams)
 {
-	float AccumulatedTime = 0.f;
+	// Cache the tracepath
+	FTracePath& TracePath = HotPathParams.GetTracePath();
+	const FTracePath CachedTracePath(TracePath);
 	// Update Pure Links
 	TMap<int32, TSharedPtr<FScriptExecutionNode>> AllPureNodes;
 	GetAllPureNodes(AllPureNodes);
 	for (auto PureIter : AllPureNodes)
 	{
-		FScriptExecutionHottestPathParams LinkedHotPathParams(HotPathParams);
 		if (!PureIter.Value->HasFlags(EScriptExecutionNodeFlags::InvalidTrace))
 		{
-			LinkedHotPathParams.TracePath.AddExitPin(PureIter.Key);
+			TracePath.AddExitPin(PureIter.Key);
 		}
-		AccumulatedTime += PureIter.Value->CalculateHottestPathStats(LinkedHotPathParams);
-	}
-	return AccumulatedTime;
-}
-
-void FScriptExecutionPureChainNode::CalculateHeatLevelStats(TSharedPtr<FScriptHeatLevelMetrics> HeatLevelMetrics)
-{
-	SCOPE_CYCLE_COUNTER(STAT_CalculateHeatLevelsCost);
-	TMap<int32, TSharedPtr<FScriptExecutionNode>> AllPureNodes;
-	GetAllPureNodes(AllPureNodes);
-	for (auto PureIter : AllPureNodes)
-	{
-		PureIter.Value->CalculateHeatLevelStats(HeatLevelMetrics);
-	}
+		PureIter.Value->UpdateHeatDisplayStats(HotPathParams);
+	}
+	// Restore the tracepath
+	TracePath = CachedTracePath;
+	// Grab local perf data by tracepath.
+	TSharedPtr<FScriptPerfData> LocalPerfData = GetOrAddPerfDataByInstanceAndTracePath(HotPathParams.GetInstanceName(), TracePath);
+	// Calculate hottest path value
+	const float PathHeatLevel = HotPathParams.CalculateHeatLevel();
+	LocalPerfData->SetHottestPathHeatLevel(PathHeatLevel);
+	// Update the heat thresholds
+	LocalPerfData->SetHeatLevels(HotPathParams.GetHeatThresholds());
+	// Update blueprint level node
+	NodePerfData.SetHeatLevels(HotPathParams.GetHeatThresholds());
 }
 
 //////////////////////////////////////////////////////////////////////////
@@ -978,8 +1066,6 @@ void FScriptExecutionBlueprint::RefreshStats(const FTracePath& TracePath)
 		// as an overall blueprint performance representation.
 		BlueprintEventIter->RefreshStats(TracePath);
 	}
-	// Determine if we require the hottest path's updating
-	const bool bUpdateHottestPaths = GetDefault<UBlueprintProfilerSettings>()->WireHeatMapDisplayMode != EBlueprintProfilerHeatMapDisplayMode::None;
 	// Update instance stats
 	for (auto InstanceIter : Instances)
 	{
@@ -989,13 +1075,10 @@ void FScriptExecutionBlueprint::RefreshStats(const FTracePath& TracePath)
 		for (auto InstanceEventIter : InstanceIter->GetChildNodes())
 		{
 			TSharedPtr<FScriptPerfData> InstanceEventPerfData = InstanceEventIter->GetOrAddPerfDataByInstanceAndTracePath(InstanceIter->GetName(), TracePath);
-			InstancePerfData->AddData(*InstanceEventPerfData.Get());
-			// Update the hottest path stats
-			if (bUpdateHottestPaths)
+			// Don't consider one time construction events in the instance runtime totals.
+			if (InstanceEventIter->IsRuntimeEvent())
 			{
-				SCOPE_CYCLE_COUNTER(STAT_CalculateHottestPathCost);
-				FScriptExecutionHottestPathParams HotPathParams(InstanceIter->GetName(), InstanceEventPerfData->GetAverageTiming());
-				InstanceEventIter->CalculateHottestPathStats(HotPathParams);
+				InstancePerfData->AddData(*InstanceEventPerfData.Get());
 			}
 		}
 		// Add for consolidation at the bottom.
@@ -1005,23 +1088,12 @@ void FScriptExecutionBlueprint::RefreshStats(const FTracePath& TracePath)
 	if (InstanceData.Num() > 0)
 	{
 		TSharedPtr<FScriptPerfData> BlueprintData = GetOrAddBlueprintPerfDataByTracePath(TracePath);
-		BlueprintData->Reset();
-
-		for (auto InstanceIter : InstanceData)
-		{
-			BlueprintData->AddData(*InstanceIter.Get());
-		}
-		
-		// Now that all stats have been updated/accumulated, update heat level metrics data.
+		BlueprintData->CreateBlueprintStats(InstanceData);
+		// Now that all stats have been updated/accumulated, update heat level metrics data. (this now also updates the displays in the stats)
 		UpdateHeatLevelMetrics(BlueprintData);
-
-		// Finally, calculate updated heat levels based on the current heat level metrics data.
-		CalculateHeatLevelStats(HeatLevelMetrics);
-		for (auto InstanceIter : Instances)
-		{
-			InstanceIter->CalculateHeatLevelStats(HeatLevelMetrics);
-		}
 	}
+	// Update the node stats
+	UpdatePerfDataForNode();
 }
 
 void FScriptExecutionBlueprint::GetAllExecNodes(TMap<FName, TSharedPtr<FScriptExecutionNode>>& ExecNodesOut)
@@ -1075,7 +1147,7 @@ void FScriptExecutionBlueprint::SortEvents()
 
 void FScriptExecutionBlueprint::UpdateHeatLevelMetrics(TSharedPtr<FScriptPerfData> BlueprintData)
 {
-	SCOPE_CYCLE_COUNTER(STAT_CalculateHeatLevelsCost);
+	SCOPE_CYCLE_COUNTER(STAT_HeatDisplayUpdateCost);
 	const UBlueprintProfilerSettings* ProfilerSettings = GetDefault<UBlueprintProfilerSettings>();
 	check(ProfilerSettings != nullptr);
 
@@ -1089,6 +1161,14 @@ void FScriptExecutionBlueprint::UpdateHeatLevelMetrics(TSharedPtr<FScriptPerfDat
 			const double InclusiveTiming = BlueprintData->GetInclusiveTiming();
 			const double MaxTiming = BlueprintData->GetMaxTiming();
 			const double TotalTiming = BlueprintData->GetTotalTiming();
+			// Find the blueprint relative event average threshold
+			FTracePath RootTracePath;
+			HeatLevelMetrics->EventPerformanceThreshold = 0.0f;
+			for (auto EventNode : ChildNodes)
+			{
+				TSharedPtr<FScriptPerfData> EventData = EventNode->GetBlueprintPerfDataByTracePath(RootTracePath);
+				HeatLevelMetrics->EventPerformanceThreshold += EventData->GetMaxTiming();
+			}
 
 			HeatLevelMetrics->AveragePerformanceThreshold	= AverageTiming > 0.0f ? 1.0f / AverageTiming : 0.0f;
 			HeatLevelMetrics->InclusivePerformanceThreshold = InclusiveTiming > 0.0f ? 1.0f / InclusiveTiming : 0.0f;
@@ -1097,6 +1177,7 @@ void FScriptExecutionBlueprint::UpdateHeatLevelMetrics(TSharedPtr<FScriptPerfDat
 		}
 		else
 		{
+			HeatLevelMetrics->EventPerformanceThreshold = 0.0f;
 			HeatLevelMetrics->AveragePerformanceThreshold = 0.0f;
 			HeatLevelMetrics->InclusivePerformanceThreshold = 0.0f;
 			HeatLevelMetrics->MaxPerformanceThreshold = 0.0f;
@@ -1132,6 +1213,39 @@ void FScriptExecutionBlueprint::UpdateHeatLevelMetrics(TSharedPtr<FScriptPerfDat
 	}
 		break;
 	}
+	// Update displays based on new heat thresholds.
+	UpdateHeatLevels();
+}
+
+void FScriptExecutionBlueprint::UpdateHeatLevels()
+{
+	// Find instances we need to update
+	TSet<FName> InstanceList;
+	InstanceList.Add(SPDN_Blueprint);
+	InstanceList.Add(ActiveInstanceName);
+	// Add all instance if the display settings dictate it.
+	const UBlueprintProfilerSettings* ProfilerSettings = GetDefault<UBlueprintProfilerSettings>();
+	if (ProfilerSettings->bDisplayByInstance && !ProfilerSettings->bScopeToDebugInstance)
+	{
+		for (auto InstanceNode : Instances)
+		{
+			InstanceList.Add(InstanceNode->GetName());
+		}
+	}
+	// Update heat displays
+	for (auto CurrentInstanceName : InstanceList)
+	{
+		FTracePath RootTracePath;
+		for (auto EventNode : ChildNodes)
+		{
+			// Grab the instance perf data
+			TSharedPtr<FScriptPerfData> InstancePerfData = EventNode->GetOrAddPerfDataByInstanceAndTracePath(CurrentInstanceName, RootTracePath);
+			// Create the heat update data
+			FScriptExecutionHottestPathParams HotPathParams(CurrentInstanceName, InstancePerfData->GetAverageTiming(), InstancePerfData->GetSampleCount(), HeatLevelMetrics);
+			// Walk through the event linked nodes and update heat display stats.
+			EventNode->UpdateHeatDisplayStats(HotPathParams);
+		}
+	}
 }
 
 #undef LOCTEXT_NAMESPACE
\ No newline at end of file
