Fix for User Struct Editor details panels not scrolling. Replaced Vertical Box with Splitter

diff --git a/Engine/Source/Editor/Kismet/Private/UserDefinedStructureEditor.cpp b/Engine/Source/Editor/Kismet/Private/UserDefinedStructureEditor.cpp
--- a/Engine/Source/Editor/Kismet/Private/UserDefinedStructureEditor.cpp
+++ b/Engine/Source/Editor/Kismet/Private/UserDefinedStructureEditor.cpp
@@ -362,7 +362,10 @@ TSharedRef<SDockTab> FUserDefinedStructureEditor::SpawnStructureTab(const FSpawn
 		EditedStruct = Cast<UUserDefinedStruct>(EditingObjs[ 0 ]);
 	}
 
-	auto Box = SNew(SVerticalBox);
+	TSharedRef<SSplitter> Splitter = SNew(SSplitter)
+		.Orientation(Orient_Vertical)
+		.PhysicalSplitterHandleSize(10.0f)
+		.ResizeMode(ESplitterResizeMode::Fixed);
 
 	{
 		// Create a property view
@@ -373,9 +376,8 @@ TSharedRef<SDockTab> FUserDefinedStructureEditor::SpawnStructureTab(const FSpawn
 		FOnGetDetailCustomizationInstance LayoutStructDetails = FOnGetDetailCustomizationInstance::CreateStatic(&FUserDefinedStructureDetails::MakeInstance);
 		PropertyView->RegisterInstancedCustomPropertyLayout(UUserDefinedStruct::StaticClass(), LayoutStructDetails);
 		PropertyView->SetObject(EditedStruct);
-		Box->AddSlot()
-		.VAlign(EVerticalAlignment::VAlign_Top)
-		.AutoHeight()
+		Splitter->AddSlot()
+		.Value(0.25f)
 		[
 			PropertyView.ToSharedRef()
 		];
@@ -391,10 +393,7 @@ TSharedRef<SDockTab> FUserDefinedStructureEditor::SpawnStructureTab(const FSpawn
 		auto DefaultValueWidget = DefaultValueView->GetWidget();
 		if (DefaultValueWidget.IsValid())
 		{
-			Box->AddSlot()
-			.VAlign(EVerticalAlignment::VAlign_Top)
-			.Padding(2.0f, 0.0f, 0.0f, 0.0f)
-			.AutoHeight()
+			Splitter->AddSlot()
 			[
 				DefaultValueWidget.ToSharedRef()
 			];
@@ -406,7 +405,7 @@ TSharedRef<SDockTab> FUserDefinedStructureEditor::SpawnStructureTab(const FSpawn
 		.Label( LOCTEXT("UserDefinedStructureEditor", "Structure") )
 		.TabColorScale( GetTabColorScale() )
 		[
-			Box
+			Splitter
 		];
 }
 
