Copying //UE4/Dev-Editor to //UE4/Dev-Main (Source: //UE4/Dev-Editor @ 3228984)

diff --git a/Engine/Source/Editor/FoliageEdit/Private/SFoliageEdit.h b/Engine/Source/Editor/FoliageEdit/Private/SFoliageEdit.h
--- a/Engine/Source/Editor/FoliageEdit/Private/SFoliageEdit.h
+++ b/Engine/Source/Editor/FoliageEdit/Private/SFoliageEdit.h
@@ -63,18 +63,27 @@ private:	// BRUSH SETTINGS
 	/** Retrieves the brush Radius for the brush. */
 	TOptional<float> GetRadius() const;
 
+	/** Checks if the brush size should appear. Dependant on the current tool being used. */
+	bool IsEnabled_BrushSize() const;
+
 	/** Sets the Paint Density for the brush. */
 	void SetPaintDensity(float InDensity);
 
 	/** Retrieves the Paint Density for the brush. */
 	TOptional<float> GetPaintDensity() const;
 
+	/** Checks if the paint density should appear. Dependant on the current tool being used. */
+	bool IsEnabled_PaintDensity() const;
+
 	/** Sets the Erase Density for the brush. */
 	void SetEraseDensity(float InDensity);
 
 	/** Retrieves the Erase Density for the brush. */
 	TOptional<float> GetEraseDensity() const;
 
+	/** Checks if the erase density should appear. Dependant on the current tool being used. */
+	bool IsEnabled_EraseDensity() const;
+
 	/** Retrieves the text for the filters option */
 	FText GetFilterText() const;
 
@@ -84,6 +93,18 @@ private:	// BRUSH SETTINGS
 	/** Retrieves the filter settings for painting on Landscapes. */
 	ECheckBoxState GetCheckState_Landscape() const;
 
+	/** Sets the instantiation mode settings */
+	void OnCheckStateChanged_SingleInstantiationMode(ECheckBoxState InState);
+
+	/** Retrieves the instantiation mode settings */
+	ECheckBoxState GetCheckState_SingleInstantiationMode() const;
+
+	/** Sets the instantiation mode settings */
+	void OnCheckStateChanged_SpawnInCurrentLevelMode(ECheckBoxState InState);
+
+	/** Retrieves the instantiation mode settings */
+	ECheckBoxState GetCheckState_SpawnInCurrentLevelMode() const;
+
 	/** Retrieves the tooltip text for the landscape filter */
 	FText GetTooltipText_Landscape() const;
 
@@ -135,12 +156,24 @@ private:	// BRUSH SETTINGS
 	/** Checks if the filters should appear. Dependant on the current tool being used. */
 	EVisibility GetVisibility_Filters() const;
 
+	/** Checks if the actions should appear. Dependant on the current tool being used. */
+	EVisibility GetVisibility_Actions() const;
+
 	/** Checks if the text in the empty list overlay should appear. If the list is has items but the the drag and drop override is true, it will return EVisibility::Visible. */
 	EVisibility GetVisibility_EmptyList() const;
 
 	/** @return Whether selecting instances is currently possible */
 	EVisibility GetVisibility_SelectionOptions() const;
 
+	/** Checks if the Options section should be displayed. Dependant on the current tool being used. */
+	EVisibility GetVisibility_Options() const;
+
+	/** Checks if the SingleInstantiationMode checkbox should appear. Dependant on the current tool being used. */
+	EVisibility GetVisibility_SingleInstantiationMode() const;	
+
+	/** Checks if the SpawnInCurrentLevelMode checkbox should appear. Dependant on the current tool being used. */
+	EVisibility GetVisibility_SpawnInCurrentLevelMode() const;
+
 private:	// SELECTION
 
 	/** Handler for 'Select All' command  */
diff --git a/Engine/Source/Editor/FontEditor/Private/FontFaceDetailsCustomization.cpp b/Engine/Source/Editor/FontEditor/Private/FontFaceDetailsCustomization.cpp
--- a/Engine/Source/Editor/FontEditor/Private/FontFaceDetailsCustomization.cpp
+++ b/Engine/Source/Editor/FontEditor/Private/FontFaceDetailsCustomization.cpp
@@ -144,7 +144,7 @@ void FFontFaceDetailsCustomization::OnFontPathPicked(const FString& InNewFontFil
 		{
 			FontFace->Modify();
 			FontFace->SourceFilename = InNewFontFilename;
-			FontFace->FontFaceData = FontData;
+			FontFace->FontFaceData = FFontFaceData::MakeFontFaceData(MoveTemp(FontData)); // Make a new instance as the existing one may be being used by the font cache
 		}
 	}
 
diff --git a/Engine/Source/Editor/FontEditor/Private/SCompositeFontEditor.cpp b/Engine/Source/Editor/FontEditor/Private/SCompositeFontEditor.cpp
--- a/Engine/Source/Editor/FontEditor/Private/SCompositeFontEditor.cpp
+++ b/Engine/Source/Editor/FontEditor/Private/SCompositeFontEditor.cpp
@@ -809,8 +809,12 @@ void STypefaceEntryEditor::OnTypefaceEntryFontPathPicked(const FString& InNewFon
 	{
 		UFontFace* TempFontFace = NewObject<UFontFace>();
 		TempFontFace->SourceFilename = InNewFontFilename;
-		if (FFileHelper::LoadFileToArray(TempFontFace->FontFaceData, *TempFontFace->SourceFilename))
+
+		TArray<uint8> TempFontFaceData;
+		if (FFileHelper::LoadFileToArray(TempFontFaceData, *TempFontFace->SourceFilename))
 		{
+			TempFontFace->FontFaceData->SetData(MoveTemp(TempFontFaceData));
+
 			UFontFace* NewFontFaceAsset = SaveFontFaceAsAsset(TempFontFace, *FPaths::GetBaseFilename(TempFontFace->SourceFilename));
 			if (NewFontFaceAsset)
 			{
diff --git a/Engine/Source/Editor/GraphEditor/Private/KismetPins/SGraphPinText.cpp b/Engine/Source/Editor/GraphEditor/Private/KismetPins/SGraphPinText.cpp
--- a/Engine/Source/Editor/GraphEditor/Private/KismetPins/SGraphPinText.cpp
+++ b/Engine/Source/Editor/GraphEditor/Private/KismetPins/SGraphPinText.cpp
@@ -65,6 +65,11 @@ namespace
 			GraphPinObj->GetSchema()->TrySetDefaultText(*GraphPinObj, InText);
 		}
 
+		virtual bool IsValidText(const FText& InText, FText& OutErrorMsg) const override
+		{
+			return true;
+		}
+
 #if USE_STABLE_LOCALIZATION_KEYS
 		virtual void GetStableTextId(const int32 InIndex, const ETextPropertyEditAction InEditAction, const FString& InTextSource, const FString& InProposedNamespace, const FString& InProposedKey, FString& OutStableNamespace, FString& OutStableKey) const override
 		{
diff --git a/Engine/Source/Editor/GraphEditor/Private/MaterialGraphConnectionDrawingPolicy.cpp b/Engine/Source/Editor/GraphEditor/Private/MaterialGraphConnectionDrawingPolicy.cpp
--- a/Engine/Source/Editor/GraphEditor/Private/MaterialGraphConnectionDrawingPolicy.cpp
+++ b/Engine/Source/Editor/GraphEditor/Private/MaterialGraphConnectionDrawingPolicy.cpp
@@ -3,6 +3,8 @@
 #include "MaterialGraphConnectionDrawingPolicy.h"
 #include "MaterialGraph/MaterialGraph.h"
 #include "MaterialGraph/MaterialGraphSchema.h"
+#include "MaterialGraphConnectionDrawingPolicy.h"
+#include "MaterialGraphNode_Knot.h"
 
 /////////////////////////////////////////////////////
 // FMaterialGraphConnectionDrawingPolicy
@@ -20,6 +22,87 @@ FMaterialGraphConnectionDrawingPolicy::FMaterialGraphConnectionDrawingPolicy(int
 	HoverDeemphasisDarkFraction = 0.4f;
 }
 
+
+
+bool FMaterialGraphConnectionDrawingPolicy::FindPinCenter(UEdGraphPin* Pin, FVector2D& OutCenter) const
+{
+	if (const TSharedRef<SGraphPin>* pPinWidget = PinToPinWidgetMap.Find(Pin))
+	{
+		if (FArrangedWidget* pPinEntry = PinGeometries->Find(*pPinWidget))
+		{
+			OutCenter = FGeometryHelper::CenterOf(pPinEntry->Geometry);
+			return true;
+		}
+	}
+
+	return false;
+}
+
+bool FMaterialGraphConnectionDrawingPolicy::GetAverageConnectedPosition(class UMaterialGraphNode_Knot* Knot, EEdGraphPinDirection Direction, FVector2D& OutPos) const
+{
+	FVector2D Result = FVector2D::ZeroVector;
+	int32 ResultCount = 0;
+
+	UEdGraphPin* Pin = (Direction == EGPD_Input) ? Knot->GetInputPin() : Knot->GetOutputPin();
+	for (UEdGraphPin* LinkedPin : Pin->LinkedTo)
+	{
+		FVector2D CenterPoint;
+		if (FindPinCenter(LinkedPin, /*out*/ CenterPoint))
+		{
+			Result += CenterPoint;
+			ResultCount++;
+		}
+	}
+
+	if (ResultCount > 0)
+	{
+		OutPos = Result * (1.0f / ResultCount);
+		return true;
+	}
+	else
+	{
+		return false;
+	}
+}
+
+bool FMaterialGraphConnectionDrawingPolicy::ShouldChangeTangentForKnot(UMaterialGraphNode_Knot* Knot)
+{
+	if (bool* pResult = KnotToReversedDirectionMap.Find(Knot))
+	{
+		return *pResult;
+	}
+	else
+	{
+		bool bPinReversed = false;
+
+		FVector2D AverageLeftPin;
+		FVector2D AverageRightPin;
+		FVector2D CenterPin;
+		bool bCenterValid = FindPinCenter(Knot->GetOutputPin(), /*out*/ CenterPin);
+		bool bLeftValid = GetAverageConnectedPosition(Knot, EGPD_Input, /*out*/ AverageLeftPin);
+		bool bRightValid = GetAverageConnectedPosition(Knot, EGPD_Output, /*out*/ AverageRightPin);
+
+		if (bLeftValid && bRightValid)
+		{
+			bPinReversed = AverageRightPin.X < AverageLeftPin.X;
+		}
+		else if (bCenterValid)
+		{
+			if (bLeftValid)
+			{
+				bPinReversed = CenterPin.X < AverageLeftPin.X;
+			}
+			else if (bRightValid)
+			{
+				bPinReversed = AverageRightPin.X < CenterPin.X;
+			}
+		}
+
+		KnotToReversedDirectionMap.Add(Knot, bPinReversed);
+
+		return bPinReversed;
+	}
+}
 void FMaterialGraphConnectionDrawingPolicy::DetermineWiringStyle(UEdGraphPin* OutputPin, UEdGraphPin* InputPin, /*inout*/ FConnectionParams& Params)
 {
 	Params.AssociatedPin1 = OutputPin;
@@ -33,6 +116,15 @@ void FMaterialGraphConnectionDrawingPolicy::DetermineWiringStyle(UEdGraphPin* Ou
 		{
 			Params.WireColor = MaterialGraphSchema->InactivePinColor;
 		}
+
+		UEdGraphNode* OutputNode = OutputPin->GetOwningNode();
+		if (UMaterialGraphNode_Knot* OutputKnotNode = Cast<UMaterialGraphNode_Knot>(OutputNode))
+		{
+			if (ShouldChangeTangentForKnot(OutputKnotNode))
+			{
+				Params.StartDirection = EGPD_Input;
+			}
+		}
 	}
 	if (InputPin)
 	{
@@ -40,6 +132,15 @@ void FMaterialGraphConnectionDrawingPolicy::DetermineWiringStyle(UEdGraphPin* Ou
 		{
 			Params.WireColor = MaterialGraphSchema->InactivePinColor;
 		}
+
+		UEdGraphNode* InputNode = InputPin->GetOwningNode();
+		if (UMaterialGraphNode_Knot* InputKnotNode = Cast<UMaterialGraphNode_Knot>(InputNode))
+		{
+			if (ShouldChangeTangentForKnot(InputKnotNode))
+			{
+				Params.EndDirection = EGPD_Output;
+			}
+		}
 	}
 
 	const bool bDeemphasizeUnhoveredPins = HoveredPins.Num() > 0;
diff --git a/Engine/Source/Editor/GraphEditor/Private/MaterialGraphConnectionDrawingPolicy.h b/Engine/Source/Editor/GraphEditor/Private/MaterialGraphConnectionDrawingPolicy.h
--- a/Engine/Source/Editor/GraphEditor/Private/MaterialGraphConnectionDrawingPolicy.h
+++ b/Engine/Source/Editor/GraphEditor/Private/MaterialGraphConnectionDrawingPolicy.h
@@ -20,6 +20,13 @@ protected:
 	UMaterialGraph* MaterialGraph;
 	const UMaterialGraphSchema* MaterialGraphSchema;
 
+	// Each time a knot is encountered, input geometry is compared to output geometry to see if the pins on the knot need to be reversed
+	TMap<class UMaterialGraphNode_Knot*, bool> KnotToReversedDirectionMap;
+
+	bool ShouldChangeTangentForKnot(class UMaterialGraphNode_Knot* Knot);
+	bool GetAverageConnectedPosition(class UMaterialGraphNode_Knot* Knot, EEdGraphPinDirection Direction, FVector2D& OutPos) const;
+	bool FindPinCenter(UEdGraphPin* Pin, FVector2D& OutCenter) const;
+
 public:
 	FMaterialGraphConnectionDrawingPolicy(int32 InBackLayerID, int32 InFrontLayerID, float ZoomFactor, const FSlateRect& InClippingRect, FSlateWindowElementList& InDrawElements, UEdGraph* InGraphObj);
 
diff --git a/Engine/Source/Editor/GraphEditor/Private/NodeFactory.cpp b/Engine/Source/Editor/GraphEditor/Private/NodeFactory.cpp
--- a/Engine/Source/Editor/GraphEditor/Private/NodeFactory.cpp
+++ b/Engine/Source/Editor/GraphEditor/Private/NodeFactory.cpp
@@ -103,6 +103,7 @@
 #include "MaterialNodes/SGraphNodeMaterialBase.h"
 #include "MaterialNodes/SGraphNodeMaterialComment.h"
 #include "MaterialNodes/SGraphNodeMaterialResult.h"
+#include "MaterialGraphNode_Knot.h"
 
 #include "MaterialPins/SGraphPinMaterialInput.h"
 
@@ -174,6 +175,10 @@ TSharedPtr<SGraphNode> FNodeFactory::CreateNodeWidget(UEdGraphNode* InNode)
 		{
 			return SNew(SGraphNodeMaterialResult, RootMaterialNode);
 		}
+		else if (UMaterialGraphNode_Knot* MaterialKnot = Cast<UMaterialGraphNode_Knot>(InNode))
+		{
+			return SNew(SGraphNodeKnot, MaterialKnot);
+		}
 		else if (UMaterialGraphNode* MaterialNode = Cast<UMaterialGraphNode>(InNode))
 		{
 			return SNew(SGraphNodeMaterialBase, MaterialNode);
diff --git a/Engine/Source/Editor/GraphEditor/Private/SGraphNodeKnot.cpp b/Engine/Source/Editor/GraphEditor/Private/SGraphNodeKnot.cpp
--- a/Engine/Source/Editor/GraphEditor/Private/SGraphNodeKnot.cpp
+++ b/Engine/Source/Editor/GraphEditor/Private/SGraphNodeKnot.cpp
@@ -40,7 +40,7 @@ public:
 	// End of FDragConnection interface
 
 
-	static TSharedRef<FAmbivalentDirectionDragConnection> New(UK2Node_Knot* InKnot, const TSharedRef<SGraphPanel>& InGraphPanel, const FDraggedPinTable& InStartingPins)
+	static TSharedRef<FAmbivalentDirectionDragConnection> New(UEdGraphNode* InKnot, const TSharedRef<SGraphPanel>& InGraphPanel, const FDraggedPinTable& InStartingPins)
 	{
 		TSharedRef<FAmbivalentDirectionDragConnection> Operation = MakeShareable(new FAmbivalentDirectionDragConnection(InKnot, InGraphPanel, InStartingPins));
 		Operation->Construct();
@@ -49,7 +49,7 @@ public:
 	}
 
 protected:
-	FAmbivalentDirectionDragConnection(UK2Node_Knot* InKnot, const TSharedRef<SGraphPanel>& InGraphPanel, const FDraggedPinTable& InStartingPins)
+	FAmbivalentDirectionDragConnection(UEdGraphNode* InKnot, const TSharedRef<SGraphPanel>& InGraphPanel, const FDraggedPinTable& InStartingPins)
 		: FDragConnection(InGraphPanel, InStartingPins)
 		, KnotPtr(InKnot)
 		, StartScreenPos(FVector2D::ZeroVector)
@@ -59,7 +59,7 @@ protected:
 	}
 
 protected:
-	TWeakObjectPtr<UK2Node_Knot> KnotPtr;
+	TWeakObjectPtr<UEdGraphNode> KnotPtr;
 
 	FVector2D StartScreenPos;
 	FVector2D MostRecentScreenPos;
@@ -70,10 +70,13 @@ UEdGraphPin* FAmbivalentDirectionDragConnection::GetBestPin() const
 {
 	if (bLatchedStartScreenPos)
 	{
-		if (UK2Node_Knot* Knot = KnotPtr.Get())
+		int32 InputPinIndex = -1;
+		int32 OutputPinIndex = -1;
+
+		if (KnotPtr.Get() && KnotPtr->ShouldDrawNodeAsControlPointOnly(InputPinIndex, OutputPinIndex))
 		{
 			const bool bIsRight = MostRecentScreenPos.X >= StartScreenPos.X;
-			return bIsRight ? Knot->GetOutputPin() : Knot->GetInputPin();
+			return bIsRight ? KnotPtr->GetPinAt(OutputPinIndex) : KnotPtr->GetPinAt(InputPinIndex);
 		}
 	}
 
@@ -111,8 +114,8 @@ void FAmbivalentDirectionDragConnection::ValidateGraphPinList(TArray<UEdGraphPin
 {
 	OutValidPins.Empty(DraggingPins.Num());
 
-	if (UK2Node_Knot* Knot = KnotPtr.Get())
-	{
+   if (KnotPtr.Get() != nullptr)
+   {
 		bool bUseOutput = true;
 
 		// Pick output or input based on if the drag op is currently to the left or to the right of the starting drag point
@@ -141,14 +144,20 @@ void FAmbivalentDirectionDragConnection::ValidateGraphPinList(TArray<UEdGraphPin
 			}
 		}
 
-		// Switch the effective valid pin so it makes sense for the current drag context
-		if (bUseOutput)
-		{
-			OutValidPins.Add(Knot->GetOutputPin());
-		}
-		else
+		int32 InputPinIndex = -1;
+		int32 OutputPinIndex = -1;
+
+		if (KnotPtr->ShouldDrawNodeAsControlPointOnly(InputPinIndex, OutputPinIndex))
 		{
-			OutValidPins.Add(Knot->GetInputPin());
+			// Switch the effective valid pin so it makes sense for the current drag context
+			if (bUseOutput)
+			{
+				OutValidPins.Add(KnotPtr->GetPinAt(OutputPinIndex));
+			}
+			else
+			{
+				OutValidPins.Add(KnotPtr->GetPinAt(InputPinIndex));
+			}
 		}
 	}
 	else
@@ -199,10 +208,20 @@ void SGraphPinKnot::OnDragEnter(const FGeometry& MyGeometry, const FDragDropEven
 		
 		if (ValidPins.Num() > 0)
 		{
-			if (UK2Node_Knot* Knot = Cast<UK2Node_Knot>(GraphPinObj->GetOwningNode()))
+			UEdGraphPin* PinToHoverOver = nullptr;
+			UEdGraphNode* Knot = GraphPinObj->GetOwningNode();
+			int32 InputPinIndex = -1;
+			int32 OutputPinIndex = -1;
+
+			if (Knot != nullptr && Knot->ShouldDrawNodeAsControlPointOnly(InputPinIndex, OutputPinIndex))
 			{
 				// Dragging to another pin, pick the opposite direction as a source to maximize connection chances
-				UEdGraphPin* PinToHoverOver = (ValidPins[0]->Direction == EGPD_Input) ? Knot->GetOutputPin() : Knot->GetInputPin();
+				PinToHoverOver = (ValidPins[0]->Direction == EGPD_Input) ? Knot->GetPinAt(OutputPinIndex) : Knot->GetPinAt(InputPinIndex);
+				check(PinToHoverOver);
+			}
+
+			if (PinToHoverOver != nullptr)
+			{
 				DragConnectionOp->SetHoveredPin(PinToHoverOver);
 
 				// Pins treat being dragged over the same as being hovered outside of drag and drop if they know how to respond to the drag action.
@@ -239,7 +258,7 @@ TSharedRef<FDragDropOperation> SGraphPinKnot::SpawnPinDragEvent(const TSharedRef
 		PinHandles.Add(PinWidget->GetPinObj());
 	}
 
-	TSharedRef<FAmbivalentDirectionDragConnection> Operation = FAmbivalentDirectionDragConnection::New(CastChecked<UK2Node_Knot>(GetPinObj()->GetOwningNode()), InGraphPanel, PinHandles);
+	TSharedRef<FAmbivalentDirectionDragConnection> Operation = FAmbivalentDirectionDragConnection::New(GetPinObj()->GetOwningNode(), InGraphPanel, PinHandles);
 	return Operation;
 }
 
@@ -254,14 +273,21 @@ FReply SGraphPinKnot::OnPinMouseDown(const FGeometry& SenderGeometry, const FPoi
 				// Normally break connections, but overloaded here to delete the node entirely
 				const FScopedTransaction Transaction(FGenericCommands::Get().Delete->GetDescription());
 
-				UK2Node_Knot* NodeToDelete = CastChecked<UK2Node_Knot>(GetPinObj()->GetOwningNode());
-				UBlueprint* OwnerBlueprint = NodeToDelete->GetBlueprint();
-				NodeToDelete->GetGraph()->Modify();
-
-				FBlueprintEditorUtils::RemoveNode(OwnerBlueprint, NodeToDelete, /*bDontRecompile=*/ true);
-				FBlueprintEditorUtils::MarkBlueprintAsModified(OwnerBlueprint);
+				UEdGraphNode* NodeToDelete = GetPinObj()->GetOwningNode();
+				if (NodeToDelete != nullptr)
+				{
+					UEdGraph* Graph = NodeToDelete->GetGraph();
+					if (Graph != nullptr)
+					{
+						const UEdGraphSchema* Schema = Graph->GetSchema();
+						if (Schema != nullptr && Schema->SafeDeleteNodeFromGraph(Graph, NodeToDelete))
+						{
+							return FReply::Handled();
+						}
+					}
+				}
 
-				return FReply::Handled();
+				return FReply::Unhandled();
 			}
 			else if (MouseEvent.IsControlDown())
 			{
@@ -278,11 +304,16 @@ FReply SGraphPinKnot::OnPinMouseDown(const FGeometry& SenderGeometry, const FPoi
 //////////////////////////////////////////////////////////////////////////
 // SGraphNodeKnot
 
-void SGraphNodeKnot::Construct(const FArguments& InArgs, UK2Node_Knot* InKnot)
+void SGraphNodeKnot::Construct(const FArguments& InArgs, UEdGraphNode* InKnot)
 {
+	int32 InputPinIndex = -1;
+	int32 OutputPinIndex = -1;
+	verify(InKnot->ShouldDrawNodeAsControlPointOnly(InputPinIndex, OutputPinIndex) == true &&
+		InputPinIndex >= 0 && OutputPinIndex >= 0);
 	SGraphNodeDefault::Construct(SGraphNodeDefault::FArguments().GraphNodeObj(InKnot));
 }
 
+
 void SGraphNodeKnot::UpdateGraphNode()
 {
 	InputPins.Empty();
@@ -387,6 +418,7 @@ void SGraphNodeKnot::AddPin(const TSharedRef<SGraphPin>& PinToAdd)
 	PinToAdd->SetOwner(SharedThis(this));
 
 	const UEdGraphPin* PinObj = PinToAdd->GetPinObj();
+	PinToAdd->SetShowLabel(false);
 
 	if (PinToAdd->GetDirection() == EEdGraphPinDirection::EGPD_Input)
 	{
diff --git a/Engine/Source/Editor/GraphEditor/Private/SGraphNodeKnot.h b/Engine/Source/Editor/GraphEditor/Private/SGraphNodeKnot.h
--- a/Engine/Source/Editor/GraphEditor/Private/SGraphNodeKnot.h
+++ b/Engine/Source/Editor/GraphEditor/Private/SGraphNodeKnot.h
@@ -9,13 +9,16 @@
 
 class SCommentBubble;
 
+/** The visual representation of a control point meant to adjust how connections are routed, also known as a Reroute node.
+ * The input knot node should have properly implemented ShouldDrawNodeAsControlPointOnly to return true with valid indices for its pins.
+ */
 class SGraphNodeKnot : public SGraphNodeDefault
 {
 public:
 	SLATE_BEGIN_ARGS(SGraphNodeKnot) {}
 	SLATE_END_ARGS()
 
-	void Construct(const FArguments& InArgs, class UK2Node_Knot* InKnot);
+	void Construct(const FArguments& InArgs, class UEdGraphNode* InKnot);
 
 	// SGraphNode interface
 	virtual void UpdateGraphNode() override;
diff --git a/Engine/Source/Editor/GraphEditor/Private/SGraphPin.cpp b/Engine/Source/Editor/GraphEditor/Private/SGraphPin.cpp
--- a/Engine/Source/Editor/GraphEditor/Private/SGraphPin.cpp
+++ b/Engine/Source/Editor/GraphEditor/Private/SGraphPin.cpp
@@ -562,15 +562,19 @@ void SGraphPin::OnMouseEnter( const FGeometry& MyGeometry, const FPointerEvent&
 	
 					for (UEdGraphPin* LinkedPin : Pin->LinkedTo)
 					{
-						if (UK2Node_Knot* OwningNode = Cast<UK2Node_Knot>(LinkedPin->GetOwningNodeUnchecked()))
+						int32 InputPinIndex = -1;
+						int32 OutputPinIndex = -1;
+						UEdGraphNode* InKnot = LinkedPin->GetOwningNodeUnchecked();
+						if (InKnot != nullptr && InKnot->ShouldDrawNodeAsControlPointOnly(InputPinIndex, OutputPinIndex) == true &&
+							InputPinIndex >= 0 && OutputPinIndex >= 0)
 						{
-							SetHovered(OwningNode);
+							SetHovered(InKnot);
 						}
 					}
 				}
 
 			private:
-				void SetHovered(UK2Node_Knot* KnotNode)
+				void SetHovered(UEdGraphNode* KnotNode)
 				{
 					bool bAlreadyTraversed = false;
 					IntermediateNodes.Add(KnotNode, &bAlreadyTraversed);
@@ -585,7 +589,7 @@ void SGraphPin::OnMouseEnter( const FGeometry& MyGeometry, const FPointerEvent&
 				}
 
 			private:
-				TSet<UK2Node_Knot*> IntermediateNodes;
+				TSet<UEdGraphNode*> IntermediateNodes;
 				TSet<FEdGraphPinReference>& PinSetOut;
 				TSharedPtr<SGraphPanel>   TargetPanel;
 			};
diff --git a/Engine/Source/Editor/GraphEditor/Private/SNodePanel.cpp b/Engine/Source/Editor/GraphEditor/Private/SNodePanel.cpp
--- a/Engine/Source/Editor/GraphEditor/Private/SNodePanel.cpp
+++ b/Engine/Source/Editor/GraphEditor/Private/SNodePanel.cpp
@@ -1006,7 +1006,7 @@ FReply SNodePanel::OnTouchEnded( const FGeometry& MyGeometry, const FPointerEven
 	return FReply::Unhandled();
 }
 
-float SNodePanel::GetRelativeLayoutScale(const FSlotBase& Child) const
+float SNodePanel::GetRelativeLayoutScale(const FSlotBase& Child, float LayoutScaleMultiplier) const
 {
 	return GetZoomAmount();
 }
diff --git a/Engine/Source/Editor/GraphEditor/Public/SNodePanel.h b/Engine/Source/Editor/GraphEditor/Public/SNodePanel.h
--- a/Engine/Source/Editor/GraphEditor/Public/SNodePanel.h
+++ b/Engine/Source/Editor/GraphEditor/Public/SNodePanel.h
@@ -353,7 +353,7 @@ public:
 			return FVector2D::ZeroVector;
 		}
 
-		virtual float GetRelativeLayoutScale(const FSlotBase& Child) const override
+		virtual float GetRelativeLayoutScale(const FSlotBase& Child, float LayoutScaleMultiplier) const override
 		{
 			const FNodeSlot& ThisSlot = static_cast<const FNodeSlot&>(Child);
 			if ( !ThisSlot.AllowScale.Get() )
@@ -400,7 +400,7 @@ public:
 						CurChild.GetWidget(),
 						CurChild.Offset.Get(),
 						Size,
-						GetRelativeLayoutScale(CurChild)
+						GetRelativeLayoutScale(CurChild, AllottedGeometry.Scale)
 					);
 					ArrangedChildren.AddWidget( ChildVisibility, ChildGeom );
 				}
@@ -650,7 +650,7 @@ public:
 	virtual void OnFocusLost( const FFocusEvent& InFocusEvent ) override;
 	virtual FReply OnTouchGesture( const FGeometry& MyGeometry, const FPointerEvent& GestureEvent ) override;
 	virtual FReply OnTouchEnded( const FGeometry& MyGeometry, const FPointerEvent& InTouchEvent ) override;
-	virtual float GetRelativeLayoutScale(const FSlotBase& Child) const override;
+	virtual float GetRelativeLayoutScale(const FSlotBase& Child, float LayoutScaleMultiplier) const override;
 	// End of SWidget interface
 public:
 	/**
diff --git a/Engine/Source/Editor/InputBindingEditor/Private/EditorKeyboardShortcutSettings.h b/Engine/Source/Editor/InputBindingEditor/Private/EditorKeyboardShortcutSettings.h
--- a/Engine/Source/Editor/InputBindingEditor/Private/EditorKeyboardShortcutSettings.h
+++ b/Engine/Source/Editor/InputBindingEditor/Private/EditorKeyboardShortcutSettings.h
@@ -11,4 +11,9 @@ UCLASS(config=EditorKeyBindings, meta=(DisplayName="Keyboard Shortcuts"))
 class UEditorKeyboardShortcutSettings : public UDeveloperSettings
 {
 	GENERATED_BODY()
+
+public:
+	virtual bool SupportsAutoRegistration() const override { return false;}
 };
+
+
diff --git a/Engine/Source/Editor/InputBindingEditor/Private/Widgets/SChordEditBox.cpp b/Engine/Source/Editor/InputBindingEditor/Private/Widgets/SChordEditBox.cpp
--- a/Engine/Source/Editor/InputBindingEditor/Private/Widgets/SChordEditBox.cpp
+++ b/Engine/Source/Editor/InputBindingEditor/Private/Widgets/SChordEditBox.cpp
@@ -20,7 +20,6 @@
 /* SChordEditBox interface
  *****************************************************************************/
 
-BEGIN_SLATE_FUNCTION_BUILD_OPTIMIZATION
 void SChordEditBox::Construct( const FArguments& InArgs, TSharedPtr<FUICommandInfo> InputCommand )
 {
 	BorderImageNormal = FEditorStyle::GetBrush( "EditableTextBox.Background.Normal" );
@@ -71,12 +70,9 @@ void SChordEditBox::Construct( const FArguments& InArgs, TSharedPtr<FUICommandIn
 						.IsFocusable(false)
 						.ToolTipText(LOCTEXT("ChordEditButtonRemove_ToolTip", "Remove this binding") )
 						[
-							SNew( SBox )
-							[
-								SNew( SImage )
-								.Image( FEditorStyle::GetBrush( "Symbols.X" ) )
-								.ColorAndOpacity( FLinearColor(.7f,0,0,.75f) )
-							]
+							SNew( SImage )
+							.Image( FEditorStyle::GetBrush( "Symbols.X" ) )
+							.ColorAndOpacity( FLinearColor(.7f,0,0,.75f) )
 						]
 					]
 				]
@@ -84,7 +80,6 @@ void SChordEditBox::Construct( const FArguments& InArgs, TSharedPtr<FUICommandIn
 		]
 	];
 }
-END_SLATE_FUNCTION_BUILD_OPTIMIZATION
 
 
 const FSlateBrush* SChordEditBox::GetBorderImage() const
diff --git a/Engine/Source/Editor/InputBindingEditor/Private/Widgets/SInputBindingEditorPanel.cpp b/Engine/Source/Editor/InputBindingEditor/Private/Widgets/SInputBindingEditorPanel.cpp
--- a/Engine/Source/Editor/InputBindingEditor/Private/Widgets/SInputBindingEditorPanel.cpp
+++ b/Engine/Source/Editor/InputBindingEditor/Private/Widgets/SInputBindingEditorPanel.cpp
@@ -14,8 +14,6 @@
 #define LOCTEXT_NAMESPACE "SInputBindingEditorPanel"
 
 
-/* SChordEditor interface
- *****************************************************************************/
 
 void FInputBindingEditorPanel::Initialize(class IDetailLayoutBuilder& InDetailBuilder)
 {
@@ -25,19 +23,17 @@ void FInputBindingEditorPanel::Initialize(class IDetailLayoutBuilder& InDetailBu
 
 	FBindingContext::CommandsChanged.AddSP( SharedThis( this ), &FInputBindingEditorPanel::OnCommandsChanged );
 
-	const bool bForceRefresh = false;
-	UpdateUI(bForceRefresh);
+	UpdateUI();
 }
 
 
 void FInputBindingEditorPanel::UpdateContextMasterList()
 {
-	ContextMasterList.Empty();
-
 	TArray< TSharedPtr<FBindingContext> > Contexts;
 
 	FInputBindingManager::Get().GetKnownInputContexts( Contexts );
 	
+	ContextMasterList.Empty(Contexts.Num());
 	struct FContextNameSort
 	{
 		bool operator()( const TSharedPtr<FBindingContext>& A, const TSharedPtr<FBindingContext>& B ) const
@@ -54,21 +50,26 @@ void FInputBindingEditorPanel::UpdateContextMasterList()
 
 		TSharedRef<FChordTreeItem> TreeItem( new FChordTreeItem );
 		TreeItem->BindingContext = Context;
+		
 		ContextMasterList.Add( TreeItem );
 	}
 }
 
 
-void FInputBindingEditorPanel::OnCommandsChanged()
+void FInputBindingEditorPanel::OnCommandsChanged(const FBindingContext& ContextThatChanged)
 {
 	UpdateContextMasterList();
 
-	const bool bForceRefresh = true;
-	UpdateUI(bForceRefresh);
+	if(DetailBuilder)
+	{
+		DetailBuilder->ForceRefreshDetails();
+		// Force refreshing is going to invalidate the detail builder anyway
+		DetailBuilder = nullptr;
+		FBindingContext::CommandsChanged.RemoveAll(this);
+	}
 }
 
-
-void FInputBindingEditorPanel::UpdateUI(const bool bForceRefresh)
+void FInputBindingEditorPanel::UpdateUI()
 {
 	for(TSharedPtr<FChordTreeItem>& TreeItem : ContextMasterList)
 	{
@@ -95,11 +96,12 @@ void FInputBindingEditorPanel::UpdateUI(const bool bForceRefresh)
 					.ToolTipText(CommandInfo->GetDescription())
 				]
 				+ SVerticalBox::Slot()
-				.Padding(0.0f, 3.0f, 0.0f, 0.0f)
+				.Padding(0.0f, 3.0f, 0.0f, 3.0f)
 				.AutoHeight()
 				[
 					SNew(STextBlock)
 					.Font(IDetailLayoutBuilder::GetDetailFont())
+					.ColorAndOpacity(FLinearColor::Gray)
 					.Text(CommandInfo->GetDescription())
 				]
 			];
@@ -113,7 +115,6 @@ void FInputBindingEditorPanel::UpdateUI(const bool bForceRefresh)
 			];
 		}
 	}
-
 }
 
 void FInputBindingEditorPanel::GetCommandsForContext(TSharedPtr<FChordTreeItem> InTreeItem, TArray< TSharedPtr< FUICommandInfo > >& OutChildren)
diff --git a/Engine/Source/Editor/InputBindingEditor/Private/Widgets/SInputBindingEditorPanel.h b/Engine/Source/Editor/InputBindingEditor/Private/Widgets/SInputBindingEditorPanel.h
--- a/Engine/Source/Editor/InputBindingEditor/Private/Widgets/SInputBindingEditorPanel.h
+++ b/Engine/Source/Editor/InputBindingEditor/Private/Widgets/SInputBindingEditorPanel.h
@@ -86,7 +86,7 @@ public:
 	void Initialize(class IDetailLayoutBuilder& InDetailBuilder);
 
 private:
-	void UpdateUI(const bool bForceRefresh);
+	void UpdateUI();
 
 	/**
 	 * Gets children FChordTreeItems from the passed in tree item.  Note: Only contexts have children and those children are the actual gestures.
@@ -97,7 +97,7 @@ private:
 	void UpdateContextMasterList();
 
 	/** Called when new commands are registered with the input binding manager. */
-	void OnCommandsChanged();
+	void OnCommandsChanged(const FBindingContext& ContextThatChanged);
 
 private:
 	IDetailLayoutBuilder* DetailBuilder;
diff --git a/Engine/Source/Editor/Kismet/Kismet.Build.cs b/Engine/Source/Editor/Kismet/Kismet.Build.cs
--- a/Engine/Source/Editor/Kismet/Kismet.Build.cs
+++ b/Engine/Source/Editor/Kismet/Kismet.Build.cs
@@ -14,7 +14,6 @@ public class Kismet : ModuleRules
 				"AssetTools",
                 "BlueprintRuntime",
                 "ClassViewer",
-				"EditorWidgets",
 				"Analytics",
                 "DerivedDataCache",
                 "LevelEditor",
@@ -30,6 +29,7 @@ public class Kismet : ModuleRules
 				"Slate",
 				"SlateCore",
                 "EditorStyle",
+				"EditorWidgets",
 				"Engine",
 				"Json",
 				"Merge",
@@ -61,7 +61,6 @@ public class Kismet : ModuleRules
                 "BlueprintRuntime",
                 "ClassViewer",
 				"Documentation",
-				"EditorWidgets",
 				"GameProjectGeneration",
                 "BlueprintCompilerCppBackend",
 			}
diff --git a/Engine/Source/Editor/Kismet/Private/BlueprintDetailsCustomization.cpp b/Engine/Source/Editor/Kismet/Private/BlueprintDetailsCustomization.cpp
--- a/Engine/Source/Editor/Kismet/Private/BlueprintDetailsCustomization.cpp
+++ b/Engine/Source/Editor/Kismet/Private/BlueprintDetailsCustomization.cpp
@@ -2387,6 +2387,7 @@ void FBlueprintGraphArgumentGroupLayout::SetOnRebuildChildren( FSimpleDelegate I
 
 void FBlueprintGraphArgumentGroupLayout::GenerateChildContent( IDetailChildrenBuilder& ChildrenBuilder )
 {
+	bool WasContentAdded = false;
 	if(TargetNode.IsValid())
 	{
 		TArray<TSharedPtr<FUserPinInfo>> Pins = TargetNode->UserDefinedPins;
@@ -2403,16 +2404,27 @@ void FBlueprintGraphArgumentGroupLayout::GenerateChildContent( IDetailChildrenBu
 					FName(*FString::Printf(bIsInputNode ? TEXT("InputArgument%i") : TEXT("OutputArgument%i"), i)),
 					bIsInputNode));
 				ChildrenBuilder.AddChildCustomBuilder(BlueprintArgumentLayout);
+				WasContentAdded = true;
 			}
 		}
-		else
-		{
-			// Add a null widget for this section, keeps it around for callbacks to refresh
-			ChildrenBuilder.AddChildContent(FText::GetEmpty()).ValueContent()
-				[
-					SNullWidget::NullWidget
-				];
-		}
+	}
+	if (!WasContentAdded)
+	{
+		// Add a text widget to let the user know to hit the + icon to add parameters.
+		ChildrenBuilder.AddChildContent(FText::GetEmpty()).WholeRowContent()
+			.MaxDesiredWidth(980.f)
+			[
+				SNew(SHorizontalBox)
+				+ SHorizontalBox::Slot()
+					.VAlign(VAlign_Center)
+					.Padding(0.0f, 0.0f, 4.0f, 0.0f)
+					.AutoWidth()
+					[
+						SNew(STextBlock)
+						.Text(LOCTEXT("NoArgumentsAddedForBlueprint", "Please press the + icon above to add parameters"))
+						.Font(IDetailLayoutBuilder::GetDetailFont())
+					]
+			];
 	}
 }
 
@@ -2455,7 +2467,8 @@ void FBlueprintGraphArgumentLayout::GenerateHeaderRowContent( FDetailWidgetRow&
 			SAssignNew(ArgumentNameWidget, SEditableTextBox)
 				.Text( this, &FBlueprintGraphArgumentLayout::OnGetArgNameText )
 				.OnTextChanged(this, &FBlueprintGraphArgumentLayout::OnArgNameChange)
-				.OnTextCommitted( this, &FBlueprintGraphArgumentLayout::OnArgNameTextCommitted )
+				.OnTextCommitted(this, &FBlueprintGraphArgumentLayout::OnArgNameTextCommitted)
+				.ToolTipText(this, &FBlueprintGraphArgumentLayout::OnGetArgToolTipText)
 				.Font( IDetailLayoutBuilder::GetDetailFont() )
 				.IsEnabled(!ShouldPinBeReadOnly())
 		]
@@ -2479,12 +2492,42 @@ void FBlueprintGraphArgumentLayout::GenerateHeaderRowContent( FDetailWidgetRow&
 				.IsEnabled(!ShouldPinBeReadOnly(true))
 				.Font( IDetailLayoutBuilder::GetDetailFont() )
 		]
-		+SHorizontalBox::Slot()
+		+ SHorizontalBox::Slot()
+		.AutoWidth()
+		[
+			SNew(SButton)
+			.ContentPadding(0)
+			.IsEnabled(!IsPinEditingReadOnly())
+			.OnClicked(this, &FBlueprintGraphArgumentLayout::OnArgMoveUp)
+			.ToolTipText(LOCTEXT("FunctionArgDetailsArgMoveUpTooltip", "Move this parameter up in the list."))
+			[
+				SNew(SImage)
+				.Image(FEditorStyle::GetBrush("BlueprintEditor.Details.ArgUpButton"))
+			]
+		]
+		+ SHorizontalBox::Slot()
+		.AutoWidth()
+		.Padding(2, 0)
+		[
+			SNew(SButton)
+			.ContentPadding(0)
+			.IsEnabled(!IsPinEditingReadOnly())
+			.OnClicked(this, &FBlueprintGraphArgumentLayout::OnArgMoveDown)
+			.ToolTipText(LOCTEXT("FunctionArgDetailsArgMoveDownTooltip", "Move this parameter down in the list."))
+			[
+				SNew(SImage)
+				.Image(FEditorStyle::GetBrush("BlueprintEditor.Details.ArgDownButton"))
+			]
+		]
+		+ SHorizontalBox::Slot()
 		.HAlign(HAlign_Right)
 		.VAlign(VAlign_Center)
+		.Padding(10, 0, 0, 0)
+		.AutoWidth()
 		[
-			PropertyCustomizationHelpers::MakeClearButton(FSimpleDelegate::CreateSP(this, &FBlueprintGraphArgumentLayout::OnRemoveClicked), FText(), !IsPinEditingReadOnly())
+			PropertyCustomizationHelpers::MakeClearButton(FSimpleDelegate::CreateSP(this, &FBlueprintGraphArgumentLayout::OnRemoveClicked), LOCTEXT("FunctionArgDetailsClearTooltip", "Remove this parameter."), !IsPinEditingReadOnly())
 		]
+
 	];
 }
 
@@ -2497,7 +2540,7 @@ void FBlueprintGraphArgumentLayout::GenerateChildContent( IDetailChildrenBuilder
 		[
 			SNew(STextBlock)
 				.Text( LOCTEXT( "FunctionArgDetailsDefaultValue", "Default Value" ) )
-				.ToolTipText( LOCTEXT("FunctionArgDetailsDefaultValueTooltip", "The name of the argument that will be visible to users of this graph.") )
+				.ToolTipText( LOCTEXT("FunctionArgDetailsDefaultValueParamTooltip", "The default value of the parameter.") )
 				.Font( IDetailLayoutBuilder::GetDetailFont() )
 		]
 		.ValueContent()
@@ -2526,42 +2569,7 @@ void FBlueprintGraphArgumentLayout::GenerateChildContent( IDetailChildrenBuilder
 		];
 	}
 		
-	// Read only graphs can't have their pins re-organized
-	if ( !IsPinEditingReadOnly() )
-	{
-		ChildrenBuilder.AddChildContent( LOCTEXT( "FunctionArgDetailsMoving", "Moving" ) )
-			[
-				SNew(SHorizontalBox)
-				+SHorizontalBox::Slot()
-				[
-					SNew(SSpacer)
-				]
-				+SHorizontalBox::Slot()
-					.AutoWidth()
-					.Padding(2, 0)
-					[
-						SNew(SButton)
-						.ContentPadding(0)
-						.OnClicked(this, &FBlueprintGraphArgumentLayout::OnArgMoveUp)
-						[
-							SNew(SImage)
-							.Image(FEditorStyle::GetBrush("BlueprintEditor.Details.ArgUpButton"))
-						]
-					]
-				+SHorizontalBox::Slot()
-					.AutoWidth()
-					.Padding(2, 0)
-					[
-						SNew(SButton)
-						.ContentPadding(0)
-						.OnClicked(this, &FBlueprintGraphArgumentLayout::OnArgMoveDown)
-						[
-							SNew(SImage)
-							.Image(FEditorStyle::GetBrush("BlueprintEditor.Details.ArgDownButton"))
-						]
-					]
-			];
-	}
+	
 }
 
 namespace {
@@ -2693,6 +2701,16 @@ FText FBlueprintGraphArgumentLayout::OnGetArgNameText() const
 	return FText();
 }
 
+FText FBlueprintGraphArgumentLayout::OnGetArgToolTipText() const
+{
+	if (ParamItemPtr.IsValid())
+	{
+		FText PinTypeText = UEdGraphSchema_K2::TypeToText(ParamItemPtr.Pin()->PinType);
+		return FText::Format(LOCTEXT("BlueprintArgToolTipText", "Name: {0}\nType: {1}"), FText::FromString(ParamItemPtr.Pin()->PinName), PinTypeText);
+	}
+	return FText::GetEmpty();
+}
+
 void FBlueprintGraphArgumentLayout::OnArgNameChange(const FText& InNewText)
 {
 	bool bVerified = true;
@@ -3216,41 +3234,104 @@ void FBlueprintGraphActionDetails::CustomizeDetails( IDetailLayoutBuilder& Detai
 		TSharedRef<FBlueprintGraphArgumentGroupLayout> InputArgumentGroup =
 			MakeShareable(new FBlueprintGraphArgumentGroupLayout(SharedThis(this), FunctionEntryNode));
 		InputsCategory.AddCustomBuilder(InputArgumentGroup);
-		
-		InputsCategory.AddCustomRow( LOCTEXT("FunctionNewInputArg", "New") )
-		[
-			SNew(SBox)
-			.HAlign(HAlign_Right)
+
+		TSharedRef<SHorizontalBox> InputsHeaderContentWidget = SNew(SHorizontalBox);
+		TWeakPtr<SWidget> WeakInputsHeaderWidget = InputsHeaderContentWidget;
+		InputsHeaderContentWidget->AddSlot()
+			[
+				SNew(SHorizontalBox)
+			];
+		InputsHeaderContentWidget->AddSlot()
+			.AutoWidth()
 			[
 				SNew(SButton)
-				.Text(LOCTEXT("FunctionNewInputArg", "New"))
+				.ButtonStyle(FEditorStyle::Get(), "RoundButton")
+				.ForegroundColor(FEditorStyle::GetSlateColor("DefaultForeground"))
+				.ContentPadding(FMargin(2, 0))
 				.OnClicked(this, &FBlueprintGraphActionDetails::OnAddNewInputClicked)
 				.Visibility(this, &FBlueprintGraphActionDetails::GetAddNewInputOutputVisibility)
-			]
-		];
+				.HAlign(HAlign_Right)
+				.ToolTipText(LOCTEXT("FunctionNewInputArgTooltip", "Create a new input argument"))
+				.VAlign(VAlign_Center)
+				.AddMetaData<FTagMetaData>(FTagMetaData(TEXT("FunctionNewInputArg")))
+				[
+					SNew(SHorizontalBox)
+
+					+ SHorizontalBox::Slot()
+						.AutoWidth()
+						.Padding(FMargin(0, 1))
+						[
+							SNew(SImage)
+							.Image(FEditorStyle::GetBrush("Plus"))
+						]
+
+					+ SHorizontalBox::Slot()
+						.VAlign(VAlign_Center)
+						.AutoWidth()
+						.Padding(FMargin(2, 0, 0, 0))
+						[
+							SNew(STextBlock)
+							.Font(IDetailLayoutBuilder::GetDetailFontBold())
+							.Text(LOCTEXT("FunctionNewInputArg", "New Parameter"))
+							.Visibility(this, &FBlueprintGraphActionDetails::OnGetSectionTextVisibility, WeakInputsHeaderWidget)
+							.ShadowOffset(FVector2D(1, 1))
+						]
+				]
+			];
+		InputsCategory.HeaderContent(InputsHeaderContentWidget);
 
 		if (bHasAGraph)
 		{
 			IDetailCategoryBuilder& OutputsCategory = DetailLayout.EditCategory("Outputs", LOCTEXT("FunctionDetailsOutputs", "Outputs"));
 		
-			if (FunctionResultNode)
-			{
-				TSharedRef<FBlueprintGraphArgumentGroupLayout> OutputArgumentGroup =
-					MakeShareable(new FBlueprintGraphArgumentGroupLayout(SharedThis(this), FunctionResultNode));
-				OutputsCategory.AddCustomBuilder(OutputArgumentGroup);
-			}
-
-			OutputsCategory.AddCustomRow( LOCTEXT("FunctionNewOutputArg", "New") )
-			[
-				SNew(SBox)
-				.HAlign(HAlign_Right)
+			TSharedRef<FBlueprintGraphArgumentGroupLayout> OutputArgumentGroup =
+				MakeShareable(new FBlueprintGraphArgumentGroupLayout(SharedThis(this), FunctionResultNode));
+			OutputsCategory.AddCustomBuilder(OutputArgumentGroup);
+		
+			TSharedRef<SHorizontalBox> OutputsHeaderContentWidget = SNew(SHorizontalBox);
+			TWeakPtr<SWidget> WeakOutputsHeaderWidget = OutputsHeaderContentWidget;
+			OutputsHeaderContentWidget->AddSlot()
+				[
+					SNew(SHorizontalBox)
+				];
+			OutputsHeaderContentWidget->AddSlot()
+				.AutoWidth()
 				[
 					SNew(SButton)
-					.Text(LOCTEXT("FunctionNewOutputArg", "New"))
+					.ButtonStyle(FEditorStyle::Get(), "RoundButton")
+					.ForegroundColor(FEditorStyle::GetSlateColor("DefaultForeground"))
+					.ContentPadding(FMargin(2, 0))
 					.OnClicked(this, &FBlueprintGraphActionDetails::OnAddNewOutputClicked)
 					.Visibility(this, &FBlueprintGraphActionDetails::GetAddNewInputOutputVisibility)
+					.HAlign(HAlign_Right)
+					.ToolTipText(LOCTEXT("FunctionNewOutputArgTooltip", "Create a new output argument"))
+					.VAlign(VAlign_Center)
+					.AddMetaData<FTagMetaData>(FTagMetaData(TEXT("FunctionNewOutputArg")))
+					[
+						SNew(SHorizontalBox)
+
+						+ SHorizontalBox::Slot()
+							.AutoWidth()
+							.Padding(FMargin(0, 1))
+					[
+						SNew(SImage)
+							.Image(FEditorStyle::GetBrush("Plus"))
+					]
+
+				+ SHorizontalBox::Slot()
+					.VAlign(VAlign_Center)
+					.AutoWidth()
+					.Padding(FMargin(2, 0, 0, 0))
+					[
+						SNew(STextBlock)
+						.Font(IDetailLayoutBuilder::GetDetailFontBold())
+						.Text(LOCTEXT("FunctionNewOutputArg", "New Parameter"))
+						.Visibility(this, &FBlueprintGraphActionDetails::OnGetSectionTextVisibility, WeakOutputsHeaderWidget)
+						.ShadowOffset(FVector2D(1, 1))
+					]
 				]
 			];
+			OutputsCategory.HeaderContent(OutputsHeaderContentWidget);
 		}
 	}
 	else
@@ -4513,6 +4594,21 @@ EVisibility FBlueprintGraphActionDetails::GetAddNewInputOutputVisibility() const
 	return EVisibility::Visible;
 }
 
+EVisibility FBlueprintGraphActionDetails::OnGetSectionTextVisibility(TWeakPtr<SWidget> RowWidget) const
+{
+	bool ShowText = RowWidget.Pin()->IsHovered();
+
+	// If the row is currently hovered, or a menu is being displayed for a button, keep the button expanded.
+	if (ShowText)
+	{
+		return EVisibility::SelfHitTestInvisible;
+	}
+	else
+	{
+		return EVisibility::Collapsed;
+	}
+}
+
 FReply FBlueprintGraphActionDetails::OnAddNewOutputClicked()
 {
 	FScopedTransaction Transaction( LOCTEXT( "AddOutParam", "Add Out Parameter" ) );
diff --git a/Engine/Source/Editor/Kismet/Private/BlueprintDetailsCustomization.h b/Engine/Source/Editor/Kismet/Private/BlueprintDetailsCustomization.h
--- a/Engine/Source/Editor/Kismet/Private/BlueprintDetailsCustomization.h
+++ b/Engine/Source/Editor/Kismet/Private/BlueprintDetailsCustomization.h
@@ -449,6 +449,7 @@ private:
 	FReply OnArgMoveDown();
 
 	FText OnGetArgNameText() const;
+	FText OnGetArgToolTipText() const;
 	void OnArgNameChange(const FText& InNewText);
 	void OnArgNameTextCommitted(const FText& NewText, ETextCommit::Type InTextCommit);
 	
@@ -590,6 +591,8 @@ private:
 	/** Callback to determine if the "New" button for adding input/output pins is visible */
 	EVisibility GetAddNewInputOutputVisibility() const;
 
+	EVisibility OnGetSectionTextVisibility(TWeakPtr<SWidget> RowWidget) const;
+
 	/** Called to set the replication type from the details view combo */
 	static void SetNetFlags( TWeakObjectPtr<UK2Node_EditablePinBase> FunctionEntryNode, uint32 NetFlags);
 
diff --git a/Engine/Source/Editor/Kismet/Private/BlueprintEditor.cpp b/Engine/Source/Editor/Kismet/Private/BlueprintEditor.cpp
--- a/Engine/Source/Editor/Kismet/Private/BlueprintEditor.cpp
+++ b/Engine/Source/Editor/Kismet/Private/BlueprintEditor.cpp
@@ -7369,17 +7369,17 @@ void FBlueprintEditor::NewDocument_OnClicked(ECreatedDocumentType GraphType)
 
 	if (GraphType == CGT_NewFunctionGraph)
 	{
-		NewGraph = FBlueprintEditorUtils::CreateNewGraph(GetBlueprintObj(), DocumentName, UEdGraph::StaticClass(), UEdGraphSchema_K2::StaticClass());
+		NewGraph = FBlueprintEditorUtils::CreateNewGraph(GetBlueprintObj(), DocumentName, UEdGraph::StaticClass(), GetDefaultSchemaClass());
 		FBlueprintEditorUtils::AddFunctionGraph<UClass>(GetBlueprintObj(), NewGraph, /*bIsUserCreated=*/ true, NULL);
 	}
 	else if (GraphType == CGT_NewMacroGraph)
 	{
-		NewGraph = FBlueprintEditorUtils::CreateNewGraph(GetBlueprintObj(), DocumentName, UEdGraph::StaticClass(),  UEdGraphSchema_K2::StaticClass());
+		NewGraph = FBlueprintEditorUtils::CreateNewGraph(GetBlueprintObj(), DocumentName, UEdGraph::StaticClass(), GetDefaultSchemaClass());
 		FBlueprintEditorUtils::AddMacroGraph(GetBlueprintObj(), NewGraph, /*bIsUserCreated=*/ true, NULL);
 	}
 	else if (GraphType == CGT_NewEventGraph)
 	{
-		NewGraph = FBlueprintEditorUtils::CreateNewGraph(GetBlueprintObj(), DocumentName, UEdGraph::StaticClass(), UEdGraphSchema_K2::StaticClass());
+		NewGraph = FBlueprintEditorUtils::CreateNewGraph(GetBlueprintObj(), DocumentName, UEdGraph::StaticClass(), GetDefaultSchemaClass());
 		FBlueprintEditorUtils::AddUbergraphPage(GetBlueprintObj(), NewGraph);
 	}
 	else if (GraphType == CGT_NewAnimationGraph)
@@ -7430,6 +7430,11 @@ bool FBlueprintEditor::NewDocument_IsVisibleForType(ECreatedDocumentType GraphTy
 	return false;
 }
 
+TSubclassOf<UEdGraphSchema> FBlueprintEditor::GetDefaultSchemaClass() const
+{
+	return UEdGraphSchema_K2::StaticClass();
+}
+
 bool FBlueprintEditor::AddNewDelegateIsVisible() const
 {
 	const UBlueprint* Blueprint = GetBlueprintObj();
diff --git a/Engine/Source/Editor/Kismet/Private/FindInBlueprintManager.cpp b/Engine/Source/Editor/Kismet/Private/FindInBlueprintManager.cpp
--- a/Engine/Source/Editor/Kismet/Private/FindInBlueprintManager.cpp
+++ b/Engine/Source/Editor/Kismet/Private/FindInBlueprintManager.cpp
@@ -32,6 +32,7 @@
 
 #include "Engine/SimpleConstructionScript.h"
 #include "Kismet2/BlueprintEditorUtils.h"
+#include "Kismet2/KismetEditorUtilities.h"
 #include "ARFilter.h"
 #include "AssetRegistryModule.h"
 #include "ImaginaryBlueprintData.h"
@@ -1347,6 +1348,7 @@ FFindInBlueprintSearchManager::~FFindInBlueprintSearchManager()
 		AssetRegistryModule->Get().OnAssetRemoved().RemoveAll(this);
 		AssetRegistryModule->Get().OnAssetRenamed().RemoveAll(this);
 	}
+	FKismetEditorUtilities::OnBlueprintUnloaded.RemoveAll(this);
 	FCoreUObjectDelegates::PreGarbageCollect.RemoveAll(this);
 	FCoreUObjectDelegates::PostGarbageCollect.RemoveAll(this);
 	FCoreUObjectDelegates::OnAssetLoaded.RemoveAll(this);
@@ -1379,6 +1381,8 @@ void FFindInBlueprintSearchManager::Initialize()
 		}
 	}
 
+	FKismetEditorUtilities::OnBlueprintUnloaded.AddRaw(this, &FFindInBlueprintSearchManager::OnBlueprintUnloaded);
+
 	FCoreUObjectDelegates::PreGarbageCollect.AddRaw(this, &FFindInBlueprintSearchManager::PauseFindInBlueprintSearch);
 	FCoreUObjectDelegates::PostGarbageCollect.AddRaw(this, &FFindInBlueprintSearchManager::UnpauseFindInBlueprintSearch);
 	FCoreUObjectDelegates::OnAssetLoaded.AddRaw(this, &FFindInBlueprintSearchManager::OnAssetLoaded);
@@ -1591,6 +1595,11 @@ void FFindInBlueprintSearchManager::OnAssetLoaded(UObject* InAsset)
 	}
 }
 
+void FFindInBlueprintSearchManager::OnBlueprintUnloaded(UBlueprint* InBlueprint)
+{
+	RemoveBlueprintByPath(*InBlueprint->GetPathName());
+}
+
 void FFindInBlueprintSearchManager::OnHotReload(bool bWasTriggeredAutomatically)
 {
 	CachedAssetClasses.Reset();
@@ -1721,10 +1730,10 @@ void FFindInBlueprintSearchManager::AddOrUpdateBlueprintSearchMetadata(UBlueprin
 	else
 	{
 		Index = *IndexPtr;
+		SearchArray[Index].Blueprint = InBlueprint; // Blueprint instance may change due to reloading
 	}
 
 	// Build the search data
-	SearchArray[Index].BlueprintPath = BlueprintPath;
 	if (UProperty* ParentClassProp = InBlueprint->GetClass()->FindPropertyByName(GET_MEMBER_NAME_CHECKED(UBlueprint, ParentClass)))
 	{
 		ParentClassProp->ExportTextItem(SearchArray[Index].ParentClass, ParentClassProp->ContainerPtrToValuePtr<uint8>(InBlueprint), nullptr, InBlueprint, 0);
diff --git a/Engine/Source/Editor/Kismet/Private/SKismetInspector.cpp b/Engine/Source/Editor/Kismet/Private/SKismetInspector.cpp
--- a/Engine/Source/Editor/Kismet/Private/SKismetInspector.cpp
+++ b/Engine/Source/Editor/Kismet/Private/SKismetInspector.cpp
@@ -752,11 +752,32 @@ bool SKismetInspector::IsPropertyVisible( const FPropertyAndParent& PropertyAndP
 	}
 
 	// Filter down to selected properties only if set.
-	// If the current property is selected then it is visible or if its parent is selected and the current property did not fail any of the above tests it should be visible.
-	if( SelectedObjectProperties.Find( &Property ) || ( PropertyAndParent.ParentProperty && SelectedObjectProperties.Find( PropertyAndParent.ParentProperty ) ) )
+	if ( SelectedObjectProperties.Find( &Property) ) 
 	{
+		// If the current property is selected, it is visible.
 		return true;
 	}
+	else if ( PropertyAndParent.ParentProperty )
+	{
+		const UProperty* ParentProperty = PropertyAndParent.ParentProperty;
+		const UProperty* ParentPropertyOuter = nullptr;
+		
+		if (ParentProperty)
+		{
+			ParentPropertyOuter = Cast<UProperty>(ParentProperty->GetOuter());
+		}
+
+		if ( SelectedObjectProperties.Find( ParentProperty ) )
+		{
+			// If its parent is selected, it should be visible
+			return true;
+		}
+		else if ( ParentPropertyOuter && SelectedObjectProperties.Find( ParentPropertyOuter ) )
+		{
+			// If its parent is part of a container and the container property is selected, it should be visible
+			return true;
+		}
+	}
 
 
 	return !SelectedObjectProperties.Num();
diff --git a/Engine/Source/Editor/Kismet/Private/UserDefinedEnumEditor.cpp b/Engine/Source/Editor/Kismet/Private/UserDefinedEnumEditor.cpp
--- a/Engine/Source/Editor/Kismet/Private/UserDefinedEnumEditor.cpp
+++ b/Engine/Source/Editor/Kismet/Private/UserDefinedEnumEditor.cpp
@@ -19,12 +19,115 @@
 #include "Widgets/SToolTip.h"
 #include "Widgets/Docking/SDockTab.h"
 #include "IDocumentation.h"
+#include "STextPropertyEditableTextBox.h"
 
 #define LOCTEXT_NAMESPACE "UserDefinedEnumEditor"
 
 const FName FUserDefinedEnumEditor::EnumeratorsTabId( TEXT( "UserDefinedEnum_EnumeratorEditor" ) );
 const FName FUserDefinedEnumEditor::UserDefinedEnumEditorAppIdentifier( TEXT( "UserDefinedEnumEditorApp" ) );
 
+/** Allows STextPropertyEditableTextBox to edit a user defined enum entry */
+class FEditableTextUserDefinedEnum : public IEditableTextProperty
+{
+public:
+	FEditableTextUserDefinedEnum(UUserDefinedEnum* InTargetEnum, const int32 InEnumeratorIndex)
+		: TargetEnum(InTargetEnum)
+		, EnumeratorIndex(InEnumeratorIndex)
+		, bCausedChange(false)
+	{
+	}
+
+	virtual bool IsMultiLineText() const override
+	{
+		return false;
+	}
+
+	virtual bool IsPassword() const override
+	{
+		return false;
+	}
+
+	virtual bool IsReadOnly() const override
+	{
+		return false;
+	}
+
+	virtual bool IsDefaultValue() const override
+	{
+		return false;
+	}
+
+	virtual FText GetToolTipText() const override
+	{
+		return FText::GetEmpty();
+	}
+
+	virtual int32 GetNumTexts() const override
+	{
+		return 1;
+	}
+
+	virtual FText GetText(const int32 InIndex) const override
+	{
+		check(InIndex == 0);
+		return FEnumEditorUtils::GetEnumeratorDisplayName(TargetEnum, EnumeratorIndex);
+	}
+
+	virtual void SetText(const int32 InIndex, const FText& InText) override
+	{
+		check(InIndex == 0);
+		bCausedChange = true;
+		FEnumEditorUtils::SetEnumeratorDisplayName(TargetEnum, EnumeratorIndex, InText);
+		bCausedChange = false;
+	}
+
+	virtual bool IsValidText(const FText& InText, FText& OutErrorMsg) const override
+	{
+		bool bValidName = true;
+
+		bool bUnchangedName = (InText.ToString() == FEnumEditorUtils::GetEnumeratorDisplayName(TargetEnum, EnumeratorIndex).ToString());
+		if (InText.IsEmpty())
+		{
+			OutErrorMsg = LOCTEXT("NameMissingError", "You must provide a name.");
+			bValidName = false;
+		}
+		else if (!FEnumEditorUtils::IsEnumeratorDisplayNameValid(TargetEnum, EnumeratorIndex, InText))
+		{
+			OutErrorMsg = FText::Format(LOCTEXT("NameInUseError", "'{0}' is already in use."), InText);
+			bValidName = false;
+		}
+
+		return bValidName && !bUnchangedName;
+	}
+
+#if USE_STABLE_LOCALIZATION_KEYS
+	virtual void GetStableTextId(const int32 InIndex, const ETextPropertyEditAction InEditAction, const FString& InTextSource, const FString& InProposedNamespace, const FString& InProposedKey, FString& OutStableNamespace, FString& OutStableKey) const override
+	{
+		check(InIndex == 0);
+		return StaticStableTextId(TargetEnum, InEditAction, InTextSource, InProposedNamespace, InProposedKey, OutStableNamespace, OutStableKey);
+	}
+#endif // USE_STABLE_LOCALIZATION_KEYS
+
+	virtual void RequestRefresh() override
+	{
+	}
+
+	bool CausedChange() const
+	{
+		return bCausedChange;
+	}
+
+private:
+	/** The user defined enum being edited */
+	UUserDefinedEnum* TargetEnum;
+
+	/** Index of enumerator entry */
+	int32 EnumeratorIndex;
+
+	/** Set while we are invoking a change to the user defined enum */
+	bool bCausedChange;
+};
+
 void FUserDefinedEnumEditor::RegisterTabSpawners(const TSharedRef<class FTabManager>& InTabManager)
 {
 	WorkspaceMenuCategory = InTabManager->AddLocalWorkspaceMenuCategory(LOCTEXT("WorkspaceMenu_UserDefinedEnumEditor", "User-Defined Enum Editor"));
@@ -273,10 +376,16 @@ void FUserDefinedEnumLayout::GenerateChildContent( IDetailChildrenBuilder& Child
 }
 
 
+bool FUserDefinedEnumIndexLayout::CausedChange() const
+{
+	return DisplayNameEditor.IsValid() && DisplayNameEditor->CausedChange();
+}
 
 void FUserDefinedEnumIndexLayout::GenerateHeaderRowContent( FDetailWidgetRow& NodeRow )
 {
-	const bool bIsEditable = true;
+	DisplayNameEditor = MakeShared<FEditableTextUserDefinedEnum>(TargetEnum, EnumeratorIndex);
+
+	const bool bIsEditable = !DisplayNameEditor->IsReadOnly();
 	const bool bIsMoveUpEnabled = (TargetEnum->NumEnums() != 1) && (EnumeratorIndex != 0) && bIsEditable;
 	const bool bIsMoveDownEnabled = (TargetEnum->NumEnums() != 1) && (EnumeratorIndex < TargetEnum->NumEnums() - 2) && bIsEditable;
 
@@ -290,11 +399,7 @@ void FUserDefinedEnumIndexLayout::GenerateHeaderRowContent( FDetailWidgetRow& No
 
 			+SHorizontalBox::Slot()
 			[
-				SAssignNew(EnumeratorNameWidget, SEditableTextBox)
-				.OnTextCommitted(this, &FUserDefinedEnumIndexLayout::OnEnumeratorNameCommitted)
-				.OnTextChanged(this, &FUserDefinedEnumIndexLayout::OnEnumeratorNameChanged)
-				.IsEnabled(bIsEditable)
-				.Text(this, &FUserDefinedEnumIndexLayout::GetEnumeratorName)
+				SNew(STextPropertyEditableTextBox, DisplayNameEditor.ToSharedRef())
 			]
 
 			+SHorizontalBox::Slot()
@@ -333,91 +438,21 @@ void FUserDefinedEnumIndexLayout::GenerateHeaderRowContent( FDetailWidgetRow& No
 		];
 }
 
-FText FUserDefinedEnumIndexLayout::GetEnumeratorName() const
-{
-	return FText::FromString(FEnumEditorUtils::GetEnumeratorDisplayName(TargetEnum.Get(), EnumeratorIndex));
-}
-
 void FUserDefinedEnumIndexLayout::OnEnumeratorRemove()
 {
-	FEnumEditorUtils::RemoveEnumeratorFromUserDefinedEnum(TargetEnum.Get(), EnumeratorIndex);
+	FEnumEditorUtils::RemoveEnumeratorFromUserDefinedEnum(TargetEnum, EnumeratorIndex);
 }
 
 FReply FUserDefinedEnumIndexLayout::OnMoveEnumeratorUp()
 {
-	FEnumEditorUtils::MoveEnumeratorInUserDefinedEnum(TargetEnum.Get(), EnumeratorIndex, true);
+	FEnumEditorUtils::MoveEnumeratorInUserDefinedEnum(TargetEnum, EnumeratorIndex, true);
 	return FReply::Handled();
 }
 
 FReply FUserDefinedEnumIndexLayout::OnMoveEnumeratorDown()
 {
-	FEnumEditorUtils::MoveEnumeratorInUserDefinedEnum(TargetEnum.Get(), EnumeratorIndex, false);
+	FEnumEditorUtils::MoveEnumeratorInUserDefinedEnum(TargetEnum, EnumeratorIndex, false);
 	return FReply::Handled();
 }
 
-struct FScopeTrue
-{
-	bool& bRef;
-
-	FScopeTrue(bool& bInRef) : bRef(bInRef)
-	{
-		ensure(!bRef);
-		bRef = true;
-	}
-
-	~FScopeTrue()
-	{
-		ensure(bRef);
-		bRef = false;
-	}
-};
-
-void FUserDefinedEnumIndexLayout::OnEnumeratorNameCommitted(const FText& NewText, ETextCommit::Type InTextCommit)
-{
-	const FString NewDisplayName = NewText.ToString();
-	if (FEnumEditorUtils::IsEnumeratorDisplayNameValid(TargetEnum.Get(), NewDisplayName))
-	{
-		FScopeTrue ScopeTrue(bCausedChange);
-		FEnumEditorUtils::SetEnumeratorDisplayName(TargetEnum.Get(), EnumeratorIndex, NewDisplayName);
-	}
-
-	auto EnumeratorNameWidgetSP = EnumeratorNameWidget.Pin();
-	if (EnumeratorNameWidgetSP.IsValid())
-	{
-		EnumeratorNameWidgetSP->SetError(FString());
-	}
-}
-
-void FUserDefinedEnumIndexLayout::OnEnumeratorNameChanged(const FText& NewText)
-{
-	IsValidEnumeratorDisplayName(NewText);
-}
-
-bool FUserDefinedEnumIndexLayout::IsValidEnumeratorDisplayName(const FText& NewText) const
-{
-	bool bValidName = true;
-
-	const FString NewName = NewText.ToString();
-	bool bUnchangedName = (NewName == FEnumEditorUtils::GetEnumeratorDisplayName(TargetEnum.Get(), EnumeratorIndex));
-	FText ErrorMsg;
-	if (NewText.IsEmpty())
-	{
-		ErrorMsg = LOCTEXT("NameMissingError", "You must provide a name.");
-		bValidName = false;
-	}
-	else if (!bUnchangedName && !FEnumEditorUtils::IsEnumeratorDisplayNameValid(TargetEnum.Get(), NewName))
-	{
-		ErrorMsg = FText::Format(LOCTEXT("NameInUseError", "'{0}' is already in use."), NewText);
-		bValidName = false;
-	}
-
-	auto EnumeratorNameWidgetSP = EnumeratorNameWidget.Pin();
-	if (EnumeratorNameWidgetSP.IsValid())
-	{
-		EnumeratorNameWidgetSP->SetError(ErrorMsg);
-	}
-
-	return bValidName && !bUnchangedName;
-}
-
 #undef LOCTEXT_NAMESPACE
diff --git a/Engine/Source/Editor/Kismet/Private/UserDefinedEnumEditor.h b/Engine/Source/Editor/Kismet/Private/UserDefinedEnumEditor.h
--- a/Engine/Source/Editor/Kismet/Private/UserDefinedEnumEditor.h
+++ b/Engine/Source/Editor/Kismet/Private/UserDefinedEnumEditor.h
@@ -139,10 +139,9 @@ public:
 	FUserDefinedEnumIndexLayout(UUserDefinedEnum* InTargetEnum, int32 InEnumeratorIndex)
 		: TargetEnum(InTargetEnum)
 		, EnumeratorIndex(InEnumeratorIndex)
-		, bCausedChange(false)
 	{}
 
-	bool CausedChange() const { return bCausedChange; }
+	bool CausedChange() const;
 
 private:
 	/** IDetailCustomNodeBuilder Interface*/
@@ -155,9 +154,6 @@ private:
 	virtual bool InitiallyCollapsed() const override { return false; }
 
 private:
-	/** Retrieves the enumerator's name */
-	FText GetEnumeratorName() const;
-
 	/** Moves the enumerator up in the list */
 	FReply OnMoveEnumeratorUp();
 	
@@ -167,30 +163,13 @@ private:
 	/** Deletes the enumerator */
 	void OnEnumeratorRemove();
 
-	/** Callback when the enumerator's name is committed */
-	void OnEnumeratorNameCommitted(const FText& NewText, ETextCommit::Type InTextCommit);
-	
-	/** Callback when changing the enumerator's name to verify the name */
-	void OnEnumeratorNameChanged(const FText& NewText);
-
-	/** 
-	 * Helper function to validate the enumerator's display name
-	 *
-	 * @param NewText		The name to validate
-	 *
-	 * @return				Returns true if the name is valid
-	 */
-	bool IsValidEnumeratorDisplayName(const FText& NewText) const;
-
 private:
 	/** The target node that this argument is on */
-	TWeakObjectPtr<UUserDefinedEnum> TargetEnum;
+	UUserDefinedEnum* TargetEnum;
 
 	/** Index of enumerator */
 	int32 EnumeratorIndex;
 
-	/** The enumerator's name widget, used for setting a enumerator's name */
-	TWeakPtr< SEditableTextBox > EnumeratorNameWidget;
-
-	bool bCausedChange;
+	/** The editable text interface for the display name data */
+	TSharedPtr<class FEditableTextUserDefinedEnum> DisplayNameEditor;
 };
diff --git a/Engine/Source/Editor/Kismet/Public/BlueprintEditor.h b/Engine/Source/Editor/Kismet/Public/BlueprintEditor.h
--- a/Engine/Source/Editor/Kismet/Public/BlueprintEditor.h
+++ b/Engine/Source/Editor/Kismet/Public/BlueprintEditor.h
@@ -950,6 +950,9 @@ protected:
 	/** Get the graph appearance of a specific graph, GetCurrentGraphAppearance() uses the currently focused graph. */
 	virtual FGraphAppearanceInfo GetGraphAppearance(class UEdGraph* InGraph) const;
 
+	/** Whenever new graphs need to be created it will use this schema by default. */
+	virtual TSubclassOf<UEdGraphSchema> GetDefaultSchemaClass() const;
+
 	/** Attempts to invoke the details tab if it's currently possible to. */
 	void TryInvokingDetailsTab(bool bFlash = true);
 
diff --git a/Engine/Source/Editor/Kismet/Public/FindInBlueprintManager.h b/Engine/Source/Editor/Kismet/Public/FindInBlueprintManager.h
--- a/Engine/Source/Editor/Kismet/Public/FindInBlueprintManager.h
+++ b/Engine/Source/Editor/Kismet/Public/FindInBlueprintManager.h
@@ -472,6 +472,9 @@ private:
 	/** Callback hook from the Asset Registry when an asset is loaded */
 	void OnAssetLoaded(class UObject* InAsset);
 
+	/** Callback from Kismet when a Blueprint is unloaded */
+	void OnBlueprintUnloaded(class UBlueprint* InBlueprint);
+
 	/** Callback hook from the Hot Reload manager that indicates that a module has been hot-reloaded */
 	void OnHotReload(bool bWasTriggeredAutomatically);
 
diff --git a/Engine/Source/Editor/LandscapeEditor/Private/LandscapeEdMode.cpp b/Engine/Source/Editor/LandscapeEditor/Private/LandscapeEdMode.cpp
--- a/Engine/Source/Editor/LandscapeEditor/Private/LandscapeEdMode.cpp
+++ b/Engine/Source/Editor/LandscapeEditor/Private/LandscapeEdMode.cpp
@@ -355,8 +355,9 @@ void FEdModeLandscape::Enter()
 	if (CurrentGizmoActor.IsValid())
 	{
 		CurrentGizmoActor->SetTargetLandscape(CurrentToolTarget.LandscapeInfo.Get());
+
+		CastChecked<ALandscapeGizmoActiveActor>(CurrentGizmoActor.Get())->bSnapToLandscapeGrid = UISettings->bSnapGizmo;
 	}
-	CastChecked<ALandscapeGizmoActiveActor>(CurrentGizmoActor.Get())->bSnapToLandscapeGrid = UISettings->bSnapGizmo;
 
 	int32 SquaredDataTex = ALandscapeGizmoActiveActor::DataTexSize * ALandscapeGizmoActiveActor::DataTexSize;
 
@@ -1365,6 +1366,12 @@ bool FEdModeLandscape::InputKey(FEditorViewportClient* ViewportClient, FViewport
 			LandscapeEditorControlType = ELandscapeFoliageEditorControlType::RequireCtrl;
 		}
 
+		// Special case to handle where user paint with Left Click then pressing a moving camera input, we do not want to process them so as long as the tool is active ignore other input
+		if (CurrentTool != nullptr && CurrentTool->IsToolActive())
+		{
+			return true;
+		}
+
 		if (Key == EKeys::LeftMouseButton && Event == IE_Pressed)
 		{
 			// When debugging it's possible to miss the "mouse released" event, if we get a "mouse pressed" event when we think it's already pressed then treat it as release first
diff --git a/Engine/Source/Editor/LandscapeEditor/Private/LandscapeEdModeTools.h b/Engine/Source/Editor/LandscapeEditor/Private/LandscapeEdModeTools.h
--- a/Engine/Source/Editor/LandscapeEditor/Private/LandscapeEdModeTools.h
+++ b/Engine/Source/Editor/LandscapeEditor/Private/LandscapeEdModeTools.h
@@ -1403,6 +1403,8 @@ public:
 		return true;
 	}
 
+	virtual bool IsToolActive() const override { return bToolActive;  }
+
 protected:
 	TArray<FLandscapeToolInteractorPosition> InteractorPositions;
 	FVector2D LastInteractorPosition;
diff --git a/Engine/Source/Editor/LandscapeEditor/Private/LandscapeEditorDetailWidgets.cpp b/Engine/Source/Editor/LandscapeEditor/Private/LandscapeEditorDetailWidgets.cpp
--- a/Engine/Source/Editor/LandscapeEditor/Private/LandscapeEditorDetailWidgets.cpp
+++ b/Engine/Source/Editor/LandscapeEditor/Private/LandscapeEditorDetailWidgets.cpp
@@ -177,6 +177,7 @@ void SToolSelector::BuildMultiBlockWidget(const ISlateStyle* StyleSet, const FNa
 		]
 		+ SVerticalBox::Slot()
 		.AutoHeight()
+		.HAlign(HAlign_Center)
 		[
 			SmallIconWidget
 		]
diff --git a/Engine/Source/Editor/LandscapeEditor/Public/LandscapeToolInterface.h b/Engine/Source/Editor/LandscapeEditor/Public/LandscapeToolInterface.h
--- a/Engine/Source/Editor/LandscapeEditor/Public/LandscapeToolInterface.h
+++ b/Engine/Source/Editor/LandscapeEditor/Public/LandscapeToolInterface.h
@@ -226,6 +226,7 @@ class FLandscapeTool : public FGCObject
 {
 public:
 	virtual void EnterTool() {}
+	virtual bool IsToolActive() const { return false;  }
 	virtual void ExitTool() {}
 	virtual bool BeginTool(FEditorViewportClient* ViewportClient, const FLandscapeToolTarget& Target, const FVector& InHitLocation, const UViewportInteractor* Interactor = nullptr) = 0;
 	virtual void EndTool(FEditorViewportClient* ViewportClient) = 0;
diff --git a/Engine/Source/Editor/LevelEditor/Private/SLevelEditor.cpp b/Engine/Source/Editor/LevelEditor/Private/SLevelEditor.cpp
--- a/Engine/Source/Editor/LevelEditor/Private/SLevelEditor.cpp
+++ b/Engine/Source/Editor/LevelEditor/Private/SLevelEditor.cpp
@@ -1283,18 +1283,19 @@ bool SLevelEditor::IsModeActive( FEditorModeID ModeID )
 
 void SLevelEditor::RefreshEditorModeCommands()
 {
-	FLevelEditorModesCommands::Unregister();
-	FLevelEditorModesCommands::Register();
-	
 	FLevelEditorModule& LevelEditorModule = FModuleManager::GetModuleChecked<FLevelEditorModule>( "LevelEditor" );
 
+	if(!FLevelEditorModesCommands::IsRegistered())
+	{
+		FLevelEditorModesCommands::Register();
+	}
 	const IWorkspaceMenuStructure& MenuStructure = WorkspaceMenu::GetMenuStructure();
 	TSharedPtr<FTabManager> LevelEditorTabManager = LevelEditorModule.GetLevelEditorTabManager();
 
 	// We need to remap all the actions to commands.
 	const FLevelEditorModesCommands& Commands = FLevelEditorModesCommands::Get();
 
-	int commandIndex = 0;
+	int32 CommandIndex = 0;
 	for( const FEditorModeInfo& Mode : FEditorModeRegistry::Get().GetSortedModeInfo() )
 	{
 		// If the mode isn't visible don't create a menu option for it.
@@ -1310,16 +1311,16 @@ void SLevelEditor::RefreshEditorModeCommands()
 			FInputBindingManager::Get().FindCommandInContext(Commands.GetContextName(), EditorModeCommandName);
 
 		// If a command isn't yet registered for this mode, we need to register one.
-		if ( ensure(EditorModeCommand.IsValid()) )
+		if ( EditorModeCommand.IsValid() && !LevelEditorCommands->IsActionMapped(Commands.EditorModeCommands[CommandIndex]) )
 		{
 			LevelEditorCommands->MapAction(
-				Commands.EditorModeCommands[commandIndex],
+				Commands.EditorModeCommands[CommandIndex],
 				FExecuteAction::CreateStatic( &SLevelEditor::ToggleEditorMode, Mode.ID ),
 				FCanExecuteAction(),
 				FIsActionChecked::CreateStatic( &SLevelEditor::IsModeActive, Mode.ID ));
 		}
 
-		commandIndex++;
+		CommandIndex++;
 	}
 
 	for( const auto& ToolBoxTab : ToolBoxTabs )
diff --git a/Engine/Source/Editor/LevelEditor/Private/SLevelViewportToolBar.cpp b/Engine/Source/Editor/LevelEditor/Private/SLevelViewportToolBar.cpp
--- a/Engine/Source/Editor/LevelEditor/Private/SLevelViewportToolBar.cpp
+++ b/Engine/Source/Editor/LevelEditor/Private/SLevelViewportToolBar.cpp
@@ -928,8 +928,22 @@ TSharedRef<SWidget> SLevelViewportToolBar::GenerateShowMenu() const
 		ShowMenu[SFData.Group].Add( Actions.ShowFlagCommands[ ShowFlag ] );
 	}
 
+	// Get all menu extenders for this context menu from the level editor module
+	FLevelEditorModule& LevelEditorModule = FModuleManager::GetModuleChecked<FLevelEditorModule>(TEXT("LevelEditor"));
+	TArray<FLevelEditorModule::FLevelEditorMenuExtender> MenuExtenderDelegates = LevelEditorModule.GetAllLevelViewportShowMenuExtenders();
+
+	TArray<TSharedPtr<FExtender>> Extenders;
+	for (int32 i = 0; i < MenuExtenderDelegates.Num(); ++i)
+	{
+		if (MenuExtenderDelegates[i].IsBound())
+		{
+			Extenders.Add(MenuExtenderDelegates[i].Execute(Viewport.Pin()->GetCommandList().ToSharedRef()));
+		}
+	}
+	TSharedPtr<FExtender> MenuExtender = FExtender::Combine(Extenders);
+
 	const bool bInShouldCloseWindowAfterMenuSelection = true;
-	FMenuBuilder ShowMenuBuilder( bInShouldCloseWindowAfterMenuSelection, Viewport.Pin()->GetCommandList() );
+	FMenuBuilder ShowMenuBuilder( bInShouldCloseWindowAfterMenuSelection, Viewport.Pin()->GetCommandList(), MenuExtender);
 	{
 		ShowMenuBuilder.AddMenuEntry( Actions.UseDefaultShowFlags );
 		
diff --git a/Engine/Source/Editor/LevelEditor/Public/LevelEditor.h b/Engine/Source/Editor/LevelEditor/Public/LevelEditor.h
--- a/Engine/Source/Editor/LevelEditor/Public/LevelEditor.h
+++ b/Engine/Source/Editor/LevelEditor/Public/LevelEditor.h
@@ -222,6 +222,7 @@ public:
 	virtual TArray<FLevelViewportMenuExtender_SelectedObjects>& GetAllLevelViewportDragDropContextMenuExtenders() {return LevelViewportDragDropContextMenuExtenders;}
 	virtual TArray<FLevelViewportMenuExtender_SelectedActors>& GetAllLevelViewportContextMenuExtenders() {return LevelViewportContextMenuExtenders;}
 	virtual TArray<FLevelEditorMenuExtender>& GetAllLevelViewportOptionsMenuExtenders() {return LevelViewportOptionsMenuExtenders;}
+	virtual TArray<FLevelEditorMenuExtender>& GetAllLevelViewportShowMenuExtenders() { return LevelViewportShowMenuExtenders; }
 	virtual TArray<FLevelEditorMenuExtender>& GetAllLevelEditorToolbarViewMenuExtenders() {return LevelEditorToolbarViewMenuExtenders;}
 	virtual TArray<FLevelEditorMenuExtender>& GetAllLevelEditorToolbarBuildMenuExtenders() {return LevelEditorToolbarBuildMenuExtenders;}
 	virtual TArray<FLevelEditorMenuExtender>& GetAllLevelEditorToolbarCompileMenuExtenders() {return LevelEditorToolbarCompileMenuExtenders;}
@@ -360,6 +361,7 @@ private:
 	TArray<FLevelViewportMenuExtender_SelectedObjects> LevelViewportDragDropContextMenuExtenders;
 	TArray<FLevelViewportMenuExtender_SelectedActors> LevelViewportContextMenuExtenders;
 	TArray<FLevelEditorMenuExtender> LevelViewportOptionsMenuExtenders;
+	TArray<FLevelEditorMenuExtender> LevelViewportShowMenuExtenders;
 	TArray<FLevelEditorMenuExtender> LevelEditorToolbarViewMenuExtenders;
 	TArray<FLevelEditorMenuExtender> LevelEditorToolbarBuildMenuExtenders;
 	TArray<FLevelEditorMenuExtender> LevelEditorToolbarCompileMenuExtenders;
diff --git a/Engine/Source/Editor/Localization/Private/LocalizationConfigurationScript.cpp b/Engine/Source/Editor/Localization/Private/LocalizationConfigurationScript.cpp
--- a/Engine/Source/Editor/Localization/Private/LocalizationConfigurationScript.cpp
+++ b/Engine/Source/Editor/Localization/Private/LocalizationConfigurationScript.cpp
@@ -451,6 +451,13 @@ namespace LocalizationConfigurationScript
 			ConfigSection.Add( TEXT("CommandletClass"), TEXT("InternationalizationExport") );
 
 			ConfigSection.Add( TEXT("bImportLoc"), TEXT("true") );
+
+			// Import-specific settings.
+			{
+				UEnum* LocalizedTextCollapseModeEnum = FindObjectChecked<UEnum>(ANY_PACKAGE, TEXT("ELocalizedTextCollapseMode"));
+				const FName CollapseModeName = LocalizedTextCollapseModeEnum->GetNameByValue((int64)Target->Settings.ExportSettings.CollapseMode);
+				ConfigSection.Add(TEXT("LocalizedTextCollapseMode"), CollapseModeName.ToString());
+			}
 		}
 
 		Script.Dirty = true;
@@ -567,6 +574,10 @@ namespace LocalizationConfigurationScript
 
 			// Export-specific settings.
 			{
+				UEnum* LocalizedTextCollapseModeEnum = FindObjectChecked<UEnum>(ANY_PACKAGE, TEXT("ELocalizedTextCollapseMode"));
+				const FName CollapseModeName = LocalizedTextCollapseModeEnum->GetNameByValue((int64)Target->Settings.ExportSettings.CollapseMode);
+				ConfigSection.Add(TEXT("LocalizedTextCollapseMode"), CollapseModeName.ToString());
+
 				ConfigSection.Add(TEXT("ShouldPersistCommentsOnExport"), Target->Settings.ExportSettings.ShouldPersistCommentsOnExport ? TEXT("true") : TEXT("false"));
 				ConfigSection.Add(TEXT("ShouldAddSourceLocationsAsComments"), Target->Settings.ExportSettings.ShouldAddSourceLocationsAsComments ? TEXT("true") : TEXT("false"));
 			}
diff --git a/Engine/Source/Editor/Localization/Public/LocalizationTargetTypes.h b/Engine/Source/Editor/Localization/Public/LocalizationTargetTypes.h
--- a/Engine/Source/Editor/Localization/Public/LocalizationTargetTypes.h
+++ b/Engine/Source/Editor/Localization/Public/LocalizationTargetTypes.h
@@ -22,6 +22,17 @@ enum class ELocalizationTargetConflictStatus : uint8
 	Clear,
 };
 
+UENUM()
+enum class ELocalizedTextCollapseMode : uint8
+{
+	/** Collapse texts with the same text identity (namespace + key) and source text (default 4.15+ behavior). */
+	IdenticalTextIdAndSource			UMETA(DisplayName = "Identical Text Identity (Namespace + Key) and Source Text"),
+	/** Collapse texts with the same package ID, text identity (namespace + key), and source text (4.14 behavior). */
+	IdenticalPackageIdTextIdAndSource	UMETA(DisplayName = "Identical Package ID, Text Identity (Namespace + Key) and Source Text"),
+	/** Collapse texts with the same namespace and source text (legacy pre-4.14 behavior). */
+	IdenticalNamespaceAndSource			UMETA(DisplayName = "Identical Namespace and Source Text"),
+};
+
 USTRUCT()
 struct FGatherTextSearchDirectory
 {
@@ -242,11 +253,16 @@ struct FLocalizationExportingSettings
 	GENERATED_BODY()
 
 	FLocalizationExportingSettings()
-		: ShouldPersistCommentsOnExport(false)
+		: CollapseMode(ELocalizedTextCollapseMode::IdenticalTextIdAndSource)
+		, ShouldPersistCommentsOnExport(false)
 		, ShouldAddSourceLocationsAsComments(true)
 	{
 	}
 
+	/* How should we collapse down text when exporting to PO? */
+	UPROPERTY(config, EditAnywhere, Category = "Collapsing")
+	ELocalizedTextCollapseMode CollapseMode;
+
 	/* Should user comments in existing PO files be persisted after export? Useful if using a third party service that stores editor/translator notes in the PO format's comment fields. */
 	UPROPERTY(config, EditAnywhere, Category = "Comments")
 	bool ShouldPersistCommentsOnExport;
diff --git a/Engine/Source/Editor/MaterialEditor/Private/MaterialEditor.cpp b/Engine/Source/Editor/MaterialEditor/Private/MaterialEditor.cpp
--- a/Engine/Source/Editor/MaterialEditor/Private/MaterialEditor.cpp
+++ b/Engine/Source/Editor/MaterialEditor/Private/MaterialEditor.cpp
@@ -55,12 +55,16 @@
 #include "Materials/MaterialExpressionScalarParameter.h"
 #include "Materials/MaterialExpressionStaticComponentMaskParameter.h"
 #include "Materials/MaterialExpressionTextureSampleParameter.h"
+#include "Materials/MaterialExpressionTextureObjectParameter.h"
 #include "Materials/MaterialExpressionTextureObject.h"
 #include "Materials/MaterialExpressionTextureSampleParameter2D.h"
 #include "Materials/MaterialExpressionTextureSampleParameterCube.h"
 #include "Materials/MaterialExpressionTextureSampleParameterSubUV.h"
 #include "Materials/MaterialExpressionTransformPosition.h"
 #include "Materials/MaterialExpressionVectorParameter.h"
+#include "Materials/MaterialExpressionStaticBoolParameter.h"
+#include "Materials/MaterialFunction.h"
+#include "Materials/MaterialParameterCollection.h"
 
 #include "MaterialEditorActions.h"
 #include "MaterialExpressionClasses.h"
@@ -103,6 +107,7 @@
 #include "Framework/Commands/GenericCommands.h"
 #include "CanvasTypes.h"
 #include "Engine/Selection.h"
+#include "Materials/Material.h"
 
 #define LOCTEXT_NAMESPACE "MaterialEditor"
 
@@ -702,6 +707,8 @@ void FMaterialEditor::CreateInternalWidgets()
 		LayoutCollectionParameterDetails
 		);
 
+	MaterialDetailsView->OnFinishedChangingProperties().AddSP(this, &FMaterialEditor::OnFinishedChangingProperties);
+
 	PropertyEditorModule.RegisterCustomClassLayout( UMaterial::StaticClass()->GetFName(), FOnGetDetailCustomizationInstance::CreateStatic(&FMaterialDetailCustomization::MakeInstance ) );
 
 	Palette = SNew(SMaterialPalette, SharedThis(this));
@@ -759,6 +766,22 @@ void FMaterialEditor::CreateInternalWidgets()
 }
 END_SLATE_FUNCTION_BUILD_OPTIMIZATION
 
+void FMaterialEditor::OnFinishedChangingProperties(const FPropertyChangedEvent& PropertyChangedEvent)
+{
+	if (PropertyChangedEvent.Property != nullptr)
+	{
+		UStructProperty* Property = Cast<UStructProperty>(PropertyChangedEvent.Property);
+
+		if (Property != nullptr)
+		{
+			if (Property->Struct->GetFName() == TEXT("LinearColor") || Property->Struct->GetFName() == TEXT("Color")) // if we changed a color property refresh the previews
+			{
+				RefreshExpressionPreviews();
+			}
+		}
+	}
+}
+
 FName FMaterialEditor::GetToolkitFName() const
 {
 	return FName("MaterialEditor");
@@ -2543,6 +2566,119 @@ void FMaterialEditor::OnFindInMaterial()
 	FindResults->FocusForUse();
 }
 
+UClass* FMaterialEditor::GetOnPromoteToParameterClass(UEdGraphPin* TargetPin)
+{
+	UMaterialGraphNode_Root* RootPinNode = Cast<UMaterialGraphNode_Root>(TargetPin->GetOwningNode());
+	UMaterialGraphNode* OtherPinNode = Cast<UMaterialGraphNode>(TargetPin->GetOwningNode());
+
+	if (RootPinNode != nullptr)
+	{
+		EMaterialProperty propertyId = (EMaterialProperty)FCString::Atoi(*TargetPin->PinType.PinSubCategory);
+
+		switch (propertyId)
+		{
+			case MP_Opacity:
+			case MP_Metallic:
+			case MP_Specular:
+			case MP_Roughness:
+			case MP_TessellationMultiplier:
+			case MP_CustomData0:
+			case MP_CustomData1:
+			case MP_AmbientOcclusion:
+			case MP_Refraction:
+			case MP_PixelDepthOffset:
+			case MP_OpacityMask: return UMaterialExpressionScalarParameter::StaticClass();
+
+			case MP_WorldPositionOffset:
+			case MP_WorldDisplacement:
+			case MP_EmissiveColor:
+			case MP_BaseColor:
+			case MP_SubsurfaceColor:
+			case MP_SpecularColor:
+			case MP_Normal:	return UMaterialExpressionVectorParameter::StaticClass();
+		}
+	}
+	else if (OtherPinNode)
+	{
+		const TArray<FExpressionInput*> ExpressionInputs = OtherPinNode->MaterialExpression->GetInputs();
+		FString TargetPinName = OtherPinNode->GetShortenPinName(TargetPin->PinName);
+
+		for (int32 Index = 0; Index < ExpressionInputs.Num(); ++Index)
+		{
+			FExpressionInput* Input = ExpressionInputs[Index];
+			FString InputName = OtherPinNode->MaterialExpression->GetInputName(Index);
+			InputName = OtherPinNode->GetShortenPinName(InputName);
+
+			if (InputName == TargetPinName)
+			{
+				switch (OtherPinNode->MaterialExpression->GetInputType(Index))
+				{
+					case MCT_Float1:
+					case MCT_Float: return UMaterialExpressionScalarParameter::StaticClass();
+
+					case MCT_Float2:
+					case MCT_Float3:
+					case MCT_Float4: return UMaterialExpressionVectorParameter::StaticClass();
+					
+					case MCT_StaticBool: return UMaterialExpressionStaticBoolParameter::StaticClass();
+
+					case MCT_Texture2D:
+					case MCT_TextureCube: 
+					case MCT_Texture: return UMaterialExpressionTextureObjectParameter::StaticClass();
+				}
+
+				break;
+			}
+		}
+	}
+
+	return nullptr;
+}
+
+void FMaterialEditor::OnPromoteToParameter()
+{
+	UEdGraphPin* TargetPin = GraphEditor->GetGraphPinForMenu();
+	UMaterialGraphNode_Base* PinNode = Cast<UMaterialGraphNode_Base>(TargetPin->GetOwningNode());
+
+	FMaterialGraphSchemaAction_NewNode Action;	
+	Action.MaterialExpressionClass = GetOnPromoteToParameterClass(TargetPin);
+
+	if (Action.MaterialExpressionClass != nullptr)
+	{
+		check(PinNode);
+		UEdGraph* GraphObj = PinNode->GetGraph();
+		check(GraphObj);
+
+		const FScopedTransaction Transaction(LOCTEXT("PromoteToParameter", "Promote To Parameter"));
+		GraphObj->Modify();
+
+		// Set position of new node to be close to node we clicked on
+		FVector2D NewNodePos;
+		NewNodePos.X = PinNode->NodePosX - 100;
+		NewNodePos.Y = PinNode->NodePosY;
+
+		UMaterialGraphNode* MaterialNode = Cast<UMaterialGraphNode>(Action.PerformAction(GraphObj, TargetPin, NewNodePos));
+
+		if (MaterialNode->MaterialExpression->HasAParameterName())
+		{
+			MaterialNode->MaterialExpression->SetParameterName(FName(*TargetPin->PinName));
+			MaterialNode->MaterialExpression->ValidateParameterName();
+		}
+	}
+}
+
+bool FMaterialEditor::OnCanPromoteToParameter()
+{
+	UEdGraphPin* TargetPin = GraphEditor->GetGraphPinForMenu();
+
+	if (TargetPin->LinkedTo.Num() == 0)
+	{
+		return GetOnPromoteToParameterClass(TargetPin) != nullptr;
+	}
+
+	return false;
+}
+
 FString FMaterialEditor::GetDocLinkForSelectedNode()
 {
 	FString DocumentationLink;
@@ -3032,12 +3168,9 @@ UMaterialExpression* FMaterialEditor::CreateNewMaterialExpression(UClass* NewExp
 
 		NewExpression->UpdateParameterGuid(true, true);
 
-		UMaterialExpressionTextureSampleParameter* TextureParameterExpression = Cast<UMaterialExpressionTextureSampleParameter>( NewExpression );
-		if( (TextureParameterExpression != nullptr) && TextureParameterExpression->CanRenameNode() )
+		if (NewExpression->HasAParameterName())
 		{
-			// Change the parameter's name on creation to mirror the object's name; this avoids issues of having colliding parameter
-			// names and having the name left as "None"
-			TextureParameterExpression->ParameterName = TextureParameterExpression->GetFName();
+			NewExpression->ValidateParameterName();
 		}
 
 		UMaterialExpressionComponentMask* ComponentMaskExpression = Cast<UMaterialExpressionComponentMask>( NewExpression );
@@ -3428,6 +3561,13 @@ void FMaterialEditor::PasteNodesHere(const FVector2D& Location)
 			UMaterialExpression* NewExpression = GraphNode->MaterialExpression;
 			NewExpression->Material = Material;
 			NewExpression->Function = MaterialFunction;
+
+			// Make sure the param name is valid after the paste
+			if (NewExpression->HasAParameterName())
+			{
+				NewExpression->ValidateParameterName();
+			}
+
 			Material->Expressions.Add(NewExpression);
 
 			// There can be only one default mesh paint texture.
@@ -3956,6 +4096,12 @@ TSharedRef<SGraphEditor> FMaterialEditor::CreateGraphEditorWidget()
 			FCanExecuteAction::CreateSP(this, &FMaterialEditor::CanGoToDocumentation)
 			);
 
+
+		GraphEditorCommands->MapAction(FMaterialEditorCommands::Get().PromoteToParameter,
+			FExecuteAction::CreateSP(this, &FMaterialEditor::OnPromoteToParameter),
+			FCanExecuteAction::CreateSP(this, &FMaterialEditor::OnCanPromoteToParameter)
+			);
+
 	}
 
 	FGraphAppearanceInfo AppearanceInfo;
diff --git a/Engine/Source/Editor/MaterialEditor/Private/MaterialEditor.h b/Engine/Source/Editor/MaterialEditor/Private/MaterialEditor.h
--- a/Engine/Source/Editor/MaterialEditor/Private/MaterialEditor.h
+++ b/Engine/Source/Editor/MaterialEditor/Private/MaterialEditor.h
@@ -180,6 +180,8 @@ struct FMaterialExpression
 {
 	FString Name;
 	UClass* MaterialClass;
+	FString CreationDescription;
+	FString CreationName;
 
 	friend bool operator==(const FMaterialExpression& A,const FMaterialExpression& B)
 	{
@@ -352,6 +354,7 @@ public:
 	virtual void PasteNodesHere(const FVector2D& Location) override;
 	virtual int32 GetNumberOfSelectedNodes() const override;
 	virtual FMaterialRenderProxy* GetExpressionPreview(UMaterialExpression* InExpression) override;
+	virtual void DeleteNodes(const TArray<class UEdGraphNode*>& NodesToDelete) override;
 
 
 	void UpdateStatsMaterials();
@@ -447,8 +450,6 @@ protected:
 	/** Whether we can select every node */
 	bool CanSelectAllNodes() const;
 
-	/** Delete an array of Material Graph Nodes and their corresponding expressions/comments */
-	void DeleteNodes(const TArray<class UEdGraphNode*>& NodesToDelete);
 	/** Whether we are able to delete the currently selected nodes */
 	bool CanDeleteNodes() const;
 	/** Delete only the currently selected nodes that can be duplicated */
@@ -575,6 +576,15 @@ private:
 	/** Bring up the search tab */
 	void OnFindInMaterial();
 
+	/** Will promote selected pin to a parameter of the pin type */
+	void OnPromoteToParameter();
+
+	/** Used to know if we can promote selected pin to a parameter of the pin type */
+	bool OnCanPromoteToParameter();
+
+	/** Will  return the UClass to create from the Pin Type */
+	UClass* GetOnPromoteToParameterClass(UEdGraphPin* TargetPin);
+
 	/** Open documentation for the selected node class */
 	void OnGoToDocumentation();
 	/** Can we open documentation for the selected node */
@@ -671,6 +681,8 @@ private:
 	TSharedRef<SDockTab> SpawnTab_Palette(const FSpawnTabArgs& Args);
 	TSharedRef<SDockTab> SpawnTab_Stats(const FSpawnTabArgs& Args);
 	TSharedRef<SDockTab> SpawnTab_Find(const FSpawnTabArgs& Args);
+
+	void OnFinishedChangingProperties(const FPropertyChangedEvent& PropertyChangedEvent);
 private:
 	/** List of open tool panels; used to ensure only one exists at any one time */
 	TMap< FName, TWeakPtr<class SDockableTab> > SpawnedToolPanels;
diff --git a/Engine/Source/Editor/MaterialEditor/Private/MaterialEditorActions.cpp b/Engine/Source/Editor/MaterialEditor/Private/MaterialEditorActions.cpp
--- a/Engine/Source/Editor/MaterialEditor/Private/MaterialEditorActions.cpp
+++ b/Engine/Source/Editor/MaterialEditor/Private/MaterialEditorActions.cpp
@@ -25,7 +25,7 @@ void FMaterialEditorCommands::RegisterCommands()
 
 	UI_COMMAND( CameraHome, "Home", "Goes home on the canvas.", EUserInterfaceActionType::Button, FInputChord() );
 	UI_COMMAND( CleanUnusedExpressions, "Clean Up", "Cleans up any unused Expressions.", EUserInterfaceActionType::Button, FInputChord() );
-	UI_COMMAND(ShowHideConnectors, "Connectors", "Show or Hide Unused Connectors", EUserInterfaceActionType::ToggleButton, FInputChord());
+	UI_COMMAND( ShowHideConnectors, "Connectors", "Show or Hide Unused Connectors", EUserInterfaceActionType::ToggleButton, FInputChord());
 	UI_COMMAND( ToggleLivePreview, "Live Preview", "Toggles real time update of the preview material.", EUserInterfaceActionType::ToggleButton, FInputChord());
 	UI_COMMAND( ToggleRealtimeExpressions, "Live Nodes", "Toggles real time update of the graph canvas.", EUserInterfaceActionType::ToggleButton, FInputChord() );
 	UI_COMMAND( AlwaysRefreshAllPreviews, "Live Update", "All nodes are previewed live.", EUserInterfaceActionType::ToggleButton, FInputChord() );
@@ -53,6 +53,7 @@ void FMaterialEditorCommands::RegisterCommands()
 	UI_COMMAND( ForceRefreshPreviews, "Force Refresh Previews", "Forces a refresh of all previews", EUserInterfaceActionType::Button, FInputChord(EKeys::SpaceBar) );
 	UI_COMMAND( CreateComponentMaskNode, "Create ComponentMask Node", "Creates a ComponentMask node at the current cursor position.", EUserInterfaceActionType::Button, FInputChord(EModifierKey::Shift, EKeys::C));
 	UI_COMMAND( FindInMaterial, "Search", "Finds expressions and comments in the current Material", EUserInterfaceActionType::Button, FInputChord(EModifierKey::Control, EKeys::F));
+	UI_COMMAND( PromoteToParameter, "Promote to Parameter", "Promote selected Pin to parameter of pin type", EUserInterfaceActionType::Button, FInputChord());	
 }
 
 //////////////////////////////////////////////////////////////////////////
diff --git a/Engine/Source/Editor/MaterialEditor/Private/MaterialEditorUtilities.cpp b/Engine/Source/Editor/MaterialEditor/Private/MaterialEditorUtilities.cpp
--- a/Engine/Source/Editor/MaterialEditor/Private/MaterialEditorUtilities.cpp
+++ b/Engine/Source/Editor/MaterialEditor/Private/MaterialEditorUtilities.cpp
@@ -81,6 +81,16 @@ void FMaterialEditorUtilities::DeleteSelectedNodes(const class UEdGraph* Graph)
 	}
 }
 
+
+void FMaterialEditorUtilities::DeleteNodes(const class UEdGraph* Graph, const TArray<UEdGraphNode*>& NodesToDelete)
+{
+	TSharedPtr<class IMaterialEditor> MaterialEditor = GetIMaterialEditorForObject(Graph);
+	if (MaterialEditor.IsValid())
+	{
+		MaterialEditor->DeleteNodes(NodesToDelete);
+	}
+}
+
 FText FMaterialEditorUtilities::GetOriginalObjectName(const class UEdGraph* Graph)
 {
 	TSharedPtr<class IMaterialEditor> MaterialEditor = GetIMaterialEditorForObject(Graph);
@@ -579,13 +589,20 @@ void FMaterialEditorUtilities::AddMaterialExpressionCategory(FGraphActionMenuBui
 		{
 			if (!ActionMenuBuilder.FromPin || HasCompatibleConnection(MaterialExpression.MaterialClass, FromPinType, ActionMenuBuilder.FromPin->Direction, bMaterialFunction))
 			{
+				FString CreationName = MaterialExpression.Name;
 				FFormatNamedArguments Arguments;
 				Arguments.Add(TEXT("Name"), FText::FromString( *MaterialExpression.Name ));
-				const FText ToolTip = FText::Format( LOCTEXT( "NewMaterialExpressionTooltip", "Adds a {Name} node here" ), Arguments );
+				FString ToolTip = FText::Format(LOCTEXT("NewMaterialExpressionTooltip", "Adds a {Name} node here"), Arguments).ToString();
+				if (MaterialExpression.CreationDescription.Len() != 0)
+					ToolTip = MaterialExpression.CreationDescription;
+				if (MaterialExpression.CreationName.Len() != 0)
+					CreationName = MaterialExpression.CreationName;
+
+
 				TSharedPtr<FMaterialGraphSchemaAction_NewNode> NewNodeAction(new FMaterialGraphSchemaAction_NewNode(
 					CategoryName,
-					FText::FromString(MaterialExpression.Name),
-					ToolTip.ToString(), 0, CastChecked<UMaterialExpression>(MaterialExpression.MaterialClass->GetDefaultObject())->GetKeywords()));
+					FText::FromString(CreationName),
+					ToolTip, 0, CastChecked<UMaterialExpression>(MaterialExpression.MaterialClass->GetDefaultObject())->GetKeywords()));
 				NewNodeAction->MaterialExpressionClass = MaterialExpression.MaterialClass;
 				ActionMenuBuilder.AddAction(NewNodeAction);
 			}
diff --git a/Engine/Source/Editor/MaterialEditor/Private/MaterialExpressionClasses.cpp b/Engine/Source/Editor/MaterialEditor/Private/MaterialExpressionClasses.cpp
--- a/Engine/Source/Editor/MaterialEditor/Private/MaterialExpressionClasses.cpp
+++ b/Engine/Source/Editor/MaterialEditor/Private/MaterialExpressionClasses.cpp
@@ -101,12 +101,24 @@ void MaterialExpressionClasses::InitMaterialExpressionClasses()
 							// Trim the material expression name and add it to the list used for filtering.
 							static const FString ExpressionPrefix = TEXT("MaterialExpression");
 							FString ClassName = *Class->GetName();
+
+							if (Class->HasMetaData("DisplayName"))
+							{
+								ClassName = Class->GetDisplayNameText().ToString();
+							}
+
 							if (ClassName.StartsWith(ExpressionPrefix, ESearchCase::CaseSensitive))
 							{
 								ClassName = ClassName.Mid(ExpressionPrefix.Len());
 							}
 							MaterialExpression.Name = ClassName;
 							MaterialExpression.MaterialClass = Class;
+							UMaterialExpression* TempObject = Cast<UMaterialExpression>(Class->GetDefaultObject());
+							if (TempObject)
+							{
+								MaterialExpression.CreationDescription = TempObject->GetCreationDescription();
+								MaterialExpression.CreationName = TempObject->GetCreationName();
+							}
 
 							AllExpressionClasses.Add(MaterialExpression);
 
@@ -130,7 +142,6 @@ void MaterialExpressionClasses::InitMaterialExpressionClasses()
 							}
 
 							// Category fill...
-							UMaterialExpression* TempObject = Cast<UMaterialExpression>(Class->GetDefaultObject());
 							if (TempObject)
 							{
 								if (TempObject->MenuCategories.Num() == 0)
diff --git a/Engine/Source/Editor/MaterialEditor/Private/SMaterialPalette.cpp b/Engine/Source/Editor/MaterialEditor/Private/SMaterialPalette.cpp
--- a/Engine/Source/Editor/MaterialEditor/Private/SMaterialPalette.cpp
+++ b/Engine/Source/Editor/MaterialEditor/Private/SMaterialPalette.cpp
@@ -92,6 +92,11 @@ TSharedRef<SWidget> SMaterialPaletteItem::CreateHotkeyDisplayWidget(const FSlate
 		.Font(NameFont);
 }
 
+FText SMaterialPaletteItem::GetItemTooltip() const
+{
+	return FText::FromString(ActionPtr.Pin()->GetTooltipDescription());
+}
+
 //////////////////////////////////////////////////////////////////////////
 
 void SMaterialPalette::Construct(const FArguments& InArgs, TWeakPtr<FMaterialEditor> InMaterialEditorPtr)
diff --git a/Engine/Source/Editor/MaterialEditor/Private/SMaterialPalette.h b/Engine/Source/Editor/MaterialEditor/Private/SMaterialPalette.h
--- a/Engine/Source/Editor/MaterialEditor/Private/SMaterialPalette.h
+++ b/Engine/Source/Editor/MaterialEditor/Private/SMaterialPalette.h
@@ -23,6 +23,8 @@ public:
 private:
 	/* Create the hotkey display widget */
 	TSharedRef<SWidget> CreateHotkeyDisplayWidget(const FSlateFontInfo& NameFont, const TSharedPtr<const FInputChord> HotkeyChord);
+
+	virtual FText GetItemTooltip() const override;
 };
 
 //////////////////////////////////////////////////////////////////////////
diff --git a/Engine/Source/Editor/MaterialEditor/Public/IMaterialEditor.h b/Engine/Source/Editor/MaterialEditor/Public/IMaterialEditor.h
--- a/Engine/Source/Editor/MaterialEditor/Public/IMaterialEditor.h
+++ b/Engine/Source/Editor/MaterialEditor/Public/IMaterialEditor.h
@@ -55,6 +55,11 @@ public:
 	 */
 	virtual void DeleteSelectedNodes() {};
 
+	/** 
+	 * Delete an array of Material Graph Nodes and their corresponding expressions/comments
+	*/
+	virtual void DeleteNodes(const TArray<class UEdGraphNode*>& NodesToDelete) {}
+
 	/**
 	 * Gets the name of the material or material function that we are editing
 	 */
diff --git a/Engine/Source/Editor/MaterialEditor/Public/MaterialEditorActions.h b/Engine/Source/Editor/MaterialEditor/Public/MaterialEditorActions.h
--- a/Engine/Source/Editor/MaterialEditor/Public/MaterialEditorActions.h
+++ b/Engine/Source/Editor/MaterialEditor/Public/MaterialEditorActions.h
@@ -163,6 +163,9 @@ public:
 	/** Create component mask node */
 	TSharedPtr< FUICommandInfo > CreateComponentMaskNode;
 
+	/** Create component mask node */
+	TSharedPtr< FUICommandInfo > PromoteToParameter;
+
 
 	/**
 	 * Initialize commands
diff --git a/Engine/Source/Editor/MaterialEditor/Public/MaterialEditorUtilities.h b/Engine/Source/Editor/MaterialEditor/Public/MaterialEditorUtilities.h
--- a/Engine/Source/Editor/MaterialEditor/Public/MaterialEditorUtilities.h
+++ b/Engine/Source/Editor/MaterialEditor/Public/MaterialEditorUtilities.h
@@ -79,6 +79,14 @@ public:
 	 */
 	static void DeleteSelectedNodes(const class UEdGraph* Graph);
 
+	/**
+	 * Delete the specified nodes from the graph.
+	 * @param Graph Graph representing the material.
+	 * @param NodesToDelete Array of nodes to be removed from the graph.
+	*/
+	static void DeleteNodes(const class UEdGraph* Graph, const TArray<UEdGraphNode*>& NodesToDelete);
+
+
 	/**
 	 * Gets the name of the material or material function that we are editing
 	 *
