Copying //UE4/Dev-Blueprints to //UE4/Dev-Main (Source: //UE4/Dev-Blueprints @ 3235800)

diff --git a/Engine/Source/Editor/BlueprintGraph/Private/BlueprintComponentNodeSpawner.cpp b/Engine/Source/Editor/BlueprintGraph/Private/BlueprintComponentNodeSpawner.cpp
--- a/Engine/Source/Editor/BlueprintGraph/Private/BlueprintComponentNodeSpawner.cpp
+++ b/Engine/Source/Editor/BlueprintGraph/Private/BlueprintComponentNodeSpawner.cpp
@@ -136,7 +136,8 @@ FBlueprintNodeSignature UBlueprintComponentNodeSpawner::GetSpawnerSignature() co
 // and FEdGraphSchemaAction_K2AddComponent::PerformAction().
 UEdGraphNode* UBlueprintComponentNodeSpawner::Invoke(UEdGraph* ParentGraph, FBindingSet const& Bindings, FVector2D const Location) const
 {
-	auto PostSpawnLambda = [](UEdGraphNode* NewNode, bool bIsTemplateNode, FCustomizeNodeDelegate UserDelegate)
+	UClass* ComponentType = ComponentClass;
+	auto PostSpawnLambda = [ComponentType](UEdGraphNode* NewNode, bool bIsTemplateNode, FCustomizeNodeDelegate UserDelegate)
 	{		
 		UK2Node_AddComponent* AddCompNode = CastChecked<UK2Node_AddComponent>(NewNode);
 		UBlueprint* Blueprint = AddCompNode->GetBlueprint();
@@ -144,10 +145,12 @@ UEdGraphNode* UBlueprintComponentNodeSpawner::Invoke(UEdGraph* ParentGraph, FBin
 		UFunction* AddComponentFunc = FindFieldChecked<UFunction>(AActor::StaticClass(), UK2Node_AddComponent::GetAddComponentFunctionName());
 		AddCompNode->FunctionReference.SetFromField<UFunction>(AddComponentFunc, !bIsTemplateNode && FBlueprintEditorUtils::IsActorBased(Blueprint));
 
+		AddCompNode->TemplateType = ComponentType;
+
 		UserDelegate.ExecuteIfBound(NewNode, bIsTemplateNode);
 	};
 
-	FCustomizeNodeDelegate PostSpawnDelegate = FCustomizeNodeDelegate::CreateStatic(PostSpawnLambda, CustomizeNodeDelegate);
+	FCustomizeNodeDelegate PostSpawnDelegate = FCustomizeNodeDelegate::CreateLambda(PostSpawnLambda, CustomizeNodeDelegate);
 	// let SpawnNode() allocate default pins (so we can modify them)
 	UK2Node_AddComponent* NewNode = Super::SpawnNode<UK2Node_AddComponent>(NodeClass, ParentGraph, FBindingSet(), Location, PostSpawnDelegate);
 	if (NewNode->Pins.Num() == 0)
@@ -161,7 +164,7 @@ UEdGraphNode* UBlueprintComponentNodeSpawner::Invoke(UEdGraph* ParentGraph, FBin
 	{
 		if (ComponentClass != nullptr)
 		{
-			ReturnPin->PinType.PinSubCategoryObject = *ComponentClass;
+			ReturnPin->PinType.PinSubCategoryObject = ComponentType;
 		}
 		else
 		{
diff --git a/Engine/Source/Editor/BlueprintGraph/Private/BlueprintEditorSettings.cpp b/Engine/Source/Editor/BlueprintGraph/Private/BlueprintEditorSettings.cpp
--- a/Engine/Source/Editor/BlueprintGraph/Private/BlueprintEditorSettings.cpp
+++ b/Engine/Source/Editor/BlueprintGraph/Private/BlueprintEditorSettings.cpp
@@ -23,7 +23,6 @@ UBlueprintEditorSettings::UBlueprintEditorSettings(const FObjectInitializer& Obj
 	, bShowEmptySections(true)
 	, bSpawnDefaultBlueprintNodes(true)
 	, bHideConstructionScriptComponentsInDetailsView(true)
-	, bEnableAdvancedContainers(false)
 	// Compiler Settings
 	, SaveOnCompile(SoC_Never)
 	, bJumpToNodeErrors(false)
diff --git a/Engine/Source/Editor/BlueprintGraph/Private/CallFunctionHandler.cpp b/Engine/Source/Editor/BlueprintGraph/Private/CallFunctionHandler.cpp
--- a/Engine/Source/Editor/BlueprintGraph/Private/CallFunctionHandler.cpp
+++ b/Engine/Source/Editor/BlueprintGraph/Private/CallFunctionHandler.cpp
@@ -126,7 +126,7 @@ void FKCHandler_CallFunction::CreateFunctionCallStatement(FKismetFunctionContext
 								LatentUUID = GetTypeHash(AssociatedNode->NodeGuid);
 
 								UEdGraphNode* ResultNode = AssociatedNode;
-								UEdGraphNode* SourceNode = Cast<UEdGraphNode>(Context.MessageLog.FindSourceObject(AssociatedNode));
+								UEdGraphNode* SourceNode = Cast<UEdGraphNode>(Context.MessageLog.GetIntermediateTunnelInstance(AssociatedNode));
 								while (SourceNode && SourceNode != ResultNode)
 								{
 									if (SourceNode->NodeGuid.IsValid())
@@ -134,7 +134,7 @@ void FKCHandler_CallFunction::CreateFunctionCallStatement(FKismetFunctionContext
 										LatentUUID = HashCombine(LatentUUID, GetTypeHash(SourceNode->NodeGuid));
 									}
 									ResultNode = SourceNode;
-									SourceNode = Cast<UEdGraphNode>(Context.MessageLog.FindSourceObject(ResultNode));
+									SourceNode = Cast<UEdGraphNode>(Context.MessageLog.GetIntermediateTunnelInstance(ResultNode));
 								}
 							}
 							else
@@ -437,6 +437,7 @@ void FKCHandler_CallFunction::CreateFunctionCallStatement(FKismetFunctionContext
 			}
 
 			// Iterate over all the contexts this functions needs to be called on, and emit a call function statement for each
+			FBlueprintCompiledStatement* LatentStatement = nullptr;
 			for (auto TargetListIt = ContextTerms.CreateIterator(); TargetListIt; ++TargetListIt)
 			{
 				FBPTerminal* Target = *TargetListIt;
@@ -471,6 +472,7 @@ void FKCHandler_CallFunction::CreateFunctionCallStatement(FKismetFunctionContext
 						check(LatentTargetParamIndex != INDEX_NONE);
 						Statement.UbergraphCallIndex = LatentTargetParamIndex;
 						Context.GotoFixupRequestMap.Add(&Statement, ThenExecPin);
+						LatentStatement = &Statement;
 					}
 				}
 
@@ -484,6 +486,12 @@ void FKCHandler_CallFunction::CreateFunctionCallStatement(FKismetFunctionContext
 				{
 					FBlueprintCompiledStatement& SuspendState = Context.AppendStatementForNode(Node);
 					SuspendState.Type = KCST_InstrumentedStateSuspend;
+					if (LatentTargetNode != NULL)
+					{
+						check(LatentTargetParamIndex != INDEX_NONE);
+						SuspendState.UbergraphCallIndex = LatentTargetParamIndex;
+						SuspendState.TargetLabel = LatentStatement;
+					}
 				}
 				// End this thread of execution; the latent function will resume it at some point in the future
 				FBlueprintCompiledStatement& PopStatement = Context.AppendStatementForNode(Node);
diff --git a/Engine/Source/Editor/BlueprintGraph/Private/EdGraphSchema_K2.cpp b/Engine/Source/Editor/BlueprintGraph/Private/EdGraphSchema_K2.cpp
--- a/Engine/Source/Editor/BlueprintGraph/Private/EdGraphSchema_K2.cpp
+++ b/Engine/Source/Editor/BlueprintGraph/Private/EdGraphSchema_K2.cpp
@@ -1332,26 +1332,50 @@ bool UEdGraphSchema_K2::PinDefaultValueIsEditable(const UEdGraphPin& InGraphPin)
 
 void UEdGraphSchema_K2::SelectAllNodesInDirection(TEnumAsByte<enum EEdGraphPinDirection> InDirection, UEdGraph* Graph, UEdGraphPin* InGraphPin)
 {
-	TArray<UEdGraphPin*> AllPins = InGraphPin->LinkedTo;
-
-	if (AllPins.Num() == 0)
+	/** Traverses the node graph out from the specified pin, logging each node that it visits along the way. */
+	struct FDirectionalNodeVisitor
 	{
-		return;
-	}
+		FDirectionalNodeVisitor(UEdGraphPin* StartingPin, EEdGraphPinDirection TargetDirection)
+			: Direction(TargetDirection)
+		{
+			TraversePin(StartingPin);
+		}
 
-	for (UEdGraphPin* Pin : AllPins)
-	{
-		UEdGraphNode* OwningNode = Pin->GetOwningNode();
-		FKismetEditorUtilities::AddToSelection(Graph, OwningNode);
+		/** If the pin is the right direction, visits each of its attached nodes */
+		void TraversePin(UEdGraphPin* Pin)
+		{
+			if (Pin->Direction == Direction)
+			{
+				for (UEdGraphPin* LinkedPin : Pin->LinkedTo)
+				{
+					VisitNode(LinkedPin->GetOwningNode());
+				}
+			}
+		}
 
-		TArray<UEdGraphPin*> LinkedPins = Pin->GetOwningNode()->GetAllPins();
-		for (UEdGraphPin* InputPin : LinkedPins)
+		/** If the node has already been visited, does nothing. Otherwise it traverses each of its pins. */
+		void VisitNode(UEdGraphNode* Node)
 		{
-			if (InputPin->Direction == InDirection)
+			bool bAlreadyVisited = false;
+			VisitedNodes.Add(Node, &bAlreadyVisited);
+
+			if (!bAlreadyVisited)
 			{
-				SelectAllNodesInDirection(InDirection, Graph, InputPin);
+				for (UEdGraphPin* Pin : Node->Pins)
+				{
+					TraversePin(Pin);
+				}
 			}
 		}
+
+		EEdGraphPinDirection Direction;
+		TSet<UEdGraphNode*>  VisitedNodes;
+	};
+
+	FDirectionalNodeVisitor NodeVisitor(InGraphPin, InDirection);
+	for (UEdGraphNode* Node : NodeVisitor.VisitedNodes)
+	{
+		FKismetEditorUtilities::AddToSelection(Graph, Node);
 	}
 }
 
diff --git a/Engine/Source/Editor/BlueprintGraph/Private/K2Node_AddComponent.cpp b/Engine/Source/Editor/BlueprintGraph/Private/K2Node_AddComponent.cpp
--- a/Engine/Source/Editor/BlueprintGraph/Private/K2Node_AddComponent.cpp
+++ b/Engine/Source/Editor/BlueprintGraph/Private/K2Node_AddComponent.cpp
@@ -14,6 +14,7 @@
 #include "UObject/ReleaseObjectVersion.h"
 #include "KismetCompilerMisc.h"
 #include "KismetCompiler.h"
+#include "BlueprintsObjectVersion.h" // for ComponentTemplateClassSupport
 
 #define LOCTEXT_NAMESPACE "K2Node_AddComponent"
 
@@ -31,6 +32,12 @@ const FString UK2Node_AddComponent::ComponentTemplateNamePrefix(TEXT("NODE_Add")
 
 void UK2Node_AddComponent::Serialize(FArchive& Ar)
 {
+	if (!TemplateType && Ar.IsSaving() && GetLinkerCustomVersion(FBlueprintsObjectVersion::GUID) < FBlueprintsObjectVersion::ComponentTemplateClassSupport)
+	{
+		UActorComponent* Template = GetTemplateFromNode();
+		TemplateType = Template ? Template->GetClass() : nullptr;
+	}
+
 	Super::Serialize(Ar);
 
 	if (Ar.IsLoading())
@@ -39,6 +46,7 @@ void UK2Node_AddComponent::Serialize(FArchive& Ar)
 		{
 			if (Ar.IsPersistent() && !Ar.HasAnyPortFlags(PPF_Duplicate | PPF_DuplicateForPIE))
 			{
+				UBlueprint* Blueprint = GetBlueprint();
 				UActorComponent* Template = GetTemplateFromNode();
 
 				// Fix up old "generic" template names to conform to the new unique naming convention.
@@ -93,9 +101,7 @@ void UK2Node_AddComponent::AllocatePinsForExposedVariables()
 {
 	const UEdGraphSchema_K2* K2Schema = GetDefault<UEdGraphSchema_K2>();
 
-	const UActorComponent* TemplateComponent = GetTemplateFromNode();
-	const UClass* ComponentClass = TemplateComponent ? TemplateComponent->GetClass() : nullptr;
-
+	const UClass* ComponentClass = GetSpawnedType();
 	if (ComponentClass != nullptr)
 	{
 		const UObject* ClassDefaultObject = ComponentClass ? ComponentClass->ClassDefaultObject : nullptr;
@@ -149,10 +155,15 @@ FName UK2Node_AddComponent::GetAddComponentFunctionName()
     return AddComponentFunctionName;
 }
 
-const UClass* UK2Node_AddComponent::GetSpawnedType() const
+UClass* UK2Node_AddComponent::GetSpawnedType() const
 {
+	if (TemplateType != nullptr)
+	{
+		return TemplateType;
+	}
+	
 	const UActorComponent* TemplateComponent = GetTemplateFromNode();
-	return TemplateComponent ? TemplateComponent->GetClass() : NULL;
+	return TemplateComponent ? TemplateComponent->GetClass() : nullptr;
 }
 
 void UK2Node_AddComponent::AllocateDefaultPinsWithoutExposedVariables()
@@ -190,14 +201,26 @@ void UK2Node_AddComponent::ReallocatePinsDuringReconstruction(TArray<UEdGraphPin
 	AllocatePinsForExposedVariables();
 }
 
+void UK2Node_AddComponent::PostReconstructNode()
+{
+	Super::PostReconstructNode();
+
+	// Set the return type to the right class of component
+	UEdGraphPin* ReturnPin = GetReturnValuePin();
+	UClass* ComponentClass = GetSpawnedType();
+	if (ReturnPin && ComponentClass)
+	{
+		ReturnPin->PinType.PinSubCategoryObject = ComponentClass->GetAuthoritativeClass();
+	}
+}
+
 void UK2Node_AddComponent::ValidateNodeDuringCompilation(FCompilerResultsLog& MessageLog) const
 {
 	Super::ValidateNodeDuringCompilation(MessageLog);
 
-	UActorComponent* Template = GetTemplateFromNode();
-	if (Template)
+	const UClass* TemplateClass = GetSpawnedType();
+	if (TemplateClass)
 	{
-		UClass* TemplateClass = Template->GetClass();
 		if (!TemplateClass->IsChildOf(UActorComponent::StaticClass()) || TemplateClass->HasAnyClassFlags(CLASS_Abstract) || !TemplateClass->HasMetaData(FBlueprintMetadata::MD_BlueprintSpawnableComponent) )
 		{
 			FFormatNamedArguments Args;
@@ -206,7 +229,7 @@ void UK2Node_AddComponent::ValidateNodeDuringCompilation(FCompilerResultsLog& Me
 			MessageLog.Error(*FText::Format(NSLOCTEXT("KismetCompiler", "InvalidComponentTemplate_Error", "Invalid class '{TemplateClass}' used as template by '{NodeTitle}' for @@"), Args).ToString(), this);
 		}
 
-		if (UChildActorComponent const* ChildActorComponent = Cast<UChildActorComponent const>(Template))
+		if (UChildActorComponent const* ChildActorComponent = Cast<UChildActorComponent const>(GetTemplateFromNode()))
 		{
 			UBlueprint const* Blueprint = GetBlueprint();
 
@@ -233,7 +256,15 @@ void UK2Node_AddComponent::ValidateNodeDuringCompilation(FCompilerResultsLog& Me
 				{
 					FFormatNamedArguments Args;
 					Args.Add(TEXT("ChildActorClass"), FText::FromString(ChildActorClass->GetName()));
-					MessageLog.Error(*FText::Format(NSLOCTEXT("KismetCompiler", "AddStaticChildActorComponent_Error", "@@ cannot add a '{ChildActorClass}' component as it has static mobility, and the ChildActorComponent does not."), Args).ToString(), this);
+
+					if (IsRunningDedicatedServer())
+					{
+						MessageLog.Error(*FText::Format(NSLOCTEXT("KismetCompiler", "StrippedComponentTemplate_Error", "Unknown template referenced by '{NodeTitle}' for @@ - Please resave the blueprint to correct this."), Args).ToString(), this);
+					}
+					else
+					{
+						MessageLog.Error(*FText::Format(NSLOCTEXT("KismetCompiler", "MissingComponentTemplate_Error", "Unknown template referenced by '{NodeTitle}' for @@"), Args).ToString(), this);
+					}
 				}
 			}
 		}
@@ -382,6 +413,20 @@ bool UK2Node_AddComponent::IsCompatibleWithGraph(UEdGraph const* Graph) const
 	return (Blueprint != nullptr) && FBlueprintEditorUtils::IsActorBased(Blueprint) && Super::IsCompatibleWithGraph(Graph);
 }
 
+void UK2Node_AddComponent::ReconstructNode()
+{
+	Super::ReconstructNode();
+
+	if (GetLinkerCustomVersion(FBlueprintsObjectVersion::GUID) < FBlueprintsObjectVersion::ComponentTemplateClassSupport)
+	{
+		UActorComponent* Template = GetTemplateFromNode();
+		if (Template != nullptr)
+		{
+			TemplateType = Template->GetClass();
+		}
+	}
+}
+
 void UK2Node_AddComponent::ExpandNode(class FKismetCompilerContext& CompilerContext, UEdGraph* SourceGraph)
 {
 	Super::ExpandNode(CompilerContext, SourceGraph);
diff --git a/Engine/Source/Editor/BlueprintGraph/Private/K2Node_CallFunction.cpp b/Engine/Source/Editor/BlueprintGraph/Private/K2Node_CallFunction.cpp
--- a/Engine/Source/Editor/BlueprintGraph/Private/K2Node_CallFunction.cpp
+++ b/Engine/Source/Editor/BlueprintGraph/Private/K2Node_CallFunction.cpp
@@ -1003,14 +1003,6 @@ void UK2Node_CallFunction::PostReconstructNode()
 		}
 	}
 
-	// Set the return type to the right class of component
-	UActorComponent* TemplateComp = GetTemplateFromNode();
-	UEdGraphPin* ReturnPin = GetReturnValuePin();
-	if(TemplateComp && ReturnPin)
-	{
-		ReturnPin->PinType.PinSubCategoryObject = TemplateComp->GetClass()->GetAuthoritativeClass();
-	}
-
 	if (UEdGraphPin* TypePickerPin = FDynamicOutputHelper::GetTypePickerPin(this))
 	{
 		FDynamicOutputHelper(TypePickerPin).ConformOutputType();
@@ -1027,22 +1019,6 @@ void UK2Node_CallFunction::PostReconstructNode()
 	}
 }
 
-void UK2Node_CallFunction::DestroyNode()
-{
-	// See if this node has a template
-	UActorComponent* Template = GetTemplateFromNode();
-	if (Template != NULL)
-	{
-		// Get the blueprint so we can remove it from it
-		UBlueprint* BlueprintObj = GetBlueprint();
-
-		// remove it
-		BlueprintObj->ComponentTemplates.Remove(Template);
-	}
-
-	Super::DestroyNode();
-}
-
 void UK2Node_CallFunction::NotifyPinConnectionListChanged(UEdGraphPin* Pin)
 {
 	Super::NotifyPinConnectionListChanged(Pin);
diff --git a/Engine/Source/Editor/BlueprintGraph/Private/K2Node_DynamicCast.cpp b/Engine/Source/Editor/BlueprintGraph/Private/K2Node_DynamicCast.cpp
--- a/Engine/Source/Editor/BlueprintGraph/Private/K2Node_DynamicCast.cpp
+++ b/Engine/Source/Editor/BlueprintGraph/Private/K2Node_DynamicCast.cpp
@@ -407,11 +407,12 @@ void UK2Node_DynamicCast::ValidateNodeDuringCompilation(FCompilerResultsLog& Mes
 	UEdGraphPin* SourcePin = GetCastSourcePin();
 	if (SourcePin->LinkedTo.Num() > 0)
 	{
-		const UClass* SourceType = *TargetType;
+		UClass* SourceType = *TargetType;
 		if (SourceType == nullptr)
 		{
 			return;
 		}
+		SourceType = SourceType->GetAuthoritativeClass();
 
 		for (UEdGraphPin* CastInput : SourcePin->LinkedTo)
 		{
@@ -440,20 +441,21 @@ void UK2Node_DynamicCast::ValidateNodeDuringCompilation(FCompilerResultsLog& Mes
 
 				continue;
 			}
+			SourceClass = SourceClass->GetAuthoritativeClass();
 
 			if (SourceClass == SourceType)
 			{
 				const FString SourcePinName = CastInput->PinFriendlyName.IsEmpty() ? CastInput->PinName : CastInput->PinFriendlyName.ToString();
 
 				FText const WarningFormat = LOCTEXT("EqualObjectCast", "'%s' is already a '%s', you don't need @@.");
-				MessageLog.Warning( *FString::Printf(*WarningFormat.ToString(), *SourcePinName, *TargetType->GetDisplayNameText().ToString()), this );
+				MessageLog.Note( *FString::Printf(*WarningFormat.ToString(), *SourcePinName, *TargetType->GetDisplayNameText().ToString()), this );
 			}
 			else if (SourceClass->IsChildOf(SourceType))
 			{
 				const FString SourcePinName = CastInput->PinFriendlyName.IsEmpty() ? CastInput->PinName : CastInput->PinFriendlyName.ToString();
 
 				FText const WarningFormat = LOCTEXT("UnneededObjectCast", "'%s' is already a '%s' (which inherits from '%s'), so you don't need @@.");
-				MessageLog.Warning( *FString::Printf(*WarningFormat.ToString(), *SourcePinName, *SourceClass->GetDisplayNameText().ToString(), *TargetType->GetDisplayNameText().ToString()), this );
+				MessageLog.Note( *FString::Printf(*WarningFormat.ToString(), *SourcePinName, *SourceClass->GetDisplayNameText().ToString(), *TargetType->GetDisplayNameText().ToString()), this );
 			}
 			else if (!SourceType->IsChildOf(SourceClass) && !FKismetEditorUtilities::IsClassABlueprintInterface(SourceType))
 			{
diff --git a/Engine/Source/Editor/BlueprintGraph/Private/K2Node_FunctionEntry.cpp b/Engine/Source/Editor/BlueprintGraph/Private/K2Node_FunctionEntry.cpp
--- a/Engine/Source/Editor/BlueprintGraph/Private/K2Node_FunctionEntry.cpp
+++ b/Engine/Source/Editor/BlueprintGraph/Private/K2Node_FunctionEntry.cpp
@@ -11,6 +11,7 @@
 #include "EdGraphSchema_K2.h"
 #include "K2Node_CallFunction.h"
 #include "K2Node_MakeArray.h"
+#include "K2Node_MakeVariable.h"
 #include "K2Node_VariableSet.h"
 #include "Kismet2/BlueprintEditorUtils.h"
 #include "EdGraphUtilities.h"
@@ -429,7 +430,6 @@ void UK2Node_FunctionEntry::ExpandNode(class FKismetCompilerContext& CompilerCon
 						VariableSetNode->SetFromProperty(Property, false);
 						Schema->ConfigureVarNode(VariableSetNode, LocalVar.VarName, Function, CompilerContext.Blueprint);
 						VariableSetNode->AllocateDefaultPins();
-						CompilerContext.MessageLog.NotifyIntermediateObjectCreation(VariableSetNode, this);
 
 						if(UEdGraphPin* SetPin = VariableSetNode->FindPin(Property->GetName()))
 						{
@@ -450,24 +450,11 @@ void UK2Node_FunctionEntry::ExpandNode(class FKismetCompilerContext& CompilerCon
 								FScriptArrayHelper_InContainer ArrayHelper(ArrayProperty, StructData->GetStructMemory());
 								FScriptArrayHelper_InContainer DefaultArrayHelper(ArrayProperty, StructData->GetStructMemory());
 
-								uint8* StructDefaults = NULL;
-								UStructProperty* StructProperty = dynamic_cast<UStructProperty*>(ArrayProperty->Inner);
-								if ( StructProperty != NULL )
-								{
-									checkSlow(StructProperty->Struct);
-									StructDefaults = (uint8*)FMemory::Malloc(StructProperty->Struct->GetStructureSize());
-									StructProperty->InitializeValue(StructDefaults);
-								}
-
 								// Go through each element in the array to set the default value
 								for( int32 ArrayIndex = 0 ; ArrayIndex < ArrayHelper.Num() ; ArrayIndex++ )
 								{
 									uint8* PropData = ArrayHelper.GetRawPtr(ArrayIndex);
 
-									// Always use struct defaults if the inner is a struct, for symmetry with the import of array inner struct defaults
-									uint8* PropDefault = ( StructProperty != NULL ) ? StructDefaults :
-										( ( StructData->GetStructMemory() && DefaultArrayHelper.Num() > ArrayIndex ) ? DefaultArrayHelper.GetRawPtr(ArrayIndex) : NULL );
-
 									// Retrieve the element's default value
 									FString DefaultValue;
 									FBlueprintEditorUtils::PropertyValueToString(ArrayProperty->Inner, PropData, DefaultValue);
@@ -481,6 +468,11 @@ void UK2Node_FunctionEntry::ExpandNode(class FKismetCompilerContext& CompilerCon
 									Schema->TrySetDefaultValue(*MakeArray->Pins[ArrayIndex + 1], DefaultValue);
 								}
 							}
+							else if(LocalVar.VarType.bIsSet || LocalVar.VarType.bIsMap)
+							{
+								UK2Node_MakeVariable* MakeVariableNode = CompilerContext.SpawnIntermediateNode<UK2Node_MakeVariable>(this, SourceGraph);
+								MakeVariableNode->SetupVariable(LocalVar, SetPin, CompilerContext, Function, Property);
+							}
 							else
 							{
 								if (CompilerContext.Blueprint->bIsRegeneratingOnLoad)
diff --git a/Engine/Source/Editor/BlueprintGraph/Private/K2Node_Knot.cpp b/Engine/Source/Editor/BlueprintGraph/Private/K2Node_Knot.cpp
--- a/Engine/Source/Editor/BlueprintGraph/Private/K2Node_Knot.cpp
+++ b/Engine/Source/Editor/BlueprintGraph/Private/K2Node_Knot.cpp
@@ -76,60 +76,116 @@ bool UK2Node_Knot::CanSplitPin(const UEdGraphPin* Pin) const
 
 void UK2Node_Knot::NotifyPinConnectionListChanged(UEdGraphPin* Pin)
 {
-	const UEdGraphSchema_K2* K2Schema = GetDefault<UEdGraphSchema_K2>();
+	PropagatePinType();
+}
+
+void UK2Node_Knot::PostReconstructNode()
+{
+	PropagatePinType();
+	Super::PostReconstructNode();
+}
+
+void UK2Node_Knot::PropagatePinType()
+{
 	UEdGraphPin* MyInputPin = GetInputPin();
 	UEdGraphPin* MyOutputPin = GetOutputPin();
 
-	const int32 NumLinks = MyInputPin->LinkedTo.Num() + MyOutputPin->LinkedTo.Num();
+	if (MyInputPin->LinkedTo.Num() > 0)
+	{
+		// If we can't mirror from output type, we should at least get the type information from the input connection chain
+		PropagatePinTypeFromInput();
+	}
+	else if (MyOutputPin->LinkedTo.Num() > 0)
+	{
+		// Try to mirror from output first to make sure we get appropriate member references
+		PropagatePinTypeFromOutput();
+	}
+	else
+	{
+		// Revert to wildcard
+		const UEdGraphSchema_K2* K2Schema = GetDefault<UEdGraphSchema_K2>();
+
+		MyInputPin->BreakAllPinLinks();
+		MyInputPin->PinType.ResetToDefaults();
+		MyInputPin->PinType.PinCategory = K2Schema->PC_Wildcard;
+
+		MyOutputPin->BreakAllPinLinks();
+		MyOutputPin->PinType.ResetToDefaults();
+		MyOutputPin->PinType.PinCategory = K2Schema->PC_Wildcard;
+	}
+}
 
-	if (Pin->LinkedTo.Num() > 0)
+void UK2Node_Knot::PropagatePinTypeFromInput()
+{
+	// Set the type of the pin based on input connections.
+	// We have to move up the chain of linked reroute nodes until we reach a node
+	// with type information before percolating that information down.
+	UEdGraphPin* MyInputPin = GetInputPin();
+	UEdGraphPin* MyOutputPin = GetOutputPin();
+
+	for (UEdGraphPin* InPin : MyInputPin->LinkedTo)
 	{
-		// Just made a connection, was it the first?
-		if (NumLinks == 1)
+		if (UK2Node_Knot* KnotNode = Cast<UK2Node_Knot>(InPin->GetOwningNode()))
 		{
-			UEdGraphPin* TypeSource = Pin->LinkedTo[0];
-
-			MyInputPin->PinType = TypeSource->PinType;
-			MyOutputPin->PinType = TypeSource->PinType;
+			KnotNode->PropagatePinTypeFromInput();
 		}
 	}
-	else
+
+	UEdGraphPin* TypeSource = MyInputPin->LinkedTo.Num() ? MyInputPin->LinkedTo[0] : nullptr;
+	if (TypeSource)
 	{
-		// Just broke a connection, was it the last?
-		if (NumLinks == 0)
+		MyInputPin->PinType = TypeSource->PinType;
+		MyOutputPin->PinType = TypeSource->PinType;
+
+		for (UEdGraphPin* InPin : MyInputPin->LinkedTo)
 		{
-			// Revert to wildcard
-			MyInputPin->BreakAllPinLinks();
-			MyInputPin->PinType.ResetToDefaults();
-			MyInputPin->PinType.PinCategory = K2Schema->PC_Wildcard;
-
-			MyOutputPin->BreakAllPinLinks();
-			MyOutputPin->PinType.ResetToDefaults();
-			MyOutputPin->PinType.PinCategory = K2Schema->PC_Wildcard;
+			if (UK2Node* OwningNode = Cast<UK2Node>(InPin->GetOwningNode()))
+			{
+				if (!OwningNode->IsA<UK2Node_Knot>())
+				{
+					OwningNode->PinConnectionListChanged(InPin);
+				}
+			}
 		}
 	}
+	else
+	{
+		// TODO?
+	}
 }
 
-void UK2Node_Knot::PostReconstructNode()
+void UK2Node_Knot::PropagatePinTypeFromOutput()
 {
+	// Set the type of the pin based on the output connection, and then percolate
+	// that type information up until we no longer reach another Reroute node
 	UEdGraphPin* MyInputPin = GetInputPin();
 	UEdGraphPin* MyOutputPin = GetOutputPin();
 
-	// Find a pin that has connections to use to jumpstart the wildcard process
-	for (int32 PinIndex = 0; PinIndex < Pins.Num(); ++PinIndex)
+	UEdGraphPin* TypeSource = MyOutputPin->LinkedTo.Num() ? MyOutputPin->LinkedTo[0] : nullptr;
+	if (TypeSource)
 	{
-		if (Pins[PinIndex]->LinkedTo.Num() > 0)
-		{
-			// The pin is linked, continue to use its type as the type for all pins.
-			UEdGraphPin* TypeSource = Pins[PinIndex]->LinkedTo[0];
+		MyInputPin->PinType = TypeSource->PinType;
+		MyOutputPin->PinType = TypeSource->PinType;
 
-			MyInputPin->PinType = TypeSource->PinType;
-			MyOutputPin->PinType = TypeSource->PinType;
-			break;
+		for (UEdGraphPin* InPin : MyInputPin->LinkedTo)
+		{
+			if (UK2Node* OwningNode = Cast<UK2Node>(InPin->GetOwningNode()))
+			{
+				if (UK2Node_Knot* KnotNode = Cast<UK2Node_Knot>(OwningNode))
+				{
+					KnotNode->PropagatePinTypeFromOutput();
+				}
+				else
+				{
+					OwningNode->PinConnectionListChanged(InPin);
+				}
+			}
 		}
 	}
-
-	Super::PostReconstructNode();
+	else 
+	{
+		// TODO?
+	}
 }
 
 void UK2Node_Knot::GetMenuActions(FBlueprintActionDatabaseRegistrar& ActionRegistrar) const
diff --git a/Engine/Source/Editor/BlueprintGraph/Private/K2Node_MathExpression.cpp b/Engine/Source/Editor/BlueprintGraph/Private/K2Node_MathExpression.cpp
--- a/Engine/Source/Editor/BlueprintGraph/Private/K2Node_MathExpression.cpp
+++ b/Engine/Source/Editor/BlueprintGraph/Private/K2Node_MathExpression.cpp
@@ -999,6 +999,14 @@ private:
 			ADD_ALIAS("RAND")
 			ADD_ALIAS("RANDOM")
 		FUNC_ALIASES_END
+
+		FUNC_ALIASES_BEGIN("Dot_VectorVector")
+			ADD_ALIAS("Dot")
+		FUNC_ALIASES_END
+
+		FUNC_ALIASES_BEGIN("Cross_VectorVector")
+			ADD_ALIAS("Cross")
+		FUNC_ALIASES_END
 				
 		// if none of the above aliases returned, then we don't have any for
 		// this function (use its regular name)
diff --git a/Engine/Source/Editor/BlueprintGraph/Private/K2Node_TunnelBoundary.cpp b/Engine/Source/Editor/BlueprintGraph/Private/K2Node_TunnelBoundary.cpp
--- a/Engine/Source/Editor/BlueprintGraph/Private/K2Node_TunnelBoundary.cpp
+++ b/Engine/Source/Editor/BlueprintGraph/Private/K2Node_TunnelBoundary.cpp
@@ -194,6 +194,11 @@ void UK2Node_TunnelBoundary::CreateBoundariesForExpansionNodes(UEdGraphNode* Sou
 		{
 			for (UEdGraphPin* LinkedPin : Pin->LinkedTo)
 			{
+				//if (ExpansionNode->IsA<UK2Node_Event>() && Pin->Direction == EGPD_Output && Pin->PinType.PinCategory == UEdGraphSchema_K2::PC_Exec)
+				//{
+				//	// Add boundary for expansion events.
+				//	LinkedPinMap.Add(Pin, LinkedPin);
+				//}
 				if (UEdGraphPin** SourcePin = LinkedPinMap.Find(LinkedPin))
 				{
 					if (Pin->Direction == EGPD_Input)
diff --git a/Engine/Source/Editor/BlueprintGraph/Private/K2Node_Variable.cpp b/Engine/Source/Editor/BlueprintGraph/Private/K2Node_Variable.cpp
--- a/Engine/Source/Editor/BlueprintGraph/Private/K2Node_Variable.cpp
+++ b/Engine/Source/Editor/BlueprintGraph/Private/K2Node_Variable.cpp
@@ -148,7 +148,7 @@ void UK2Node_Variable::CreatePinForSelf()
 				{
 					TargetClass = Property->GetOwnerClass()->GetAuthoritativeClass();
 				}
-				else
+				else if(GetBlueprint()->SkeletonGeneratedClass)
 				{
 					TargetClass = GetBlueprint()->SkeletonGeneratedClass->GetAuthoritativeClass();
 				}
diff --git a/Engine/Source/Editor/BlueprintGraph/Public/BlueprintEditorSettings.h b/Engine/Source/Editor/BlueprintGraph/Public/BlueprintEditorSettings.h
--- a/Engine/Source/Editor/BlueprintGraph/Public/BlueprintEditorSettings.h
+++ b/Engine/Source/Editor/BlueprintGraph/Public/BlueprintEditorSettings.h
@@ -81,9 +81,6 @@ public:
 	UPROPERTY(EditAnywhere, config, Category=Workflow)
 	bool bHideConstructionScriptComponentsInDetailsView;
 
-	UPROPERTY(EditAnywhere, config, Category = Experimental, AdvancedDisplay, meta = (DisplayName = "Enable (incomplete, non-functional) support for TSet and TMap"))
-	bool bEnableAdvancedContainers;
-
 // Compiler Settings
 public:
 	/** Determines when to save Blueprints post-compile */
diff --git a/Engine/Source/Editor/EditorStyle/Private/SlateEditorStyle.cpp b/Engine/Source/Editor/EditorStyle/Private/SlateEditorStyle.cpp
--- a/Engine/Source/Editor/EditorStyle/Private/SlateEditorStyle.cpp
+++ b/Engine/Source/Editor/EditorStyle/Private/SlateEditorStyle.cpp
@@ -4057,7 +4057,7 @@ void FSlateEditorStyle::FStyle::SetupGraphEditorStyles()
 				);
 
 			Set( "Graph.CompactNode.Title", FTextBlockStyle(NormalText)
-				.SetFont( TTF_FONT( "Fonts/Roboto-BoldCondensed", 40 ) )
+				.SetFont( TTF_FONT( "Fonts/Roboto-BoldCondensed", 20 ) )
 				.SetColorAndOpacity( FLinearColor(1.0f, 1.0f, 1.0f, 0.5f) )
 				.SetShadowOffset( FVector2D::ZeroVector )
 				.SetShadowColorAndOpacity( FLinearColor::White )
diff --git a/Engine/Source/Editor/IntroTutorials/Private/EditorTutorialImportFactory.cpp b/Engine/Source/Editor/IntroTutorials/Private/EditorTutorialImportFactory.cpp
--- a/Engine/Source/Editor/IntroTutorials/Private/EditorTutorialImportFactory.cpp
+++ b/Engine/Source/Editor/IntroTutorials/Private/EditorTutorialImportFactory.cpp
@@ -65,7 +65,7 @@ UObject* UEditorTutorialImportFactory::FactoryCreateBinary(UClass* InClass, UObj
 bool UEditorTutorialImportFactory::CanReimport(UObject* Obj, TArray<FString>& OutFilenames)
 {
 	UBlueprint* Blueprint = Cast<UBlueprint>(Obj);
-	if(Blueprint != nullptr)
+	if(Blueprint != nullptr && Blueprint->GeneratedClass)
 	{
 		UEditorTutorial* EditorTutorial = Cast<UEditorTutorial>(Blueprint->GeneratedClass->GetDefaultObject());
 		if(EditorTutorial != nullptr)
diff --git a/Engine/Source/Editor/Kismet/Private/BPFunctionDragDropAction.cpp b/Engine/Source/Editor/Kismet/Private/BPFunctionDragDropAction.cpp
--- a/Engine/Source/Editor/Kismet/Private/BPFunctionDragDropAction.cpp
+++ b/Engine/Source/Editor/Kismet/Private/BPFunctionDragDropAction.cpp
@@ -11,6 +11,7 @@
 #include "Kismet2/BlueprintEditorUtils.h"
 #include "BlueprintNodeBinder.h"
 #include "BlueprintFunctionNodeSpawner.h"
+#include "ScopedTransaction.h"
 
 #define LOCTEXT_NAMESPACE "FunctionDragDropAction"
 
@@ -237,6 +238,8 @@ FReply FKismetFunctionDragDropAction::DroppedOnPanel(TSharedRef<SWidget> const&
 			{
 				AnalyticCallback.ExecuteIfBound();
 
+				const FScopedTransaction Transaction( LOCTEXT("KismetFunction_DroppedOnPanel", "Function Dropped on Graph") );
+
 				IBlueprintNodeBinder::FBindingSet Bindings;
 				FunctionNodeSpawner->Invoke(&Graph, Bindings, GraphPosition);
 
diff --git a/Engine/Source/Editor/Kismet/Private/BlueprintDetailsCustomization.cpp b/Engine/Source/Editor/Kismet/Private/BlueprintDetailsCustomization.cpp
--- a/Engine/Source/Editor/Kismet/Private/BlueprintDetailsCustomization.cpp
+++ b/Engine/Source/Editor/Kismet/Private/BlueprintDetailsCustomization.cpp
@@ -4997,7 +4997,7 @@ void FBlueprintGlobalOptionsDetails::CustomizeDetails(IDetailLayoutBuilder& Deta
 			UProperty* Property = *PropertyIt;
 			FString Category = Property->GetMetaData(TEXT("Category"));
 
-			if (Category != TEXT("BlueprintOptions") && Category != TEXT("ClassOptions") && Category != TEXT("Experimental"))
+			if (Category != TEXT("BlueprintOptions") && Category != TEXT("ClassOptions") && Category != TEXT("Packaging"))
 			{
 				DetailLayout.HideProperty(DetailLayout.GetProperty(Property->GetFName()));
 			}
diff --git a/Engine/Source/Editor/Kismet/Private/BlueprintEditor.cpp b/Engine/Source/Editor/Kismet/Private/BlueprintEditor.cpp
--- a/Engine/Source/Editor/Kismet/Private/BlueprintEditor.cpp
+++ b/Engine/Source/Editor/Kismet/Private/BlueprintEditor.cpp
@@ -2613,6 +2613,16 @@ void FBlueprintEditor::CreateDefaultCommands()
 	ToolkitCommands->MapAction(FBlueprintEditorCommands::Get().GenerateNativeCode,
 		FExecuteAction::CreateSP(this, &FBlueprintEditor::OpenNativeCodeGenerationTool),
 		FCanExecuteAction::CreateSP(this, &FBlueprintEditor::CanGenerateNativeCode));
+
+	ToolkitCommands->MapAction(FBlueprintEditorCommands::Get().ShowActionMenuItemSignatures,
+		FExecuteAction::CreateLambda([]()
+			{ 
+				UBlueprintEditorSettings* Settings = GetMutableDefault<UBlueprintEditorSettings>();
+				Settings->bShowActionMenuItemSignatures = !Settings->bShowActionMenuItemSignatures;
+				Settings->SaveConfig();
+			}),
+		FCanExecuteAction(),
+		FIsActionChecked::CreateLambda([]()->bool{ return GetDefault<UBlueprintEditorSettings>()->bShowActionMenuItemSignatures; }));
 }
 
 bool FBlueprintEditor::IsProfilerAvailable() const
@@ -2736,7 +2746,7 @@ void FBlueprintEditor::ReparentBlueprint_NewParentChosen(UClass* ChosenClass)
 		}
 
 		// If the chosen class differs hierarchically from the current class, warn that there may be data loss
-		if (bReparent && !ChosenClass->GetDefaultObject()->IsA(BlueprintObj->ParentClass))
+		if (bReparent && (!BlueprintObj->ParentClass || !ChosenClass->GetDefaultObject()->IsA(BlueprintObj->ParentClass)))
 		{
 			const FText Title = LOCTEXT("ReparentTitle", "Reparent Blueprint"); 
 			const FText Message = LOCTEXT("ReparentWarning", "Reparenting this blueprint may cause data loss.  Continue reparenting?"); 
@@ -2756,7 +2766,7 @@ void FBlueprintEditor::ReparentBlueprint_NewParentChosen(UClass* ChosenClass)
 
 		if ( bReparent )
 		{
-			UE_LOG(LogBlueprint, Warning, TEXT("Reparenting blueprint %s from %s to %s..."), *BlueprintObj->GetFullName(), *BlueprintObj->ParentClass->GetName(), *ChosenClass->GetName());
+			UE_LOG(LogBlueprint, Warning, TEXT("Reparenting blueprint %s from %s to %s..."), *BlueprintObj->GetFullName(), BlueprintObj->ParentClass ? *BlueprintObj->ParentClass->GetName() : TEXT("[None]"), *ChosenClass->GetName());
 
 			UClass* OldParentClass = BlueprintObj->ParentClass ;
 			BlueprintObj->ParentClass = ChosenClass;
diff --git a/Engine/Source/Editor/Kismet/Private/BlueprintEditorCommands.cpp b/Engine/Source/Editor/Kismet/Private/BlueprintEditorCommands.cpp
--- a/Engine/Source/Editor/Kismet/Private/BlueprintEditorCommands.cpp
+++ b/Engine/Source/Editor/Kismet/Private/BlueprintEditorCommands.cpp
@@ -77,6 +77,8 @@ void FBlueprintEditorCommands::RegisterCommands()
 
 	UI_COMMAND(GenerateNativeCode, "Generate Native Code", "Generate C++ code from the blueprint", EUserInterfaceActionType::Button, FInputChord());
 
+	UI_COMMAND(ShowActionMenuItemSignatures, "Show Action Menu Item Signatures", "If enabled, tooltips on action menu items will show the associated action's signature id (can be used to setup custom favorites menus).", EUserInterfaceActionType::ToggleButton, FInputChord());
+
 	// SCC commands
 	UI_COMMAND( BeginBlueprintMerge, "Merge", "Shows the Blueprint merge panel and toolbar, allowing the user to resolve conflicted blueprints", EUserInterfaceActionType::Button, FInputChord() );
 }
diff --git a/Engine/Source/Editor/Kismet/Private/BlueprintEditorCommands.h b/Engine/Source/Editor/Kismet/Private/BlueprintEditorCommands.h
--- a/Engine/Source/Editor/Kismet/Private/BlueprintEditorCommands.h
+++ b/Engine/Source/Editor/Kismet/Private/BlueprintEditorCommands.h
@@ -75,6 +75,7 @@ public:
 	TSharedPtr< FUICommandInfo > RecompileBlueprintEditor;
 	TSharedPtr< FUICommandInfo > RecompilePersona;
 	TSharedPtr< FUICommandInfo > GenerateNativeCode;
+	TSharedPtr< FUICommandInfo > ShowActionMenuItemSignatures;
 
 	// SSC commands
 	TSharedPtr< FUICommandInfo > BeginBlueprintMerge;
diff --git a/Engine/Source/Editor/Kismet/Private/DetailsDiff.cpp b/Engine/Source/Editor/Kismet/Private/DetailsDiff.cpp
--- a/Engine/Source/Editor/Kismet/Private/DetailsDiff.cpp
+++ b/Engine/Source/Editor/Kismet/Private/DetailsDiff.cpp
@@ -92,7 +92,7 @@ void FDetailsDiff::DiffAgainst(const FDetailsDiff& Newer, TArray< FSingleObjectD
 	ToSet(NewProperties, NewPropertiesSet);
 
 	// detect removed properties:
-	auto RemovedProperties = OldPropertiesSet.Difference(NewPropertiesSet);
+	TSet<FPropertySoftPath> RemovedProperties = OldPropertiesSet.Difference(NewPropertiesSet);
 	for (const auto& RemovedProperty : RemovedProperties)
 	{
 		// @todo: (doc) label these as removed, rather than added to a
@@ -101,25 +101,30 @@ void FDetailsDiff::DiffAgainst(const FDetailsDiff& Newer, TArray< FSingleObjectD
 	}
 
 	// detect added properties:
-	auto AddededProperties = NewPropertiesSet.Difference(OldPropertiesSet);
-	for (const auto& AddedProperty : AddededProperties)
+	TSet<FPropertySoftPath> AddedProperties = NewPropertiesSet.Difference(OldPropertiesSet);
+	for (const auto& AddedProperty : AddedProperties)
 	{
 		FSingleObjectDiffEntry Entry(AddedProperty, EPropertyDiffType::PropertyAddedToB);
 		OutDifferences.Push(Entry);
 	}
 
 	// check for changed properties
-	auto CommonProperties = NewPropertiesSet.Intersect(OldPropertiesSet);
+	TSet<FPropertySoftPath> CommonProperties = NewPropertiesSet.Intersect(OldPropertiesSet);
 	for (const auto& CommonProperty : CommonProperties)
 	{
 		// get value, diff:
 		check(NewSelectedObjects.Num() == 1);
-		auto OldPoperty = CommonProperty.Resolve(OldSelectedObjects[0].Get());
-		auto NewProperty = CommonProperty.Resolve(NewSelectedObjects[0].Get());
+		FResolvedProperty OldProperty = CommonProperty.Resolve(OldSelectedObjects[0].Get());
+		FResolvedProperty NewProperty = CommonProperty.Resolve(NewSelectedObjects[0].Get());
 
-		if (!DiffUtils::Identical(OldPoperty, NewProperty))
+		TArray<FPropertySoftPath> DifferingSubProperties;
+
+		if (!DiffUtils::Identical(OldProperty, NewProperty, CommonProperty, DifferingSubProperties))
 		{
-			OutDifferences.Push(FSingleObjectDiffEntry(CommonProperty, EPropertyDiffType::PropertyValueChanged));
+			for (int32 DifferingIndex = 0; DifferingIndex < DifferingSubProperties.Num(); DifferingIndex++)
+			{
+				OutDifferences.Push(FSingleObjectDiffEntry(DifferingSubProperties[DifferingIndex], EPropertyDiffType::PropertyValueChanged));
+			}
 		}
 	}
 }
diff --git a/Engine/Source/Editor/Kismet/Private/DiffUtils.cpp b/Engine/Source/Editor/Kismet/Private/DiffUtils.cpp
--- a/Engine/Source/Editor/Kismet/Private/DiffUtils.cpp
+++ b/Engine/Source/Editor/Kismet/Private/DiffUtils.cpp
@@ -14,6 +14,11 @@
 
 static const UProperty* Resolve( const UStruct* Class, FName PropertyName )
 {
+	if(Class == nullptr )
+	{
+		return nullptr;
+	}
+
 	for (TFieldIterator<UProperty> PropertyIt(Class); PropertyIt; ++PropertyIt)
 	{
 		if( PropertyIt->GetFName() == PropertyName )
@@ -72,29 +77,189 @@ FResolvedProperty FPropertySoftPath::Resolve(const UObject* Object) const
 
 FPropertyPath FPropertySoftPath::ResolvePath(const UObject* Object) const
 {
-	const UStruct* ClassOverride = nullptr;
+	auto UpdateContainerAddress = [](const UProperty* Property, const void* Instance, const void*& OutContainerAddress, const UStruct*& OutContainerStruct)
+	{
+		if( ensure(Instance) )
+		{
+			if(const UObjectProperty* ObjectProperty = Cast<UObjectProperty>(Property))
+			{
+				const UObject* const* InstanceObject = reinterpret_cast<const UObject* const*>(Instance);
+				if( *InstanceObject)
+				{
+					OutContainerAddress = *InstanceObject;
+					OutContainerStruct = (*InstanceObject)->GetClass();
+				}
+			}
+			else if(const UStructProperty* StructProperty = Cast<UStructProperty>(Property))
+			{
+				OutContainerAddress = Instance;
+				OutContainerStruct = StructProperty->Struct;
+			}
+		}
+	};
+
+	auto TryReadIndex = [](const TArray<FName>& LocalPropertyChain, int32& OutIndex) -> int32
+	{
+		if(OutIndex + 1 < LocalPropertyChain.Num())
+		{
+			FString AsString = LocalPropertyChain[OutIndex + 1].ToString();
+			if(AsString.IsNumeric())
+			{
+				++OutIndex;
+				return FCString::Atoi(*AsString);
+			}
+		}
+		return INDEX_NONE;
+	};
+
+	const void* ContainerAddress = Object;
+	const UStruct* ContainerStruct = Object->GetClass();
+
 	FPropertyPath Ret;
-	for( const auto& Property : PropertyChain )
+	for( int32 I = 0; I < PropertyChain.Num(); ++I )
 	{
-		const UProperty* ResolvedProperty = ::Resolve(ClassOverride ? ClassOverride : Object->GetClass(), Property);
-		if (const UObjectProperty* ObjectProperty = Cast<UObjectProperty>(ResolvedProperty))
+		FName PropertyIdentifier = PropertyChain[I];
+		const UProperty * ResolvedProperty = ::Resolve(ContainerStruct, PropertyIdentifier);
+			
+		FPropertyInfo Info = { ResolvedProperty, INDEX_NONE };
+		Ret.AddProperty(Info);
+
+		int32 PropertyIndex = TryReadIndex(PropertyChain, I);
+		
+		
+		// calculate offset so we can continue resolving object properties/structproperties:
+		if( const UArrayProperty* ArrayProperty = Cast<UArrayProperty>(ResolvedProperty) )
+		{
+			if(PropertyIndex != INDEX_NONE)
+			{
+				FScriptArrayHelper ArrayHelper(ArrayProperty, ArrayProperty->ContainerPtrToValuePtr<const void*>( ContainerAddress ));
+
+				UpdateContainerAddress( ArrayProperty->Inner, ArrayHelper.GetRawPtr(PropertyIndex), ContainerAddress, ContainerStruct );
+
+				FPropertyInfo ArrayInfo = { ArrayProperty->Inner, PropertyIndex };
+				Ret.AddProperty(ArrayInfo);
+			}
+		}
+		else if( const USetProperty* SetProperty = Cast<USetProperty>(ResolvedProperty) )
+		{
+			if(PropertyIndex != INDEX_NONE)
+			{
+				FScriptSetHelper SetHelper(SetProperty, SetProperty->ContainerPtrToValuePtr<const void*>( ContainerAddress ));
+
+				// Figure out the real index in this instance of the set (sets have gaps in them):
+				int32 RealIndex = -1;
+				for( int32 J = 0; PropertyIndex >= 0; ++J)
+				{
+					++RealIndex;
+					if(SetHelper.IsValidIndex(J))
+					{
+						--PropertyIndex;
+					}
+				}
+
+				UpdateContainerAddress( SetProperty->ElementProp, SetHelper.GetElementPtr(RealIndex), ContainerAddress, ContainerStruct );
+
+				FPropertyInfo SetInfo = { SetProperty->ElementProp, RealIndex };
+				Ret.AddProperty(SetInfo);
+			}
+		}
+		else if( const UMapProperty* MapProperty = Cast<UMapProperty>(ResolvedProperty) )
+		{
+			if(PropertyIndex != INDEX_NONE)
+			{
+				FScriptMapHelper MapHelper(MapProperty, MapProperty->ContainerPtrToValuePtr<const void*>( ContainerAddress ));
+				
+				// Figure out the real index in this instance of the map (map have gaps in them):
+				int32 RealIndex = -1;
+				for( int32 J = 0; PropertyIndex >= 0; ++J)
+				{
+					++RealIndex;
+					if(MapHelper.IsValidIndex(J))
+					{
+						--PropertyIndex;
+					}
+				}
+
+				// we have an index, but are we looking into a key or value? Peek ahead to find out:
+				if(ensure(I + 1 < PropertyChain.Num()))
+				{
+					if(PropertyChain[I+1] == MapProperty->KeyProp->GetFName())
+					{
+						++I;
+
+						UpdateContainerAddress( MapProperty->KeyProp, MapHelper.GetKeyPtr(RealIndex), ContainerAddress, ContainerStruct );
+
+						FPropertyInfo MakKeyInfo = { MapProperty->KeyProp, RealIndex };
+						Ret.AddProperty(MakKeyInfo);
+					}
+					else if(ensure( PropertyChain[I+1] == MapProperty->ValueProp->GetFName() ))
+					{	
+						++I;
+
+						UpdateContainerAddress( MapProperty->ValueProp, MapHelper.GetValuePtr(RealIndex), ContainerAddress, ContainerStruct );
+						
+						FPropertyInfo MapValueInfo = { MapProperty->ValueProp, RealIndex };
+						Ret.AddProperty(MapValueInfo);
+					}
+				}
+			}
+		}
+		else if (const UObjectProperty* ObjectProperty = Cast<UObjectProperty>(ResolvedProperty))
 		{
-			// At the moment the cast of Object to void* just avoids an overzealous assert in ContainerPtrToValuePtr (for some of our
-			// properties, the outer is a UStruct, not a UClass). At some point I want to support diffing of individual array elements, 
-			// and that will necessitate the container be a void* anyway:
-			const UObject* const* BaseObject = reinterpret_cast<const UObject* const*>( ObjectProperty->ContainerPtrToValuePtr<void>( static_cast<const void*>(Object)) );
-			if( BaseObject && *BaseObject )
+			UpdateContainerAddress( ObjectProperty, ObjectProperty->ContainerPtrToValuePtr<const void*>( ContainerAddress, FMath::Max(PropertyIndex, 0) ), ContainerAddress, ContainerStruct );
+			
+			// handle static arrays:
+			if(PropertyIndex != INDEX_NONE )
 			{
-				Object = *BaseObject;
+				FPropertyInfo ObjectInfo = { ResolvedProperty, PropertyIndex };
+				Ret.AddProperty(ObjectInfo);
 			}
-			ClassOverride = nullptr;
 		}
 		else if( const UStructProperty* StructProperty = Cast<UStructProperty>(ResolvedProperty) )
 		{
-			ClassOverride = StructProperty->Struct;
+			UpdateContainerAddress( StructProperty, StructProperty->ContainerPtrToValuePtr<const void*>( ContainerAddress, FMath::Max(PropertyIndex, 0) ), ContainerAddress, ContainerStruct );
+			
+			// handle static arrays:
+			if(PropertyIndex != INDEX_NONE )
+			{
+				FPropertyInfo StructInfo = { ResolvedProperty, PropertyIndex };
+				Ret.AddProperty(StructInfo);
+			}
+		}
+		else
+		{
+			// handle static arrays:
+			if(PropertyIndex != INDEX_NONE )
+			{
+				FPropertyInfo StaticArrayInfo = { ResolvedProperty, PropertyIndex };
+				Ret.AddProperty(StaticArrayInfo);
+			}
+		}
+	}
+	return Ret;
+}
+
+FString FPropertySoftPath::ToDisplayName() const
+{
+	FString Ret;
+	for( FName Property : PropertyChain )
+	{
+		FString PropertyAsString = Property.ToString();
+		if(Ret.IsEmpty())
+		{
+			Ret.Append(PropertyAsString);
+		}
+		else if( PropertyAsString.IsNumeric())
+		{
+			Ret.AppendChar('[');
+			Ret.Append(PropertyAsString);
+			Ret.AppendChar(']');
+		}
+		else
+		{
+			Ret.AppendChar(' ');
+			Ret.Append(PropertyAsString);
 		}
-		FPropertyInfo Info = { ResolvedProperty, -1 };
-		Ret.AddProperty(Info);
 	}
 	return Ret;
 }
@@ -139,9 +304,13 @@ void DiffUtils::CompareUnrelatedObjects(const UObject* A, const UObject* B, TArr
 			FResolvedProperty BProp = PropertyName.Resolve(B);
 
 			check(AProp != FResolvedProperty() && BProp != FResolvedProperty());
-			if (!DiffUtils::Identical(AProp, BProp))
+			TArray<FPropertySoftPath> DifferingSubProperties;
+			if (!DiffUtils::Identical(AProp, BProp, PropertyName, DifferingSubProperties))
 			{
-				OutDifferingProperties.Push(FSingleObjectDiffEntry( PropertyName, EPropertyDiffType::PropertyValueChanged ));
+				for (int DifferingIndex = 0; DifferingIndex < DifferingSubProperties.Num(); DifferingIndex++)
+				{
+					OutDifferingProperties.Push(FSingleObjectDiffEntry(DifferingSubProperties[DifferingIndex], EPropertyDiffType::PropertyValueChanged));
+				}
 			}
 		}
 	}
@@ -205,12 +374,216 @@ void DiffUtils::CompareUnrelatedSCS(const UBlueprint* Old, const TArray< FSCSRes
 	}
 }
 
-bool DiffUtils::Identical(const FResolvedProperty& AProp, const FResolvedProperty& BProp)
+static void AdvanceSetIterator( FScriptSetHelper& SetHelper, int32& Index)
+{
+	while(Index < SetHelper.Num() && !SetHelper.IsValidIndex(Index))
+	{
+		++Index;
+	}
+}
+
+static void AdvanceMapIterator( FScriptMapHelper& MapHelper, int32& Index)
+{
+	while(Index < MapHelper.Num() && !MapHelper.IsValidIndex(Index))
+	{
+		++Index;
+	}
+}
+
+static void IdenticalHelper(const UProperty* AProperty, const UProperty* BProperty, const void* AValue, const void* BValue, const FPropertySoftPath& RootPath, TArray<FPropertySoftPath>& DifferingSubProperties, bool bStaticArrayHandled = false)
+{
+	if(AProperty == nullptr || BProperty == nullptr || AProperty->ArrayDim != BProperty->ArrayDim || AProperty->GetClass() != BProperty->GetClass())
+	{
+		DifferingSubProperties.Push(RootPath);
+		return;
+	}
+
+	if(!bStaticArrayHandled && AProperty->ArrayDim != 1)
+	{
+		// Identical does not handle static array case automatically and we have to do the offset calculation ourself because 
+		// our container (e.g. the struct or class or dynamic array) has already done the initial offset calculation:
+		for( int32 I = 0; I < AProperty->ArrayDim; ++I )
+		{
+			int32 Offset = AProperty->ElementSize * I;
+			const void* CurAValue = reinterpret_cast<const void*>(reinterpret_cast<const uint8*>(AValue) + Offset);
+			const void* CurBValue = reinterpret_cast<const void*>(reinterpret_cast<const uint8*>(BValue) + Offset);
+
+			IdenticalHelper(AProperty, BProperty, CurAValue, CurBValue, FPropertySoftPath(RootPath, I), DifferingSubProperties, true);
+		}
+
+		return;
+	}
+	else if (AProperty->Identical(AValue, BValue, PPF_DeepComparison))
+	{
+		return;
+	}
+
+	const UStructProperty* APropAsStruct = Cast<UStructProperty>(AProperty);
+	const UArrayProperty* APropAsArray = Cast<UArrayProperty>(AProperty);
+	const USetProperty* APropAsSet = Cast<USetProperty>(AProperty);
+	const UMapProperty* APropAsMap = Cast<UMapProperty>(AProperty);
+	const UObjectProperty* APropAsObject = Cast<UObjectProperty>(AProperty);
+	if (APropAsStruct != nullptr)
+	{
+		const UStructProperty* BPropAsStruct = CastChecked<UStructProperty>(BProperty);
+		if (APropAsStruct->Struct->StructFlags & STRUCT_IdenticalNative || BPropAsStruct->Struct != APropAsStruct->Struct)
+		{
+			// If the struct uses CPP identical tests, then we can't dig into it, and we already know it's not identical from the test when we started
+			DifferingSubProperties.Push(RootPath);
+		}
+		else
+		{
+			for (TFieldIterator<UProperty> PropertyIt(APropAsStruct->Struct); PropertyIt; ++PropertyIt)
+			{
+				const UProperty* StructProp = *PropertyIt;
+				IdenticalHelper(StructProp, StructProp, StructProp->ContainerPtrToValuePtr<void>(AValue, 0), StructProp->ContainerPtrToValuePtr<void>(BValue, 0), FPropertySoftPath(RootPath, StructProp), DifferingSubProperties);
+			}
+		}
+	}
+	else if (APropAsArray != nullptr)
+	{
+		const UArrayProperty* BPropAsArray = CastChecked<UArrayProperty>(BProperty);
+		if(BPropAsArray->Inner == APropAsArray->Inner)
+		{
+			FScriptArrayHelper ArrayHelperA(APropAsArray, AValue);
+			FScriptArrayHelper ArrayHelperB(BPropAsArray, BValue);
+		
+			// note any differences in contained types:
+			for (int32 ArrayIndex = 0; ArrayIndex < ArrayHelperA.Num() && ArrayIndex < ArrayHelperB.Num(); ArrayIndex++)
+			{
+				IdenticalHelper(APropAsArray->Inner, APropAsArray->Inner, ArrayHelperA.GetRawPtr(ArrayIndex), ArrayHelperB.GetRawPtr(ArrayIndex), FPropertySoftPath(RootPath, ArrayIndex), DifferingSubProperties);
+			}
+
+			// note any size difference:
+			if (ArrayHelperA.Num() != ArrayHelperB.Num())
+			{
+				DifferingSubProperties.Push(RootPath);
+			}
+		}
+		else
+		{
+			DifferingSubProperties.Push(RootPath);
+		}
+	}
+	else if(APropAsSet != nullptr)
+	{
+		const USetProperty* BPropAsSet = CastChecked<USetProperty>(BProperty);
+		if(BPropAsSet->ElementProp == APropAsSet->ElementProp)
+		{
+			FScriptSetHelper SetHelperA(APropAsSet, AValue);
+			FScriptSetHelper SetHelperB(BPropAsSet, BValue);
+
+			if (SetHelperA.Num() != SetHelperB.Num())
+			{
+				// API not robust enough to indicate changes made to # of set elements, would
+				// need to return something more detailed than DifferingSubProperties array:
+				DifferingSubProperties.Push(RootPath);
+			}
+
+			// note any differences in contained elements:
+			int32 SetSizeA = SetHelperA.Num();
+			int32 SetSizeB = SetHelperA.Num();
+			
+			int32 SetIndexA = 0;
+			int32 SetIndexB = 0;
+
+			AdvanceSetIterator(SetHelperA, SetIndexA);
+			AdvanceSetIterator(SetHelperB, SetIndexB);
+
+			for (int32 VirtualIndex = 0; VirtualIndex < SetSizeA && VirtualIndex < SetSizeB; ++VirtualIndex)
+			{
+				IdenticalHelper(APropAsSet->ElementProp, APropAsSet->ElementProp, SetHelperA.GetElementPtr(SetIndexA), SetHelperB.GetElementPtr(SetIndexB), FPropertySoftPath(RootPath, VirtualIndex), DifferingSubProperties);
+
+				// advance iterators in step:
+				AdvanceSetIterator(SetHelperA, SetIndexA);
+				AdvanceSetIterator(SetHelperB, SetIndexB);
+			}
+		}
+		else
+		{
+			DifferingSubProperties.Push(RootPath);
+		}
+	}
+	else if(APropAsMap != nullptr)
+	{
+		const UMapProperty* BPropAsMap = CastChecked<UMapProperty>(BProperty);
+		if(APropAsMap->KeyProp == BPropAsMap->KeyProp && APropAsMap->ValueProp == BPropAsMap->ValueProp)
+		{
+			FScriptMapHelper MapHelperA(APropAsMap, AValue);
+			FScriptMapHelper MapHelperB(BPropAsMap, BValue);
+
+			if (MapHelperA.Num() != MapHelperB.Num())
+			{
+				// API not robust enough to indicate changes made to # of set elements, would
+				// need to return something more detailed than DifferingSubProperties array:
+				DifferingSubProperties.Push(RootPath);
+			}
+
+			int32 MapSizeA = MapHelperA.Num();
+			int32 MapSizeB = MapHelperB.Num();
+			
+			int32 MapIndexA = 0;
+			int32 MapIndexB = 0;
+
+			AdvanceMapIterator(MapHelperA, MapIndexA);
+			AdvanceMapIterator(MapHelperB, MapIndexB);
+			
+			for (int32 VirtualIndex = 0; VirtualIndex < MapSizeA && VirtualIndex < MapSizeB; ++VirtualIndex)
+			{
+				IdenticalHelper(APropAsMap->KeyProp, APropAsMap->KeyProp, MapHelperA.GetKeyPtr(MapIndexA), MapHelperB.GetKeyPtr(MapIndexB), FPropertySoftPath(RootPath, VirtualIndex), DifferingSubProperties);
+				IdenticalHelper(APropAsMap->ValueProp, APropAsMap->ValueProp, MapHelperA.GetValuePtr(MapIndexA), MapHelperB.GetValuePtr(MapIndexB), FPropertySoftPath(RootPath, VirtualIndex), DifferingSubProperties);
+
+				AdvanceMapIterator(MapHelperA, MapIndexA);
+				AdvanceMapIterator(MapHelperB, MapIndexB);
+			}
+		}
+		else
+		{
+			DifferingSubProperties.Push(RootPath);
+		}
+	}
+	else if(APropAsObject != nullptr)
+	{
+		// dig into the objects if they are in the same package as our initial object:
+		const UObjectProperty* BPropAsObject = CastChecked<UObjectProperty>(BProperty);
+
+		const UObject* A = *((const UObject* const*)AValue);
+		const UObject* B = *((const UObject* const*)BValue);
+
+		if(BPropAsObject->HasAnyPropertyFlags(CPF_InstancedReference) && APropAsObject->HasAnyPropertyFlags(CPF_InstancedReference) && A && B && A->GetClass() == B->GetClass())
+		{
+			// dive into the object to find actual differences:
+			const UClass* AClass = A->GetClass(); // BClass and AClass are identical!
+
+			for (TFieldIterator<UProperty> PropertyIt(AClass); PropertyIt; ++PropertyIt)
+			{
+				const UProperty* ClassProp = *PropertyIt;
+				IdenticalHelper(ClassProp, ClassProp, ClassProp->ContainerPtrToValuePtr<void>(A, 0), ClassProp->ContainerPtrToValuePtr<void>(B, 0), FPropertySoftPath(RootPath, ClassProp), DifferingSubProperties);
+			}
+		}
+		else
+		{		
+			DifferingSubProperties.Push(RootPath);
+		}
+	}
+	else
+	{
+		DifferingSubProperties.Push(RootPath);
+	}
+}
+
+bool DiffUtils::Identical(const FResolvedProperty& AProp, const FResolvedProperty& BProp, const FPropertySoftPath& RootPath, TArray<FPropertySoftPath>& DifferingProperties)
 {
 	const void* AValue = AProp.Property->ContainerPtrToValuePtr<void>(AProp.Object);
 	const void* BValue = BProp.Property->ContainerPtrToValuePtr<void>(BProp.Object);
 
-	return AProp.Property->Identical(AValue, BValue, PPF_DeepComparison);
+	// We _could_ just ask the property for comparison but that would make the "identical" functions significantly more complex.
+	// Instead let's write a new function, specific to DiffUtils, that handles the sub properties
+	// NOTE: For Static Arrays, AValue and BValue were, and are, only references to the value at index 0.  So changes to values past index 0 didn't show up before and
+	// won't show up now.  Changes to index 0 will show up as a change to the entire array.
+	IdenticalHelper(AProp.Property, BProp.Property, AValue, BValue, RootPath, DifferingProperties);
+	
+	return DifferingProperties.Num() == 0;
 }
 
 TArray<FPropertySoftPath> DiffUtils::GetVisiblePropertiesInOrderDeclared(const UObject* ForObj, const TArray<FName>& Scope /*= TArray<FName>()*/)
@@ -582,7 +955,7 @@ FLinearColor DiffViewUtils::Conflicting()
 FText DiffViewUtils::PropertyDiffMessage(FSingleObjectDiffEntry Difference, FText ObjectName)
 {
 	FText Message;
-	FString PropertyName = Difference.Identifier.LastPropertyName().GetPlainNameString();
+	FString PropertyName = Difference.Identifier.ToDisplayName();
 	switch (Difference.DiffType)
 	{
 	case EPropertyDiffType::PropertyAddedToA:
diff --git a/Engine/Source/Editor/Kismet/Private/SBlueprintDiff.cpp b/Engine/Source/Editor/Kismet/Private/SBlueprintDiff.cpp
--- a/Engine/Source/Editor/Kismet/Private/SBlueprintDiff.cpp
+++ b/Engine/Source/Editor/Kismet/Private/SBlueprintDiff.cpp
@@ -169,7 +169,6 @@ FCDODiffControl::FCDODiffControl(
 	: OldDetails(InOldCDO, FDetailsDiff::FOnDisplayedPropertiesChanged() )
 	, NewDetails(InNewCDO, FDetailsDiff::FOnDisplayedPropertiesChanged() )
 {
-	// @todo: (doc): rename InDifferences
 	TArray< FSingleObjectDiffEntry > DifferingProperties;
 	OldDetails.DiffAgainst(NewDetails, DifferingProperties);
 
@@ -180,17 +179,25 @@ FCDODiffControl::FCDODiffControl(
 	TArray<FPropertySoftPath> OldProperties = OldDetails.GetDisplayedProperties();
 	TArray<FPropertySoftPath> NewProperties = NewDetails.GetDisplayedProperties();
 
-	const auto FindAndPushDiff = [&OrderedProperties, &DifferingProperties](const FPropertySoftPath& PropertyIdentifer) -> bool
+	const auto FindAndPushDiff = [&OrderedProperties, &DifferingProperties](const FPropertySoftPath& PropertyIdentifier) -> bool
 	{
+		bool bDiffers = false;
 		for (const auto& Difference : DifferingProperties)
 		{
-			if (Difference.Identifier == PropertyIdentifer)
+			if (Difference.Identifier == PropertyIdentifier)
 			{
+				bDiffers = true;
+				// if there are any nested differences associated with PropertyIdentifier, add those
+				// as well:
+				OrderedProperties.Push(&Difference);
+			}
+			else if (Difference.Identifier.IsSubPropertyMatch(PropertyIdentifier))
+			{
+				bDiffers = true;
 				OrderedProperties.Push(&Difference);
-				return true;
 			}
 		}
-		return false;
+		return bDiffers;
 	};
 
 	// zip the two sets of properties, zip iterators are not trivial to write in c++,
diff --git a/Engine/Source/Editor/Kismet/Private/SBlueprintEditorSelectedDebugObjectWidget.cpp b/Engine/Source/Editor/Kismet/Private/SBlueprintEditorSelectedDebugObjectWidget.cpp
--- a/Engine/Source/Editor/Kismet/Private/SBlueprintEditorSelectedDebugObjectWidget.cpp
+++ b/Engine/Source/Editor/Kismet/Private/SBlueprintEditorSelectedDebugObjectWidget.cpp
@@ -450,7 +450,7 @@ void SBlueprintEditorSelectedDebugObjectWidget::GenerateDebugObjectNames(bool bR
 
 		const bool bPassesFlags = !TestObject->HasAnyFlags(RF_ClassDefaultObject) && !TestObject->IsPendingKill();
 		const bool bGeneratedByAnyBlueprint = TestObject->GetClass()->ClassGeneratedBy != nullptr;
-		const bool bGeneratedByThisBlueprint = bGeneratedByAnyBlueprint && TestObject->IsA(GetBlueprintObj()->GeneratedClass);
+		const bool bGeneratedByThisBlueprint = bGeneratedByAnyBlueprint && GetBlueprintObj()->GeneratedClass && TestObject->IsA(GetBlueprintObj()->GeneratedClass);
 
 		if (bPassesFlags && bGeneratedByThisBlueprint)
 		{
diff --git a/Engine/Source/Editor/Kismet/Private/SBlueprintEditorToolbar.cpp b/Engine/Source/Editor/Kismet/Private/SBlueprintEditorToolbar.cpp
--- a/Engine/Source/Editor/Kismet/Private/SBlueprintEditorToolbar.cpp
+++ b/Engine/Source/Editor/Kismet/Private/SBlueprintEditorToolbar.cpp
@@ -130,6 +130,12 @@ void FKismet2Menu::FillDeveloperMenu( FMenuBuilder& MenuBuilder )
 		}
 		MenuBuilder.EndSection();
 	}
+
+	MenuBuilder.BeginSection("SchemaDeveloperSettings", LOCTEXT("SchemaDevUtilsHeading", "Schema Utilities"));
+	{
+		MenuBuilder.AddMenuEntry(FBlueprintEditorCommands::Get().ShowActionMenuItemSignatures);
+	}
+	MenuBuilder.EndSection();
 }
 
 void FKismet2Menu::FillEditMenu( FMenuBuilder& MenuBuilder )
@@ -399,6 +405,7 @@ static void BlueprintEditorToolbarImpl::MakeCompileDeveloperSubMenu(FMenuBuilder
 {
 	const FBlueprintEditorCommands& EditorCommands = FBlueprintEditorCommands::Get();
 	InMenuBuilder.AddMenuEntry(EditorCommands.SaveIntermediateBuildProducts);
+	InMenuBuilder.AddMenuEntry(EditorCommands.ShowActionMenuItemSignatures);
 }
 
 
diff --git a/Engine/Source/Editor/Kismet/Private/SKismetInspector.cpp b/Engine/Source/Editor/Kismet/Private/SKismetInspector.cpp
--- a/Engine/Source/Editor/Kismet/Private/SKismetInspector.cpp
+++ b/Engine/Source/Editor/Kismet/Private/SKismetInspector.cpp
@@ -29,6 +29,7 @@
 #include "K2Node_FormatText.h"
 #include "K2Node_VariableGet.h"
 #include "K2Node_VariableSet.h"
+#include "K2Node_AddComponent.h" // for GetTemplateFromNode()
 #include "IDetailCustomization.h"
 #include "Editor.h"
 #include "PropertyEditorModule.h"
@@ -252,10 +253,10 @@ FText SKismetInspector::GetContextualEditingWidgetTitle() const
 					}
 				}
 			}
-			else if (UK2Node* K2Node = Cast<UK2Node>(Object))
+			else if (UK2Node_AddComponent* ComponentNode = Cast<UK2Node_AddComponent>(Object))
 			{
 				// Edit the component template
-				if (UActorComponent* Template = K2Node->GetTemplateFromNode())
+				if (UActorComponent* Template = ComponentNode->GetTemplateFromNode())
 				{
 					Title = FText::Format(LOCTEXT("Name_TemplateFmt", "{0} Template"), FText::FromString(Template->GetClass()->GetName()));
 				}
@@ -583,10 +584,13 @@ void SKismetInspector::UpdateFromObjects(const TArray<UObject*>& PropertyObjects
 			else if (UK2Node* K2Node = Cast<UK2Node>(Object))
 			{
 				// Edit the component template if it exists
-				if (UActorComponent* Template = K2Node->GetTemplateFromNode())
+				if (UK2Node_AddComponent* ComponentNode = Cast<UK2Node_AddComponent>(K2Node))
 				{
-					SelectionInfo.ObjectsForPropertyEditing.Add(Template);
-					SelectionInfo.EditableComponentTemplates.Add(Template);
+					if (UActorComponent* Template = ComponentNode->GetTemplateFromNode())
+					{
+						SelectionInfo.ObjectsForPropertyEditing.Add(Template);
+						SelectionInfo.EditableComponentTemplates.Add(Template);
+					}
 				}
 
 				// See if we should edit properties of the node
diff --git a/Engine/Source/Editor/Kismet/Private/SSCSEditor.cpp b/Engine/Source/Editor/Kismet/Private/SSCSEditor.cpp
--- a/Engine/Source/Editor/Kismet/Private/SSCSEditor.cpp
+++ b/Engine/Source/Editor/Kismet/Private/SSCSEditor.cpp
@@ -4667,6 +4667,10 @@ UActorComponent* SSCSEditor::AddNewComponent( UClass* NewComponentClass, UObject
 				NewComponentName = *FComponentEditorUtils::GenerateValidVariableName(NewComponentClass, ActorInstance);
 			}
 
+			// Get the set of owned components that exists prior to instancing the new component.
+			TInlineComponentArray<UActorComponent*> PreInstanceComponents;
+			ActorInstance->GetComponents(PreInstanceComponents);
+
 			// Construct the new component and attach as needed
 			UActorComponent* NewInstanceComponent = NewObject<UActorComponent>(ActorInstance, NewComponentClass, NewComponentName, RF_Transactional);
 			if (USceneComponent* NewSceneComponent = Cast<USceneComponent>(NewInstanceComponent))
@@ -4693,6 +4697,17 @@ UActorComponent* SSCSEditor::AddNewComponent( UClass* NewComponentClass, UObject
 			NewInstanceComponent->OnComponentCreated();
 			NewInstanceComponent->RegisterComponent();
 
+			// Register any new components that may have been created during construction of the instanced component, but were not explicitly registered.
+			TInlineComponentArray<UActorComponent*> PostInstanceComponents;
+			ActorInstance->GetComponents(PostInstanceComponents);
+			for (UActorComponent* ActorComponent : PostInstanceComponents)
+			{
+				if (!ActorComponent->IsRegistered() && ActorComponent->bAutoRegister && !ActorComponent->IsPendingKill() && !PreInstanceComponents.Contains(ActorComponent))
+				{
+					ActorComponent->RegisterComponent();
+				}
+			}
+
 			// Rerun construction scripts
 			ActorInstance->RerunConstructionScripts();
 
diff --git a/Engine/Source/Editor/Kismet/Public/DiffUtils.h b/Engine/Source/Editor/Kismet/Public/DiffUtils.h
--- a/Engine/Source/Editor/Kismet/Public/DiffUtils.h
+++ b/Engine/Source/Editor/Kismet/Public/DiffUtils.h
@@ -59,9 +59,39 @@ struct FPropertySoftPath
 		}
 	}
 
+	FPropertySoftPath(const FPropertySoftPath& SubPropertyPath, const UProperty* LeafProperty)
+		: PropertyChain(SubPropertyPath.PropertyChain)
+	{
+		PropertyChain.Push(LeafProperty->GetFName());
+	}
+
+	FPropertySoftPath(const FPropertySoftPath& SubPropertyPath, int32 ContainerIndex)
+		: PropertyChain(SubPropertyPath.PropertyChain)
+	{
+		PropertyChain.Push(FName(*FString::FromInt(ContainerIndex)));
+	}
+
 	FResolvedProperty Resolve(const UObject* Object) const;
 	FPropertyPath ResolvePath(const UObject* Object) const;
-	FName LastPropertyName() const;
+	FString ToDisplayName() const;
+
+	inline bool IsSubPropertyMatch(const FPropertySoftPath& PotentialBasePropertyPath) const
+	{
+		if (PropertyChain.Num() <= PotentialBasePropertyPath.PropertyChain.Num())
+		{
+			return false;
+		}
+
+		for (int32 CurChainElement = 0; CurChainElement < PotentialBasePropertyPath.PropertyChain.Num(); CurChainElement++)
+		{
+			if (PotentialBasePropertyPath.PropertyChain[CurChainElement] != PropertyChain[CurChainElement])
+			{
+				return false;
+			}
+		}
+
+		return true;
+	}
 
 	inline bool operator==(FPropertySoftPath const& RHS) const
 	{
@@ -99,15 +129,6 @@ FORCEINLINE bool operator!=(const FSCSIdentifier& A, const FSCSIdentifier& B)
 	return !(A == B);
 }
 
-FORCEINLINE FName FPropertySoftPath::LastPropertyName() const
-{
-	if( PropertyChain.Num() != 0 )
-	{
-		return PropertyChain.Last();
-	}
-	return FName();
-}
-
 FORCEINLINE uint32 GetTypeHash( FPropertySoftPath const& Path )
 {
 	uint32 Ret = 0;
@@ -188,7 +209,7 @@ namespace DiffUtils
 	KISMET_API const UObject* GetCDO(const UBlueprint* ForBlueprint);
 	KISMET_API void CompareUnrelatedObjects(const UObject* A, const UObject* B, TArray<FSingleObjectDiffEntry>& OutDifferingProperties);
 	KISMET_API void CompareUnrelatedSCS(const UBlueprint* Old, const TArray< FSCSResolvedIdentifier >& OldHierarchy, const UBlueprint* New, const TArray< FSCSResolvedIdentifier >& NewHierarchy, FSCSDiffRoot& OutDifferingEntries );
-	KISMET_API bool Identical(const FResolvedProperty& AProp, const FResolvedProperty& BProp);
+	KISMET_API bool Identical(const FResolvedProperty& AProp, const FResolvedProperty& BProp, const FPropertySoftPath& RootPath, TArray<FPropertySoftPath>& DifferingProperties); 
 	TArray<FPropertySoftPath> GetVisiblePropertiesInOrderDeclared(const UObject* ForObj, const TArray<FName>& Scope = TArray<FName>());
 
 	KISMET_API TArray<FPropertyPath> ResolveAll(const UObject* Object, const TArray<FPropertySoftPath>& InSoftProperties);
diff --git a/Engine/Source/Editor/KismetCompiler/Private/KismetCompiler.cpp b/Engine/Source/Editor/KismetCompiler/Private/KismetCompiler.cpp
--- a/Engine/Source/Editor/KismetCompiler/Private/KismetCompiler.cpp
+++ b/Engine/Source/Editor/KismetCompiler/Private/KismetCompiler.cpp
@@ -3717,16 +3717,6 @@ void FKismetCompilerContext::Compile()
 			}
 		}
 
-		// Create a mapping between compile time generated events and what created them.
-		if (TargetClass->bHasInstrumentation)
-		{
-			FBlueprintDebugData& DebugData = TargetClass->GetDebugData();
-			for (const TPair<UEdGraphPin*, UK2Node_Event*>& EventInfo : SourcePinToExpansionEvent)
-			{
-				DebugData.RegisterPinToEventName(EventInfo.Key, EventInfo.Value->GetFunctionName());
-			}
-		}
-
 		// Save off intermediate build products if requested
 		if (CompileOptions.bSaveIntermediateProducts && !Blueprint->bIsRegeneratingOnLoad)
 		{
@@ -3919,7 +3909,7 @@ void FKismetCompilerContext::Compile()
 			}
 		}
 
-		if (bDisplayBytecode && bIsFullCompile)
+		if (bDisplayBytecode && bIsFullCompile && !IsRunningCommandlet())
 		{
 			TGuardValue<ELogTimes::Type> DisableLogTimes(GPrintLogTimes, ELogTimes::None);
 
@@ -3938,14 +3928,14 @@ void FKismetCompilerContext::Compile()
 		}
 
 		// Generate code thru the backend(s)
-		if ((bDisplayCpp && bIsFullCompile) || CompileOptions.DoesRequireCppCodeGeneration())
+		if ((bDisplayCpp && bIsFullCompile && !IsRunningCommandlet()) || CompileOptions.DoesRequireCppCodeGeneration())
 		{
 			FString CppSourceCode;
 			FString HeaderSourceCode;
 
 			{
 				TUniquePtr<IBlueprintCompilerCppBackend> Backend_CPP(IBlueprintCompilerCppBackendModuleInterface::Get().Create());
-				HeaderSourceCode = Backend_CPP->GenerateCodeFromClass(NewClass, FunctionList, !bIsFullCompile, CppSourceCode);
+				HeaderSourceCode = Backend_CPP->GenerateCodeFromClass(NewClass, FunctionList, !bIsFullCompile, CompileOptions.NativizationOptions, CppSourceCode);
 			}
 
 			if (CompileOptions.OutHeaderSourceCode.IsValid())
@@ -3958,7 +3948,7 @@ void FKismetCompilerContext::Compile()
 				*CompileOptions.OutCppSourceCode = CppSourceCode;
 			}
 
-			if (bDisplayCpp)
+			if (bDisplayCpp && !IsRunningCommandlet())
 			{
 				UE_LOG(LogK2Compiler, Log, TEXT("[header]\n\n\n%s"), *HeaderSourceCode);
 				UE_LOG(LogK2Compiler, Log, TEXT("[body]\n\n\n%s"), *CppSourceCode);
@@ -3966,7 +3956,7 @@ void FKismetCompilerContext::Compile()
 		}
 
 		static const FBoolConfigValueHelper DisplayLayout(TEXT("Kismet"), TEXT("bDisplaysLayout"), GEngineIni);
-		if (!Blueprint->bIsRegeneratingOnLoad && bIsFullCompile && DisplayLayout && NewClass)
+		if (!Blueprint->bIsRegeneratingOnLoad && bIsFullCompile && DisplayLayout && NewClass && !IsRunningCommandlet())
 		{
 			UE_LOG(LogK2Compiler, Log, TEXT("\n\nLAYOUT CLASS %s:"), *GetNameSafe(NewClass));
 
diff --git a/Engine/Source/Editor/KismetCompiler/Private/KismetCompilerModule.cpp b/Engine/Source/Editor/KismetCompiler/Private/KismetCompilerModule.cpp
--- a/Engine/Source/Editor/KismetCompiler/Private/KismetCompilerModule.cpp
+++ b/Engine/Source/Editor/KismetCompiler/Private/KismetCompilerModule.cpp
@@ -34,13 +34,31 @@ DECLARE_CYCLE_STAT(TEXT("Compile Time"), EKismetCompilerStats_CompileTime, STATG
 DECLARE_CYCLE_STAT(TEXT("Compile Skeleton Class"), EKismetCompilerStats_CompileSkeletonClass, STATGROUP_KismetCompiler);
 DECLARE_CYCLE_STAT(TEXT("Compile Generated Class"), EKismetCompilerStats_CompileGeneratedClass, STATGROUP_KismetCompiler);
 
-IMPLEMENT_MODULE( FKismet2CompilerModule, KismetCompiler );
-
-
 #define LOCTEXT_NAMESPACE "KismetCompiler"
 
 //////////////////////////////////////////////////////////////////////////
-// FKismet2CompilerModule
+// FKismet2CompilerModule - The Kismet 2 Compiler module
+class FKismet2CompilerModule : public IKismetCompilerInterface
+{
+public:
+	// Implementation of the IKismetCompilerInterface
+	virtual void CompileBlueprint(class UBlueprint* Blueprint, const FKismetCompilerOptions& CompileOptions, FCompilerResultsLog& Results, TSharedPtr<class FBlueprintCompileReinstancer> ParentReinstancer = NULL, TArray<UObject*>* ObjLoaded = NULL) override;
+	virtual void CompileStructure(class UUserDefinedStruct* Struct, FCompilerResultsLog& Results) override;
+	virtual void RecoverCorruptedBlueprint(class UBlueprint* Blueprint) override;
+	virtual void RemoveBlueprintGeneratedClasses(class UBlueprint* Blueprint) override;
+	virtual TArray<IBlueprintCompiler*>& GetCompilers() override { return Compilers; }
+	virtual void GetBlueprintTypesForClass(UClass* ParentClass, UClass*& OutBlueprintClass, UClass*& OutBlueprintGeneratedClass) const override;
+	virtual void GenerateCppCodeForEnum(UUserDefinedEnum* UDEnum, FString& OutHeaderCode, FString& OutCPPCode) override;
+	virtual FString GenerateCppCodeForStruct(UUserDefinedStruct* UDStruct) override;
+	virtual FString GenerateCppWrapper(UBlueprintGeneratedClass* BPGC) override;
+	// End implementation
+private:
+	void CompileBlueprintInner(class UBlueprint* Blueprint, const FKismetCompilerOptions& CompileOptions, FCompilerResultsLog& Results, TSharedPtr<FBlueprintCompileReinstancer> Reinstancer, TArray<UObject*>* ObjLoaded);
+
+	TArray<IBlueprintCompiler*> Compilers;
+};
+
+IMPLEMENT_MODULE( FKismet2CompilerModule, KismetCompiler );
 
 struct FBlueprintIsBeingCompiledHelper
 {
@@ -160,10 +178,10 @@ void FKismet2CompilerModule::CompileStructure(UUserDefinedStruct* Struct, FCompi
 	FUserDefinedStructureCompilerUtils::CompileStruct(Struct, Results, true);
 }
 
-FString FKismet2CompilerModule::GenerateCppCodeForEnum(UUserDefinedEnum* UDEnum)
+void FKismet2CompilerModule::GenerateCppCodeForEnum(UUserDefinedEnum* UDEnum, FString& OutHeaderCode, FString& OutCPPCode)
 {
 	TUniquePtr<IBlueprintCompilerCppBackend> Backend_CPP(IBlueprintCompilerCppBackendModuleInterface::Get().Create());
-	return Backend_CPP->GenerateCodeFromEnum(UDEnum);
+	Backend_CPP->GenerateCodeFromEnum(UDEnum, OutHeaderCode, OutCPPCode);
 }
 
 FString FKismet2CompilerModule::GenerateCppCodeForStruct(UUserDefinedStruct* UDStruct)
diff --git a/Engine/Source/Editor/KismetCompiler/Private/KismetCompilerVMBackend.cpp b/Engine/Source/Editor/KismetCompiler/Private/KismetCompilerVMBackend.cpp
--- a/Engine/Source/Editor/KismetCompiler/Private/KismetCompilerVMBackend.cpp
+++ b/Engine/Source/Editor/KismetCompiler/Private/KismetCompilerVMBackend.cpp
@@ -219,6 +219,30 @@ struct FSkipOffsetEmitter
 	}
 };
 
+//////////////////////////////////////////////////////////////////////////
+// FCodeSkipInfo
+
+class FCodeSkipInfo
+{
+public:
+	enum ECodeSkipType
+	{
+		Fixup = 0,
+		InstrumentedDelegateFixup
+	};
+	FCodeSkipInfo(ECodeSkipType TypeIn, FBlueprintCompiledStatement* TargetLabelIn = nullptr, FBlueprintCompiledStatement* SourceLabelIn = nullptr)
+		: Type(TypeIn)
+		, SourceLabel(SourceLabelIn)
+		, TargetLabel(TargetLabelIn)
+	{
+	}
+
+	ECodeSkipType Type;
+	FBlueprintCompiledStatement* SourceLabel;
+	FBlueprintCompiledStatement* TargetLabel;
+	FName DelegateName;
+};
+
 //////////////////////////////////////////////////////////////////////////
 // FScriptBuilderBase
 
@@ -241,7 +265,7 @@ private:
 	FKismetCompilerVMBackend::TStatementToSkipSizeMap& UbergraphStatementLabelMap;
 
 	// Fixup list for jump targets (location to overwrite; jump target)
-	TMap<CodeSkipSizeType, FBlueprintCompiledStatement*> JumpTargetFixupMap;
+	TMap<CodeSkipSizeType, FCodeSkipInfo> JumpTargetFixupMap;
 	
 	// Is this compiling the ubergraph?
 	bool bIsUbergraph;
@@ -768,10 +792,13 @@ public:
 				else if (Struct == TransformStruct)
 				{
 					FTransform T = FTransform::Identity;
-					const bool bParsedUsingCustomFormat = T.InitFromString(Term->Name);
-					if (!bParsedUsingCustomFormat)
+					if (!Term->Name.IsEmpty())
 					{
-						Struct->ImportText(*Term->Name, &T, nullptr, PPF_None, GWarn, GetPathNameSafe(StructProperty));
+						const bool bParsedUsingCustomFormat = T.InitFromString(Term->Name);
+						if (!bParsedUsingCustomFormat)
+						{
+							Struct->ImportText(*Term->Name, &T, nullptr, PPF_None, GWarn, GetPathNameSafe(StructProperty));
+						}
 					}
 					Writer << EX_TransformConst;
 					Writer << T;
@@ -973,7 +1000,7 @@ public:
 				// Emit the literal and queue a fixup to correct it once the address is known
 				Writer << EX_SkipOffsetConst;
 				CodeSkipSizeType PatchUpNeededAtOffset = Writer.EmitPlaceholderSkip();
-				JumpTargetFixupMap.Add(PatchUpNeededAtOffset, TargetLabel);
+				JumpTargetFixupMap.Add(PatchUpNeededAtOffset, FCodeSkipInfo(FCodeSkipInfo::Fixup, TargetLabel));
 			}
 			else if (Prop->GetBoolMetaData(FBlueprintMetadata::MD_LatentCallbackTarget))
 			{
@@ -1248,7 +1275,7 @@ public:
 	{
 		check(Schema && DestinationExpression && !DestinationExpression->Type.PinCategory.IsEmpty());
 
-		const bool bIsArray = DestinationExpression->Type.bIsArray;
+		const bool bIsContainer = DestinationExpression->Type.bIsArray || DestinationExpression->Type.bIsSet || DestinationExpression->Type.bIsMap;
 		const bool bIsDelegate = Schema->PC_Delegate == DestinationExpression->Type.PinCategory;
 		const bool bIsMulticastDelegate = Schema->PC_MCDelegate == DestinationExpression->Type.PinCategory;
 		const bool bIsBoolean = Schema->PC_Boolean == DestinationExpression->Type.PinCategory;
@@ -1256,7 +1283,7 @@ public:
 		const bool bIsAsset = Schema->PC_Asset == DestinationExpression->Type.PinCategory;
 		const bool bIsWeakObjPtr = DestinationExpression->Type.bIsWeakPointer;
 
-		if (bIsArray)
+		if (bIsContainer)
 		{
 			Writer << EX_Let;
 			ensure(DestinationExpression->AssociatedVarProperty);
@@ -1491,6 +1518,53 @@ public:
 		Writer << EX_EndArray;
 	}
 
+	void EmitCreateSetStatement(FBlueprintCompiledStatement& Statement)
+	{
+		Writer << EX_SetSet;
+
+		FBPTerminal* SetTerm = Statement.LHS;
+		EmitTerm(SetTerm);
+		int32 ElementNum = Statement.RHS.Num();
+
+		Writer << ElementNum; // number of elements in the set, used for reserve call
+		
+		USetProperty* SetProperty = CastChecked<USetProperty>(SetTerm->AssociatedVarProperty);
+		UProperty* InnerProperty = SetProperty->ElementProp;
+
+		for(FBPTerminal* Item : Statement.RHS)
+		{
+			EmitTerm(Item, (Item->bIsLiteral ? InnerProperty : NULL));
+		}
+
+		Writer << EX_EndSet;
+	}
+
+	void EmitCreateMapStatement(FBlueprintCompiledStatement& Statement)
+	{
+		Writer << EX_SetMap;
+
+		FBPTerminal* MapTerm = Statement.LHS;
+		EmitTerm(MapTerm);
+		
+		ensureMsgf(Statement.RHS.Num() % 2 == 0, TEXT("Expected even number of key/values whe emitting map statement"));
+		int32 ElementNum = Statement.RHS.Num() / 2;
+
+		Writer << ElementNum;
+
+		UMapProperty* MapProperty = CastChecked<UMapProperty>(MapTerm->AssociatedVarProperty);
+		
+		for(auto MapItemIt = Statement.RHS.CreateIterator(); MapItemIt; ++MapItemIt)
+		{
+			FBPTerminal* Item = *MapItemIt;
+			EmitTerm(Item, (Item->bIsLiteral ? MapProperty->KeyProp : NULL));
+			++MapItemIt;
+			Item = *MapItemIt;
+			EmitTerm(Item, (Item->bIsLiteral ? MapProperty->ValueProp : NULL));
+		}
+
+		Writer << EX_EndMap;
+	}
+
 	void EmitGoto(FBlueprintCompiledStatement& Statement)
 	{
 		if (Statement.Type == KCST_ComputedGoto)
@@ -1508,7 +1582,7 @@ public:
 			CodeSkipSizeType PatchUpNeededAtOffset = Writer.EmitPlaceholderSkip();
 
 			// Queue up a fixup to be done once all label offsets are known
-			JumpTargetFixupMap.Add(PatchUpNeededAtOffset, Statement.TargetLabel);
+			JumpTargetFixupMap.Add(PatchUpNeededAtOffset, FCodeSkipInfo(FCodeSkipInfo::Fixup, Statement.TargetLabel));
 
 			// Now include the boolean expression
 			EmitTerm(Statement.LHS, (UProperty*)(GetDefault<UBoolProperty>()));
@@ -1528,7 +1602,7 @@ public:
 			CodeSkipSizeType PatchUpNeededAtOffset = Writer.EmitPlaceholderSkip();
 
 			// Queue up a fixup to be done once all label offsets are known
-			JumpTargetFixupMap.Add(PatchUpNeededAtOffset, Statement.TargetLabel);
+			JumpTargetFixupMap.Add(PatchUpNeededAtOffset, FCodeSkipInfo(FCodeSkipInfo::Fixup, Statement.TargetLabel));
 		}
 		else if (Statement.Type == KCST_GotoReturn)
 		{
@@ -1537,7 +1611,7 @@ public:
 			CodeSkipSizeType PatchUpNeededAtOffset = Writer.EmitPlaceholderSkip();
 
 			// Queue up a fixup to be done once all label offsets are known
-			JumpTargetFixupMap.Add(PatchUpNeededAtOffset, &ReturnStatement);
+			JumpTargetFixupMap.Add(PatchUpNeededAtOffset, FCodeSkipInfo(FCodeSkipInfo::Fixup, &ReturnStatement));
 		}
 		else if (Statement.Type == KCST_GotoReturnIfNot)
 		{
@@ -1546,7 +1620,7 @@ public:
 			CodeSkipSizeType PatchUpNeededAtOffset = Writer.EmitPlaceholderSkip();
 
 			// Queue up a fixup to be done once all label offsets are known
-			JumpTargetFixupMap.Add(PatchUpNeededAtOffset, &ReturnStatement);
+			JumpTargetFixupMap.Add(PatchUpNeededAtOffset, FCodeSkipInfo(FCodeSkipInfo::Fixup, &ReturnStatement));
 
 			// Now include the boolean expression
 			EmitTerm(Statement.LHS, (UProperty*)(GetDefault<UBoolProperty>()));
@@ -1564,7 +1638,7 @@ public:
 		CodeSkipSizeType PatchUpNeededAtOffset = Writer.EmitPlaceholderSkip();
 
 		// Mark the target for fixup once the addresses have been resolved
-		JumpTargetFixupMap.Add(PatchUpNeededAtOffset, Statement.TargetLabel);
+		JumpTargetFixupMap.Add(PatchUpNeededAtOffset, FCodeSkipInfo(FCodeSkipInfo::Fixup, Statement.TargetLabel));
 	}
 
 	void EmitPopExecState(FBlueprintCompiledStatement& Statement)
@@ -1671,6 +1745,20 @@ public:
 				FName EventName(*Statement.Comment);
 				Writer << EventName;
 			}
+			else if (EventType == EScriptInstrumentation::SuspendState)
+			{
+				if (Statement.TargetLabel->TargetLabel)
+				{
+					CodeSkipSizeType PatchUpNeededAtOffset = Writer.EmitPlaceholderSkip();
+					FCodeSkipInfo CodeSkipInfo(FCodeSkipInfo::InstrumentedDelegateFixup, Statement.TargetLabel->TargetLabel, &Statement);
+					if (Statement.TargetLabel && Statement.TargetLabel->FunctionToCall)
+					{
+						CodeSkipInfo.DelegateName = Statement.TargetLabel->FunctionToCall->GetFName();
+					}
+					// Queue up a fixup to be done once all label offsets are known
+					JumpTargetFixupMap.Add(PatchUpNeededAtOffset, CodeSkipInfo);
+				}
+			}
 		}
 
 		TArray<UEdGraphPin*> PinContextArray(Statement.PureOutputContextArray);
@@ -1772,7 +1860,7 @@ public:
 		Writer << EX_PushExecutionFlow;
 		CodeSkipSizeType PatchUpNeededAtOffset = Writer.EmitPlaceholderSkip();
 
-		JumpTargetFixupMap.Add(PatchUpNeededAtOffset, &ReturnTarget);
+		JumpTargetFixupMap.Add(PatchUpNeededAtOffset, FCodeSkipInfo(FCodeSkipInfo::Fixup, &ReturnTarget));
 	}
 
 	void CloseScript()
@@ -1887,6 +1975,12 @@ public:
 		case KCST_ArrayGetByRef:
 			EmitArrayGetByRef(Statement);
 			break;
+		case KCST_CreateSet:
+			EmitCreateSetStatement(Statement);
+			break;
+		case KCST_CreateMap:
+			EmitCreateMapStatement(Statement);
+			break;
 		default:
 			UE_LOG(LogK2Compiler, Warning, TEXT("VM backend encountered unsupported statement type %d"), (int32)Statement.Type);
 		}
@@ -1895,14 +1989,20 @@ public:
 	// Fix up all jump targets
 	void PerformFixups()
 	{
-		for (TMap<CodeSkipSizeType, FBlueprintCompiledStatement*>::TIterator It(JumpTargetFixupMap); It; ++It)
+		for (TMap<CodeSkipSizeType, FCodeSkipInfo>::TIterator It(JumpTargetFixupMap); It; ++It)
 		{
 			CodeSkipSizeType OffsetToFix = It.Key();
-			FBlueprintCompiledStatement* TargetStatement = It.Value();
+			FCodeSkipInfo& CodeSkipInfo = It.Value();
 
-			CodeSkipSizeType TargetStatementOffset = StatementLabelMap.FindChecked(TargetStatement);
+			CodeSkipSizeType TargetStatementOffset = StatementLabelMap.FindChecked(CodeSkipInfo.TargetLabel);
 
 			Writer.CommitSkip(OffsetToFix, TargetStatementOffset);
+
+			if (CodeSkipInfo.Type == FCodeSkipInfo::InstrumentedDelegateFixup)
+			{
+				// Register delegate entrypoint offsets
+				ClassBeingBuilt->GetDebugData().RegisterEntryPoint(TargetStatementOffset, CodeSkipInfo.DelegateName);
+			}
 		}
 
 		JumpTargetFixupMap.Empty();
@@ -1986,6 +2086,21 @@ void FKismetCompilerVMBackend::ConstructFunction(FKismetFunctionContext& Functio
 	if (bIsUbergraph)
 	{
 		ScriptWriter.CopyStatementMapToUbergraphMap();
+
+		// Create a mapping between compile time generated events and what created them.
+		if (FunctionContext.bInstrumentScriptCode)
+		{
+			FBlueprintDebugData& DebugData = CompilerContext.NewClass->GetDebugData();
+			for (const TPair<UEdGraphPin*, UK2Node_Event*>& EventInfo : CompilerContext.SourcePinToExpansionEvent)
+			{
+				TArray<int32> PinOffsets;
+				DebugData.FindAllCodeLocationsFromSourcePin(EventInfo.Key, Function, PinOffsets);
+				for (const int32 ScriptOffset : PinOffsets)
+				{
+					DebugData.RegisterEntryPoint(ScriptOffset, EventInfo.Value->GetFunctionName());
+				}
+			}
+		}
 	}
 
 	// Make sure we didn't overflow the maximum bytecode size
diff --git a/Engine/Source/Editor/KismetCompiler/Public/BlueprintCompiledStatement.h b/Engine/Source/Editor/KismetCompiler/Public/BlueprintCompiledStatement.h
--- a/Engine/Source/Editor/KismetCompiler/Public/BlueprintCompiledStatement.h
+++ b/Engine/Source/Editor/KismetCompiler/Public/BlueprintCompiledStatement.h
@@ -53,6 +53,8 @@ enum EKismetCompiledStatementType
 	KCST_InstrumentedTunnelEndOfThread,	// Instrumented tunnel exit
 	//
 	KCST_ArrayGetByRef,
+	KCST_CreateSet,
+	KCST_CreateMap,
 };
 
 //@TODO: Too rigid / icky design
diff --git a/Engine/Source/Editor/KismetCompiler/Public/BlueprintCompilerCppBackendInterface.h b/Engine/Source/Editor/KismetCompiler/Public/BlueprintCompilerCppBackendInterface.h
--- a/Engine/Source/Editor/KismetCompiler/Public/BlueprintCompilerCppBackendInterface.h
+++ b/Engine/Source/Editor/KismetCompiler/Public/BlueprintCompilerCppBackendInterface.h
@@ -13,8 +13,8 @@ class UUserDefinedStruct;
 class IBlueprintCompilerCppBackend
 {
 public:
-	virtual FString GenerateCodeFromClass(UClass* SourceClass, TIndirectArray<FKismetFunctionContext>& Functions, bool bGenerateStubsOnly, FString& OutCppBody) = 0;
-	virtual FString GenerateCodeFromEnum(UUserDefinedEnum* SourceEnum) = 0;
+	virtual FString GenerateCodeFromClass(UClass* SourceClass, TIndirectArray<FKismetFunctionContext>& Functions, bool bGenerateStubsOnly, FCompilerNativizationOptions NativizationOptions, FString& OutCppBody) = 0;
+	virtual void GenerateCodeFromEnum(UUserDefinedEnum* SourceEnum, FString& OutHeaderCode, FString& OutCPPCode) = 0;
 	virtual FString GenerateCodeFromStruct(UUserDefinedStruct* SourceStruct) = 0;
 
 	// Generate a wrapper class, that helps accessing non-native properties and calling non-native functions
diff --git a/Engine/Source/Editor/KismetCompiler/Public/KismetCompiledFunctionContext.h b/Engine/Source/Editor/KismetCompiler/Public/KismetCompiledFunctionContext.h
--- a/Engine/Source/Editor/KismetCompiler/Public/KismetCompiledFunctionContext.h
+++ b/Engine/Source/Editor/KismetCompiler/Public/KismetCompiledFunctionContext.h
@@ -320,10 +320,9 @@ public:
 		}
 		else
 		{
-			MessageLog.Warning(
-				*FString::Printf(TEXT("A node that generated no code of it's own (@@) tried to inject code into @@. %s"), *GetFullNameSafe(Source))
-				, Source
-				, Destination);
+			// A node that generated no code of it's own (Source) tried to inject code into (Destination).
+			// It is ok, for example: UK2Node_GetClassDefaults works like this.
+			// Moreover when KismetCompilerDebugOptions::EmitNodeComments is enabled there is always a Comment state generated anyway.
 		}
 	}
 
diff --git a/Engine/Source/Editor/KismetCompiler/Public/KismetCompilerModule.h b/Engine/Source/Editor/KismetCompiler/Public/KismetCompilerModule.h
--- a/Engine/Source/Editor/KismetCompiler/Public/KismetCompilerModule.h
+++ b/Engine/Source/Editor/KismetCompiler/Public/KismetCompilerModule.h
@@ -91,35 +91,10 @@ public:
 	 */
 	virtual void GetBlueprintTypesForClass(UClass* ParentClass, UClass*& OutBlueprintClass, UClass*& OutBlueprintGeneratedClass) const = 0;
 
-	virtual FString GenerateCppCodeForEnum(UUserDefinedEnum* UDEnum) = 0;
+	virtual void GenerateCppCodeForEnum(UUserDefinedEnum* UDEnum, FString& OutHeaderCode, FString& OutCPPCode) = 0;
 	virtual FString GenerateCppCodeForStruct(UUserDefinedStruct* UDStruct) = 0;
 	// Generate a wrapper class, that helps accessing non-native properties and calling non-native functions
 	virtual FString GenerateCppWrapper(UBlueprintGeneratedClass* BPGC) = 0;
 };
 
-//////////////////////////////////////////////////////////////////////////
-// FKismet2CompilerModule
-
-/**
- * The Kismet 2 Compiler module
- */
-class FKismet2CompilerModule : public IKismetCompilerInterface
-{
-public:
-	// Implementation of the IKismetCompilerInterface
-	virtual void CompileBlueprint(class UBlueprint* Blueprint, const FKismetCompilerOptions& CompileOptions, FCompilerResultsLog& Results, TSharedPtr<class FBlueprintCompileReinstancer> ParentReinstancer = NULL, TArray<UObject*>* ObjLoaded = NULL) override;
-	virtual void CompileStructure(class UUserDefinedStruct* Struct, FCompilerResultsLog& Results) override;
-	virtual void RecoverCorruptedBlueprint(class UBlueprint* Blueprint) override;
-	virtual void RemoveBlueprintGeneratedClasses(class UBlueprint* Blueprint) override;
-	virtual TArray<IBlueprintCompiler*>& GetCompilers() override { return Compilers; }
-	virtual void GetBlueprintTypesForClass(UClass* ParentClass, UClass*& OutBlueprintClass, UClass*& OutBlueprintGeneratedClass) const override;
-	virtual FString GenerateCppCodeForEnum(UUserDefinedEnum* UDEnum) override;
-	virtual FString GenerateCppCodeForStruct(UUserDefinedStruct* UDStruct) override;
-	virtual FString GenerateCppWrapper(UBlueprintGeneratedClass* BPGC) override;
-	// End implementation
-private:
-	void CompileBlueprintInner(class UBlueprint* Blueprint, const FKismetCompilerOptions& CompileOptions, FCompilerResultsLog& Results, TSharedPtr<FBlueprintCompileReinstancer> Reinstancer, TArray<UObject*>* ObjLoaded);
-
-	TArray<IBlueprintCompiler*> Compilers;
-};
 
diff --git a/Engine/Source/Editor/MainFrame/Private/Frame/MainFrameActions.cpp b/Engine/Source/Editor/MainFrame/Private/Frame/MainFrameActions.cpp
--- a/Engine/Source/Editor/MainFrame/Private/Frame/MainFrameActions.cpp
+++ b/Engine/Source/Editor/MainFrame/Private/Frame/MainFrameActions.cpp
@@ -526,7 +526,7 @@ void FMainFrameActionCallbacks::PackageProject( const FName InPlatformInfoName )
 				UnrecoverableError = true;
 			}
 
-			if ((Result & ETargetPlatformReadyStatus::RemoveServerNameEmpty) != 0 && (bProjectHasCode || (!FApp::GetEngineIsPromotedBuild() && !FApp::IsEngineInstalled()) || PackagingSettings->bNativizeBlueprintAssets))
+			if ((Result & ETargetPlatformReadyStatus::RemoveServerNameEmpty) != 0 && (bProjectHasCode || (!FApp::GetEngineIsPromotedBuild() && !FApp::IsEngineInstalled()) || PackagingSettings->BlueprintNativizationMethod != EProjectPackagingBlueprintNativizationMethod::Disabled))
 			{
 				AddMessageLog(
 					LOCTEXT("RemoveServerNameNotFound", "Remote compiling requires a server name. "),
@@ -626,7 +626,7 @@ void FMainFrameActionCallbacks::PackageProject( const FName InPlatformInfoName )
 		OptionalParams += TEXT(" -nodebuginfo");
 	}
 
-	if (PackagingSettings->bNativizeBlueprintAssets)
+	if (PackagingSettings->BlueprintNativizationMethod != EProjectPackagingBlueprintNativizationMethod::Disabled)
 	{
 		OptionalParams += TEXT(" -nativizeAssets");
 	}
@@ -661,7 +661,7 @@ void FMainFrameActionCallbacks::PackageProject( const FName InPlatformInfoName )
 	}
 
 	// only build if the project has code that might need to be built
-	if (bProjectHasCode || (!FApp::GetEngineIsPromotedBuild() && !FApp::IsEngineInstalled()) || PackagingSettings->bNativizeBlueprintAssets)
+	if (bProjectHasCode || (!FApp::GetEngineIsPromotedBuild() && !FApp::IsEngineInstalled()) || PackagingSettings->BlueprintNativizationMethod != EProjectPackagingBlueprintNativizationMethod::Disabled)
 	{
 		OptionalParams += TEXT(" -build");
 	}
diff --git a/Engine/Source/Editor/PropertyEditor/Private/SDetailsViewBase.cpp b/Engine/Source/Editor/PropertyEditor/Private/SDetailsViewBase.cpp
--- a/Engine/Source/Editor/PropertyEditor/Private/SDetailsViewBase.cpp
+++ b/Engine/Source/Editor/PropertyEditor/Private/SDetailsViewBase.cpp
@@ -140,7 +140,8 @@ TArray< FPropertyPath > SDetailsViewBase::GetPropertiesInOrderDisplayed() const
 	return Ret;
 }
 
-static TSharedPtr< IDetailTreeNode > FindTreeNodeFromPropertyRecursive( const TArray< TSharedRef<IDetailTreeNode> >& Nodes, const FPropertyPath& Property )
+// @return populates OutNodes with the leaf node corresponding to property as the first entry in the list (e.g. [leaf, parent, grandparent]):
+static void FindTreeNodeFromPropertyRecursive( const TArray< TSharedRef<IDetailTreeNode> >& Nodes, const FPropertyPath& Property, TArray< TSharedPtr< IDetailTreeNode > >& OutNodes )
 {
 	for (auto& TreeNode : Nodes)
 	{
@@ -149,22 +150,21 @@ static TSharedPtr< IDetailTreeNode > FindTreeNodeFromPropertyRecursive( const TA
 			FPropertyPath tmp = TreeNode->GetPropertyPath();
 			if( Property == tmp )
 			{
-				return TreeNode;
+				OutNodes.Push(TreeNode);
+				return;
 			}
 		}
-		else
+
+		// Need to check children even if we're a leaf, because all DetailItemNodes are leaves, even if they may have sub-children
+		TArray< TSharedRef<IDetailTreeNode> > Children;
+		TreeNode->GetChildren(Children);
+		FindTreeNodeFromPropertyRecursive(Children, Property, OutNodes);
+		if (OutNodes.Num() > 0)
 		{
-			TArray< TSharedRef<IDetailTreeNode> > Children;
-			TreeNode->GetChildren(Children);
-			auto Result = FindTreeNodeFromPropertyRecursive(Children, Property);
-			if( Result.IsValid() )
-			{
-				return Result;
-			}
+			OutNodes.Push(TreeNode);
+			return;
 		}
 	}
-
-	return TSharedPtr< IDetailTreeNode >();
 }
 
 void SDetailsViewBase::HighlightProperty(const FPropertyPath& Property)
@@ -175,18 +175,25 @@ void SDetailsViewBase::HighlightProperty(const FPropertyPath& Property)
 		PrevHighlightedNodePtr->SetIsHighlighted(false);
 	}
 
-	auto NextNodePtr = FindTreeNodeFromPropertyRecursive( RootTreeNodes, Property );
-	if (NextNodePtr.IsValid())
+	TSharedPtr< IDetailTreeNode > FinalNodePtr = nullptr;
+	TArray< TSharedPtr< IDetailTreeNode > > TreeNodeChain;
+	FindTreeNodeFromPropertyRecursive(RootTreeNodes, Property, TreeNodeChain);
+	if (TreeNodeChain.Num() > 0)
 	{
-		NextNodePtr->SetIsHighlighted(true);
-		auto ParentCategory = NextNodePtr->GetParentCategory();
-		if (ParentCategory.IsValid())
+		FinalNodePtr = TreeNodeChain[0];
+		check(FinalNodePtr.IsValid());
+		FinalNodePtr->SetIsHighlighted(true);
+
+		for (int ParentIndex = 1; ParentIndex < TreeNodeChain.Num(); ++ParentIndex)
 		{
-			DetailTree->SetItemExpansion(ParentCategory.ToSharedRef(), true);
+			TSharedPtr< IDetailTreeNode > CurrentParent = TreeNodeChain[ParentIndex];
+			check(CurrentParent.IsValid());
+			DetailTree->SetItemExpansion(CurrentParent.ToSharedRef(), true);
 		}
-		DetailTree->RequestScrollIntoView(NextNodePtr.ToSharedRef());
+
+		DetailTree->RequestScrollIntoView(FinalNodePtr.ToSharedRef());
 	}
-	CurrentlyHighlightedNode = NextNodePtr;
+	CurrentlyHighlightedNode = FinalNodePtr;
 }
 
 void SDetailsViewBase::ShowAllAdvancedProperties()
@@ -1344,8 +1351,11 @@ void SDetailsViewBase::UpdateFilteredDetails()
 	NumVisbleTopLevelObjectNodes = 0;
 	FRootPropertyNodeList& RootPropertyNodes = GetRootNodes();
 
-	CurrentFilter.bShowAllAdvanced = GetDefault<UEditorStyleSettings>()->bShowAllAdvancedDetails;
-
+	if( GetDefault<UEditorStyleSettings>()->bShowAllAdvancedDetails )
+	{
+		CurrentFilter.bShowAllAdvanced = true;
+	}
+	
 	for(int32 RootNodeIndex = 0; RootNodeIndex < RootPropertyNodes.Num(); ++RootNodeIndex)
 	{
 		TSharedPtr<FComplexPropertyNode>& RootPropertyNode = RootPropertyNodes[RootNodeIndex];
diff --git a/Engine/Source/Editor/UnrealEd/Classes/Settings/ProjectPackagingSettings.h b/Engine/Source/Editor/UnrealEd/Classes/Settings/ProjectPackagingSettings.h
--- a/Engine/Source/Editor/UnrealEd/Classes/Settings/ProjectPackagingSettings.h
+++ b/Engine/Source/Editor/UnrealEd/Classes/Settings/ProjectPackagingSettings.h
@@ -25,7 +25,7 @@ enum EProjectPackagingBuildConfigurations
 };
 
 /**
- * Enumerates the the available internationalization data presets for project packaging.
+ * Enumerates the available internationalization data presets for project packaging.
  */
 UENUM()
 enum class EProjectPackagingInternationalizationPresets : uint8
@@ -46,6 +46,22 @@ enum class EProjectPackagingInternationalizationPresets : uint8
 	All
 };
 
+/**
+* Enumerates the available methods for Blueprint nativization during project packaging.
+*/
+UENUM()
+enum class EProjectPackagingBlueprintNativizationMethod : uint8
+{
+	/** Disable Blueprint nativization (default). */
+	Disabled,
+
+	/** Enable nativization for all Blueprint assets. */
+	Inclusive,
+
+	/** Enable nativization for selected Blueprint assets only. */
+	Exclusive
+};
+
 /**
  * Implements the Editor's user settings.
  */
@@ -86,12 +102,16 @@ public:
 	bool IncludeDebugFiles;
 
 	/** If enabled, then the project's Blueprint assets (including structs and enums) will be intermediately converted into C++ and used in the packaged project (in place of the .uasset files).*/
-	UPROPERTY(config, EditAnywhere, Category=Experimental)
-	bool bNativizeBlueprintAssets;
+	UPROPERTY(config, EditAnywhere, Category = Blueprints)
+	EProjectPackagingBlueprintNativizationMethod BlueprintNativizationMethod;
+
+	/** List of Blueprints to include for nativization when using the exclusive method. */
+	UPROPERTY(config, EditAnywhere, AdvancedDisplay, Category = Blueprints, meta = (DisplayName = "List of Blueprint assets to nativize", RelativeToGameContentDir, LongPackageName))
+	TArray<FFilePath> NativizeBlueprintAssets;
 
-	/** WHen nativization is enabled, only Blueprints with "Nativize" set true will be converted to C++. */
-	UPROPERTY(config, EditAnywhere, Category = Experimental)
-	bool bNativizeOnlySelectedBlueprints;
+	/** If enabled, a warning will be emitted at build/cook time if nativization is turned on in the Project Settings, but the nativization flag was omitted from the command line. */
+	UPROPERTY(config, EditAnywhere, Category = Blueprints)
+	bool bWarnIfPackagedWithoutNativizationFlag;
 
 	/** If enabled, all content will be put into a single .pak file instead of many individual files (default = enabled). */
 	UPROPERTY(config, EditAnywhere, Category=Packaging)
@@ -226,12 +246,27 @@ public:
 	UPROPERTY(config, EditAnywhere, Category=Packaging, AdvancedDisplay, meta=(DisplayName="Additional Non-Asset Directories To Copy", RelativeToGameContentDir))
 	TArray<FDirectoryPath> DirectoriesToAlwaysStageAsNonUFS;	
 
+private:
+	/** Helper array used to mirror Blueprint asset selections across edits */
+	TArray<FFilePath> CachedNativizeBlueprintAssets;
+
+	UPROPERTY(config)
+	bool bNativizeBlueprintAssets_DEPRECATED;
 
+	UPROPERTY(config)
+	bool bNativizeOnlySelectedBlueprints_DEPRECATED;
 	
 public:
 
 	// UObject Interface
 
+	virtual void PostInitProperties() override;
 	virtual void PostEditChangeProperty( struct FPropertyChangedEvent& PropertyChangedEvent ) override;
 	virtual bool CanEditChange( const UProperty* InProperty ) const override;
+
+	/** Adds the given Blueprint asset to the exclusive nativization list. */
+	bool AddBlueprintAssetToNativizationList(const class UBlueprint* InBlueprint);
+
+	/** Removes the given Blueprint asset from the exclusive nativization list. */
+	bool RemoveBlueprintAssetFromNativizationList(const class UBlueprint* InBlueprint);
 };
diff --git a/Engine/Source/Editor/UnrealEd/Private/Commandlets/NativeCodeGenCommandlet.cpp b/Engine/Source/Editor/UnrealEd/Private/Commandlets/NativeCodeGenCommandlet.cpp
--- a/Engine/Source/Editor/UnrealEd/Private/Commandlets/NativeCodeGenCommandlet.cpp
+++ b/Engine/Source/Editor/UnrealEd/Private/Commandlets/NativeCodeGenCommandlet.cpp
@@ -23,12 +23,17 @@ int32 UNativeCodeGenCommandlet::Main(const FString& Params)
 		return 0;
 	}
 
-	TArray< TPair< FString, FString > > CodeGenTargets;
+	TArray<FPlatformNativizationDetails> CodeGenTargets;
 	{
 		for (auto& Platform : Platforms)
 		{
+			FPlatformNativizationDetails PlatformNativizationDetails;
+			PlatformNativizationDetails.PlatformName = Platform;
 			// If you change this target path you must also update logic in CookCommand.Automation.CS. Passing a single directory around is cumbersome for testing, so I have hard coded it.
-			CodeGenTargets.Push(TPairInitializer<FString, FString>(Platform, FString(FPaths::Combine(*FPaths::GameIntermediateDir(), *Platform))));
+			PlatformNativizationDetails.PlatformTargetDirectory = FString(FPaths::Combine(*FPaths::GameIntermediateDir(), *Platform));
+			// CompilerNativizationOptions will be serialized from saved manifest
+
+			CodeGenTargets.Push(PlatformNativizationDetails);
 		}
 	}
 
