Fixing non-unity CIS errors (missing includes, fwd declarations, etc.)

diff --git a/Engine/Source/Editor/Kismet/Private/BlueprintCompilationManager.cpp b/Engine/Source/Editor/Kismet/Private/BlueprintCompilationManager.cpp
--- a/Engine/Source/Editor/Kismet/Private/BlueprintCompilationManager.cpp
+++ b/Engine/Source/Editor/Kismet/Private/BlueprintCompilationManager.cpp
@@ -4,6 +4,7 @@
 #include "BlueprintEditorUtils.h"
 #include "Kismet2/KismetEditorUtilities.h"
 #include "Kismet2/KismetReinstanceUtilities.h"
+#include "TickableEditorObject.h"
 
 struct FBlueprintCompilationManagerImpl : public FGCObject, FTickableEditorObject
 {
@@ -109,11 +110,6 @@ void FBlueprintCompilationManager::Shutdown()
 	BPCMImpl = nullptr;
 }
 
-FBlueprintCompilationManagerImpl* FBlueprintCompilationManager::DebugAccessImpl()
-{
-	return BPCMImpl;
-}
-
 // Forward to impl:
 void FBlueprintCompilationManager::QueueForCompilation(UBlueprint* BPToCompile)
 {
diff --git a/Engine/Source/Editor/Kismet/Public/BlueprintCompilationManager.h b/Engine/Source/Editor/Kismet/Public/BlueprintCompilationManager.h
--- a/Engine/Source/Editor/Kismet/Public/BlueprintCompilationManager.h
+++ b/Engine/Source/Editor/Kismet/Public/BlueprintCompilationManager.h
@@ -2,13 +2,12 @@
 
 #pragma once
 
-struct FBlueprintCompilationManagerImpl;
+class UBlueprint;
 
 struct KISMET_API FBlueprintCompilationManager
 {
 	static void Initialize();
 	static void Shutdown();
-	static FBlueprintCompilationManagerImpl* DebugAccessImpl();
 
 	static void QueueForCompilation(UBlueprint* BPToCompile);
 	static void NotifyBlueprintLoaded(UBlueprint* BPLoaded);
