CIS fix - missing DLLEXPORT (_API) definition, added include

diff --git a/Engine/Source/Editor/Kismet/Public/BlueprintCompilationManager.h b/Engine/Source/Editor/Kismet/Public/BlueprintCompilationManager.h
--- a/Engine/Source/Editor/Kismet/Public/BlueprintCompilationManager.h
+++ b/Engine/Source/Editor/Kismet/Public/BlueprintCompilationManager.h
@@ -2,6 +2,8 @@
 
 #pragma once
 
+#include "HAL/Platform.h" // for DLLEXPORT (KISMET_API)
+
 class UBlueprint;
 
 struct KISMET_API FBlueprintCompilationManager
