CIS fix - Replacing include (for prior non-unity CIS errors) with CoreMinimal.h. Other include relied on some defines that weren't included.

diff --git a/Engine/Source/Editor/Kismet/Public/BlueprintCompilationManager.h b/Engine/Source/Editor/Kismet/Public/BlueprintCompilationManager.h
--- a/Engine/Source/Editor/Kismet/Public/BlueprintCompilationManager.h
+++ b/Engine/Source/Editor/Kismet/Public/BlueprintCompilationManager.h
@@ -2,7 +2,7 @@
 
 #pragma once
 
-#include "HAL/Platform.h" // for DLLEXPORT (KISMET_API)
+#include "CoreMinimal.h" // for DLLEXPORT (KISMET_API)
 
 class UBlueprint;
 
