Copying //UE4/Dev-Framework to //UE4/Dev-Main (Source: //UE4/Dev-Framework @ 3431384)

diff --git a/Engine/Source/Editor/GraphEditor/Public/SGraphPin.h b/Engine/Source/Editor/GraphEditor/Public/SGraphPin.h
--- a/Engine/Source/Editor/GraphEditor/Public/SGraphPin.h
+++ b/Engine/Source/Editor/GraphEditor/Public/SGraphPin.h
@@ -137,18 +137,12 @@ public:
 	/** @return whether this pin is a map value */
 	bool IsMap() const;
 
-	/** @return whether this pin is passed by ref */
-	bool IsByRef() const;
-
 	/** @return whether this pin is passed by mutable ref */
 	bool IsByMutableRef() const;
 
 	/** @return whether this pin is passed by mutable ref */
 	bool IsDelegate() const;
 
-	/** @return whether this pin is passed by const ref */
-	bool IsByConstRef() const;
-
 	/** @return whether this pin is connected to another pin */
 	bool IsConnected() const;
 
@@ -172,6 +166,9 @@ public:
 		PinColorModifier = InColor;
 	}
 
+	/** Set this pin to only be used to display default value */
+	void SetOnlyShowDefaultValue(bool bNewOnlyShowDefaultValue);
+
 	/** If pin in node is visible at all */
 	EVisibility IsPinVisibleAsAdvanced() const;
 
@@ -181,6 +178,9 @@ public:
 	/** Returns whether or not this pin is currently connectable */
 	bool GetIsConnectable() const;
 
+	/** Get the widget we should put into the 'default value' space, shown when nothing connected */
+	virtual TSharedRef<SWidget>	GetDefaultValueWidget();
+
 protected:
 	FText GetPinLabel() const;
 
@@ -190,9 +190,6 @@ protected:
 		return bShowLabel ? EVisibility::Visible : EVisibility::Collapsed;
 	}
 
-	/** Get the widget we should put into the 'default value' space, shown when nothing connected */
-	virtual TSharedRef<SWidget>	GetDefaultValueWidget();
-
 	/** @return The brush with which to paint this graph pin's incoming/outgoing bullet point */
 	virtual const FSlateBrush* GetPinIcon() const;
 
@@ -255,6 +252,9 @@ protected:
 	/** If we should draw the label on this pin */
 	bool bShowLabel;
 
+	/** If we should draw the label on this pin */
+	bool bOnlyShowDefaultValue;
+
 	/** true when we're moving links between pins. */
 	bool bIsMovingLinks;
 
diff --git a/Engine/Source/Editor/Kismet/Private/BitmaskLiteralDetails.cpp b/Engine/Source/Editor/Kismet/Private/BitmaskLiteralDetails.cpp
--- a/Engine/Source/Editor/Kismet/Private/BitmaskLiteralDetails.cpp
+++ b/Engine/Source/Editor/Kismet/Private/BitmaskLiteralDetails.cpp
@@ -96,7 +96,7 @@ void FBitmaskLiteralDetails::OnBitmaskEnumTypeChanged(TSharedPtr<FString> ItemSe
 	{
 		if (const UEdGraphSchema_K2* K2Schema = Cast<UEdGraphSchema_K2>(InputPin->GetSchema()))
 		{
-			K2Schema->SetPinDefaultValueBasedOnType(InputPin);
+			K2Schema->SetPinAutogeneratedDefaultValueBasedOnType(InputPin);
 		}
 	}
 }
diff --git a/Engine/Source/Editor/Kismet/Private/BlueprintCompilationManager.cpp b/Engine/Source/Editor/Kismet/Private/BlueprintCompilationManager.cpp
--- a/Engine/Source/Editor/Kismet/Private/BlueprintCompilationManager.cpp
+++ b/Engine/Source/Editor/Kismet/Private/BlueprintCompilationManager.cpp
@@ -6,10 +6,13 @@
 #include "Blueprint/BlueprintSupport.h"
 #include "CompilerResultsLog.h"
 #include "Components/TimelineComponent.h"
+#include "Editor.h"
 #include "Engine/Engine.h"
+#include "Engine/LevelScriptBlueprint.h"
 #include "Engine/SCS_Node.h"
 #include "Engine/SimpleConstructionScript.h"
 #include "Engine/TimelineTemplate.h"
+#include "FindInBlueprintManager.h"
 #include "K2Node_CustomEvent.h"
 #include "K2Node_FunctionEntry.h"
 #include "K2Node_FunctionResult.h"
@@ -24,120 +27,48 @@
 #include "UObject/UObjectHash.h"
 #include "WidgetBlueprint.h"
 
+/*
+	BLUEPRINT COMPILATION MANAGER IMPLEMENTATION NOTES
+
+	INPUTS: UBlueprint, UEdGraph, UEdGraphNode, UEdGraphPin, references to UClass, UProperties
+	INTERMEDIATES: Cloned Graph, Nodes, Pins
+	OUPUTS: UClass, UProperties
+
+	The blueprint compilation manager addresses shortcomings of compilation 
+	behavior (performance, correctness) that occur when compiling blueprints 
+	that are inter-dependent. If you are using blueprints and there are no dependencies
+	between blueprint compilation outputs and inputs, then this code is completely
+	unnecessary and you can directly interface with FKismetCompilerContext and its
+	derivatives.
+
+	In order to handle compilation correctly the manager splits compilation into
+	the following stages (implemented below in FlushCompilationQueueImpl):
+
+	STAGE I: GATHER
+	STAGE II: FILTER
+	STAGE III: SORT
+	STAGE IV: SET TEMPORARY BLUEPRINT FLAGS
+	STAGE V: VALIDATE
+	STAGE VI: PURGE (LOAD ONLY)
+	STAGE VII: DISCARD SKELETON CDO
+	STAGE VIII: RECOMPILE SKELETON
+	STAGE IX: RECONSTRUCT NODES, REPLACE DEPRECATED NODES (LOAD ONLY)
+	STAGE X: CREATE REINSTANCER (DISCARD 'OLD' CLASS)
+	STAGE XI: CREATE UPDATED CLASS HIERARCHY
+	STAGE XII: COMPILE CLASS LAYOUT
+	STAGE XIII: COMPILE CLASS FUNCTIONS
+	STAGE XIV: CLEAR TEMPORARY FLAGS
+	STAGE XV: REINSTANCE
+
+	The code that implements these stages are labeled below. At some later point a final
+	reinstancing operation will occur, unless the client is using CompileSynchronously, 
+	in which case the expensive object graph find and replace will occur immediately
+*/
+
 // Debugging switches:
 #define VERIFY_NO_STALE_CLASS_REFERENCES 0
 #define VERIFY_NO_BAD_SKELETON_REFERENCES 0
 
-class FFixupBytecodeReferences : public FArchiveUObject
-{
-public:
-	FFixupBytecodeReferences(UObject* InObject)
-	{
-		ArIsObjectReferenceCollector = true;
-		
-		InObject->Serialize(*this);
-		class FArchiveProxyCollector : public FReferenceCollector
-		{
-			/** Archive we are a proxy for */
-			FArchive& Archive;
-		public:
-			FArchiveProxyCollector(FArchive& InArchive)
-				: Archive(InArchive)
-			{
-			}
-			virtual void HandleObjectReference(UObject*& Object, const UObject* ReferencingObject, const UProperty* ReferencingProperty) override
-			{
-				Archive << Object;
-			}
-			virtual void HandleObjectReferences(UObject** InObjects, const int32 ObjectNum, const UObject* InReferencingObject, const UProperty* InReferencingProperty) override
-			{
-				for (int32 ObjectIndex = 0; ObjectIndex < ObjectNum; ++ObjectIndex)
-				{
-					UObject*& Object = InObjects[ObjectIndex];
-					Archive << Object;
-				}
-			}
-			virtual bool IsIgnoringArchetypeRef() const override
-			{
-				return false;
-			}
-			virtual bool IsIgnoringTransient() const override
-			{
-				return false;
-			}
-		} ArchiveProxyCollector(*this);
-		
-		InObject->GetClass()->CallAddReferencedObjects(InObject, ArchiveProxyCollector);
-	}
-
-	FArchive& operator<<( UObject*& Obj )
-	{
-		if (Obj != NULL)
-		{
-			if(UClass* RelatedClass = Cast<UClass>(Obj))
-			{
-				UClass* NewClass = RelatedClass->GetAuthoritativeClass();
-				ensure(NewClass);
-				if(NewClass != RelatedClass)
-				{
-					Obj = NewClass;
-				}
-			}
-			else if(UField* AsField = Cast<UField>(Obj))
-			{
-				UClass* OwningClass = AsField->GetOwnerClass();
-				if(OwningClass)
-				{
-					UClass* NewClass = OwningClass->GetAuthoritativeClass();
-					ensure(NewClass);
-					if(NewClass != OwningClass)
-					{
-						// drill into new class finding equivalent object:
-						TArray<FName> Names;
-						UObject* Iter = Obj;
-						while (Iter && Iter != OwningClass)
-						{
-							Names.Add(Iter->GetFName());
-							Iter = Iter->GetOuter();
-						}
-
-						UObject* Owner = NewClass;
-						UObject* Match = nullptr;
-						for(int32 I = Names.Num() - 1; I >= 0; --I)
-						{
-							UObject* Next = StaticFindObjectFast( UObject::StaticClass(), Owner, Names[I]);
-							if( Next )
-							{
-								if(I == 0)
-								{
-									Match = Next;
-								}
-								else
-								{
-									Owner = Match;
-								}
-							}
-							else
-							{
-								break;
-							}
-						}
-						
-						if(Match)
-						{
-							Obj = Match;
-						}
-					}
-				}
-			}
-		}
-		return *this;
-	}
-};
-
-// free function that we use to cross a module boundary (from CoreUObject to here)
-void FlushReinstancingQueueImplWrapper();
-
 struct FBlueprintCompilationManagerImpl : public FGCObject
 {
 	FBlueprintCompilationManagerImpl();
@@ -146,19 +77,40 @@ struct FBlueprintCompilationManagerImpl : public FGCObject
 	// FGCObject:
 	virtual void AddReferencedObjects(FReferenceCollector& Collector);
 
-	void QueueForCompilation(UBlueprint* BPToCompile);
-	void FlushCompilationQueueImpl(TArray<UObject*>* ObjLoaded);
+	void QueueForCompilation(const FBPCompileRequest& CompileJob);
+	void CompileSynchronouslyImpl(const FBPCompileRequest& Request);
+	void FlushCompilationQueueImpl(TArray<UObject*>* ObjLoaded, bool bSuppressBroadcastCompiled);
 	void FlushReinstancingQueueImpl();
 	bool HasBlueprintsToCompile() const;
 	bool IsGeneratedClassLayoutReady() const;
 	static void ReinstanceBatch(TArray<TUniquePtr<FBlueprintCompileReinstancer>>& Reinstancers, TMap< UClass*, UClass* >& InOutOldToNewClassMap, TArray<UObject*>* ObjLoaded);
 	static UClass* FastGenerateSkeletonClass(UBlueprint* BP, FKismetCompilerContext& CompilerContext);
+	static bool IsQueuedForCompilation(UBlueprint* BP);
+	
+	// Declaration of archive to fix up bytecode references of blueprints that are actively compiled:
+	class FFixupBytecodeReferences : public FArchiveUObject
+	{
+	public:
+		FFixupBytecodeReferences(UObject* InObject);
 
-	TArray<UBlueprint*> QueuedBPs;
+	private:
+		virtual FArchive& operator<<( UObject*& Obj ) override;
+	};
+
+	// Queued requests to be processed in the next FlushCompilationQueueImpl call:
+	TArray<FBPCompileRequest> QueuedRequests;
+	
+	// Data stored for reinstancing, which finishes much later than compilation,
+	// populated by FlushCompilationQueueImpl, cleared by FlushReinstancingQueueImpl:
 	TMap<UClass*, UClass*> ClassesToReinstance;
+
+	// State stored so that we can check what stage of compilation we're in:
 	bool bGeneratedClassLayoutReady;
 };
 
+// free function that we use to cross a module boundary (from CoreUObject to here)
+void FlushReinstancingQueueImplWrapper();
+
 FBlueprintCompilationManagerImpl::FBlueprintCompilationManagerImpl()
 {
 	FBlueprintSupport::SetFlushReinstancingQueueFPtr(&FlushReinstancingQueueImplWrapper);
@@ -172,63 +124,149 @@ FBlueprintCompilationManagerImpl::~FBlueprintCompilationManagerImpl()
 
 void FBlueprintCompilationManagerImpl::AddReferencedObjects(FReferenceCollector& Collector)
 {
-	Collector.AddReferencedObjects(QueuedBPs);
+	for( FBPCompileRequest& Job : QueuedRequests )
+	{
+		Collector.AddReferencedObject(Job.BPToCompile);
+	}
+}
+
+void FBlueprintCompilationManagerImpl::QueueForCompilation(const FBPCompileRequest& CompileJob)
+{
+	if(!CompileJob.BPToCompile->bQueuedForCompilation)
+	{
+		CompileJob.BPToCompile->bQueuedForCompilation = true;
+		QueuedRequests.Add(CompileJob);
+	}
 }
 
-void FBlueprintCompilationManagerImpl::QueueForCompilation(UBlueprint* BPToCompile)
+void FBlueprintCompilationManagerImpl::CompileSynchronouslyImpl(const FBPCompileRequest& Request)
 {
-	QueuedBPs.Add(BPToCompile);
+	Request.BPToCompile->bQueuedForCompilation = true;
+
+	const bool bIsRegeneratingOnLoad		= (Request.CompileOptions & EBlueprintCompileOptions::IsRegeneratingOnLoad		) != EBlueprintCompileOptions::None;
+	const bool bSkipGarbageCollection		= (Request.CompileOptions & EBlueprintCompileOptions::SkipGarbageCollection		) != EBlueprintCompileOptions::None;
+	const bool bBatchCompile				= (Request.CompileOptions & EBlueprintCompileOptions::BatchCompile				) != EBlueprintCompileOptions::None;
+	const bool bSkipReinstancing			= (Request.CompileOptions & EBlueprintCompileOptions::SkipReinstancing			) != EBlueprintCompileOptions::None;
+
+	// Wipe the PreCompile log, any generated messages are now irrelevant
+	Request.BPToCompile->PreCompileLog.Reset();
+
+	// Reset the flag, so if the user tries to use PIE it will warn them if the BP did not compile
+	Request.BPToCompile->bDisplayCompilePIEWarning = true;
+	
+	// Do not want to run this code without the editor present nor when running commandlets.
+	if (GEditor && GIsEditor)
+	{
+		// We do not want to regenerate a search Guid during loads, nothing has changed in the Blueprint and it is cached elsewhere
+		if (!bIsRegeneratingOnLoad)
+		{
+			FFindInBlueprintSearchManager::Get().AddOrUpdateBlueprintSearchMetadata(Request.BPToCompile);
+		}
+	}
+
+	ensure(!bIsRegeneratingOnLoad); // unexpected code path, compile on load handled with different function call
+	ensure(!bSkipReinstancing); // This is an internal option, should not go through CompileSynchronouslyImpl
+
+	ensure(QueuedRequests.Num() == 0);
+	QueuedRequests.Add(Request);
+	// We suppress normal compilation boradcasts because the old code path 
+	// did this after GC and we want to match the old behavior:
+	const bool bSuppressBroadcastCompiled = true;
+	FlushCompilationQueueImpl(nullptr, bSuppressBroadcastCompiled);
+	FlushReinstancingQueueImpl();
+
+	ensure(Request.BPToCompile->bQueuedForCompilation == false);
+
+	if(!bSkipGarbageCollection)
+	{
+		CollectGarbage(GARBAGE_COLLECTION_KEEPFLAGS);
+	}
+	
+	if (!bBatchCompile)
+	{
+		Request.BPToCompile->BroadcastCompiled();
+
+		if(GEditor)
+		{
+			GEditor->BroadcastBlueprintCompiled();	
+		}
+	}
 }
 
 static double GTimeCompiling = 0.f;
 static double GTimeReinstancing = 0.f;
 
-enum class ECompilationManagerFlags
+enum class ECompilationManagerJobType
 {
-	None = 0x0,
-
-	SkeletonOnly				= 0x1,
+	Normal,
+	SkeletonOnly,
+	RelinkOnly,
 };
 
-ENUM_CLASS_FLAGS(ECompilationManagerFlags)
-
 struct FCompilerData
 {
-	FCompilerData(UBlueprint* InBP, ECompilationManagerFlags InFlags)
+	explicit FCompilerData(UBlueprint* InBP, ECompilationManagerJobType InJobType, FCompilerResultsLog* InResultsLogOverride, EBlueprintCompileOptions UserOptions)
 	{
 		check(InBP);
 		BP = InBP;
-		Flags = InFlags;
-		ResultsLog = MakeUnique<FCompilerResultsLog>();
-		ResultsLog->BeginEvent(TEXT("BlueprintCompilationManager Compile"));
-		ResultsLog->SetSourcePath(InBP->GetPathName());
-		Options.bRegenerateSkelton = false;
-		Options.bReinstanceAndStubOnFailure = false;
+		JobType = InJobType;
+		UPackage* Package = Cast<UPackage>(BP->GetOutermost());
+		bPackageWasDirty = Package ? Package->IsDirty() : false;
+
+		ActiveResultsLog = InResultsLogOverride;
+		if(InResultsLogOverride == nullptr)
+		{
+			ResultsLog = MakeUnique<FCompilerResultsLog>();
+			ResultsLog->BeginEvent(TEXT("BlueprintCompilationManager Compile"));
+			ResultsLog->SetSourcePath(InBP->GetPathName());
+			ActiveResultsLog = ResultsLog.Get();
+		}
+		
+		static const FBoolConfigValueHelper IgnoreCompileOnLoadErrorsOnBuildMachine(TEXT("Kismet"), TEXT("bIgnoreCompileOnLoadErrorsOnBuildMachine"), GEngineIni);
+		ActiveResultsLog->bLogInfoOnly = !BP->bHasBeenRegenerated && GIsBuildMachine && IgnoreCompileOnLoadErrorsOnBuildMachine;
+
+		InternalOptions.bRegenerateSkelton = false;
+		InternalOptions.bReinstanceAndStubOnFailure = false;
+		InternalOptions.bSaveIntermediateProducts = (UserOptions & EBlueprintCompileOptions::SaveIntermediateProducts ) != EBlueprintCompileOptions::None;
+
 		if( UWidgetBlueprint* WidgetBP = Cast<UWidgetBlueprint>(BP))
 		{
-			Compiler = UWidgetBlueprint::GetCompilerForWidgetBP(WidgetBP, *ResultsLog, Options);
+			Compiler = UWidgetBlueprint::GetCompilerForWidgetBP(WidgetBP, *ActiveResultsLog, InternalOptions);
 		}
 		else
 		{
-			Compiler = FKismetCompilerContext::GetCompilerForBP(BP, *ResultsLog, Options);
+			Compiler = FKismetCompilerContext::GetCompilerForBP(BP, *ActiveResultsLog, InternalOptions);
 		}
 	}
 
-	bool IsSkeletonOnly() const { return !!(Flags & ECompilationManagerFlags::SkeletonOnly); }
+	bool IsSkeletonOnly() const { return JobType == ECompilationManagerJobType::SkeletonOnly; }
+	bool ShouldResetClassMembers() const { return JobType != ECompilationManagerJobType::RelinkOnly; }
+	bool ShouldSetTemporaryBlueprintFlags() const { return JobType != ECompilationManagerJobType::RelinkOnly; }
+	bool ShouldValidateVariableNames() const { return JobType == ECompilationManagerJobType::Normal; }
+	bool ShouldRegenerateSkeleton() const { return JobType != ECompilationManagerJobType::RelinkOnly; }
+	bool ShouldMarkUpToDateAfterSkeletonStage() const { return IsSkeletonOnly(); }
+	bool ShouldReconstructNodes() const { return JobType == ECompilationManagerJobType::Normal; }
+	bool ShouldSkipReinstancerCreation() const { return (IsSkeletonOnly() && BP->ParentClass->IsNative()); }
+	bool ShouldCompileClassLayout() const { return JobType == ECompilationManagerJobType::Normal; }
+	bool ShouldCompileClassFunctions() const { return JobType == ECompilationManagerJobType::Normal; }
+	bool ShouldRelinkAfterSkippingCompile() const { return JobType == ECompilationManagerJobType::RelinkOnly; }
 
 	UBlueprint* BP;
+	FCompilerResultsLog* ActiveResultsLog;
 	TUniquePtr<FCompilerResultsLog> ResultsLog;
 	TUniquePtr<FKismetCompilerContext> Compiler;
-	FKismetCompilerOptions Options;
+	FKismetCompilerOptions InternalOptions;
 
-	ECompilationManagerFlags Flags;
+	ECompilationManagerJobType JobType;
+	bool bPackageWasDirty;
 };
 	
-void FBlueprintCompilationManagerImpl::FlushCompilationQueueImpl(TArray<UObject*>* ObjLoaded)
+void FBlueprintCompilationManagerImpl::FlushCompilationQueueImpl(TArray<UObject*>* ObjLoaded, bool bSuppressBroadcastCompiled)
 {
+	TGuardValue<bool> GuardTemplateNameFlag(GCompilingBlueprint, true);
 	ensure(bGeneratedClassLayoutReady);
 
-	if( QueuedBPs.Num() == 0 )
+	if( QueuedRequests.Num() == 0 )
 	{
 		return;
 	}
@@ -237,12 +275,36 @@ void FBlueprintCompilationManagerImpl::FlushCompilationQueueImpl(TArray<UObject*
 	TArray<TUniquePtr<FBlueprintCompileReinstancer>> Reinstancers;
 	{ // begin GTimeCompiling scope 
 		FScopedDurationTimer SetupTimer(GTimeCompiling); 
-	
-		for(int32 I = 0; I < QueuedBPs.Num(); ++I)
+
+		// STAGE I: Add any child blueprints that were not compiled, we will need to run a link pass on them:
+		for(const FBPCompileRequest& ComplieJob : QueuedRequests)
+		{
+			if(UClass* OldSkeletonClass = ComplieJob.BPToCompile->SkeletonGeneratedClass)
+			{
+				TArray<UClass*> SkeletonClassesToReparentList;
+				GetDerivedClasses(OldSkeletonClass, SkeletonClassesToReparentList);
+		
+				for(UClass* ChildClass : SkeletonClassesToReparentList)
+				{
+					if(UBlueprint* ChildBlueprint = UBlueprint::GetBlueprintFromClass(ChildClass))
+					{
+						if(!IsQueuedForCompilation(ChildBlueprint))
+						{
+							ChildBlueprint->bQueuedForCompilation = true;
+							ensure(ChildBlueprint->bHasBeenRegenerated);
+							CurrentlyCompilingBPs.Add(FCompilerData(ChildBlueprint, ECompilationManagerJobType::RelinkOnly, nullptr, EBlueprintCompileOptions::None));
+						}
+					}
+				}
+			}
+		}
+
+		// STAGE II: Filter out data only and interface blueprints:
+		for(int32 I = 0; I < QueuedRequests.Num(); ++I)
 		{
-			// filter out data only and interface blueprints:
 			bool bSkipCompile = false;
-			UBlueprint* QueuedBP = QueuedBPs[I];
+			FBPCompileRequest& QueuedJob = QueuedRequests[I];
+			UBlueprint* QueuedBP = QueuedJob.BPToCompile;
 
 			ensure(!(QueuedBP->GeneratedClass->ClassDefaultObject->HasAnyFlags(RF_NeedLoad)));
 			if( FBlueprintEditorUtils::IsDataOnlyBlueprint(QueuedBP) )
@@ -264,7 +326,7 @@ void FBlueprintCompilationManagerImpl::FlushCompilationQueueImpl(TArray<UObject*
 
 			if(bSkipCompile)
 			{
-				CurrentlyCompilingBPs.Add(FCompilerData(QueuedBP, ECompilationManagerFlags::SkeletonOnly));
+				CurrentlyCompilingBPs.Add(FCompilerData(QueuedBP, ECompilationManagerJobType::SkeletonOnly, QueuedJob.ClientResultsLog, QueuedJob.CompileOptions));
 				if (QueuedBP->IsGeneratedClassAuthoritative() && (QueuedBP->GeneratedClass != nullptr))
 				{
 					// set bIsRegeneratingOnLoad so that we don't reset loaders:
@@ -281,18 +343,18 @@ void FBlueprintCompilationManagerImpl::FlushCompilationQueueImpl(TArray<UObject*
 					FBlueprintEditorUtils::RecreateClassMetaData(QueuedBP, QueuedBP->GeneratedClass, true);
 				}
 
-				QueuedBPs.RemoveAtSwap(I);
+				QueuedRequests.RemoveAtSwap(I);
 				--I;
 			}
 			else
 			{
-				CurrentlyCompilingBPs.Add(FCompilerData(QueuedBP, ECompilationManagerFlags::None));
+				CurrentlyCompilingBPs.Add(FCompilerData(QueuedBP, ECompilationManagerJobType::Normal, QueuedJob.ClientResultsLog, QueuedJob.CompileOptions));
 			}
 		}
-	
-		QueuedBPs.Empty();
 
-		// sort into sane compilation order. Not needed if we introduce compilation phases:
+		QueuedRequests.Empty();
+
+		// STAGE III: Sort into correct compilation order. We want to compile root types before their derived (child) types:
 		auto HierarchyDepthSortFn = [](const FCompilerData& CompilerDataA, const FCompilerData& CompilerDataB)
 		{
 			UBlueprint& A = *(CompilerDataA.BP);
@@ -334,92 +396,135 @@ void FBlueprintCompilationManagerImpl::FlushCompilationQueueImpl(TArray<UObject*
 		};
 		CurrentlyCompilingBPs.Sort( HierarchyDepthSortFn );
 
+		// STAGE IV: Set UBlueprint flags (bBeingCompiled, bIsRegeneratingOnLoad)
 		for (FCompilerData& CompilerData : CurrentlyCompilingBPs)
 		{
-			if(CompilerData.IsSkeletonOnly())
+			if(CompilerData.ShouldSetTemporaryBlueprintFlags())
+			{
+				UBlueprint* BP = CompilerData.BP;
+				BP->bBeingCompiled = true;
+				BP->bIsRegeneratingOnLoad = !BP->bHasBeenRegenerated && BP->GetLinker();
+			}
+		}
+
+		// STAGE V: Validate Variable Names
+		for (FCompilerData& CompilerData : CurrentlyCompilingBPs)
+		{
+			if(!CompilerData.ShouldValidateVariableNames())
 			{
 				continue;
 			}
 
-			UBlueprint* BP = CompilerData.BP;
-		
-			// even ValidateVariableNames will trigger compilation if 
-			// it thinks the BP is not being compiled:
-			BP->bBeingCompiled = true;
 			CompilerData.Compiler->ValidateVariableNames();
-			BP->bBeingCompiled = false;
 		}
 
-		// purge null graphs, could be done only on load
+		// STAGE VI: Purge null graphs, could be done only on load
 		for (FCompilerData& CompilerData : CurrentlyCompilingBPs)
 		{
 			UBlueprint* BP = CompilerData.BP;
-			UPackage* Package = Cast<UPackage>(BP->GetOutermost());
-			bool bIsPackageDirty = Package ? Package->IsDirty() : false;
-
 			FBlueprintEditorUtils::PurgeNullGraphs(BP);
+		}
 
-			if( Package )
+		// STAGE VII: safely throw away old skeleton CDOs:
+		{
+			TMap<UClass*, UClass*> OldSkeletonToNewSkeleton;
+			for (FCompilerData& CompilerData : CurrentlyCompilingBPs)
 			{
-				Package->SetDirtyFlag(bIsPackageDirty);
+				UBlueprint* BP = CompilerData.BP;
+				if(BP->SkeletonGeneratedClass)
+				{
+					UClass* CopyOfOldClass = FBlueprintCompileReinstancer::MoveCDOToNewClass(BP->SkeletonGeneratedClass, OldSkeletonToNewSkeleton, true);
+					OldSkeletonToNewSkeleton.Add(BP->SkeletonGeneratedClass, CopyOfOldClass);
+					// Child types will need to use CopyOfOldClass
+				}
 			}
 		}
-
-		// recompile skeleton
-		for (FCompilerData& CompilerData : CurrentlyCompilingBPs)
+		
+		// STAGE VIII: recompile skeleton
 		{
-			UBlueprint* BP = CompilerData.BP;
-			UPackage* Package = Cast<UPackage>(BP->GetOutermost());
-			bool bIsPackageDirty = Package ? Package->IsDirty() : false;
+			TMap<UClass*, UClass*> OldSkeletonToNewSkeleton;
+			for (FCompilerData& CompilerData : CurrentlyCompilingBPs)
+			{
+				UBlueprint* BP = CompilerData.BP;
 		
-			BP->SkeletonGeneratedClass = FastGenerateSkeletonClass(BP, *(CompilerData.Compiler) );
+				if(CompilerData.ShouldRegenerateSkeleton())
+				{
+					UClass* OldSkeleton = BP->SkeletonGeneratedClass;
+					BP->SkeletonGeneratedClass = FastGenerateSkeletonClass(BP, *(CompilerData.Compiler) );
+					OldSkeletonToNewSkeleton.Add(OldSkeleton, BP->SkeletonGeneratedClass);
+				}
+				else
+				{
+					// Just relink, note that UProperties that reference *other* types may be stale until
+					// we fixup below:
+					UClass* SkeletonToRelink = BP->SkeletonGeneratedClass;
 
-			if(CompilerData.IsSkeletonOnly())
-			{
-				// Flag data only blueprints as being up-to-date
-				BP->Status = BS_UpToDate;
-				BP->bHasBeenRegenerated = true;
-				BP->GeneratedClass->ClearFunctionMapsCaches();
-			}
+					// CDO needs to be moved aside already:
+					ensure(SkeletonToRelink->ClassDefaultObject == nullptr);
 
-			if( Package )
-			{
-				Package->SetDirtyFlag(bIsPackageDirty);
+					UClass * const* OverridenParent = OldSkeletonToNewSkeleton.Find(SkeletonToRelink->GetSuperClass());
+					if(OverridenParent && *OverridenParent)
+					{
+						SkeletonToRelink->SetSuperStruct(*OverridenParent);
+					}
+
+					SkeletonToRelink->Bind();
+					SkeletonToRelink->StaticLink(true);
+				}
+
+				if(CompilerData.ShouldMarkUpToDateAfterSkeletonStage())
+				{
+					// Flag data only blueprints as being up-to-date
+					BP->Status = BS_UpToDate;
+					BP->bHasBeenRegenerated = true;
+					BP->GeneratedClass->ClearFunctionMapsCaches();
+				}
 			}
 		}
 
-		// reconstruct nodes
+		// STAGE IX: Reconstruct nodes and replace deprecated nodes, then broadcast 'precompile
 		for (FCompilerData& CompilerData : CurrentlyCompilingBPs)
 		{
-			if(CompilerData.IsSkeletonOnly())
+			if(!CompilerData.ShouldReconstructNodes())
 			{
 				continue;
 			}
 
 			UBlueprint* BP = CompilerData.BP;
-			UPackage* Package = Cast<UPackage>(BP->GetOutermost());
-			bool bIsPackageDirty = Package ? Package->IsDirty() : false;
 
-			// Setting bBeingCompiled to make sure that FBlueprintEditorUtils::MarkBlueprintAsStructurallyModified
-			// doesn't run skeleton compile while the compilation manager is compiling things:
-			BP->bBeingCompiled = true;
-			FBlueprintEditorUtils::ReconstructAllNodes(BP);
-			FBlueprintEditorUtils::ReplaceDeprecatedNodes(BP);
-			BP->bBeingCompiled = false;
+			// Some nodes are set up to do things during reconstruction only when this flag is NOT set.
+			if(BP->bIsRegeneratingOnLoad)
+			{
+				FBlueprintEditorUtils::ReconstructAllNodes(BP);
+				FBlueprintEditorUtils::ReplaceDeprecatedNodes(BP);
+			}
+			
+			// Broadcast pre-compile
+			{
+				if(GEditor && GIsEditor)
+				{
+					GEditor->BroadcastBlueprintPreCompile(BP);
+				}
+			}
+
+			// Do not want to run this code without the editor present nor when running commandlets.
+			if (GEditor && GIsEditor)
+			{
+				// We do not want to regenerate a search Guid during loads, nothing has changed in the Blueprint and it is cached elsewhere
+				if (!BP->bIsRegeneratingOnLoad)
+				{
+					FFindInBlueprintSearchManager::Get().AddOrUpdateBlueprintSearchMetadata(BP);
+				}
+			}
 
 			// we are regenerated, tag ourself as such so that
 			// old logic to 'fix' circular dependencies doesn't
 			// cause redundant regeneration (e.g. bForceRegenNodes
 			// in ExpandTunnelsAndMacros):
 			BP->bHasBeenRegenerated = true;
-		
-			if( Package )
-			{
-				Package->SetDirtyFlag(bIsPackageDirty);
-			}
 		}
 	
-		// reinstance every blueprint that is queued, note that this means classes in the hierarchy that are *not* being 
+		// STAGE X: reinstance every blueprint that is queued, note that this means classes in the hierarchy that are *not* being 
 		// compiled will be parented to REINST versions of the class, so type checks (IsA, etc) involving those types
 		// will be incoherent!
 		TMap<UBlueprint*, UObject*> OldCDOs;
@@ -429,14 +534,12 @@ void FBlueprintCompilationManagerImpl::FlushCompilationQueueImpl(TArray<UObject*
 				// we including skeleton only compilation jobs for reinstancing because we need UpdateCustomPropertyListForPostConstruction
 				// to happen (at the right time) for those generated classes as well. This means we *don't* need to reinstance if 
 				// the parent is a native type (unless we hot reload, but that should not need to be handled here):
-				if(CompilerData.IsSkeletonOnly() && CompilerData.BP->ParentClass->IsNative())
+				if(CompilerData.ShouldSkipReinstancerCreation())
 				{
 					continue;
 				}
 
 				UBlueprint* BP = CompilerData.BP;
-				UPackage* Package = Cast<UPackage>(BP->GetOutermost());
-				bool bIsPackageDirty = Package ? Package->IsDirty() : false;
 
 				OldCDOs.Add(BP, BP->GeneratedClass->ClassDefaultObject);
 				Reinstancers.Emplace(TUniquePtr<FBlueprintCompileReinstancer>( 
@@ -445,117 +548,143 @@ void FBlueprintCompilationManagerImpl::FlushCompilationQueueImpl(TArray<UObject*
 						EBlueprintCompileReinstancerFlags::AutoInferSaveOnCompile|EBlueprintCompileReinstancerFlags::AvoidCDODuplication
 					)
 				));
-
-				if( Package )
-				{
-					Package->SetDirtyFlag(bIsPackageDirty);
-				}
 			}
 		}
 
-		// Reinstancing done, lets fix up child->parent pointers:
+		// STAGE XI: Reinstancing done, lets fix up child->parent pointers
 		for (FCompilerData& CompilerData : CurrentlyCompilingBPs)
 		{
 			UBlueprint* BP = CompilerData.BP;
 			if(BP->GeneratedClass->GetSuperClass()->HasAnyClassFlags(CLASS_NewerVersionExists))
 			{
 				BP->GeneratedClass->SetSuperStruct(BP->GeneratedClass->GetSuperClass()->GetAuthoritativeClass());
-				BP->GeneratedClass->Children = NULL;
-				BP->GeneratedClass->Script.Empty();
-				BP->GeneratedClass->MinAlignment = 0;
-				BP->GeneratedClass->RefLink = NULL;
-				BP->GeneratedClass->PropertyLink = NULL;
-				BP->GeneratedClass->DestructorLink = NULL;
-				BP->GeneratedClass->ScriptObjectReferences.Empty();
-				BP->GeneratedClass->PropertyLink = NULL;
+				if(CompilerData.ShouldResetClassMembers())
+				{
+					BP->GeneratedClass->Children = NULL;
+					BP->GeneratedClass->Script.Empty();
+					BP->GeneratedClass->MinAlignment = 0;
+					BP->GeneratedClass->RefLink = NULL;
+					BP->GeneratedClass->PropertyLink = NULL;
+					BP->GeneratedClass->DestructorLink = NULL;
+					BP->GeneratedClass->ScriptObjectReferences.Empty();
+					BP->GeneratedClass->PropertyLink = NULL;
+				}
 			}
 		}
 
-		// recompile every blueprint
+		// STAGE XII: Recompile every blueprint
 		bGeneratedClassLayoutReady = false;
 		for (FCompilerData& CompilerData : CurrentlyCompilingBPs)
 		{
-			if(CompilerData.IsSkeletonOnly())
-			{
-				continue;
-			}
-
 			UBlueprint* BP = CompilerData.BP;
-			UPackage* Package = Cast<UPackage>(BP->GetOutermost());
-			bool bIsPackageDirty = Package ? Package->IsDirty() : false;
-
-			ensure(BP->GeneratedClass->ClassDefaultObject == nullptr || 
-				BP->GeneratedClass->ClassDefaultObject->GetClass() != BP->GeneratedClass);
-			// default value propagation occurrs below:
-			BP->GeneratedClass->ClassDefaultObject = nullptr;
-			BP->bIsRegeneratingOnLoad = true;
-			BP->bBeingCompiled = true;
-			// Reset the flag, so if the user tries to use PIE it will warn them if the BP did not compile
-			BP->bDisplayCompilePIEWarning = true;
-		
-			FKismetCompilerContext& CompilerContext = *(CompilerData.Compiler);
-			CompilerContext.CompileClassLayout(EInternalCompilerFlags::PostponeLocalsGenerationUntilPhaseTwo);
-			BP->bIsRegeneratingOnLoad = false;
-
-			if( Package )
+			if(CompilerData.ShouldCompileClassLayout())
 			{
-				Package->SetDirtyFlag(bIsPackageDirty);
+				ensure( BP->GeneratedClass->ClassDefaultObject == nullptr || 
+						BP->GeneratedClass->ClassDefaultObject->GetClass() != BP->GeneratedClass);
+				// default value propagation occurs in ReinstaneBatch, CDO will be created via CompileFunctions call:
+				BP->GeneratedClass->ClassDefaultObject = nullptr;
+				// Reset the flag, so if the user tries to use PIE it will warn them if the BP did not compile
+				BP->bDisplayCompilePIEWarning = true;
+		
+				FKismetCompilerContext& CompilerContext = *(CompilerData.Compiler);
+				CompilerContext.CompileClassLayout(EInternalCompilerFlags::PostponeLocalsGenerationUntilPhaseTwo);
 			}
 		}
 		bGeneratedClassLayoutReady = true;
 	
+		// STAGE XIII: Compile functions
 		for (FCompilerData& CompilerData : CurrentlyCompilingBPs)
 		{
-			if(CompilerData.IsSkeletonOnly())
+			UBlueprint* BP = CompilerData.BP;
+			if(!CompilerData.ShouldCompileClassFunctions())
 			{
-				continue;
+				if( BP->GeneratedClass->ClassDefaultObject == nullptr || 
+					BP->GeneratedClass->ClassDefaultObject->GetClass() != BP->GeneratedClass)
+				{
+					// relink, generate CDO:
+					UClass* BPGC = BP->GeneratedClass;
+					BPGC->Bind();
+					BPGC->StaticLink(true);
+					BPGC->ClassDefaultObject = nullptr;
+					BPGC->GetDefaultObject(true);
+				}
 			}
-
-			UBlueprint* BP = CompilerData.BP;
-			UPackage* Package = Cast<UPackage>(BP->GetOutermost());
-			bool bIsPackageDirty = Package ? Package->IsDirty() : false;
-
-			ensure(BP->GeneratedClass->ClassDefaultObject == nullptr || 
-				BP->GeneratedClass->ClassDefaultObject->GetClass() != BP->GeneratedClass);
-			// default value propagation occurrs below:
-			BP->GeneratedClass->ClassDefaultObject = nullptr;
-			BP->bIsRegeneratingOnLoad = true;
+			else
+			{
+				ensure( BP->GeneratedClass->ClassDefaultObject == nullptr || 
+						BP->GeneratedClass->ClassDefaultObject->GetClass() != BP->GeneratedClass);
+				// default value propagation occurrs below:
+				BP->GeneratedClass->ClassDefaultObject = nullptr;
 		
-			FKismetCompilerContext& CompilerContext = *(CompilerData.Compiler);
-			CompilerContext.CompileFunctions(EInternalCompilerFlags::PostponeLocalsGenerationUntilPhaseTwo);
+				FKismetCompilerContext& CompilerContext = *(CompilerData.Compiler);
+				CompilerContext.CompileFunctions(EInternalCompilerFlags::PostponeLocalsGenerationUntilPhaseTwo);
+
+				if (CompilerData.ActiveResultsLog->NumErrors == 0)
+				{
+					// Blueprint is error free.  Go ahead and fix up debug info
+					BP->Status = (0 == CompilerData.ActiveResultsLog->NumWarnings) ? BS_UpToDate : BS_UpToDateWithWarnings;
+				}
+				else
+				{
+					BP->Status = BS_Error; // do we still have the old version of the class?
+				}
+			}
+
+			ensure(BP->GeneratedClass->ClassDefaultObject->GetClass() == *(BP->GeneratedClass));
+		}
 		
-			BP->bBeingCompiled = false;
+		// STAGE XIV: CLEAR TEMPORARY FLAGS
+		for (FCompilerData& CompilerData : CurrentlyCompilingBPs)
+		{
+			UBlueprint* BP = CompilerData.BP;
+			FBlueprintEditorUtils::UpdateDelegatesInBlueprint(BP);
+
+			if (FBlueprintEditorUtils::IsLevelScriptBlueprint(BP))
+			{
+				// When the Blueprint is recompiled, then update the bound events for level scripting
+				ULevelScriptBlueprint* LevelScriptBP = CastChecked<ULevelScriptBlueprint>(BP);
+
+				if (ULevel* BPLevel = LevelScriptBP->GetLevel())
+				{
+					BPLevel->OnLevelScriptBlueprintChanged(LevelScriptBP);
+				}
+			}
 
-			if (CompilerData.ResultsLog->NumErrors == 0)
+			if(!BP->bIsRegeneratingOnLoad)
 			{
-				// Blueprint is error free.  Go ahead and fix up debug info
-				BP->Status = (0 == CompilerData.ResultsLog->NumWarnings) ? BS_UpToDate : BS_UpToDateWithWarnings;
+				FKismetEditorUtilities::StripExternalComponents(BP);
+				
+				if(BP->SimpleConstructionScript)
+				{
+					BP->SimpleConstructionScript->FixupRootNodeParentReferences();
+				}
+				
+				UBlueprint::ValidateGeneratedClass(BP->GeneratedClass);
 			}
-			else
+
+			if(CompilerData.ShouldSetTemporaryBlueprintFlags())
 			{
-				BP->Status = BS_Error; // do we still have the old version of the class?
+				BP->bBeingCompiled = false;
+				BP->bIsRegeneratingOnLoad = false;
 			}
 
-			BP->bIsRegeneratingOnLoad = false;
-			ensure(BP->GeneratedClass->ClassDefaultObject->GetClass() == *(BP->GeneratedClass));
-		
-			if( Package )
+			if(UPackage* Package = Cast<UPackage>(BP->GetOutermost()))
 			{
-				Package->SetDirtyFlag(bIsPackageDirty);
+				Package->SetDirtyFlag(CompilerData.bPackageWasDirty);
 			}
 		}
 	} // end GTimeCompiling scope
 
-	// and now we can finish the first stage of the reinstancing operation, moving old classes to new classes:
+	// STAGE XV: Now we can finish the first stage of the reinstancing operation, moving old classes to new classes:
 	{
 		FScopedDurationTimer ReinstTimer(GTimeReinstancing);
 		ReinstanceBatch(Reinstancers, ClassesToReinstance, ObjLoaded);
 		
-		// make sure no junk in bytecode, this can happen only for blueprints that were in QueuedBPsLocal because
+		// Make sure no junk in bytecode, this can happen only for blueprints that were in CurrentlyCompilingBPs because
 		// the reinstancer can detect all other references (see UpdateBytecodeReferences):
 		for (FCompilerData& CompilerData : CurrentlyCompilingBPs)
 		{
-			if(CompilerData.IsSkeletonOnly())
+			if(CompilerData.ShouldCompileClassFunctions())
 			{
 				continue;
 			}
@@ -569,22 +698,47 @@ void FBlueprintCompilationManagerImpl::FlushCompilationQueueImpl(TArray<UObject*
 					FFixupBytecodeReferences ValidateAr(CurrentFunction);
 				}
 			}
+			
+			if (!bSuppressBroadcastCompiled)
+			{
+				BP->BroadcastCompiled();
+			}
+		}
+
+		if ( GEditor )
+		{
+			GEditor->BroadcastBlueprintReinstanced();
+		}
+
+		if (!bSuppressBroadcastCompiled)
+		{
+			if(GEditor)
+			{
+				GEditor->BroadcastBlueprintCompiled();	
+			}
 		}
 	}
 
 	for (FCompilerData& CompilerData : CurrentlyCompilingBPs)
 	{
-		CompilerData.ResultsLog->EndEvent();
+		if(CompilerData.ResultsLog)
+		{
+			CompilerData.ResultsLog->EndEvent();
+		}
+		CompilerData.BP->bQueuedForCompilation = false;
 	}
 
+	UEdGraphPin::Purge();
+
 	UE_LOG(LogBlueprint, Display, TEXT("Time Compiling: %f, Time Reinstancing: %f"),  GTimeCompiling, GTimeReinstancing);
 	//GTimeCompiling = 0.0;
 	//GTimeReinstancing = 0.0;
-	ensure(QueuedBPs.Num() == 0);
+	ensure(QueuedRequests.Num() == 0);
 }
 
 void FBlueprintCompilationManagerImpl::FlushReinstancingQueueImpl()
 {
+	TGuardValue<bool> GuardTemplateNameFlag(GCompilingBlueprint, true);
 	// we can finalize reinstancing now:
 	if(ClassesToReinstance.Num() == 0)
 	{
@@ -593,7 +747,8 @@ void FBlueprintCompilationManagerImpl::FlushReinstancingQueueImpl()
 
 	{
 		FScopedDurationTimer ReinstTimer(GTimeReinstancing);
-
+		
+		TGuardValue<bool> ReinstancingGuard(GIsReinstancing, true);
 		FBlueprintCompileReinstancer::BatchReplaceInstancesOfClass(ClassesToReinstance);
 		ClassesToReinstance.Empty();
 	}
@@ -611,7 +766,7 @@ void FBlueprintCompilationManagerImpl::FlushReinstancingQueueImpl()
 
 bool FBlueprintCompilationManagerImpl::HasBlueprintsToCompile() const
 {
-	return QueuedBPs.Num() != 0;
+	return QueuedRequests.Num() != 0;
 }
 
 bool FBlueprintCompilationManagerImpl::IsGeneratedClassLayoutReady() const
@@ -818,6 +973,7 @@ void FBlueprintCompilationManagerImpl::ReinstanceBatch(TArray<TUniquePtr<FBluepr
 
 	// 3. Update any remaining instances that are tagged as RF_ArchetypeObject or RF_InheritableComponentTemplate - 
 	// we may need to do further sorting to ensure that interdependent archetypes are initialized correctly:
+	TSet<UObject*> ArchetypeReferencers;
 	for( const TUniquePtr<FBlueprintCompileReinstancer>& Reinstancer : Reinstancers )
 	{
 		UClass* OldClass = Reinstancer->DuplicatedClass;
@@ -845,6 +1001,25 @@ void FBlueprintCompilationManagerImpl::ReinstanceBatch(TArray<TUniquePtr<FBluepr
 			// for each archetype:
 			for(UObject* Archetype : ArchetypeObjects )
 			{
+				// make sure we fix up references in the owner:
+				if(UBlueprintGeneratedClass* ArchetypeOwnerClass = Archetype->GetTypedOuter<UBlueprintGeneratedClass>())
+				{
+					ArchetypeReferencers.Add(ArchetypeOwnerClass);
+					
+					TArray<UObject*> ContainedObjects;
+					GetObjectsWithOuter(ArchetypeOwnerClass, ContainedObjects);
+					ArchetypeReferencers.Append(ContainedObjects);
+				}
+				
+				if(UBlueprint* ArchetypeOwnerBP = Archetype->GetTypedOuter<UBlueprint>())
+				{
+					ArchetypeReferencers.Add(ArchetypeOwnerBP);
+
+					TArray<UObject*> ContainedObjects;
+					GetObjectsWithOuter(ArchetypeOwnerBP, ContainedObjects);
+					ArchetypeReferencers.Append(ContainedObjects);
+				}
+				
 				// move aside:
 				FName OriginalName = Archetype->GetFName();
 				UObject* OriginalOuter = Archetype->GetOuter();
@@ -866,6 +1041,7 @@ void FBlueprintCompilationManagerImpl::ReinstanceBatch(TArray<TUniquePtr<FBluepr
 				FBlueprintCompileReinstancer::CopyPropertiesForUnrelatedObjects(Archetype, NewArchetype);
 
 				OldArchetypeToNewArchetype.Add(Archetype, NewArchetype);
+				ArchetypeReferencers.Add(Archetype);
 				ObjectsThatShouldUseOldStuff.Add(Archetype);
 
 				Archetype->MarkPendingKill();
@@ -876,13 +1052,13 @@ void FBlueprintCompilationManagerImpl::ReinstanceBatch(TArray<TUniquePtr<FBluepr
 	// 4. update known references to archetypes (e.g. component templates, WidgetTree). We don't want to run the normal 
 	// reference finder to update these because searching the entire object graph is time consuming. Instead we just replace
 	// all references in our UBlueprint and its generated class:
-	TSet<UObject*> ArchetypeReferencers;
 	for( const TUniquePtr<FBlueprintCompileReinstancer>& Reinstancer : Reinstancers )
 	{
 		ArchetypeReferencers.Add(Reinstancer->ClassToReinstance);
 		ArchetypeReferencers.Add(Reinstancer->ClassToReinstance->ClassGeneratedBy);
 		if(UBlueprint* BP = Cast<UBlueprint>(Reinstancer->ClassToReinstance->ClassGeneratedBy))
 		{
+			ArchetypeReferencers.Add(BP->SkeletonGeneratedClass);
 			ensure(BP->bCachedDependenciesUpToDate);
 			for(TWeakObjectPtr<UBlueprint> Dependendency : BP->CachedDependencies)
 			{
@@ -897,7 +1073,13 @@ void FBlueprintCompilationManagerImpl::ReinstanceBatch(TArray<TUniquePtr<FBluepr
 	}
 }
 
+/*
+	This function completely replaces the 'skeleton only' compilation pass in the Kismet compiler. Long
+	term that code path will be removed and clients will be redirected to this function.
 
+	Notes to maintainers: any UObject created here and outered to the resulting class must be marked as transient
+	or you will create a cook error!
+*/
 UClass* FBlueprintCompilationManagerImpl::FastGenerateSkeletonClass(UBlueprint* BP, FKismetCompilerContext& CompilerContext)
 {
 	const UEdGraphSchema_K2* Schema = GetDefault<UEdGraphSchema_K2>();
@@ -916,14 +1098,25 @@ UClass* FBlueprintCompilationManagerImpl::FastGenerateSkeletonClass(UBlueprint*
 		}
 	}
 
-	ensure(BP->SkeletonGeneratedClass == nullptr);
+	UBlueprintGeneratedClass* Ret = nullptr;
 	UBlueprintGeneratedClass* OriginalNewClass = CompilerContext.NewClass;
-	CompilerContext.SpawnNewClass(FString::Printf(TEXT("SKEL_%s_C"), *BP->GetName()));
-	UBlueprintGeneratedClass* Ret = CompilerContext.NewClass;
-	CompilerContext.NewClass = OriginalNewClass;
-	Ret->ClassGeneratedBy = BP;
+	if(BP->SkeletonGeneratedClass == nullptr )
+	{
+		CompilerContext.SpawnNewClass(FString::Printf(TEXT("SKEL_%s_C"), *BP->GetName()));
+		Ret = CompilerContext.NewClass;
+		Ret->SetFlags(RF_Transient);
+		CompilerContext.NewClass = OriginalNewClass;
+	}
+	else
+	{
+		Ret = CastChecked<UBlueprintGeneratedClass>(*(BP->SkeletonGeneratedClass));
+		CompilerContext.CleanAndSanitizeClass(Ret, Ret->ClassDefaultObject);
+	}
 	
-	const auto MakeFunction = [Ret, ParentClass, Schema, BP, &MessageLog](FName FunctionNameFName, UField**& InCurrentFieldStorageLocation, UField**& InCurrentParamStorageLocation, bool bIsStaticFunction, uint32 InFunctionFlags, const TArray<UK2Node_FunctionResult*>& ReturnNodes, const TArray<UEdGraphPin*>& InputPins) -> UFunction*
+	Ret->ClassGeneratedBy = BP;
+
+	// This is a version of PrecompileFunction that does not require 'terms' and graph cloning:
+	const auto MakeFunction = [Ret, ParentClass, Schema, BP, &MessageLog](FName FunctionNameFName, UField**& InCurrentFieldStorageLocation, UField**& InCurrentParamStorageLocation, bool bIsStaticFunction, uint32 InFunctionFlags, const TArray<UK2Node_FunctionResult*>& ReturnNodes, const TArray<UEdGraphPin*>& InputPins, bool bForceArrayStructRefsConst) -> UFunction*
 	{
 		if(!ensure(FunctionNameFName != FName())
 			|| FindObjectFast<UFunction>(Ret, FunctionNameFName, true ))
@@ -931,7 +1124,7 @@ UClass* FBlueprintCompilationManagerImpl::FastGenerateSkeletonClass(UBlueprint*
 			return nullptr;
 		}
 		
-		UFunction* NewFunction = NewObject<UFunction>(Ret, FunctionNameFName, RF_Public);
+		UFunction* NewFunction = NewObject<UFunction>(Ret, FunctionNameFName, RF_Public|RF_Transient);
 					
 		Ret->AddFunctionToFunctionMap(NewFunction);
 
@@ -991,43 +1184,48 @@ UClass* FBlueprintCompilationManagerImpl::FastGenerateSkeletonClass(UBlueprint*
 				{
 					// Reimplmentation of FKismetCompilerContext::CreatePropertiesFromList without dependence on 'terms'
 					UProperty* Param = FKismetCompilerUtilities::CreatePropertyOnScope(NewFunction, *(Pin->PinName), Pin->PinType, Ret, 0, Schema, MessageLog);
-					Param->PropertyFlags |= CPF_Parm;
-					if(Pin->PinType.bIsReference)
-					{
-						Param->PropertyFlags |= CPF_ReferenceParm | CPF_OutParm;
-					}
-					if(Pin->PinType.bIsConst)
-					{
-						Param->PropertyFlags |= CPF_ConstParm;
-					}
-
-					if (UObjectProperty* ObjProp = Cast<UObjectProperty>(Param))
+					if(Param)
 					{
-						UClass* EffectiveClass = nullptr;
-						if (ObjProp->PropertyClass != nullptr)
+						Param->SetFlags(RF_Transient);
+						Param->PropertyFlags |= CPF_Parm;
+						if(Pin->PinType.bIsReference)
 						{
-							EffectiveClass = ObjProp->PropertyClass;
+							Param->PropertyFlags |= CPF_ReferenceParm | CPF_OutParm;
 						}
-						else if (UClassProperty* ClassProp = Cast<UClassProperty>(ObjProp))
+
+						if(Pin->PinType.bIsConst || (bForceArrayStructRefsConst && (Pin->PinType.IsArray() || Pin->PinType.PinCategory == UEdGraphSchema_K2::PC_Struct) && Pin->PinType.bIsReference))
 						{
-							EffectiveClass = ClassProp->MetaClass;
+							Param->PropertyFlags |= CPF_ConstParm;
 						}
 
-						if ((EffectiveClass != nullptr) && (EffectiveClass->HasAnyClassFlags(CLASS_Const)))
+						if (UObjectProperty* ObjProp = Cast<UObjectProperty>(Param))
 						{
-							Param->PropertyFlags |= CPF_ConstParm;
+							UClass* EffectiveClass = nullptr;
+							if (ObjProp->PropertyClass != nullptr)
+							{
+								EffectiveClass = ObjProp->PropertyClass;
+							}
+							else if (UClassProperty* ClassProp = Cast<UClassProperty>(ObjProp))
+							{
+								EffectiveClass = ClassProp->MetaClass;
+							}
+
+							if ((EffectiveClass != nullptr) && (EffectiveClass->HasAnyClassFlags(CLASS_Const)))
+							{
+								Param->PropertyFlags |= CPF_ConstParm;
+							}
 						}
-					}
-					else if (UArrayProperty* ArrayProp = Cast<UArrayProperty>(Param))
-					{
-						Param->PropertyFlags |= CPF_ReferenceParm;
+						else if (UArrayProperty* ArrayProp = Cast<UArrayProperty>(Param))
+						{
+							Param->PropertyFlags |= CPF_ReferenceParm;
 
-						// ALWAYS pass array parameters as out params, so they're set up as passed by ref
-						Param->PropertyFlags |= CPF_OutParm;
-					}
+							// ALWAYS pass array parameters as out params, so they're set up as passed by ref
+							Param->PropertyFlags |= CPF_OutParm;
+						}
 
-					*InCurrentParamStorageLocation = Param;
-					InCurrentParamStorageLocation = &Param->Next;
+						*InCurrentParamStorageLocation = Param;
+						InCurrentParamStorageLocation = &Param->Next;
+					}
 				}
 			}
 			
@@ -1036,6 +1234,7 @@ UClass* FBlueprintCompilationManagerImpl::FastGenerateSkeletonClass(UBlueprint*
 				// Gather all input pins on these nodes, these are 
 				// the outputs of the function:
 				TSet<FString> UsedPinNames;
+				static const FName RetValName = FName(TEXT("ReturnValue"));
 				for(UK2Node_FunctionResult* Node : ReturnNodes)
 				{
 					for(UEdGraphPin* Pin : Node->Pins)
@@ -1047,9 +1246,18 @@ UClass* FBlueprintCompilationManagerImpl::FastGenerateSkeletonClass(UBlueprint*
 								UsedPinNames.Add(Pin->PinName);
 							
 								UProperty* Param = FKismetCompilerUtilities::CreatePropertyOnScope(NewFunction, *(Pin->PinName), Pin->PinType, Ret, 0, Schema, MessageLog);
-								Param->PropertyFlags |= CPF_Parm|CPF_ReturnParm;
-								*InCurrentParamStorageLocation = Param;
-								InCurrentParamStorageLocation = &Param->Next;
+								if(Param)
+								{
+									Param->SetFlags(RF_Transient);
+									// we only tag things as CPF_ReturnParm if the value is named ReturnValue.... this is *terrible* behavior:
+									if(Param->GetFName() == RetValName)
+									{
+										Param->PropertyFlags |= CPF_ReturnParm;
+									}
+									Param->PropertyFlags |= CPF_Parm|CPF_OutParm;
+									*InCurrentParamStorageLocation = Param;
+									InCurrentParamStorageLocation = &Param->Next;
+								}
 							}
 						}
 					}
@@ -1078,9 +1286,10 @@ UClass* FBlueprintCompilationManagerImpl::FastGenerateSkeletonClass(UBlueprint*
 					InCurrentFieldStorageLocation, 
 					CurrentParamStorageLocation, 
 					bIsStaticFunction, 
-					EntryNode->GetFunctionFlags(), 
+					EntryNode->GetFunctionFlags() & ~FUNC_Native, 
 					ReturnNodes, 
-					EntryNode->Pins
+					EntryNode->Pins,
+					false
 				);
 
 				if(NewFunction)
@@ -1088,23 +1297,33 @@ UClass* FBlueprintCompilationManagerImpl::FastGenerateSkeletonClass(UBlueprint*
 					// locals:
 					for( const FBPVariableDescription& BPVD : EntryNode->LocalVariables )
 					{
-						UProperty* NewProperty = FKismetCompilerUtilities::CreatePropertyOnScope(NewFunction, BPVD.VarName, BPVD.VarType, Ret, 0, Schema, MessageLog);
-						*CurrentParamStorageLocation = NewProperty;
-						CurrentParamStorageLocation = &NewProperty->Next;
+						if(UProperty* LocalVariable = FKismetCompilerContext::CreateUserDefinedLocalVariableForFunction(BPVD, NewFunction, Ret, CurrentParamStorageLocation, Schema, MessageLog) )
+						{
+							LocalVariable->SetFlags(RF_Transient);
+						}
 					}
 
 					// __WorldContext:
-					if(bIsStaticFunction && FindField<UObjectProperty>(NewFunction, TEXT("__WorldContext")) == nullptr)
+					if(bIsStaticFunction)
 					{
-						FEdGraphPinType WorldContextPinType(Schema->PC_Object, FString(), UObject::StaticClass(), false, false, false, false, FEdGraphTerminalType());
-						UProperty* Param = FKismetCompilerUtilities::CreatePropertyOnScope(NewFunction, TEXT("__WorldContext"), WorldContextPinType, Ret, 0, Schema, MessageLog);
-						Param->PropertyFlags |= CPF_Parm;
-						*CurrentParamStorageLocation = Param;
-						CurrentParamStorageLocation = &Param->Next;
+						if( FindField<UObjectProperty>(NewFunction, TEXT("__WorldContext")) == nullptr )
+						{
+							FEdGraphPinType WorldContextPinType(Schema->PC_Object, FString(), UObject::StaticClass(), EPinContainerType::None, false, FEdGraphTerminalType());
+							UProperty* Param = FKismetCompilerUtilities::CreatePropertyOnScope(NewFunction, TEXT("__WorldContext"), WorldContextPinType, Ret, 0, Schema, MessageLog);
+							if(Param)
+							{
+								Param->SetFlags(RF_Transient);
+								Param->PropertyFlags |= CPF_Parm;
+								*CurrentParamStorageLocation = Param;
+								CurrentParamStorageLocation = &Param->Next;
+							}
+						}
+						
+						// set the metdata:
+						NewFunction->SetMetaData(FBlueprintMetadata::MD_WorldContext, TEXT("__WorldContext"));
 					}
 
-					NewFunction->Bind();
-					NewFunction->StaticLink(true);
+					FKismetCompilerContext::SetCalculatedMetaDataAndFlags(NewFunction, EntryNode, Schema);
 				}
 			}
 		}
@@ -1114,8 +1333,8 @@ UClass* FBlueprintCompilationManagerImpl::FastGenerateSkeletonClass(UBlueprint*
 
 	Ret->SetSuperStruct(ParentClass);
 	
-	/*Ret->ClassFlags |= (ParentClass->ClassFlags & CLASS_Inherit);
-	Ret->ClassCastFlags |= ParentClass->ClassCastFlags;*/
+	Ret->ClassFlags |= (ParentClass->ClassFlags & CLASS_Inherit);
+	Ret->ClassCastFlags |= ParentClass->ClassCastFlags;
 	
 	if (FBlueprintEditorUtils::IsInterfaceBlueprint(BP))
 	{
@@ -1152,7 +1371,8 @@ UClass* FBlueprintCompilationManagerImpl::FastGenerateSkeletonClass(UBlueprint*
 				false, 
 				Event->FunctionFlags|FUNC_BlueprintCallable, 
 				TArray<UK2Node_FunctionResult*>(), 
-				Event->Pins
+				Event->Pins,
+				true
 			);
 
 			if(NewFunction)
@@ -1179,19 +1399,160 @@ UClass* FBlueprintCompilationManagerImpl::FastGenerateSkeletonClass(UBlueprint*
 	
 	AddFunctionForGraphs(TEXT(""), BP->FunctionGraphs, CurrentFieldStorageLocation, BPTYPE_FunctionLibrary == BP->BlueprintType);
 
+	// Add interface functions, often these are added by normal detection of implemented functions, but they won't be
+	// if the interface is added but the function is not implemented:
+	for(const FBPInterfaceDescription& BPID : BP->ImplementedInterfaces)
+	{
+		UClass* InterfaceClass = BPID.Interface;
+		if(UBlueprint* Owner = Cast<UBlueprint>(InterfaceClass->ClassGeneratedBy))
+		{
+			if( ensure(Owner->SkeletonGeneratedClass) )
+			{
+				InterfaceClass = Owner->SkeletonGeneratedClass;
+			}
+		}
+
+		for (TFieldIterator<UFunction> FunctionIt(InterfaceClass, EFieldIteratorFlags::ExcludeSuper); FunctionIt; ++FunctionIt)
+		{
+			UFunction* Fn = *FunctionIt;
+			
+			UField** CurrentParamStorageLocation = nullptr;
+
+			// Note that MakeFunction will early out if the function was created above:
+			MakeFunction(
+				Fn->GetFName(), 
+				CurrentFieldStorageLocation, 
+				CurrentParamStorageLocation, 
+				false, 
+				Fn->FunctionFlags & ~FUNC_Native, 
+				TArray<UK2Node_FunctionResult*>(), 
+				TArray<UEdGraphPin*>(),
+				false
+			);
+		}
+	}
+
 	CompilerContext.NewClass = Ret;
 	CompilerContext.bAssignDelegateSignatureFunction = true;
 	CompilerContext.FinishCompilingClass(Ret);
 	CompilerContext.bAssignDelegateSignatureFunction = false;
 	CompilerContext.NewClass = OriginalNewClass;
 
-	Ret->Bind();
-	Ret->StaticLink(/*bRelinkExistingProperties =*/true);
+	Ret->GetDefaultObject()->SetFlags(RF_Transient);
 
 	return Ret;
 }
 
-// Singleton boilperplate:
+bool FBlueprintCompilationManagerImpl::IsQueuedForCompilation(UBlueprint* BP)
+{
+	return BP->bQueuedForCompilation;
+}
+
+// FFixupBytecodeReferences Implementation:
+FBlueprintCompilationManagerImpl::FFixupBytecodeReferences::FFixupBytecodeReferences(UObject* InObject)
+{
+	ArIsObjectReferenceCollector = true;
+		
+	InObject->Serialize(*this);
+	class FArchiveProxyCollector : public FReferenceCollector
+	{
+		/** Archive we are a proxy for */
+		FArchive& Archive;
+	public:
+		FArchiveProxyCollector(FArchive& InArchive)
+			: Archive(InArchive)
+		{
+		}
+		virtual void HandleObjectReference(UObject*& Object, const UObject* ReferencingObject, const UProperty* ReferencingProperty) override
+		{
+			Archive << Object;
+		}
+		virtual void HandleObjectReferences(UObject** InObjects, const int32 ObjectNum, const UObject* InReferencingObject, const UProperty* InReferencingProperty) override
+		{
+			for (int32 ObjectIndex = 0; ObjectIndex < ObjectNum; ++ObjectIndex)
+			{
+				UObject*& Object = InObjects[ObjectIndex];
+				Archive << Object;
+			}
+		}
+		virtual bool IsIgnoringArchetypeRef() const override
+		{
+			return false;
+		}
+		virtual bool IsIgnoringTransient() const override
+		{
+			return false;
+		}
+	} ArchiveProxyCollector(*this);
+		
+	InObject->GetClass()->CallAddReferencedObjects(InObject, ArchiveProxyCollector);
+}
+
+FArchive& FBlueprintCompilationManagerImpl::FFixupBytecodeReferences::operator<<( UObject*& Obj )
+{
+	if (Obj != NULL)
+	{
+		if(UClass* RelatedClass = Cast<UClass>(Obj))
+		{
+			UClass* NewClass = RelatedClass->GetAuthoritativeClass();
+			ensure(NewClass);
+			if(NewClass != RelatedClass)
+			{
+				Obj = NewClass;
+			}
+		}
+		else if(UField* AsField = Cast<UField>(Obj))
+		{
+			UClass* OwningClass = AsField->GetOwnerClass();
+			if(OwningClass)
+			{
+				UClass* NewClass = OwningClass->GetAuthoritativeClass();
+				ensure(NewClass);
+				if(NewClass != OwningClass)
+				{
+					// drill into new class finding equivalent object:
+					TArray<FName> Names;
+					UObject* Iter = Obj;
+					while (Iter && Iter != OwningClass)
+					{
+						Names.Add(Iter->GetFName());
+						Iter = Iter->GetOuter();
+					}
+
+					UObject* Owner = NewClass;
+					UObject* Match = nullptr;
+					for(int32 I = Names.Num() - 1; I >= 0; --I)
+					{
+						UObject* Next = StaticFindObjectFast( UObject::StaticClass(), Owner, Names[I]);
+						if( Next )
+						{
+							if(I == 0)
+							{
+								Match = Next;
+							}
+							else
+							{
+								Owner = Match;
+							}
+						}
+						else
+						{
+							break;
+						}
+					}
+						
+					if(Match)
+					{
+						Obj = Match;
+					}
+				}
+			}
+		}
+	}
+	return *this;
+}
+
+// Singleton boilerplate, simply forwarding to the implementation above:
 FBlueprintCompilationManagerImpl* BPCMImpl = nullptr;
 
 void FlushReinstancingQueueImplWrapper()
@@ -1218,14 +1579,16 @@ void FBlueprintCompilationManager::FlushCompilationQueue(TArray<UObject*>* ObjLo
 {
 	if(BPCMImpl)
 	{
-		BPCMImpl->FlushCompilationQueueImpl(ObjLoaded);
+		BPCMImpl->FlushCompilationQueueImpl(ObjLoaded, false);
 	}
 }
 
-void FBlueprintCompilationManager::QueueForCompilation(UBlueprint* BPToCompile)
+void FBlueprintCompilationManager::CompileSynchronously(const FBPCompileRequest& Request)
 {
-	check(BPToCompile->GetLinker());
-	BPCMImpl->QueueForCompilation(BPToCompile);
+	if(BPCMImpl)
+	{
+		BPCMImpl->CompileSynchronouslyImpl(Request);
+	}
 }
 	
 void FBlueprintCompilationManager::NotifyBlueprintLoaded(UBlueprint* BPLoaded)
@@ -1235,8 +1598,14 @@ void FBlueprintCompilationManager::NotifyBlueprintLoaded(UBlueprint* BPLoaded)
 	{
 		FBlueprintCompilationManager::Initialize();
 	}
+
+	if(FBlueprintEditorUtils::IsCompileOnLoadDisabled(BPLoaded))
+	{
+		return;
+	}
+
 	check(BPLoaded->GetLinker());
-	BPCMImpl->QueueForCompilation(BPLoaded);
+	BPCMImpl->QueueForCompilation(FBPCompileRequest(BPLoaded, EBlueprintCompileOptions::IsRegeneratingOnLoad, nullptr));
 }
 
 bool FBlueprintCompilationManager::IsGeneratedClassLayoutReady()
diff --git a/Engine/Source/Editor/Kismet/Private/BlueprintDetailsCustomization.cpp b/Engine/Source/Editor/Kismet/Private/BlueprintDetailsCustomization.cpp
--- a/Engine/Source/Editor/Kismet/Private/BlueprintDetailsCustomization.cpp
+++ b/Engine/Source/Editor/Kismet/Private/BlueprintDetailsCustomization.cpp
@@ -54,6 +54,7 @@
 #include "SKismetInspector.h"
 #include "SSCSEditor.h"
 #include "SPinTypeSelector.h"
+#include "NodeFactory.h"
 #include "Kismet2/Kismet2NameValidators.h"
 #include "Widgets/Layout/SWidgetSwitcher.h"
 
@@ -1989,7 +1990,7 @@ bool FBlueprintVarActionDetails::IsConfigCheckBoxEnabled() const
 		if (UProperty* VariableProperty = CachedVariableProperty.Get())
 		{
 			// meant to match up with UHT's FPropertyBase::IsObject(), which it uses to block object properties from being marked with CPF_Config
-			bEnabled = VariableProperty->IsA<UClassProperty>() || VariableProperty->IsA<UAssetClassProperty>() || 
+			bEnabled = VariableProperty->IsA<UClassProperty>() || VariableProperty->IsA<UAssetClassProperty>() || VariableProperty->IsA<UAssetObjectProperty>() ||
 				(!VariableProperty->IsA<UObjectPropertyBase>() && !VariableProperty->IsA<UInterfaceProperty>());
 		}
 	}
@@ -2513,8 +2514,12 @@ void FBlueprintVarActionDetails::OnFinishedChangingProperties(const FPropertyCha
 			// Search out the correct local variable in the Function Entry Node and set the default value
 			for(auto& LocalVar : FuncEntry->LocalVariables)
 			{
-				if(LocalVar.VarName == DirectProperty->GetFName())
+				if(LocalVar.VarName == DirectProperty->GetFName() && LocalVar.DefaultValue != DefaultValueString)
 				{
+					const FScopedTransaction Transaction(LOCTEXT("ChangeDefaults", "Change Defaults"));
+
+					FuncEntry->Modify();
+					GetBlueprintObj()->Modify();
 					LocalVar.DefaultValue = DefaultValueString;
 					FBlueprintEditorUtils::MarkBlueprintAsModified(GetBlueprintObj());
 					break;
@@ -2701,22 +2706,39 @@ void FBlueprintGraphArgumentLayout::GenerateChildContent( IDetailChildrenBuilder
 {
 	if (bHasDefaultValue)
 	{
-		ChildrenBuilder.AddChildContent( LOCTEXT( "FunctionArgDetailsDefaultValue", "Default Value" ) )
-		.NameContent()
-		[
-			SNew(STextBlock)
-				.Text( LOCTEXT( "FunctionArgDetailsDefaultValue", "Default Value" ) )
-				.ToolTipText( LOCTEXT("FunctionArgDetailsDefaultValueParamTooltip", "The default value of the parameter.") )
-				.Font( IDetailLayoutBuilder::GetDetailFont() )
-		]
-		.ValueContent()
-		[
-			SNew(SEditableTextBox)
-				.Text( this, &FBlueprintGraphArgumentLayout::OnGetArgDefaultValueText )
-				.OnTextCommitted( this, &FBlueprintGraphArgumentLayout::OnArgDefaultValueCommitted )
-				.IsEnabled(!ShouldPinBeReadOnly())
-				.Font( IDetailLayoutBuilder::GetDetailFont() )
-		];
+		if (UEdGraphPin* FoundPin = GetPin())
+		{
+			// Certain types are outlawed at the compiler level
+			bool bTypeWithNoDefaults = (FoundPin->PinType.PinCategory == UEdGraphSchema_K2::PC_Object) || (FoundPin->PinType.PinCategory == UEdGraphSchema_K2::PC_Class) || (FoundPin->PinType.PinCategory == UEdGraphSchema_K2::PC_Interface);
+
+			if (!FoundPin->PinType.bIsReference && !bTypeWithNoDefaults)
+			{
+				DefaultValuePinWidget = FNodeFactory::CreatePinWidget(FoundPin);
+				DefaultValuePinWidget->SetOnlyShowDefaultValue(true);
+				TSharedRef<SWidget> DefaultValueWidget = DefaultValuePinWidget->GetDefaultValueWidget();
+
+				if (DefaultValueWidget != SNullWidget::NullWidget)
+				{
+					ChildrenBuilder.AddChildContent(LOCTEXT("FunctionArgDetailsDefaultValue", "Default Value"))
+						.NameContent()
+						[
+							SNew(STextBlock)
+							.Text(LOCTEXT("FunctionArgDetailsDefaultValue", "Default Value"))
+						.ToolTipText(LOCTEXT("FunctionArgDetailsDefaultValueParamTooltip", "The default value of the parameter."))
+						.Font(IDetailLayoutBuilder::GetDetailFont())
+						]
+					.ValueContent()
+						.MaxDesiredWidth(512)
+						[
+							DefaultValueWidget
+						];
+				}
+				else
+				{
+					DefaultValuePinWidget.Reset();
+				}
+			}
+		}
 
 		ChildrenBuilder.AddChildContent( LOCTEXT( "FunctionArgDetailsPassByReference", "Pass-by-Reference" ) )
 		.NameContent()
@@ -2926,6 +2948,15 @@ FEdGraphPinType FBlueprintGraphArgumentLayout::OnGetPinInfo() const
 	return FEdGraphPinType();
 }
 
+UEdGraphPin* FBlueprintGraphArgumentLayout::GetPin() const
+{
+	if (ParamItemPtr.IsValid() && TargetNode)
+	{
+		return TargetNode->FindPin(ParamItemPtr.Pin()->PinName, ParamItemPtr.Pin()->DesiredPinDirection);
+	}
+	return nullptr;
+}
+
 ECheckBoxState FBlueprintGraphArgumentLayout::IsRefChecked() const
 {
 	FEdGraphPinType PinType = OnGetPinInfo();
@@ -2936,9 +2967,9 @@ void FBlueprintGraphArgumentLayout::OnRefCheckStateChanged(ECheckBoxState InStat
 {
 	FEdGraphPinType PinType = OnGetPinInfo();
 	PinType.bIsReference = (InState == ECheckBoxState::Checked)? true : false;
-	// Note: Array types are implicitly passed by reference. For custom event nodes, the reference flag is essentially
-	//  treated as being redundant on array inputs, but we also need to implicitly set the 'const' flag to avoid a compiler note.
-	PinType.bIsConst = (PinType.bIsArray || PinType.bIsReference) && TargetNode && TargetNode->IsA<UK2Node_CustomEvent>();
+	// Note: Container types are implicitly passed by reference. For custom event nodes, the reference flag is essentially
+	//  treated as being redundant on container inputs, but we also need to implicitly set the 'const' flag to avoid a compiler note.
+	PinType.bIsConst = (PinType.IsContainer() || PinType.bIsReference) && TargetNode && TargetNode->IsA<UK2Node_CustomEvent>();
 	PinInfoChanged(PinType);
 }
 
@@ -2970,8 +3001,11 @@ void FBlueprintGraphArgumentLayout::PinInfoChanged(const FEdGraphPinType& PinTyp
 						{
 							(*UDPinPtr)->PinType = PinType;
 
-							// Array types are implicitly passed by reference. For custom event nodes, since they are inputs, also implicitly treat them as 'const' so that they don't result in a compiler note.
-							(*UDPinPtr)->PinType.bIsConst = PinType.bIsArray && Node->IsA<UK2Node_CustomEvent>();
+							// Container types are implicitly passed by reference. For custom event nodes, since they are inputs, also implicitly treat them as 'const' so that they don't result in a compiler note.
+							(*UDPinPtr)->PinType.bIsConst = PinType.IsContainer() && Node->IsA<UK2Node_CustomEvent>();
+
+							// Reset default value, it probably doesn't match
+							(*UDPinPtr)->PinDefaultValue.Reset();
 						}
 						GraphActionDetailsPinned->OnParamsChanged(Node);
 					}
@@ -2996,33 +3030,6 @@ void FBlueprintGraphArgumentLayout::OnPrePinInfoChange(const FEdGraphPinType& Pi
 	}
 }
 
-FText FBlueprintGraphArgumentLayout::OnGetArgDefaultValueText() const
-{
-	if (ParamItemPtr.IsValid())
-	{
-		return FText::FromString(ParamItemPtr.Pin()->PinDefaultValue);
-	}
-	return FText();
-}
-
-void FBlueprintGraphArgumentLayout::OnArgDefaultValueCommitted(const FText& NewText, ETextCommit::Type InTextCommit)
-{
-	auto GraphActionDetailsPinned = GraphActionDetailsPtr.Pin();
-	if (!NewText.IsEmpty() 
-		&& !ShouldPinBeReadOnly() 
-		&& (InTextCommit == ETextCommit::OnEnter || InTextCommit == ETextCommit::OnUserMovedFocus) 
-		&& ParamItemPtr.IsValid() 
-		&& GraphActionDetailsPinned.IsValid())
-	{
-		bool bSuccess = TargetNode->ModifyUserDefinedPinDefaultValue(ParamItemPtr.Pin(), NewText.ToString());
-		if (bSuccess)
-		{
-			GraphActionDetailsPinned->OnParamsChanged(TargetNode);
-		}
-	}
-}
-
-
 BEGIN_SLATE_FUNCTION_BUILD_OPTIMIZATION
 void FBlueprintGraphActionDetails::CustomizeDetails( IDetailLayoutBuilder& DetailLayout )
 {
@@ -3513,6 +3520,11 @@ void FBlueprintGraphActionDetails::CustomizeDetails( IDetailLayoutBuilder& Detai
 			];
 		}
 	}
+
+	if (bHasAGraph)
+	{
+		MyRegisteredGraphChangedDelegateHandle = GetGraph()->AddOnGraphChangedHandler(FOnGraphChanged::FDelegate::CreateSP(this, &FBaseBlueprintGraphActionDetails::OnGraphChanged));
+	}
 }
 END_SLATE_FUNCTION_BUILD_OPTIMIZATION
 
@@ -3597,6 +3609,21 @@ bool FBaseBlueprintGraphActionDetails::AttemptToCreateResultNode()
 	return FunctionResultNodePtr.IsValid();
 }
 
+FBaseBlueprintGraphActionDetails::~FBaseBlueprintGraphActionDetails()
+{
+	UEdGraph* MyGraph = GetGraph();
+	if (MyRegisteredGraphChangedDelegateHandle.IsValid() && MyGraph)
+	{
+		MyGraph->RemoveOnGraphChangedHandler(MyRegisteredGraphChangedDelegateHandle);
+	}
+}
+
+void FBaseBlueprintGraphActionDetails::OnGraphChanged(const FEdGraphEditAction& Action)
+{
+	/** Graph changed, need to refresh inputs in case pin UI changed */
+	RegenerateInputsChildrenDelegate.ExecuteIfBound();
+}
+
 void FBaseBlueprintGraphActionDetails::SetRefreshDelegate(FSimpleDelegate RefreshDelegate, bool bForInputs)
 {
 	((bForInputs) ? RegenerateInputsChildrenDelegate : RegenerateOutputsChildrenDelegate) = RefreshDelegate;
@@ -3962,28 +3989,10 @@ void FBaseBlueprintGraphActionDetails::OnParamsChanged(UK2Node_EditablePinBase*
 
 		// Reconstruct the entry/exit definition and recompile the blueprint to make sure the signature has changed before any fixups
 		TargetNode->ReconstructNode();
-		FParamsChangedHelper ParamsChangedHelper;
-		ParamsChangedHelper.ModifiedBlueprints.Add(GetBlueprintObj());
-		FBlueprintEditorUtils::MarkBlueprintAsStructurallyModified(GetBlueprintObj());
-
-		ParamsChangedHelper.Broadcast(GetBlueprintObj(), TargetNode, Graph);
 
-		for (auto GraphIt = ParamsChangedHelper.ModifiedGraphs.CreateIterator(); GraphIt; ++GraphIt)
-		{
-			if(UEdGraph* ModifiedGraph = *GraphIt)
-			{
-				ModifiedGraph->NotifyGraphChanged();
-			}
-		}
+		const UEdGraphSchema_K2* K2Schema = GetDefault<UEdGraphSchema_K2>();
 
-		// Now update all the blueprints that got modified
-		for (auto BlueprintIt = ParamsChangedHelper.ModifiedBlueprints.CreateIterator(); BlueprintIt; ++BlueprintIt)
-		{
-			if(UBlueprint* Blueprint = *BlueprintIt)
-			{
-				Blueprint->BroadcastChanged();
-			}
-		}
+		K2Schema->HandleParameterDefaultValueChanged(TargetNode);
 	}
 }
 
diff --git a/Engine/Source/Editor/Kismet/Private/BlueprintDetailsCustomization.h b/Engine/Source/Editor/Kismet/Private/BlueprintDetailsCustomization.h
--- a/Engine/Source/Editor/Kismet/Private/BlueprintDetailsCustomization.h
+++ b/Engine/Source/Editor/Kismet/Private/BlueprintDetailsCustomization.h
@@ -15,6 +15,7 @@
 #include "Widgets/Views/STableViewBase.h"
 #include "Widgets/Views/STableRow.h"
 #include "SMyBlueprint.h"
+#include "SGraphPin.h"
 
 class Error;
 class FBlueprintGlobalOptionsDetails;
@@ -285,7 +286,7 @@ private:
 class FBaseBlueprintGraphActionDetails : public IDetailCustomization
 {
 public:
-	virtual ~FBaseBlueprintGraphActionDetails(){};
+	virtual ~FBaseBlueprintGraphActionDetails();
 
 	FBaseBlueprintGraphActionDetails(TWeakPtr<SMyBlueprint> InMyBlueprint)
 		: MyBlueprint(InMyBlueprint)
@@ -324,14 +325,13 @@ public:
 	
 	FReply OnAddNewInputClicked();
 
+	/** Called when blueprint graph changes */
+	void OnGraphChanged(const FEdGraphEditAction& Action);
+
 protected:
 	/** Tries to create the result node (if there are output args) */
 	bool AttemptToCreateResultNode();
 
-	/** Toggles the ability to be called in editor */
-	void OnToggleEditorCallableEvent( const ECheckBoxState NewCheckedState, TWeakObjectPtr<UK2Node_EditablePinBase> SelectedNode ) const;
-
-protected:
 	/** Pointer to the parent */
 	TWeakPtr<SMyBlueprint> MyBlueprint;
 	
@@ -347,6 +347,9 @@ protected:
 
 	/** Details layout builder we need to hold on to to refresh it at times */
 	IDetailLayoutBuilder* DetailsLayoutPtr;
+
+	/** Handle for graph refresh delegate */
+	FDelegateHandle MyRegisteredGraphChangedDelegateHandle;
 	
 	/** Array of nodes were were constructed to represent */
 	TArray< TWeakObjectPtr<UObject> > ObjectsBeingEdited;
@@ -469,15 +472,15 @@ private:
 	void PinInfoChanged(const FEdGraphPinType& PinType);
 	void OnPrePinInfoChange(const FEdGraphPinType& PinType);
 
+	/** Returns the graph pin representing this variable */
+	UEdGraphPin* GetPin() const;
+
 	/** Returns whether the "Pass-by-Reference" checkbox is checked or not */
 	ECheckBoxState IsRefChecked() const;
 
 	/** Handles toggling the "Pass-by-Reference" checkbox */
 	void OnRefCheckStateChanged(ECheckBoxState InState);
 
-	FText OnGetArgDefaultValueText() const;
-	void OnArgDefaultValueCommitted(const FText& NewText, ETextCommit::Type InTextCommit);
-
 private:
 	/** The parent graph action details customization */
 	TWeakPtr<class FBaseBlueprintGraphActionDetails> GraphActionDetailsPtr;
@@ -496,6 +499,9 @@ private:
 
 	/** Holds a weak pointer to the argument name widget, used for error notifications */
 	TWeakPtr<SEditableTextBox> ArgumentNameWidget;
+
+	/** The SGraphPin widget created to show/edit default value */
+	TSharedPtr<SGraphPin> DefaultValuePinWidget;
 };
 
 /** Details customization for functions and graphs selected in the MyBlueprint panel */
diff --git a/Engine/Source/Editor/Kismet/Private/BlueprintEditor.cpp b/Engine/Source/Editor/Kismet/Private/BlueprintEditor.cpp
--- a/Engine/Source/Editor/Kismet/Private/BlueprintEditor.cpp
+++ b/Engine/Source/Editor/Kismet/Private/BlueprintEditor.cpp
@@ -996,6 +996,11 @@ TSharedRef<SGraphEditor> FBlueprintEditor::CreateGraphEditorWidget(TSharedRef<FT
 				FCanExecuteAction::CreateSP(this, &FBlueprintEditor::CanRestoreAllStructVarPins)
 				);
 
+			GraphEditorCommands->MapAction(FGraphEditorCommands::Get().ResetPinToDefaultValue,
+				FExecuteAction::CreateSP(this, &FBlueprintEditor::OnResetPinToDefaultValue),
+				FCanExecuteAction::CreateSP(this, &FBlueprintEditor::CanResetPinToDefaultValue)
+			);
+
 			GraphEditorCommands->MapAction( FGraphEditorCommands::Get().AddOptionPin,
 				FExecuteAction::CreateSP( this, &FBlueprintEditor::OnAddOptionPin ),
 				FCanExecuteAction::CreateSP( this, &FBlueprintEditor::CanAddOptionPin )
@@ -3983,25 +3988,24 @@ void  FBlueprintEditor::OnRemoveExecutionPin()
 
 bool FBlueprintEditor::CanRemoveExecutionPin() const
 {
-	bool bReturnValue = true;
+	TSharedPtr<SGraphEditor> FocusedGraphEd = FocusedGraphEdPtr.Pin();
+	if (FocusedGraphEd.IsValid())
+	{
+		const FScopedTransaction Transaction(LOCTEXT("RemoveExecutionPin", "Remove Execution Pin"));
 
-	const FGraphPanelSelectionSet& SelectedNodes = GetSelectedNodes();
+		UEdGraphPin* SelectedPin = FocusedGraphEd->GetGraphPinForMenu();
+		UEdGraphNode* OwningNode = SelectedPin->GetOwningNode();
 
-	// Iterate over all nodes, and make sure all execution sequence nodes will always have at least 2 outs
-	for (FGraphPanelSelectionSet::TConstIterator It(SelectedNodes); It; ++It)
-	{
-		UK2Node_ExecutionSequence* SeqNode = Cast<UK2Node_ExecutionSequence>(*It);
-		if (SeqNode != NULL)
+		if (UK2Node_ExecutionSequence* SeqNode = Cast<UK2Node_ExecutionSequence>(OwningNode))
 		{
-			bReturnValue = SeqNode->CanRemoveExecutionPin();
-			if (!bReturnValue)
-			{
-				break;
-			}
+			return SeqNode->CanRemoveExecutionPin();
+		}
+		else if (UK2Node_Switch* SwitchNode = Cast<UK2Node_Switch>(OwningNode))
+		{
+			return SwitchNode->CanRemoveExecutionPin(SelectedPin);
 		}
 	}
-
-	return bReturnValue;
+	return false;
 }
 
 void  FBlueprintEditor::OnRemoveThisStructVarPin()
@@ -4094,6 +4098,39 @@ bool FBlueprintEditor::CanRestoreAllStructVarPins() const
 	return Node && !Node->AllPinsAreShown();
 }
 
+void FBlueprintEditor::OnResetPinToDefaultValue()
+{
+	TSharedPtr<SGraphEditor> FocusedGraphEd = FocusedGraphEdPtr.Pin();
+	if (FocusedGraphEd.IsValid())
+	{
+		UEdGraphPin* TargetPin = FocusedGraphEd->GetGraphPinForMenu();
+
+		check(TargetPin);
+
+		const FScopedTransaction Transaction(LOCTEXT("ResetPinToDefaultValue", "Reset Pin To Default Value"));
+
+		const UEdGraphSchema_K2* K2Schema = GetDefault<UEdGraphSchema_K2>();
+		K2Schema->ResetPinToAutogeneratedDefaultValue(TargetPin);
+	}
+}
+
+bool FBlueprintEditor::CanResetPinToDefaultValue() const
+{
+	const UEdGraphSchema_K2* K2Schema = GetDefault<UEdGraphSchema_K2>();
+
+	bool bCanRecombine = false;
+	TSharedPtr<SGraphEditor> FocusedGraphEd = FocusedGraphEdPtr.Pin();
+	if (FocusedGraphEd.IsValid())
+	{
+		if (UEdGraphPin* Pin = FocusedGraphEd->GetGraphPinForMenu())
+		{
+			return !Pin->DoesDefaultValueMatchAutogenerated();
+		}
+	}
+
+	return false;
+}
+
 void FBlueprintEditor::OnAddOptionPin()
 {
 	const FGraphPanelSelectionSet& SelectedNodes = GetSelectedNodes();
@@ -6207,11 +6244,11 @@ void FBlueprintEditor::OnDisallowedPinConnection(const UEdGraphPin* PinA, const
 {
 	FDisallowedPinConnection NewRecord;
 	NewRecord.PinTypeCategoryA = PinA->PinType.PinCategory;
-	NewRecord.bPinIsArrayA = PinA->PinType.bIsArray;
+	NewRecord.bPinIsArrayA = PinA->PinType.IsArray();
 	NewRecord.bPinIsReferenceA = PinA->PinType.bIsReference;
 	NewRecord.bPinIsWeakPointerA = PinA->PinType.bIsWeakPointer;
 	NewRecord.PinTypeCategoryB = PinB->PinType.PinCategory;
-	NewRecord.bPinIsArrayB = PinB->PinType.bIsArray;
+	NewRecord.bPinIsArrayB = PinB->PinType.IsArray();
 	NewRecord.bPinIsReferenceB = PinB->PinType.bIsReference;
 	NewRecord.bPinIsWeakPointerB = PinB->PinType.bIsWeakPointer;
 	AnalyticsStats.GraphDisallowedPinConnections.Add(NewRecord);
diff --git a/Engine/Source/Editor/Kismet/Private/FindInBlueprintManager.cpp b/Engine/Source/Editor/Kismet/Private/FindInBlueprintManager.cpp
--- a/Engine/Source/Editor/Kismet/Private/FindInBlueprintManager.cpp
+++ b/Engine/Source/Editor/Kismet/Private/FindInBlueprintManager.cpp
@@ -635,7 +635,7 @@ namespace BlueprintSearchMetaDataHelpers
 		{
 			InWriter->WriteValue(FFindInBlueprintSearchTags::FiB_ObjectClass, FText::FromString(InPinType.PinSubCategoryObject->GetName()));
 		}
-		InWriter->WriteValue(FFindInBlueprintSearchTags::FiB_IsArray, InPinType.bIsArray);
+		InWriter->WriteValue(FFindInBlueprintSearchTags::FiB_IsArray, InPinType.IsArray());
 		InWriter->WriteValue(FFindInBlueprintSearchTags::FiB_IsReference, InPinType.bIsReference);
 	}
 
diff --git a/Engine/Source/Editor/Kismet/Private/FindInBlueprints.cpp b/Engine/Source/Editor/Kismet/Private/FindInBlueprints.cpp
--- a/Engine/Source/Editor/Kismet/Private/FindInBlueprints.cpp
+++ b/Engine/Source/Editor/Kismet/Private/FindInBlueprints.cpp
@@ -92,7 +92,7 @@ bool FindInBlueprintsHelpers::ParsePinType(FText InKey, FText InValue, FEdGraphP
 	}
 	else if(InKey.CompareTo(FFindInBlueprintSearchTags::FiB_IsArray) == 0)
 	{
-		InOutPinType.bIsArray = InValue.ToString().ToBool();
+		InOutPinType.ContainerType = (InValue.ToString().ToBool() ? EPinContainerType::Array : EPinContainerType::None);
 	}
 	else if(InKey.CompareTo(FFindInBlueprintSearchTags::FiB_IsReference) == 0)
 	{
@@ -332,9 +332,9 @@ FFindInBlueprintsPin::FFindInBlueprintsPin(const FText& InValue, TSharedPtr<FFin
 
 TSharedRef<SWidget> FFindInBlueprintsPin::CreateIcon() const
 {
-	const FSlateBrush* Brush = NULL;
+	const FSlateBrush* Brush = nullptr;
 
-	if( PinType.bIsArray )
+	if( PinType.IsArray() )
 	{
 		Brush = FEditorStyle::GetBrush( TEXT("GraphEditor.ArrayPinIcon") );
 	}
diff --git a/Engine/Source/Editor/Kismet/Private/SBlueprintActionMenu.cpp b/Engine/Source/Editor/Kismet/Private/SBlueprintActionMenu.cpp
--- a/Engine/Source/Editor/Kismet/Private/SBlueprintActionMenu.cpp
+++ b/Engine/Source/Editor/Kismet/Private/SBlueprintActionMenu.cpp
@@ -228,7 +228,7 @@ void SBlueprintActionMenu::Construct( const FArguments& InArgs, TSharedPtr<FBlue
 	//@TODO: Should probably be somewhere more schema-sensitive than the graph panel!
 	FSlateColor TypeColor;
 	FString TypeOfDisplay;
-	const FSlateBrush* ContextIcon = NULL;
+	const FSlateBrush* ContextIcon = nullptr;
 
 	if (DraggedFromPins.Num() == 1)
 	{
@@ -241,7 +241,7 @@ void SBlueprintActionMenu::Construct( const FArguments& InArgs, TSharedPtr<FBlue
 		{
 			// Get the type color and icon
 			TypeColor = Schema->GetPinTypeColor(OnePin->PinType);
-			ContextIcon = FEditorStyle::GetBrush( OnePin->PinType.bIsArray ? TEXT("Graph.ArrayPin.Connected") : TEXT("Graph.Pin.Connected") );
+			ContextIcon = FEditorStyle::GetBrush( OnePin->PinType.IsArray() ? TEXT("Graph.ArrayPin.Connected") : TEXT("Graph.Pin.Connected") );
 		}
 	}
 
diff --git a/Engine/Source/Editor/Kismet/Private/SBlueprintEditorToolbar.cpp b/Engine/Source/Editor/Kismet/Private/SBlueprintEditorToolbar.cpp
--- a/Engine/Source/Editor/Kismet/Private/SBlueprintEditorToolbar.cpp
+++ b/Engine/Source/Editor/Kismet/Private/SBlueprintEditorToolbar.cpp
@@ -254,11 +254,8 @@ static void OnDiffRevisionPicked(FRevisionInfo const& RevisionInfo, TWeakObjectP
 					FString PreviousTempPkgName;
 					if (Revision->Get(PreviousTempPkgName))
 					{
-						// Forcibly disable compile on load in case we are loading old blueprints that might try to update/compile
-						TGuardValue<bool> DisableCompileOnLoad(GForceDisableBlueprintCompileOnLoad, true);
-
 						// Try and load that package
-						UPackage* PreviousTempPkg = LoadPackage(NULL, *PreviousTempPkgName, LOAD_None);
+						UPackage* PreviousTempPkg = LoadPackage(NULL, *PreviousTempPkgName, LOAD_DisableCompileOnLoad);
 
 						if (PreviousTempPkg != NULL)
 						{
diff --git a/Engine/Source/Editor/Kismet/Private/SBlueprintPalette.cpp b/Engine/Source/Editor/Kismet/Private/SBlueprintPalette.cpp
--- a/Engine/Source/Editor/Kismet/Private/SBlueprintPalette.cpp
+++ b/Engine/Source/Editor/Kismet/Private/SBlueprintPalette.cpp
@@ -1332,6 +1332,11 @@ void SBlueprintPaletteItem::OnNameTextCommitted(const FText& NewText, ETextCommi
 				// Search through all function entry nodes for local variables to update their scope name
 				TArray<UK2Node_Variable*> VariableNodes;
 				Graph->GetNodesOfClass<UK2Node_Variable>(VariableNodes);
+				for (const UEdGraph* SubGraph : Graph->SubGraphs)
+				{
+					check(SubGraph != nullptr);
+					SubGraph->GetNodesOfClass<UK2Node_Variable>(VariableNodes);
+				}
 
 				for (UK2Node_Variable* const VariableNode : VariableNodes)
 				{
diff --git a/Engine/Source/Editor/Kismet/Private/SSCSEditor.cpp b/Engine/Source/Editor/Kismet/Private/SSCSEditor.cpp
--- a/Engine/Source/Editor/Kismet/Private/SSCSEditor.cpp
+++ b/Engine/Source/Editor/Kismet/Private/SSCSEditor.cpp
@@ -32,7 +32,7 @@
 #include "Kismet2/BlueprintEditorUtils.h"
 #include "ComponentAssetBroker.h"
 #include "ClassViewerFilter.h"
-
+#include "SSearchBox.h"
 #include "PropertyPath.h"
 
 #include "AssetSelection.h"
@@ -190,23 +190,26 @@ FReply SSCSEditor::TryHandleAssetDragDropOperation(const FDragDropEvent& DragDro
 					}
 				}
 
+				// Only set focus to the last item created
+				const bool bSetFocusToNewItem = (DroppedAssetIdx == NumAssets - 1);
+
 				TSubclassOf<UActorComponent> MatchingComponentClassForAsset = FComponentAssetBrokerage::GetPrimaryComponentForAsset(AssetClass);
 				if (MatchingComponentClassForAsset != nullptr)
 				{
-					AddNewComponent(MatchingComponentClassForAsset, Asset, true );
+					AddNewComponent(MatchingComponentClassForAsset, Asset, true, bSetFocusToNewItem );
 					bMarkBlueprintAsModified = true;
 				}
 				else if ((PotentialComponentClass != nullptr) && !PotentialComponentClass->HasAnyClassFlags(CLASS_Deprecated | CLASS_Abstract | CLASS_NewerVersionExists))
 				{
 					if (PotentialComponentClass->HasMetaData(FBlueprintMetadata::MD_BlueprintSpawnableComponent))
 					{
-						AddNewComponent(PotentialComponentClass, nullptr, true );
+						AddNewComponent(PotentialComponentClass, nullptr, true, bSetFocusToNewItem );
 						bMarkBlueprintAsModified = true;
 					}
 				}
 				else if ((PotentialActorClass != nullptr) && !PotentialActorClass->HasAnyClassFlags(CLASS_Deprecated | CLASS_Abstract | CLASS_NewerVersionExists))
 				{
-					AddNewComponent(UChildActorComponent::StaticClass(), PotentialActorClass, true );
+					AddNewComponent(UChildActorComponent::StaticClass(), PotentialActorClass, true, bSetFocusToNewItem );
 					bMarkBlueprintAsModified = true;
 				}
 			}
@@ -376,6 +379,7 @@ FSCSEditorTreeNode::FSCSEditorTreeNode(FSCSEditorTreeNode::ENodeType InNodeType)
 	: ComponentTemplatePtr(nullptr)
 	, NodeType(InNodeType)
 	, bNonTransactionalRename(false)
+	, FilterFlags((uint8)EFilteredState::Unknown)
 {
 }
 
@@ -476,6 +480,105 @@ bool FSCSEditorTreeNode::IsAttachedTo(FSCSEditorTreeNodePtrType InNodePtr) const
 	return false; 
 }
 
+void FSCSEditorTreeNode::UpdateCachedFilterState(bool bMatchesFilter, bool bUpdateParent)
+{
+	bool bFlagsChanged = false;
+	if ((FilterFlags & EFilteredState::Unknown) == EFilteredState::Unknown)
+	{
+		FilterFlags   = 0x00;
+		bFlagsChanged = true;
+	}
+
+	if (bMatchesFilter)
+	{
+		bFlagsChanged |= (FilterFlags & EFilteredState::MatchesFilter) == 0;
+		FilterFlags |= EFilteredState::MatchesFilter;
+	}
+	else
+	{
+		bFlagsChanged |= (FilterFlags & EFilteredState::MatchesFilter) != 0;
+		FilterFlags &= ~EFilteredState::MatchesFilter;
+	}
+
+	const bool bHadChildMatch = (FilterFlags & EFilteredState::ChildMatches) != 0;
+	// refresh the cached child state (don't update the parent, we'll do that below if it's needed)
+	RefreshCachedChildFilterState(/*bUpdateParent =*/false);
+
+	bFlagsChanged |= bHadChildMatch != ((FilterFlags & EFilteredState::ChildMatches) != 0);
+	if (bUpdateParent && bFlagsChanged)
+	{
+		ApplyFilteredStateToParent();
+	}
+}
+
+void FSCSEditorTreeNode::RefreshCachedChildFilterState(bool bUpdateParent)
+{
+	const bool bCointainedMatch = !IsFlaggedForFiltration();
+
+	FilterFlags &= ~EFilteredState::ChildMatches;
+	for (FSCSEditorTreeNodePtrType Child : Children)
+	{
+		if (!Child->IsFlaggedForFiltration())
+		{
+			FilterFlags |= EFilteredState::ChildMatches;
+			break;
+		}
+	}
+	const bool bCointainsMatch = !IsFlaggedForFiltration();
+
+	const bool bStateChange = bCointainedMatch != bCointainsMatch;
+	if (bUpdateParent && bStateChange)
+	{
+		ApplyFilteredStateToParent();
+	}
+}
+
+void FSCSEditorTreeNode::ApplyFilteredStateToParent()
+{
+	FSCSEditorTreeNode* Child = this;
+	while (Child->ParentNodePtr.IsValid())
+	{
+		FSCSEditorTreeNode* Parent = Child->ParentNodePtr.Get();
+
+		if ( !IsFlaggedForFiltration() )
+		{
+			if ((Parent->FilterFlags & EFilteredState::ChildMatches) == 0)
+			{
+				Parent->FilterFlags |= EFilteredState::ChildMatches;
+			}
+			else
+			{
+				// all parents from here on up should have the flag
+				break;
+			}
+		}
+		// have to see if this was the only child contributing to this flag
+		else if (Parent->FilterFlags & EFilteredState::ChildMatches)
+		{
+			Parent->FilterFlags &= ~EFilteredState::ChildMatches;
+			for (const FSCSEditorTreeNodePtrType& Sibling : Parent->Children)
+			{
+				if (Sibling.Get() == Child)
+				{
+					continue;
+				}
+
+				if (Sibling->FilterFlags & EFilteredState::FilteredInMask)
+				{
+					Parent->FilterFlags |= EFilteredState::ChildMatches;
+					break;
+				}
+			}
+
+			if (Parent->FilterFlags & EFilteredState::ChildMatches)
+			{
+				// another child added the flag back
+				break;
+			}
+		}
+		Child = Parent;
+	}
+}
 
 FSCSEditorTreeNodePtrType FSCSEditorTreeNode::FindClosestParent(TArray<FSCSEditorTreeNodePtrType> InNodes)
 {
@@ -516,6 +619,16 @@ void FSCSEditorTreeNode::AddChild(FSCSEditorTreeNodePtrType InChildNodePtr)
 	Children.AddUnique(InChildNodePtr);
 	InChildNodePtr->ParentNodePtr = AsShared();
 
+	if (InChildNodePtr->FilterFlags != EFilteredState::Unknown && !InChildNodePtr->IsFlaggedForFiltration())
+	{
+		FSCSEditorTreeNodePtrType AncestorPtr = InChildNodePtr->ParentNodePtr;
+		while (AncestorPtr.IsValid() && (AncestorPtr->FilterFlags & EFilteredState::ChildMatches) == 0)
+		{
+			AncestorPtr->FilterFlags |= EFilteredState::ChildMatches;
+			AncestorPtr = AncestorPtr->GetParent();
+		}
+	}
+
 	// Add a child node to the SCS tree node if not already present
 	USCS_Node* SCS_ChildNode = InChildNodePtr->GetSCSNode();
 	if(SCS_ChildNode != NULL)
@@ -769,6 +882,11 @@ void FSCSEditorTreeNode::RemoveChild(FSCSEditorTreeNodePtrType InChildNodePtr)
 	Children.Remove(InChildNodePtr);
 	InChildNodePtr->ParentNodePtr.Reset();
 	InChildNodePtr->RemoveMeAsChild();
+
+	if (InChildNodePtr->IsFlaggedForFiltration())
+	{
+		RefreshCachedChildFilterState(/*bUpdateParent =*/true);
+	}
 }
 
 void FSCSEditorTreeNode::OnRequestRename(bool bTransactional)
@@ -1299,8 +1417,8 @@ void SSCS_RowWidget::Construct( const FArguments& InArgs, TSharedPtr<SSCSEditor>
 	
 	auto Args = FSuperRowType::FArguments()
 		.Style(bIsSeparator ?
-				&FEditorStyle::Get().GetWidgetStyle<FTableRowStyle>("TableView.NoHoverTableRow") :
-				&FEditorStyle::Get().GetWidgetStyle<FTableRowStyle>("SceneOutliner.TableViewRow")) //@todo create editor style for the SCS tree
+			&FEditorStyle::Get().GetWidgetStyle<FTableRowStyle>("TableView.NoHoverTableRow") :
+			&FEditorStyle::Get().GetWidgetStyle<FTableRowStyle>("SceneOutliner.TableViewRow")) //@todo create editor style for the SCS tree
 		.Padding(FMargin(0.f, 0.f, 0.f, 4.f))
 		.ShowSelection(!bIsSeparator)
 		.OnDragDetected(this, &SSCS_RowWidget::HandleOnDragDetected)
@@ -2355,7 +2473,7 @@ void SSCS_RowWidget::OnAttachToDropAction(const TArray<FSCSEditorTreeNodePtrType
 				check(ComponentTemplate);
 
 				// Note: This will mark the Blueprint as structurally modified
-				UActorComponent* ClonedComponent = SCSEditorPtr->AddNewComponent(ComponentTemplate->GetClass(), NULL);
+				UActorComponent* ClonedComponent = SCSEditorPtr->AddNewComponent(ComponentTemplate->GetClass(), nullptr);
 				check(ClonedComponent);
 
 				//Serialize object properties using write/read operations.
@@ -2650,7 +2768,7 @@ void SSCS_RowWidget::OnMakeNewRootDropAction(FSCSEditorTreeNodePtrType DroppedNo
 	{
 		// Get the current Blueprint context
 		UBlueprint* Blueprint = GetBlueprint();
-		check(Blueprint != NULL && Blueprint->SimpleConstructionScript != nullptr);
+		check(Blueprint && Blueprint->SimpleConstructionScript);
 
 		// Clone the component if it's being dropped into a different SCS
 		if(DroppedNodePtr->GetBlueprint() != Blueprint)
@@ -2659,7 +2777,7 @@ void SSCS_RowWidget::OnMakeNewRootDropAction(FSCSEditorTreeNodePtrType DroppedNo
 			check(ComponentTemplate);
 
 			// Note: This will mark the Blueprint as structurally modified
-			UActorComponent* ClonedComponent = SCSEditorPtr->AddNewComponent(ComponentTemplate->GetClass(), NULL);
+			UActorComponent* ClonedComponent = SCSEditorPtr->AddNewComponent(ComponentTemplate->GetClass(), nullptr);
 			check(ClonedComponent);
 
 			//Serialize object properties using write/read operations.
@@ -3338,7 +3456,7 @@ void SSCSEditor::Construct( const FArguments& InArgs )
 	SCSTreeWidget = SNew(SSCSTreeType)
 		.ToolTipText(LOCTEXT("DropAssetToAddComponent", "Drop asset here to add a component."))
 		.SCSEditor(this)
-		.TreeItemsSource(&RootNodes)
+		.TreeItemsSource(&FilteredRootNodes)
 		.SelectionMode(ESelectionMode::Multi)
 		.OnGenerateRow(this, &SSCSEditor::MakeTableRowWidget)
 		.OnGetChildren(this, &SSCSEditor::OnGetChildrenForTree)
@@ -3441,6 +3559,15 @@ void SSCSEditor::Construct( const FArguments& InArgs )
 					.IsEnabled(AllowEditing)
 				]
 
+				+ SHorizontalBox::Slot()
+					.FillWidth(1.0f)
+					.VAlign(VAlign_Center)
+					.Padding(3.0f, 3.0f)
+				[
+					SAssignNew(FilterBox, SSearchBox)
+						.OnTextChanged(this, &SSCSEditor::OnFilterTextChanged)
+				]
+
 				+ SHorizontalBox::Slot()
 				.FillWidth(1.0f)
 				.HAlign(HAlign_Right)
@@ -3884,8 +4011,7 @@ void SSCSEditor::OnDuplicateComponent()
 
 		for (int32 i = 0; i < SelectedNodes.Num(); ++i)
 		{
-			UActorComponent* ComponentTemplate = SelectedNodes[i]->GetComponentTemplate();
-			if(ComponentTemplate != NULL)
+			if (UActorComponent* ComponentTemplate = SelectedNodes[i]->GetComponentTemplate())
 			{
 				UActorComponent* CloneComponent = AddNewComponent(ComponentTemplate->GetClass(), ComponentTemplate);
 				UActorComponent* OriginalComponent = ComponentTemplate;
@@ -3935,9 +4061,25 @@ void SSCSEditor::OnDuplicateComponent()
 
 void SSCSEditor::OnGetChildrenForTree( FSCSEditorTreeNodePtrType InNodePtr, TArray<FSCSEditorTreeNodePtrType>& OutChildren )
 {
-	if(InNodePtr.IsValid())
+	if (InNodePtr.IsValid())
 	{
-		OutChildren = InNodePtr->GetChildren();
+		const TArray<FSCSEditorTreeNodePtrType>& Children = InNodePtr->GetChildren();
+		OutChildren.Reserve(Children.Num());
+
+		if (!GetFilterText().IsEmpty())
+		{
+			for (FSCSEditorTreeNodePtrType Child : Children)
+			{
+				if (!Child->IsFlaggedForFiltration())
+				{
+					OutChildren.Add(Child);
+				}
+			}
+		}
+		else
+		{
+			OutChildren = Children;
+		}
 	}
 	else
 	{
@@ -4274,7 +4416,7 @@ void SSCSEditor::UpdateTree(bool bRegenerateTreeNodes)
 				}
 				
 				for (UActorComponent* Component : Components)
-				{				
+				{
 					if (USceneComponent* SceneComp = Cast<USceneComponent>(Component))
 					{
 						// Add the rest of the native base class SceneComponent hierarchy
@@ -4445,6 +4587,8 @@ void SSCSEditor::UpdateTree(bool bRegenerateTreeNodes)
 				SCSTreeWidget->RequestScrollIntoView(NodeToRenamePtr);
 			}
 		}
+
+		RebuildFilteredRootList();
 	}
 
 	// refresh widget
@@ -4625,11 +4769,18 @@ TSharedPtr<FSCSEditorTreeNode> SSCSEditor::AddRootComponentTreeNode(UActorCompon
 	if (RootTreeNode.IsValid())
 	{
 		NewTreeNode = RootTreeNode->AddChildFromComponent(ActorComp);
+		RefreshFilteredState(NewTreeNode, /*bRecursive =*/false);
 	}
 	else
 	{
 		NewTreeNode = FSCSEditorTreeNode::FactoryNodeFromComponent(ActorComp);
 		RootNodes.Add(NewTreeNode);
+
+		bool bIsFilteredOut = RefreshFilteredState(NewTreeNode, /*bRecursive =*/false);
+		if (!bIsFilteredOut)
+		{
+			FilteredRootNodes.Add(NewTreeNode);
+		}
 	}
 
 	RootComponentNodes.Add(NewTreeNode);
@@ -4735,6 +4886,36 @@ UClass* SSCSEditor::CreateNewBPComponent(TSubclassOf<UActorComponent> ComponentC
 	return NewClass;
 }
 
+void SSCSEditor::RebuildFilteredRootList()
+{
+ 	FilteredRootNodes.Empty(RootNodes.Num());
+
+	FSCSEditorTreeNodePtrType PendingSeparator;
+	for (const FSCSEditorTreeNodePtrType& Node : RootNodes)
+	{
+		switch (Node->GetNodeType())
+		{
+		case FSCSEditorTreeNode::ENodeType::ComponentNode:
+			if (Node->IsFlaggedForFiltration())
+			{	
+				break;
+			}
+		case FSCSEditorTreeNode::ENodeType::RootActorNode:
+			if (PendingSeparator.IsValid())
+			{
+				FilteredRootNodes.Add(PendingSeparator);
+				PendingSeparator.Reset();
+			}
+			FilteredRootNodes.Add(Node);
+			break;
+
+		case FSCSEditorTreeNode::ENodeType::SeparatorNode:
+			PendingSeparator = Node;
+			break;
+		}
+	}
+}
+
 void SSCSEditor::ClearSelection()
 {
 	if ( bUpdatingSelection == false )
@@ -4776,7 +4957,7 @@ bool SSCSEditor::IsEditingAllowed() const
 	return AllowEditing.Get() && nullptr == GEditor->PlayWorld;
 }
 
-UActorComponent* SSCSEditor::AddNewComponent( UClass* NewComponentClass, UObject* Asset, const bool bSkipMarkBlueprintModified )
+UActorComponent* SSCSEditor::AddNewComponent( UClass* NewComponentClass, UObject* Asset, const bool bSkipMarkBlueprintModified, const bool bSetFocusToNewItem )
 {
 	if (NewComponentClass->ClassWithin && NewComponentClass->ClassWithin != UObject::StaticClass())
 	{
@@ -4816,7 +4997,7 @@ UActorComponent* SSCSEditor::AddNewComponent( UClass* NewComponentClass, UObject
 		}
 		
 		const FName NewVariableName = (Asset ? FName(*FComponentEditorUtils::GenerateValidVariableNameFromAsset(Asset, nullptr)) : NAME_None);
-		NewComponent = AddNewNode(Blueprint->SimpleConstructionScript->CreateNode(NewComponentClass, NewVariableName), Asset, bMarkBlueprintModified);
+		NewComponent = AddNewNode(Blueprint->SimpleConstructionScript->CreateNode(NewComponentClass, NewVariableName), Asset, bMarkBlueprintModified, bSetFocusToNewItem);
 
 		if (ComponentTemplate)
 		{
@@ -4839,7 +5020,7 @@ UActorComponent* SSCSEditor::AddNewComponent( UClass* NewComponentClass, UObject
 		if (ComponentTemplate)
 		{
 			// Create a duplicate of the provided template
-			NewComponent = AddNewNodeForInstancedComponent(FComponentEditorUtils::DuplicateComponent(ComponentTemplate), nullptr);
+			NewComponent = AddNewNodeForInstancedComponent(FComponentEditorUtils::DuplicateComponent(ComponentTemplate), nullptr, bSetFocusToNewItem);
 		}
 		else if (AActor* ActorInstance = GetActorContext())
 		{
@@ -4901,7 +5082,7 @@ UActorComponent* SSCSEditor::AddNewComponent( UClass* NewComponentClass, UObject
 			// Rerun construction scripts
 			ActorInstance->RerunConstructionScripts();
 
-			NewComponent = AddNewNodeForInstancedComponent(NewInstanceComponent, Asset);
+			NewComponent = AddNewNodeForInstancedComponent(NewInstanceComponent, Asset, bSetFocusToNewItem);
 		}
 	}
 
@@ -5478,6 +5659,7 @@ FSCSEditorTreeNodePtrType SSCSEditor::AddTreeNode(USCS_Node* InSCSNode, FSCSEdit
 		{
 			// do this first, because we need a FSCSEditorTreeNodePtrType for the new node
 			NewNodePtr = ParentPtr->AddChild(InSCSNode, bIsInheritedSCS);
+			RefreshFilteredState(NewNodePtr, /*bRecursive =*/false);
 
 			bool bParentIsEditorOnly = ParentPtr->GetComponentTemplate()->IsEditorOnly();
 			// if you can't nest this new node under the proposed parent (then swap the two)
@@ -5520,6 +5702,12 @@ FSCSEditorTreeNodePtrType SSCSEditor::AddTreeNode(USCS_Node* InSCSNode, FSCSEdit
 			{
 				NewNodePtr = MakeShareable(new FSCSEditorTreeNodeComponent(InSCSNode, bIsInheritedSCS));
 				RootNodes.Add(NewNodePtr);
+
+				bool bIsFilteredOut = RefreshFilteredState(NewNodePtr, /*bRecursive =*/false);
+				if (!bIsFilteredOut)
+				{
+					FilteredRootNodes.Add(NewNodePtr);
+				}
 			}
 			
 			NodeSCS->AddNode(InSCSNode);
@@ -5545,6 +5733,12 @@ FSCSEditorTreeNodePtrType SSCSEditor::AddTreeNode(USCS_Node* InSCSNode, FSCSEdit
 		{
 			NewNodePtr = MakeShareable(new FSCSEditorTreeNodeComponent(InSCSNode, bIsInheritedSCS));
 			RootNodes.Add(NewNodePtr);
+
+			bool bIsFilteredOut = RefreshFilteredState(NewNodePtr, /*bRecursive =*/false);
+			if (!bIsFilteredOut)
+			{
+				FilteredRootNodes.Add(NewNodePtr);
+			}
 		}
 
 		RootComponentNodes.Add(NewNodePtr);
@@ -5594,6 +5788,7 @@ FSCSEditorTreeNodePtrType SSCSEditor::AddTreeNodeFromComponent(USceneComponent*
 		// Add a new tree node for the given scene component
 		check(ParentNodePtr.IsValid());
 		NewNodePtr = ParentNodePtr->AddChildFromComponent(InSceneComponent);
+		RefreshFilteredState(NewNodePtr, /*bRecursive =*/false);
 
 		// Expand parent nodes by default
 		SCSTreeWidget->SetItemExpansion(ParentNodePtr, true);
@@ -5615,6 +5810,7 @@ FSCSEditorTreeNodePtrType SSCSEditor::AddTreeNodeFromComponent(USceneComponent*
 		else if (SceneRootNodePtr->GetComponentTemplate() != InSceneComponent)
 		{
 			NewNodePtr = SceneRootNodePtr->AddChildFromComponent(InSceneComponent);
+			RefreshFilteredState(NewNodePtr, /*bRecursive =*/false);
 		}
 	}
 
@@ -6171,5 +6367,114 @@ void SSCSEditor::SetItemExpansionRecursive(FSCSEditorTreeNodePtrType Model, bool
 	}
 }
 
+FText SSCSEditor::GetFilterText() const
+{
+	return FilterBox->GetText();
+}
+
+void SSCSEditor::OnFilterTextChanged(const FText& InFilterText)
+{
+	struct OnFilterTextChanged_Inner
+	{
+		static FSCSEditorTreeNodePtrType ExpandToFilteredChildren(SSCSEditor* SCSEditor, FSCSEditorTreeNodePtrType TreeNode)
+		{
+			FSCSEditorTreeNodePtrType NodeToFocus;
+
+			const TArray<FSCSEditorTreeNodePtrType>& Children = TreeNode->GetChildren();
+			// iterate backwards so we select from the top down
+			for (int32 ChildIndex = Children.Num() - 1; ChildIndex >= 0; --ChildIndex)
+			{
+				const FSCSEditorTreeNodePtrType& Child = Children[ChildIndex];
+				if (!Child->IsFlaggedForFiltration())
+				{
+					SCSEditor->SetNodeExpansionState(TreeNode, /*bIsExpanded =*/true);
+					NodeToFocus = ExpandToFilteredChildren(SCSEditor, Child);
+				}
+			}
+
+			if (!NodeToFocus.IsValid() && !TreeNode->IsFlaggedForFiltration())
+			{
+				NodeToFocus = TreeNode;
+			}
+			return NodeToFocus;
+		}
+	};
+
+	FSCSEditorTreeNodePtrType NewSelection;
+	const bool bIsFilterBlank = GetFilterText().IsEmpty();
+
+	bool bRootItemFilteredBackIn = false;
+	// iterate backwards so we select from the top down
+	for (int32 ComponentIndex = RootComponentNodes.Num() - 1; ComponentIndex >= 0; --ComponentIndex)
+	{
+		FSCSEditorTreeNodePtrType Component = RootComponentNodes[ComponentIndex];
+
+		const bool bWasFilteredOut = Component->IsFlaggedForFiltration();
+		bool bFilteredOut = RefreshFilteredState(Component, /*bRecursive =*/true);
+
+		if (!bFilteredOut)
+		{
+			if (!bIsFilterBlank)
+			{
+				NewSelection = OnFilterTextChanged_Inner::ExpandToFilteredChildren(this, Component);
+			}
+			bRootItemFilteredBackIn |= bWasFilteredOut;
+		}
+		else
+		{
+			FilteredRootNodes.Remove(Component);
+		}
+	}
+
+	if (NewSelection.IsValid() && !SCSTreeWidget->IsItemSelected(NewSelection))
+	{
+		SelectNode(NewSelection, /*IsCntrlDown =*/false);
+	}
+	
+	if (bRootItemFilteredBackIn)
+	{
+		RebuildFilteredRootList();
+	}
+	UpdateTree(/*bRegenerateTreeNodes =*/false);
+}
+
+bool SSCSEditor::RefreshFilteredState(FSCSEditorTreeNodePtrType TreeNode, bool bRecursive)
+{
+	FString FilterText = FText::TrimPrecedingAndTrailing( GetFilterText() ).ToString();
+	TArray<FString> FilterTerms;
+	FilterText.ParseIntoArray(FilterTerms, TEXT(" "), /*CullEmpty =*/true);
+
+	struct RefreshFilteredState_Inner
+	{
+		static void RefreshFilteredState(FSCSEditorTreeNodePtrType TreeNodeIn, const TArray<FString>& FilterTermsIn, bool bRecursiveIn)
+		{
+			if (bRecursiveIn)
+			{
+				for (FSCSEditorTreeNodePtrType Child : TreeNodeIn->GetChildren())
+				{
+					RefreshFilteredState(Child, FilterTermsIn, bRecursiveIn);
+				}
+			}
+			
+			FString DisplayStr = TreeNodeIn->GetDisplayString();
+
+			bool bIsFilteredOut = false;
+			for (const FString& FilterTerm : FilterTermsIn)
+			{
+				if (!DisplayStr.Contains(FilterTerm))
+				{
+					bIsFilteredOut = true;
+				}
+			}
+			// if we're not recursing, then assume this is for a new node and we need to update the parent
+			// otherwise, assume the parent was hit as part of the recursion
+			TreeNodeIn->UpdateCachedFilterState(!bIsFilteredOut, /*bUpdateParent =*/!bRecursiveIn);
+		}
+	};
+
+	RefreshFilteredState_Inner::RefreshFilteredState(TreeNode, FilterTerms, bRecursive);
+	return TreeNode->IsFlaggedForFiltration();
+}
+
 #undef LOCTEXT_NAMESPACE
 
