Copying //UE4/Dev-Framework to //UE4/Dev-Main (Source: //UE4/Dev-Framework @ 3459469)

diff --git a/Engine/Source/Editor/BlueprintGraph/Private/K2Node_VariableSet.cpp b/Engine/Source/Editor/BlueprintGraph/Private/K2Node_VariableSet.cpp
--- a/Engine/Source/Editor/BlueprintGraph/Private/K2Node_VariableSet.cpp
+++ b/Engine/Source/Editor/BlueprintGraph/Private/K2Node_VariableSet.cpp
@@ -6,6 +6,7 @@
 #include "Engine/BlueprintGeneratedClass.h"
 #include "EdGraphSchema_K2.h"
 #include "K2Node_VariableGet.h"
+#include "K2Node_CallFunction.h"
 #include "Kismet2/BlueprintEditorUtils.h"
 #include "KismetCompiler.h"
 #include "VariableSetHandler.h"
@@ -80,10 +81,8 @@ UK2Node_VariableSet::UK2Node_VariableSet(const FObjectInitializer& ObjectInitial
 
 void UK2Node_VariableSet::AllocateDefaultPins()
 {
-	const UEdGraphSchema_K2* K2Schema = GetDefault<UEdGraphSchema_K2>();
-
-	CreatePin(EGPD_Input, K2Schema->PC_Exec, FString(), nullptr, K2Schema->PN_Execute);
-	CreatePin(EGPD_Output, K2Schema->PC_Exec, FString(), nullptr, K2Schema->PN_Then);
+	CreatePin(EGPD_Input, UEdGraphSchema_K2::PC_Exec, FString(), nullptr, UEdGraphSchema_K2::PN_Execute);
+	CreatePin(EGPD_Output, UEdGraphSchema_K2::PC_Exec, FString(), nullptr, UEdGraphSchema_K2::PN_Then);
 
 	if (GetVarName() != NAME_None)
 	{
@@ -103,10 +102,8 @@ void UK2Node_VariableSet::AllocateDefaultPins()
 
 void UK2Node_VariableSet::ReallocatePinsDuringReconstruction(TArray<UEdGraphPin*>& OldPins)
 {
-	const UEdGraphSchema_K2* K2Schema = GetDefault<UEdGraphSchema_K2>();
-
-	CreatePin(EGPD_Input, K2Schema->PC_Exec, FString(), nullptr, K2Schema->PN_Execute);
-	CreatePin(EGPD_Output, K2Schema->PC_Exec, FString(), nullptr, K2Schema->PN_Then);
+	CreatePin(EGPD_Input, UEdGraphSchema_K2::PC_Exec, FString(), nullptr, UEdGraphSchema_K2::PN_Execute);
+	CreatePin(EGPD_Output, UEdGraphSchema_K2::PC_Exec, FString(), nullptr, UEdGraphSchema_K2::PN_Then);
 
 	if (GetVarName() != NAME_None)
 	{
@@ -342,8 +339,7 @@ void UK2Node_VariableSet::CreateOutputPinTooltip()
 FText UK2Node_VariableSet::GetPinNameOverride(const UEdGraphPin& Pin) const
 {
 	// Stop the output pin for the variable, effectively the "get" pin, from displaying a name.
-	const UEdGraphSchema_K2* K2Schema = GetDefault<UEdGraphSchema_K2>();
-	if(Pin.ParentPin == nullptr && (Pin.Direction == EGPD_Output || Pin.PinType.PinCategory == K2Schema->PC_Exec))
+	if(Pin.ParentPin == nullptr && (Pin.Direction == EGPD_Output || Pin.PinType.PinCategory == UEdGraphSchema_K2::PC_Exec))
 	{
 		return FText::GetEmpty();
 	}
@@ -357,37 +353,75 @@ void UK2Node_VariableSet::ExpandNode(class FKismetCompilerContext& CompilerConte
 
 	if (CompilerContext.bIsFullCompile)
 	{
+		UProperty* VariableProperty = GetPropertyForVariable();
+
 		const UEdGraphSchema_K2* K2Schema = CompilerContext.GetSchema();
 
-		UEdGraphPin* Pin = FindPin(GetVariableOutputPinName());
-		if(Pin)
+		if (UEdGraphPin* VariableGetPin = FindPin(GetVariableOutputPinName()))
 		{
 			// If the output pin is linked, we need to spawn a separate "Get" node and hook it up.
-			if(Pin->LinkedTo.Num())
+			if (VariableGetPin->LinkedTo.Num())
 			{
-				UProperty* VariableProperty = GetPropertyForVariable();
-
-				if(VariableProperty)
+				if (VariableProperty)
 				{
 					UK2Node_VariableGet* VariableGetNode = CompilerContext.SpawnIntermediateNode<UK2Node_VariableGet>(this, SourceGraph);
 					VariableGetNode->VariableReference = VariableReference;
 					VariableGetNode->AllocateDefaultPins();
 					CompilerContext.MessageLog.NotifyIntermediateObjectCreation(VariableGetNode, this);
-					CompilerContext.MovePinLinksToIntermediate(*Pin, *VariableGetNode->FindPin(GetVarNameString()));
+					CompilerContext.MovePinLinksToIntermediate(*VariableGetPin, *VariableGetNode->FindPin(GetVarNameString()));
 
 					// Duplicate the connection to the self pin.
 					UEdGraphPin* SetSelfPin = K2Schema->FindSelfPin(*this, EGPD_Input);
 					UEdGraphPin* GetSelfPin = K2Schema->FindSelfPin(*VariableGetNode, EGPD_Input);
-					if(SetSelfPin && GetSelfPin)
+					if (SetSelfPin && GetSelfPin)
 					{
 						CompilerContext.CopyPinLinksToIntermediate(*SetSelfPin, *GetSelfPin);
 					}
 				}
 			}
-			Pins.Remove(Pin);
-			Pin->MarkPendingKill();
+			Pins.Remove(VariableGetPin);
+			VariableGetPin->MarkPendingKill();
+		}
+
+		// If property has a BlueprintSetter accessor, then replace the variable get node with a call function
+		if (VariableProperty)
+		{
+			const FString& SetFunctionName = VariableProperty->GetMetaData(FBlueprintMetadata::MD_PropertySetFunction);
+			if (!SetFunctionName.IsEmpty())
+			{
+				UClass* OwnerClass = VariableProperty->GetOwnerClass();
+				UFunction* SetFunction = OwnerClass->FindFunctionByName(*SetFunctionName);
+				check(SetFunction);
+
+				UK2Node_CallFunction* CallFuncNode = CompilerContext.SpawnIntermediateNode<UK2Node_CallFunction>(this, SourceGraph);
+				CallFuncNode->SetFromFunction(SetFunction);
+				CallFuncNode->AllocateDefaultPins();
+
+				// Move Exec pin connections
+				CompilerContext.MovePinLinksToIntermediate(*GetExecPin(), *CallFuncNode->GetExecPin());
+
+				// Move Then pin connections
+				CompilerContext.MovePinLinksToIntermediate(*FindPinChecked(UEdGraphSchema_K2::PN_Then, EGPD_Output), *CallFuncNode->GetThenPin());
+				
+				// Move Self pin connections
+				CompilerContext.MovePinLinksToIntermediate(*K2Schema->FindSelfPin(*this, EGPD_Input), *K2Schema->FindSelfPin(*CallFuncNode, EGPD_Input));
+
+				// Move Value pin connections
+				UEdGraphPin* SetFunctionValuePin = nullptr;
+				for (UEdGraphPin* CallFuncPin : CallFuncNode->Pins)
+				{
+					if (!K2Schema->IsMetaPin(*CallFuncPin))
+					{
+						check(CallFuncPin->Direction == EGPD_Input);
+						SetFunctionValuePin = CallFuncPin;
+						break;
+					}
+				}
+				check(SetFunctionValuePin);
+
+				CompilerContext.MovePinLinksToIntermediate(*FindPin(GetVarNameString(), EGPD_Input), *SetFunctionValuePin);
+			}
 		}
-		
 	}
 
 }
diff --git a/Engine/Source/Editor/BlueprintGraph/Private/VariableSetHandler.cpp b/Engine/Source/Editor/BlueprintGraph/Private/VariableSetHandler.cpp
--- a/Engine/Source/Editor/BlueprintGraph/Private/VariableSetHandler.cpp
+++ b/Engine/Source/Editor/BlueprintGraph/Private/VariableSetHandler.cpp
@@ -156,7 +156,7 @@ void FKCHandler_VariableSet::Transform(FKismetFunctionContext& Context, UEdGraph
 		if (SetNotify->ShouldFlushDormancyOnSet())
 		{
 			// Create CallFuncNode
-			UK2Node_CallFunction* CallFuncNode = Node->GetGraph()->CreateBlankNode<UK2Node_CallFunction>();
+			UK2Node_CallFunction* CallFuncNode = Node->GetGraph()->CreateIntermediateNode<UK2Node_CallFunction>();
 			CallFuncNode->FunctionReference.SetExternalMember(NAME_FlushNetDormancy, AActor::StaticClass() );
 			CallFuncNode->AllocateDefaultPins();
 
@@ -180,7 +180,7 @@ void FKCHandler_VariableSet::Transform(FKismetFunctionContext& Context, UEdGraph
 
 		if (SetNotify->HasLocalRepNotify())
 		{
-			UK2Node_CallFunction* CallFuncNode = Node->GetGraph()->CreateBlankNode<UK2Node_CallFunction>();
+			UK2Node_CallFunction* CallFuncNode = Node->GetGraph()->CreateIntermediateNode<UK2Node_CallFunction>();
 			CallFuncNode->FunctionReference.SetExternalMember(SetNotify->GetRepNotifyName(), SetNotify->GetVariableSourceClass() );
 			CallFuncNode->AllocateDefaultPins();
 
diff --git a/Engine/Source/Editor/ContentBrowser/Private/AssetContextMenu.cpp b/Engine/Source/Editor/ContentBrowser/Private/AssetContextMenu.cpp
--- a/Engine/Source/Editor/ContentBrowser/Private/AssetContextMenu.cpp
+++ b/Engine/Source/Editor/ContentBrowser/Private/AssetContextMenu.cpp
@@ -2668,7 +2668,7 @@ void FAssetContextMenu::ExecuteAssignChunkID()
 		const FVector2D CursorPos = FSlateApplication::Get().GetCursorPos();
 		FSlateRect Anchor(CursorPos.X, CursorPos.Y, CursorPos.X, CursorPos.Y);
 
-		FVector2D AdjustedSummonLocation = FSlateApplication::Get().CalculatePopupWindowPosition(Anchor, SColorPicker::DEFAULT_WINDOW_SIZE, Orient_Horizontal);
+		FVector2D AdjustedSummonLocation = FSlateApplication::Get().CalculatePopupWindowPosition(Anchor, SColorPicker::DEFAULT_WINDOW_SIZE, FVector2D::ZeroVector, Orient_Horizontal);
 
 		TSharedPtr<SWindow> Window = SNew(SWindow)
 			.AutoCenter(EAutoCenter::None)
diff --git a/Engine/Source/Editor/DetailCustomizations/Private/CurveColorCustomization.cpp b/Engine/Source/Editor/DetailCustomizations/Private/CurveColorCustomization.cpp
--- a/Engine/Source/Editor/DetailCustomizations/Private/CurveColorCustomization.cpp
+++ b/Engine/Source/Editor/DetailCustomizations/Private/CurveColorCustomization.cpp
@@ -410,7 +410,7 @@ FReply FCurveColorCustomization::OnCurvePreviewDoubleClick(const FGeometry& InMy
 			const FVector2D CursorPos = FSlateApplication::Get().GetCursorPos();
 			FSlateRect Anchor(CursorPos.X, CursorPos.Y, CursorPos.X, CursorPos.Y);
 
-			FVector2D AdjustedSummonLocation = FSlateApplication::Get().CalculatePopupWindowPosition( Anchor, FCurveColorCustomization::DEFAULT_WINDOW_SIZE, Orient_Horizontal );
+			FVector2D AdjustedSummonLocation = FSlateApplication::Get().CalculatePopupWindowPosition( Anchor, FCurveColorCustomization::DEFAULT_WINDOW_SIZE, FVector2D::ZeroVector, Orient_Horizontal );
 
 			TSharedPtr<SWindow> Window = SNew(SWindow)
 				.Title( FText::Format( LOCTEXT("WindowHeader", "{0} - Internal Color Curve Editor"), StructPropertyHandle->GetPropertyDisplayName()) )
diff --git a/Engine/Source/Editor/DetailCustomizations/Private/CurveStructCustomization.cpp b/Engine/Source/Editor/DetailCustomizations/Private/CurveStructCustomization.cpp
--- a/Engine/Source/Editor/DetailCustomizations/Private/CurveStructCustomization.cpp
+++ b/Engine/Source/Editor/DetailCustomizations/Private/CurveStructCustomization.cpp
@@ -390,7 +390,7 @@ FReply FCurveStructCustomization::OnCurvePreviewDoubleClick( const FGeometry& In
 			const FVector2D CursorPos = FSlateApplication::Get().GetCursorPos();
 			FSlateRect Anchor(CursorPos.X, CursorPos.Y, CursorPos.X, CursorPos.Y);
 
-			FVector2D AdjustedSummonLocation = FSlateApplication::Get().CalculatePopupWindowPosition( Anchor, FCurveStructCustomization::DEFAULT_WINDOW_SIZE, Orient_Horizontal );
+			FVector2D AdjustedSummonLocation = FSlateApplication::Get().CalculatePopupWindowPosition( Anchor, FCurveStructCustomization::DEFAULT_WINDOW_SIZE, FVector2D::ZeroVector, Orient_Horizontal );
 
 			TSharedPtr<SWindow> Window = SNew(SWindow)
 				.Title( FText::Format( LOCTEXT("WindowHeader", "{0} - Internal Curve Editor"), StructPropertyHandle->GetPropertyDisplayName()) )
diff --git a/Engine/Source/Editor/DetailCustomizations/Private/DirectoryPathStructCustomization.cpp b/Engine/Source/Editor/DetailCustomizations/Private/DirectoryPathStructCustomization.cpp
--- a/Engine/Source/Editor/DetailCustomizations/Private/DirectoryPathStructCustomization.cpp
+++ b/Engine/Source/Editor/DetailCustomizations/Private/DirectoryPathStructCustomization.cpp
@@ -30,8 +30,7 @@ void FDirectoryPathStructCustomization::CustomizeHeader( TSharedRef<IPropertyHan
 
 	const bool bRelativeToGameContentDir = StructPropertyHandle->HasMetaData( TEXT("RelativeToGameContentDir") );
 	const bool bUseRelativePath = StructPropertyHandle->HasMetaData( TEXT("RelativePath") );
-	const bool bLongPackageName = StructPropertyHandle->HasMetaData( TEXT("LongPackageName") );
-	const bool bContentDir = StructPropertyHandle->HasMetaData( TEXT("ContentDir") );
+	const bool bContentDir = StructPropertyHandle->HasMetaData( TEXT("ContentDir") ) || StructPropertyHandle->HasMetaData(TEXT("LongPackageName"));
 
 	AbsoluteGameContentDir = FPaths::ConvertRelativePathToFull(FPaths::GameContentDir());
 
@@ -60,7 +59,7 @@ void FDirectoryPathStructCustomization::CustomizeHeader( TSharedRef<IPropertyHan
 			PickerWidget = SAssignNew(BrowseButton, SButton)
 			.ButtonStyle( FEditorStyle::Get(), "HoverHintOnly" )
 			.ToolTipText( LOCTEXT( "FolderButtonToolTipText", "Choose a directory from this computer") )
-			.OnClicked( FOnClicked::CreateSP(this, &FDirectoryPathStructCustomization::OnPickDirectory, PathProperty.ToSharedRef(), bRelativeToGameContentDir, bUseRelativePath,bLongPackageName) )
+			.OnClicked( FOnClicked::CreateSP(this, &FDirectoryPathStructCustomization::OnPickDirectory, PathProperty.ToSharedRef(), bRelativeToGameContentDir, bUseRelativePath) )
 			.ContentPadding( 2.0f )
 			.ForegroundColor( FSlateColor::UseForeground() )
 			.IsFocusable( false )
@@ -127,7 +126,7 @@ FReply FDirectoryPathStructCustomization::OnPickContent(TSharedRef<IPropertyHand
 	return FReply::Handled();
 }
 
-FReply FDirectoryPathStructCustomization::OnPickDirectory(TSharedRef<IPropertyHandle> PropertyHandle, const bool bRelativeToGameContentDir, const bool bUseRelativePath, const bool bLongPackageName) const
+FReply FDirectoryPathStructCustomization::OnPickDirectory(TSharedRef<IPropertyHandle> PropertyHandle, const bool bRelativeToGameContentDir, const bool bUseRelativePath) const
 {
 	FString Directory;
 	IDesktopPlatform* DesktopPlatform = FDesktopPlatformModule::Get();
@@ -153,25 +152,11 @@ FReply FDirectoryPathStructCustomization::OnPickDirectory(TSharedRef<IPropertyHa
 				{
 					FEditorDirectories::Get().SetLastDirectory(ELastDirectory::GENERIC_IMPORT, Directory);
 
-					if (bLongPackageName)
-					{
-						FString LongPackageName;
-						FString StringFailureReason;
-						if (FPackageName::TryConvertFilenameToLongPackageName(Directory, LongPackageName, &StringFailureReason) == false)
-						{
-							StartDirectory = Directory;
-							FMessageDialog::Open(EAppMsgType::Ok, FText::FromString(StringFailureReason));
-							continue;
-						}
-						Directory = LongPackageName;
-					}
-
 					if (bRelativeToGameContentDir)
 					{
 						Directory = Directory.RightChop(AbsoluteGameContentDir.Len());
 					}
-
-					if (bUseRelativePath)
+					else if (bUseRelativePath)
 					{
 						Directory = IFileManager::Get().ConvertToRelativePath(*Directory);
 					}
diff --git a/Engine/Source/Editor/DetailCustomizations/Private/DirectoryPathStructCustomization.h b/Engine/Source/Editor/DetailCustomizations/Private/DirectoryPathStructCustomization.h
--- a/Engine/Source/Editor/DetailCustomizations/Private/DirectoryPathStructCustomization.h
+++ b/Engine/Source/Editor/DetailCustomizations/Private/DirectoryPathStructCustomization.h
@@ -30,7 +30,7 @@ private:
 	FReply OnPickContent(TSharedRef<IPropertyHandle> PropertyHandle) ;
 
 	/** Delegate used to display a directory picker */
-	FReply OnPickDirectory(TSharedRef<IPropertyHandle> PropertyHandle, const bool bRelativeToGameContentDir, const bool bUseRelativePaths, const bool bLongPackageName) const;
+	FReply OnPickDirectory(TSharedRef<IPropertyHandle> PropertyHandle, const bool bRelativeToGameContentDir, const bool bUseRelativePaths) const;
 
 	/** Check whether that the chosen path is valid */
 	bool IsValidPath(const FString& AbsolutePath, const bool bRelativeToGameContentDir, FText* const OutReason = nullptr) const;
diff --git a/Engine/Source/Editor/GameplayTasksEditor/Private/K2Node_LatentGameplayTaskCall.cpp b/Engine/Source/Editor/GameplayTasksEditor/Private/K2Node_LatentGameplayTaskCall.cpp
--- a/Engine/Source/Editor/GameplayTasksEditor/Private/K2Node_LatentGameplayTaskCall.cpp
+++ b/Engine/Source/Editor/GameplayTasksEditor/Private/K2Node_LatentGameplayTaskCall.cpp
@@ -433,6 +433,7 @@ bool UK2Node_LatentGameplayTaskCall::ConnectSpawnProperties(UClass* ClassToSpawn
 					continue;
 				}
 
+				// This is sloppy, we should be comparing to defaults mutch later in the compile process:
 				if (ClassToSpawn->ClassDefaultObject != nullptr)
 				{
 					// We don't want to generate an assignment node unless the default value 
@@ -496,9 +497,20 @@ bool UK2Node_LatentGameplayTaskCall::ConnectSpawnProperties(UClass* ClassToSpawn
 				}
 				else
 				{
-					// Move connection from the variable pin on the spawn node to the 'value' pin
-					CompilerContext.MovePinLinksToIntermediate(*SpawnVarPin, *ValuePin);
-					SetVarNode->PinConnectionListChanged(ValuePin);
+					// For non-array struct pins that are not linked, transfer the pin type so that the node will expand an auto-ref that will assign the value by-ref.
+					if (SpawnVarPin->PinType.IsArray() == false && SpawnVarPin->PinType.PinCategory == Schema->PC_Struct && SpawnVarPin->LinkedTo.Num() == 0)
+					{
+						ValuePin->PinType.PinCategory = SpawnVarPin->PinType.PinCategory;
+						ValuePin->PinType.PinSubCategory = SpawnVarPin->PinType.PinSubCategory;
+						ValuePin->PinType.PinSubCategoryObject = SpawnVarPin->PinType.PinSubCategoryObject;
+						CompilerContext.MovePinLinksToIntermediate(*SpawnVarPin, *ValuePin);
+					}
+					else
+					{
+						// Move connection from the variable pin on the spawn node to the 'value' pin
+						CompilerContext.MovePinLinksToIntermediate(*SpawnVarPin, *ValuePin);
+						SetVarNode->PinConnectionListChanged(ValuePin);
+					}
 				}
 			}
 		}
@@ -699,7 +711,7 @@ void UK2Node_LatentGameplayTaskCall::ExpandNode(class FKismetCompilerContext& Co
 	// Branch based on return value of Prespawn
 	// -------------------------------------------
 		
-	UK2Node_IfThenElse* BranchNode = SourceGraph->CreateBlankNode<UK2Node_IfThenElse>();
+	UK2Node_IfThenElse* BranchNode = SourceGraph->CreateIntermediateNode<UK2Node_IfThenElse>();
 	BranchNode->AllocateDefaultPins();
 	CompilerContext.MessageLog.NotifyIntermediateObjectCreation(BranchNode, this);
 
diff --git a/Engine/Source/Editor/GraphEditor/Private/GraphEditorDragDropAction.cpp b/Engine/Source/Editor/GraphEditor/Private/GraphEditorDragDropAction.cpp
--- a/Engine/Source/Editor/GraphEditor/Private/GraphEditorDragDropAction.cpp
+++ b/Engine/Source/Editor/GraphEditor/Private/GraphEditorDragDropAction.cpp
@@ -22,7 +22,7 @@ UEdGraphPin* FGraphEditorDragDropAction::GetHoveredPin() const
 
 UEdGraphNode* FGraphEditorDragDropAction::GetHoveredNode() const
 {
-	return HoveredNode.IsValid() ? HoveredNode->GetNodeObj() : NULL;
+	return HoveredNode.Get();
 }
 
 UEdGraph* FGraphEditorDragDropAction::GetHoveredGraph() const
@@ -55,6 +55,11 @@ void FGraphEditorDragDropAction::SetHoveredPin(UEdGraphPin* InPin)
 }
 
 void FGraphEditorDragDropAction::SetHoveredNode(const TSharedPtr<SGraphNode>& InNode)
+{
+	SetHoveredNode(InNode.IsValid() ? InNode->GetNodeObj() : nullptr);
+}
+
+void FGraphEditorDragDropAction::SetHoveredNode(UEdGraphNode* InNode)
 {
 	if (HoveredNode != InNode)
 	{
@@ -210,3 +215,23 @@ FReply FGraphSchemaActionDragDropAction::DroppedOnPanel( const TSharedRef< SWidg
 	}
 	return FReply::Unhandled();
 }
+
+FReply FGraphSchemaActionDragDropAction::DroppedOnPin(FVector2D ScreenPosition, FVector2D GraphPosition)
+{
+	if (UEdGraph* Graph = GetHoveredGraph())
+	{
+		if (ActionNode.IsValid())
+		{
+			TArray<UEdGraphPin*> DummyPins;
+			if (UEdGraphPin* Pin = GetHoveredPin())
+			{
+				DummyPins.Add(Pin);
+			}
+
+			ActionNode->PerformAction(Graph, DummyPins, GraphPosition);
+
+			return FReply::Handled();
+		}
+	}
+	return FReply::Unhandled();
+}
diff --git a/Engine/Source/Editor/GraphEditor/Public/GraphEditorDragDropAction.h b/Engine/Source/Editor/GraphEditor/Public/GraphEditorDragDropAction.h
--- a/Engine/Source/Editor/GraphEditor/Public/GraphEditorDragDropAction.h
+++ b/Engine/Source/Editor/GraphEditor/Public/GraphEditorDragDropAction.h
@@ -15,6 +15,9 @@ class SGraphPanel;
 class SWidget;
 class UEdGraph;
 struct FSlateBrush;
+class UEdGraphPin;
+class UEdGraphNode;
+class UEdGraphSchema;
 
 // Base class for drag-drop actions that pass into the graph editor and perform an action when dropped
 class GRAPHEDITOR_API FGraphEditorDragDropAction : public FDragDropOperation
@@ -22,9 +25,10 @@ class GRAPHEDITOR_API FGraphEditorDragDropAction : public FDragDropOperation
 public:
 	DRAG_DROP_OPERATOR_TYPE(FGraphEditorDragDropAction, FDragDropOperation)
 
-	void SetHoveredPin(class UEdGraphPin* InPin);
-	void SetHoveredNode(const TSharedPtr<class SGraphNode>& InNode);
-	void SetHoveredGraph(const TSharedPtr<class SGraphPanel>& InGraph);
+	void SetHoveredPin(UEdGraphPin* InPin);
+	void SetHoveredNode(const TSharedPtr<SGraphNode>& InNode);
+	void SetHoveredNode(UEdGraphNode* InNode);
+	void SetHoveredGraph(const TSharedPtr<SGraphPanel>& InGraph);
 	void SetHoveredCategoryName(const FText& InHoverCategoryName);
 	void SetHoveredAction(TSharedPtr<struct FEdGraphSchemaAction> Action);
 	void SetDropTargetValid( bool bValid ) { bDropTargetValid = bValid; }
@@ -33,12 +37,12 @@ public:
 	virtual void HoverTargetChanged() {}
 	virtual FReply DroppedOnPin(FVector2D ScreenPosition, FVector2D GraphPosition) { return FReply::Unhandled(); }
 	virtual FReply DroppedOnNode(FVector2D ScreenPosition, FVector2D GraphPosition) { return FReply::Unhandled(); }
-	virtual FReply DroppedOnPanel( const TSharedRef< class SWidget >& Panel, FVector2D ScreenPosition, FVector2D GraphPosition, UEdGraph& Graph) { return FReply::Unhandled(); }
+	virtual FReply DroppedOnPanel( const TSharedRef< SWidget >& Panel, FVector2D ScreenPosition, FVector2D GraphPosition, UEdGraph& Graph) { return FReply::Unhandled(); }
 	virtual FReply DroppedOnAction(TSharedRef<struct FEdGraphSchemaAction> Action) { return FReply::Unhandled(); }
 	virtual FReply DroppedOnCategory(FText Category) { return FReply::Unhandled(); }
 	// End of interface to override
 	
-	virtual bool IsSupportedBySchema(const class UEdGraphSchema* Schema) const { return true; }
+	virtual bool IsSupportedBySchema(const UEdGraphSchema* Schema) const { return true; }
 
 	bool HasFeedbackMessage();
 	void SetFeedbackMessage(const TSharedPtr<SWidget>& Message);
@@ -61,10 +65,10 @@ private:
 	FEdGraphPinReference HoveredPin;
 
 	// The node that the drag action is currently hovering over
-	TSharedPtr<class SGraphNode> HoveredNode;
+	TWeakObjectPtr<UEdGraphNode> HoveredNode;
 
 	// The graph that the drag action is currently hovering over
-	TSharedPtr<class SGraphPanel> HoveredGraph;
+	TSharedPtr<SGraphPanel> HoveredGraph;
 
 protected:
 
@@ -87,7 +91,8 @@ public:
 
 	// FGraphEditorDragDropAction interface
 	virtual void HoverTargetChanged() override;
-	virtual FReply DroppedOnPanel( const TSharedRef< class SWidget >& Panel, FVector2D ScreenPosition, FVector2D GraphPosition, UEdGraph& Graph) override;
+	virtual FReply DroppedOnPanel(const TSharedRef< SWidget >& Panel, FVector2D ScreenPosition, FVector2D GraphPosition, UEdGraph& Graph) override;
+	virtual FReply DroppedOnPin(FVector2D ScreenPosition, FVector2D GraphPosition) override;
 	// End of FGraphEditorDragDropAction
 
 	static TSharedRef<FGraphSchemaActionDragDropAction> New(TSharedPtr<FEdGraphSchemaAction> InActionNode )
diff --git a/Engine/Source/Editor/Kismet/Private/BPFunctionDragDropAction.cpp b/Engine/Source/Editor/Kismet/Private/BPFunctionDragDropAction.cpp
--- a/Engine/Source/Editor/Kismet/Private/BPFunctionDragDropAction.cpp
+++ b/Engine/Source/Editor/Kismet/Private/BPFunctionDragDropAction.cpp
@@ -221,26 +221,44 @@ void FKismetFunctionDragDropAction::HoverTargetChanged()
 
 //------------------------------------------------------------------------------
 FReply FKismetFunctionDragDropAction::DroppedOnPanel(TSharedRef<SWidget> const& Panel, FVector2D ScreenPosition, FVector2D GraphPosition, UEdGraph& Graph)
-{	
+{
+	return DroppedOnPin(ScreenPosition, GraphPosition);
+}
+
+//------------------------------------------------------------------------------
+FReply FKismetFunctionDragDropAction::DroppedOnPin(FVector2D ScreenPosition, FVector2D GraphPosition)
+{
 	FReply Reply = FReply::Unhandled();
 
+	UEdGraph* Graph = GetHoveredGraph();
+	check(Graph);
+
 	// The ActionNode set during construction points to the Graph, this is suitable for displaying the mouse decorator but needs to be more complete based on the current graph
-	UBlueprintFunctionNodeSpawner* FunctionNodeSpawner = GetDropAction(Graph);
+	UBlueprintFunctionNodeSpawner* FunctionNodeSpawner = GetDropAction(*Graph);
 
 	if (FunctionNodeSpawner)
 	{
 		FText CannotDropReason = FText::GetEmpty();
-		if (!CanBeDroppedDelegate.IsBound() || CanBeDroppedDelegate.Execute(nullptr, GetHoveredGraph(), CannotDropReason))
+		if (!CanBeDroppedDelegate.IsBound() || CanBeDroppedDelegate.Execute(nullptr, Graph, CannotDropReason))
 		{
 			UFunction const* Function = GetFunctionProperty();
 			if ((Function != NULL) && UEdGraphSchema_K2::CanUserKismetCallFunction(Function))
 			{
 				AnalyticCallback.ExecuteIfBound();
 
-				const FScopedTransaction Transaction( LOCTEXT("KismetFunction_DroppedOnPanel", "Function Dropped on Graph") );
+				const FScopedTransaction Transaction(LOCTEXT("KismetFunction_DroppedOnPanel", "Function Dropped on Graph"));
 
 				IBlueprintNodeBinder::FBindingSet Bindings;
-				FunctionNodeSpawner->Invoke(&Graph, Bindings, GraphPosition);
+				UEdGraphNode* ResultNode = FunctionNodeSpawner->Invoke(Graph, Bindings, GraphPosition);
+
+				// Autowire the node if we were dragging on top of a pin
+				if (ResultNode != nullptr)
+				{
+					if (UEdGraphPin* FromPin = GetHoveredPin())
+					{
+						ResultNode->AutowireNewNode(FromPin);
+					}
+				}
 
 				Reply = FReply::Handled();
 			}
diff --git a/Engine/Source/Editor/Kismet/Private/BPFunctionDragDropAction.h b/Engine/Source/Editor/Kismet/Private/BPFunctionDragDropAction.h
--- a/Engine/Source/Editor/Kismet/Private/BPFunctionDragDropAction.h
+++ b/Engine/Source/Editor/Kismet/Private/BPFunctionDragDropAction.h
@@ -62,6 +62,7 @@ public:
 	// FGraphEditorDragDropAction interface
 	virtual void HoverTargetChanged() override;
 	virtual FReply DroppedOnPanel( const TSharedRef< class SWidget >& Panel, FVector2D ScreenPosition, FVector2D GraphPosition, UEdGraph& Graph) override;
+	virtual FReply DroppedOnPin(FVector2D ScreenPosition, FVector2D GraphPosition) override;
 	// End of FGraphEditorDragDropAction
 
 	static TSharedRef<FKismetFunctionDragDropAction> New(TSharedPtr<FEdGraphSchemaAction> InActionNode, FName InFunctionName, UClass* InOwningClass, const FMemberReference& InCallOnMember, FNodeCreationAnalytic AnalyticCallback, FCanBeDroppedDelegate CanBeDroppedDelegate = FCanBeDroppedDelegate());
diff --git a/Engine/Source/Editor/Kismet/Private/BPVariableDragDropAction.cpp b/Engine/Source/Editor/Kismet/Private/BPVariableDragDropAction.cpp
--- a/Engine/Source/Editor/Kismet/Private/BPVariableDragDropAction.cpp
+++ b/Engine/Source/Editor/Kismet/Private/BPVariableDragDropAction.cpp
@@ -126,15 +126,16 @@ void FKismetVariableDragDropAction::HoverTargetChanged()
 		if(CanVariableBeDropped(VariableProperty, *PinUnderCursor->GetOwningNode()->GetGraph()))
 		{
 			const UEdGraphSchema_K2* Schema = CastChecked<const UEdGraphSchema_K2>(PinUnderCursor->GetSchema());
+			const bool bIsExecPin = Schema->IsExecPin(*PinUnderCursor);
 
-			const bool bIsRead = PinUnderCursor->Direction == EGPD_Input;
+			const bool bIsRead = (PinUnderCursor->Direction == EGPD_Input) && !bIsExecPin;
 			const UBlueprint* Blueprint = FBlueprintEditorUtils::FindBlueprintForNode(PinUnderCursor->GetOwningNode());
 			const bool bReadOnlyProperty = FBlueprintEditorUtils::IsPropertyReadOnlyInCurrentBlueprint(Blueprint, VariableProperty);
 			const bool bCanWriteIfNeeded = bIsRead || !bReadOnlyProperty;
 
 			FEdGraphPinType VariablePinType;
 			Schema->ConvertPropertyToPinType(VariableProperty, VariablePinType);
-			const bool bTypeMatch = Schema->ArePinTypesCompatible(VariablePinType, PinUnderCursor->PinType);
+			const bool bTypeMatch = Schema->ArePinTypesCompatible(VariablePinType, PinUnderCursor->PinType) || bIsExecPin;
 
 			Args.Add(TEXT("PinUnderCursor"), FText::FromString(PinUnderCursor->PinName));
 
@@ -329,18 +330,20 @@ FReply FKismetVariableDragDropAction::DroppedOnPin(FVector2D ScreenPosition, FVe
 	{
 		if (const UEdGraphSchema_K2* Schema = CastChecked<const UEdGraphSchema_K2>(TargetPin->GetSchema()))
 		{
+			const bool bIsExecPin = Schema->IsExecPin(*TargetPin);
+
 			UProperty* VariableProperty = GetVariableProperty();
 
 			if(CanVariableBeDropped(VariableProperty, *TargetPin->GetOwningNode()->GetGraph()))
 			{
-				const bool bIsRead = TargetPin->Direction == EGPD_Input;
+				const bool bIsRead = (TargetPin->Direction == EGPD_Input) && !bIsExecPin;
 				const UBlueprint* Blueprint = FBlueprintEditorUtils::FindBlueprintForNode(TargetPin->GetOwningNode());
 				const bool bReadOnlyProperty = FBlueprintEditorUtils::IsPropertyReadOnlyInCurrentBlueprint(Blueprint, VariableProperty);
 				const bool bCanWriteIfNeeded = bIsRead || !bReadOnlyProperty;
 
 				FEdGraphPinType VariablePinType;
 				Schema->ConvertPropertyToPinType(VariableProperty, VariablePinType);
-				const bool bTypeMatch = Schema->ArePinTypesCompatible(VariablePinType, TargetPin->PinType);
+				const bool bTypeMatch = Schema->ArePinTypesCompatible(VariablePinType, TargetPin->PinType) || bIsExecPin;
 
 				if (bTypeMatch && bCanWriteIfNeeded)
 				{
diff --git a/Engine/Source/Editor/Kismet/Private/BlueprintCompilationManager.cpp b/Engine/Source/Editor/Kismet/Private/BlueprintCompilationManager.cpp
--- a/Engine/Source/Editor/Kismet/Private/BlueprintCompilationManager.cpp
+++ b/Engine/Source/Editor/Kismet/Private/BlueprintCompilationManager.cpp
@@ -57,8 +57,8 @@
 	STAGE XI: CREATE UPDATED CLASS HIERARCHY
 	STAGE XII: COMPILE CLASS LAYOUT
 	STAGE XIII: COMPILE CLASS FUNCTIONS
-	STAGE XIV: CLEAR TEMPORARY FLAGS
-	STAGE XV: REINSTANCE
+	STAGE XIV: REINSTANCE
+	STAGE XV: CLEAR TEMPORARY FLAGS
 
 	The code that implements these stages are labeled below. At some later point a final
 	reinstancing operation will occur, unless the client is using CompileSynchronously, 
@@ -69,6 +69,8 @@
 #define VERIFY_NO_STALE_CLASS_REFERENCES 0
 #define VERIFY_NO_BAD_SKELETON_REFERENCES 0
 
+struct FReinstancingJob;
+
 struct FBlueprintCompilationManagerImpl : public FGCObject
 {
 	FBlueprintCompilationManagerImpl();
@@ -83,7 +85,7 @@ struct FBlueprintCompilationManagerImpl : public FGCObject
 	void FlushReinstancingQueueImpl();
 	bool HasBlueprintsToCompile() const;
 	bool IsGeneratedClassLayoutReady() const;
-	static void ReinstanceBatch(TArray<TUniquePtr<FBlueprintCompileReinstancer>>& Reinstancers, TMap< UClass*, UClass* >& InOutOldToNewClassMap, TArray<UObject*>* ObjLoaded);
+	static void ReinstanceBatch(TArray<FReinstancingJob>& Reinstancers, TMap< UClass*, UClass* >& InOutOldToNewClassMap, TArray<UObject*>* ObjLoaded);
 	static UClass* FastGenerateSkeletonClass(UBlueprint* BP, FKismetCompilerContext& CompilerContext);
 	static bool IsQueuedForCompilation(UBlueprint* BP);
 	
@@ -128,6 +130,8 @@ void FBlueprintCompilationManagerImpl::AddReferencedObjects(FReferenceCollector&
 	{
 		Collector.AddReferencedObject(Job.BPToCompile);
 	}
+
+	Collector.AddReferencedObjects(ClassesToReinstance);
 }
 
 void FBlueprintCompilationManagerImpl::QueueForCompilation(const FBPCompileRequest& CompileJob)
@@ -184,8 +188,6 @@ void FBlueprintCompilationManagerImpl::CompileSynchronouslyImpl(const FBPCompile
 	
 	if (!bBatchCompile)
 	{
-		Request.BPToCompile->BroadcastCompiled();
-
 		if(GEditor)
 		{
 			GEditor->BroadcastBlueprintCompiled();	
@@ -205,7 +207,7 @@ enum class ECompilationManagerJobType
 
 struct FCompilerData
 {
-	explicit FCompilerData(UBlueprint* InBP, ECompilationManagerJobType InJobType, FCompilerResultsLog* InResultsLogOverride, EBlueprintCompileOptions UserOptions)
+	explicit FCompilerData(UBlueprint* InBP, ECompilationManagerJobType InJobType, FCompilerResultsLog* InResultsLogOverride, EBlueprintCompileOptions UserOptions, bool bBytecodeOnly)
 	{
 		check(InBP);
 		BP = InBP;
@@ -228,6 +230,7 @@ struct FCompilerData
 		InternalOptions.bRegenerateSkelton = false;
 		InternalOptions.bReinstanceAndStubOnFailure = false;
 		InternalOptions.bSaveIntermediateProducts = (UserOptions & EBlueprintCompileOptions::SaveIntermediateProducts ) != EBlueprintCompileOptions::None;
+		InternalOptions.CompileType = bBytecodeOnly ? EKismetCompileType::BytecodeOnly : EKismetCompileType::Full;
 
 		if( UWidgetBlueprint* WidgetBP = Cast<UWidgetBlueprint>(BP))
 		{
@@ -254,12 +257,19 @@ struct FCompilerData
 	UBlueprint* BP;
 	FCompilerResultsLog* ActiveResultsLog;
 	TUniquePtr<FCompilerResultsLog> ResultsLog;
-	TUniquePtr<FKismetCompilerContext> Compiler;
+	TSharedPtr<FKismetCompilerContext> Compiler;
 	FKismetCompilerOptions InternalOptions;
+	TSharedPtr<FBlueprintCompileReinstancer> Reinstancer;
 
 	ECompilationManagerJobType JobType;
 	bool bPackageWasDirty;
 };
+
+struct FReinstancingJob
+{
+	TSharedPtr<FBlueprintCompileReinstancer> Reinstancer;
+	TSharedPtr<FKismetCompilerContext> Compiler;
+};
 	
 void FBlueprintCompilationManagerImpl::FlushCompilationQueueImpl(TArray<UObject*>* ObjLoaded, bool bSuppressBroadcastCompiled)
 {
@@ -272,7 +282,6 @@ void FBlueprintCompilationManagerImpl::FlushCompilationQueueImpl(TArray<UObject*
 	}
 
 	TArray<FCompilerData> CurrentlyCompilingBPs;
-	TArray<TUniquePtr<FBlueprintCompileReinstancer>> Reinstancers;
 	{ // begin GTimeCompiling scope 
 		FScopedDurationTimer SetupTimer(GTimeCompiling); 
 
@@ -292,11 +301,27 @@ void FBlueprintCompilationManagerImpl::FlushCompilationQueueImpl(TArray<UObject*
 						{
 							ChildBlueprint->bQueuedForCompilation = true;
 							ensure(ChildBlueprint->bHasBeenRegenerated);
-							CurrentlyCompilingBPs.Add(FCompilerData(ChildBlueprint, ECompilationManagerJobType::RelinkOnly, nullptr, EBlueprintCompileOptions::None));
+							CurrentlyCompilingBPs.Add(FCompilerData(ChildBlueprint, ECompilationManagerJobType::RelinkOnly, nullptr, EBlueprintCompileOptions::None, false));
 						}
 					}
 				}
 			}
+			
+			// Add any dependent blueprints for a bytecode compile, this is needed because we 
+			// have no way to keep bytecode safe when a function is renamed or parameters are
+			// added or removed - strictly speaking we only need to do this when function 
+			// parameters changed, but that's a somewhat dubious optmization - ideally this
+			// work *never* needs to happen:
+			TArray<UBlueprint*> DependentBlueprints;
+			FBlueprintEditorUtils::GetDependentBlueprints(ComplieJob.BPToCompile, DependentBlueprints);
+			for(UBlueprint* DependentBlueprint : DependentBlueprints)
+			{
+				if(!IsQueuedForCompilation(DependentBlueprint))
+				{
+					DependentBlueprint->bQueuedForCompilation = true;
+					CurrentlyCompilingBPs.Add(FCompilerData(DependentBlueprint, ECompilationManagerJobType::Normal, nullptr, EBlueprintCompileOptions::None, true));
+				}
+			}
 		}
 
 		// STAGE II: Filter out data only and interface blueprints:
@@ -307,7 +332,7 @@ void FBlueprintCompilationManagerImpl::FlushCompilationQueueImpl(TArray<UObject*
 			UBlueprint* QueuedBP = QueuedJob.BPToCompile;
 
 			ensure(!(QueuedBP->GeneratedClass->ClassDefaultObject->HasAnyFlags(RF_NeedLoad)));
-			if( FBlueprintEditorUtils::IsDataOnlyBlueprint(QueuedBP) )
+			if( FBlueprintEditorUtils::IsDataOnlyBlueprint(QueuedBP) && !QueuedBP->bHasBeenRegenerated)
 			{
 				const UClass* ParentClass = QueuedBP->ParentClass;
 				if (ParentClass && ParentClass->HasAllClassFlags(CLASS_Native))
@@ -326,7 +351,7 @@ void FBlueprintCompilationManagerImpl::FlushCompilationQueueImpl(TArray<UObject*
 
 			if(bSkipCompile)
 			{
-				CurrentlyCompilingBPs.Add(FCompilerData(QueuedBP, ECompilationManagerJobType::SkeletonOnly, QueuedJob.ClientResultsLog, QueuedJob.CompileOptions));
+				CurrentlyCompilingBPs.Add(FCompilerData(QueuedBP, ECompilationManagerJobType::SkeletonOnly, QueuedJob.ClientResultsLog, QueuedJob.CompileOptions, false));
 				if (QueuedBP->IsGeneratedClassAuthoritative() && (QueuedBP->GeneratedClass != nullptr))
 				{
 					// set bIsRegeneratingOnLoad so that we don't reset loaders:
@@ -348,7 +373,7 @@ void FBlueprintCompilationManagerImpl::FlushCompilationQueueImpl(TArray<UObject*
 			}
 			else
 			{
-				CurrentlyCompilingBPs.Add(FCompilerData(QueuedBP, ECompilationManagerJobType::Normal, QueuedJob.ClientResultsLog, QueuedJob.CompileOptions));
+				CurrentlyCompilingBPs.Add(FCompilerData(QueuedBP, ECompilationManagerJobType::Normal, QueuedJob.ClientResultsLog, QueuedJob.CompileOptions, false));
 			}
 		}
 
@@ -498,6 +523,21 @@ void FBlueprintCompilationManagerImpl::FlushCompilationQueueImpl(TArray<UObject*
 				FBlueprintEditorUtils::ReconstructAllNodes(BP);
 				FBlueprintEditorUtils::ReplaceDeprecatedNodes(BP);
 			}
+			else
+			{
+				// matching existing behavior, when compiling a BP not on load we refresh nodes
+				// before compiling:
+				FBlueprintCompileReinstancer::OptionallyRefreshNodes(BP);
+				TArray<UBlueprint*> DependentBlueprints;
+				FBlueprintEditorUtils::GetDependentBlueprints(BP, DependentBlueprints);
+
+				for (UBlueprint* CurrentBP : DependentBlueprints)
+				{
+					const EBlueprintStatus OriginalStatus = CurrentBP->Status;
+					FBlueprintEditorUtils::RefreshExternalBlueprintDependencyNodes(CurrentBP, BP->GeneratedClass);
+					CurrentBP->Status = OriginalStatus;
+				}
+			}
 			
 			// Broadcast pre-compile
 			{
@@ -542,12 +582,12 @@ void FBlueprintCompilationManagerImpl::FlushCompilationQueueImpl(TArray<UObject*
 				UBlueprint* BP = CompilerData.BP;
 
 				OldCDOs.Add(BP, BP->GeneratedClass->ClassDefaultObject);
-				Reinstancers.Emplace(TUniquePtr<FBlueprintCompileReinstancer>( 
+				CompilerData.Reinstancer = TSharedPtr<FBlueprintCompileReinstancer>( 
 					new FBlueprintCompileReinstancer(
 						BP->GeneratedClass, 
 						EBlueprintCompileReinstancerFlags::AutoInferSaveOnCompile|EBlueprintCompileReinstancerFlags::AvoidCDODuplication
 					)
-				));
+				);
 			}
 		}
 
@@ -589,6 +629,10 @@ void FBlueprintCompilationManagerImpl::FlushCompilationQueueImpl(TArray<UObject*
 				FKismetCompilerContext& CompilerContext = *(CompilerData.Compiler);
 				CompilerContext.CompileClassLayout(EInternalCompilerFlags::PostponeLocalsGenerationUntilPhaseTwo);
 			}
+			else if(CompilerData.Compiler.IsValid())
+			{
+				CompilerData.Compiler->NewClass = CastChecked<UBlueprintGeneratedClass>(BP->GeneratedClass);
+			}
 		}
 		bGeneratedClassLayoutReady = true;
 	
@@ -617,7 +661,11 @@ void FBlueprintCompilationManagerImpl::FlushCompilationQueueImpl(TArray<UObject*
 				BP->GeneratedClass->ClassDefaultObject = nullptr;
 		
 				FKismetCompilerContext& CompilerContext = *(CompilerData.Compiler);
-				CompilerContext.CompileFunctions(EInternalCompilerFlags::PostponeLocalsGenerationUntilPhaseTwo);
+				CompilerContext.CompileFunctions(
+					EInternalCompilerFlags::PostponeLocalsGenerationUntilPhaseTwo
+					|EInternalCompilerFlags::PostponeDefaultObjectAssignmentUntilReinstancing
+					|EInternalCompilerFlags::SkipRefreshExternalBlueprintDependencyNodes
+				); 
 
 				if (CompilerData.ActiveResultsLog->NumErrors == 0)
 				{
@@ -632,8 +680,32 @@ void FBlueprintCompilationManagerImpl::FlushCompilationQueueImpl(TArray<UObject*
 
 			ensure(BP->GeneratedClass->ClassDefaultObject->GetClass() == *(BP->GeneratedClass));
 		}
+	} // end GTimeCompiling scope
+
+	// STAGE XIV: Now we can finish the first stage of the reinstancing operation, moving old classes to new classes:
+	{
+		{
+			TArray<FReinstancingJob> Reinstancers;
+			// Set up reinstancing jobs - we need a reference to the compiler in order to honor 
+			// CopyTermDefaultsToDefaultObject
+			for (FCompilerData& CompilerData : CurrentlyCompilingBPs)
+			{
+				if(CompilerData.Reinstancer.IsValid())
+				{
+					Reinstancers.Push(
+						FReinstancingJob {
+							CompilerData.Reinstancer,
+							CompilerData.Compiler
+						}
+					);
+				}
+			}
+
+			FScopedDurationTimer ReinstTimer(GTimeReinstancing);
+			ReinstanceBatch(Reinstancers, ClassesToReinstance, ObjLoaded);
+		}
 		
-		// STAGE XIV: CLEAR TEMPORARY FLAGS
+		// STAGE XV: CLEAR TEMPORARY FLAGS
 		for (FCompilerData& CompilerData : CurrentlyCompilingBPs)
 		{
 			UBlueprint* BP = CompilerData.BP;
@@ -644,6 +716,8 @@ void FBlueprintCompilationManagerImpl::FlushCompilationQueueImpl(TArray<UObject*
 				// When the Blueprint is recompiled, then update the bound events for level scripting
 				ULevelScriptBlueprint* LevelScriptBP = CastChecked<ULevelScriptBlueprint>(BP);
 
+				// ULevel::OnLevelScriptBlueprintChanged needs to be run after the CDO has
+				// been updated as it respawns the actor:
 				if (ULevel* BPLevel = LevelScriptBP->GetLevel())
 				{
 					BPLevel->OnLevelScriptBlueprintChanged(LevelScriptBP);
@@ -673,19 +747,14 @@ void FBlueprintCompilationManagerImpl::FlushCompilationQueueImpl(TArray<UObject*
 				Package->SetDirtyFlag(CompilerData.bPackageWasDirty);
 			}
 		}
-	} // end GTimeCompiling scope
 
-	// STAGE XV: Now we can finish the first stage of the reinstancing operation, moving old classes to new classes:
-	{
-		FScopedDurationTimer ReinstTimer(GTimeReinstancing);
-		ReinstanceBatch(Reinstancers, ClassesToReinstance, ObjLoaded);
-		
 		// Make sure no junk in bytecode, this can happen only for blueprints that were in CurrentlyCompilingBPs because
 		// the reinstancer can detect all other references (see UpdateBytecodeReferences):
 		for (FCompilerData& CompilerData : CurrentlyCompilingBPs)
 		{
 			if(CompilerData.ShouldCompileClassFunctions())
 			{
+				CompilerData.BP->BroadcastCompiled();
 				continue;
 			}
 
@@ -698,11 +767,6 @@ void FBlueprintCompilationManagerImpl::FlushCompilationQueueImpl(TArray<UObject*
 					FFixupBytecodeReferences ValidateAr(CurrentFunction);
 				}
 			}
-			
-			if (!bSuppressBroadcastCompiled)
-			{
-				BP->BroadcastCompiled();
-			}
 		}
 
 		if ( GEditor )
@@ -774,7 +838,7 @@ bool FBlueprintCompilationManagerImpl::IsGeneratedClassLayoutReady() const
 	return bGeneratedClassLayoutReady;
 }
 
-void FBlueprintCompilationManagerImpl::ReinstanceBatch(TArray<TUniquePtr<FBlueprintCompileReinstancer>>& Reinstancers, TMap< UClass*, UClass* >& InOutOldToNewClassMap, TArray<UObject*>* ObjLoaded)
+void FBlueprintCompilationManagerImpl::ReinstanceBatch(TArray<FReinstancingJob>& Reinstancers, TMap< UClass*, UClass* >& InOutOldToNewClassMap, TArray<UObject*>* ObjLoaded)
 {
 	const auto FilterOutOfDateClasses = [](TArray<UClass*>& ClassList)
 	{
@@ -794,10 +858,11 @@ void FBlueprintCompilationManagerImpl::ReinstanceBatch(TArray<TUniquePtr<FBluepr
 
 	// Reinstancers may contain *part* of a class hierarchy, so we first need to reparent any child types that 
 	// haven't already been reinstanced:
-	for (const TUniquePtr<FBlueprintCompileReinstancer>& Reinstancer : Reinstancers)
+	for (const FReinstancingJob& ReinstancingJob : Reinstancers)
 	{
-		UClass* OldClass = Reinstancer->DuplicatedClass;
-		InOutOldToNewClassMap.Add(Reinstancer->DuplicatedClass, Reinstancer->ClassToReinstance);
+		const TSharedPtr<FBlueprintCompileReinstancer>& CurrentReinstancer = ReinstancingJob.Reinstancer;
+		UClass* OldClass = CurrentReinstancer->DuplicatedClass;
+		InOutOldToNewClassMap.Add(CurrentReinstancer->DuplicatedClass, CurrentReinstancer->ClassToReinstance);
 		if(!OldClass)
 		{
 			continue;
@@ -808,7 +873,7 @@ void FBlueprintCompilationManagerImpl::ReinstanceBatch(TArray<TUniquePtr<FBluepr
 			continue;
 		}
 
-		bool bParentLayoutChanged = FStructUtils::TheSameLayout(OldClass, Reinstancer->ClassToReinstance);
+		bool bParentLayoutChanged = FStructUtils::TheSameLayout(OldClass, CurrentReinstancer->ClassToReinstance);
 		if(bParentLayoutChanged)
 		{
 			// we need *all* derived types:
@@ -854,15 +919,20 @@ void FBlueprintCompilationManagerImpl::ReinstanceBatch(TArray<TUniquePtr<FBluepr
 	for(UClass* Class : ClassesToReinstance)
 	{
 		UObject* OriginalCDO = Class->ClassDefaultObject;
-		Reinstancers.Emplace(TUniquePtr<FBlueprintCompileReinstancer>( 
-			new FBlueprintCompileReinstancer(
-				Class, 
-				EBlueprintCompileReinstancerFlags::AutoInferSaveOnCompile|EBlueprintCompileReinstancerFlags::AvoidCDODuplication
-			)
-		));
+		Reinstancers.Emplace(
+			FReinstancingJob {
+				TSharedPtr<FBlueprintCompileReinstancer>( 
+					new FBlueprintCompileReinstancer(
+						Class, 
+						EBlueprintCompileReinstancerFlags::AutoInferSaveOnCompile|EBlueprintCompileReinstancerFlags::AvoidCDODuplication
+					)
+				),
+				TSharedPtr<FKismetCompilerContext>()
+			}
+		);
 
 		// make sure we have the newest parent now that CDO has been moved to duplicate class:
-		TUniquePtr<FBlueprintCompileReinstancer>& NewestReinstancer = Reinstancers.Last();
+		TSharedPtr<FBlueprintCompileReinstancer>& NewestReinstancer = Reinstancers.Last().Reinstancer;
 
 		UClass* SuperClass = NewestReinstancer->ClassToReinstance->GetSuperClass();
 		if(ensure(SuperClass))
@@ -879,15 +949,13 @@ void FBlueprintCompilationManagerImpl::ReinstanceBatch(TArray<TUniquePtr<FBluepr
 	}
 
 	// run UpdateBytecodeReferences:
-	for (const TUniquePtr<FBlueprintCompileReinstancer>& Reinstancer : Reinstancers)
+	for (const FReinstancingJob& ReinstancingJob : Reinstancers)
 	{
-		InOutOldToNewClassMap.Add( Reinstancer->DuplicatedClass, Reinstancer->ClassToReinstance );
+		const TSharedPtr<FBlueprintCompileReinstancer>& CurrentReinstancer = ReinstancingJob.Reinstancer;
+		InOutOldToNewClassMap.Add( CurrentReinstancer->DuplicatedClass, CurrentReinstancer->ClassToReinstance );
 			
-		UBlueprint* CompiledBlueprint = UBlueprint::GetBlueprintFromClass(Reinstancer->ClassToReinstance);
-		CompiledBlueprint->bIsRegeneratingOnLoad = true;
-		Reinstancer->UpdateBytecodeReferences();
-
-		CompiledBlueprint->bIsRegeneratingOnLoad = false;
+		UBlueprint* CompiledBlueprint = UBlueprint::GetBlueprintFromClass(CurrentReinstancer->ClassToReinstance);
+		CurrentReinstancer->UpdateBytecodeReferences();
 	}
 	
 	// Now we can update templates and archetypes - note that we don't look for direct references to archetypes - doing
@@ -898,8 +966,11 @@ void FBlueprintCompilationManagerImpl::ReinstanceBatch(TArray<TUniquePtr<FBluepr
 	// also sorts, but we don't want to make too many assumptions about caller. We could refine this API so that
 	// we're not taking a raw list of reinstancers:
 	Reinstancers.Sort(
-		[](const TUniquePtr<FBlueprintCompileReinstancer>& CompilerDataA, const TUniquePtr<FBlueprintCompileReinstancer>& CompilerDataB)
+		[](const FReinstancingJob& ReinstancingDataA, const FReinstancingJob& ReinstancingDataB)
 		{
+			const TSharedPtr<FBlueprintCompileReinstancer>& CompilerDataA = ReinstancingDataA.Reinstancer;
+			const TSharedPtr<FBlueprintCompileReinstancer>& CompilerDataB = ReinstancingDataB.Reinstancer;
+
 			UClass* A = CompilerDataA->ClassToReinstance;
 			UClass* B = CompilerDataB->ClassToReinstance;
 			int32 DepthA = 0;
@@ -928,14 +999,21 @@ void FBlueprintCompilationManagerImpl::ReinstanceBatch(TArray<TUniquePtr<FBluepr
 
 	// 2. Copy defaults from old CDO - CDO may be missing if this class was reinstanced and relinked here,
 	// so use GetDefaultObject(true):
-	for( const TUniquePtr<FBlueprintCompileReinstancer>& Reinstancer : Reinstancers )
+	for (const FReinstancingJob& ReinstancingJob : Reinstancers)
 	{
-		UObject* OldCDO = Reinstancer->DuplicatedClass->ClassDefaultObject;
+		const TSharedPtr<FBlueprintCompileReinstancer>& CurrentReinstancer = ReinstancingJob.Reinstancer;
+		UObject* OldCDO = CurrentReinstancer->DuplicatedClass->ClassDefaultObject;
 		if(OldCDO)
 		{
-			FBlueprintCompileReinstancer::CopyPropertiesForUnrelatedObjects(OldCDO,  Reinstancer->ClassToReinstance->GetDefaultObject(true));
-			
-			if (UBlueprintGeneratedClass* BPGClass = CastChecked<UBlueprintGeneratedClass>(Reinstancer->ClassToReinstance))
+			UObject* NewCDO = CurrentReinstancer->ClassToReinstance->GetDefaultObject(true);
+			FBlueprintCompileReinstancer::CopyPropertiesForUnrelatedObjects(OldCDO, NewCDO);
+
+			if(ReinstancingJob.Compiler.IsValid())
+			{
+				ReinstancingJob.Compiler->PropagateValuesToCDO(NewCDO, OldCDO);
+			}
+
+			if (UBlueprintGeneratedClass* BPGClass = CastChecked<UBlueprintGeneratedClass>(CurrentReinstancer->ClassToReinstance))
 			{
 				BPGClass->UpdateCustomPropertyListForPostConstruction();
 
@@ -974,9 +1052,10 @@ void FBlueprintCompilationManagerImpl::ReinstanceBatch(TArray<TUniquePtr<FBluepr
 	// 3. Update any remaining instances that are tagged as RF_ArchetypeObject or RF_InheritableComponentTemplate - 
 	// we may need to do further sorting to ensure that interdependent archetypes are initialized correctly:
 	TSet<UObject*> ArchetypeReferencers;
-	for( const TUniquePtr<FBlueprintCompileReinstancer>& Reinstancer : Reinstancers )
+	for (const FReinstancingJob& ReinstancingJob : Reinstancers)
 	{
-		UClass* OldClass = Reinstancer->DuplicatedClass;
+		const TSharedPtr<FBlueprintCompileReinstancer>& CurrentReinstancer = ReinstancingJob.Reinstancer;
+		UClass* OldClass = CurrentReinstancer->DuplicatedClass;
 		if(ensure(OldClass))
 		{
 			TArray<UObject*> ArchetypeObjects;
@@ -1002,22 +1081,23 @@ void FBlueprintCompilationManagerImpl::ReinstanceBatch(TArray<TUniquePtr<FBluepr
 			for(UObject* Archetype : ArchetypeObjects )
 			{
 				// make sure we fix up references in the owner:
-				if(UBlueprintGeneratedClass* ArchetypeOwnerClass = Archetype->GetTypedOuter<UBlueprintGeneratedClass>())
 				{
-					ArchetypeReferencers.Add(ArchetypeOwnerClass);
-					
-					TArray<UObject*> ContainedObjects;
-					GetObjectsWithOuter(ArchetypeOwnerClass, ContainedObjects);
-					ArchetypeReferencers.Append(ContainedObjects);
-				}
-				
-				if(UBlueprint* ArchetypeOwnerBP = Archetype->GetTypedOuter<UBlueprint>())
-				{
-					ArchetypeReferencers.Add(ArchetypeOwnerBP);
+					UObject* Iter = Archetype->GetOuter();
+					while(Iter)
+					{
+						if(Iter->HasAnyFlags(RF_ClassDefaultObject)
+							|| Cast<UBlueprintGeneratedClass>(Iter)
+							|| Cast<UBlueprint>(Iter) )
+						{
+							ArchetypeReferencers.Add(Iter);
 
-					TArray<UObject*> ContainedObjects;
-					GetObjectsWithOuter(ArchetypeOwnerBP, ContainedObjects);
-					ArchetypeReferencers.Append(ContainedObjects);
+							// this handles nested subobjects:
+							TArray<UObject*> ContainedObjects;
+							GetObjectsWithOuter(Iter, ContainedObjects);
+							ArchetypeReferencers.Append(ContainedObjects);
+						}
+						Iter = Iter->GetOuter();
+					}
 				}
 				
 				// move aside:
@@ -1033,9 +1113,9 @@ void FBlueprintCompilationManagerImpl::ReinstanceBatch(TArray<TUniquePtr<FBluepr
 					REN_DoNotDirty | REN_DontCreateRedirectors | REN_ForceNoResetLoaders );
 
 				// reconstruct
-				FMakeClassSpawnableOnScope TemporarilySpawnable(Reinstancer->ClassToReinstance);
+				FMakeClassSpawnableOnScope TemporarilySpawnable(CurrentReinstancer->ClassToReinstance);
 				const EObjectFlags FlagMask = RF_Public | RF_ArchetypeObject | RF_Transactional | RF_Transient | RF_TextExportTransient | RF_InheritableComponentTemplate | RF_Standalone; //TODO: what about RF_RootSet?
-				UObject* NewArchetype = NewObject<UObject>(OriginalOuter, Reinstancer->ClassToReinstance, OriginalName, OriginalFlags & FlagMask);
+				UObject* NewArchetype = NewObject<UObject>(OriginalOuter, CurrentReinstancer->ClassToReinstance, OriginalName, OriginalFlags & FlagMask);
 
 				// copy old data:
 				FBlueprintCompileReinstancer::CopyPropertiesForUnrelatedObjects(Archetype, NewArchetype);
@@ -1052,11 +1132,12 @@ void FBlueprintCompilationManagerImpl::ReinstanceBatch(TArray<TUniquePtr<FBluepr
 	// 4. update known references to archetypes (e.g. component templates, WidgetTree). We don't want to run the normal 
 	// reference finder to update these because searching the entire object graph is time consuming. Instead we just replace
 	// all references in our UBlueprint and its generated class:
-	for( const TUniquePtr<FBlueprintCompileReinstancer>& Reinstancer : Reinstancers )
+	for (const FReinstancingJob& ReinstancingJob : Reinstancers)
 	{
-		ArchetypeReferencers.Add(Reinstancer->ClassToReinstance);
-		ArchetypeReferencers.Add(Reinstancer->ClassToReinstance->ClassGeneratedBy);
-		if(UBlueprint* BP = Cast<UBlueprint>(Reinstancer->ClassToReinstance->ClassGeneratedBy))
+		const TSharedPtr<FBlueprintCompileReinstancer>& CurrentReinstancer = ReinstancingJob.Reinstancer;
+		ArchetypeReferencers.Add(CurrentReinstancer->ClassToReinstance);
+		ArchetypeReferencers.Add(CurrentReinstancer->ClassToReinstance->ClassGeneratedBy);
+		if(UBlueprint* BP = Cast<UBlueprint>(CurrentReinstancer->ClassToReinstance->ClassGeneratedBy))
 		{
 			ArchetypeReferencers.Add(BP->SkeletonGeneratedClass);
 			ensure(BP->bCachedDependenciesUpToDate);
@@ -1264,6 +1345,12 @@ UClass* FBlueprintCompilationManagerImpl::FastGenerateSkeletonClass(UBlueprint*
 				}
 			}
 		}
+
+		// We're linking the skeleton function because TProperty::LinkInternal
+		// will assign add TTypeFundamentals::GetComputedFlagsPropertyFlags()
+		// to PropertyFlags. PropertyFlags must (mostly) match in order for 
+		// functions to be compatible:
+		NewFunction->StaticLink(true);
 		return NewFunction;
 	};
 
diff --git a/Engine/Source/Editor/Kismet/Private/BlueprintDetailsCustomization.cpp b/Engine/Source/Editor/Kismet/Private/BlueprintDetailsCustomization.cpp
--- a/Engine/Source/Editor/Kismet/Private/BlueprintDetailsCustomization.cpp
+++ b/Engine/Source/Editor/Kismet/Private/BlueprintDetailsCustomization.cpp
@@ -776,15 +776,15 @@ void FBlueprintVarActionDetails::CustomizeDetails( IDetailLayoutBuilder& DetailL
 				OriginalProperty = VariableProperty;
 			}
 
-			if (OriginalProperty == NULL || bVariableRenamed)
+			if (OriginalProperty == nullptr || bVariableRenamed)
 			{
 				// Prevent editing the default value of a skeleton property
-				VariableProperty = NULL;
+				VariableProperty = nullptr;
 			}
-			else  if (auto StructProperty = Cast<const UStructProperty>(OriginalProperty))
+			else if (const UStructProperty* StructProperty = Cast<const UStructProperty>(OriginalProperty))
 			{
 				// Prevent editing the default value of a stale struct
-				auto BGStruct = Cast<const UUserDefinedStruct>(StructProperty->Struct);
+				const UUserDefinedStruct* BGStruct = Cast<const UUserDefinedStruct>(StructProperty->Struct);
 				if (BGStruct && (EUserDefinedStructureStatus::UDSS_UpToDate != BGStruct->Status))
 				{
 					VariableProperty = nullptr;
@@ -793,7 +793,7 @@ void FBlueprintVarActionDetails::CustomizeDetails( IDetailLayoutBuilder& DetailL
 		}
 
 		// Find the class containing the variable
-		UClass* VariableClass = (VariableProperty != NULL) ? VariableProperty->GetTypedOuter<UClass>() : nullptr;
+		UClass* VariableClass = (VariableProperty ? VariableProperty->GetTypedOuter<UClass>() : nullptr);
 
 		FText ErrorMessage;
 		IDetailCategoryBuilder& DefaultValueCategory = DetailLayout.EditCategory(TEXT("DefaultValueCategory"), LOCTEXT("DefaultValueCategoryHeading", "Default Value"));
@@ -843,7 +843,7 @@ void FBlueprintVarActionDetails::CustomizeDetails( IDetailLayoutBuilder& DetailL
 				const bool bUDSProperty = PotentialUDSProperty && Cast<const UUserDefinedStruct>(PotentialUDSProperty->Struct);
 
 				UK2Node_FunctionEntry* FuncEntry = EntryNodes[0];
-				for(auto& LocalVar : FuncEntry->LocalVariables)
+				for (const FBPVariableDescription& LocalVar : FuncEntry->LocalVariables)
 				{
 					if(LocalVar.VarName == VariableProperty->GetFName()) //Property->GetFName())
 					{
@@ -1355,11 +1355,11 @@ void FBlueprintVarActionDetails::PopulateCategories(SMyBlueprint* MyBlueprint, T
 		}
 	}
 
-	CategorySource.Empty();
+	CategorySource.Reset();
 	CategorySource.Add(MakeShareable(new FText(LOCTEXT("Default", "Default"))));
 	for (int32 i = 0; i < VisibleVariables.Num(); ++i)
 	{
-		FText Category = FBlueprintEditorUtils::GetBlueprintVariableCategory(Blueprint, VisibleVariables[i], NULL);
+		FText Category = FBlueprintEditorUtils::GetBlueprintVariableCategory(Blueprint, VisibleVariables[i], nullptr);
 		if (!Category.IsEmpty() && !Category.EqualTo(FText::FromString(Blueprint->GetName())))
 		{
 			bool bNewCategory = true;
@@ -1396,7 +1396,7 @@ void FBlueprintVarActionDetails::PopulateCategories(SMyBlueprint* MyBlueprint, T
 			}
 		}
 
-		auto EntryNode = FBlueprintEditorUtils::GetEntryNode(FunctionGraph);
+		UK2Node_EditablePinBase* EntryNode = FBlueprintEditorUtils::GetEntryNode(FunctionGraph);
 		if (UK2Node_FunctionEntry* FunctionEntryNode = Cast<UK2Node_FunctionEntry>(EntryNode))
 		{
 			for (FBPVariableDescription& Variable : FunctionEntryNode->LocalVariables)
@@ -1416,7 +1416,7 @@ void FBlueprintVarActionDetails::PopulateCategories(SMyBlueprint* MyBlueprint, T
 
 	for (UEdGraph* MacroGraph : Blueprint->MacroGraphs)
 	{
-		auto EntryNode = FBlueprintEditorUtils::GetEntryNode(MacroGraph);
+		UK2Node_EditablePinBase* EntryNode = FBlueprintEditorUtils::GetEntryNode(MacroGraph);
 		if (UK2Node_Tunnel* TypedEntryNode = ExactCast<UK2Node_Tunnel>(EntryNode))
 		{
 			bool bNewCategory = true;
@@ -2512,9 +2512,9 @@ void FBlueprintVarActionDetails::OnFinishedChangingProperties(const FPropertyCha
 			UK2Node_FunctionEntry* FuncEntry = Cast<UK2Node_FunctionEntry>(InEntryNode.Get());
 
 			// Search out the correct local variable in the Function Entry Node and set the default value
-			for(auto& LocalVar : FuncEntry->LocalVariables)
+			for (FBPVariableDescription& LocalVar : FuncEntry->LocalVariables)
 			{
-				if(LocalVar.VarName == DirectProperty->GetFName() && LocalVar.DefaultValue != DefaultValueString)
+				if (LocalVar.VarName == DirectProperty->GetFName() && LocalVar.DefaultValue != DefaultValueString)
 				{
 					const FScopedTransaction Transaction(LOCTEXT("ChangeDefaults", "Change Defaults"));
 
@@ -2724,10 +2724,10 @@ void FBlueprintGraphArgumentLayout::GenerateChildContent( IDetailChildrenBuilder
 						[
 							SNew(STextBlock)
 							.Text(LOCTEXT("FunctionArgDetailsDefaultValue", "Default Value"))
-						.ToolTipText(LOCTEXT("FunctionArgDetailsDefaultValueParamTooltip", "The default value of the parameter."))
-						.Font(IDetailLayoutBuilder::GetDetailFont())
+							.ToolTipText(LOCTEXT("FunctionArgDetailsDefaultValueParamTooltip", "The default value of the parameter."))
+							.Font(IDetailLayoutBuilder::GetDetailFont())
 						]
-					.ValueContent()
+						.ValueContent()
 						.MaxDesiredWidth(512)
 						[
 							DefaultValueWidget
@@ -2760,11 +2760,11 @@ void FBlueprintGraphArgumentLayout::GenerateChildContent( IDetailChildrenBuilder
 	
 }
 
-namespace {
-
+namespace 
+{
 	static TArray<UK2Node_EditablePinBase*> GatherAllResultNodes(UK2Node_EditablePinBase* TargetNode)
 	{
-		if (auto ResultNode = Cast<UK2Node_FunctionResult>(TargetNode))
+		if (UK2Node_FunctionResult* ResultNode = Cast<UK2Node_FunctionResult>(TargetNode))
 		{
 			return (TArray<UK2Node_EditablePinBase*>)ResultNode->GetAllResultNodes();
 		}
@@ -2780,14 +2780,14 @@ namespace {
 
 void FBlueprintGraphArgumentLayout::OnRemoveClicked()
 {
-	auto ParamItem = ParamItemPtr.Pin();
+	TSharedPtr<FUserPinInfo> ParamItem = ParamItemPtr.Pin();
 	if (ParamItem.IsValid())
 	{
 		const FScopedTransaction Transaction( LOCTEXT( "RemoveParam", "Remove Parameter" ) );
 
-		auto GraphActionDetails = GraphActionDetailsPtr.Pin();
-		auto TargetNodes = GatherAllResultNodes(TargetNode);
-		for (auto Node : TargetNodes)
+		TSharedPtr<FBaseBlueprintGraphActionDetails> GraphActionDetails = GraphActionDetailsPtr.Pin();
+		TArray<UK2Node_EditablePinBase*> TargetNodes = GatherAllResultNodes(TargetNode);
+		for (UK2Node_EditablePinBase* Node : TargetNodes)
 		{
 			Node->Modify();
 			Node->RemoveUserDefinedPinByName(ParamItem->PinName);
@@ -2807,13 +2807,13 @@ FReply FBlueprintGraphArgumentLayout::OnArgMoveUp()
 	if (ThisParamIndex != INDEX_NONE && NewParamIndex >= 0)
 	{
 		const FScopedTransaction Transaction( LOCTEXT("K2_MovePinUp", "Move Pin Up") );
-		auto TargetNodes = GatherAllResultNodes(TargetNode);
-		for (auto Node : TargetNodes)
+		TArray<UK2Node_EditablePinBase*> TargetNodes = GatherAllResultNodes(TargetNode);
+		for (UK2Node_EditablePinBase* Node : TargetNodes)
 		{
 			Node->Modify();
 			Node->UserDefinedPins.Swap(ThisParamIndex, NewParamIndex);
 
-			auto GraphActionDetails = GraphActionDetailsPtr.Pin();
+			TSharedPtr<FBaseBlueprintGraphActionDetails> GraphActionDetails = GraphActionDetailsPtr.Pin();
 			if (GraphActionDetails.IsValid())
 			{
 				GraphActionDetails->OnParamsChanged(Node, true);
@@ -2830,13 +2830,13 @@ FReply FBlueprintGraphArgumentLayout::OnArgMoveDown()
 	if (ThisParamIndex != INDEX_NONE && NewParamIndex < TargetNode->UserDefinedPins.Num())
 	{
 		const FScopedTransaction Transaction( LOCTEXT("K2_MovePinDown", "Move Pin Down") );
-		auto TargetNodes = GatherAllResultNodes(TargetNode);
-		for (auto Node : TargetNodes)
+		TArray<UK2Node_EditablePinBase*> TargetNodes = GatherAllResultNodes(TargetNode);
+		for (UK2Node_EditablePinBase* Node : TargetNodes)
 		{
 			Node->Modify();
 			Node->UserDefinedPins.Swap(ThisParamIndex, NewParamIndex);
 			
-			auto GraphActionDetails = GraphActionDetailsPtr.Pin();
+			TSharedPtr<FBaseBlueprintGraphActionDetails> GraphActionDetails = GraphActionDetailsPtr.Pin();
 			if (GraphActionDetails.IsValid())
 			{
 				GraphActionDetails->OnParamsChanged(Node, true);
@@ -2993,7 +2993,7 @@ void FBlueprintGraphArgumentLayout::PinInfoChanged(const FEdGraphPinType& PinTyp
 				{
 					if (Node)
 					{
-						auto UDPinPtr = Node->UserDefinedPins.FindByPredicate([&](TSharedPtr<FUserPinInfo>& UDPin)
+						TSharedPtr<FUserPinInfo>* UDPinPtr = Node->UserDefinedPins.FindByPredicate([&](TSharedPtr<FUserPinInfo>& UDPin)
 						{
 							return UDPin.IsValid() && (UDPin->PinName == PinName);
 						});
@@ -3017,10 +3017,10 @@ void FBlueprintGraphArgumentLayout::PinInfoChanged(const FEdGraphPinType& PinTyp
 
 void FBlueprintGraphArgumentLayout::OnPrePinInfoChange(const FEdGraphPinType& PinType)
 {
-	if( !ShouldPinBeReadOnly(true))
+	if (!ShouldPinBeReadOnly(true))
 	{
-		auto TargetNodes = GatherAllResultNodes(TargetNode);
-		for (auto Node : TargetNodes)
+		TArray<UK2Node_EditablePinBase*> TargetNodes = GatherAllResultNodes(TargetNode);
+		for (UK2Node_EditablePinBase* Node : TargetNodes)
 		{
 			if (Node)
 			{
@@ -3959,7 +3959,7 @@ void FBlueprintDelegateActionDetails::OnFunctionSelected(TSharedPtr<FString> Fun
 		{
 			while (FunctionEntryNode->UserDefinedPins.Num())
 			{
-				auto Pin = FunctionEntryNode->UserDefinedPins[0];
+				TSharedPtr<FUserPinInfo> Pin = FunctionEntryNode->UserDefinedPins[0];
 				FunctionEntryNode->RemoveUserDefinedPin(Pin);
 			}
 
@@ -4037,8 +4037,8 @@ bool FBaseBlueprintGraphActionDetails::OnVerifyPinRename(UK2Node_EditablePinBase
 	if (InTargetNode)
 	{
 		// Check if the name conflicts with any of the other internal UFunction's property names (local variables and parameters).
-		const auto FoundFunction = FFunctionFromNodeHelper::FunctionFromNode(InTargetNode);
-		const auto ExistingProperty = FindField<const UProperty>(FoundFunction, *InNewName);
+		const UFunction* FoundFunction = FFunctionFromNodeHelper::FunctionFromNode(InTargetNode);
+		const UProperty* ExistingProperty = FindField<const UProperty>(FoundFunction, *InNewName);
 		if (ExistingProperty)
 		{
 			OutErrorMessage = LOCTEXT("ConflictsWithProperty", "Conflicts with another another local variable or function parameter!");
@@ -4065,12 +4065,12 @@ bool FBaseBlueprintGraphActionDetails::OnPinRenamed(UK2Node_EditablePinBase* Tar
 
 		const FScopedTransaction Transaction(LOCTEXT("RenameParam", "Rename Parameter"));
 
-		auto TerminalNodes = GatherAllResultNodes(FunctionResultNodePtr.Get());
-		if (auto EntryNode = FunctionEntryNodePtr.Get())
+		TArray<UK2Node_EditablePinBase*> TerminalNodes = GatherAllResultNodes(FunctionResultNodePtr.Get());
+		if (UK2Node_EditablePinBase* EntryNode = FunctionEntryNodePtr.Get())
 		{
 			TerminalNodes.Add(EntryNode);
 		}
-		for (auto TerminalNode : TerminalNodes)
+		for (UK2Node_EditablePinBase* TerminalNode : TerminalNodes)
 		{
 			TerminalNode->Modify();
 			PinRenamedHelper.NodesToRename.Add(TerminalNode);
@@ -4082,24 +4082,23 @@ bool FBaseBlueprintGraphActionDetails::OnPinRenamed(UK2Node_EditablePinBase* Tar
 		PinRenamedHelper.Broadcast(GetBlueprintObj(), TargetNode, Graph);
 
 		// TEST
-		for(auto NodeIter = PinRenamedHelper.NodesToRename.CreateIterator(); NodeIter; ++NodeIter)
+		for (UK2Node* NodeToRename : PinRenamedHelper.NodesToRename)
 		{
-			if(ERenamePinResult::ERenamePinResult_NameCollision == (*NodeIter)->RenameUserDefinedPin(OldName, NewName, true))
+			if (ERenamePinResult::ERenamePinResult_NameCollision == NodeToRename->RenameUserDefinedPin(OldName, NewName, true))
 			{
-				// log 
 				return false;
 			}
 		}
 
 		// UPDATE
-		for(auto NodeIter = PinRenamedHelper.NodesToRename.CreateIterator(); NodeIter; ++NodeIter)
+		for (UK2Node* NodeToRename : PinRenamedHelper.NodesToRename)
 		{
-			(*NodeIter)->RenameUserDefinedPin(OldName, NewName, false);
+			NodeToRename->RenameUserDefinedPin(OldName, NewName, false);
 		}
 
-		for (auto TerminalNode : TerminalNodes)
+		for (UK2Node_EditablePinBase* TerminalNode : TerminalNodes)
 		{
-			auto UDPinPtr = TerminalNode->UserDefinedPins.FindByPredicate([&](TSharedPtr<FUserPinInfo>& Pin)
+			TSharedPtr<FUserPinInfo>* UDPinPtr = TerminalNode->UserDefinedPins.FindByPredicate([&](TSharedPtr<FUserPinInfo>& Pin)
 			{
 				return Pin.IsValid() && (Pin->PinName == OldName);
 			});
@@ -4237,7 +4236,7 @@ void FBlueprintGraphActionDetails::OnTooltipTextCommitted(const FText& NewText,
 	if (FKismetUserDeclaredFunctionMetadata* Metadata = GetMetadataBlock())
 	{
 		Metadata->ToolTip = NewText;
-		if(auto Function = FindFunction())
+		if (UFunction* Function = FindFunction())
 		{
 			Function->Modify();
 			Function->SetMetaData(FBlueprintMetadata::MD_Tooltip, *NewText.ToString());
@@ -4282,7 +4281,7 @@ void FBlueprintGraphActionDetails::OnCategoryTextCommitted(const FText& NewText,
 				Metadata->Category = CategoryName;
 			}
 
-			if(auto Function = FindFunction())
+			if (UFunction* Function = FindFunction())
 			{
 				Function->Modify();
 				Function->SetMetaData(FBlueprintMetadata::MD_FunctionCategory, *CategoryName.ToString());
@@ -4300,7 +4299,7 @@ void FBlueprintGraphActionDetails::OnCategorySelectionChanged( TSharedPtr<FText>
 		if (FKismetUserDeclaredFunctionMetadata* Metadata = GetMetadataBlock())
 		{
 			Metadata->Category = *ProposedSelection.Get();
-			if(auto Function = FindFunction())
+			if (UFunction* Function = FindFunction())
 			{
 				Function->Modify();
 				Function->SetMetaData(FBlueprintMetadata::MD_FunctionCategory, *ProposedSelection.Get()->ToString());
@@ -4346,7 +4345,7 @@ void FBlueprintGraphActionDetails::OnKeywordsTextCommitted(const FText& NewText,
 			{
 				Metadata->Keywords = Keywords;
 
-				if(auto Function = FindFunction())
+				if (UFunction* Function = FindFunction())
 				{
 					Function->Modify();
 					Function->SetMetaData(FBlueprintMetadata::MD_FunctionKeywords, *Keywords.ToString());
@@ -4381,7 +4380,7 @@ void FBlueprintGraphActionDetails::OnCompactNodeTitleTextCommitted(const FText&
 			{
 				Metadata->CompactNodeTitle = CompactNodeTitle;
 
-				if(auto Function = FindFunction())
+				if (UFunction* Function = FindFunction())
 				{
 					Function->Modify();
 
@@ -4488,7 +4487,7 @@ void FBlueprintGraphActionDetails::OnAccessSpecifierSelected( TSharedPtr<FAccess
 		const FScopedTransaction Transaction( LOCTEXT( "ChangeAccessSpecifier", "Change Access Specifier" ) );
 
 		FunctionEntryNode->Modify();
-		auto Function = FindFunction();
+		UFunction* Function = FindFunction();
 		if(Function)
 		{
 			Function->Modify();
@@ -4651,9 +4650,9 @@ bool FBlueprintGraphActionDetails::IsPureFunctionVisible() const
 void FBlueprintGraphActionDetails::OnIsPureFunctionModified( const ECheckBoxState NewCheckedState )
 {
 	UK2Node_EditablePinBase * FunctionEntryNode = FunctionEntryNodePtr.Get();
-	auto Function = FindFunction();
-	auto EntryNode = Cast<UK2Node_FunctionEntry>(FunctionEntryNode);
-	if(EntryNode && Function)
+	UFunction* Function = FindFunction();
+	UK2Node_FunctionEntry* EntryNode = Cast<UK2Node_FunctionEntry>(FunctionEntryNode);
+	if (EntryNode && Function)
 	{
 		const FScopedTransaction Transaction( LOCTEXT( "ChangePure", "Change Pure" ) );
 		EntryNode->Modify();
@@ -4668,8 +4667,8 @@ void FBlueprintGraphActionDetails::OnIsPureFunctionModified( const ECheckBoxStat
 
 ECheckBoxState FBlueprintGraphActionDetails::GetIsPureFunction() const
 {
-	UK2Node_EditablePinBase * FunctionEntryNode = FunctionEntryNodePtr.Get();
-	auto EntryNode = Cast<UK2Node_FunctionEntry>(FunctionEntryNode);
+	UK2Node_EditablePinBase* FunctionEntryNode = FunctionEntryNodePtr.Get();
+	UK2Node_FunctionEntry* EntryNode = Cast<UK2Node_FunctionEntry>(FunctionEntryNode);
 	if(!EntryNode)
 	{
 		return ECheckBoxState::Undetermined;
@@ -4695,8 +4694,8 @@ bool FBlueprintGraphActionDetails::IsConstFunctionVisible() const
 void FBlueprintGraphActionDetails::OnIsConstFunctionModified( const ECheckBoxState NewCheckedState )
 {
 	UK2Node_EditablePinBase * FunctionEntryNode = FunctionEntryNodePtr.Get();
-	auto Function = FindFunction();
-	auto EntryNode = Cast<UK2Node_FunctionEntry>(FunctionEntryNode);
+	UFunction* Function = FindFunction();
+	UK2Node_FunctionEntry* EntryNode = Cast<UK2Node_FunctionEntry>(FunctionEntryNode);
 	if(EntryNode && Function)
 	{
 		const FScopedTransaction Transaction( LOCTEXT( "ChangeConst", "Change Const" ) );
@@ -4712,8 +4711,8 @@ void FBlueprintGraphActionDetails::OnIsConstFunctionModified( const ECheckBoxSta
 
 ECheckBoxState FBlueprintGraphActionDetails::GetIsConstFunction() const
 {
-	UK2Node_EditablePinBase * FunctionEntryNode = FunctionEntryNodePtr.Get();
-	auto EntryNode = Cast<UK2Node_FunctionEntry>(FunctionEntryNode);
+	UK2Node_EditablePinBase* FunctionEntryNode = FunctionEntryNodePtr.Get();
+	UK2Node_FunctionEntry* EntryNode = Cast<UK2Node_FunctionEntry>(FunctionEntryNode);
 	if(!EntryNode)
 	{
 		return ECheckBoxState::Undetermined;
@@ -4816,13 +4815,13 @@ FReply FBlueprintGraphActionDetails::OnAddNewOutputClicked()
 		}
 
 		const FString NewPinName = FunctionResultNode->CreateUniquePinName(TEXT("NewParam"));
-		auto TargetNodes = GatherAllResultNodes(FunctionResultNode);
+		TArray<UK2Node_EditablePinBase*> TargetNodes = GatherAllResultNodes(FunctionResultNode);
 		bool bAllChanged = TargetNodes.Num() > 0;
-		for (auto Node : TargetNodes)
+		for (UK2Node_EditablePinBase* Node : TargetNodes)
 		{
 			Node->Modify();
-			auto NewPin = Node->CreateUserDefinedPin(NewPinName, PinType, EGPD_Input, false);
-			bAllChanged &= nullptr != NewPin;
+			UEdGraphPin* NewPin = Node->CreateUserDefinedPin(NewPinName, PinType, EGPD_Input, false);
+			bAllChanged &= (nullptr != NewPin);
 
 			if (bAllChanged)
 			{
@@ -4876,10 +4875,9 @@ void FBlueprintInterfaceLayout::GenerateChildContent( IDetailChildrenBuilder& Ch
 	if (!bShowsInheritedInterfaces)
 	{
 		// Generate a list of interfaces already implemented
-		for (TArray<FBPInterfaceDescription>::TConstIterator It(Blueprint->ImplementedInterfaces); It; ++It)
+		for (const FBPInterfaceDescription& ImplementedInterface : Blueprint->ImplementedInterfaces)
 		{
-			auto Interface = (*It).Interface;
-			if (Interface)
+			if (const TSubclassOf<UInterface> Interface = ImplementedInterface.Interface)
 			{
 				Interfaces.AddUnique(FInterfaceName(Interface->GetFName(), Interface->GetDisplayNameText()));
 			}
@@ -6052,14 +6050,6 @@ void FChildActorComponentDetails::CustomizeDetails(IDetailLayoutBuilder& DetailB
 			}
 		}
 
-
-		auto TemplatesVisible = [](const TArray<TWeakObjectPtr<UObject>>& WeakObjects)
-		{
-			bool bResult = true;
-			
-			return bResult;
-		};
-
 		TArray<TWeakObjectPtr<UObject>> ObjectsBeingCustomized;
 		DetailBuilder.GetObjectsBeingCustomized(ObjectsBeingCustomized);
 
@@ -6235,9 +6225,9 @@ TSharedRef<SWidget> FBlueprintDocumentationDetails::GenerateExcerptList()
 		TArray<FExcerpt> Excerpts;
 		DocumentationPage->GetExcerpts( Excerpts );
 
-		for( auto ExcerptIter = Excerpts.CreateConstIterator(); ExcerptIter; ++ExcerptIter )
+		for (const FExcerpt& Excerpt : Excerpts)
 		{
-			ExcerptList.Add( MakeShareable( new FString( ExcerptIter->Name )));
+			ExcerptList.Add( MakeShareable( new FString( Excerpt.Name )));
 		}
 	}
 
diff --git a/Engine/Source/Editor/Kismet/Private/SBlueprintPalette.cpp b/Engine/Source/Editor/Kismet/Private/SBlueprintPalette.cpp
--- a/Engine/Source/Editor/Kismet/Private/SBlueprintPalette.cpp
+++ b/Engine/Source/Editor/Kismet/Private/SBlueprintPalette.cpp
@@ -101,14 +101,7 @@ static FString GetVarType(UStruct* VarScope, FName VarName, bool bUseObjToolTip,
 				FEdGraphPinType PinType;
 				if (K2Schema->ConvertPropertyToPinType(Property, PinType)) // use schema to get the color
 				{
-					if (bDetailed && PinType.PinSubCategoryObject.IsValid())
-					{
-						VarDesc = FString::Printf(TEXT("%s '%s'"), *PinType.PinCategory, *PinType.PinSubCategoryObject->GetName());
-					}
-					else
-					{
-						VarDesc = UEdGraphSchema_K2::TypeToText(PinType).ToString();
-					}
+					VarDesc = UEdGraphSchema_K2::TypeToText(PinType).ToString();
 				}
 			}
 		}
diff --git a/Engine/Source/Editor/Kismet/Private/SSCSEditor.cpp b/Engine/Source/Editor/Kismet/Private/SSCSEditor.cpp
--- a/Engine/Source/Editor/Kismet/Private/SSCSEditor.cpp
+++ b/Engine/Source/Editor/Kismet/Private/SSCSEditor.cpp
@@ -26,6 +26,8 @@
 #include "UnrealEdGlobals.h"
 #include "Kismet2/KismetEditorUtilities.h"
 #include "EdGraphSchema_K2.h"
+#include "GraphEditorActions.h"
+#include "ToolkitManager.h"
 #include "K2Node_Variable.h"
 #include "K2Node_ComponentBoundEvent.h"
 #include "K2Node_VariableGet.h"
@@ -3446,6 +3448,10 @@ void SSCSEditor::Construct( const FArguments& InArgs )
 			FCanExecuteAction::CreateSP( this, &SSCSEditor::CanRenameComponent ) ) 
 		);
 
+	CommandList->MapAction( FGraphEditorCommands::Get().FindReferences,
+		FUIAction( FExecuteAction::CreateSP( this, &SSCSEditor::OnFindReferences ) )
+	);
+
 	FSlateBrush const* MobilityHeaderBrush = FEditorStyle::GetBrush(TEXT("ClassIcon.ComponentMobilityHeaderIcon"));
 	
 	TSharedPtr<SHeaderRow> HeaderRow = SNew(SHeaderRow)
@@ -3526,6 +3532,12 @@ void SSCSEditor::Construct( const FArguments& InArgs )
 		FUIAction(FExecuteAction::CreateSP(this, &SSCSEditor::PromoteToBlueprint))
 	);
 
+	TSharedPtr<SHorizontalBox> ButtonBox;
+	TSharedPtr<SVerticalBox>   HeaderBox;
+	TSharedPtr<SWidget> SearchBar = SAssignNew(FilterBox, SSearchBox)
+		.OnTextChanged(this, &SSCSEditor::OnFilterTextChanged);
+	const bool  bInlineSearchBarWithButtons = (EditorMode == EComponentEditorMode::BlueprintSCS);
+
 	bool bHideComponentClassCombo = InArgs._HideComponentClassCombo.Get();
 
 	Contents = SNew(SVerticalBox)
@@ -3544,112 +3556,115 @@ void SSCSEditor::Construct( const FArguments& InArgs )
 			.AddMetaData<FTagMetaData>(FTagMetaData(TEXT("ComponentsPanel")))
 			.BorderBackgroundColor( FLinearColor( .6,.6,.6, 1.0f ) )
 			[
-				SNew(SHorizontalBox)
-				
-				+ SHorizontalBox::Slot()
-				.Padding( 3.0f, 3.0f )
-				.AutoWidth()
-				.HAlign(HAlign_Left)
-				[
-					SNew(SComponentClassCombo)
-					.AddMetaData<FTagMetaData>(FTagMetaData(TEXT("Actor.AddComponent")))
-					.Visibility(bHideComponentClassCombo ? EVisibility::Hidden : EVisibility::Visible)
-					.OnComponentClassSelected(this, &SSCSEditor::PerformComboAddClass)
-					.ToolTipText(LOCTEXT("AddComponent_Tooltip", "Adds a new component to this actor"))
-					.IsEnabled(AllowEditing)
-				]
-
-				+ SHorizontalBox::Slot()
-					.FillWidth(1.0f)
-					.VAlign(VAlign_Center)
-					.Padding(3.0f, 3.0f)
-				[
-					SAssignNew(FilterBox, SSearchBox)
-						.OnTextChanged(this, &SSCSEditor::OnFilterTextChanged)
-				]
-
-				+ SHorizontalBox::Slot()
-				.FillWidth(1.0f)
-				.HAlign(HAlign_Right)
-				.Padding( 3.0f, 3.0f )
-				[
-					SNew( SButton )
-					.AddMetaData<FTagMetaData>(FTagMetaData(TEXT("Actor.ConvertToBlueprint")))
-					.Visibility( this, &SSCSEditor::GetPromoteToBlueprintButtonVisibility )
-					.OnClicked( this, &SSCSEditor::OnPromoteToBlueprintClicked )
-					.ButtonStyle(FEditorStyle::Get(), "FlatButton.Primary")
-					.ContentPadding(FMargin(10,0))
-					.ToolTip(IDocumentation::Get()->CreateToolTip(
-						LOCTEXT("PromoteToBluerprintTooltip","Converts this actor into a reusable Blueprint Class that can have script behavior" ),
-						NULL,
-						TEXT("Shared/LevelEditor"),
-						TEXT("ConvertToBlueprint")))
+				SAssignNew(HeaderBox, SVerticalBox)
+					+ SVerticalBox::Slot()
+						.AutoHeight()
+						.VAlign(VAlign_Top)
 					[
-						SNew(SHorizontalBox)
-						
+						SAssignNew(ButtonBox, SHorizontalBox)
+				
 						+ SHorizontalBox::Slot()
-						.VAlign(VAlign_Center)
-						.Padding(3.f)
+						.Padding( 3.0f, 3.0f )
 						.AutoWidth()
+						.HAlign(HAlign_Left)
 						[
-							SNew(STextBlock)
-							.TextStyle(FEditorStyle::Get(), "ContentBrowser.TopBar.Font")
-							.Font( FEditorStyle::Get().GetFontStyle( "FontAwesome.10" ) )
-							.Text(FText::FromString(FString(TEXT("\xf085"))) /*fa-cogs*/)
+							SNew(SComponentClassCombo)
+							.AddMetaData<FTagMetaData>(FTagMetaData(TEXT("Actor.AddComponent")))
+							.Visibility(bHideComponentClassCombo ? EVisibility::Hidden : EVisibility::Visible)
+							.OnComponentClassSelected(this, &SSCSEditor::PerformComboAddClass)
+							.ToolTipText(LOCTEXT("AddComponent_Tooltip", "Adds a new component to this actor"))
+							.IsEnabled(AllowEditing)
 						]
 
-						+ SHorizontalBox::Slot()
-						.VAlign(VAlign_Center)
-						.Padding(3.f)
-						.AutoWidth()
-						[
-							SNew(STextBlock)
-							.TextStyle(FEditorStyle::Get(), "ContentBrowser.TopBar.Font")
-							//.Text( LOCTEXT("PromoteToBlueprint", "Add Script") )
-							.Text(LOCTEXT("PromoteToBlueprint", "Blueprint/Add Script"))
-						]
-					]
-				]
-				+ SHorizontalBox::Slot()
-				.FillWidth(1.0f)
-				.Padding( 3.0f, 3.0f )
-				.HAlign(HAlign_Right)
-				.Padding(3.0f, 3.0f)
-				[
-					SNew(SComboButton)
-					.AddMetaData<FTagMetaData>(FTagMetaData(TEXT("Actor.EditBlueprint")))
-					.Visibility(this, &SSCSEditor::GetEditBlueprintButtonVisibility)
-					.ContentPadding(FMargin(10, 0))
-					.ComboButtonStyle(FEditorStyle::Get(), "ToolbarComboButton")
-					.ButtonStyle(FEditorStyle::Get(), "FlatButton.Primary")
-					.ForegroundColor(FLinearColor::White)
-					.ButtonContent()
-					[
-						SNew( SHorizontalBox )
+						//
+						// horizontal slot (index) #1 => reserved for BP-editor search bar (see 'ButtonBox' usage below)
 
 						+ SHorizontalBox::Slot()
-						.AutoWidth()
-						.VAlign(VAlign_Center)
-						.Padding(3.f)
+						.FillWidth(1.0f)
+						.HAlign(HAlign_Right)
+						.Padding( 3.0f, 3.0f )
 						[
-							SNew(STextBlock)
-							.TextStyle(FEditorStyle::Get(), "ContentBrowser.TopBar.Font")
-							.Font(FEditorStyle::Get().GetFontStyle("FontAwesome.10"))
-							.Text(FText::FromString(FString(TEXT("\xf085"))) /*fa-cogs*/)
-						]
+							SNew( SButton )
+							.AddMetaData<FTagMetaData>(FTagMetaData(TEXT("Actor.ConvertToBlueprint")))
+							.Visibility( this, &SSCSEditor::GetPromoteToBlueprintButtonVisibility )
+							.OnClicked( this, &SSCSEditor::OnPromoteToBlueprintClicked )
+							.ButtonStyle(FEditorStyle::Get(), "FlatButton.Primary")
+							.ContentPadding(FMargin(10,0))
+							.ToolTip(IDocumentation::Get()->CreateToolTip(
+								LOCTEXT("PromoteToBluerprintTooltip","Converts this actor into a reusable Blueprint Class that can have script behavior" ),
+								NULL,
+								TEXT("Shared/LevelEditor"),
+								TEXT("ConvertToBlueprint")))
+							[
+								SNew(SHorizontalBox)
 						
+								+ SHorizontalBox::Slot()
+								.VAlign(VAlign_Center)
+								.Padding(3.f)
+								.AutoWidth()
+								[
+									SNew(STextBlock)
+									.TextStyle(FEditorStyle::Get(), "ContentBrowser.TopBar.Font")
+									.Font( FEditorStyle::Get().GetFontStyle( "FontAwesome.10" ) )
+									.Text(FText::FromString(FString(TEXT("\xf085"))) /*fa-cogs*/)
+								]
+
+								+ SHorizontalBox::Slot()
+								.VAlign(VAlign_Center)
+								.Padding(3.f)
+								.AutoWidth()
+								[
+									SNew(STextBlock)
+									.TextStyle(FEditorStyle::Get(), "ContentBrowser.TopBar.Font")
+									//.Text( LOCTEXT("PromoteToBlueprint", "Add Script") )
+									.Text(LOCTEXT("PromoteToBlueprint", "Blueprint/Add Script"))
+								]
+							]
+						]
 						+ SHorizontalBox::Slot()
+						.FillWidth(1.0f)
+						.Padding( 3.0f, 3.0f )
+						.HAlign(HAlign_Right)
+						.Padding(3.0f, 3.0f)
 						[
-							SNew(STextBlock)
-							.TextStyle(FEditorStyle::Get(), "ContentBrowser.TopBar.Font")
-							.Text(LOCTEXT("EditBlueprint", "Edit Blueprint"))
+							SNew(SComboButton)
+							.AddMetaData<FTagMetaData>(FTagMetaData(TEXT("Actor.EditBlueprint")))
+							.Visibility(this, &SSCSEditor::GetEditBlueprintButtonVisibility)
+							.ContentPadding(FMargin(10, 0))
+							.ComboButtonStyle(FEditorStyle::Get(), "ToolbarComboButton")
+							.ButtonStyle(FEditorStyle::Get(), "FlatButton.Primary")
+							.ForegroundColor(FLinearColor::White)
+							.ButtonContent()
+							[
+								SNew( SHorizontalBox )
+
+								+ SHorizontalBox::Slot()
+								.AutoWidth()
+								.VAlign(VAlign_Center)
+								.Padding(3.f)
+								[
+									SNew(STextBlock)
+									.TextStyle(FEditorStyle::Get(), "ContentBrowser.TopBar.Font")
+									.Font(FEditorStyle::Get().GetFontStyle("FontAwesome.10"))
+									.Text(FText::FromString(FString(TEXT("\xf085"))) /*fa-cogs*/)
+								]
+						
+								+ SHorizontalBox::Slot()
+								[
+									SNew(STextBlock)
+									.TextStyle(FEditorStyle::Get(), "ContentBrowser.TopBar.Font")
+									.Text(LOCTEXT("EditBlueprint", "Edit Blueprint"))
+								]
+							]
+							.MenuContent()
+							[
+								EditBlueprintMenuBuilder.MakeWidget()
+							]
 						]
 					]
-					.MenuContent()
-					[
-						EditBlueprintMenuBuilder.MakeWidget()
-					]
-				]
+
+				//
+				// vertical slot (index) #1 => reserved for instance-editor search bar (see 'HeaderBox' usage below)
 			]
 		]
 
@@ -3666,6 +3681,32 @@ void SSCSEditor::Construct( const FArguments& InArgs )
 		]
 	];
 
+	// insert the search bar, depending on which editor this widget is in (depending on convert/edit button visibility)
+	if (bInlineSearchBarWithButtons)
+	{
+		const int32 SearchBarHorizontalSlotIndex = 1;
+
+		ButtonBox->InsertSlot(SearchBarHorizontalSlotIndex)
+			.FillWidth(1.0f)
+			.VAlign(VAlign_Center)
+			.Padding(3.0f, 3.0f)
+		[
+			SearchBar.ToSharedRef()
+		];
+	}
+	else
+	{
+		const int32 SearchBarVerticalSlotIndex = 1;
+
+		HeaderBox->InsertSlot(SearchBarVerticalSlotIndex)
+			.VAlign(VAlign_Center)
+			.Padding(3.0f, 1.0f)
+		[
+			SearchBar.ToSharedRef()
+		];
+	}
+
+
 	this->ChildSlot
 	[
 		Contents.ToSharedRef()
@@ -3830,6 +3871,11 @@ TSharedPtr< SWidget > SSCSEditor::CreateContextMenu()
 
 					if (EditorMode == EComponentEditorMode::BlueprintSCS)
 					{
+						if (SelectedItems.Num() == 1)
+						{
+							MenuBuilder.AddMenuEntry(FGraphEditorCommands::Get().FindReferences);
+						}
+
 						// Collect the classes of all selected objects
 						TArray<UClass*> SelectionClasses;
 						for( auto NodeIter = SelectedNodes.CreateConstIterator(); NodeIter; ++NodeIter )
@@ -3992,6 +4038,23 @@ void SSCSEditor::ViewEvent(UBlueprint* Blueprint, const FName EventName, const F
 	}
 }
 
+void SSCSEditor::OnFindReferences()
+{
+	TArray<FSCSEditorTreeNodePtrType> SelectedNodes = SCSTreeWidget->GetSelectedItems();
+	if (SelectedNodes.Num() == 1)
+	{
+		TSharedPtr<IToolkit> FoundAssetEditor = FToolkitManager::Get().FindEditorForAsset(GetBlueprint());
+		if (FoundAssetEditor.IsValid())
+		{
+			const FString VariableName = SelectedNodes[0]->GetVariableName().ToString();
+			const FString SearchTerm = FString::Printf(TEXT("Nodes(VariableReference(MemberName=+\"%s\"))"), *VariableName);
+
+			TSharedRef<IBlueprintEditor> BlueprintEditor = StaticCastSharedRef<IBlueprintEditor>(FoundAssetEditor.ToSharedRef());
+			BlueprintEditor->SummonSearchUI(true, SearchTerm);
+		}
+	}
+}
+
 bool SSCSEditor::CanDuplicateComponent() const
 {
 	if(!IsEditingAllowed())
diff --git a/Engine/Source/Editor/Kismet/Public/SSCSEditor.h b/Engine/Source/Editor/Kismet/Public/SSCSEditor.h
--- a/Engine/Source/Editor/Kismet/Public/SSCSEditor.h
+++ b/Engine/Source/Editor/Kismet/Public/SSCSEditor.h
@@ -852,6 +852,9 @@ public:
 	void OnDeleteNodes();
 	bool CanDeleteNodes() const;
 
+	/** Callbacks to find references of the selected component */
+	void OnFindReferences();
+
 	/** Removes an existing component node from the tree */
 	void RemoveComponentNode(FSCSEditorTreeNodePtrType InNodePtr);
 
diff --git a/Engine/Source/Editor/KismetCompiler/Private/AnimBlueprintCompiler.cpp b/Engine/Source/Editor/KismetCompiler/Private/AnimBlueprintCompiler.cpp
--- a/Engine/Source/Editor/KismetCompiler/Private/AnimBlueprintCompiler.cpp
+++ b/Engine/Source/Editor/KismetCompiler/Private/AnimBlueprintCompiler.cpp
@@ -662,7 +662,7 @@ void FAnimBlueprintCompiler::ProcessSubInstance(UAnimGraphNode_SubInstance* SubI
 		FString PrefixedName = SubInstance->GetPinTargetVariableName(Pin);
 
 		// Create a property on the new class to hold the pin data
-		UProperty* NewProperty = FKismetCompilerUtilities::CreatePropertyOnScope(NewAnimBlueprintClass, FName(*PrefixedName), Pin->PinType, NewAnimBlueprintClass, 0, GetSchema(), MessageLog);
+		UProperty* NewProperty = FKismetCompilerUtilities::CreatePropertyOnScope(NewAnimBlueprintClass, FName(*PrefixedName), Pin->PinType, NewAnimBlueprintClass, CPF_BlueprintVisible, GetSchema(), MessageLog);
 		if(NewProperty)
 		{
 			NewProperty->SetMetaData(TEXT("Category"), TEXT("SubInstance"));
diff --git a/Engine/Source/Editor/KismetCompiler/Private/KismetCompiler.cpp b/Engine/Source/Editor/KismetCompiler/Private/KismetCompiler.cpp
--- a/Engine/Source/Editor/KismetCompiler/Private/KismetCompiler.cpp
+++ b/Engine/Source/Editor/KismetCompiler/Private/KismetCompiler.cpp
@@ -691,21 +691,21 @@ void FKismetCompilerContext::CreateClassVariablesFromBlueprint()
 		CreateVariable(Timeline->GetDirectionPropertyName(), DirectionPinType);
 
 		FEdGraphPinType FloatPinType(Schema->PC_Float, FString(), nullptr, EPinContainerType::None, false, FEdGraphTerminalType());
-		for(int32 i=0; i<Timeline->FloatTracks.Num(); i++)
+		for (const FTTFloatTrack& FloatTrack : Timeline->FloatTracks)
 		{
-			CreateVariable(Timeline->GetTrackPropertyName(Timeline->FloatTracks[i].TrackName), FloatPinType);
+			CreateVariable(Timeline->GetTrackPropertyName(FloatTrack.TrackName), FloatPinType);
 		}
 
 		FEdGraphPinType VectorPinType(Schema->PC_Struct, FString(), VectorStruct, EPinContainerType::None, false, FEdGraphTerminalType());
-		for(int32 i=0; i<Timeline->VectorTracks.Num(); i++)
+		for (const FTTVectorTrack& VectorTrack : Timeline->VectorTracks)
 		{
-			CreateVariable(Timeline->GetTrackPropertyName(Timeline->VectorTracks[i].TrackName), VectorPinType);
+			CreateVariable(Timeline->GetTrackPropertyName(VectorTrack.TrackName), VectorPinType);
 		}
 
 		FEdGraphPinType LinearColorPinType(Schema->PC_Struct, FString(), LinearColorStruct, EPinContainerType::None, false, FEdGraphTerminalType());
-		for(int32 i=0; i<Timeline->LinearColorTracks.Num(); i++)
+		for (const FTTLinearColorTrack& LinearColorTrack : Timeline->LinearColorTracks)
 		{
-			CreateVariable(Timeline->GetTrackPropertyName(Timeline->LinearColorTracks[i].TrackName), LinearColorPinType);
+			CreateVariable(Timeline->GetTrackPropertyName(LinearColorTrack.TrackName), LinearColorPinType);
 		}
 	}
 
@@ -741,11 +741,9 @@ void FKismetCompilerContext::CreateClassVariablesFromBlueprint()
 
 void FKismetCompilerContext::CreatePropertiesFromList(UStruct* Scope, UField**& PropertyStorageLocation, TIndirectArray<FBPTerminal>& Terms, uint64 PropertyFlags, bool bPropertiesAreLocal, bool bPropertiesAreParameters)
 {
-	for (int32 i = 0; i < Terms.Num(); ++i)
+	for (FBPTerminal& Term : Terms)
 	{
-		FBPTerminal& Term = Terms[i];
-
-		if(NULL != Term.AssociatedVarProperty)
+		if (Term.AssociatedVarProperty)
 		{
 			if(Term.Context && !Term.Context->IsObjectContextType())
 			{
@@ -759,11 +757,8 @@ void FKismetCompilerContext::CreatePropertiesFromList(UStruct* Scope, UField**&
 			MessageLog.Error(*FString::Printf(*LOCTEXT("PropertyForLiteral_Error", "Cannot create property for a literal: %s from @@ type (%s)").ToString(), *Term.Name, *UEdGraphSchema_K2::TypeToText(Term.Type).ToString()), Term.Source);
 		}
 
-		UProperty* NewProperty = FKismetCompilerUtilities::CreatePropertyOnScope(Scope, FName(*Term.Name), Term.Type, NewClass, PropertyFlags, Schema, MessageLog);
-		if (NewProperty != NULL)
+		if (UProperty* NewProperty = FKismetCompilerUtilities::CreatePropertyOnScope(Scope, FName(*Term.Name), Term.Type, NewClass, PropertyFlags, Schema, MessageLog))
 		{
-			NewProperty->PropertyFlags |= PropertyFlags;
-
 			if (bPropertiesAreParameters && Term.Type.bIsConst)
 			{
 				NewProperty->SetPropertyFlags(CPF_ConstParm);
@@ -1095,6 +1090,13 @@ void FKismetCompilerContext::CopyTermDefaultsToDefaultObject(UObject* DefaultObj
 	}
 }
 
+void FKismetCompilerContext::PropagateValuesToCDO(UObject* InNewCDO, UObject* InOldCDO)
+{
+	ensure(InNewCDO);
+	CopyTermDefaultsToDefaultObject(InNewCDO);
+	SetCanEverTick();
+}
+
 void FKismetCompilerContext::PrintVerboseInfoStruct(UStruct* Struct) const
 {
 	for (TFieldIterator<UProperty> PropIt(Struct); PropIt; ++PropIt)
@@ -3834,7 +3836,12 @@ void FKismetCompilerContext::CompileClassLayout(EInternalCompilerFlags InternalF
 
 void FKismetCompilerContext::CompileFunctions(EInternalCompilerFlags InternalFlags)
 {
+	// This is phase two, so we want to generated locals if PostponeLocalsGenerationUntilPhaseTwo is set:
 	const bool bGenerateLocals = !!(InternalFlags & EInternalCompilerFlags::PostponeLocalsGenerationUntilPhaseTwo);
+	// Don't propagate values to CDO if we're going to do that in reinstancing:
+	const bool bPropagateValuesToCDO = !(InternalFlags & EInternalCompilerFlags::PostponeDefaultObjectAssignmentUntilReinstancing);
+	// Don't RefreshExternalBlueprintDependencyNodes if the calling code has done so already:
+	const bool bSkipRefreshExternalBlueprintDependencyNodes = !!(InternalFlags & EInternalCompilerFlags::SkipRefreshExternalBlueprintDependencyNodes);
 	FKismetCompilerVMBackend Backend_VM(Blueprint, Schema, *this);
 
 	if( bGenerateLocals )
@@ -3969,9 +3976,10 @@ void FKismetCompilerContext::CompileFunctions(EInternalCompilerFlags InternalFla
 				Blueprint->PRIVATE_InnermostPreviousCDO = OldCDO;
 			}
 		}
-		else
+
+		if(bPropagateValuesToCDO)
 		{
-			if( NewCDO )
+			if( !Blueprint->HasAnyFlags(RF_BeingRegenerated) )
 			{
 				// Propagate the old CDO's properties to the new
 				if( OldCDO )
@@ -4006,10 +4014,9 @@ void FKismetCompilerContext::CompileFunctions(EInternalCompilerFlags InternalFla
 				}
 				// <<< End Backwards Compatibility
 			}
-		}
 
-		CopyTermDefaultsToDefaultObject(NewCDO);
-		SetCanEverTick();
+			PropagateValuesToCDO(NewCDO, OldCDO);
+		}
 
 		// Note: The old->new CDO copy is deferred when regenerating, so we skip this step in that case.
 		if (!Blueprint->HasAnyFlags(RF_BeingRegenerated))
@@ -4148,7 +4155,7 @@ void FKismetCompilerContext::CompileFunctions(EInternalCompilerFlags InternalFla
 	}
 
 	// For full compiles, find other blueprints that may need refreshing, and mark them dirty, in case they try to run
-	if( bIsFullCompile && !Blueprint->bIsRegeneratingOnLoad )
+	if( bIsFullCompile && !Blueprint->bIsRegeneratingOnLoad && !bSkipRefreshExternalBlueprintDependencyNodes )
 	{
 		TArray<UBlueprint*> DependentBlueprints;
 		FBlueprintEditorUtils::GetDependentBlueprints(Blueprint, DependentBlueprints);
@@ -4468,15 +4475,15 @@ FString FKismetCompilerContext::GetGuid(const UEdGraphNode* Node) const
 	return Ret.ToString();
 }
 
-TUniquePtr<FKismetCompilerContext> FKismetCompilerContext::GetCompilerForBP(UBlueprint* BP, FCompilerResultsLog& InMessageLog, const FKismetCompilerOptions& InCompileOptions)
+TSharedPtr<FKismetCompilerContext> FKismetCompilerContext::GetCompilerForBP(UBlueprint* BP, FCompilerResultsLog& InMessageLog, const FKismetCompilerOptions& InCompileOptions)
 {
 	if(UAnimBlueprint* AnimBP = Cast<UAnimBlueprint>(BP))
 	{
-		return TUniquePtr<FKismetCompilerContext>(new FAnimBlueprintCompiler(AnimBP, InMessageLog, InCompileOptions, nullptr));
+		return TSharedPtr<FKismetCompilerContext>(new FAnimBlueprintCompiler(AnimBP, InMessageLog, InCompileOptions, nullptr));
 	}
 	else
 	{
-		return TUniquePtr<FKismetCompilerContext>(new FKismetCompilerContext(BP, InMessageLog, InCompileOptions, nullptr));
+		return TSharedPtr<FKismetCompilerContext>(new FKismetCompilerContext(BP, InMessageLog, InCompileOptions, nullptr));
 	}
 }
 
diff --git a/Engine/Source/Editor/KismetCompiler/Private/KismetCompilerMisc.cpp b/Engine/Source/Editor/KismetCompiler/Private/KismetCompilerMisc.cpp
--- a/Engine/Source/Editor/KismetCompiler/Private/KismetCompilerMisc.cpp
+++ b/Engine/Source/Editor/KismetCompiler/Private/KismetCompilerMisc.cpp
@@ -1057,9 +1057,6 @@ UProperty* FKismetCompilerUtilities::CreatePropertyOnScope(UStruct* Scope, const
 {
 	const EObjectFlags ObjectFlags = RF_Public;
 
-	UProperty* NewProperty = NULL;
-	UObject* PropertyScope = NULL;
-
 	FName ValidatedPropertyName = PropertyName;
 
 	// Check to see if there's already a object on this scope with the same name, and throw an internal compiler error if so
@@ -1079,12 +1076,15 @@ UProperty* FKismetCompilerUtilities::CreatePropertyOnScope(UStruct* Scope, const
 				FString TestNameString = PropertyName.ToString() + FString::Printf(TEXT("_ERROR_DUPLICATE_%d"), Counter++);
 				TestName = FName(*TestNameString);
 
-			} while (CheckPropertyNameOnScope(Scope, TestName) != NULL);
+			} while (CheckPropertyNameOnScope(Scope, TestName) != nullptr);
 
 			ValidatedPropertyName = TestName;
 		}
 	}
 
+	UProperty* NewProperty = nullptr;
+	UObject* PropertyScope = nullptr;
+
 	// Handle creating a container property, if necessary
 	const bool bIsMapProperty = Type.IsMap();
 	const bool bIsSetProperty = Type.IsSet();
@@ -1148,7 +1148,6 @@ UProperty* FKismetCompilerUtilities::CreatePropertyOnScope(UStruct* Scope, const
 		{
 			if (!NewProperty->HasAnyPropertyFlags(CPF_HasGetValueTypeHash))
 			{
-
 				MessageLog.Error(
 					*FString::Printf(
 						*LOCTEXT("MapKeyTypeUnhashable_Error", "Map Property @@ has key type of %s which cannot be hashed and is therefore invalid").ToString(),
@@ -1157,7 +1156,7 @@ UProperty* FKismetCompilerUtilities::CreatePropertyOnScope(UStruct* Scope, const
 				);
 			}
 			// make the value property:
-			// not feelign good about myself..
+			// not feeling good about myself..
 			// Fix up the array property to have the new type-specific property as its inner, and return the new UArrayProperty
 			NewMapProperty->KeyProp = NewProperty;
 			// make sure the value property does not collide with the key property:
@@ -1221,6 +1220,11 @@ UProperty* FKismetCompilerUtilities::CreatePropertyOnScope(UStruct* Scope, const
 		}
 	}
 
+	if (NewProperty)
+	{
+		NewProperty->SetPropertyFlags(PropertyFlags);
+	}
+
 	return NewProperty;
 }
 
diff --git a/Engine/Source/Editor/KismetCompiler/Private/KismetCompilerModule.cpp b/Engine/Source/Editor/KismetCompiler/Private/KismetCompilerModule.cpp
--- a/Engine/Source/Editor/KismetCompiler/Private/KismetCompilerModule.cpp
+++ b/Engine/Source/Editor/KismetCompiler/Private/KismetCompilerModule.cpp
@@ -360,8 +360,11 @@ void FKismet2CompilerModule::RefreshVariables(UBlueprint* Blueprint)
 				K2Schema,
 				MessageLog);
 
-			NewProperty->PropertyLinkNext = Blueprint->GeneratedClass->PropertyLink;
-			Blueprint->GeneratedClass->PropertyLink = NewProperty;
+			if(NewProperty)
+			{
+				NewProperty->PropertyLinkNext = Blueprint->GeneratedClass->PropertyLink;
+				Blueprint->GeneratedClass->PropertyLink = NewProperty;
+			}
 		}
 		
 		Blueprint->GeneratedClass->Bind();
diff --git a/Engine/Source/Editor/KismetCompiler/Public/KismetCompiler.h b/Engine/Source/Editor/KismetCompiler/Public/KismetCompiler.h
--- a/Engine/Source/Editor/KismetCompiler/Public/KismetCompiler.h
+++ b/Engine/Source/Editor/KismetCompiler/Public/KismetCompiler.h
@@ -30,6 +30,8 @@ enum class EInternalCompilerFlags
 	None = 0x0,
 
 	PostponeLocalsGenerationUntilPhaseTwo = 0x1,
+	PostponeDefaultObjectAssignmentUntilReinstancing = 0x2,
+	SkipRefreshExternalBlueprintDependencyNodes = 0x4,
 };
 ENUM_CLASS_FLAGS(EInternalCompilerFlags)
 
@@ -151,7 +153,7 @@ public:
 			ParentGraph = SourceNode->GetGraph();
 		}
 
-		NodeType* Result = ParentGraph->CreateBlankNode<NodeType>();
+		NodeType* Result = ParentGraph->CreateIntermediateNode<NodeType>();
 		//check (Cast<UK2Node_Event>(Result) == nullptr); -- Removed to avoid any fallout, will replace with care later
 		MessageLog.NotifyIntermediateObjectCreation(Result, SourceNode); // this might be useful to track back function entry nodes to events.
 		Result->CreateNewGuid();
@@ -170,7 +172,7 @@ public:
 			ParentGraph = SourceNode->GetGraph();
 		}
 
-		NodeType* Result = ParentGraph->CreateBlankNode<NodeType>();
+		NodeType* Result = ParentGraph->CreateIntermediateNode<NodeType>();
 		//check (Cast<UK2Node_Event>(Result) != nullptr); -- Removed to avoid any fallout, will replace with care later
 		MessageLog.NotifyIntermediateObjectCreation(Result, SourceNode); // this might be useful to track back function entry nodes to events.
 		Result->CreateNewGuid();
@@ -225,7 +227,7 @@ public:
 
 	FString GetGuid(const UEdGraphNode* Node) const;
 
-	static TUniquePtr<FKismetCompilerContext> GetCompilerForBP(UBlueprint* BP, FCompilerResultsLog& InMessageLog, const FKismetCompilerOptions& InCompileOptions);
+	static TSharedPtr<FKismetCompilerContext> GetCompilerForBP(UBlueprint* BP, FCompilerResultsLog& InMessageLog, const FKismetCompilerOptions& InCompileOptions);
 	
 	/** Ensures that all variables have valid names for compilation/replication */
 	void ValidateVariableNames();
@@ -344,6 +346,9 @@ protected:
 	/** Copies default values cached for the terms in the DefaultPropertyValueMap to the final CDO */
 	virtual void CopyTermDefaultsToDefaultObject(UObject* DefaultObject);
 
+	/** Non virtual wrapper to encapsulate functions that occur when the CDO is ready for values: */
+	void PropagateValuesToCDO(UObject* NewCDO, UObject* OldCDO);
+
 	/** 
 	 * Function works only if subclass of AActor or UActorComponent.
 	 * If ReceiveTick event is defined, force CanEverTick.
diff --git a/Engine/Source/Editor/UMGEditor/Private/WidgetBlueprint.cpp b/Engine/Source/Editor/UMGEditor/Private/WidgetBlueprint.cpp
--- a/Engine/Source/Editor/UMGEditor/Private/WidgetBlueprint.cpp
+++ b/Engine/Source/Editor/UMGEditor/Private/WidgetBlueprint.cpp
@@ -699,9 +699,9 @@ bool UWidgetBlueprint::ValidateGeneratedClass(const UClass* InClass)
 	return Result;
 }
 
-TUniquePtr<FKismetCompilerContext> UWidgetBlueprint::GetCompilerForWidgetBP(UWidgetBlueprint* BP, FCompilerResultsLog& InMessageLog, const FKismetCompilerOptions& InCompileOptions)
+TSharedPtr<FKismetCompilerContext> UWidgetBlueprint::GetCompilerForWidgetBP(UWidgetBlueprint* BP, FCompilerResultsLog& InMessageLog, const FKismetCompilerOptions& InCompileOptions)
 {
-	return TUniquePtr<FKismetCompilerContext>(new FWidgetBlueprintCompiler(BP, InMessageLog, InCompileOptions, nullptr));
+	return TSharedPtr<FKismetCompilerContext>(new FWidgetBlueprintCompiler(BP, InMessageLog, InCompileOptions, nullptr));
 }
 
 void UWidgetBlueprint::GetReparentingRules(TSet< const UClass* >& AllowedChildrenOfClasses, TSet< const UClass* >& DisallowedChildrenOfClasses) const
diff --git a/Engine/Source/Editor/UMGEditor/Public/WidgetBlueprint.h b/Engine/Source/Editor/UMGEditor/Public/WidgetBlueprint.h
--- a/Engine/Source/Editor/UMGEditor/Public/WidgetBlueprint.h
+++ b/Engine/Source/Editor/UMGEditor/Public/WidgetBlueprint.h
@@ -257,7 +257,7 @@ public:
 
 	static bool ValidateGeneratedClass(const UClass* InClass);
 	
-	static TUniquePtr<FKismetCompilerContext> GetCompilerForWidgetBP(UWidgetBlueprint* BP, FCompilerResultsLog& InMessageLog, const FKismetCompilerOptions& InCompileOptions);
+	static TSharedPtr<FKismetCompilerContext> GetCompilerForWidgetBP(UWidgetBlueprint* BP, FCompilerResultsLog& InMessageLog, const FKismetCompilerOptions& InCompileOptions);
 
 private:
 	void ForEachSourceWidgetImpl(TFunctionRef<void(UWidget*)> Fn) const;
