Copying //UE4/Dev-Framework to //UE4/Dev-Main (Source: //UE4/Dev-Framework @ 3510040)

diff --git a/Engine/Source/Editor/Kismet/Private/BPVariableDragDropAction.cpp b/Engine/Source/Editor/Kismet/Private/BPVariableDragDropAction.cpp
--- a/Engine/Source/Editor/Kismet/Private/BPVariableDragDropAction.cpp
+++ b/Engine/Source/Editor/Kismet/Private/BPVariableDragDropAction.cpp
@@ -117,7 +117,7 @@ void FKismetVariableDragDropAction::HoverTargetChanged()
 			Message = FText::Format( LOCTEXT("IncorrectGraphForVariable_Error", "Cannot place variable '{VariableName}' in external scope '{Scope}'"), Args);
 		}
 	}
-	else if (PinUnderCursor != NULL)
+	else if (PinUnderCursor)
 	{
 		FFormatNamedArguments Args;
 		Args.Add(TEXT("PinUnderCursor"), FText::FromString(PinUnderCursor->PinName));
@@ -125,43 +125,51 @@ void FKismetVariableDragDropAction::HoverTargetChanged()
 
 		if(CanVariableBeDropped(VariableProperty, *PinUnderCursor->GetOwningNode()->GetGraph()))
 		{
-			const UEdGraphSchema_K2* Schema = CastChecked<const UEdGraphSchema_K2>(PinUnderCursor->GetSchema());
-			const bool bIsExecPin = Schema->IsExecPin(*PinUnderCursor);
-
-			const bool bIsRead = (PinUnderCursor->Direction == EGPD_Input) && !bIsExecPin;
-			const UBlueprint* Blueprint = FBlueprintEditorUtils::FindBlueprintForNode(PinUnderCursor->GetOwningNode());
-			const bool bReadOnlyProperty = FBlueprintEditorUtils::IsPropertyReadOnlyInCurrentBlueprint(Blueprint, VariableProperty);
-			const bool bCanWriteIfNeeded = bIsRead || !bReadOnlyProperty;
-
-			FEdGraphPinType VariablePinType;
-			Schema->ConvertPropertyToPinType(VariableProperty, VariablePinType);
-			const bool bTypeMatch = Schema->ArePinTypesCompatible(VariablePinType, PinUnderCursor->PinType) || bIsExecPin;
-
-			Args.Add(TEXT("PinUnderCursor"), FText::FromString(PinUnderCursor->PinName));
-
-			if (bTypeMatch && bCanWriteIfNeeded)
+			if (PinUnderCursor->bOrphanedPin)
 			{
-				StatusSymbol = FEditorStyle::GetBrush(TEXT("Graph.ConnectorFeedback.OK"));
-
-				if (bIsRead)
-				{
-					Message = FText::Format(LOCTEXT("MakeThisEqualThat_PinEqualVariableName", "Make {PinUnderCursor} = {VariableName}"), Args);
-				}
-				else
-				{
-					Message = FText::Format(LOCTEXT("MakeThisEqualThat_VariableNameEqualPin", "Make {VariableName} = {PinUnderCursor}"), Args);
-				}
+				StatusSymbol = FEditorStyle::GetBrush(TEXT("Graph.ConnectorFeedback.Error"));
+				Message = FText::Format(LOCTEXT("OrphanedPin_Error", "Cannot make connection to orphaned pin {PinUnderCursor}"), Args);
 			}
 			else
 			{
-				StatusSymbol = FEditorStyle::GetBrush(TEXT("Graph.ConnectorFeedback.Error"));
-				if (!bCanWriteIfNeeded)
+				const UEdGraphSchema_K2* Schema = CastChecked<const UEdGraphSchema_K2>(PinUnderCursor->GetSchema());
+				const bool bIsExecPin = Schema->IsExecPin(*PinUnderCursor);
+
+				const bool bIsRead = (PinUnderCursor->Direction == EGPD_Input) && !bIsExecPin;
+				const UBlueprint* Blueprint = FBlueprintEditorUtils::FindBlueprintForNode(PinUnderCursor->GetOwningNode());
+				const bool bWritableProperty = (FBlueprintEditorUtils::IsPropertyWritableInBlueprint(Blueprint, VariableProperty) == FBlueprintEditorUtils::EPropertyWritableState::Writable);
+				const bool bCanWriteIfNeeded = bIsRead || bWritableProperty;
+
+				FEdGraphPinType VariablePinType;
+				Schema->ConvertPropertyToPinType(VariableProperty, VariablePinType);
+				const bool bTypeMatch = Schema->ArePinTypesCompatible(VariablePinType, PinUnderCursor->PinType) || bIsExecPin;
+
+				Args.Add(TEXT("PinUnderCursor"), FText::FromString(PinUnderCursor->PinName));
+
+				if (bTypeMatch && bCanWriteIfNeeded)
 				{
-					Message = FText::Format(LOCTEXT("ReadOnlyVar_Error", "Cannot write to read-only variable '{VariableName}'"), Args);
+					StatusSymbol = FEditorStyle::GetBrush(TEXT("Graph.ConnectorFeedback.OK"));
+
+					if (bIsRead)
+					{
+						Message = FText::Format(LOCTEXT("MakeThisEqualThat_PinEqualVariableName", "Make {PinUnderCursor} = {VariableName}"), Args);
+					}
+					else
+					{
+						Message = FText::Format(LOCTEXT("MakeThisEqualThat_VariableNameEqualPin", "Make {VariableName} = {PinUnderCursor}"), Args);
+					}
 				}
 				else
 				{
-					Message = FText::Format(LOCTEXT("NotCompatible_Error", "The type of '{VariableName}' is not compatible with {PinUnderCursor}"), Args);
+					StatusSymbol = FEditorStyle::GetBrush(TEXT("Graph.ConnectorFeedback.Error"));
+					if (!bCanWriteIfNeeded)
+					{
+						Message = FText::Format(LOCTEXT("ReadOnlyVar_Error", "Cannot write to read-only variable '{VariableName}'"), Args);
+					}
+					else
+					{
+						Message = FText::Format(LOCTEXT("NotCompatible_Error", "The type of '{VariableName}' is not compatible with {PinUnderCursor}"), Args);
+					}
 				}
 			}
 		}
@@ -173,7 +181,7 @@ void FKismetVariableDragDropAction::HoverTargetChanged()
 			Message = FText::Format( LOCTEXT("IncorrectGraphForPin_Error", "Cannot place local variable '{VariableName}' in external scope '{Scope}'"), Args);
 		}
 	}
-	else if (VarNodeUnderCursor != NULL)
+	else if (VarNodeUnderCursor)
 	{
 		FFormatNamedArguments Args;
 		Args.Add(TEXT("VariableName"), FText::FromString(VariableString));
@@ -182,8 +190,8 @@ void FKismetVariableDragDropAction::HoverTargetChanged()
 		{
 			const bool bIsRead = VarNodeUnderCursor->IsA(UK2Node_VariableGet::StaticClass());
 			const UBlueprint* Blueprint = FBlueprintEditorUtils::FindBlueprintForNode(VarNodeUnderCursor);
-			const bool bReadOnlyProperty = FBlueprintEditorUtils::IsPropertyReadOnlyInCurrentBlueprint(Blueprint, VariableProperty);
-			const bool bCanWriteIfNeeded = bIsRead || !bReadOnlyProperty;
+			const bool bWritableProperty = (FBlueprintEditorUtils::IsPropertyWritableInBlueprint(Blueprint, VariableProperty) == FBlueprintEditorUtils::EPropertyWritableState::Writable);
+			const bool bCanWriteIfNeeded = bIsRead || bWritableProperty;
 
 			if (bCanWriteIfNeeded)
 			{
@@ -229,7 +237,7 @@ void FKismetVariableDragDropAction::HoverTargetChanged()
 			Blueprint = UBlueprint::GetBlueprintFromClass(Cast<UClass>(VariableSource.Get()));
 		}
 
-		if(Blueprint != NULL)
+		if (Blueprint)
 		{
 			Category = FBlueprintEditorUtils::GetBlueprintVariableCategory(Blueprint, VariableProperty->GetFName(), GetLocalVariableScope() );
 		}
@@ -271,8 +279,7 @@ void FKismetVariableDragDropAction::HoverTargetChanged()
 			// Needs to have a valid index to move it (this excludes variables added through other means, like timelines/components
 			int32 MoveVarIndex = INDEX_NONE;
 			int32 TargetVarIndex = INDEX_NONE;
-			UBlueprint* Blueprint = UBlueprint::GetBlueprintFromClass(Cast<UClass>(VariableSource.Get()));
-			if(Blueprint != NULL)
+			if (UBlueprint* Blueprint = UBlueprint::GetBlueprintFromClass(Cast<UClass>(VariableSource.Get())))
 			{
 				MoveVarIndex = FBlueprintEditorUtils::FindNewVariableIndex(Blueprint, VariableName);
 				TargetVarIndex = FBlueprintEditorUtils::FindNewVariableIndex(Blueprint, TargetVarName);
@@ -324,38 +331,39 @@ void FKismetVariableDragDropAction::HoverTargetChanged()
 
 FReply FKismetVariableDragDropAction::DroppedOnPin(FVector2D ScreenPosition, FVector2D GraphPosition)
 {
-	UEdGraphPin* TargetPin = GetHoveredPin();
-
-	if (TargetPin != NULL)
+	if (UEdGraphPin* TargetPin = GetHoveredPin())
 	{
-		if (const UEdGraphSchema_K2* Schema = CastChecked<const UEdGraphSchema_K2>(TargetPin->GetSchema()))
+		if (!TargetPin->bOrphanedPin)
 		{
-			const bool bIsExecPin = Schema->IsExecPin(*TargetPin);
-
-			UProperty* VariableProperty = GetVariableProperty();
-
-			if(CanVariableBeDropped(VariableProperty, *TargetPin->GetOwningNode()->GetGraph()))
+			if (const UEdGraphSchema_K2* Schema = CastChecked<const UEdGraphSchema_K2>(TargetPin->GetSchema()))
 			{
-				const bool bIsRead = (TargetPin->Direction == EGPD_Input) && !bIsExecPin;
-				const UBlueprint* Blueprint = FBlueprintEditorUtils::FindBlueprintForNode(TargetPin->GetOwningNode());
-				const bool bReadOnlyProperty = FBlueprintEditorUtils::IsPropertyReadOnlyInCurrentBlueprint(Blueprint, VariableProperty);
-				const bool bCanWriteIfNeeded = bIsRead || !bReadOnlyProperty;
+				const bool bIsExecPin = Schema->IsExecPin(*TargetPin);
 
-				FEdGraphPinType VariablePinType;
-				Schema->ConvertPropertyToPinType(VariableProperty, VariablePinType);
-				const bool bTypeMatch = Schema->ArePinTypesCompatible(VariablePinType, TargetPin->PinType) || bIsExecPin;
+				UProperty* VariableProperty = GetVariableProperty();
 
-				if (bTypeMatch && bCanWriteIfNeeded)
+				if (CanVariableBeDropped(VariableProperty, *TargetPin->GetOwningNode()->GetGraph()))
 				{
-					FEdGraphSchemaAction_K2NewNode Action;
+					const bool bIsRead = (TargetPin->Direction == EGPD_Input) && !bIsExecPin;
+					const UBlueprint* Blueprint = FBlueprintEditorUtils::FindBlueprintForNode(TargetPin->GetOwningNode());
+					const bool bWritableProperty = (FBlueprintEditorUtils::IsPropertyWritableInBlueprint(Blueprint, VariableProperty) == FBlueprintEditorUtils::EPropertyWritableState::Writable);
+					const bool bCanWriteIfNeeded = bIsRead || bWritableProperty;
+
+					FEdGraphPinType VariablePinType;
+					Schema->ConvertPropertyToPinType(VariableProperty, VariablePinType);
+					const bool bTypeMatch = Schema->ArePinTypesCompatible(VariablePinType, TargetPin->PinType) || bIsExecPin;
 
-					UK2Node_Variable* VarNode = bIsRead ? (UK2Node_Variable*)NewObject<UK2Node_VariableGet>() : (UK2Node_Variable*)NewObject<UK2Node_VariableSet>();
-					Action.NodeTemplate = VarNode;
+					if (bTypeMatch && bCanWriteIfNeeded)
+					{
+						FEdGraphSchemaAction_K2NewNode Action;
 
-					UBlueprint* DropOnBlueprint = FBlueprintEditorUtils::FindBlueprintForGraph(TargetPin->GetOwningNode()->GetGraph());
-					UEdGraphSchema_K2::ConfigureVarNode(VarNode, VariableName, VariableSource.Get(), DropOnBlueprint);
+						UK2Node_Variable* VarNode = bIsRead ? (UK2Node_Variable*)NewObject<UK2Node_VariableGet>() : (UK2Node_Variable*)NewObject<UK2Node_VariableSet>();
+						Action.NodeTemplate = VarNode;
 
-					Action.PerformAction(TargetPin->GetOwningNode()->GetGraph(), TargetPin, GraphPosition);
+						UBlueprint* DropOnBlueprint = FBlueprintEditorUtils::FindBlueprintForGraph(TargetPin->GetOwningNode()->GetGraph());
+						UEdGraphSchema_K2::ConfigureVarNode(VarNode, VariableName, VariableSource.Get(), DropOnBlueprint);
+
+						Action.PerformAction(TargetPin->GetOwningNode()->GetGraph(), TargetPin, GraphPosition);
+					}
 				}
 			}
 		}
@@ -453,9 +461,9 @@ bool FKismetVariableDragDropAction::CanExecuteMakeSetter(FNodeConstructionParams
 	if(UClass* VariableSourceClass = Cast<UClass>(InParams.VariableSource.Get()))
 	{
 		const UBlueprint* Blueprint = FBlueprintEditorUtils::FindBlueprintForGraph(InParams.Graph);
-		const bool bReadOnlyProperty = FBlueprintEditorUtils::IsPropertyReadOnlyInCurrentBlueprint(Blueprint, InVariableProperty);
+		const bool bWritableProperty = (FBlueprintEditorUtils::IsPropertyWritableInBlueprint(Blueprint, InVariableProperty) == FBlueprintEditorUtils::EPropertyWritableState::Writable);
 
-		return (!bReadOnlyProperty) && (!VariableSourceClass->HasAnyClassFlags(CLASS_Const));
+		return (bWritableProperty && !VariableSourceClass->HasAnyClassFlags(CLASS_Const));
 	}
 
 	return true;
@@ -463,10 +471,10 @@ bool FKismetVariableDragDropAction::CanExecuteMakeSetter(FNodeConstructionParams
 
 FReply FKismetVariableDragDropAction::DroppedOnPanel( const TSharedRef< SWidget >& Panel, FVector2D ScreenPosition, FVector2D GraphPosition, UEdGraph& Graph)
 {	
-	if (Cast<const UEdGraphSchema_K2>(Graph.GetSchema()) != NULL)
+	if (Graph.GetSchema()->IsA<UEdGraphSchema_K2>())
 	{
 		UProperty* VariableProperty = GetVariableProperty();
-		if (VariableProperty != nullptr && CanVariableBeDropped(VariableProperty, Graph))
+		if (VariableProperty && CanVariableBeDropped(VariableProperty, Graph))
 		{
 			UStruct* Outer = CastChecked<UStruct>(VariableProperty->GetOuter());
 			
@@ -480,7 +488,7 @@ FReply FKismetVariableDragDropAction::DroppedOnPanel( const TSharedRef< SWidget
 			AnalyticCallback.ExecuteIfBound();
 
 			// Take into account current state of modifier keys in case the user changed his mind
-			auto ModifierKeys = FSlateApplication::Get().GetModifierKeys();
+			FModifierKeysState ModifierKeys = FSlateApplication::Get().GetModifierKeys();
 			const bool bModifiedKeysActive = ModifierKeys.IsControlDown() || ModifierKeys.IsAltDown();
 			const bool bAutoCreateGetter = bModifiedKeysActive ? ModifierKeys.IsControlDown() : bControlDrag;
 			const bool bAutoCreateSetter = bModifiedKeysActive ? ModifierKeys.IsAltDown() : bAltDrag;
diff --git a/Engine/Source/Editor/Kismet/Private/BlueprintActionMenuBuilder.cpp b/Engine/Source/Editor/Kismet/Private/BlueprintActionMenuBuilder.cpp
--- a/Engine/Source/Editor/Kismet/Private/BlueprintActionMenuBuilder.cpp
+++ b/Engine/Source/Editor/Kismet/Private/BlueprintActionMenuBuilder.cpp
@@ -518,23 +518,32 @@ void FBlueprintActionMenuBuilder::RebuildActionList()
 		MenuSection->Empty();
 	}
 	
-	FBlueprintActionDatabase::FActionRegistry const& ActionDatabase = FBlueprintActionDatabase::Get().GetAllActions();
-	for (auto const& ActionEntry : ActionDatabase)
+	FBlueprintActionDatabase& ActionDatabase = FBlueprintActionDatabase::Get();
+	FBlueprintActionDatabase::FActionRegistry const& ActionRegistry = ActionDatabase.GetAllActions();
+	for (auto const& ActionEntry : ActionRegistry)
 	{
-		for (UBlueprintNodeSpawner const* NodeSpawner : ActionEntry.Value)
+		if (UObject *ActionObject = ActionEntry.Key.ResolveObjectPtr())
 		{
-			FBlueprintActionInfo BlueprintAction(ActionEntry.Key, NodeSpawner);
-
-			// @TODO: could probably have a super filter that spreads across 
-			//        all MenuSctions (to pair down on performance?)
-			for (TSharedRef<FMenuSectionDefinition> const& MenuSection : MenuSections)
+			for (UBlueprintNodeSpawner const* NodeSpawner : ActionEntry.Value)
 			{
-				for (TSharedPtr<FEdGraphSchemaAction> MenuEntry : MenuSection->MakeMenuItems(BlueprintEditorPtr, BlueprintAction))
+				FBlueprintActionInfo BlueprintAction(ActionObject, NodeSpawner);
+
+				// @TODO: could probably have a super filter that spreads across 
+				//        all MenuSctions (to pair down on performance?)
+				for (TSharedRef<FMenuSectionDefinition> const& MenuSection : MenuSections)
 				{
-					AddAction(MenuEntry);
+					for (TSharedPtr<FEdGraphSchemaAction> MenuEntry : MenuSection->MakeMenuItems(BlueprintEditorPtr, BlueprintAction))
+					{
+						AddAction(MenuEntry);
+					}
 				}
 			}
 		}
+		else
+		{
+			// Remove this (invalid) entry on the next tick.
+			ActionDatabase.DeferredRemoveEntry(ActionEntry.Key);
+		}
 	}	
 }
 
diff --git a/Engine/Source/Editor/Kismet/Private/BlueprintCompilationManager.cpp b/Engine/Source/Editor/Kismet/Private/BlueprintCompilationManager.cpp
--- a/Engine/Source/Editor/Kismet/Private/BlueprintCompilationManager.cpp
+++ b/Engine/Source/Editor/Kismet/Private/BlueprintCompilationManager.cpp
@@ -23,6 +23,7 @@
 #include "Serialization/ArchiveHasReferences.h"
 #include "Serialization/ArchiveReplaceObjectRef.h"
 #include "TickableEditorObject.h"
+#include "UObject/MetaData.h"
 #include "UObject/ReferenceChainSearch.h"
 #include "UObject/UObjectHash.h"
 #include "WidgetBlueprint.h"
@@ -178,7 +179,12 @@ void FBlueprintCompilationManagerImpl::CompileSynchronouslyImpl(const FBPCompile
 	const bool bSuppressBroadcastCompiled = true;
 	FlushCompilationQueueImpl(nullptr, bSuppressBroadcastCompiled);
 	FlushReinstancingQueueImpl();
-
+	
+	if ( GEditor )
+	{
+		GEditor->BroadcastBlueprintReinstanced();
+	}
+	
 	ensure(Request.BPToCompile->bQueuedForCompilation == false);
 
 	if(!bSkipGarbageCollection)
@@ -312,14 +318,17 @@ void FBlueprintCompilationManagerImpl::FlushCompilationQueueImpl(TArray<UObject*
 			// added or removed - strictly speaking we only need to do this when function 
 			// parameters changed, but that's a somewhat dubious optmization - ideally this
 			// work *never* needs to happen:
-			TArray<UBlueprint*> DependentBlueprints;
-			FBlueprintEditorUtils::GetDependentBlueprints(ComplieJob.BPToCompile, DependentBlueprints);
-			for(UBlueprint* DependentBlueprint : DependentBlueprints)
+			if(!FBlueprintEditorUtils::IsInterfaceBlueprint(ComplieJob.BPToCompile))
 			{
-				if(!IsQueuedForCompilation(DependentBlueprint))
+				TArray<UBlueprint*> DependentBlueprints;
+				FBlueprintEditorUtils::GetDependentBlueprints(ComplieJob.BPToCompile, DependentBlueprints);
+				for(UBlueprint* DependentBlueprint : DependentBlueprints)
 				{
-					DependentBlueprint->bQueuedForCompilation = true;
-					CurrentlyCompilingBPs.Add(FCompilerData(DependentBlueprint, ECompilationManagerJobType::Normal, nullptr, EBlueprintCompileOptions::None, true));
+					if(!IsQueuedForCompilation(DependentBlueprint))
+					{
+						DependentBlueprint->bQueuedForCompilation = true;
+						CurrentlyCompilingBPs.Add(FCompilerData(DependentBlueprint, ECompilationManagerJobType::Normal, nullptr, EBlueprintCompileOptions::None, true));
+					}
 				}
 			}
 		}
@@ -331,8 +340,22 @@ void FBlueprintCompilationManagerImpl::FlushCompilationQueueImpl(TArray<UObject*
 			FBPCompileRequest& QueuedJob = QueuedRequests[I];
 			UBlueprint* QueuedBP = QueuedJob.BPToCompile;
 
-			ensure(!(QueuedBP->GeneratedClass->ClassDefaultObject->HasAnyFlags(RF_NeedLoad)));
-			if( FBlueprintEditorUtils::IsDataOnlyBlueprint(QueuedBP) && !QueuedBP->bHasBeenRegenerated)
+			ensure(!QueuedBP->GeneratedClass || !(QueuedBP->GeneratedClass->ClassDefaultObject->HasAnyFlags(RF_NeedLoad)));
+			bool bDefaultComponentMustBeAdded = false;
+			bool bHasPendingUberGraphFrame = false;
+
+			if(UBlueprintGeneratedClass* BPGC = Cast<UBlueprintGeneratedClass>(QueuedBP->GeneratedClass))
+			{
+				if( BPGC->SimpleConstructionScript &&
+					BPGC->SimpleConstructionScript->GetSceneRootComponentTemplate() == nullptr)
+				{
+					bDefaultComponentMustBeAdded = true;
+				}
+				
+				bHasPendingUberGraphFrame = BPGC->UberGraphFramePointerProperty || BPGC->UberGraphFunction;
+			}
+			
+			if( FBlueprintEditorUtils::IsDataOnlyBlueprint(QueuedBP) && !QueuedBP->bHasBeenRegenerated && !bDefaultComponentMustBeAdded && !bHasPendingUberGraphFrame)
 			{
 				const UClass* ParentClass = QueuedBP->ParentClass;
 				if (ParentClass && ParentClass->HasAllClassFlags(CLASS_Native))
@@ -429,6 +452,11 @@ void FBlueprintCompilationManagerImpl::FlushCompilationQueueImpl(TArray<UObject*
 				UBlueprint* BP = CompilerData.BP;
 				BP->bBeingCompiled = true;
 				BP->bIsRegeneratingOnLoad = !BP->bHasBeenRegenerated && BP->GetLinker();
+				if(BP->bIsRegeneratingOnLoad)
+				{
+					// we may have cached dependencies before being fully loaded:
+					BP->bCachedDependenciesUpToDate = false;
+				}
 			}
 		}
 
@@ -534,8 +562,16 @@ void FBlueprintCompilationManagerImpl::FlushCompilationQueueImpl(TArray<UObject*
 				for (UBlueprint* CurrentBP : DependentBlueprints)
 				{
 					const EBlueprintStatus OriginalStatus = CurrentBP->Status;
+					UPackage* const Package = CurrentBP->GetOutermost();
+					const bool bStartedWithUnsavedChanges = Package != nullptr ? Package->IsDirty() : true;
+
 					FBlueprintEditorUtils::RefreshExternalBlueprintDependencyNodes(CurrentBP, BP->GeneratedClass);
+
 					CurrentBP->Status = OriginalStatus;
+					if(Package != nullptr && Package->IsDirty() && !bStartedWithUnsavedChanges)
+					{
+						Package->SetDirtyFlag(false);
+					}
 				}
 			}
 			
@@ -581,7 +617,10 @@ void FBlueprintCompilationManagerImpl::FlushCompilationQueueImpl(TArray<UObject*
 
 				UBlueprint* BP = CompilerData.BP;
 
-				OldCDOs.Add(BP, BP->GeneratedClass->ClassDefaultObject);
+				if(BP->GeneratedClass)
+				{
+					OldCDOs.Add(BP, BP->GeneratedClass->ClassDefaultObject);
+				}
 				CompilerData.Reinstancer = TSharedPtr<FBlueprintCompileReinstancer>( 
 					new FBlueprintCompileReinstancer(
 						BP->GeneratedClass, 
@@ -595,7 +634,7 @@ void FBlueprintCompilationManagerImpl::FlushCompilationQueueImpl(TArray<UObject*
 		for (FCompilerData& CompilerData : CurrentlyCompilingBPs)
 		{
 			UBlueprint* BP = CompilerData.BP;
-			if(BP->GeneratedClass->GetSuperClass()->HasAnyClassFlags(CLASS_NewerVersionExists))
+			if(BP->GeneratedClass && BP->GeneratedClass->GetSuperClass()->HasAnyClassFlags(CLASS_NewerVersionExists))
 			{
 				BP->GeneratedClass->SetSuperStruct(BP->GeneratedClass->GetSuperClass()->GetAuthoritativeClass());
 				if(CompilerData.ShouldResetClassMembers())
@@ -619,19 +658,23 @@ void FBlueprintCompilationManagerImpl::FlushCompilationQueueImpl(TArray<UObject*
 			UBlueprint* BP = CompilerData.BP;
 			if(CompilerData.ShouldCompileClassLayout())
 			{
-				ensure( BP->GeneratedClass->ClassDefaultObject == nullptr || 
+				ensure( BP->GeneratedClass == nullptr ||
+						BP->GeneratedClass->ClassDefaultObject == nullptr || 
 						BP->GeneratedClass->ClassDefaultObject->GetClass() != BP->GeneratedClass);
 				// default value propagation occurs in ReinstaneBatch, CDO will be created via CompileFunctions call:
-				BP->GeneratedClass->ClassDefaultObject = nullptr;
-				// Reset the flag, so if the user tries to use PIE it will warn them if the BP did not compile
-				BP->bDisplayCompilePIEWarning = true;
+				if(BP->GeneratedClass)
+				{
+					BP->GeneratedClass->ClassDefaultObject = nullptr;
+					// Reset the flag, so if the user tries to use PIE it will warn them if the BP did not compile
+					BP->bDisplayCompilePIEWarning = true;
 		
-				FKismetCompilerContext& CompilerContext = *(CompilerData.Compiler);
-				CompilerContext.CompileClassLayout(EInternalCompilerFlags::PostponeLocalsGenerationUntilPhaseTwo);
+					FKismetCompilerContext& CompilerContext = *(CompilerData.Compiler);
+					CompilerContext.CompileClassLayout(EInternalCompilerFlags::PostponeLocalsGenerationUntilPhaseTwo);
+				}
 			}
-			else if(CompilerData.Compiler.IsValid())
+			else if(CompilerData.Compiler.IsValid() && BP->GeneratedClass)
 			{
-				CompilerData.Compiler->NewClass = CastChecked<UBlueprintGeneratedClass>(BP->GeneratedClass);
+				CompilerData.Compiler->SetNewClass( CastChecked<UBlueprintGeneratedClass>(BP->GeneratedClass) );
 			}
 		}
 		bGeneratedClassLayoutReady = true;
@@ -642,8 +685,9 @@ void FBlueprintCompilationManagerImpl::FlushCompilationQueueImpl(TArray<UObject*
 			UBlueprint* BP = CompilerData.BP;
 			if(!CompilerData.ShouldCompileClassFunctions())
 			{
-				if( BP->GeneratedClass->ClassDefaultObject == nullptr || 
-					BP->GeneratedClass->ClassDefaultObject->GetClass() != BP->GeneratedClass)
+				if( BP->GeneratedClass &&
+					(	BP->GeneratedClass->ClassDefaultObject == nullptr || 
+						BP->GeneratedClass->ClassDefaultObject->GetClass() != BP->GeneratedClass) )
 				{
 					// relink, generate CDO:
 					UClass* BPGC = BP->GeneratedClass;
@@ -655,17 +699,23 @@ void FBlueprintCompilationManagerImpl::FlushCompilationQueueImpl(TArray<UObject*
 			}
 			else
 			{
-				ensure( BP->GeneratedClass->ClassDefaultObject == nullptr || 
+				ensure( BP->GeneratedClass == nullptr ||
+						BP->GeneratedClass->ClassDefaultObject == nullptr || 
 						BP->GeneratedClass->ClassDefaultObject->GetClass() != BP->GeneratedClass);
+				
 				// default value propagation occurrs below:
-				BP->GeneratedClass->ClassDefaultObject = nullptr;
+				if(BP->GeneratedClass)
+				{
+					BP->GeneratedClass->ClassDefaultObject = nullptr;
+				
 		
-				FKismetCompilerContext& CompilerContext = *(CompilerData.Compiler);
-				CompilerContext.CompileFunctions(
-					EInternalCompilerFlags::PostponeLocalsGenerationUntilPhaseTwo
-					|EInternalCompilerFlags::PostponeDefaultObjectAssignmentUntilReinstancing
-					|EInternalCompilerFlags::SkipRefreshExternalBlueprintDependencyNodes
-				); 
+					FKismetCompilerContext& CompilerContext = *(CompilerData.Compiler);
+					CompilerContext.CompileFunctions(
+						EInternalCompilerFlags::PostponeLocalsGenerationUntilPhaseTwo
+						|EInternalCompilerFlags::PostponeDefaultObjectAssignmentUntilReinstancing
+						|EInternalCompilerFlags::SkipRefreshExternalBlueprintDependencyNodes
+					); 
+				}
 
 				if (CompilerData.ActiveResultsLog->NumErrors == 0)
 				{
@@ -678,7 +728,7 @@ void FBlueprintCompilationManagerImpl::FlushCompilationQueueImpl(TArray<UObject*
 				}
 			}
 
-			ensure(BP->GeneratedClass->ClassDefaultObject->GetClass() == *(BP->GeneratedClass));
+			ensure(BP->GeneratedClass == nullptr || BP->GeneratedClass->ClassDefaultObject->GetClass() == *(BP->GeneratedClass));
 		}
 	} // end GTimeCompiling scope
 
@@ -690,7 +740,7 @@ void FBlueprintCompilationManagerImpl::FlushCompilationQueueImpl(TArray<UObject*
 			// CopyTermDefaultsToDefaultObject
 			for (FCompilerData& CompilerData : CurrentlyCompilingBPs)
 			{
-				if(CompilerData.Reinstancer.IsValid())
+				if(CompilerData.Reinstancer.IsValid() && CompilerData.Reinstancer->ClassToReinstance)
 				{
 					Reinstancers.Push(
 						FReinstancingJob {
@@ -724,7 +774,7 @@ void FBlueprintCompilationManagerImpl::FlushCompilationQueueImpl(TArray<UObject*
 				}
 			}
 
-			if(!BP->bIsRegeneratingOnLoad)
+			if(!BP->bIsRegeneratingOnLoad && BP->GeneratedClass)
 			{
 				FKismetEditorUtilities::StripExternalComponents(BP);
 				
@@ -732,6 +782,28 @@ void FBlueprintCompilationManagerImpl::FlushCompilationQueueImpl(TArray<UObject*
 				{
 					BP->SimpleConstructionScript->FixupRootNodeParentReferences();
 				}
+
+				const bool bIsInterface = FBlueprintEditorUtils::IsInterfaceBlueprint(BP);
+
+				TArray<UBlueprint*> DependentBPs;
+				FBlueprintEditorUtils::GetDependentBlueprints(BP, DependentBPs);
+
+				// refresh each dependent blueprint
+				for (UBlueprint* Dependent : DependentBPs)
+				{
+					// for interface changes, auto-refresh nodes on any dependent blueprints
+					// note: RefreshAllNodes() will internally send a change notification event to the dependent blueprint
+					if (bIsInterface)
+					{
+						ensure(!Dependent->bIsRegeneratingOnLoad);
+						FBlueprintEditorUtils::RefreshAllNodes(Dependent);
+					}
+					else if(!BP->bIsRegeneratingOnLoad)
+					{
+						// for non-interface changes, nodes with an external dependency have already been refreshed, and it is now safe to send a change notification event
+						Dependent->BroadcastChanged();
+					}
+				}
 				
 				UBlueprint::ValidateGeneratedClass(BP->GeneratedClass);
 			}
@@ -769,11 +841,6 @@ void FBlueprintCompilationManagerImpl::FlushCompilationQueueImpl(TArray<UObject*
 			}
 		}
 
-		if ( GEditor )
-		{
-			GEditor->BroadcastBlueprintReinstanced();
-		}
-
 		if (!bSuppressBroadcastCompiled)
 		{
 			if(GEditor)
@@ -1046,7 +1113,6 @@ void FBlueprintCompilationManagerImpl::ReinstanceBatch(TArray<FReinstancingJob>&
 		}
 	}
 
-	TSet<UObject*> ObjectsThatShouldUseOldStuff;
 	TMap<UObject*, UObject*> OldArchetypeToNewArchetype;
 
 	// 3. Update any remaining instances that are tagged as RF_ArchetypeObject or RF_InheritableComponentTemplate - 
@@ -1121,9 +1187,37 @@ void FBlueprintCompilationManagerImpl::ReinstanceBatch(TArray<FReinstancingJob>&
 				FBlueprintCompileReinstancer::CopyPropertiesForUnrelatedObjects(Archetype, NewArchetype);
 
 				OldArchetypeToNewArchetype.Add(Archetype, NewArchetype);
-				ArchetypeReferencers.Add(Archetype);
-				ObjectsThatShouldUseOldStuff.Add(Archetype);
+				// Map old subobjects to new subobjects. This is needed by UMG right now, which allows owning archetypes to reference subobjects
+				// in subwidgets:
+				{
+					TArray<UObject*> OldSubobjects;
+					GetObjectsWithOuter(Archetype, OldSubobjects);
+					TArray<UObject*> NewSubobjects;
+					GetObjectsWithOuter(NewArchetype, NewSubobjects);
+
+					TMap<FName, UObject*> OldNameMap;
+					for(UObject* OldSubobject : OldSubobjects )
+					{
+						OldNameMap.Add( OldSubobject->GetFName(), OldSubobject );
+					}
+
+					TMap<FName, UObject*> NewNameMap;
+					for(UObject* NewSubobject : NewSubobjects )
+					{
+						NewNameMap.Add( NewSubobject->GetFName(), NewSubobject );
+					}
+
+					for(TPair< FName, UObject* > OldSubobject : OldNameMap )
+					{
+						UObject** NewSubobject = NewNameMap.Find(OldSubobject.Key);
+						OldArchetypeToNewArchetype.Add(OldSubobject.Value, NewSubobject ? *NewSubobject : nullptr );
+					}
+					
+				}
+
+				ArchetypeReferencers.Add(NewArchetype);
 
+				Archetype->RemoveFromRoot();
 				Archetype->MarkPendingKill();
 			}
 		}
@@ -1168,6 +1262,11 @@ UClass* FBlueprintCompilationManagerImpl::FastGenerateSkeletonClass(UBlueprint*
 	FCompilerResultsLog MessageLog;
 
 	UClass* ParentClass = BP->ParentClass;
+	if(ParentClass == nullptr)
+	{
+		return nullptr;
+	}
+
 	if(ParentClass->ClassGeneratedBy)
 	{
 		if(UBlueprint* ParentBP = Cast<UBlueprint>(ParentClass->ClassGeneratedBy))
@@ -1181,9 +1280,17 @@ UClass* FBlueprintCompilationManagerImpl::FastGenerateSkeletonClass(UBlueprint*
 
 	UBlueprintGeneratedClass* Ret = nullptr;
 	UBlueprintGeneratedClass* OriginalNewClass = CompilerContext.NewClass;
-	if(BP->SkeletonGeneratedClass == nullptr )
+	FString SkelClassName = FString::Printf(TEXT("SKEL_%s_C"), *BP->GetName());
+
+	if (BP->SkeletonGeneratedClass == nullptr)
 	{
-		CompilerContext.SpawnNewClass(FString::Printf(TEXT("SKEL_%s_C"), *BP->GetName()));
+		// This might exist in the package because we are being reloaded in place
+		BP->SkeletonGeneratedClass = FindObject<UBlueprintGeneratedClass>(BP->GetOutermost(), *SkelClassName);
+	}
+
+	if (BP->SkeletonGeneratedClass == nullptr)
+	{
+		CompilerContext.SpawnNewClass(SkelClassName);
 		Ret = CompilerContext.NewClass;
 		Ret->SetFlags(RF_Transient);
 		CompilerContext.NewClass = OriginalNewClass;
@@ -1197,7 +1304,16 @@ UClass* FBlueprintCompilationManagerImpl::FastGenerateSkeletonClass(UBlueprint*
 	Ret->ClassGeneratedBy = BP;
 
 	// This is a version of PrecompileFunction that does not require 'terms' and graph cloning:
-	const auto MakeFunction = [Ret, ParentClass, Schema, BP, &MessageLog](FName FunctionNameFName, UField**& InCurrentFieldStorageLocation, UField**& InCurrentParamStorageLocation, bool bIsStaticFunction, EFunctionFlags InFunctionFlags, const TArray<UK2Node_FunctionResult*>& ReturnNodes, const TArray<UEdGraphPin*>& InputPins, bool bForceArrayStructRefsConst) -> UFunction*
+	const auto MakeFunction = [Ret, ParentClass, Schema, BP, &MessageLog]
+		(	FName FunctionNameFName, 
+			UField**& InCurrentFieldStorageLocation, 
+			UField**& InCurrentParamStorageLocation, 
+			EFunctionFlags InFunctionFlags, 
+			const TArray<UK2Node_FunctionResult*>& ReturnNodes, 
+			const TArray<UEdGraphPin*>& InputPins, 
+			bool bIsStaticFunction, 
+			bool bForceArrayStructRefsConst, 
+			UFunction* SignatureOverride) -> UFunction*
 	{
 		if(!ensure(FunctionNameFName != FName())
 			|| FindObjectFast<UFunction>(Ret, FunctionNameFName, true ))
@@ -1241,20 +1357,23 @@ UClass* FBlueprintCompilationManagerImpl::FastGenerateSkeletonClass(UBlueprint*
 			}
 		}
 		NewFunction->SetSuperStruct( ParentFn );
-					
+		
 		InCurrentParamStorageLocation = &NewFunction->Children;
 
 		// params:
-		if(ParentFn)
+		if(ParentFn || SignatureOverride)
 		{
-			NewFunction->FunctionFlags |= (ParentFn->FunctionFlags & (FUNC_FuncInherit | FUNC_Public | FUNC_Protected | FUNC_Private | FUNC_BlueprintPure));
-			for (TFieldIterator<UProperty> PropIt(ParentFn); PropIt && (PropIt->PropertyFlags & CPF_Parm); ++PropIt)
+			UFunction* SignatureFn = ParentFn ? ParentFn : SignatureOverride;
+			NewFunction->FunctionFlags |= (SignatureFn->FunctionFlags & (FUNC_FuncInherit | FUNC_Public | FUNC_Protected | FUNC_Private | FUNC_BlueprintPure));
+			for (TFieldIterator<UProperty> PropIt(SignatureFn); PropIt && (PropIt->PropertyFlags & CPF_Parm); ++PropIt)
 			{
 				UProperty* ClonedParam = CastChecked<UProperty>(StaticDuplicateObject(*PropIt, NewFunction, PropIt->GetFName(), RF_AllFlags, nullptr, EDuplicateMode::Normal, EInternalObjectFlags::AllFlags & ~(EInternalObjectFlags::Native) ));
+				ClonedParam->PropertyFlags |= CPF_BlueprintVisible|CPF_BlueprintReadOnly;
 				ClonedParam->Next = nullptr;
 				*InCurrentParamStorageLocation = ClonedParam;
 				InCurrentParamStorageLocation = &ClonedParam->Next;
 			}
+			UMetaData::CopyMetadata(SignatureFn, NewFunction);
 		}
 		else
 		{
@@ -1263,8 +1382,8 @@ UClass* FBlueprintCompilationManagerImpl::FastGenerateSkeletonClass(UBlueprint*
 			{
 				if(Pin->Direction == EEdGraphPinDirection::EGPD_Output && !Schema->IsExecPin(*Pin) && Pin->ParentPin == nullptr && Pin->GetName() != UK2Node_Event::DelegateOutputName)
 				{
-					// Reimplmentation of FKismetCompilerContext::CreatePropertiesFromList without dependence on 'terms'
-					UProperty* Param = FKismetCompilerUtilities::CreatePropertyOnScope(NewFunction, *(Pin->PinName), Pin->PinType, Ret, 0, Schema, MessageLog);
+					// Reimplementation of FKismetCompilerContext::CreatePropertiesFromList without dependence on 'terms'
+					UProperty* Param = FKismetCompilerUtilities::CreatePropertyOnScope(NewFunction, *(Pin->PinName), Pin->PinType, Ret, CPF_BlueprintVisible|CPF_BlueprintReadOnly, Schema, MessageLog);
 					if(Param)
 					{
 						Param->SetFlags(RF_Transient);
@@ -1369,14 +1488,15 @@ UClass* FBlueprintCompilationManagerImpl::FastGenerateSkeletonClass(UBlueprint*
 				
 				UField** CurrentParamStorageLocation = nullptr;
 				UFunction* NewFunction = MakeFunction(
-					FName(*(Graph->GetName() + FunctionNamePostfix)),
-					InCurrentFieldStorageLocation,
-					CurrentParamStorageLocation,
-					bIsStaticFunction,
+					FName(*(Graph->GetName() + FunctionNamePostfix)), 
+					InCurrentFieldStorageLocation, 
+					CurrentParamStorageLocation, 
 					(EFunctionFlags)(EntryNode->GetFunctionFlags() & ~FUNC_Native),
-					ReturnNodes,
+					ReturnNodes, 
 					EntryNode->Pins,
-					false
+					bIsStaticFunction, 
+					false,
+					nullptr
 				);
 
 				if(NewFunction)
@@ -1449,17 +1569,20 @@ UClass* FBlueprintCompilationManagerImpl::FastGenerateSkeletonClass(UBlueprint*
 		{
 			FName EventNodeName = CompilerContext.GetEventStubFunctionName(Event);
 
+			UFunction* SourceFN = Event->FindEventSignatureFunction();
+
 			UField** CurrentParamStorageLocation = nullptr;
 
 			UFunction* NewFunction = MakeFunction(
-				EventNodeName,
-				CurrentFieldStorageLocation,
-				CurrentParamStorageLocation,
-				false,
+				EventNodeName, 
+				CurrentFieldStorageLocation, 
+				CurrentParamStorageLocation, 
 				(EFunctionFlags)(Event->FunctionFlags|FUNC_BlueprintCallable),
-				TArray<UK2Node_FunctionResult*>(),
+				TArray<UK2Node_FunctionResult*>(), 
 				Event->Pins,
-				true
+				false, 
+				true,
+				SourceFN
 			);
 
 			if(NewFunction)
@@ -1510,11 +1633,12 @@ UClass* FBlueprintCompilationManagerImpl::FastGenerateSkeletonClass(UBlueprint*
 				Fn->GetFName(), 
 				CurrentFieldStorageLocation, 
 				CurrentParamStorageLocation, 
-				false, 
 				Fn->FunctionFlags & ~FUNC_Native, 
 				TArray<UK2Node_FunctionResult*>(), 
 				TArray<UEdGraphPin*>(),
-				false
+				false, 
+				false,
+				nullptr
 			);
 		}
 	}
diff --git a/Engine/Source/Editor/Kismet/Private/BlueprintDetailsCustomization.cpp b/Engine/Source/Editor/Kismet/Private/BlueprintDetailsCustomization.cpp
--- a/Engine/Source/Editor/Kismet/Private/BlueprintDetailsCustomization.cpp
+++ b/Engine/Source/Editor/Kismet/Private/BlueprintDetailsCustomization.cpp
@@ -2706,10 +2706,11 @@ void FBlueprintGraphArgumentLayout::GenerateChildContent( IDetailChildrenBuilder
 {
 	if (bHasDefaultValue)
 	{
-		if (UEdGraphPin* FoundPin = GetPin())
+		UEdGraphPin* FoundPin = GetPin();
+		if (FoundPin)
 		{
 			// Certain types are outlawed at the compiler level
-			bool bTypeWithNoDefaults = (FoundPin->PinType.PinCategory == UEdGraphSchema_K2::PC_Object) || (FoundPin->PinType.PinCategory == UEdGraphSchema_K2::PC_Class) || (FoundPin->PinType.PinCategory == UEdGraphSchema_K2::PC_Interface);
+			const bool bTypeWithNoDefaults = (FoundPin->PinType.PinCategory == UEdGraphSchema_K2::PC_Object) || (FoundPin->PinType.PinCategory == UEdGraphSchema_K2::PC_Class) || (FoundPin->PinType.PinCategory == UEdGraphSchema_K2::PC_Interface) || UEdGraphSchema_K2::IsExecPin(*FoundPin);
 
 			if (!FoundPin->PinType.bIsReference && !bTypeWithNoDefaults)
 			{
@@ -2740,21 +2741,25 @@ void FBlueprintGraphArgumentLayout::GenerateChildContent( IDetailChildrenBuilder
 			}
 		}
 
-		ChildrenBuilder.AddCustomRow( LOCTEXT( "FunctionArgDetailsPassByReference", "Pass-by-Reference" ) )
-		.NameContent()
-		[
-			SNew(STextBlock)
-				.Text( LOCTEXT( "FunctionArgDetailsPassByReference", "Pass-by-Reference" ) )
-				.ToolTipText( LOCTEXT("FunctionArgDetailsPassByReferenceTooltip", "Pass this paremeter by reference?") )
-				.Font( IDetailLayoutBuilder::GetDetailFont() )
-		]
-		.ValueContent()
-		[
-			SNew(SCheckBox)
-				.IsChecked( this, &FBlueprintGraphArgumentLayout::IsRefChecked )
-				.OnCheckStateChanged( this, &FBlueprintGraphArgumentLayout::OnRefCheckStateChanged)
+		// Exec pins can't be passed by reference
+		if (FoundPin == nullptr || !UEdGraphSchema_K2::IsExecPin(*FoundPin))
+		{
+			ChildrenBuilder.AddCustomRow(LOCTEXT("FunctionArgDetailsPassByReference", "Pass-by-Reference"))
+				.NameContent()
+				[
+					SNew(STextBlock)
+					.Text(LOCTEXT("FunctionArgDetailsPassByReference", "Pass-by-Reference"))
+				.ToolTipText(LOCTEXT("FunctionArgDetailsPassByReferenceTooltip", "Pass this paremeter by reference?"))
+				.Font(IDetailLayoutBuilder::GetDetailFont())
+				]
+			.ValueContent()
+				[
+					SNew(SCheckBox)
+					.IsChecked(this, &FBlueprintGraphArgumentLayout::IsRefChecked)
+				.OnCheckStateChanged(this, &FBlueprintGraphArgumentLayout::OnRefCheckStateChanged)
 				.IsEnabled(!ShouldPinBeReadOnly())
-		];
+				];
+		}
 	}
 		
 	
@@ -3679,13 +3684,13 @@ void FBlueprintGraphActionDetails::OnEditorCallableEventModified( const ECheckBo
 	}
 }
 
-UMulticastDelegateProperty* FBlueprintDelegateActionDetails::GetDelegatePoperty() const
+UMulticastDelegateProperty* FBlueprintDelegateActionDetails::GetDelegateProperty() const
 {
 	if (MyBlueprint.IsValid())
 	{
 		if (const FEdGraphSchemaAction_K2Delegate* DelegateVar = MyBlueprint.Pin()->SelectionAsDelegate())
 		{
-			return DelegateVar->GetDelegatePoperty();
+			return DelegateVar->GetDelegateProperty();
 		}
 	}
 	return NULL;
@@ -3693,7 +3698,7 @@ UMulticastDelegateProperty* FBlueprintDelegateActionDetails::GetDelegatePoperty(
 
 bool FBlueprintDelegateActionDetails::IsBlueprintProperty() const
 {
-	const UMulticastDelegateProperty* Property = GetDelegatePoperty();
+	const UMulticastDelegateProperty* Property = GetDelegateProperty();
 	const UBlueprint* Blueprint = GetBlueprintObj();
 	if(Property && Blueprint)
 	{
@@ -3731,7 +3736,7 @@ UEdGraph* FBlueprintDelegateActionDetails::GetGraph() const
 
 FText FBlueprintDelegateActionDetails::OnGetTooltipText() const
 {
-	if (UMulticastDelegateProperty* DelegateProperty = GetDelegatePoperty())
+	if (UMulticastDelegateProperty* DelegateProperty = GetDelegateProperty())
 	{
 		FString Result;
 		FBlueprintEditorUtils::GetBlueprintVariableMetaData(GetBlueprintObj(), DelegateProperty->GetFName(), NULL, TEXT("tooltip"), Result);
@@ -3742,7 +3747,7 @@ FText FBlueprintDelegateActionDetails::OnGetTooltipText() const
 
 void FBlueprintDelegateActionDetails::OnTooltipTextCommitted(const FText& NewText, ETextCommit::Type InTextCommit)
 {
-	if (UMulticastDelegateProperty* DelegateProperty = GetDelegatePoperty())
+	if (UMulticastDelegateProperty* DelegateProperty = GetDelegateProperty())
 	{
 		FBlueprintEditorUtils::SetBlueprintVariableMetaData(GetBlueprintObj(), DelegateProperty->GetFName(), NULL, TEXT("tooltip"), NewText.ToString() );
 	}
@@ -3750,7 +3755,7 @@ void FBlueprintDelegateActionDetails::OnTooltipTextCommitted(const FText& NewTex
 
 FText FBlueprintDelegateActionDetails::OnGetCategoryText() const
 {
-	if (UMulticastDelegateProperty* DelegateProperty = GetDelegatePoperty())
+	if (UMulticastDelegateProperty* DelegateProperty = GetDelegateProperty())
 	{
 		const UEdGraphSchema_K2* K2Schema = GetDefault<UEdGraphSchema_K2>();
 		FName DelegateName = DelegateProperty->GetFName();
@@ -3773,7 +3778,7 @@ void FBlueprintDelegateActionDetails::OnCategoryTextCommitted(const FText& NewTe
 {
 	if (InTextCommit == ETextCommit::OnEnter || InTextCommit == ETextCommit::OnUserMovedFocus)
 	{
-		if (UMulticastDelegateProperty* DelegateProperty = GetDelegatePoperty())
+		if (UMulticastDelegateProperty* DelegateProperty = GetDelegateProperty())
 		{
 			// Remove excess whitespace and prevent categories with just spaces
 			FText CategoryName = FText::TrimPrecedingAndTrailing(NewText);
@@ -3796,7 +3801,7 @@ TSharedRef< ITableRow > FBlueprintDelegateActionDetails::MakeCategoryViewWidget(
 
 void FBlueprintDelegateActionDetails::OnCategorySelectionChanged( TSharedPtr<FText> ProposedSelection, ESelectInfo::Type /*SelectInfo*/ )
 {
-	UMulticastDelegateProperty* DelegateProperty = GetDelegatePoperty();
+	UMulticastDelegateProperty* DelegateProperty = GetDelegateProperty();
 	if (DelegateProperty && ProposedSelection.IsValid())
 	{
 		FText NewCategory = *ProposedSelection.Get();
@@ -3922,7 +3927,7 @@ void FBlueprintDelegateActionDetails::CustomizeDetails( IDetailLayoutBuilder& De
 void FBlueprintDelegateActionDetails::CollectAvailibleSignatures()
 {
 	FunctionsToCopySignatureFrom.Empty();
-	if (UMulticastDelegateProperty* Property = GetDelegatePoperty())
+	if (UMulticastDelegateProperty* Property = GetDelegateProperty())
 	{
 		if (UClass* ScopeClass = Cast<UClass>(Property->GetOuterUField()))
 		{
@@ -3948,7 +3953,7 @@ void FBlueprintDelegateActionDetails::CollectAvailibleSignatures()
 void FBlueprintDelegateActionDetails::OnFunctionSelected(TSharedPtr<FString> FunctionName, ESelectInfo::Type SelectInfo)
 {
 	UK2Node_EditablePinBase* FunctionEntryNode = FunctionEntryNodePtr.Get();
-	UMulticastDelegateProperty* Property = GetDelegatePoperty();
+	UMulticastDelegateProperty* Property = GetDelegateProperty();
 	UClass* ScopeClass = Property ? Cast<UClass>(Property->GetOuterUField()) : NULL;
 	const UEdGraphSchema_K2* Schema = GetDefault<UEdGraphSchema_K2>();
 
@@ -3988,7 +3993,10 @@ void FBaseBlueprintGraphActionDetails::OnParamsChanged(UK2Node_EditablePinBase*
 		RegenerateOutputsChildrenDelegate.ExecuteIfBound();
 
 		// Reconstruct the entry/exit definition and recompile the blueprint to make sure the signature has changed before any fixups
-		TargetNode->ReconstructNode();
+		{
+			TGuardValue<ESaveOrphanPinMode> GuardSaveMode(TargetNode->OrphanedPinSaveMode, ESaveOrphanPinMode::SaveNone);
+			TargetNode->ReconstructNode();
+		}
 
 		const UEdGraphSchema_K2* K2Schema = GetDefault<UEdGraphSchema_K2>();
 
diff --git a/Engine/Source/Editor/Kismet/Private/BlueprintDetailsCustomization.h b/Engine/Source/Editor/Kismet/Private/BlueprintDetailsCustomization.h
--- a/Engine/Source/Editor/Kismet/Private/BlueprintDetailsCustomization.h
+++ b/Engine/Source/Editor/Kismet/Private/BlueprintDetailsCustomization.h
@@ -379,7 +379,7 @@ private:
 
 	void SetEntryNode();
 
-	UMulticastDelegateProperty* GetDelegatePoperty() const;
+	UMulticastDelegateProperty* GetDelegateProperty() const;
 	FText OnGetTooltipText() const;
 	void OnTooltipTextCommitted(const FText& NewText, ETextCommit::Type InTextCommit);
 	FText OnGetCategoryText() const;
diff --git a/Engine/Source/Editor/Kismet/Private/BlueprintEditor.cpp b/Engine/Source/Editor/Kismet/Private/BlueprintEditor.cpp
--- a/Engine/Source/Editor/Kismet/Private/BlueprintEditor.cpp
+++ b/Engine/Source/Editor/Kismet/Private/BlueprintEditor.cpp
@@ -3290,6 +3290,7 @@ void FBlueprintEditor::Compile()
 		}
 
 		FCompilerResultsLog LogResults;
+		LogResults.SetSourcePath(BlueprintObj->GetPathName());
 		LogResults.BeginEvent(TEXT("Compile"));
 		LogResults.bLogDetailedResults = GetDefault<UBlueprintEditorSettings>()->bShowDetailedCompileResults;
 		LogResults.EventDisplayThresholdMs = GetDefault<UBlueprintEditorSettings>()->CompileEventDisplayThresholdMs;
@@ -3304,7 +3305,9 @@ void FBlueprintEditor::Compile()
 		}
 		FKismetEditorUtilities::CompileBlueprint(BlueprintObj, CompileOptions, &LogResults);
 
-		bool bForceMessageDisplay = ((LogResults.NumWarnings > 0) || (LogResults.NumErrors > 0)) && !BlueprintObj->bIsRegeneratingOnLoad;
+		LogResults.EndEvent();
+
+		const bool bForceMessageDisplay = ((LogResults.NumWarnings > 0) || (LogResults.NumErrors > 0)) && !BlueprintObj->bIsRegeneratingOnLoad;
 		DumpMessagesToCompilerLog(LogResults.Messages, bForceMessageDisplay);
 
 		UBlueprintEditorSettings const* BpEditorSettings = GetDefault<UBlueprintEditorSettings>();
@@ -3323,9 +3326,7 @@ void FBlueprintEditor::Compile()
 
 		AppendExtraCompilerResults(CompilerResultsListing);
 
-		LogResults.EndEvent();
-
-	    // Update the blueprint instrumenation state and show the profiler window if required
+	    // Update the blueprint instrumentation state and show the profiler window if required
 	    if (bProfilerAvailable)
 	    {
 		    bBlueprintHasInstrumentation = BlueprintObj && BlueprintObj->GeneratedClass ? BlueprintObj->GeneratedClass->HasInstrumentation() : false;
@@ -3848,7 +3849,7 @@ bool FBlueprintEditor::CanPromoteToVariable(bool bInToMemberVariable) const
 	{
 		if (UEdGraphPin* Pin = FocusedGraphEd->GetGraphPinForMenu())
 		{
-			if (bInToMemberVariable || (!bInToMemberVariable && FBlueprintEditorUtils::DoesSupportLocalVariables(FocusedGraphEd->GetCurrentGraph())))
+			if (!Pin->bOrphanedPin && (bInToMemberVariable || (!bInToMemberVariable && FBlueprintEditorUtils::DoesSupportLocalVariables(FocusedGraphEd->GetCurrentGraph()))))
 			{
 				bCanPromote = K2Schema->CanPromotePinToVariable(*Pin);
 			}
@@ -3941,7 +3942,7 @@ void FBlueprintEditor::OnAddExecutionPin()
 			const FScopedTransaction Transaction( LOCTEXT("AddExecutionPin", "Add Execution Pin") );
 			SeqNode->Modify();
 
-			SeqNode->AddPinToExecutionNode();
+			SeqNode->AddInputPin();
 
 			const UEdGraphSchema* Schema = SeqNode->GetSchema();
 			Schema->ReconstructNode(*SeqNode);
@@ -7303,7 +7304,7 @@ void FBlueprintEditor::RefreshStandAloneDefaultsEditor()
 		}
 	}
 
-	if ( DefaultObjects.Num() )
+	if ( DefaultObjects.Num() && DefaultEditor.IsValid() )
 	{
 		DefaultEditor->ShowDetailsForObjects(DefaultObjects);
 	}
@@ -7327,12 +7328,6 @@ void FBlueprintEditor::OnAddNewVariable()
 {
 	const FScopedTransaction Transaction( LOCTEXT("AddVariable", "Add Variable") );
 
-	// Reset MyBlueprint item filter so new variable is visible
-	if (MyBlueprintWidget.IsValid())
-	{
-		MyBlueprintWidget->OnResetItemFilter();
-	}
-
 	FName VarName = FBlueprintEditorUtils::FindUniqueKismetName(GetBlueprintObj(), TEXT("NewVar"));
 
 	bool bSuccess = MyBlueprintWidget.IsValid() && FBlueprintEditorUtils::AddMemberVariable(GetBlueprintObj(), VarName, MyBlueprintWidget->GetLastPinTypeUsed());
@@ -7385,12 +7380,6 @@ void FBlueprintEditor::OnAddNewDelegate()
 	UBlueprint* const Blueprint = GetBlueprintObj();
 	check(NULL != Blueprint);
 
-	// Reset MyBlueprint item filter so new variable is visible
-	if (MyBlueprintWidget.IsValid())
-	{
-		MyBlueprintWidget->OnResetItemFilter();
-	}
-
 	FName Name = FBlueprintEditorUtils::FindUniqueKismetName(GetBlueprintObj(), TEXT("NewEventDispatcher"));
 
 
@@ -7454,12 +7443,6 @@ void FBlueprintEditor::NewDocument_OnClicked(ECreatedDocumentType GraphType)
 		DocumentNameText = LOCTEXT("NewDocNewName", "NewDocument");
 		break;
 	}
-	
-	// Reset MyBlueprint item filter so new variable is visible
-	if( bResetMyBlueprintFilter )
-	{
-		MyBlueprintWidget->OnResetItemFilter();
-	}
 
 	FName DocumentName = FName(*DocumentNameText.ToString());
 
@@ -7569,34 +7552,37 @@ void FBlueprintEditor::NotifyPostChange(const FPropertyChangedEvent& PropertyCha
 		FocusedGraphEdPtr.Pin()->NotifyPostPropertyChange(PropertyChangedEvent, PropertyName);
 	}
 	
-	UBlueprint* Blueprint = GetBlueprintObj();
-	UPackage* BlueprintPackage = Blueprint->GetOutermost();
-
-	// if any of the objects being edited are in our package, mark us as dirty
-	bool bPropertyInBlueprint = false;
-	for (int32 ObjectIndex = 0; ObjectIndex < PropertyChangedEvent.GetNumObjectsBeingEdited(); ++ObjectIndex)
+	if (IsEditingSingleBlueprint())
 	{
-		const UObject* Object = PropertyChangedEvent.GetObjectBeingEdited(ObjectIndex);
-		if (Object && Object->GetOutermost() == BlueprintPackage)
+		UBlueprint* Blueprint = GetBlueprintObj();
+		UPackage* BlueprintPackage = Blueprint->GetOutermost();
+
+		// if any of the objects being edited are in our package, mark us as dirty
+		bool bPropertyInBlueprint = false;
+		for (int32 ObjectIndex = 0; ObjectIndex < PropertyChangedEvent.GetNumObjectsBeingEdited(); ++ObjectIndex)
 		{
-			bPropertyInBlueprint = true;
-			break;
+			const UObject* Object = PropertyChangedEvent.GetObjectBeingEdited(ObjectIndex);
+			if (Object && Object->GetOutermost() == BlueprintPackage)
+			{
+				bPropertyInBlueprint = true;
+				break;
+			}
 		}
-	}
 
-	if (bPropertyInBlueprint)
-	{
-		// Note: if change type is "interactive," hold off on applying the change (e.g. this will occur if the user is scrubbing a spinbox value; we don't want to apply the change until the mouse is released, for performance reasons)
-		if (IsEditingSingleBlueprint() && PropertyChangedEvent.ChangeType != EPropertyChangeType::Interactive)
+		if (bPropertyInBlueprint)
 		{
-			FBlueprintEditorUtils::MarkBlueprintAsModified(Blueprint, PropertyChangedEvent);
+			// Note: if change type is "interactive," hold off on applying the change (e.g. this will occur if the user is scrubbing a spinbox value; we don't want to apply the change until the mouse is released, for performance reasons)
+			if (PropertyChangedEvent.ChangeType != EPropertyChangeType::Interactive)
+			{
+				FBlueprintEditorUtils::MarkBlueprintAsModified(Blueprint, PropertyChangedEvent);
 
-			// Call PostEditChange() on any Actors that might be based on this Blueprint
-			FBlueprintEditorUtils::PostEditChangeBlueprintActors(Blueprint);
-		}
+				// Call PostEditChange() on any Actors that might be based on this Blueprint
+				FBlueprintEditorUtils::PostEditChangeBlueprintActors(Blueprint);
+			}
 
-		// Force updates to occur immediately during interactive mode (otherwise the preview won't refresh because it won't be ticking)
-		UpdateSCSPreview(PropertyChangedEvent.ChangeType == EPropertyChangeType::Interactive);
+			// Force updates to occur immediately during interactive mode (otherwise the preview won't refresh because it won't be ticking)
+			UpdateSCSPreview(PropertyChangedEvent.ChangeType == EPropertyChangeType::Interactive);
+		}
 	}
 }
 
diff --git a/Engine/Source/Editor/Kismet/Private/BlueprintPaletteFavorites.cpp b/Engine/Source/Editor/Kismet/Private/BlueprintPaletteFavorites.cpp
--- a/Engine/Source/Editor/Kismet/Private/BlueprintPaletteFavorites.cpp
+++ b/Engine/Source/Editor/Kismet/Private/BlueprintPaletteFavorites.cpp
@@ -71,7 +71,7 @@ static FBlueprintNodeSignature BlueprintPaletteFavoritesImpl::ConstructLegacySig
 		FEdGraphSchemaAction_K2Delegate* DelegateAction = (FEdGraphSchemaAction_K2Delegate*)InPaletteAction.Get();
 
 		SignatureNodeClass = UK2Node_BaseMCDelegate::StaticClass();
-		SignatureSubObject = DelegateAction->GetDelegatePoperty();
+		SignatureSubObject = DelegateAction->GetDelegateProperty();
 	}
 	// if we can pull out a node associated with this action
 	else if (UK2Node const* NodeTemplate = FBlueprintActionMenuUtils::ExtractNodeTemplateFromAction(InPaletteAction))
diff --git a/Engine/Source/Editor/Kismet/Private/Debugging/SKismetDebuggingView.cpp b/Engine/Source/Editor/Kismet/Private/Debugging/SKismetDebuggingView.cpp
--- a/Engine/Source/Editor/Kismet/Private/Debugging/SKismetDebuggingView.cpp
+++ b/Engine/Source/Editor/Kismet/Private/Debugging/SKismetDebuggingView.cpp
@@ -196,7 +196,7 @@ FText FLatentActionLineItem::GetDescription() const
 {
 	if (UObject* ParentObject = ParentObjectRef.Get())
 	{
-		if (UWorld* World = GEngine->GetWorldFromContextObject(ParentObject))
+		if (UWorld* World = GEngine->GetWorldFromContextObject(ParentObject, EGetWorldErrorMode::ReturnNull))
 		{
 			FLatentActionManager& LatentActionManager = World->GetLatentActionManager();
 			return FText::FromString(LatentActionManager.GetDescription(ParentObject, UUID));
@@ -659,7 +659,7 @@ public:
 				}
 
 				// It could also have active latent behaviors
-				if (UWorld* World = GEngine->GetWorldFromContextObject(ParentObject))
+				if (UWorld* World = GEngine->GetWorldFromContextObject(ParentObject, EGetWorldErrorMode::ReturnNull))
 				{
 					FLatentActionManager& LatentActionManager = World->GetLatentActionManager();
 
diff --git a/Engine/Source/Editor/Kismet/Private/SBlueprintPalette.cpp b/Engine/Source/Editor/Kismet/Private/SBlueprintPalette.cpp
--- a/Engine/Source/Editor/Kismet/Private/SBlueprintPalette.cpp
+++ b/Engine/Source/Editor/Kismet/Private/SBlueprintPalette.cpp
@@ -1286,29 +1286,18 @@ bool SBlueprintPaletteItem::OnNameTextVerifyChanged(const FText& InNewText, FTex
 	TSharedPtr<INameValidatorInterface> NameValidator = MakeShareable(new FKismetNameValidator(BlueprintObj, OriginalName, ValidationScope));
 
 	EValidatorResult ValidatorResult = NameValidator->IsValid(TextAsString);
-	if(ValidatorResult == EValidatorResult::AlreadyInUse)
+	switch (ValidatorResult)
 	{
-		OutErrorMessage = FText::Format(LOCTEXT("RenameFailed_InUse", "{0} is in use by another variable or function!"), InNewText);
-	}
-	else if(ValidatorResult == EValidatorResult::EmptyName)
-	{
-		OutErrorMessage = LOCTEXT("RenameFailed_LeftBlank", "Names cannot be left blank!");
-	}
-	else if(ValidatorResult == EValidatorResult::TooLong)
-	{
-		OutErrorMessage = LOCTEXT("RenameFailed_NameTooLong", "Names must have fewer than 100 characters!");
-	}
-	else if(ValidatorResult == EValidatorResult::LocallyInUse)
-	{
-		OutErrorMessage = LOCTEXT("ConflictsWithProperty", "Conflicts with another another local variable or function parameter!");
-	}
-
-	if(OutErrorMessage.IsEmpty())
-	{
-		return true;
+	case EValidatorResult::Ok:
+	case EValidatorResult::ExistingName:
+		// These are fine, don't need to surface to the user, the rename can 'proceed' even if the name is the existing one
+		break;
+	default:
+		OutErrorMessage = INameValidatorInterface::GetErrorText(TextAsString, ValidatorResult);
+		break;
 	}
 
-	return false;
+	return OutErrorMessage.IsEmpty();
 }
 
 //------------------------------------------------------------------------------
diff --git a/Engine/Source/Editor/Kismet/Private/SMyBlueprint.cpp b/Engine/Source/Editor/Kismet/Private/SMyBlueprint.cpp
--- a/Engine/Source/Editor/Kismet/Private/SMyBlueprint.cpp
+++ b/Engine/Source/Editor/Kismet/Private/SMyBlueprint.cpp
@@ -16,6 +16,7 @@
 #include "Dialogs/Dialogs.h"
 #include "Kismet2/KismetEditorUtilities.h"
 #include "EdGraphSchema_K2.h"
+#include "K2Node_AddComponent.h"
 #include "K2Node_Event.h"
 #include "K2Node_CallFunction.h"
 #include "K2Node_Tunnel.h"
@@ -474,7 +475,7 @@ void SMyBlueprint::OnCategoryNameCommitted(const FText& InNewText, ETextCommit::
 			else if (Actions[i]->GetTypeId() == FEdGraphSchemaAction_K2Delegate::StaticGetTypeId())
 			{
 				FEdGraphSchemaAction_K2Delegate* DelegateAction = (FEdGraphSchemaAction_K2Delegate*)Actions[i].Get();
-				FBlueprintEditorUtils::SetBlueprintVariableCategory(GetBlueprintObj(), DelegateAction->GetDelegatePoperty()->GetFName(), NULL, CategoryName, true);
+				FBlueprintEditorUtils::SetBlueprintVariableCategory(GetBlueprintObj(), DelegateAction->GetDelegateProperty()->GetFName(), NULL, CategoryName, true);
 			}
 			else if (Actions[i]->GetTypeId() == FEdGraphSchemaAction_K2Graph::StaticGetTypeId())
 			{
@@ -1630,7 +1631,7 @@ void SMyBlueprint::OnActionSelectedHelper(TSharedPtr<FEdGraphSchemaAction> InAct
 		else if (InAction->GetTypeId() == FEdGraphSchemaAction_K2Delegate::StaticGetTypeId())
 		{
 			FEdGraphSchemaAction_K2Delegate* DelegateAction = (FEdGraphSchemaAction_K2Delegate*)InAction.Get();
-			if (UMulticastDelegateProperty* Property = DelegateAction->GetDelegatePoperty())
+			if (UMulticastDelegateProperty* Property = DelegateAction->GetDelegateProperty())
 			{
 				Inspector->ShowDetailsForSingleObject(Property, SKismetInspector::FShowDetailsOptions(FText::FromString(Property->GetName())));
 			}
@@ -2617,13 +2618,18 @@ void SMyBlueprint::OnDuplicateAction()
 
 		DuplicatedGraph->Modify();
 
-		// Generate new Guids for all nodes in the graph
-		// *NOTE* this cannot occur during PostDuplicate, node Guids need to remain static during duplication for Blueprint compilation
+		// Generate new Guids and component templates for all relevant nodes in the graph
+		// *NOTE* this cannot occur during PostDuplicate, node Guids and component templates need to remain static during duplication for Blueprint compilation
 		for (UEdGraphNode* EdGraphNode : DuplicatedGraph->Nodes)
 		{
 			if (EdGraphNode)
 			{
 				EdGraphNode->CreateNewGuid();
+
+				if (UK2Node_AddComponent* AddComponentNode = Cast<UK2Node_AddComponent>(EdGraphNode))
+				{
+					AddComponentNode->MakeNewComponentTemplate();
+				}
 			}
 		}
 		// Only function and macro duplication is supported
@@ -2787,7 +2793,15 @@ void SMyBlueprint::SelectItemByName(const FName& ItemName, ESelectInfo::Type Sel
 	else
 	{
 		// Attempt to select the item in the main graph action menu
-		GraphActionMenu->SelectItemByName(ItemName, SelectInfo, SectionId, bIsCategory);
+		const bool bSucceededAtSelecting = GraphActionMenu->SelectItemByName(ItemName, SelectInfo, SectionId, bIsCategory);
+		if (!bSucceededAtSelecting)
+		{
+			// We failed to select the item, maybe because it was filtered out?
+			// Reset the item filter and try again (we don't do this first because someone went to the effort of typing
+			// a filter and probably wants to keep it unless it is getting in the way, as it just has)
+			OnResetItemFilter();
+			GraphActionMenu->SelectItemByName(ItemName, SelectInfo, SectionId, bIsCategory);
+		}
 	}
 }
 
diff --git a/Engine/Source/Editor/Kismet/Private/UserDefinedEnumEditor.cpp b/Engine/Source/Editor/Kismet/Private/UserDefinedEnumEditor.cpp
--- a/Engine/Source/Editor/Kismet/Private/UserDefinedEnumEditor.cpp
+++ b/Engine/Source/Editor/Kismet/Private/UserDefinedEnumEditor.cpp
@@ -1,6 +1,7 @@
 // Copyright 1998-2017 Epic Games, Inc. All Rights Reserved.
 
 #include "UserDefinedEnumEditor.h"
+#include "Editor.h"
 #include "Widgets/Text/STextBlock.h"
 #include "EditorStyleSet.h"
 #include "PropertyEditorModule.h"
@@ -261,8 +262,6 @@ FLinearColor FUserDefinedEnumEditor::GetWorldCentricTabColorScale() const
 	return FLinearColor( 0.5f, 0.0f, 0.0f, 0.5f );
 }
 
-
-
 void FEnumDetails::CustomizeDetails( IDetailLayoutBuilder& DetailLayout )
 {
 	const TArray<TWeakObjectPtr<UObject>> Objects = DetailLayout.GetDetailsView().GetSelectedObjects();
@@ -310,8 +309,15 @@ void FEnumDetails::CustomizeDetails( IDetailLayoutBuilder& DetailLayout )
 	}
 }
 
+FEnumDetails::FEnumDetails()
+	: TargetEnum(nullptr)
+{
+	GEditor->RegisterForUndo(this);
+}
+
 FEnumDetails::~FEnumDetails()
 {
+	GEditor->UnregisterForUndo( this );
 }
 
 void FEnumDetails::OnForceRefresh()
@@ -322,6 +328,11 @@ void FEnumDetails::OnForceRefresh()
 	}
 }
 
+void FEnumDetails::PostUndo(bool bSuccess)
+{
+	OnForceRefresh();
+}
+
 void FEnumDetails::PreChange(const class UUserDefinedEnum* Enum, FEnumEditorUtils::EEnumEditorChangeInfo Info)
 {
 }
@@ -352,9 +363,8 @@ void FEnumDetails::OnBitmaskFlagsAttributeStateChanged(ECheckBoxState InNewState
 
 bool FUserDefinedEnumLayout::CausedChange() const
 {
-	for(auto Iter = Children.CreateConstIterator(); Iter; ++Iter)
+	for (const TWeakPtr<class FUserDefinedEnumIndexLayout>& Child : Children)
 	{
-		const auto Child = (*Iter);
 		if (Child.IsValid() && Child.Pin()->CausedChange())
 		{
 			return true;
@@ -365,8 +375,8 @@ bool FUserDefinedEnumLayout::CausedChange() const
 
 void FUserDefinedEnumLayout::GenerateChildContent( IDetailChildrenBuilder& ChildrenBuilder )
 {
-	Children.Empty();
 	const int32 EnumToShowNum = FMath::Max<int32>(0, TargetEnum->NumEnums() - 1);
+	Children.Reset(EnumToShowNum);
 	for (int32 EnumIdx = 0; EnumIdx < EnumToShowNum; ++EnumIdx)
 	{
 		TSharedRef<class FUserDefinedEnumIndexLayout> EnumIndexLayout = MakeShareable(new FUserDefinedEnumIndexLayout(TargetEnum.Get(), EnumIdx) );
diff --git a/Engine/Source/Editor/Kismet/Private/UserDefinedEnumEditor.h b/Engine/Source/Editor/Kismet/Private/UserDefinedEnumEditor.h
--- a/Engine/Source/Editor/Kismet/Private/UserDefinedEnumEditor.h
+++ b/Engine/Source/Editor/Kismet/Private/UserDefinedEnumEditor.h
@@ -9,6 +9,7 @@
 #include "IDetailCustomization.h"
 #include "IDetailCustomNodeBuilder.h"
 #include "IDetailsView.h"
+#include "EditorUndoClient.h"
 
 #include "Editor/UnrealEd/Public/Kismet2/EnumEditorUtils.h"
 
@@ -59,7 +60,7 @@ protected:
 };
 
 /** Details customization for functions and graphs selected in the MyBlueprint panel */
-class FEnumDetails : public IDetailCustomization, FEnumEditorUtils::INotifyOnEnumChanged
+class FEnumDetails : public IDetailCustomization, FEnumEditorUtils::INotifyOnEnumChanged, FEditorUndoClient
 {
 public:
 	/** Makes a new instance of this detail layout class for a specific detail view requesting it */
@@ -68,11 +69,7 @@ public:
 		return MakeShareable(new FEnumDetails);
 	}
 
-	FEnumDetails()
-		: TargetEnum(NULL)
-	{
-	}
-
+	FEnumDetails();
 	~FEnumDetails();
 
 	/** IDetailCustomization interface */
@@ -85,6 +82,10 @@ public:
 	virtual void PreChange(const class UUserDefinedEnum* Enum, FEnumEditorUtils::EEnumEditorChangeInfo Info) override;
 	virtual void PostChange(const class UUserDefinedEnum* Enum, FEnumEditorUtils::EEnumEditorChangeInfo Info) override;
 
+	/** FEditorUndoClient Interface */
+	virtual void PostUndo(bool bSuccess) override;
+	virtual void PostRedo(bool bSuccess) override { PostUndo(bSuccess); }
+
 private:
 	/** Handles new enum element request */
 	FReply OnAddNewEnumerator();
diff --git a/Engine/Source/Editor/KismetCompiler/Private/AnimBlueprintCompiler.cpp b/Engine/Source/Editor/KismetCompiler/Private/AnimBlueprintCompiler.cpp
--- a/Engine/Source/Editor/KismetCompiler/Private/AnimBlueprintCompiler.cpp
+++ b/Engine/Source/Editor/KismetCompiler/Private/AnimBlueprintCompiler.cpp
@@ -653,7 +653,7 @@ void FAnimBlueprintCompiler::ProcessSubInstance(UAnimGraphNode_SubInstance* SubI
 
 	for(UEdGraphPin* Pin : SubInstance->Pins)
 	{
-		if(Pin->PinName == TEXT("InPose") || Pin->PinName == TEXT("Pose"))
+		if (Pin->bOrphanedPin || Pin->PinName == TEXT("InPose") || Pin->PinName == TEXT("Pose"))
 		{
 			continue;
 		}
@@ -669,14 +669,17 @@ void FAnimBlueprintCompiler::ProcessSubInstance(UAnimGraphNode_SubInstance* SubI
 			FKismetCompilerUtilities::LinkAddedProperty(NewAnimBlueprintClass, NewProperty);
 
 			// Add mappings to the node
-			FAnimNode_SubInstance& RuntimeNode = SubInstance->Node;
-
-			RuntimeNode.SourcePropertyNames.Add(NewProperty->GetFName());
-			
-			// Find the property on the internal instance
-			if(UProperty* FoundProperty = FindField<UProperty>(*RuntimeNode.InstanceClass, *Pin->PinName))
+			if(!bGenerateSubInstanceVariables)
 			{
-				RuntimeNode.DestPropertyNames.Add(FoundProperty->GetFName());
+				FAnimNode_SubInstance& RuntimeNode = SubInstance->Node;
+				
+				if(UProperty* FoundProperty = FindField<UProperty>(*RuntimeNode.InstanceClass, *Pin->PinName))
+				{
+					RuntimeNode.SourcePropertyNames.Add(NewProperty->GetFName());
+			
+					// Find the property on the internal instance
+					RuntimeNode.DestPropertyNames.Add(FoundProperty->GetFName());
+				}
 			}
 		}
 	}
@@ -1794,13 +1797,17 @@ void FAnimBlueprintCompiler::SpawnNewClass(const FString& NewClassName)
 	NewClass = NewAnimBlueprintClass;
 }
 
+void FAnimBlueprintCompiler::OnNewClassSet(UBlueprintGeneratedClass* ClassToUse)
+{
+	NewAnimBlueprintClass = CastChecked<UAnimBlueprintGeneratedClass>(ClassToUse);
+}
+
 void FAnimBlueprintCompiler::CleanAndSanitizeClass(UBlueprintGeneratedClass* ClassToClean, UObject*& InOldCDO)
 {
 	Super::CleanAndSanitizeClass(ClassToClean, InOldCDO);
 
 	// Make sure our typed pointer is set
-	check(ClassToClean == NewClass);
-	NewAnimBlueprintClass = CastChecked<UAnimBlueprintGeneratedClass>((UObject*)NewClass);
+	check(ClassToClean == NewClass && NewAnimBlueprintClass == NewClass);
 
 	NewAnimBlueprintClass->AnimBlueprintDebugData = FAnimBlueprintDebugData();
 
@@ -1849,6 +1856,8 @@ void FAnimBlueprintCompiler::FinishCompilingClass(UClass* Class)
 
 void FAnimBlueprintCompiler::PostCompile()
 {
+	Super::PostCompile();
+
 	UAnimBlueprintGeneratedClass* AnimBlueprintGeneratedClass = CastChecked<UAnimBlueprintGeneratedClass>(NewClass);
 
 	UAnimInstance* DefaultAnimInstance = CastChecked<UAnimInstance>(AnimBlueprintGeneratedClass->GetDefaultObject());
diff --git a/Engine/Source/Editor/KismetCompiler/Private/KismetCompiler.cpp b/Engine/Source/Editor/KismetCompiler/Private/KismetCompiler.cpp
--- a/Engine/Source/Editor/KismetCompiler/Private/KismetCompiler.cpp
+++ b/Engine/Source/Editor/KismetCompiler/Private/KismetCompiler.cpp
@@ -56,6 +56,9 @@
 
 static bool bDebugPropertyPropagation = false;
 
+FSimpleMulticastDelegate FKismetCompilerContext::OnPreCompile;
+FSimpleMulticastDelegate FKismetCompilerContext::OnPostCompile;
+
 #define USE_TRANSIENT_SKELETON 0
 
 #define LOCTEXT_NAMESPACE "KismetCompiler"
@@ -243,7 +246,7 @@ void FKismetCompilerContext::CleanAndSanitizeClass(UBlueprintGeneratedClass* Cla
 	TransientClass->ClassGeneratedBy = Blueprint;
 	TransientClass->ClassFlags |= CLASS_CompiledFromBlueprint;
 
-	NewClass = ClassToClean;
+	SetNewClass( ClassToClean );
 	InOldCDO = ClassToClean->ClassDefaultObject; // we don't need to create the CDO at this point
 	
 	const ERenameFlags RenFlags = REN_DontCreateRedirectors |  ((bRecompilingOnLoad) ? REN_ForceNoResetLoaders : 0) | REN_NonTransactional | REN_DoNotDirty;
@@ -370,25 +373,33 @@ void FKismetCompilerContext::ValidateLink(const UEdGraphPin* PinA, const UEdGrap
 {
 	Super::ValidateLink(PinA, PinB);
 
-	// At this point we can assume the pins are linked, and as such the connection response should not be to disallow
-	// @todo: Potentially revisit this later.
-	// This API is intended to describe how to handle a potentially new connection to a pin that may already have a connection.
-	// However it also checks all necessary constraints for a valid connection to exist. We rely on the fact that the "disallow"
-	// response will be returned if the pins are not compatible; any other response here then means that the connection is valid.
-	const FPinConnectionResponse ConnectResponse = Schema->CanCreateConnection(PinA, PinB);
+	// We don't want to validate orphaned pin connections to avoid noisy connection errors that are
+	// already being reported
+	const bool bShouldValidatePinA = (PinA == nullptr || !PinA->bOrphanedPin);
+	const bool bShouldValidatePinB = (PinB == nullptr || !PinB->bOrphanedPin);
 
-	const bool bForbiddenConnection = (ConnectResponse.Response == CONNECT_RESPONSE_DISALLOW);
-	const bool bMissingConversion   = (ConnectResponse.Response == CONNECT_RESPONSE_MAKE_WITH_CONVERSION_NODE);
-	if (bForbiddenConnection || bMissingConversion)
+	if (bShouldValidatePinA && bShouldValidatePinB)
 	{
-		const FString ErrorMessage = FString::Printf(*LOCTEXT("PinTypeMismatch_Error", "Can't connect pins @@ and @@: %s").ToString(), *ConnectResponse.Message.ToString());
-		if (ConnectResponse.IsFatal())
-		{
-			MessageLog.Error(*ErrorMessage, PinA, PinB);
-		}
-		else
+		// At this point we can assume the pins are linked, and as such the connection response should not be to disallow
+		// @todo: Potentially revisit this later.
+		// This API is intended to describe how to handle a potentially new connection to a pin that may already have a connection.
+		// However it also checks all necessary constraints for a valid connection to exist. We rely on the fact that the "disallow"
+		// response will be returned if the pins are not compatible; any other response here then means that the connection is valid.
+		const FPinConnectionResponse ConnectResponse = Schema->CanCreateConnection(PinA, PinB);
+
+		const bool bForbiddenConnection = (ConnectResponse.Response == CONNECT_RESPONSE_DISALLOW);
+		const bool bMissingConversion = (ConnectResponse.Response == CONNECT_RESPONSE_MAKE_WITH_CONVERSION_NODE);
+		if (bForbiddenConnection || bMissingConversion)
 		{
-			MessageLog.Warning(*ErrorMessage, PinA, PinB);
+			const FString ErrorMessage = FString::Printf(*LOCTEXT("PinTypeMismatch_Error", "Can't connect pins @@ and @@: %s").ToString(), *ConnectResponse.Message.ToString());
+			if (ConnectResponse.IsFatal())
+			{
+				MessageLog.Error(*ErrorMessage, PinA, PinB);
+			}
+			else
+			{
+				MessageLog.Warning(*ErrorMessage, PinA, PinB);
+			}
 		}
 	}
 
@@ -633,6 +644,12 @@ void FKismetCompilerContext::CreateClassVariablesFromBlueprint()
 			if(!Variable.DefaultValue.IsEmpty())
 			{
 				SetPropertyDefaultValue(NewProperty, Variable.DefaultValue);
+
+				// We're copying the value to the real CDO, so clear the version stored in the blueprint editor data
+				if (CompileOptions.CompileType == EKismetCompileType::Full)
+				{
+					Variable.DefaultValue.Empty();
+				}
 			}
 
 			if (NewProperty->HasAnyPropertyFlags(CPF_Net))
@@ -889,7 +906,7 @@ static void SwapElementsInSingleLinkedList(UField* & PtrToFirstElement, UField*
 void FKismetCompilerContext::CreateParametersForFunction(FKismetFunctionContext& Context, UFunction* ParameterSignature, UField**& FunctionPropertyStorageLocation)
 {
 	const bool bArePropertiesLocal = true;
-	CreatePropertiesFromList(Context.Function, FunctionPropertyStorageLocation, Context.Parameters, CPF_Parm, bArePropertiesLocal, /*bPropertiesAreParameters=*/ true);
+	CreatePropertiesFromList(Context.Function, FunctionPropertyStorageLocation, Context.Parameters, CPF_Parm|CPF_BlueprintVisible | CPF_BlueprintReadOnly, bArePropertiesLocal, /*bPropertiesAreParameters=*/ true);
 	CreatePropertiesFromList(Context.Function, FunctionPropertyStorageLocation, Context.Results, CPF_Parm | CPF_OutParm, bArePropertiesLocal, /*bPropertiesAreParameters=*/ true);
 
 	//MAKE SURE THE PARAMETERS ORDER MATCHES THE OVERRIDEN FUNCTION
@@ -919,6 +936,13 @@ void FKismetCompilerContext::CreateParametersForFunction(FKismetFunctionContext&
 					break;
 				}
 			}
+
+			// Ensure that the 'CPF_UObjectWrapper' flag is propagated through to new parameters, so that wrapper types like 'TSubclassOf' can be preserved if the compiled UFunction is ever nativized.
+			if (SignatureIt->HasAllPropertyFlags(CPF_UObjectWrapper))
+			{
+				CastChecked<UProperty>(*CurrentFieldStorageLocation)->SetPropertyFlags(CPF_UObjectWrapper);
+			}
+
 			CurrentFieldStorageLocation = &((*CurrentFieldStorageLocation)->Next);
 		}
 		FunctionPropertyStorageLocation = CurrentFieldStorageLocation;
@@ -1735,7 +1759,8 @@ void FKismetCompilerContext::PrecompileFunction(FKismetFunctionContext& Context,
 		if (Context.IsDelegateSignature())
 		{
 			Context.Function->FunctionFlags |= FUNC_Delegate;
-			if(UMulticastDelegateProperty* Property = FindObject<UMulticastDelegateProperty>(Context.NewClass, *Context.DelegateSignatureName.ToString()))
+
+			if(UMulticastDelegateProperty* Property = Cast<UMulticastDelegateProperty>(StaticFindObjectFast(UMulticastDelegateProperty::StaticClass(), NewClass, Context.DelegateSignatureName)))
 			{
 				Property->SignatureFunction = Context.Function;
 			}
@@ -2522,22 +2547,34 @@ void FKismetCompilerContext::ExpandTimelineNodes(UEdGraph* SourceGraph)
 
 FPinConnectionResponse FKismetCompilerContext::MovePinLinksToIntermediate(UEdGraphPin& SourcePin, UEdGraphPin& IntermediatePin)
 {
-	 UEdGraphSchema_K2 const* K2Schema = GetSchema();
-	 FPinConnectionResponse ConnectionResult = K2Schema->MovePinLinks(SourcePin, IntermediatePin, true);
+	FPinConnectionResponse ConnectionResult;
 
-	 CheckConnectionResponse(ConnectionResult, SourcePin.GetOwningNode());
-	 MessageLog.NotifyIntermediatePinCreation(&IntermediatePin, &SourcePin);
+	// If we're modifying a removed pin there will be other compile errors and we don't want odd connection disallowed error so don't even try to move the pin links
+	if (!SourcePin.bOrphanedPin)
+	{
+		UEdGraphSchema_K2 const* K2Schema = GetSchema();
+		ConnectionResult = K2Schema->MovePinLinks(SourcePin, IntermediatePin, true);
+
+		CheckConnectionResponse(ConnectionResult, SourcePin.GetOwningNode());
+		MessageLog.NotifyIntermediatePinCreation(&IntermediatePin, &SourcePin);
+	}
 
 	 return ConnectionResult;
 }
 
 FPinConnectionResponse FKismetCompilerContext::CopyPinLinksToIntermediate(UEdGraphPin& SourcePin, UEdGraphPin& IntermediatePin)
 {
-	UEdGraphSchema_K2 const* K2Schema = GetSchema();
-	FPinConnectionResponse ConnectionResult = K2Schema->CopyPinLinks(SourcePin, IntermediatePin, true);
+	FPinConnectionResponse ConnectionResult;
 
-	CheckConnectionResponse(ConnectionResult, SourcePin.GetOwningNode());
-	MessageLog.NotifyIntermediatePinCreation(&IntermediatePin, &SourcePin);
+	// If we're modifying a removed pin there will be other compile errors and we don't want odd connection disallowed error so don't even try to move the pin links
+	if (!SourcePin.bOrphanedPin)
+	{
+		UEdGraphSchema_K2 const* K2Schema = GetSchema();
+		ConnectionResult = K2Schema->CopyPinLinks(SourcePin, IntermediatePin, true);
+
+		CheckConnectionResponse(ConnectionResult, SourcePin.GetOwningNode());
+		MessageLog.NotifyIntermediatePinCreation(&IntermediatePin, &SourcePin);
+	}
 
 	return ConnectionResult;
 }
@@ -4333,6 +4370,12 @@ void FKismetCompilerContext::Compile()
 	CompileFunctions(EInternalCompilerFlags::None);
 }
 
+void FKismetCompilerContext::SetNewClass(UBlueprintGeneratedClass* ClassToUse)
+{
+	NewClass = ClassToUse;
+	OnNewClassSet(ClassToUse);
+}
+
 bool FKismetCompilerContext::ValidateGeneratedClass(UBlueprintGeneratedClass* Class)
 {
 	// Our CDO should be properly constructed by this point and should always exist
diff --git a/Engine/Source/Editor/KismetCompiler/Private/KismetCompilerMisc.cpp b/Engine/Source/Editor/KismetCompiler/Private/KismetCompilerMisc.cpp
--- a/Engine/Source/Editor/KismetCompiler/Private/KismetCompilerMisc.cpp
+++ b/Engine/Source/Editor/KismetCompiler/Private/KismetCompilerMisc.cpp
@@ -867,13 +867,13 @@ UProperty* FKismetCompilerUtilities::CreatePrimitiveProperty(UObject* PropertySc
 	{
 		UClass* SubType = (PinSubCategory == Schema->PSC_Self) ? SelfClass : Cast<UClass>(PinSubCategoryObject);
 
-		if (SubType == NULL)
+		if (SubType == nullptr)
 		{
 			// If this is from a degenerate pin, because the object type has been removed, default this to a UObject subtype so we can make a dummy term for it to allow the compiler to continue
 			SubType = UObject::StaticClass();
 		}
 
-		if (SubType != NULL)
+		if (SubType)
 		{
 			const bool bIsInterface = SubType->HasAnyClassFlags(CLASS_Interface)
 				|| ((SubType == SelfClass) && ensure(SelfClass->ClassGeneratedBy) && FBlueprintEditorUtils::IsInterfaceBlueprint(CastChecked<UBlueprint>(SelfClass->ClassGeneratedBy)));
@@ -914,8 +914,7 @@ UProperty* FKismetCompilerUtilities::CreatePrimitiveProperty(UObject* PropertySc
 	}
 	else if (PinCategory == Schema->PC_Struct)
 	{
-		UScriptStruct* SubType = Cast<UScriptStruct>(PinSubCategoryObject);
-		if (SubType != NULL)
+		if (UScriptStruct* SubType = Cast<UScriptStruct>(PinSubCategoryObject))
 		{
 			FString StructureError;
 			if (FStructureEditorUtils::EStructureError::Ok == FStructureEditorUtils::IsStructureValid(SubType, NULL, &StructureError))
@@ -951,7 +950,7 @@ UProperty* FKismetCompilerUtilities::CreatePrimitiveProperty(UObject* PropertySc
 	{
 		UClass* SubType = Cast<UClass>(PinSubCategoryObject);
 
-		if (SubType == NULL)
+		if (SubType == nullptr)
 		{
 			// If this is from a degenerate pin, because the object type has been removed, default this to a UObject subtype so we can make a dummy term for it to allow the compiler to continue
 			SubType = UObject::StaticClass();
@@ -963,11 +962,11 @@ UProperty* FKismetCompilerUtilities::CreatePrimitiveProperty(UObject* PropertySc
 					));
 		}
 
-		if (SubType != NULL)
+		if (SubType)
 		{
 			if (PinCategory == Schema->PC_AssetClass)
 			{
-				auto AssetClassProperty = NewObject<UAssetClassProperty>(PropertyScope, ValidatedPropertyName, ObjectFlags);
+				UAssetClassProperty* AssetClassProperty = NewObject<UAssetClassProperty>(PropertyScope, ValidatedPropertyName, ObjectFlags);
 				// we want to use this setter function instead of setting the 
 				// MetaClass member directly, because it properly handles  
 				// placeholder classes (classes that are stubbed in during load)
@@ -1164,6 +1163,8 @@ UProperty* FKismetCompilerUtilities::CreatePropertyOnScope(UStruct* Scope, const
 			NewMapProperty->ValueProp = CreatePrimitiveProperty(PropertyScope, ValueName, Type.PinValueType.TerminalCategory, Type.PinValueType.TerminalSubCategory, Type.PinValueType.TerminalSubCategoryObject.Get(), SelfClass, Type.bIsWeakPointer, Schema, MessageLog);;
 			if (!NewMapProperty->ValueProp)
 			{
+				NewProperty->MarkPendingKill();
+				NewProperty = nullptr;
 				NewMapProperty->MarkPendingKill();
 			}
 			else
@@ -1649,27 +1650,30 @@ FBPTerminal* FNodeHandlingFunctor::RegisterLiteral(FKismetFunctionContext& Conte
 
 void FNodeHandlingFunctor::RegisterNets(FKismetFunctionContext& Context, UEdGraphNode* Node)
 {
-	for (int32 PinIndex = 0; PinIndex < Node->Pins.Num(); ++PinIndex)
+	for (UEdGraphPin* Pin : Node->Pins)
 	{
-		UEdGraphPin* Pin = Node->Pins[PinIndex];
-		if (Pin->bNotConnectable && Pin->LinkedTo.Num() > 0)
+		if (!Pin->bOrphanedPin)
 		{
-			CompilerContext.MessageLog.Warning(*LOCTEXT("NotConnectablePinLinked", "@@ is linked to another pin but is marked as not connectable. This pin connection will not be compiled.").ToString(), Pin);
-		}
-		else if (!CompilerContext.GetSchema()->IsMetaPin(*Pin)
-			|| (CompilerContext.GetSchema()->IsSelfPin(*Pin) && Pin->LinkedTo.Num() == 0 && Pin->DefaultObject) )
-		{
-			UEdGraphPin* Net = FEdGraphUtilities::GetNetFromPin(Pin);
-
-			if (Context.NetMap.Find(Net) == NULL)
+			if (Pin->bNotConnectable && Pin->LinkedTo.Num() > 0)
 			{
-				if ((Net->Direction == EGPD_Input) && (Net->LinkedTo.Num() == 0))
-				{
-					RegisterLiteral(Context, Net);
-				}
-				else
+				// If it is not connectible due to being orphaned no need to warn as we have other messaging for that
+				CompilerContext.MessageLog.Warning(*LOCTEXT("NotConnectablePinLinked", "@@ is linked to another pin but is marked as not connectable. This pin connection will not be compiled.").ToString(), Pin);
+			}
+			else if (!CompilerContext.GetSchema()->IsMetaPin(*Pin)
+				|| (Pin->LinkedTo.Num() == 0 && Pin->DefaultObject && CompilerContext.GetSchema()->IsSelfPin(*Pin) ))
+			{
+				UEdGraphPin* Net = FEdGraphUtilities::GetNetFromPin(Pin);
+
+				if (Context.NetMap.Find(Net) == nullptr)
 				{
-					RegisterNet(Context, Pin);
+					if ((Net->Direction == EGPD_Input) && (Net->LinkedTo.Num() == 0))
+					{
+						RegisterLiteral(Context, Net);
+					}
+					else
+					{
+						RegisterNet(Context, Pin);
+					}
 				}
 			}
 		}
diff --git a/Engine/Source/Editor/KismetCompiler/Public/KismetCompiler.h b/Engine/Source/Editor/KismetCompiler/Public/KismetCompiler.h
--- a/Engine/Source/Editor/KismetCompiler/Public/KismetCompiler.h
+++ b/Engine/Source/Editor/KismetCompiler/Public/KismetCompiler.h
@@ -129,6 +129,10 @@ public:
 
 	// Flag to trigger ProcessSubInstance in CreateClassVariablesFromBlueprint:
 	bool bGenerateSubInstanceVariables;
+
+	static FSimpleMulticastDelegate OnPreCompile;
+	static FSimpleMulticastDelegate OnPostCompile;
+
 public:
 	FKismetCompilerContext(UBlueprint* SourceSketch, FCompilerResultsLog& InMessageLog, const FKismetCompilerOptions& InCompilerOptions, TArray<UObject*>* InObjLoaded);
 	virtual ~FKismetCompilerContext();
@@ -142,6 +146,9 @@ public:
 	/** Compile a blueprint into a class and a set of functions */
 	void Compile();
 
+	/** Function used to assign the new class that will be used by the compiler */
+	void SetNewClass(UBlueprintGeneratedClass* ClassToUse);
+
 	const UEdGraphSchema_K2* GetSchema() const { return Schema; }
 
 	// Spawns an intermediate node associated with the source node (for error purposes)
@@ -243,6 +250,7 @@ protected:
 	virtual UEdGraphSchema_K2* CreateSchema();
 	virtual void PostCreateSchema();
 	virtual void SpawnNewClass(const FString& NewClassName);
+	virtual void OnNewClassSet(UBlueprintGeneratedClass* ClassToUse) {}
 
 	/**
 	 * Backwards Compatability:  Ensures that the passed in TargetClass is of the proper type (e.g. BlueprintGeneratedClass, AnimBlueprintGeneratedClass), and NULLs the reference if it is not 
@@ -319,8 +327,8 @@ protected:
 	virtual void PostCompileDiagnostics() {}
 
 	// Gives derived classes a chance to hook up any custom logic
-	virtual void PreCompile() {}
-	virtual void PostCompile() {}
+	virtual void PreCompile() { OnPreCompile.Broadcast(); }
+	virtual void PostCompile() { OnPostCompile.Broadcast(); }
 
 	/** Determines if a node is pure */
 	virtual bool IsNodePure(const UEdGraphNode* Node) const;
diff --git a/Engine/Source/Editor/LevelEditor/Private/LevelEditorToolBar.cpp b/Engine/Source/Editor/LevelEditor/Private/LevelEditorToolBar.cpp
--- a/Engine/Source/Editor/LevelEditor/Private/LevelEditorToolBar.cpp
+++ b/Engine/Source/Editor/LevelEditor/Private/LevelEditorToolBar.cpp
@@ -316,7 +316,9 @@ namespace LevelEditorActionHelpers
 
 		FText RootClassName = FText::FromString(InRootClass->GetName());
 		TSharedRef<SWidget> ClassViewer = FModuleManager::LoadModuleChecked<FClassViewerModule>("ClassViewer").CreateClassViewer(Options, InOnClassPicked);
-		InMenuBuilder.BeginSection(NAME_None, FText::Format( NSLOCTEXT("LevelToolBarViewMenu", "SelectGameModeLabel", "Select {RootClass} class"), RootClassName ));
+		FFormatNamedArguments FormatArgs;
+		FormatArgs.Add(TEXT("RootClass"), RootClassName);
+		InMenuBuilder.BeginSection(NAME_None, FText::Format(NSLOCTEXT("LevelToolBarViewMenu", "SelectGameModeLabel", "Select {RootClass} class"), FormatArgs));
 		InMenuBuilder.AddWidget(ClassViewer, FText::GetEmpty(), true);
 		InMenuBuilder.EndSection();
 	}
@@ -347,7 +349,9 @@ namespace LevelEditorActionHelpers
 
 		FText RootClassName = FText::FromString(InRootClass->GetName());
 		TSharedRef<SWidget> ClassViewer = FModuleManager::LoadModuleChecked<FClassViewerModule>("ClassViewer").CreateClassViewer(Options, InOnClassPicked);
-		InMenuBuilder.BeginSection(NAME_None, FText::Format( NSLOCTEXT("LevelToolBarViewMenu", "CreateGameModeLabel", "Select {RootClass} parent class"), RootClassName ));
+		FFormatNamedArguments FormatArgs;
+		FormatArgs.Add(TEXT("RootClass"), RootClassName);
+		InMenuBuilder.BeginSection(NAME_None, FText::Format(NSLOCTEXT("LevelToolBarViewMenu", "CreateGameModeLabel", "Select {RootClass} parent class"), FormatArgs));
 		InMenuBuilder.AddWidget(ClassViewer, FText::GetEmpty(), true);
 		InMenuBuilder.EndSection();
 	}
@@ -708,7 +712,9 @@ void LevelEditorActionHelpers::GetBlueprintSettingsSubMenu(FMenuBuilder& InMenuB
 	);
 
 	// Select a game mode, this is always available so the user can switch his selection
-	InMenuBuilder.AddSubMenu( FText::Format( LOCTEXT("SelectGameModeClass", "Select {RootClass} Class"), RootClassName ),
+	FFormatNamedArguments Args;
+	Args.Add(TEXT("RootClass"), RootClassName);
+	InMenuBuilder.AddSubMenu(FText::Format(LOCTEXT("SelectGameModeClass", "Select {RootClass} Class"), Args),
 		GetSelectMenuTooltip(GetGameModeClass(InSettingsData.LevelEditor, InSettingsData.bIsProjectSettings), InSettingsData.RootClass, InSettingsData.bIsProjectSettings),
 		FNewMenuDelegate::CreateStatic( &LevelEditorActionHelpers::GetSelectSettingsClassSubMenu, InSettingsData.RootClass, InSettingsData.OnSelectClassPicked ),
 		FUIAction(
@@ -860,12 +866,15 @@ FText LevelEditorActionHelpers::GetOpenGameStateBlueprintLabel(TWeakPtr< SLevelE
 #define LOCTEXT_NAMESPACE "LevelToolBarViewMenu"
 	if(UClass* GameStateClass = GetGameStateClass(InLevelEditor, bInIsProjectSettings))
 	{
+		FFormatNamedArguments FormatArgs;
 		if(GameStateClass->ClassGeneratedBy)
 		{
-			return FText::Format( LOCTEXT("GameStateEditBlueprint", "GameState: Edit {GameStateName}"), FText::FromString(GameStateClass->ClassGeneratedBy->GetName()));
+			FormatArgs.Add(TEXT("GameStateName"), FText::FromString(GameStateClass->ClassGeneratedBy->GetName()));
+			return FText::Format(LOCTEXT("GameStateEditBlueprint", "GameState: Edit {GameStateName}"), FormatArgs);
 		}
 
-		return FText::Format( LOCTEXT("GameStateBlueprint", "GameState: {GameStateName}"), FText::FromString(GameStateClass->GetName()));
+		FormatArgs.Add(TEXT("GameStateName"), FText::FromString(GameStateClass->GetName()));
+		return FText::Format(LOCTEXT("GameStateBlueprint", "GameState: {GameStateName}"), FormatArgs);
 	}
 
 	return LOCTEXT("GameStateCreateBlueprint", "GameState: New...");
@@ -931,12 +940,15 @@ FText LevelEditorActionHelpers::GetOpenPawnBlueprintLabel(TWeakPtr< SLevelEditor
 #define LOCTEXT_NAMESPACE "LevelToolBarViewMenu"
 	if(UClass* PawnClass = GetPawnClass(InLevelEditor, bInIsProjectSettings))
 	{
+		FFormatNamedArguments FormatArgs;
 		if(PawnClass->ClassGeneratedBy)
 		{
-			return FText::Format( LOCTEXT("PawnEditBlueprint", "Pawn: Edit {PawnName}"), FText::FromString(PawnClass->ClassGeneratedBy->GetName()));
+			FormatArgs.Add(TEXT("PawnName"), FText::FromString(PawnClass->ClassGeneratedBy->GetName()));
+			return FText::Format(LOCTEXT("PawnEditBlueprint", "Pawn: Edit {PawnName}"), FormatArgs);
 		}
 
-		return FText::Format( LOCTEXT("PawnBlueprint", "Pawn: {PawnName}"), FText::FromString(PawnClass->GetName()));
+		FormatArgs.Add(TEXT("PawnName"), FText::FromString(PawnClass->GetName()));
+		return FText::Format(LOCTEXT("PawnBlueprint", "Pawn: {PawnName}"), FormatArgs);
 	}
 
 	return LOCTEXT("PawnCreateBlueprint", "Pawn: New...");
@@ -1002,12 +1014,15 @@ FText LevelEditorActionHelpers::GetOpenHUDBlueprintLabel(TWeakPtr< SLevelEditor
 #define LOCTEXT_NAMESPACE "LevelToolBarViewMenu"
 	if(UClass* HUDClass = GetHUDClass(InLevelEditor, bInIsProjectSettings))
 	{
-		if(HUDClass->ClassGeneratedBy)
+		FFormatNamedArguments FormatArgs;
+		if (HUDClass->ClassGeneratedBy)
 		{
-			return FText::Format( LOCTEXT("HUDEditBlueprint", "HUD: Edit {HUDName}"), FText::FromString(HUDClass->ClassGeneratedBy->GetName()));
+			FormatArgs.Add(TEXT("HUDName"), FText::FromString(HUDClass->ClassGeneratedBy->GetName()));
+			return FText::Format(LOCTEXT("HUDEditBlueprint", "HUD: Edit {HUDName}"), FormatArgs);
 		}
 
-		return FText::Format( LOCTEXT("HUDBlueprint", "HUD: {HUDName}"), FText::FromString(HUDClass->GetName()));
+		FormatArgs.Add(TEXT("HUDName"), FText::FromString(HUDClass->GetName()));
+		return FText::Format(LOCTEXT("HUDBlueprint", "HUD: {HUDName}"), FormatArgs);
 	}
 
 	return LOCTEXT("HUDCreateBlueprint", "HUD: New...");
@@ -1073,12 +1088,15 @@ FText LevelEditorActionHelpers::GetOpenPlayerControllerBlueprintLabel(TWeakPtr<
 #define LOCTEXT_NAMESPACE "LevelToolBarViewMenu"
 	if(UClass* PlayerControllerClass = GetPlayerControllerClass(InLevelEditor, bInIsProjectSettings))
 	{
-		if(PlayerControllerClass->ClassGeneratedBy)
+		FFormatNamedArguments FormatArgs;
+		if (PlayerControllerClass->ClassGeneratedBy)
 		{
-			return FText::Format( LOCTEXT("PlayerControllerEditBlueprint", "PlayerController: Edit {PlayerControllerName}"), FText::FromString(PlayerControllerClass->ClassGeneratedBy->GetName()));
+			FormatArgs.Add(TEXT("PlayerControllerName"), FText::FromString(PlayerControllerClass->ClassGeneratedBy->GetName()));
+			return FText::Format(LOCTEXT("PlayerControllerEditBlueprint", "PlayerController: Edit {PlayerControllerName}"), FormatArgs);
 		}
 
-		return FText::Format( LOCTEXT("PlayerControllerBlueprint", "PlayerController: {PlayerControllerName}"), FText::FromString(PlayerControllerClass->GetName()));
+		FormatArgs.Add(TEXT("PlayerControllerName"), FText::FromString(PlayerControllerClass->GetName()));
+		return FText::Format(LOCTEXT("PlayerControllerBlueprint", "PlayerController: {PlayerControllerName}"), FormatArgs);
 	}
 
 	return LOCTEXT("PlayerControllerCreateBlueprint", "PlayerController: New...");
diff --git a/Engine/Source/Editor/MaterialEditor/Private/MaterialEditingLibrary.cpp b/Engine/Source/Editor/MaterialEditor/Private/MaterialEditingLibrary.cpp
--- a/Engine/Source/Editor/MaterialEditor/Private/MaterialEditingLibrary.cpp
+++ b/Engine/Source/Editor/MaterialEditor/Private/MaterialEditingLibrary.cpp
@@ -22,6 +22,7 @@
 #include "MaterialEditor/MaterialEditorInstanceConstant.h"
 #include "Particles/ParticleSystemComponent.h"
 #include "EditorSupportDelegates.h"
+#include "Misc/RuntimeErrors.h"
 
 
 DEFINE_LOG_CATEGORY_STATIC(LogMaterialEditingLibrary, Warning, All);
@@ -426,7 +427,7 @@ bool UMaterialEditingLibrary::ConnectMaterialExpressions(UMaterialExpression* Fr
 
 void UMaterialEditingLibrary::RecompileMaterial(UMaterial* Material)
 {
-	if (Material)
+	if (ensureAsRuntimeWarning(Material != nullptr))
 	{
 		{
 			FMaterialUpdateContext UpdateContext;
@@ -669,4 +670,4 @@ void UMaterialEditingLibrary::UpdateMaterialInstance(UMaterialInstanceConstant*
 		FEditorDelegates::RefreshEditor.Broadcast();
 		FEditorSupportDelegates::RedrawAllViewports.Broadcast();
 	}
-}
+}
\ No newline at end of file
diff --git a/Engine/Source/Editor/UMGEditor/Private/Designer/SDesignerView.cpp b/Engine/Source/Editor/UMGEditor/Private/Designer/SDesignerView.cpp
--- a/Engine/Source/Editor/UMGEditor/Private/Designer/SDesignerView.cpp
+++ b/Engine/Source/Editor/UMGEditor/Private/Designer/SDesignerView.cpp
@@ -1407,6 +1407,10 @@ void SDesignerView::ClearDropPreviews()
 
 		UWidgetBlueprint* BP = GetBlueprint();
 		BP->WidgetTree->RemoveWidget(DropPreview.Widget);
+
+		// Since the widget has been removed from the widget tree, move it into the transient package. Otherwise,
+		// it will remain outered to the widget tree and end up as a property in the BP class layout as a result.
+		DropPreview.Widget->Rename(nullptr, GetTransientPackage());
 	}
 	DropPreviews.Empty();
 }
@@ -2586,6 +2590,18 @@ void SDesignerView::ProcessDropAndAddWidget(const FGeometry& MyGeometry, const F
 			Transaction.Cancel();
 		}
 
+		// Remove widgets tracked by the 'DropPreviews' set. We don't consider them to be transient at this point because they have been inserted into the widget tree hierarchy.
+		for (const FDropPreview& DropPreview : DropPreviews)
+		{
+			DragDropPreviewWidgets.RemoveSwap(DropPreview.Widget);
+		}
+
+		// Move the remaining widgets into the transient package. Otherwise, they will remain outered to the WidgetTree and end up as properties in the BP class layout as a result.
+		for (UWidget* Widget : DragDropPreviewWidgets)
+		{
+			Widget->Rename(nullptr, GetTransientPackage());
+		}
+
 		// If we had preview widgets, we know that we can not be performing a selected widget drag/drop operation. Bail.
 		return;
 	}
diff --git a/Engine/Source/Editor/UMGEditor/Private/WidgetBlueprintCompiler.cpp b/Engine/Source/Editor/UMGEditor/Private/WidgetBlueprintCompiler.cpp
--- a/Engine/Source/Editor/UMGEditor/Private/WidgetBlueprintCompiler.cpp
+++ b/Engine/Source/Editor/UMGEditor/Private/WidgetBlueprintCompiler.cpp
@@ -155,7 +155,7 @@ void FWidgetBlueprintCompiler::CleanAndSanitizeClass(UBlueprintGeneratedClass* C
 	if ( !Blueprint->bIsRegeneratingOnLoad && bIsFullCompile )
 	{
 		UPackage* WidgetTemplatePackage = WidgetBP->GetWidgetTemplatePackage();
-		UUserWidget* OldArchetype = LoadObject<UUserWidget>(WidgetTemplatePackage, TEXT("WidgetArchetype"), nullptr, LOAD_NoWarn);
+		UUserWidget* OldArchetype = FindObjectFast<UUserWidget>(WidgetTemplatePackage, TEXT("WidgetArchetype"));
 		if ( OldArchetype )
 		{
 			const bool bRecompilingOnLoad = Blueprint->bIsRegeneratingOnLoad;
@@ -185,8 +185,7 @@ void FWidgetBlueprintCompiler::CleanAndSanitizeClass(UBlueprintGeneratedClass* C
 	Super::CleanAndSanitizeClass(ClassToClean, InOutOldCDO);
 
 	// Make sure our typed pointer is set
-	check(ClassToClean == NewClass);
-	NewWidgetBlueprintClass = CastChecked<UWidgetBlueprintGeneratedClass>((UObject*)NewClass);
+	check(ClassToClean == NewClass && NewWidgetBlueprintClass == NewClass);
 
 	for ( UWidgetAnimation* Animation : NewWidgetBlueprintClass->Animations )
 	{
@@ -469,9 +468,10 @@ void FWidgetBlueprintCompiler::FinishCompilingClass(UClass* Class)
 	// Make sure that we don't have dueling widget hierarchies
 	if (UWidgetBlueprintGeneratedClass* SuperBPGClass = Cast<UWidgetBlueprintGeneratedClass>(BPGClass->GetSuperClass()))
 	{
-		if (SuperBPGClass->WidgetTree != nullptr)
+		UWidgetBlueprint* SuperBlueprint = Cast<UWidgetBlueprint>(SuperBPGClass->ClassGeneratedBy);
+		if (SuperBlueprint->WidgetTree != nullptr)
 		{
-			if ((SuperBPGClass->WidgetTree->RootWidget != nullptr) && (BPGClass->WidgetTree->RootWidget != nullptr))
+			if ((SuperBlueprint->WidgetTree->RootWidget != nullptr) && (WidgetBlueprint()->WidgetTree->RootWidget != nullptr))
 			{
 				// We both have a widget tree, terrible things will ensue
 				// @todo: nickd - we need to switch this back to a warning in engine, but note for games
@@ -593,6 +593,11 @@ void FWidgetBlueprintCompiler::SpawnNewClass(const FString& NewClassName)
 	NewClass = NewWidgetBlueprintClass;
 }
 
+void FWidgetBlueprintCompiler::OnNewClassSet(UBlueprintGeneratedClass* ClassToUse)
+{
+	NewWidgetBlueprintClass = CastChecked<UWidgetBlueprintGeneratedClass>(ClassToUse);
+}
+
 void FWidgetBlueprintCompiler::PrecompileFunction(FKismetFunctionContext& Context, EInternalCompilerFlags InternalFlags)
 {
 	Super::PrecompileFunction(Context, InternalFlags);
diff --git a/Engine/Source/Editor/UnrealEd/Classes/CookOnTheSide/CookOnTheFlyServer.h b/Engine/Source/Editor/UnrealEd/Classes/CookOnTheSide/CookOnTheFlyServer.h
--- a/Engine/Source/Editor/UnrealEd/Classes/CookOnTheSide/CookOnTheFlyServer.h
+++ b/Engine/Source/Editor/UnrealEd/Classes/CookOnTheSide/CookOnTheFlyServer.h
@@ -44,11 +44,11 @@ enum class ECookInitializationFlags
 	OutputVerboseCookerWarnings =				0x00000800, // output additional cooker warnings about content issues
 	EnablePartialGC =							0x00001000, // mark up with an object flag objects which are in packages which we are about to use or in the middle of using, this means we can gc more often but only gc stuff which we have finished with
 	TestCook =									0x00002000, // test the cooker garbage collection process and cooking (cooker will never end just keep testing).
-	IterateOnHash =								0x00004000, // when using iterative cooking use hashes of original files instead of timestamps
+	//unused =									0x00004000,
 	LogDebugInfo =								0x00008000, // enables additional debug log information
 	IterateSharedBuild =						0x00010000, // iterate from a build in the SharedIterativeBuild directory 
 	IgnoreIniSettingsOutOfDate =				0x00020000, // if the inisettings say the cook is out of date keep using the previously cooked build
-	IterateOnAssetRegistry =					0x00040000, // when using iterative cooking use the asset registry dependencies
+	IgnoreScriptPackagesOutOfDate =				0x00040000, // for incremental cooking, ignore script package changes
 };
 ENUM_CLASS_FLAGS(ECookInitializationFlags);
 
@@ -1292,12 +1292,6 @@ public:
 	*/
 	void MarkPackageDirtyForCooker( UPackage *Package );
 
-	/**
-	* Mark any packages which reference this package as dirty also
-	* causes all packages to be recooked on next request
-	*/
-	void MarkDependentPackagesDirtyForCooker(const FName& PackageName);
-
 	/**
 	* MaybeMarkPackageAsAlreadyLoaded
 	* Mark the package as already loaded if we have already cooked the package for all requested target platforms
@@ -1366,14 +1360,6 @@ private:
 	*/
 	bool GetAllPackageFilenamesFromAssetRegistry( const FString& AssetRegistryPath, TArray<FName>& OutPackageFilenames ) const;
 
-	/**
-	* SaveCookedAssetRegistry
-	* Save an asset registry which contains all the packages which were cooked
-	* this asset registry is used in the editor to display information about package sizes and asset sizes
-	* this asset registry is also used by the cooker for iterative cooking
-	*/
-	bool SaveCookedAssetRegistry(const FString& CookedAssetRegistryFilename, const TArray<FName>& AllCookedPackages, const TArray<FName>& UncookedEditorOnlyPackages, const TArray<FName>& FailedToSavePackages, const FString& PlatformName, const bool Append) const;
-
 	/**
 	* BuildMapDependencyGraph
 	* builds a map of dependencies from maps
@@ -1729,8 +1715,6 @@ private:
 	*/
 	void PopulateCookedPackagesFromDisk( const TArray<ITargetPlatform*>& Platforms );
 
-
-
 	/**
 	* Searches the disk for all the cooked files in the sandbox path provided
 	* Returns a map of the uncooked file path matches to the cooked file path for each package which exists
@@ -1740,13 +1724,6 @@ private:
 	*/
 	void GetAllCookedFiles(TMap<FName, FName>& UncookedPathToCookedPath, const FString& SandboxPath);
 
-	/**
-	* Verify the cooked package list hasn't got any packages which are out of date in it and remove any which are out of date deleting content from disk
-	* 
-	* @params Platforms to process
-	*/
-	void VerifyCookedPackagesAreUptodate( const TArray<ITargetPlatform*>& Platforms );
-
 	/** Generates asset registry */
 	void GenerateAssetRegistry();
 
diff --git a/Engine/Source/Editor/UnrealEd/Classes/Editor/EditorEngine.h b/Engine/Source/Editor/UnrealEd/Classes/Editor/EditorEngine.h
--- a/Engine/Source/Editor/UnrealEd/Classes/Editor/EditorEngine.h
+++ b/Engine/Source/Editor/UnrealEd/Classes/Editor/EditorEngine.h
@@ -2913,8 +2913,8 @@ protected:
 	void HandlePackageReloaded(const EPackageReloadPhase InPackageReloadPhase, FPackageReloadedEvent* InPackageReloadedEvent);
 
 public:
-	/** True if world assets are enabled */
-	static bool IsUsingWorldAssets();
+	DEPRECATED(4.17, "IsUsingWorldAssets is now always true, remove any code that assumes it could be false")
+	static bool IsUsingWorldAssets() { return true; }
 
 private:
 	/** Handler for when any asset is loaded in the editor */
diff --git a/Engine/Source/Editor/UnrealEd/Classes/UserDefinedStructure/UserDefinedStructEditorData.h b/Engine/Source/Editor/UnrealEd/Classes/UserDefinedStructure/UserDefinedStructEditorData.h
--- a/Engine/Source/Editor/UnrealEd/Classes/UserDefinedStructure/UserDefinedStructEditorData.h
+++ b/Engine/Source/Editor/UnrealEd/Classes/UserDefinedStructure/UserDefinedStructEditorData.h
@@ -9,6 +9,8 @@
 #include "UObject/AssetPtr.h"
 #include "EdGraph/EdGraphPin.h"
 #include "UObject/StructOnScope.h"
+#include "EditorUndoClient.h"
+#include "StructureEditorUtils.h"
 #include "UserDefinedStructEditorData.generated.h"
 
 class ITransactionObjectAnnotation;
@@ -119,7 +121,7 @@ public:
 };
 
 UCLASS()
-class UNREALED_API UUserDefinedStructEditorData : public UObject
+class UNREALED_API UUserDefinedStructEditorData : public UObject, public FEditorUndoClient
 {
 	GENERATED_UCLASS_BODY()
 
@@ -152,10 +154,25 @@ public:
 	virtual void PostLoadSubobjects(struct FObjectInstancingGraph* OuterInstanceGraph) override;
 	// End of UObject interface.
 
+	// FEditorUndoClient interface
+	virtual void PostUndo(bool bSuccess) override;
+	virtual void PostRedo(bool bSuccess) override { PostUndo(bSuccess); }
+	// End of FEditorUndoClient interface.
+
+
 	uint32 GenerateUniqueNameIdForMemberVariable();
 	class UUserDefinedStruct* GetOwnerStruct() const;
 
 	const uint8* GetDefaultInstance() const;
 	void RecreateDefaultInstance(FString* OutLog = nullptr);
 	void CleanDefaultInstance();
+
+private:
+
+	// Track the structure change that PostEditUndo undid to pass to FUserDefinedStructureCompilerUtils::CompileStruct
+	FStructureEditorUtils::EStructureEditorChangeInfo CachedStructureChange;
+
+	// Utility function for both PostEditUndo to route through
+	void ConsolidatedPostEditUndo(FStructureEditorUtils::EStructureEditorChangeInfo ActiveChange);
+
 };
