Mirror 3510829, 3514133, and 3514517 in Main. These are fixes for regressions introduced by the compilation manager (3510073)
#rnx

diff --git a/Engine/Source/Editor/Kismet/Private/BlueprintCompilationManager.cpp b/Engine/Source/Editor/Kismet/Private/BlueprintCompilationManager.cpp
--- a/Engine/Source/Editor/Kismet/Private/BlueprintCompilationManager.cpp
+++ b/Engine/Source/Editor/Kismet/Private/BlueprintCompilationManager.cpp
@@ -21,7 +21,7 @@
 #include "KismetCompiler.h"
 #include "ProfilingDebugging/ScopedTimers.h"
 #include "Serialization/ArchiveHasReferences.h"
-#include "Serialization/ArchiveReplaceObjectRef.h"
+#include "Serialization/ArchiveReplaceOrClearExternalReferences.h"
 #include "TickableEditorObject.h"
 #include "UObject/MetaData.h"
 #include "UObject/ReferenceChainSearch.h"
@@ -869,6 +869,11 @@ void FBlueprintCompilationManagerImpl::FlushCompilationQueueImpl(TArray<UObject*
 
 void FBlueprintCompilationManagerImpl::FlushReinstancingQueueImpl()
 {
+	if(GCompilingBlueprint)
+	{
+		return;
+	}
+
 	TGuardValue<bool> GuardTemplateNameFlag(GCompilingBlueprint, true);
 	// we can finalize reinstancing now:
 	if(ClassesToReinstance.Num() == 0)
@@ -1073,7 +1078,7 @@ void FBlueprintCompilationManagerImpl::ReinstanceBatch(TArray<FReinstancingJob>&
 		if(OldCDO)
 		{
 			UObject* NewCDO = CurrentReinstancer->ClassToReinstance->GetDefaultObject(true);
-			FBlueprintCompileReinstancer::CopyPropertiesForUnrelatedObjects(OldCDO, NewCDO);
+			FBlueprintCompileReinstancer::CopyPropertiesForUnrelatedObjects(OldCDO, NewCDO, true);
 
 			if(ReinstancingJob.Compiler.IsValid())
 			{
@@ -1184,7 +1189,7 @@ void FBlueprintCompilationManagerImpl::ReinstanceBatch(TArray<FReinstancingJob>&
 				UObject* NewArchetype = NewObject<UObject>(OriginalOuter, CurrentReinstancer->ClassToReinstance, OriginalName, OriginalFlags & FlagMask);
 
 				// copy old data:
-				FBlueprintCompileReinstancer::CopyPropertiesForUnrelatedObjects(Archetype, NewArchetype);
+				FBlueprintCompileReinstancer::CopyPropertiesForUnrelatedObjects(Archetype, NewArchetype, false);
 
 				OldArchetypeToNewArchetype.Add(Archetype, NewArchetype);
 				// Map old subobjects to new subobjects. This is needed by UMG right now, which allows owning archetypes to reference subobjects
@@ -1244,7 +1249,8 @@ void FBlueprintCompilationManagerImpl::ReinstanceBatch(TArray<FReinstancingJob>&
 
 	for(UObject* ArchetypeReferencer : ArchetypeReferencers)
 	{
-		FArchiveReplaceObjectRef< UObject > ReplaceAr(ArchetypeReferencer, OldArchetypeToNewArchetype, false, false, false);
+		UPackage* NewPackage = ArchetypeReferencer->GetOutermost();
+		FArchiveReplaceOrClearExternalReferences<UObject> ReplaceInCDOAr(ArchetypeReferencer, OldArchetypeToNewArchetype, NewPackage);
 	}
 }
 
@@ -1622,6 +1628,8 @@ UClass* FBlueprintCompilationManagerImpl::FastGenerateSkeletonClass(UBlueprint*
 			}
 		}
 
+		AddFunctionForGraphs(TEXT(""), BPID.Graphs, CurrentFieldStorageLocation, BPTYPE_FunctionLibrary == BP->BlueprintType);
+
 		for (TFieldIterator<UFunction> FunctionIt(InterfaceClass, EFieldIteratorFlags::ExcludeSuper); FunctionIt; ++FunctionIt)
 		{
 			UFunction* Fn = *FunctionIt;
diff --git a/Engine/Source/Editor/UnrealEd/Private/Kismet2/KismetReinstanceUtilities.cpp b/Engine/Source/Editor/UnrealEd/Private/Kismet2/KismetReinstanceUtilities.cpp
--- a/Engine/Source/Editor/UnrealEd/Private/Kismet2/KismetReinstanceUtilities.cpp
+++ b/Engine/Source/Editor/UnrealEd/Private/Kismet2/KismetReinstanceUtilities.cpp
@@ -2143,7 +2143,7 @@ void FBlueprintCompileReinstancer::ReparentChild(UClass* ChildClass)
 	ChildClass->StaticLink(true);
 }
 
-void FBlueprintCompileReinstancer::CopyPropertiesForUnrelatedObjects(UObject* OldObject, UObject* NewObject )
+void FBlueprintCompileReinstancer::CopyPropertiesForUnrelatedObjects(UObject* OldObject, UObject* NewObject, bool bClearExternalReferences)
 {
 	InstancedPropertyUtils::FInstancedPropertyMap InstancedPropertyMap;
 	InstancedPropertyUtils::FArchiveInstancedSubObjCollector  InstancedSubObjCollector(OldObject, InstancedPropertyMap);
@@ -2153,6 +2153,7 @@ void FBlueprintCompileReinstancer::CopyPropertiesForUnrelatedObjects(UObject* Ol
 	Params.bDoDelta = false;
 	Params.bCopyDeprecatedProperties = true;
 	Params.bSkipCompilerGeneratedDefaults = true;
+	Params.bClearReferences = bClearExternalReferences;
 	UEngine::CopyPropertiesForUnrelatedObjects(OldObject, NewObject, Params);
 
 	InstancedPropertyUtils::FArchiveInsertInstancedSubObjects InstancedSubObjSpawner(NewObject, InstancedPropertyMap);
diff --git a/Engine/Source/Editor/UnrealEd/Public/Kismet2/KismetReinstanceUtilities.h b/Engine/Source/Editor/UnrealEd/Public/Kismet2/KismetReinstanceUtilities.h
--- a/Engine/Source/Editor/UnrealEd/Public/Kismet2/KismetReinstanceUtilities.h
+++ b/Engine/Source/Editor/UnrealEd/Public/Kismet2/KismetReinstanceUtilities.h
@@ -195,7 +195,7 @@ protected:
 	/** Determine whether reinstancing actors should preserve the root component of the new actor */
 	virtual bool ShouldPreserveRootComponentOfReinstancedActor() const { return true; }
 
-	static void CopyPropertiesForUnrelatedObjects(UObject* OldObject, UObject* NewObject );
+	static void CopyPropertiesForUnrelatedObjects(UObject* OldObject, UObject* NewObject, bool bClearExternalReferences);
 
 private:
 	/**
diff --git a/Engine/Source/Runtime/Engine/Classes/Engine/Engine.h b/Engine/Source/Runtime/Engine/Classes/Engine/Engine.h
--- a/Engine/Source/Runtime/Engine/Classes/Engine/Engine.h
+++ b/Engine/Source/Runtime/Engine/Classes/Engine/Engine.h
@@ -2349,6 +2349,7 @@ public:
 		/** Skips copying properties with BlueprintCompilerGeneratedDefaults metadata */
 		bool bSkipCompilerGeneratedDefaults;
 		bool bNotifyObjectReplacement;
+		bool bClearReferences;
 
 		FCopyPropertiesForUnrelatedObjectsParams()
 			: bAggressiveDefaultSubobjectReplacement(false)
@@ -2358,6 +2359,7 @@ public:
 			, bPreserveRootComponent(true)
 			, bSkipCompilerGeneratedDefaults(false)
 			, bNotifyObjectReplacement(true)
+			, bClearReferences(true)
 		{}
 	};
 	static void CopyPropertiesForUnrelatedObjects(UObject* OldObject, UObject* NewObject, FCopyPropertiesForUnrelatedObjectsParams Params = FCopyPropertiesForUnrelatedObjectsParams());//bool bAggressiveDefaultSubobjectReplacement = false, bool bDoDelta = true);
diff --git a/Engine/Source/Runtime/Engine/Private/UnrealEngine.cpp b/Engine/Source/Runtime/Engine/Private/UnrealEngine.cpp
--- a/Engine/Source/Runtime/Engine/Private/UnrealEngine.cpp
+++ b/Engine/Source/Runtime/Engine/Private/UnrealEngine.cpp
@@ -11810,15 +11810,18 @@ void UEngine::CopyPropertiesForUnrelatedObjects(UObject* OldObject, UObject* New
 		}
 	});
 
-	// Replace references to old classes and instances on this object with the corresponding new ones
-	UPackage* NewPackage = NewObject->GetOutermost();
-	FArchiveReplaceOrClearExternalReferences<UObject> ReplaceInCDOAr(NewObject, ReferenceReplacementMap, NewPackage);
-
-	// Replace references inside each individual component. This is always required because if something is in ReferenceReplacementMap, the above replace code will skip fixing child properties
-	for (int32 ComponentIndex = 0; ComponentIndex < ComponentsOnNewObject.Num(); ++ComponentIndex)
+	if(Params.bClearReferences)
 	{
-		UObject* NewComponent = ComponentsOnNewObject[ComponentIndex];
-		FArchiveReplaceOrClearExternalReferences<UObject> ReplaceInComponentAr(NewComponent, ReferenceReplacementMap, NewPackage);
+		UPackage* NewPackage = NewObject->GetOutermost();
+		// Replace references to old classes and instances on this object with the corresponding new ones
+		FArchiveReplaceOrClearExternalReferences<UObject> ReplaceInCDOAr(NewObject, ReferenceReplacementMap, NewPackage);
+
+		// Replace references inside each individual component. This is always required because if something is in ReferenceReplacementMap, the above replace code will skip fixing child properties
+		for (int32 ComponentIndex = 0; ComponentIndex < ComponentsOnNewObject.Num(); ++ComponentIndex)
+		{
+			UObject* NewComponent = ComponentsOnNewObject[ComponentIndex];
+			FArchiveReplaceOrClearExternalReferences<UObject> ReplaceInComponentAr(NewComponent, ReferenceReplacementMap, NewPackage);
+		}
 	}
 
 	// Restore the root component reference
