Merge CL# 3538471 from Release-4.17 to Release-Staging-4.17

diff --git a/Engine/Source/Editor/Kismet/Private/BlueprintCompilationManager.cpp b/Engine/Source/Editor/Kismet/Private/BlueprintCompilationManager.cpp
--- a/Engine/Source/Editor/Kismet/Private/BlueprintCompilationManager.cpp
+++ b/Engine/Source/Editor/Kismet/Private/BlueprintCompilationManager.cpp
@@ -1238,9 +1238,12 @@ void FBlueprintCompilationManagerImpl::ReinstanceBatch(TArray<FReinstancingJob>&
 		{
 			ArchetypeReferencers.Add(BP->SkeletonGeneratedClass);
 			ensure(BP->bCachedDependenciesUpToDate);
-			for(TWeakObjectPtr<UBlueprint> Dependendency : BP->CachedDependencies)
+			for(const TWeakObjectPtr<UBlueprint>& Dependency : BP->CachedDependencies)
 			{
-				ArchetypeReferencers.Add(Dependendency.Get());
+				if (UBlueprint* DependencyBP = Dependency.Get())
+				{
+					ArchetypeReferencers.Add(DependencyBP);
+				}
 			}
 		}
 	}
