Copying //UE4/Dev-Framework to //UE4/Dev-Main (Source: //UE4/Dev-Framework @ 3582324)

diff --git a/Engine/Source/Editor/Kismet/Private/BPVariableDragDropAction.cpp b/Engine/Source/Editor/Kismet/Private/BPVariableDragDropAction.cpp
--- a/Engine/Source/Editor/Kismet/Private/BPVariableDragDropAction.cpp
+++ b/Engine/Source/Editor/Kismet/Private/BPVariableDragDropAction.cpp
@@ -19,11 +19,26 @@
 
 FKismetVariableDragDropAction::FKismetVariableDragDropAction()
 	: VariableName(NAME_None)
-	, bControlDrag(false)
-	, bAltDrag(false)
 {
 }
 
+UBlueprint* FKismetVariableDragDropAction::GetSourceBlueprint() const
+{
+	check(VariableSource.IsValid());
+
+	UClass* VariableSourceClass = nullptr;
+	if (VariableSource.Get()->IsA(UClass::StaticClass()))
+	{
+		VariableSourceClass = CastChecked<UClass>(VariableSource.Get());
+	}
+	else
+	{
+		check(VariableSource.Get()->GetOuter());
+		VariableSourceClass = CastChecked<UClass>(VariableSource.Get()->GetOuter());
+	}
+	return UBlueprint::GetBlueprintFromClass(VariableSourceClass);
+}
+
 void FKismetVariableDragDropAction::GetLinksThatWillBreak(	UEdGraphNode* Node, UProperty* NewVariableProperty, 
 						   TArray<class UEdGraphPin*>& OutBroken)
 {
@@ -59,10 +74,6 @@ void FKismetVariableDragDropAction::HoverTargetChanged()
 	FString VariableString = VariableName.ToString();
 
 	// Icon/text to draw on tooltip
-	FSlateColor IconColor = FLinearColor::White;
-	const FSlateBrush* StatusSymbol = FEditorStyle::GetBrush(TEXT("Graph.ConnectorFeedback.Error"));
-	FSlateBrush const* SecondaryIcon = nullptr;
-	FSlateColor SecondaryColor = FLinearColor::Transparent;
 	FText Message = LOCTEXT("InvalidDropTarget", "Invalid drop target!");
 
 	UEdGraphPin* PinUnderCursor = GetHoveredPin();
@@ -77,7 +88,7 @@ void FKismetVariableDragDropAction::HoverTargetChanged()
 		{
 			bBadSchema = true;
 		}
-		else if(!CanVariableBeDropped(VariableProperty, *TheHoveredGraph))
+		else if (!CanVariableBeDropped(VariableProperty, *TheHoveredGraph))
 		{
 			bBadGraph = true;
 		}
@@ -97,24 +108,21 @@ void FKismetVariableDragDropAction::HoverTargetChanged()
 
 	if (bBadSchema)
 	{
-		StatusSymbol = FEditorStyle::GetBrush(TEXT("Graph.ConnectorFeedback.Error"));
-		Message = LOCTEXT("CannotCreateInThisSchema", "Cannot access variables in this type of graph");
+		SetFeedbackMessageError(LOCTEXT("CannotCreateInThisSchema", "Cannot access variables in this type of graph"));
 	}
-	else if(bBadGraph)
+	else if (bBadGraph)
 	{
 		FFormatNamedArguments Args;
 		Args.Add(TEXT("VariableName"), FText::FromString(VariableString));
 		Args.Add(TEXT("Scope"), FText::FromString(TheHoveredGraph->GetName()));
 
-		StatusSymbol = FEditorStyle::GetBrush(TEXT("Graph.ConnectorFeedback.Error"));
-
-		if(IsFromBlueprint(FBlueprintEditorUtils::FindBlueprintForGraph(TheHoveredGraph)) && VariableProperty->GetOuter()->IsA(UFunction::StaticClass()))
+		if (IsFromBlueprint(FBlueprintEditorUtils::FindBlueprintForGraph(TheHoveredGraph)) && VariableProperty->GetOuter()->IsA(UFunction::StaticClass()))
 		{
-			Message = FText::Format( LOCTEXT("IncorrectGraphForLocalVariable_Error", "Cannot place local variable '{VariableName}' in external scope '{Scope}'"), Args);
+			SetFeedbackMessageError(FText::Format( LOCTEXT("IncorrectGraphForLocalVariable_Error", "Cannot place local variable '{VariableName}' in external scope '{Scope}'"), Args));
 		}
 		else
 		{
-			Message = FText::Format( LOCTEXT("IncorrectGraphForVariable_Error", "Cannot place variable '{VariableName}' in external scope '{Scope}'"), Args);
+			SetFeedbackMessageError(FText::Format( LOCTEXT("IncorrectGraphForVariable_Error", "Cannot place variable '{VariableName}' in external scope '{Scope}'"), Args));
 		}
 	}
 	else if (PinUnderCursor)
@@ -123,12 +131,11 @@ void FKismetVariableDragDropAction::HoverTargetChanged()
 		Args.Add(TEXT("PinUnderCursor"), FText::FromString(PinUnderCursor->PinName));
 		Args.Add(TEXT("VariableName"), FText::FromString(VariableString));
 
-		if(CanVariableBeDropped(VariableProperty, *PinUnderCursor->GetOwningNode()->GetGraph()))
+		if (CanVariableBeDropped(VariableProperty, *PinUnderCursor->GetOwningNode()->GetGraph()))
 		{
 			if (PinUnderCursor->bOrphanedPin)
 			{
-				StatusSymbol = FEditorStyle::GetBrush(TEXT("Graph.ConnectorFeedback.Error"));
-				Message = FText::Format(LOCTEXT("OrphanedPin_Error", "Cannot make connection to orphaned pin {PinUnderCursor}"), Args);
+				SetFeedbackMessageError(FText::Format(LOCTEXT("OrphanedPin_Error", "Cannot make connection to orphaned pin {PinUnderCursor}"), Args));
 			}
 			else
 			{
@@ -148,28 +155,15 @@ void FKismetVariableDragDropAction::HoverTargetChanged()
 
 				if (bTypeMatch && bCanWriteIfNeeded)
 				{
-					StatusSymbol = FEditorStyle::GetBrush(TEXT("Graph.ConnectorFeedback.OK"));
-
-					if (bIsRead)
-					{
-						Message = FText::Format(LOCTEXT("MakeThisEqualThat_PinEqualVariableName", "Make {PinUnderCursor} = {VariableName}"), Args);
-					}
-					else
-					{
-						Message = FText::Format(LOCTEXT("MakeThisEqualThat_VariableNameEqualPin", "Make {VariableName} = {PinUnderCursor}"), Args);
-					}
+					SetFeedbackMessageOK(bIsRead ?
+						FText::Format(LOCTEXT("MakeThisEqualThat_PinEqualVariableName", "Make {PinUnderCursor} = {VariableName}"), Args) :
+						FText::Format(LOCTEXT("MakeThisEqualThat_VariableNameEqualPin", "Make {VariableName} = {PinUnderCursor}"), Args));
 				}
 				else
 				{
-					StatusSymbol = FEditorStyle::GetBrush(TEXT("Graph.ConnectorFeedback.Error"));
-					if (!bCanWriteIfNeeded)
-					{
-						Message = FText::Format(LOCTEXT("ReadOnlyVar_Error", "Cannot write to read-only variable '{VariableName}'"), Args);
-					}
-					else
-					{
-						Message = FText::Format(LOCTEXT("NotCompatible_Error", "The type of '{VariableName}' is not compatible with {PinUnderCursor}"), Args);
-					}
+					SetFeedbackMessageError(bCanWriteIfNeeded ?
+						FText::Format(LOCTEXT("NotCompatible_Error", "The type of '{VariableName}' is not compatible with {PinUnderCursor}"), Args) :
+						FText::Format(LOCTEXT("ReadOnlyVar_Error", "Cannot write to read-only variable '{VariableName}'"), Args));
 				}
 			}
 		}
@@ -177,8 +171,7 @@ void FKismetVariableDragDropAction::HoverTargetChanged()
 		{
 			Args.Add(TEXT("Scope"), FText::FromString(PinUnderCursor->GetOwningNode()->GetGraph()->GetName()));
 
-			StatusSymbol = FEditorStyle::GetBrush(TEXT("Graph.ConnectorFeedback.Error"));
-			Message = FText::Format( LOCTEXT("IncorrectGraphForPin_Error", "Cannot place local variable '{VariableName}' in external scope '{Scope}'"), Args);
+			SetFeedbackMessageError(FText::Format( LOCTEXT("IncorrectGraphForPin_Error", "Cannot place local variable '{VariableName}' in external scope '{Scope}'"), Args));
 		}
 	}
 	else if (VarNodeUnderCursor)
@@ -186,7 +179,7 @@ void FKismetVariableDragDropAction::HoverTargetChanged()
 		FFormatNamedArguments Args;
 		Args.Add(TEXT("VariableName"), FText::FromString(VariableString));
 
-		if(CanVariableBeDropped(VariableProperty, *VarNodeUnderCursor->GetGraph()))
+		if (CanVariableBeDropped(VariableProperty, *VarNodeUnderCursor->GetGraph()))
 		{
 			const bool bIsRead = VarNodeUnderCursor->IsA(UK2Node_VariableGet::StaticClass());
 			const UBlueprint* Blueprint = FBlueprintEditorUtils::FindBlueprintForNode(VarNodeUnderCursor);
@@ -196,119 +189,20 @@ void FKismetVariableDragDropAction::HoverTargetChanged()
 			if (bCanWriteIfNeeded)
 			{
 				Args.Add(TEXT("ReadOrWrite"), bIsRead ? LOCTEXT("Read", "read") : LOCTEXT("Write", "write"));
-				if(WillBreakLinks(VarNodeUnderCursor, VariableProperty))
-				{
-					StatusSymbol = FEditorStyle::GetBrush(TEXT("Graph.ConnectorFeedback.OKWarn"));
-					Message = FText::Format( LOCTEXT("ChangeNodeToWarnBreakLinks", "Change node to {ReadOrWrite} '{VariableName}', WARNING this will break links!"), Args);
-				}
-				else
-				{
-					StatusSymbol = FEditorStyle::GetBrush(TEXT("Graph.ConnectorFeedback.OK"));
-					Message = FText::Format( LOCTEXT("ChangeNodeTo", "Change node to {ReadOrWrite} '{VariableName}'"), Args);
-				}
+				SetFeedbackMessageOK(WillBreakLinks(VarNodeUnderCursor, VariableProperty) ?
+					FText::Format(LOCTEXT("ChangeNodeToWarnBreakLinks", "Change node to {ReadOrWrite} '{VariableName}', WARNING this will break links!"), Args) :
+					FText::Format(LOCTEXT("ChangeNodeTo", "Change node to {ReadOrWrite} '{VariableName}'"), Args));
 			}
 			else
 			{
-				StatusSymbol = FEditorStyle::GetBrush(TEXT("Graph.ConnectorFeedback.Error"));
-				Message = FText::Format( LOCTEXT("ReadOnlyVar_Error", "Cannot write to read-only variable '{VariableName}'"), Args);
+				SetFeedbackMessageError(FText::Format( LOCTEXT("ReadOnlyVar_Error", "Cannot write to read-only variable '{VariableName}'"), Args));
 			}
 		}
 		else
 		{
 			Args.Add(TEXT("Scope"), FText::FromString(VarNodeUnderCursor->GetGraph()->GetName()));
 
-			StatusSymbol = FEditorStyle::GetBrush(TEXT("Graph.ConnectorFeedback.Error"));
-			Message = FText::Format( LOCTEXT("IncorrectGraphForNodeReplace_Error", "Cannot replace node with local variable '{VariableName}' in external scope '{Scope}'"), Args);
-		}
-	}
-	else if (!HoveredCategoryName.IsEmpty())
-	{
-		// Find Blueprint that made this class and get category of variable
-		FText Category;
-		UBlueprint* Blueprint;
-		
-		// Find the Blueprint for this property
-		if(Cast<UFunction>(VariableSource.Get()))
-		{
-			Blueprint = UBlueprint::GetBlueprintFromClass(Cast<UClass>(VariableSource->GetOuter()));
-		}
-		else
-		{
-			Blueprint = UBlueprint::GetBlueprintFromClass(Cast<UClass>(VariableSource.Get()));
-		}
-
-		if (Blueprint)
-		{
-			Category = FBlueprintEditorUtils::GetBlueprintVariableCategory(Blueprint, VariableProperty->GetFName(), GetLocalVariableScope() );
-		}
-
-		// See if class is native
-		UClass* OuterClass = Cast<UClass>(VariableProperty->GetOuter());
-		if(OuterClass || Cast<UFunction>(VariableProperty->GetOuter()))
-		{
-			const bool bIsNativeVar = (OuterClass && OuterClass->ClassGeneratedBy == NULL);
-
-			FFormatNamedArguments Args;
-			Args.Add(TEXT("VariableName"), FText::FromString(VariableString));
-			Args.Add(TEXT("HoveredCategoryName"), HoveredCategoryName);
-
-			if (bIsNativeVar)
-			{
-				StatusSymbol = FEditorStyle::GetBrush(TEXT("Graph.ConnectorFeedback.Error"));
-				Message = FText::Format( LOCTEXT("ChangingCatagoryNotThisVar", "Cannot change category for variable '{VariableName}'"), Args );
-			}
-			else if (Category.EqualTo(HoveredCategoryName))
-			{
-				StatusSymbol = FEditorStyle::GetBrush(TEXT("Graph.ConnectorFeedback.Error"));
-				Message = FText::Format( LOCTEXT("ChangingCatagoryAlreadyIn", "Variable '{VariableName}' is already in category '{HoveredCategoryName}'"), Args );
-			}
-			else
-			{
-				StatusSymbol = FEditorStyle::GetBrush(TEXT("Graph.ConnectorFeedback.OK"));
-				Message = FText::Format( LOCTEXT("ChangingCatagoryOk", "Move variable '{VariableName}' to category '{HoveredCategoryName}'"), Args );
-			}
-		}
-	}
-	else if (HoveredAction.IsValid())
-	{
-		if(HoveredAction.Pin()->GetTypeId() == FEdGraphSchemaAction_K2Var::StaticGetTypeId())
-		{
-			FEdGraphSchemaAction_K2Var* VarAction = (FEdGraphSchemaAction_K2Var*)HoveredAction.Pin().Get();
-			FName TargetVarName = VarAction->GetVariableName();
-
-			// Needs to have a valid index to move it (this excludes variables added through other means, like timelines/components
-			int32 MoveVarIndex = INDEX_NONE;
-			int32 TargetVarIndex = INDEX_NONE;
-			if (UBlueprint* Blueprint = UBlueprint::GetBlueprintFromClass(Cast<UClass>(VariableSource.Get())))
-			{
-				MoveVarIndex = FBlueprintEditorUtils::FindNewVariableIndex(Blueprint, VariableName);
-				TargetVarIndex = FBlueprintEditorUtils::FindNewVariableIndex(Blueprint, TargetVarName);
-			}
-
-			FFormatNamedArguments Args;
-			Args.Add(TEXT("VariableName"), FText::FromString(VariableString));
-			Args.Add(TEXT("TargetVarName"), FText::FromName(TargetVarName));
-
-			if(MoveVarIndex == INDEX_NONE)
-			{
-				StatusSymbol = FEditorStyle::GetBrush(TEXT("Graph.ConnectorFeedback.Error"));
-				Message = FText::Format( LOCTEXT("MoveVarDiffClass", "Cannot reorder variable '{VariableName}'."), Args );
-			}
-			else if(TargetVarIndex == INDEX_NONE)
-			{
-				StatusSymbol = FEditorStyle::GetBrush(TEXT("Graph.ConnectorFeedback.Error"));
-				Message = FText::Format( LOCTEXT("MoveVarOther", "Cannot reorder variable '{VariableName}' before '{TargetVarName}'."), Args );
-			}
-			else if(VariableName == TargetVarName)
-			{
-				StatusSymbol = FEditorStyle::GetBrush(TEXT("Graph.ConnectorFeedback.Error"));
-				Message = FText::Format( LOCTEXT("MoveVarYourself", "Cannot reorder variable '{VariableName}' before itself."), Args );
-			}
-			else
-			{
-				StatusSymbol = FEditorStyle::GetBrush(TEXT("Graph.ConnectorFeedback.OK"));
-				Message = FText::Format( LOCTEXT("MoveVarOK", "Reorder variable '{VariableName}' before '{TargetVarName}'"), Args );
-			}
+			SetFeedbackMessageError(FText::Format( LOCTEXT("IncorrectGraphForNodeReplace_Error", "Cannot replace node with local variable '{VariableName}' in external scope '{Scope}'"), Args));
 		}
 	}
 	else if (bAltDrag && !bCanMakeSetter)
@@ -316,17 +210,17 @@ void FKismetVariableDragDropAction::HoverTargetChanged()
 		FFormatNamedArguments Args;
 		Args.Add(TEXT("VariableName"), FText::FromString(VariableString));
 
-		StatusSymbol = FEditorStyle::GetBrush(TEXT("Graph.ConnectorFeedback.Error"));
-		Message = FText::Format(LOCTEXT("CannotPlaceSetter", "Variable '{VariableName}' is readonly, you cannot set this variable."), Args);
+		SetFeedbackMessageError(FText::Format(LOCTEXT("CannotPlaceSetter", "Variable '{VariableName}' is readonly, you cannot set this variable."), Args));
 	}
-	// Draw variable icon
 	else
 	{
-		StatusSymbol = FBlueprintEditor::GetVarIconAndColor(VariableSource.Get(), VariableName, IconColor, SecondaryIcon, SecondaryColor);
-		Message = FText::FromString(VariableString);
+		FMyBlueprintItemDragDropAction::HoverTargetChanged();
 	}
+}
 
-	SetSimpleFeedbackMessage(StatusSymbol, IconColor, Message, SecondaryIcon, SecondaryColor);
+void FKismetVariableDragDropAction::GetDefaultStatusSymbol(const FSlateBrush*& PrimaryBrushOut, FSlateColor& IconColorOut, FSlateBrush const*& SecondaryBrushOut, FSlateColor& SecondaryColorOut) const
+{
+	PrimaryBrushOut = FBlueprintEditor::GetVarIconAndColor(VariableSource.Get(), VariableName, IconColorOut, SecondaryBrushOut, SecondaryColorOut);
 }
 
 FReply FKismetVariableDragDropAction::DroppedOnPin(FVector2D ScreenPosition, FVector2D GraphPosition)
@@ -514,7 +408,7 @@ FReply FKismetVariableDragDropAction::DroppedOnPanel( const TSharedRef< SWidget
 				MenuBuilder.BeginSection("BPVariableDroppedOn", VariableNameText );
 
 				MenuBuilder.AddMenuEntry(
-					LOCTEXT("CreateGetVariable", "Get"),
+					FText::Format( LOCTEXT("CreateGetVariable", "Get {0}"), VariableNameText ),
 					FText::Format( LOCTEXT("CreateVariableGetterToolTip", "Create Getter for variable '{0}'\n(Ctrl-drag to automatically create a getter)"), VariableNameText ),
 					FSlateIcon(),
 					FUIAction(
@@ -522,7 +416,7 @@ FReply FKismetVariableDragDropAction::DroppedOnPanel( const TSharedRef< SWidget
 					);
 
 				MenuBuilder.AddMenuEntry(
-					LOCTEXT("CreateSetVariable", "Set"),
+					FText::Format( LOCTEXT("CreateSetVariable", "Set {0}"), VariableNameText ),
 					FText::Format( LOCTEXT("CreateVariableSetterToolTip", "Create Setter for variable '{0}'\n(Alt-drag to automatically create a setter)"), VariableNameText ),
 					FSlateIcon(),
 					FUIAction(
@@ -548,70 +442,6 @@ FReply FKismetVariableDragDropAction::DroppedOnPanel( const TSharedRef< SWidget
 	return FReply::Handled();
 }
 
-
-
-FReply FKismetVariableDragDropAction::DroppedOnAction(TSharedRef<FEdGraphSchemaAction> Action)
-{
-	if(Action->GetTypeId() == FEdGraphSchemaAction_K2Var::StaticGetTypeId())
-	{
-		FEdGraphSchemaAction_K2Var* VarAction = (FEdGraphSchemaAction_K2Var*)&Action.Get();
-
-		// Only let you drag and drop if variables are from same BP class, and not onto itself
-		UBlueprint* BP = UBlueprint::GetBlueprintFromClass(Cast<UClass>(VariableSource.Get()));
-		FName TargetVarName = VarAction->GetVariableName();
-		if( (BP != NULL) && 
-			(VariableName != TargetVarName) && 
-			(VariableSource == VarAction->GetVariableClass()) )
-		{
-			bool bMoved = FBlueprintEditorUtils::MoveVariableBeforeVariable(BP, VariableName, TargetVarName, true);
-			// If we moved successfully
-			if(bMoved)
-			{
-				// Change category of var to match the one we dragged on to as well
-				FText MovedVarCategory = FBlueprintEditorUtils::GetBlueprintVariableCategory(BP, VariableName, GetLocalVariableScope());
-				FText TargetVarCategory = FBlueprintEditorUtils::GetBlueprintVariableCategory(BP, TargetVarName, GetLocalVariableScope());
-				if(!MovedVarCategory.EqualTo(TargetVarCategory))
-				{
-					FBlueprintEditorUtils::SetBlueprintVariableCategory(BP, VariableName, GetLocalVariableScope(), TargetVarCategory, true);
-				}
-
-				// Update Blueprint after changes so they reflect in My Blueprint tab.
-				FBlueprintEditorUtils::MarkBlueprintAsStructurallyModified(BP);
-			}
-		}
-
-		return FReply::Handled();
-	}
-	return FReply::Unhandled();
-}
-
-FReply FKismetVariableDragDropAction::DroppedOnCategory(FText Category)
-{
-	UE_LOG(LogTemp, Log, TEXT("Dropped %s on Category %s"), *VariableName.ToString(), *Category.ToString());
-
-	UBlueprint* BP = NULL;
-	if(VariableSource.Get()->IsA(UFunction::StaticClass()))
-	{
-		BP = UBlueprint::GetBlueprintFromClass(Cast<UClass>(VariableSource.Get()->GetOuter()));
-	}
-	else
-	{
-		BP = UBlueprint::GetBlueprintFromClass(Cast<UClass>(VariableSource.Get()));
-	}
-
-	if(BP != NULL)
-	{
-		// Check this is actually a different category
-		FText CurrentCategory = FBlueprintEditorUtils::GetBlueprintVariableCategory(BP, VariableName, GetLocalVariableScope());
-		if(!Category.EqualTo(CurrentCategory))
-		{
-			FBlueprintEditorUtils::SetBlueprintVariableCategory(BP, VariableName, GetLocalVariableScope(), Category, false);
-		}
-	}
-
-	return FReply::Handled();
-}
-
 bool FKismetVariableDragDropAction::CanVariableBeDropped(const UProperty* InVariableProperty, const UEdGraph& InGraph) const
 {
 	bool bCanVariableBeDropped = false;
diff --git a/Engine/Source/Editor/Kismet/Private/BPVariableDragDropAction.h b/Engine/Source/Editor/Kismet/Private/BPVariableDragDropAction.h
--- a/Engine/Source/Editor/Kismet/Private/BPVariableDragDropAction.h
+++ b/Engine/Source/Editor/Kismet/Private/BPVariableDragDropAction.h
@@ -10,53 +10,34 @@
 #include "Engine/Blueprint.h"
 #include "BlueprintEditor.h"
 #include "GraphEditorDragDropAction.h"
+#include "MyBlueprintItemDragDropAction.h"
 
 class UEdGraph;
 
 /** DragDropAction class for dropping a Variable onto a graph */
-class KISMET_API FKismetVariableDragDropAction : public FGraphEditorDragDropAction
+class KISMET_API FKismetVariableDragDropAction : public FMyBlueprintItemDragDropAction
 {
 public:
-	DRAG_DROP_OPERATOR_TYPE(FKismetVariableDragDropAction, FGraphEditorDragDropAction)
+	DRAG_DROP_OPERATOR_TYPE(FKismetVariableDragDropAction, FMyBlueprintItemDragDropAction)
 
 	// FGraphEditorDragDropAction interface
 	virtual void HoverTargetChanged() override;
 	virtual FReply DroppedOnPin(FVector2D ScreenPosition, FVector2D GraphPosition) override;
 	virtual FReply DroppedOnNode(FVector2D ScreenPosition, FVector2D GraphPosition) override;
 	virtual FReply DroppedOnPanel(const TSharedRef< class SWidget >& Panel, FVector2D ScreenPosition, FVector2D GraphPosition, UEdGraph& Graph) override;
-	virtual FReply DroppedOnAction(TSharedRef<struct FEdGraphSchemaAction> Action) override;
-	virtual FReply DroppedOnCategory(FText Category) override;
 	// End of FGraphEditorDragDropAction
 
-	static TSharedRef<FKismetVariableDragDropAction> New(FName InVariableName, UStruct* InVariableSource, FNodeCreationAnalytic AnalyticCallback)
+	static TSharedRef<FKismetVariableDragDropAction> New(TSharedPtr<FEdGraphSchemaAction> InAction, FName InVariableName, UStruct* InVariableSource, FNodeCreationAnalytic AnalyticCallback)
 	{
 		TSharedRef<FKismetVariableDragDropAction> Operation = MakeShareable(new FKismetVariableDragDropAction);
 		Operation->VariableName = InVariableName;
 		Operation->VariableSource = InVariableSource;
 		Operation->AnalyticCallback = AnalyticCallback;
+		Operation->SourceAction = InAction;
 		Operation->Construct();
 		return Operation;
 	}
 
-	/** Helper method to see if we're dragging in the same blueprint */
-	bool IsFromBlueprint(class UBlueprint* InBlueprint) const
-	{ 
-		check(VariableSource.IsValid());
-
-		UClass* VariableSourceClass = NULL;
-		if(VariableSource.Get()->IsA(UClass::StaticClass()))
-		{
-			VariableSourceClass = CastChecked<UClass>(VariableSource.Get());
-		}
-		else
-		{
-			check(VariableSource.Get()->GetOuter());
-			VariableSourceClass = CastChecked<UClass>(VariableSource.Get()->GetOuter());
-		}
-		UBlueprint* Blueprint = UBlueprint::GetBlueprintFromClass(VariableSourceClass);
-		return (InBlueprint == Blueprint);
-	}
-
 	UProperty* GetVariableProperty()
 	{
 		if (VariableSource.IsValid() && VariableName != NAME_None)
@@ -66,12 +47,6 @@ public:
 		return nullptr;
 	}
 
-	/** Set if operation is modified by alt */
-	void SetAltDrag(bool InIsAltDrag) {	bAltDrag = InIsAltDrag;}
-
-	/** Set if operation is modified by the ctrl key */
-	void SetCtrlDrag(bool InIsCtrlDrag) {bControlDrag = InIsCtrlDrag;}
-
 protected:
 	 /** Construct a FKismetVariableDragDropAction */
 	FKismetVariableDragDropAction();
@@ -85,6 +60,14 @@ protected:
 		TWeakObjectPtr<UStruct> VariableSource;
 	};
 
+	// FGraphSchemaActionDragDropAction interface
+	virtual void GetDefaultStatusSymbol(const FSlateBrush*& PrimaryBrushOut, FSlateColor& IconColorOut, FSlateBrush const*& SecondaryBrushOut, FSlateColor& SecondaryColorOut) const override;
+	// End of FGraphSchemaActionDragDropAction interface
+
+	// FMyBlueprintItemDragDropAction interface
+	virtual UBlueprint* GetSourceBlueprint() const override;
+	// End of FMyBlueprintItemDragDropAction interface
+
 	/** Called when user selects to create a Getter for the variable */
 	static void MakeGetter(FNodeConstructionParams InParams);
 	/** Called when user selects to create a Setter for the variable */
@@ -126,10 +109,4 @@ protected:
 	FName VariableName;
 	/** Scope this variable belongs to */
 	TWeakObjectPtr<UStruct> VariableSource;
-	/** Was ctrl held down at start of drag */
-	bool bControlDrag;
-	/** Was alt held down at the start of drag */
-	bool bAltDrag;
-	/** Analytic delegate to track node creation */
-	FNodeCreationAnalytic AnalyticCallback;
 };
diff --git a/Engine/Source/Editor/Kismet/Private/BlueprintActionMenuBuilder.cpp b/Engine/Source/Editor/Kismet/Private/BlueprintActionMenuBuilder.cpp
--- a/Engine/Source/Editor/Kismet/Private/BlueprintActionMenuBuilder.cpp
+++ b/Engine/Source/Editor/Kismet/Private/BlueprintActionMenuBuilder.cpp
@@ -385,7 +385,7 @@ void FBlueprintActionMenuBuilderImpl::FMenuSectionDefinition::AddBoundMenuItems(
 
 					if (Flags & FBlueprintActionMenuBuilder::FlattenCategoryHierarcy)
 					{
-						LastMadeMenuItem->UpdateCategory( ItemFactory.RootCategory );
+						LastMadeMenuItem->CosmeticUpdateCategory( ItemFactory.RootCategory );
 					}
 				}
 				else
@@ -450,7 +450,7 @@ FBlueprintActionMenuBuilderImpl::MenuItemList FBlueprintActionMenuBuilderImpl::F
 		UnBoundMenuEntry = ItemFactory.MakeActionMenuItem(EditorContext, DatabaseAction);
 		if (Flags & FBlueprintActionMenuBuilder::FlattenCategoryHierarcy)
 		{
-			UnBoundMenuEntry->UpdateCategory( ItemFactory.RootCategory );
+			UnBoundMenuEntry->CosmeticUpdateCategory( ItemFactory.RootCategory );
 		}
 	}
 
diff --git a/Engine/Source/Editor/Kismet/Private/BlueprintActionMenuUtils.cpp b/Engine/Source/Editor/Kismet/Private/BlueprintActionMenuUtils.cpp
--- a/Engine/Source/Editor/Kismet/Private/BlueprintActionMenuUtils.cpp
+++ b/Engine/Source/Editor/Kismet/Private/BlueprintActionMenuUtils.cpp
@@ -751,7 +751,6 @@ const UK2Node* FBlueprintActionMenuUtils::ExtractNodeTemplateFromAction(TSharedP
 		else if (ActionId == FEdGraphSchemaAction_K2NewNode::StaticGetTypeId() ||
 			ActionId == FEdGraphSchemaAction_K2AssignDelegate::StaticGetTypeId() ||
 			ActionId == FEdGraphSchemaAction_K2AddComponent::StaticGetTypeId() ||
-			ActionId == FEdGraphSchemaAction_K2AddTimeline::StaticGetTypeId() ||
 			ActionId == FEdGraphSchemaAction_K2AddCustomEvent::StaticGetTypeId() ||
 			ActionId == FEdGraphSchemaAction_K2AddCallOnActor::StaticGetTypeId() ||
 			ActionId == FEdGraphSchemaAction_K2TargetNode::StaticGetTypeId() ||
diff --git a/Engine/Source/Editor/Kismet/Private/BlueprintCompilationManager.cpp b/Engine/Source/Editor/Kismet/Private/BlueprintCompilationManager.cpp
--- a/Engine/Source/Editor/Kismet/Private/BlueprintCompilationManager.cpp
+++ b/Engine/Source/Editor/Kismet/Private/BlueprintCompilationManager.cpp
@@ -1393,18 +1393,23 @@ UClass* FBlueprintCompilationManagerImpl::FastGenerateSkeletonClass(UBlueprint*
 			{
 				// we only want the *skeleton* version of the function:
 				UClass* InterfaceClass = BPID.Interface;
-				if(UBlueprint* Owner = Cast<UBlueprint>(InterfaceClass->ClassGeneratedBy))
+				// We need to null check because FBlueprintEditorUtils::ConformImplementedInterfaces won't run until 
+				// after the skeleton classes have been generated:
+				if(InterfaceClass)
 				{
-					if( ensure(Owner->SkeletonGeneratedClass) )
+					if(UBlueprint* Owner = Cast<UBlueprint>(InterfaceClass->ClassGeneratedBy))
 					{
-						InterfaceClass = Owner->SkeletonGeneratedClass;
+						if( ensure(Owner->SkeletonGeneratedClass) )
+						{
+							InterfaceClass = Owner->SkeletonGeneratedClass;
+						}
 					}
-				}
 
-				if(UFunction* ParentInterfaceFn = InterfaceClass->FindFunctionByName(NewFunction->GetFName()))
-				{
-					ParentFn = ParentInterfaceFn;
-					break;
+					if(UFunction* ParentInterfaceFn = InterfaceClass->FindFunctionByName(NewFunction->GetFName()))
+					{
+						ParentFn = ParentInterfaceFn;
+						break;
+					}
 				}
 			}
 		}
@@ -1685,34 +1690,39 @@ UClass* FBlueprintCompilationManagerImpl::FastGenerateSkeletonClass(UBlueprint*
 	for(const FBPInterfaceDescription& BPID : BP->ImplementedInterfaces)
 	{
 		UClass* InterfaceClass = BPID.Interface;
-		if(UBlueprint* Owner = Cast<UBlueprint>(InterfaceClass->ClassGeneratedBy))
+		// Again, once the skeleton has been created we will purge null ImplementedInterfaces entries,
+		// but not yet:
+		if(InterfaceClass)
 		{
-			if( ensure(Owner->SkeletonGeneratedClass) )
+			if(UBlueprint* Owner = Cast<UBlueprint>(InterfaceClass->ClassGeneratedBy))
 			{
-				InterfaceClass = Owner->SkeletonGeneratedClass;
+				if( ensure(Owner->SkeletonGeneratedClass) )
+				{
+					InterfaceClass = Owner->SkeletonGeneratedClass;
+				}
 			}
-		}
 
-		AddFunctionForGraphs(TEXT(""), BPID.Graphs, CurrentFieldStorageLocation, BPTYPE_FunctionLibrary == BP->BlueprintType);
+			AddFunctionForGraphs(TEXT(""), BPID.Graphs, CurrentFieldStorageLocation, BPTYPE_FunctionLibrary == BP->BlueprintType);
 
-		for (TFieldIterator<UFunction> FunctionIt(InterfaceClass, EFieldIteratorFlags::ExcludeSuper); FunctionIt; ++FunctionIt)
-		{
-			UFunction* Fn = *FunctionIt;
+			for (TFieldIterator<UFunction> FunctionIt(InterfaceClass, EFieldIteratorFlags::ExcludeSuper); FunctionIt; ++FunctionIt)
+			{
+				UFunction* Fn = *FunctionIt;
 			
-			UField** CurrentParamStorageLocation = nullptr;
-
-			// Note that MakeFunction will early out if the function was created above:
-			MakeFunction(
-				Fn->GetFName(), 
-				CurrentFieldStorageLocation, 
-				CurrentParamStorageLocation, 
-				Fn->FunctionFlags & ~FUNC_Native, 
-				TArray<UK2Node_FunctionResult*>(), 
-				TArray<UEdGraphPin*>(),
-				false, 
-				false,
-				nullptr
-			);
+				UField** CurrentParamStorageLocation = nullptr;
+
+				// Note that MakeFunction will early out if the function was created above:
+				MakeFunction(
+					Fn->GetFName(), 
+					CurrentFieldStorageLocation, 
+					CurrentParamStorageLocation, 
+					Fn->FunctionFlags & ~FUNC_Native, 
+					TArray<UK2Node_FunctionResult*>(), 
+					TArray<UEdGraphPin*>(),
+					false, 
+					false,
+					nullptr
+				);
+			}
 		}
 	}
 
diff --git a/Engine/Source/Editor/Kismet/Private/BlueprintDetailsCustomization.cpp b/Engine/Source/Editor/Kismet/Private/BlueprintDetailsCustomization.cpp
--- a/Engine/Source/Editor/Kismet/Private/BlueprintDetailsCustomization.cpp
+++ b/Engine/Source/Editor/Kismet/Private/BlueprintDetailsCustomization.cpp
@@ -2002,7 +2002,7 @@ bool FBlueprintVarActionDetails::IsConfigCheckBoxEnabled() const
 		if (UProperty* VariableProperty = CachedVariableProperty.Get())
 		{
 			// meant to match up with UHT's FPropertyBase::IsObject(), which it uses to block object properties from being marked with CPF_Config
-			bEnabled = VariableProperty->IsA<UClassProperty>() || VariableProperty->IsA<UAssetClassProperty>() || VariableProperty->IsA<UAssetObjectProperty>() ||
+			bEnabled = VariableProperty->IsA<UClassProperty>() || VariableProperty->IsA<USoftClassProperty>() || VariableProperty->IsA<USoftObjectProperty>() ||
 				(!VariableProperty->IsA<UObjectPropertyBase>() && !VariableProperty->IsA<UInterfaceProperty>());
 		}
 	}
@@ -4241,7 +4241,7 @@ UFunction* FBlueprintGraphActionDetails::FindFunction() const
 
 FKismetUserDeclaredFunctionMetadata* FBlueprintGraphActionDetails::GetMetadataBlock() const
 {
-	UK2Node_EditablePinBase * FunctionEntryNode = FunctionEntryNodePtr.Get();
+	UK2Node_EditablePinBase* FunctionEntryNode = FunctionEntryNodePtr.Get();
 	if (UK2Node_FunctionEntry* TypedEntryNode = Cast<UK2Node_FunctionEntry>(FunctionEntryNode))
 	{
 		return &(TypedEntryNode->MetaData);
@@ -4302,29 +4302,11 @@ void FBlueprintGraphActionDetails::OnCategoryTextCommitted(const FText& NewText,
 {
 	if (InTextCommit == ETextCommit::OnEnter || InTextCommit == ETextCommit::OnUserMovedFocus)
 	{
-		if (FKismetUserDeclaredFunctionMetadata* Metadata = GetMetadataBlock())
-		{
-			// Remove excess whitespace and prevent categories with just spaces
-			FText CategoryName = FText::TrimPrecedingAndTrailing(NewText);
-
-			if(CategoryName.IsEmpty())
-			{
-				const UEdGraphSchema_K2* K2Schema = GetDefault<UEdGraphSchema_K2>();
-				Metadata->Category = K2Schema->VR_DefaultCategory;
-			}
-			else
-			{
-				Metadata->Category = CategoryName;
-			}
+		// Remove excess whitespace and prevent categories with just spaces
+		FText CategoryName = FText::TrimPrecedingAndTrailing(NewText);
 
-			if (UFunction* Function = FindFunction())
-			{
-				Function->Modify();
-				Function->SetMetaData(FBlueprintMetadata::MD_FunctionCategory, *CategoryName.ToString());
-			}
-			MyBlueprint.Pin()->Refresh();
-			FBlueprintEditorUtils::MarkBlueprintAsModified(GetBlueprintObj());
-		}
+		FBlueprintEditorUtils::SetBlueprintFunctionOrMacroCategory(GetGraph(), CategoryName);
+		MyBlueprint.Pin()->Refresh();
 	}
 }
 
@@ -4334,14 +4316,8 @@ void FBlueprintGraphActionDetails::OnCategorySelectionChanged( TSharedPtr<FText>
 	{
 		if (FKismetUserDeclaredFunctionMetadata* Metadata = GetMetadataBlock())
 		{
-			Metadata->Category = *ProposedSelection.Get();
-			if (UFunction* Function = FindFunction())
-			{
-				Function->Modify();
-				Function->SetMetaData(FBlueprintMetadata::MD_FunctionCategory, *ProposedSelection.Get()->ToString());
-			}
+			FBlueprintEditorUtils::SetBlueprintFunctionOrMacroCategory(GetGraph(), *ProposedSelection.Get());
 			MyBlueprint.Pin()->Refresh();
-			FBlueprintEditorUtils::MarkBlueprintAsModified(GetBlueprintObj());
 
 			CategoryListView.Pin()->ClearSelection();
 			CategoryComboButton.Pin()->SetIsOpen(false);
@@ -5766,7 +5742,7 @@ void FBlueprintComponentDetails::PopulateVariableCategories()
 	check(BlueprintObj);
 	check(BlueprintObj->SkeletonGeneratedClass);
 
-	TArray<FName> VisibleVariables;
+	TSet<FName> VisibleVariables;
 	for (TFieldIterator<UProperty> PropertyIt(BlueprintObj->SkeletonGeneratedClass, EFieldIteratorFlags::IncludeSuper); PropertyIt; ++PropertyIt)
 	{
 		UProperty* Property = *PropertyIt;
@@ -5781,9 +5757,9 @@ void FBlueprintComponentDetails::PopulateVariableCategories()
 
 	VariableCategorySource.Empty();
 	VariableCategorySource.Add(MakeShareable(new FText(LOCTEXT("Default", "Default"))));
-	for (int32 i = 0; i < VisibleVariables.Num(); ++i)
+	for (const FName& VariableName : VisibleVariables)
 	{
-		FText Category = FBlueprintEditorUtils::GetBlueprintVariableCategory(BlueprintObj, VisibleVariables[i], nullptr);
+		FText Category = FBlueprintEditorUtils::GetBlueprintVariableCategory(BlueprintObj, VariableName, nullptr);
 		if (!Category.IsEmpty() && !Category.EqualTo(FText::FromString(BlueprintObj->GetName())))
 		{
 			bool bNewCategory = true;
diff --git a/Engine/Source/Editor/Kismet/Private/BlueprintDragDropMenuItem.cpp b/Engine/Source/Editor/Kismet/Private/BlueprintDragDropMenuItem.cpp
--- a/Engine/Source/Editor/Kismet/Private/BlueprintDragDropMenuItem.cpp
+++ b/Engine/Source/Editor/Kismet/Private/BlueprintDragDropMenuItem.cpp
@@ -113,7 +113,7 @@ TSharedPtr<FDragDropOperation> FBlueprintDragDropMenuItem::OnDragged(FNodeCreati
 		if (UMulticastDelegateProperty const* Property = DelegateSpawner->GetDelegateProperty())
 		{
 			UStruct* const PropertyOwner = CastChecked<UStruct>(Property->GetOuterUField());
-			TSharedRef<FDragDropOperation> DragDropOpRef = FKismetDelegateDragDropAction::New(Property->GetFName(), PropertyOwner, AnalyticsDelegate);
+			TSharedRef<FDragDropOperation> DragDropOpRef = FKismetDelegateDragDropAction::New(nullptr, Property->GetFName(), PropertyOwner, AnalyticsDelegate);
 			DragDropAction = TSharedPtr<FDragDropOperation>(DragDropOpRef);
 		}
 	}
@@ -122,7 +122,7 @@ TSharedPtr<FDragDropOperation> FBlueprintDragDropMenuItem::OnDragged(FNodeCreati
 		if (UProperty const* Property = VariableSpawner->GetVarProperty())
 		{
 			UStruct* const PropertyOwner = CastChecked<UStruct>(Property->GetOuterUField());
-			TSharedRef<FDragDropOperation> DragDropOpRef = FKismetVariableDragDropAction::New(Property->GetFName(), PropertyOwner, AnalyticsDelegate);
+			TSharedRef<FDragDropOperation> DragDropOpRef = FKismetVariableDragDropAction::New(nullptr, Property->GetFName(), PropertyOwner, AnalyticsDelegate);
 			DragDropAction = TSharedPtr<FDragDropOperation>(DragDropOpRef);
 		}
 		// @TODO: handle local variables as well
diff --git a/Engine/Source/Editor/Kismet/Private/BlueprintEditor.cpp b/Engine/Source/Editor/Kismet/Private/BlueprintEditor.cpp
--- a/Engine/Source/Editor/Kismet/Private/BlueprintEditor.cpp
+++ b/Engine/Source/Editor/Kismet/Private/BlueprintEditor.cpp
@@ -114,7 +114,6 @@
 
 // Debugging
 #include "Debugging/SKismetDebuggingView.h"
-#include "Profiler/SBlueprintProfilerView.h"
 #include "Debugging/KismetDebugCommands.h"
 // End of debugging
 
@@ -151,10 +150,6 @@
 #include "Widgets/Notifications/SNotificationList.h"
 #include "NativeCodeGenerationTool.h"
 
-
-// Blueprint Profiler
-#include "BlueprintProfilerModule.h"
-
 #define LOCTEXT_NAMESPACE "BlueprintEditor"
 
 /////////////////////////////////////////////////////
@@ -577,10 +572,6 @@ bool FBlueprintEditor::IsInAScriptingMode() const
 
 bool FBlueprintEditor::OnRequestClose()
 {
-	if (IsProfilerAvailable())
-	{
-		TabManager->InvokeTab(FBlueprintEditorTabs::BlueprintProfilerID)->RequestCloseTab();
-	}
 	bEditorMarkedAsClosed = true;
 	return FWorkflowCentricApplication::OnRequestClose();
 }
@@ -623,6 +614,12 @@ FGraphPanelSelectionSet FBlueprintEditor::GetSelectedNodes() const
 	return CurrentSelection;
 }
 
+UEdGraphNode* FBlueprintEditor::GetSingleSelectedNode() const
+{
+	const FGraphPanelSelectionSet SelectedNodes = GetSelectedNodes();
+	return (SelectedNodes.Num() == 1) ? Cast<UEdGraphNode>(*SelectedNodes.CreateConstIterator()) : nullptr;
+}
+
 void FBlueprintEditor::AnalyticsTrackNodeEvent( UBlueprint* Blueprint, UEdGraphNode *GraphNode, bool bNodeDelete ) const
 {
 	if( Blueprint && GraphNode && FEngineAnalytics::IsAvailable() )
@@ -1267,25 +1264,11 @@ TSharedRef<SGraphEditor> FBlueprintEditor::CreateGraphEditorWidget(TSharedRef<FT
 				FCanExecuteAction::CreateSP( this, &FBlueprintEditor::CanFindReferences )
 				);
 
-			GraphEditorCommands->MapAction( FGraphEditorCommands::Get().GotoNativeFunctionDefinition,
-				FExecuteAction::CreateSP( this, &FBlueprintEditor::GotoNativeFunctionDefinition ),
-				FCanExecuteAction::CreateSP(this, &FBlueprintEditor::IsSelectionNativeFunction),
-				FIsActionChecked(),
-				FIsActionButtonVisible::CreateSP( this, &FBlueprintEditor::IsNativeCodeBrowsingAvailable )
-				);
-
 			GraphEditorCommands->MapAction( FGraphEditorCommands::Get().GoToDefinition,
 				FExecuteAction::CreateSP( this, &FBlueprintEditor::OnGoToDefinition ),
 				FCanExecuteAction::CreateSP( this, &FBlueprintEditor::CanGoToDefinition )
 				);
 
-			GraphEditorCommands->MapAction( FGraphEditorCommands::Get().GotoNativeVariableDefinition,
-				FExecuteAction::CreateSP( this, &FBlueprintEditor::GotoNativeVariableDefinition ),
-				FCanExecuteAction::CreateSP(this, &FBlueprintEditor::IsSelectionNativeVariable),
-				FIsActionChecked(),
-				FIsActionButtonVisible::CreateSP( this, &FBlueprintEditor::IsNativeCodeBrowsingAvailable )
-				);
-
 			GraphEditorCommands->MapAction(FGraphEditorCommands::Get().GoToDocumentation,
 				FExecuteAction::CreateSP(this, &FBlueprintEditor::OnGoToDocumentation),
 				FCanExecuteAction::CreateSP(this, &FBlueprintEditor::CanGoToDocumentation)
@@ -1429,7 +1412,7 @@ FGraphAppearanceInfo FBlueprintEditor::GetGraphAppearance(UEdGraph* InGraph) con
 }
 
 // Open the editor for a given graph
-void FBlueprintEditor::OnChangeBreadCrumbGraph(class UEdGraph* InGraph)
+void FBlueprintEditor::OnChangeBreadCrumbGraph(UEdGraph* InGraph)
 {
 	if (InGraph && FocusedGraphEdPtr.IsValid())
 	{
@@ -1446,10 +1429,8 @@ FBlueprintEditor::FBlueprintEditor()
 	, bIsActionMenuContextSensitive(true)
 	, CurrentUISelection(NAME_None)
 	, bEditorMarkedAsClosed(false)
-	, bCodeBasedProject(false)
 	, HasOpenActionMenu(nullptr)
 	, InstructionsFadeCountdown(0.f)
-	, bBlueprintHasInstrumentation(false)
 {
 	AnalyticsStats.GraphActionMenusNonCtxtSensitiveExecCount = 0;
 	AnalyticsStats.GraphActionMenusCtxtSensitiveExecCount = 0;
@@ -1754,34 +1735,9 @@ void FBlueprintEditor::InitBlueprintEditor(
 
 	RegisterApplicationModes(InBlueprints, bShouldOpenInDefaultsMode, bNewlyCreated);
 
-	// Cache the project type ( Blueprint or Code Based )
-	if( FPaths::IsProjectFilePathSet() )
-	{
-		FProjectStatus ProjectStatus;
-		if( IProjectManager::Get().QueryStatusForProject( FPaths::GetProjectFilePath(), ProjectStatus ) && ProjectStatus.bCodeBasedProject )
-		{
-			bCodeBasedProject = true;
-		}
-	}
-
 	// Post-layout initialization
 	PostLayoutBlueprintEditorInitialization();
 
-	// Ensure the profiler UI respects the current blueprint state if the blueprint has instrumentation.
-	if (IsProfilerAvailable())
-	{
-		UBlueprint* MyBlueprint = GetBlueprintObj();
-		bBlueprintHasInstrumentation = MyBlueprint && MyBlueprint->GeneratedClass ? MyBlueprint->GeneratedClass->HasInstrumentation() : false;
-		if (bBlueprintHasInstrumentation)
-		{
-			TabManager->InvokeTab(FBlueprintEditorTabs::BlueprintProfilerID);
-		}
-		else
-		{
-			TabManager->InvokeTab(FBlueprintEditorTabs::BlueprintProfilerID)->RequestCloseTab();
-		}
-	}
-
 	// Find and set any instances of this blueprint type if any exists and we are not already editing one
 	FBlueprintEditorUtils::FindAndSetDebuggableBlueprintInstances();
 
@@ -2267,11 +2223,6 @@ void FBlueprintEditor::CreateDefaultTabContents(const TArray<UBlueprint*>& InBlu
 			. BlueprintToWatch(InBlueprint)
 			. IsEnabled(!bIsInterface && !bIsMacro);
 
-		this->BlueprintProfiler =
-			SNew(SBlueprintProfilerView)
-			.AssetEditor(SharedThis(this))
-			.ProfileViewType(EBlueprintPerfViewType::ExecutionGraph);
-
 		this->Palette = 
 			SNew(SBlueprintPalette, SharedThis(this))
 				.IsEnabled(this, &FBlueprintEditor::IsFocusedGraphEditable);
@@ -2337,7 +2288,7 @@ void FBlueprintEditor::CreateSCSEditors()
 
 void FBlueprintEditor::OnLogTokenClicked(const TSharedRef<IMessageToken>& Token)
 {
-	if( Token->GetType() == EMessageToken::Object)
+	if (Token->GetType() == EMessageToken::Object)
 	{
 		const TSharedRef<FUObjectToken> UObjectToken = StaticCastSharedRef<FUObjectToken>(Token);
 		if(UObjectToken->GetObject().IsValid())
@@ -2545,12 +2496,6 @@ void FBlueprintEditor::CreateDefaultCommands()
 		FCanExecuteAction::CreateSP(this, &FBlueprintEditor::HasAnyWatches)
 		);
 
-	// Blueprint Profiler Commands
-	ToolkitCommands->MapAction(FFullBlueprintEditorCommands::Get().ToggleProfiler,
-		FExecuteAction::CreateSP(this, &FBlueprintEditor::ToggleProfiler),
-		FCanExecuteAction::CreateSP(this, &FBlueprintEditor::IsProfilerAvailable),
-		FIsActionChecked::CreateSP(this, &FBlueprintEditor::IsProfilerActive));
-
 	// New document actions
 	ToolkitCommands->MapAction( FBlueprintEditorCommands::Get().AddNewVariable,
 		FExecuteAction::CreateSP(this, &FBlueprintEditor::OnAddNewVariable),
@@ -2644,54 +2589,6 @@ void FBlueprintEditor::CreateDefaultCommands()
 		FIsActionChecked::CreateLambda([]()->bool{ return GetDefault<UBlueprintEditorSettings>()->bShowActionMenuItemSignatures; }));
 }
 
-bool FBlueprintEditor::IsProfilerAvailable() const
-{
-	bool bIsBPProfilerSupported = false;
-	if (GetDefault<UEditorExperimentalSettings>()->bBlueprintPerformanceAnalysisTools)
-	{
-		if (GetToolkitFName() == FName("BlueprintEditor") && GetCurrentMode() == FBlueprintEditorApplicationModes::StandardBlueprintEditorMode)
-		{
-			if (UBlueprint* Blueprint = GetBlueprintObj())
-			{
-				if (Blueprint->BlueprintType == BPTYPE_Normal || Blueprint->BlueprintType == BPTYPE_LevelScript)
-				{
-					bIsBPProfilerSupported = true;
-				}
-			}
-		}
-	}
-	return bIsBPProfilerSupported;
-
-}
-
-bool FBlueprintEditor::IsProfilerActive() const
-{
-	if (IsProfilerAvailable())
-	{
-		IBlueprintProfilerInterface& ProfilerModule = FModuleManager::LoadModuleChecked<IBlueprintProfilerInterface>("BlueprintProfiler");
-		return ProfilerModule.IsProfilerEnabled();
-	}
-	return false;
-}
-
-void FBlueprintEditor::ToggleProfiler()
-{
-	if (IsProfilerAvailable())
-	{
-		IBlueprintProfilerInterface& ProfilerModule = FModuleManager::LoadModuleChecked<IBlueprintProfilerInterface>("BlueprintProfiler");
-		ProfilerModule.ToggleProfilingCapture();
-	
-		if (ProfilerModule.IsProfilerEnabled())
-		{
-			TabManager->InvokeTab(FBlueprintEditorTabs::BlueprintProfilerID);
-		}
-		else
-		{
-			TabManager->InvokeTab(FBlueprintEditorTabs::BlueprintProfilerID)->RequestCloseTab();
-		}
-	}
-}
-
 void FBlueprintEditor::OpenNativeCodeGenerationTool()
 {
 	UBlueprint* Blueprint = GetBlueprintObj();
@@ -3261,18 +3158,6 @@ void FBlueprintEditor::Compile()
 		Arguments.Add(TEXT("BlueprintName"), FText::FromString(BlueprintObj->GetName()));
 		BlueprintLog.NewPage(FText::Format(LOCTEXT("CompilationPageLabel", "Compile {BlueprintName}"), Arguments));
 
-		// Decide if we want an instrumented compile.
-		const bool bProfilerAvailable = IsProfilerAvailable();
-		bool bAddInstrumentation = false;
-	    if (bProfilerAvailable)
-	    {
-		    if (GetDefault<UEditorExperimentalSettings>()->bBlueprintPerformanceAnalysisTools)
-		    {
-			    IBlueprintProfilerInterface& ProfilerModule = FModuleManager::LoadModuleChecked<IBlueprintProfilerInterface>("BlueprintProfiler");
-			    bAddInstrumentation = ProfilerModule.IsProfilerEnabled();
-		    }
-		}
-
 		FCompilerResultsLog LogResults;
 		LogResults.SetSourcePath(BlueprintObj->GetPathName());
 		LogResults.BeginEvent(TEXT("Compile"));
@@ -3283,10 +3168,6 @@ void FBlueprintEditor::Compile()
 		{
 			CompileOptions |= EBlueprintCompileOptions::SaveIntermediateProducts;
 		}
-		if(bAddInstrumentation)
-		{
-			CompileOptions |= EBlueprintCompileOptions::AddInstrumentation;
-		}
 		FKismetEditorUtilities::CompileBlueprint(BlueprintObj, CompileOptions, &LogResults);
 
 		LogResults.EndEvent();
@@ -3310,25 +3191,6 @@ void FBlueprintEditor::Compile()
 
 		AppendExtraCompilerResults(CompilerResultsListing);
 
-	    // Update the blueprint instrumentation state and show the profiler window if required
-	    if (bProfilerAvailable)
-	    {
-		    bBlueprintHasInstrumentation = BlueprintObj && BlueprintObj->GeneratedClass ? BlueprintObj->GeneratedClass->HasInstrumentation() : false;
-		    if (bBlueprintHasInstrumentation)
-		    {
-				if (GetDefault<UEditorExperimentalSettings>()->bBlueprintPerformanceAnalysisTools)
-				{
-					IBlueprintProfilerInterface& ProfilerModule = FModuleManager::LoadModuleChecked<IBlueprintProfilerInterface>("BlueprintProfiler");
-					ProfilerModule.AddInstrumentedBlueprint(BlueprintObj);
-				    TabManager->InvokeTab(FBlueprintEditorTabs::BlueprintProfilerID);
-				}
-		    }
-		    else
-		    {
-			    TabManager->InvokeTab(FBlueprintEditorTabs::BlueprintProfilerID)->RequestCloseTab();
-		    }
-		}
-
 		// send record when player clicks compile and send the result
 		// this will make sure how the users activity is
 		AnalyticsTrackCompileEvent(BlueprintObj, LogResults.NumErrors, LogResults.NumWarnings);
@@ -3538,31 +3400,58 @@ void FBlueprintEditor::JumpToHyperlink(const UObject* ObjectReference, bool bReq
 		}
 
 		// Open the document
-		OpenDocument(const_cast<UEdGraph*>(Graph), OpenMode);
+		OpenDocument(Graph, OpenMode);
 	}
 	else if (const AActor* ReferencedActor = Cast<const AActor>(ObjectReference))
 	{
-		// Select the in-level actors referred to by this node
+		// Check if the world is active in the editor. It's possible to open level BPs without formally opening
+		// the levels through Find-in-Blueprints
+		bool bInOpenWorld = false;
+		const TIndirectArray<FWorldContext>& WorldContextList = GEditor->GetWorldContexts();
+		const UWorld* ReferencedActorOwningWorld = ReferencedActor->GetWorld();
+		for (const FWorldContext& WorldContext : WorldContextList)
+		{
+			if (WorldContext.World() == ReferencedActorOwningWorld)
+			{
+				bInOpenWorld = true;
+				break;
+			}
+		}
+
+		// Clear the selection even if we couldn't find it, so the existing selection doesn't get mistaken for the desired to be selected actor
 		GEditor->SelectNone(false, false);
-		GEditor->SelectActor(const_cast<AActor*>(ReferencedActor), true, true, true);
 
-		// Point the camera at it
-		GUnrealEd->Exec( ReferencedActor->GetWorld(), TEXT("CAMERA ALIGN ACTIVEVIEWPORTONLY"));
+		if (bInOpenWorld)
+		{
+			// Select the in-level actor
+			GEditor->SelectActor(const_cast<AActor*>(ReferencedActor), true, true, true);
+
+			// Point the camera at it
+			GUnrealEd->Exec(ReferencedActor->GetWorld(), TEXT("CAMERA ALIGN ACTIVEVIEWPORTONLY"));
+		}
 	}
 	else if(const UFunction* Function = Cast<const UFunction>(ObjectReference))
 	{
 		if (UEdGraph* FunctionGraph = FBlueprintEditorUtils::FindScopeGraph(GetBlueprintObj(), Function))
 		{
-			OpenDocument(const_cast<UEdGraph*>(FunctionGraph), FDocumentTracker::OpenNewDocument);
+			OpenDocument(FunctionGraph, FDocumentTracker::OpenNewDocument);
 		}
 	}
+	else if(const UBlueprintGeneratedClass* Class = Cast<const UBlueprintGeneratedClass>(ObjectReference))
+	{
+		FAssetEditorManager::Get().OpenEditorForAsset(Class->ClassGeneratedBy);
+	}
+	else if (const UTimelineTemplate* Timeline = Cast<const UTimelineTemplate>(ObjectReference))
+	{
+		OpenDocument(Timeline, FDocumentTracker::OpenNewDocument);
+	}
 	else if ((ObjectReference != nullptr) && ObjectReference->IsAsset())
 	{
 		FAssetEditorManager::Get().OpenEditorForAsset(const_cast<UObject*>(ObjectReference));
 	}
 	else
 	{
-		UE_LOG(LogBlueprint, Warning, TEXT("Unknown type of hyperlinked object (%s)"), *GetNameSafe(ObjectReference));
+		UE_LOG(LogBlueprint, Warning, TEXT("Unknown type of hyperlinked object (%s), cannot focus it"), *GetNameSafe(ObjectReference));
 	}
 
 	//@TODO: Hacky way to ensure a message is seen when hitting an exception and doing intraframe debugging
@@ -4326,15 +4215,8 @@ void FBlueprintEditor::OnAddParentNode()
 {
 	const UEdGraphSchema_K2* Schema = GetDefault<UEdGraphSchema_K2>();
 
-	const FGraphPanelSelectionSet& SelectedNodes = GetSelectedNodes();
-	if( SelectedNodes.Num() == 1 )
+	if (UEdGraphNode* SelectedObj = GetSingleSelectedNode())
 	{
-		UObject* SelectedObj = NULL;
-		for (FGraphPanelSelectionSet::TConstIterator It(SelectedNodes); It; ++It)
-		{
-			SelectedObj = *It;
-		}
-
 		// Get the function that the event node or function entry represents
 		FFunctionFromNodeHelper FunctionFromNode(SelectedObj);
 		if (FunctionFromNode.Function && FunctionFromNode.Node)
@@ -4365,15 +4247,8 @@ bool FBlueprintEditor::CanAddParentNode() const
 {
 	const UEdGraphSchema_K2* Schema = GetDefault<UEdGraphSchema_K2>();
 
-	const FGraphPanelSelectionSet& SelectedNodes = GetSelectedNodes();
-	if (SelectedNodes.Num() == 1)
+	if (UEdGraphNode* SelectedObj = GetSingleSelectedNode())
 	{
-		UObject* SelectedObj = NULL;
-		for (FGraphPanelSelectionSet::TConstIterator It(SelectedNodes); It; ++It)
-		{
-			SelectedObj = *It;
-		}
-
 		// Get the function that the event node or function entry represents
 		FFunctionFromNodeHelper FunctionFromNode(SelectedObj);
 		if (FunctionFromNode.Function)
@@ -5951,7 +5826,7 @@ void FBlueprintEditor::PasteNodesHere(class UEdGraph* DestinationGraph, const FV
 				bNeedToModifyStructurally = true;
 			}
 
-			// For pasted Event nodes, we need to see if there is an already existing node in a disabled state that needs to be cleaned up
+			// For pasted Event nodes, we need to see if there is an already existing node in a ghost state that needs to be cleaned up
 			if (UK2Node_Event* EventNode = Cast<UK2Node_Event>(Node))
 			{
 				// Gather all existing event nodes
@@ -5962,14 +5837,13 @@ void FBlueprintEditor::PasteNodesHere(class UEdGraph* DestinationGraph, const FV
 				{
 					check(ExistingEventNode);
 
-					bool bIdenticalNode = EventNode != ExistingEventNode && ExistingEventNode->bOverrideFunction && UK2Node_Event::AreEventNodesIdentical(EventNode, ExistingEventNode);
+					const bool bIdenticalNode = (EventNode != ExistingEventNode) && ExistingEventNode->bOverrideFunction && UK2Node_Event::AreEventNodesIdentical(EventNode, ExistingEventNode);
 
 					// Check if the nodes are identical, if they are we need to delete the original because it is disabled. Identical nodes that are in an enabled state will never make it this far and still be enabled.
 					if (bIdenticalNode)
 					{
 						// Should not have made it to being a pasted node if the pre-existing node wasn't disabled or was otherwise explicitly disabled by the user.
-						ensure(!ExistingEventNode->IsNodeEnabled());
-						ensure(!ExistingEventNode->bUserSetEnabledState);
+						ensure(ExistingEventNode->IsAutomaticallyPlacedGhostNode());
 
 						// Destroy the pre-existing node, we do not need it.
 						ExistingEventNode->DestroyNode();
@@ -6317,59 +6191,30 @@ bool FBlueprintEditor::CanStopWatchingPin() const
 
 bool FBlueprintEditor::CanGoToDefinition() const
 {
-	const FGraphPanelSelectionSet SelectedNodes = GetSelectedNodes();
-	if(SelectedNodes.Num() == 1)
-	{
-		UObject* Node = *SelectedNodes.CreateConstIterator();
-		if(UK2Node_CallFunction* FunctionCall = Cast<UK2Node_CallFunction>(Node))
-		{
-			const UEdGraphNode* ResultEventNode = NULL;
-			if(FunctionCall->GetFunctionGraph(ResultEventNode))
-			{
-				return true;
-			}
-		}
-		else if (UK2Node_MacroInstance* MacroNode = Cast<UK2Node_MacroInstance>(Node))
-		{
-			UEdGraph* MacroGraph = MacroNode->GetMacroGraph();
-			if (MacroGraph)
-			{
-				return true;
-			}
-		}
-		else if(UK2Node_Composite* Composite = Cast<UK2Node_Composite>(Node))
-		{
-			return true;
-		}
-	}
-
-	return false;
+	const UEdGraphNode* Node = GetSingleSelectedNode();
+	return (Node != nullptr) && Node->CanJumpToDefinition();
 }
 
 void FBlueprintEditor::OnGoToDefinition()
 {
-	const FGraphPanelSelectionSet SelectedNodes = GetSelectedNodes();
-	check(SelectedNodes.Num() == 1);
-	OnNodeDoubleClicked(Cast<UEdGraphNode>(*SelectedNodes.CreateConstIterator()));
+	if (UEdGraphNode* SelectedGraphNode = GetSingleSelectedNode())
+	{
+		OnNodeDoubleClicked(SelectedGraphNode);
+	}
 }
 
 FString FBlueprintEditor::GetDocLinkForSelectedNode()
 {
 	FString DocumentationLink;
 
-	const FGraphPanelSelectionSet SelectedNodes = GetSelectedNodes();
-	if (SelectedNodes.Num() == 1)
+	if (const UEdGraphNode* SelectedGraphNode = GetSingleSelectedNode())
 	{
-		UEdGraphNode* SelectedGraphNode = Cast<UEdGraphNode>(*SelectedNodes.CreateConstIterator());
-		if (SelectedGraphNode != NULL)
-		{
-			FString DocLink = SelectedGraphNode->GetDocumentationLink();
-			FString DocExcerpt = SelectedGraphNode->GetDocumentationExcerptName();
+		const FString DocLink = SelectedGraphNode->GetDocumentationLink();
+		const FString DocExcerpt = SelectedGraphNode->GetDocumentationExcerptName();
 
-			if (!DocLink.IsEmpty() && !DocExcerpt.IsEmpty())
-			{
-				DocumentationLink = FEditorClassUtils::GetDocumentationLinkFromExcerpt(DocLink, DocExcerpt);
-			}
+		if (!DocLink.IsEmpty() && !DocExcerpt.IsEmpty())
+		{
+			DocumentationLink = FEditorClassUtils::GetDocumentationLinkFromExcerpt(DocLink, DocExcerpt);
 		}
 	}
 
@@ -6378,7 +6223,7 @@ FString FBlueprintEditor::GetDocLinkForSelectedNode()
 
 void FBlueprintEditor::OnGoToDocumentation()
 {
-	FString DocumentationLink = GetDocLinkForSelectedNode();
+	const FString DocumentationLink = GetDocLinkForSelectedNode();
 	if (!DocumentationLink.IsEmpty())
 	{
 		IDocumentation::Get()->Open(DocumentationLink, FDocumentationSourceInfo(TEXT("rightclick_bpnode")));
@@ -6398,12 +6243,10 @@ void FBlueprintEditor::OnSetEnabledStateForSelectedNodes(ENodeEnabledState NewSt
 	const FGraphPanelSelectionSet SelectedNodes = GetSelectedNodes();
 	for (UObject* SelectedNode : SelectedNodes)
 	{
-		UEdGraphNode* SelectedGraphNode = Cast<UEdGraphNode>(SelectedNode);
-		if(SelectedGraphNode)
+		if (UEdGraphNode* SelectedGraphNode = Cast<UEdGraphNode>(SelectedNode))
 		{
 			SelectedGraphNode->Modify();
-			SelectedGraphNode->EnabledState = NewState;
-			SelectedGraphNode->bUserSetEnabledState = true;
+			SelectedGraphNode->SetEnabledState(NewState);
 		}
 	}
 
@@ -6439,19 +6282,17 @@ ECheckBoxState FBlueprintEditor::GetEnabledCheckBoxStateForSelectedNodes()
 ECheckBoxState FBlueprintEditor::CheckEnabledStateForSelectedNodes(ENodeEnabledState CheckState)
 {
 	const FGraphPanelSelectionSet SelectedNodes = GetSelectedNodes();
-	ECheckBoxState Result = SelectedNodes.Num() > 0 ? ECheckBoxState::Undetermined : ECheckBoxState::Unchecked;
+	ECheckBoxState Result = (SelectedNodes.Num() > 0) ? ECheckBoxState::Undetermined : ECheckBoxState::Unchecked;
 	for (UObject* SelectedNode : SelectedNodes)
 	{
-		UEdGraphNode* SelectedGraphNode = Cast<UEdGraphNode>(SelectedNode);
-		if(SelectedGraphNode)
+		if (UEdGraphNode* SelectedGraphNode = Cast<UEdGraphNode>(SelectedNode))
 		{
-			const ENodeEnabledState NodeState = SelectedGraphNode->EnabledState;
-			if(Result == ECheckBoxState::Undetermined)
+			const ENodeEnabledState NodeState = SelectedGraphNode->GetDesiredEnabledState();
+			if (Result == ECheckBoxState::Undetermined)
 			{
 				Result = (NodeState == CheckState) ? ECheckBoxState::Checked : ECheckBoxState::Unchecked;
 			}
-			else if((NodeState != CheckState && Result == ECheckBoxState::Checked)
-				|| (NodeState == CheckState && Result == ECheckBoxState::Unchecked))
+			else if ((NodeState != CheckState && Result == ECheckBoxState::Checked) || (NodeState == CheckState && Result == ECheckBoxState::Unchecked))
 			{
 				Result = ECheckBoxState::Undetermined;
 				break;
@@ -6472,109 +6313,11 @@ bool FBlueprintEditor::GetSaveIntermediateBuildProducts() const
 	return bSaveIntermediateBuildProducts;
 }
 
-void FBlueprintEditor::OnNodeDoubleClicked(class UEdGraphNode* Node)
+void FBlueprintEditor::OnNodeDoubleClicked(UEdGraphNode* Node)
 {
-	//@TODO: Pull these last few stragglers out into their respective nodes; requires asking the same question in a different way without knowledge of the Explorer
-	if (UK2Node_CallFunction* FunctionCall = Cast<UK2Node_CallFunction>(Node))
-	{
-		const UEdGraphNode* ResultEventNode = NULL;
-
-		if(UEdGraph* FunctionGraph = FunctionCall->GetFunctionGraph(ResultEventNode))
-		{
-			// If there is an event node, jump to it, otherwise jump to the function graph
-			if(ResultEventNode)
-			{
-				FKismetEditorUtilities::BringKismetToFocusAttentionOnObject(ResultEventNode, false);
-			}
-			else
-			{
-				FKismetEditorUtilities::BringKismetToFocusAttentionOnObject(FunctionGraph);
-			}
-		}
-	}
-	else if (UK2Node_MacroInstance* MacroNode = Cast<UK2Node_MacroInstance>(Node))
-	{
-		UEdGraph* MacroGraph = MacroNode->GetMacroGraph();
-		if (MacroGraph)
-		{
-			FKismetEditorUtilities::BringKismetToFocusAttentionOnObject(MacroGraph);
-		}
-	}
-	else if (UK2Node_Timeline* TimelineNode = Cast<UK2Node_Timeline>(Node))
-	{
-		if (UTimelineTemplate* Timeline = GetBlueprintObj()->FindTimelineTemplateByVariableName(TimelineNode->TimelineName))
-		{
-			OpenDocument(Timeline, FDocumentTracker::OpenNewDocument);
-		}
-	}
-	else if(UK2Node_Variable* VariableNode = Cast<UK2Node_Variable>(Node))
-	{
-		// Jump to the RepNotify function graph
-		FName RepNotifyFunc = FBlueprintEditorUtils::GetBlueprintVariableRepNotifyFunc(GetBlueprintObj(), VariableNode->GetVarName());
-		if(RepNotifyFunc != NAME_None)
-		{
-			for( UEdGraph* Graph : GetBlueprintObj()->FunctionGraphs )
-			{
-				if(Graph->GetFName() == RepNotifyFunc)
-				{
-					FKismetEditorUtilities::BringKismetToFocusAttentionOnObject(Graph);
-					break;
-				}
-			}
-		}
-	}
-	else if (UObject* HyperlinkTarget = Node->GetJumpTargetForDoubleClick())
+	if (Node->CanJumpToDefinition())
 	{
-		// true if jumping to a level actor, do not handle the node further
-		bool bJumpToLevelActor = false;
-
-		// If double clicking on a node that references level actors, jump to hyperlink so that they are selected in the level
-		UK2Node* K2Node = Cast<UK2Node>(Node);
-		if(K2Node)
-		{
-			if (AActor* Actor = K2Node->GetReferencedLevelActor())
-			{
-				// Check if the world is active in the editor. It's possible to open level BPs without formally opening
-				// the levels through Find-in-Blueprints
-				const TIndirectArray<FWorldContext>& WorldContextList = GEditor->GetWorldContexts();
-				for (const FWorldContext& WorldContext : WorldContextList)
-				{
-					if (WorldContext.World() == Actor->GetWorld())
-					{
-						JumpToHyperlink(Actor);
-						break;
-					}
-				}
-				bJumpToLevelActor = true;
-			}
-		}
-
-		if(!bJumpToLevelActor)
-		{
-			// Check to see if our outer chain contains a blueprint. If we're inside a blueprint (a graph, pin, etc.) then
-			// focus on the target; otherwise open the editor for the target.
-			UBlueprint* TargetBP = Cast<UBlueprint>(const_cast<UObject*>(HyperlinkTarget));
-			if(TargetBP == NULL)
-			{
-				for(UObject* TestOuter = HyperlinkTarget->GetOuter(); TestOuter; TestOuter = TestOuter->GetOuter())
-				{
-					TargetBP = Cast<UBlueprint>(TestOuter);
-					if(TargetBP != NULL)
-					{
-						break;
-					}
-				}
-			}
-
-			if(TargetBP)
-			{
-				FKismetEditorUtilities::BringKismetToFocusAttentionOnObject(HyperlinkTarget);
-			}
-			else
-			{
-				FAssetEditorManager::Get().OpenEditorForAsset(HyperlinkTarget);
-			}
-		}
+		Node->JumpToDefinition();
 	}
 }
 
@@ -7079,12 +6822,14 @@ void FBlueprintEditor::ExpandNode(UEdGraphNode* InNodeToExpand, UEdGraph* InSour
 		// Successfully added the node to the graph, we may need to remove flags
 		if (Node->HasAllFlags(RF_Transient) && !DestinationGraph->HasAllFlags(RF_Transient))
 		{
+			Node->SetFlags(RF_Transactional);
 			Node->ClearFlags(RF_Transient);
 			TArray<UObject*> Subobjects;
 			GetObjectsWithOuter(Node, Subobjects);
 			for (UObject* Subobject : Subobjects)
 			{
 				Subobject->ClearFlags(RF_Transient);
+				Subobject->SetFlags(RF_Transactional);
 			}
 		}
 		DestinationGraph->Nodes.Add(Node);
@@ -7788,19 +7533,10 @@ bool FBlueprintEditor::CanRenameNodes() const
 {
 	if (IsEditable(GetFocusedGraph()))
 	{
-		const FGraphPanelSelectionSet SelectedNodes = GetSelectedNodes();
-		bool bCanRenameNodes = (SelectedNodes.Num() == 1) ? true : false;
-
-		for (FGraphPanelSelectionSet::TConstIterator NodeIt(SelectedNodes); NodeIt; ++NodeIt)
+		if (const UEdGraphNode* SelectedNode = GetSingleSelectedNode())
 		{
-			UEdGraphNode* SelectedNode = Cast<UEdGraphNode>(*NodeIt);
-			if (SelectedNode && !SelectedNode->bCanRenameNode)
-			{
-				bCanRenameNodes = false;
-				break;
-			}
+			return SelectedNode->bCanRenameNode;
 		}
-		return bCanRenameNodes;
 	}
 	return false;
 }
@@ -7885,7 +7621,7 @@ TSharedPtr<SGraphEditor> FBlueprintEditor::OpenGraphAndBringToFront(UEdGraph* Gr
 	return NewGraphEditor;
 }
 
-TSharedPtr<SDockTab> FBlueprintEditor::OpenDocument(UObject* DocumentID, FDocumentTracker::EOpenDocumentCause Cause)
+TSharedPtr<SDockTab> FBlueprintEditor::OpenDocument(const UObject* DocumentID, FDocumentTracker::EOpenDocumentCause Cause)
 {
 	TSharedRef<FTabPayload_UObject> Payload = FTabPayload_UObject::Make(DocumentID);
 	return DocumentManager->OpenDocument(Payload, Cause);
@@ -7896,14 +7632,14 @@ void FBlueprintEditor::NavigateTab(FDocumentTracker::EOpenDocumentCause InCause)
 	OpenDocument(NULL, InCause);
 }
 
-void FBlueprintEditor::CloseDocumentTab(UObject* DocumentID)
+void FBlueprintEditor::CloseDocumentTab(const UObject* DocumentID)
 {
 	TSharedRef<FTabPayload_UObject> Payload = FTabPayload_UObject::Make(DocumentID);
 	DocumentManager->CloseTab(Payload);
 }
 
 // Finds any open tabs containing the specified document and adds them to the specified array; returns true if at least one is found
-bool FBlueprintEditor::FindOpenTabsContainingDocument(UObject* DocumentID, /*inout*/ TArray< TSharedPtr<SDockTab> >& Results)
+bool FBlueprintEditor::FindOpenTabsContainingDocument(const UObject* DocumentID, /*inout*/ TArray< TSharedPtr<SDockTab> >& Results)
 {
 	int32 StartingCount = Results.Num();
 
@@ -8005,121 +7741,6 @@ void FBlueprintEditor::SetPinVisibility(SGraphEditor::EPinVisibility Visibility)
 	OnSetPinVisibility.Broadcast(PinVisibility);
 }
 
-void FBlueprintEditor::GotoNativeFunctionDefinition()
-{
-	TSharedPtr<SGraphEditor> GraphEditor = FocusedGraphEdPtr.Pin();
-	if( GraphEditor.IsValid() && IsSelectionNativeFunction() )
-	{
-		const FGraphPanelSelectionSet SelectedNodes = GraphEditor->GetSelectedNodes();
-		FGraphPanelSelectionSet::TConstIterator NodeIter( SelectedNodes );
-
-		const UK2Node_CallFunction* FunctionNode = Cast<UK2Node_CallFunction>( *NodeIter );
-
-		if( FunctionNode )
-		{
-			UFunction* TargetFunction = FunctionNode->GetTargetFunction();
-
-			if( TargetFunction )
-			{
-				FString NativeParentClassHeaderPath;
-				const bool bFileFound = FSourceCodeNavigation::FindClassHeaderPath(TargetFunction, NativeParentClassHeaderPath) 
-										&& (IFileManager::Get().FileSize(*NativeParentClassHeaderPath) != INDEX_NONE);
-				if( bFileFound )
-				{
-					const FString AbsNativeParentClassHeaderPath = FPaths::ConvertRelativePathToFull( NativeParentClassHeaderPath );
-					FSourceCodeNavigation::OpenSourceFile( AbsNativeParentClassHeaderPath );
-				}
-			}
-		}
-	}
-}
-
-bool FBlueprintEditor::IsSelectionNativeFunction()
-{
-	TSharedPtr<SGraphEditor> GraphEditor = FocusedGraphEdPtr.Pin();
-	if( GraphEditor.IsValid() && FSourceCodeNavigation::IsCompilerAvailable() )
-	{
-		FGraphPanelSelectionSet SelectedNodes = GraphEditor->GetSelectedNodes();
-		FGraphPanelSelectionSet::TIterator NodeIter( SelectedNodes );
-
-		if(NodeIter)
-		{
-			const UK2Node_CallFunction* FunctionNode = Cast<UK2Node_CallFunction>( *NodeIter );
-
-			if( FunctionNode && SelectedNodes.Num() == 1 )
-			{
-				UFunction* FunctionPtr = FunctionNode->FunctionReference.ResolveMember<UFunction>( FunctionNode->GetBlueprintClassFromNode() );
-				UClass* OwningClass = FunctionPtr ? FunctionPtr->GetOuterUClass() : NULL;
-
-				if( OwningClass && OwningClass->HasAllClassFlags( CLASS_Native ))
-				{
-					return true;
-				}
-			}
-		}
-	}
-	return false;
-}
-
-void FBlueprintEditor::GotoNativeVariableDefinition()
-{
-	TSharedPtr<SGraphEditor> GraphEditor = FocusedGraphEdPtr.Pin();
-	if( GraphEditor.IsValid() && IsSelectionNativeVariable() )
-	{
-		const FGraphPanelSelectionSet SelectedNodes = GraphEditor->GetSelectedNodes();
-		FGraphPanelSelectionSet::TConstIterator NodeIter( SelectedNodes );
-		const UK2Node_Variable* VarNode = Cast<UK2Node_Variable>( *NodeIter );
-
-		if( VarNode )
-		{
-			UProperty* VariableProperty = VarNode->VariableReference.ResolveMember<UProperty>( VarNode->GetBlueprintClassFromNode() );
-
-			if( VariableProperty )
-			{
-				FString NativeParentClassHeaderPath;
-				const bool bFileFound = FSourceCodeNavigation::FindClassHeaderPath(VariableProperty, NativeParentClassHeaderPath) 
-										&& (IFileManager::Get().FileSize(*NativeParentClassHeaderPath) != INDEX_NONE);
-				if( bFileFound )
-				{
-					const FString AbsNativeParentClassHeaderPath = FPaths::ConvertRelativePathToFull( NativeParentClassHeaderPath );
-					FSourceCodeNavigation::OpenSourceFile( AbsNativeParentClassHeaderPath );
-				}
-			}
-		}
-	}
-}
-
-bool FBlueprintEditor::IsSelectionNativeVariable()
-{
-	TSharedPtr<SGraphEditor> GraphEditor = FocusedGraphEdPtr.Pin();
-	if( GraphEditor.IsValid() )
-	{
-		FGraphPanelSelectionSet SelectedNodes = GraphEditor->GetSelectedNodes();
-		FGraphPanelSelectionSet::TIterator NodeIter( SelectedNodes );
-
-		if(NodeIter)
-		{
-			const UK2Node_Variable* VarNode = Cast<UK2Node_Variable>( *NodeIter );
-
-			if( VarNode && SelectedNodes.Num() == 1 )
-			{
-				UProperty* VariableProperty = VarNode->VariableReference.ResolveMember<UProperty>( VarNode->GetBlueprintClassFromNode() );
-
-				if( VariableProperty && VariableProperty->IsNative())
-				{
-					return true;
-				}
-			}
-		}
-	}
-	return false;
-}
-
-bool FBlueprintEditor::IsNativeCodeBrowsingAvailable() const
-{
-	return bCodeBasedProject;
-}
-
 void FBlueprintEditor::OnFindReferences()
 {
 	TSharedPtr<SGraphEditor> GraphEditor = FocusedGraphEdPtr.Pin();
@@ -8156,16 +7777,7 @@ void FBlueprintEditor::OnFindReferences()
 
 bool FBlueprintEditor::CanFindReferences()
 {
-	TSharedPtr<SGraphEditor> GraphEditor = FocusedGraphEdPtr.Pin();
-	if (GraphEditor.IsValid())
-	{
-		const FGraphPanelSelectionSet SelectedNodes = GraphEditor->GetSelectedNodes();
-		if( SelectedNodes.Num() == 1 )
-		{
-			return true;
-		}
-	}
-	return false;
+	return GetSingleSelectedNode() != nullptr;
 }
 
 AActor* FBlueprintEditor::GetPreviewActor() const
diff --git a/Engine/Source/Editor/Kismet/Private/BlueprintEditorCommands.cpp b/Engine/Source/Editor/Kismet/Private/BlueprintEditorCommands.cpp
--- a/Engine/Source/Editor/Kismet/Private/BlueprintEditorCommands.cpp
+++ b/Engine/Source/Editor/Kismet/Private/BlueprintEditorCommands.cpp
@@ -27,7 +27,7 @@ PRAGMA_DISABLE_OPTIMIZATION
 void FBlueprintEditorCommands::RegisterCommands()
 {
 	// Edit commands
-	UI_COMMAND( FindInBlueprint, "Search", "Finds references to functions, events, variables, and pins in the current Blueprint", EUserInterfaceActionType::Button, FInputChord(EModifierKey::Control, EKeys::F) );
+	UI_COMMAND( FindInBlueprint, "Find", "Finds references to functions, events, variables, and pins in the current Blueprint (use Ctrl+Shift+F to search in all Blueprints)", EUserInterfaceActionType::Button, FInputChord(EModifierKey::Control, EKeys::F) );
 	UI_COMMAND( FindInBlueprints, "Find in Blueprints", "Find references to functions, events and variables in ALL Blueprints", EUserInterfaceActionType::Button, FInputChord(EModifierKey::Control | EModifierKey::Shift, EKeys::F) );
 	UI_COMMAND( ReparentBlueprint, "Reparent Blueprint", "Change the parent of this Blueprint", EUserInterfaceActionType::Button, FInputChord() );
 
diff --git a/Engine/Source/Editor/Kismet/Private/BlueprintEditorModes.cpp b/Engine/Source/Editor/Kismet/Private/BlueprintEditorModes.cpp
--- a/Engine/Source/Editor/Kismet/Private/BlueprintEditorModes.cpp
+++ b/Engine/Source/Editor/Kismet/Private/BlueprintEditorModes.cpp
@@ -61,7 +61,6 @@ TSharedPtr<FTabManager::FLayout> GetDefaltEditorLayout(TSharedPtr<class FBluepri
 					->SetSizeCoefficient( 0.20f )
 					->AddTab( FBlueprintEditorTabs::CompilerResultsID, ETabState::ClosedTab )
 					->AddTab( FBlueprintEditorTabs::FindResultsID, ETabState::ClosedTab )
-					->AddTab( FBlueprintEditorTabs::BlueprintProfilerID, ETabState::ClosedTab )
 				)
 			)
 			->Split
@@ -100,7 +99,6 @@ FBlueprintEditorApplicationMode::FBlueprintEditorApplicationMode(TSharedPtr<clas
 	}
 	BlueprintEditorTabFactories.RegisterFactory(MakeShareable(new FCompilerResultsSummoner(InBlueprintEditor)));
 	BlueprintEditorTabFactories.RegisterFactory(MakeShareable(new FFindResultsSummoner(InBlueprintEditor)));
-	BlueprintEditorTabFactories.RegisterFactory(MakeShareable(new FBlueprintProfilerSummoner(InBlueprintEditor)));
 	
 	if( bRegisterViewport )
 	{
@@ -121,7 +119,6 @@ FBlueprintEditorApplicationMode::FBlueprintEditorApplicationMode(TSharedPtr<clas
 	InBlueprintEditor->GetToolbarBuilder()->AddScriptingToolbar(ToolbarExtender);
 	InBlueprintEditor->GetToolbarBuilder()->AddBlueprintGlobalOptionsToolbar(ToolbarExtender);
 	InBlueprintEditor->GetToolbarBuilder()->AddDebuggingToolbar(ToolbarExtender);
-	InBlueprintEditor->GetToolbarBuilder()->AddProfilerToolbar(ToolbarExtender);
 
 	FBlueprintEditorModule& BlueprintEditorModule = FModuleManager::LoadModuleChecked<FBlueprintEditorModule>("Kismet");
 	BlueprintEditorModule.OnRegisterTabsForEditor().Broadcast(BlueprintEditorTabFactories, InModeName, InBlueprintEditor);
@@ -569,7 +566,6 @@ FBlueprintEditorUnifiedMode::FBlueprintEditorUnifiedMode(TSharedPtr<class FBluep
 	}
 	BlueprintEditorTabFactories.RegisterFactory(MakeShareable(new FCompilerResultsSummoner(InBlueprintEditor)));
 	BlueprintEditorTabFactories.RegisterFactory(MakeShareable(new FFindResultsSummoner(InBlueprintEditor)));
-	BlueprintEditorTabFactories.RegisterFactory(MakeShareable(new FBlueprintProfilerSummoner(InBlueprintEditor)));
 	
 	if( bRegisterViewport )
 	{
@@ -629,7 +625,6 @@ FBlueprintEditorUnifiedMode::FBlueprintEditorUnifiedMode(TSharedPtr<class FBluep
 						->SetSizeCoefficient( 0.20f )
 						->AddTab( FBlueprintEditorTabs::CompilerResultsID, ETabState::ClosedTab )
 						->AddTab( FBlueprintEditorTabs::FindResultsID, ETabState::ClosedTab )
-						->AddTab( FBlueprintEditorTabs::BlueprintProfilerID, ETabState::ClosedTab )
 					)
 				)
 				->Split
@@ -690,7 +685,6 @@ FBlueprintEditorUnifiedMode::FBlueprintEditorUnifiedMode(TSharedPtr<class FBluep
 						->SetSizeCoefficient( 0.20f )
 						->AddTab( FBlueprintEditorTabs::CompilerResultsID, ETabState::ClosedTab )
 						->AddTab( FBlueprintEditorTabs::FindResultsID, ETabState::ClosedTab )
-						->AddTab( FBlueprintEditorTabs::BlueprintProfilerID, ETabState::ClosedTab )
 					)
 				)
 				->Split
@@ -715,10 +709,6 @@ FBlueprintEditorUnifiedMode::FBlueprintEditorUnifiedMode(TSharedPtr<class FBluep
 	InBlueprintEditor->GetToolbarBuilder()->AddCompileToolbar(ToolbarExtender);
 	InBlueprintEditor->GetToolbarBuilder()->AddScriptingToolbar(ToolbarExtender);
 	InBlueprintEditor->GetToolbarBuilder()->AddBlueprintGlobalOptionsToolbar(ToolbarExtender);
-	if (GetDefault<UEditorExperimentalSettings>()->bBlueprintPerformanceAnalysisTools)
-	{
-		InBlueprintEditor->GetToolbarBuilder()->AddProfilerToolbar(ToolbarExtender);
-	}
 	
 	if ( bRegisterViewport )
 	{
diff --git a/Engine/Source/Editor/Kismet/Private/BlueprintEditorTabFactories.cpp b/Engine/Source/Editor/Kismet/Private/BlueprintEditorTabFactories.cpp
--- a/Engine/Source/Editor/Kismet/Private/BlueprintEditorTabFactories.cpp
+++ b/Engine/Source/Editor/Kismet/Private/BlueprintEditorTabFactories.cpp
@@ -11,7 +11,6 @@
 #include "BlueprintEditorTabs.h"
 #include "STimelineEditor.h"
 #include "Debugging/SKismetDebuggingView.h"
-#include "Profiler/SBlueprintProfilerView.h"
 #include "SKismetInspector.h"
 #include "SSCSEditor.h"
 #include "SSCSEditorViewport.h"
@@ -136,24 +135,6 @@ TSharedRef<SWidget> FDebugInfoSummoner::CreateTabBody(const FWorkflowTabSpawnInf
 	return BlueprintEditorPtr->GetDebuggingView();
 }
 
-FBlueprintProfilerSummoner::FBlueprintProfilerSummoner(TSharedPtr<class FAssetEditorToolkit> InHostingApp) : FWorkflowTabFactory(FBlueprintEditorTabs::BlueprintProfilerID, InHostingApp)
-{
-	TabLabel = LOCTEXT("BlueprintProfilerViewTitle", "BlueprintProfiler");
-	TabIcon = FSlateIcon(FEditorStyle::GetStyleSetName(), "PerfTools.TabIcon");
-
-	EnableTabPadding();
-	bIsSingleton = true;
-
-	ViewMenuDescription = LOCTEXT("BlueprintProfilerView", "Profiler");
-	ViewMenuTooltip = LOCTEXT("BlueprintProfilerView_ToolTip", "Shows the current profiler data");
-}
-
-TSharedRef<SWidget> FBlueprintProfilerSummoner::CreateTabBody(const FWorkflowTabSpawnInfo& Info) const
-{
-	TSharedPtr<FBlueprintEditor> BlueprintEditorPtr = StaticCastSharedPtr<FBlueprintEditor>(HostingApp.Pin());
-	return BlueprintEditorPtr->GetBlueprintProfilerView();
-}
-
 FDefaultsEditorSummoner::FDefaultsEditorSummoner(TSharedPtr<class FAssetEditorToolkit> InHostingApp)
 	: FWorkflowTabFactory(FBlueprintEditorTabs::DefaultEditorID, InHostingApp)
 {
diff --git a/Engine/Source/Editor/Kismet/Private/BlueprintEditorTabFactories.h b/Engine/Source/Editor/Kismet/Private/BlueprintEditorTabFactories.h
--- a/Engine/Source/Editor/Kismet/Private/BlueprintEditorTabFactories.h
+++ b/Engine/Source/Editor/Kismet/Private/BlueprintEditorTabFactories.h
@@ -89,17 +89,6 @@ public:
 	virtual TSharedRef<SWidget> CreateTabBody(const FWorkflowTabSpawnInfo& Info) const override;
 };
 
-/////////////////////////////////////////////////////
-// FBlueprintProfilerSummoner
-
-struct FBlueprintProfilerSummoner : public FWorkflowTabFactory
-{
-public:
-	FBlueprintProfilerSummoner(TSharedPtr<class FAssetEditorToolkit> InHostingApp);
-
-	virtual TSharedRef<SWidget> CreateTabBody(const FWorkflowTabSpawnInfo& Info) const override;
-};
-
 /////////////////////////////////////////////////////
 // FDefaultsEditorSummoner
 
diff --git a/Engine/Source/Editor/Kismet/Private/BlueprintEditorTabs.cpp b/Engine/Source/Editor/Kismet/Private/BlueprintEditorTabs.cpp
--- a/Engine/Source/Editor/Kismet/Private/BlueprintEditorTabs.cpp
+++ b/Engine/Source/Editor/Kismet/Private/BlueprintEditorTabs.cpp
@@ -11,7 +11,6 @@ const FName FBlueprintEditorTabs::FindResultsID( TEXT( "FindResults" ) );
 const FName FBlueprintEditorTabs::ConstructionScriptEditorID( TEXT( "ConstructionScriptEditor" ) );
 const FName FBlueprintEditorTabs::SCSViewportID( TEXT( "SCSViewport" ) );
 const FName FBlueprintEditorTabs::MyBlueprintID( TEXT( "MyBlueprint" ) );
-const FName FBlueprintEditorTabs::BlueprintProfilerID( TEXT( "BlueprintProfiler" ) );
 const FName FBlueprintEditorTabs::ReplaceNodeReferencesID( TEXT( "ReplaceNodeReferences" ) );
 
 const FName FBlueprintEditorTabs::GraphEditorID( TEXT( "GraphEditor" ) );
diff --git a/Engine/Source/Editor/Kismet/Private/SBlueprintEditorSelectedDebugObjectWidget.cpp b/Engine/Source/Editor/Kismet/Private/SBlueprintEditorSelectedDebugObjectWidget.cpp
--- a/Engine/Source/Editor/Kismet/Private/SBlueprintEditorSelectedDebugObjectWidget.cpp
+++ b/Engine/Source/Editor/Kismet/Private/SBlueprintEditorSelectedDebugObjectWidget.cpp
@@ -664,6 +664,11 @@ void SBlueprintEditorSelectedDebugObjectWidget::AddDebugObject(UObject* TestObje
 	if (AActor* Actor = Cast<AActor>(TestObject))
 	{
 		Label = Actor->GetActorLabel();
+
+		if (Actor->IsSelected())
+		{
+			Label += TEXT(" (selected)");
+		}
 	}
 	else
 	{
diff --git a/Engine/Source/Editor/Kismet/Private/SBlueprintEditorToolbar.cpp b/Engine/Source/Editor/Kismet/Private/SBlueprintEditorToolbar.cpp
--- a/Engine/Source/Editor/Kismet/Private/SBlueprintEditorToolbar.cpp
+++ b/Engine/Source/Editor/Kismet/Private/SBlueprintEditorToolbar.cpp
@@ -350,8 +350,6 @@ void FFullBlueprintEditorCommands::RegisterCommands()
 	UI_COMMAND(SwitchToBlueprintDefaultsMode, "Defaults", "Switches to Class Defaults Mode", EUserInterfaceActionType::ToggleButton, FInputChord());
 	UI_COMMAND(SwitchToComponentsMode, "Components", "Switches to Components Mode", EUserInterfaceActionType::ToggleButton, FInputChord());
 
-	UI_COMMAND(ToggleProfiler,"Profiler", "Toggle Profiler State", EUserInterfaceActionType::ToggleButton, FInputChord());
-
 	UI_COMMAND(EditGlobalOptions, "Class Settings", "Edit Class Settings (Previously known as Blueprint Props)", EUserInterfaceActionType::ToggleButton, FInputChord());
 	UI_COMMAND(EditClassDefaults, "Class Defaults", "Edit the initial values of your class.", EUserInterfaceActionType::ToggleButton, FInputChord());
 
@@ -480,17 +478,6 @@ void FBlueprintEditorToolbar::AddComponentsToolbar(TSharedPtr<FExtender> Extende
 		FToolBarExtensionDelegate::CreateSP( this, &FBlueprintEditorToolbar::FillComponentsToolbar ) );
 }
 
-void FBlueprintEditorToolbar::AddProfilerToolbar(TSharedPtr<FExtender> Extender)
-{
-	TSharedPtr<FBlueprintEditor> BlueprintEditorPtr = BlueprintEditor.Pin();
-
-	Extender->AddToolBarExtension(
-		"Asset",
-		EExtensionHook::After,
-		BlueprintEditorPtr->GetToolkitCommands(),
-		FToolBarExtensionDelegate::CreateSP( this, &FBlueprintEditorToolbar::FillProfilerToolbar ) );
-}
-
 void FBlueprintEditorToolbar::FillBlueprintEditorModesToolbar(FToolBarBuilder& ToolbarBuilder)
 {
 	TSharedPtr<FBlueprintEditor> BlueprintEditorPtr = BlueprintEditor.Pin();
@@ -553,29 +540,6 @@ void FBlueprintEditorToolbar::FillCompileToolbar(FToolBarBuilder& ToolbarBuilder
 	ToolbarBuilder.EndSection();
 }
 
-void FBlueprintEditorToolbar::FillProfilerToolbar(FToolBarBuilder& ToolbarBuilder)
-{
-	const FFullBlueprintEditorCommands& Commands = FFullBlueprintEditorCommands::Get();
-	TSharedPtr<FBlueprintEditor> BlueprintEditorPtr = BlueprintEditor.Pin();
-	UBlueprint* BlueprintObj = BlueprintEditorPtr->GetBlueprintObj();
-	
-	ToolbarBuilder.BeginSection("Profiler");
-	
-	if(BlueprintObj != NULL)
-	{
-		ToolbarBuilder.AddToolBarButton( Commands.ToggleProfiler,
-										 NAME_None, 
-										 TAttribute<FText>(),
-										 TAttribute<FText>(this, &FBlueprintEditorToolbar::GetProfilerStatusTooltip),
-										 FSlateIcon(FEditorStyle::GetStyleSetName(), "BlueprintEditor.EnableProfiling"),
-										 FName(TEXT("ToggleProfiler")));
-
-
-	}
-	
-	ToolbarBuilder.EndSection();
-}
-
 void FBlueprintEditorToolbar::FillNewToolbar(FToolBarBuilder& ToolbarBuilder)
 {
 	const FBlueprintEditorCommands& Commands = FBlueprintEditorCommands::Get();
@@ -665,14 +629,7 @@ FSlateIcon FBlueprintEditorToolbar::GetStatusImage() const
 	case BS_Error:
 		return FSlateIcon(FEditorStyle::GetStyleSetName(), "Kismet.Status.Error");
 	case BS_UpToDate:
-	{
-		const bool bInstrumented = BlueprintObj->GeneratedClass->HasInstrumentation();
-		if (BlueprintEditor.Pin()->IsProfilerActive())
-		{
-			return bInstrumented ? FSlateIcon(FEditorStyle::GetStyleSetName(), "Kismet.Status.Instrumented") : FSlateIcon(FEditorStyle::GetStyleSetName(), "Kismet.Status.NotInstrumented");
-		}
-		return bInstrumented ? FSlateIcon(FEditorStyle::GetStyleSetName(), "Kismet.Status.Instrumented") : FSlateIcon(FEditorStyle::GetStyleSetName(), "Kismet.Status.Good");
-	}
+		return FSlateIcon(FEditorStyle::GetStyleSetName(), "Kismet.Status.Good");
 	case BS_UpToDateWithWarnings:
 		return FSlateIcon(FEditorStyle::GetStyleSetName(), "Kismet.Status.Warning");
 	}
@@ -699,25 +656,12 @@ FText FBlueprintEditorToolbar::GetStatusTooltip() const
 	case BS_Error:
 		return LOCTEXT("CompileError_Status", "There was an error during compilation, see the log for details");
 	case BS_UpToDate:
-	{
-		const bool bInstrumented = BlueprintObj->GeneratedClass->HasInstrumentation();
-		if (BlueprintEditor.Pin()->IsProfilerActive())
-		{
-			return bInstrumented ? LOCTEXT("GoodToGoInstrumented_Status", "Instrumentation is active") : LOCTEXT("GoodToGoNotInstrumented_Status", "Compile to add instrumentation");
-		}
-		return bInstrumented ? LOCTEXT("GoodToGoInstrumented_Status", "Instrumentation is active") : LOCTEXT("GoodToGo_Status", "Good to go");
-	}
+		return LOCTEXT("GoodToGo_Status", "Good to go");
 	case BS_UpToDateWithWarnings:
 		return LOCTEXT("GoodToGoWarning_Status", "There was a warning during compilation, see the log for details");
 	}
 }
 
-FText FBlueprintEditorToolbar::GetProfilerStatusTooltip() const
-{
-	return BlueprintEditor.Pin()->IsProfilerActive() ?	LOCTEXT("Profiling_Enabled", "Profiling Active") :
-														LOCTEXT("Profiling_Disabled", "Profiling Inactive");
-}
-
 TArray< TSharedPtr< SWidget> > FBlueprintEditorToolbar::GenerateToolbarWidgets(const UBlueprint* BlueprintObj, TAttribute<FName> ActiveModeGetter, FOnModeChangeRequested ActiveModeSetter)
 {
 	TArray< TSharedPtr< SWidget> > Ret;
diff --git a/Engine/Source/Editor/Kismet/Private/SBlueprintSubPalette.cpp b/Engine/Source/Editor/Kismet/Private/SBlueprintSubPalette.cpp
--- a/Engine/Source/Editor/Kismet/Private/SBlueprintSubPalette.cpp
+++ b/Engine/Source/Editor/Kismet/Private/SBlueprintSubPalette.cpp
@@ -331,21 +331,17 @@ FReply SBlueprintSubPalette::OnActionDragged( const TArray< TSharedPtr<FEdGraphS
 		else if(InAction->GetTypeId() == FEdGraphSchemaAction_K2Var::StaticGetTypeId())
 		{
 			FEdGraphSchemaAction_K2Var* VarAction = (FEdGraphSchemaAction_K2Var*)InAction.Get();
-
-			UClass* VarClass = VarAction->GetVariableClass();
-			if(VarClass != NULL)
+			if (UClass* VarClass = VarAction->GetVariableClass())
 			{
-				return FReply::Handled().BeginDragDrop(FKismetVariableDragDropAction::New(VarAction->GetVariableName(), VarClass, AnalyticsDelegate));
+				return FReply::Handled().BeginDragDrop(FKismetVariableDragDropAction::New(InAction, VarAction->GetVariableName(), VarClass, AnalyticsDelegate));
 			}
 		}
 		else if(InAction->GetTypeId() == FEdGraphSchemaAction_K2Delegate::StaticGetTypeId())
 		{
 			FEdGraphSchemaAction_K2Delegate* DelegateAction = (FEdGraphSchemaAction_K2Delegate*)InAction.Get();
-			UClass* VarClass = DelegateAction->GetDelegateClass();
-			check(VarClass);
-			if(VarClass != NULL)
+			if (UClass* VarClass = DelegateAction->GetDelegateClass())
 			{
-				return FReply::Handled().BeginDragDrop(FKismetDelegateDragDropAction::New(DelegateAction->GetDelegateName(), VarClass, AnalyticsDelegate));
+				return FReply::Handled().BeginDragDrop(FKismetDelegateDragDropAction::New(InAction, DelegateAction->GetDelegateName(), VarClass, AnalyticsDelegate));
 			}
 		}
 		else if (InAction->GetTypeId() == FBlueprintDragDropMenuItem::StaticGetTypeId())
