Copying //UE4/Dev-Framework to //UE4/Dev-Main (Source: //UE4/Dev-Framework @ 3582324)

diff --git a/Engine/Source/Editor/Kismet/Private/SMyBlueprint.cpp b/Engine/Source/Editor/Kismet/Private/SMyBlueprint.cpp
--- a/Engine/Source/Editor/Kismet/Private/SMyBlueprint.cpp
+++ b/Engine/Source/Editor/Kismet/Private/SMyBlueprint.cpp
@@ -140,6 +140,92 @@ public:
 	TWeakPtr<SMyBlueprint>	MyBlueprintPtr;
 };
 
+//////////////////////////////////////////////////////////////////////////
+// FGraphActionSort
+
+// Helper structure to aid category sorting
+struct FGraphActionSort
+{
+public:
+	FGraphActionSort(TArray<FName>& BlueprintCategorySorting)
+		: bCategoriesModified(false)
+		, CategorySortIndices(BlueprintCategorySorting)
+	{
+		CategoryUsage.Init(0, CategorySortIndices.Num());
+	}
+
+	void AddAction(const FString& Category, TSharedPtr<FEdGraphSchemaAction> Action)
+	{
+		// Find root category
+		int32 RootCategoryDelim = Category.Find(TEXT("|"));
+		FName RootCategory = RootCategoryDelim == INDEX_NONE ? *Category : *Category.Left(RootCategoryDelim);
+		// Get root sort index
+		const int32 SortIndex = GetSortIndex(RootCategory) + Action->GetSectionID();
+
+		SortedActions.Add(SortIndex, Action);
+	}
+
+	void AddAction(TSharedPtr<FEdGraphSchemaAction> Action)
+	{
+		const FString UserCategoryName = FEditorCategoryUtils::GetCategoryDisplayString(Action->GetCategory().ToString());
+		AddAction(UserCategoryName, Action);
+	}
+
+	void GetAllActions(FGraphActionListBuilderBase& OutActions)
+	{
+		SortedActions.KeySort(TLess<int32>());
+
+		for (const auto& Iter : SortedActions)
+		{
+			OutActions.AddAction(Iter.Value);
+		}
+	}
+
+	void CleanupCategories()
+	{
+		// Scrub unused categories from the blueprint
+		if (bCategoriesModified)
+		{
+			for (int32 CategoryIdx = CategoryUsage.Num() - 1; CategoryIdx >= 0; CategoryIdx--)
+			{
+				if (CategoryUsage[CategoryIdx] == 0)
+				{
+					CategorySortIndices.RemoveAt(CategoryIdx);
+				}
+			}
+			bCategoriesModified = false;
+		}
+	}
+
+private:
+	const int32 GetSortIndex(FName Category)
+	{
+		int32 SortIndex = CategorySortIndices.Find(Category);
+
+		if (SortIndex == INDEX_NONE)
+		{
+			bCategoriesModified = true;
+			SortIndex = CategorySortIndices.Add(Category);
+			CategoryUsage.Add(0);
+		}
+		CategoryUsage[SortIndex]++;
+		// Spread the sort values so we can fine tune sorting
+		SortIndex *= 1000;
+
+		return SortIndex + SortedActions.Num();
+	}
+
+private:
+	/** Signals if the blueprint categories have been modified and require cleanup */
+	bool bCategoriesModified;
+	/** Tracks category usage to aid removal of unused categories */
+	TArray<int32> CategoryUsage;
+	/** Reference to the category sorting in the blueprint */
+	TArray<FName>& CategorySortIndices;
+	/** Map used to sort Graph actions */
+	TMultiMap<int32, TSharedPtr<FEdGraphSchemaAction>> SortedActions;
+};
+
 //////////////////////////////////////////////////////////////////////////
 
 void SMyBlueprint::Construct(const FArguments& InArgs, TWeakPtr<FBlueprintEditor> InBlueprintEditor, const UBlueprint* InBlueprint )
@@ -767,59 +853,22 @@ bool SMyBlueprint::CanRequestRenameOnActionNode(TWeakPtr<FGraphActionNode> InSel
 {
 	bool bIsReadOnly = true;
 
-	// If checking if renaming is available on a category node, the category must have a non-native variable.
-	if(InSelectedNode.Pin()->IsCategoryNode())
+	// If checking if renaming is available on a category node, the category must have a non-native entry
+	if (InSelectedNode.Pin()->IsCategoryNode())
 	{
 		TArray<TSharedPtr<FEdGraphSchemaAction>> Actions;
 		GraphActionMenu->GetCategorySubActions(InSelectedNode, Actions);
 
-		for (int32 i = 0; i < Actions.Num(); ++i)
+		for (TSharedPtr<FEdGraphSchemaAction> Action : Actions)
 		{
-			if (Actions[i]->GetTypeId() == FEdGraphSchemaAction_K2Var::StaticGetTypeId())
-			{
-				FEdGraphSchemaAction_K2Var* VarAction = (FEdGraphSchemaAction_K2Var*)Actions[i].Get();
-
-				UClass* SkeletonGeneratedClass = GetBlueprintObj()->SkeletonGeneratedClass;
-				if (UProperty* TargetProperty = FindField<UProperty>(SkeletonGeneratedClass, VarAction->GetVariableName()))
-				{
-					UClass* OuterClass = CastChecked<UClass>(TargetProperty->GetOuter());
-					const bool bIsNativeVar = (OuterClass->ClassGeneratedBy == NULL);
-
-					// If the variable is not native and it's outer is the skeleton generated class, we can rename the category
-					if(!bIsNativeVar && OuterClass == GetBlueprintObj()->SkeletonGeneratedClass)
-					{
-						bIsReadOnly = false;
-						break;
-					}
-				}
-			}
-			else if (Actions[i]->GetTypeId() == FEdGraphSchemaAction_K2Graph::StaticGetTypeId())
-			{
-				FEdGraphSchemaAction_K2Graph* GraphAction = (FEdGraphSchemaAction_K2Graph*)(Actions[i].Get());
-
-				// Do not allow renaming of any graph actions outside of the following
-				if(Actions[i]->GetSectionID() == NodeSectionID::FUNCTION || Actions[i]->GetSectionID() == NodeSectionID::MACRO)
-				{
-					if(!FBlueprintEditorUtils::IsPaletteActionReadOnly(Actions[i], BlueprintEditorPtr.Pin())
-						&& FBlueprintEditorUtils::FindBlueprintForGraph(GraphAction->EdGraph) == GetBlueprintObj()
-						&& FBlueprintEditorUtils::FindFunctionInImplementedInterfaces(GetBlueprintObj(), GraphAction->EdGraph->GetFName(), nullptr, true) == nullptr)
-					{
-						bIsReadOnly = false;
-						break;
-					}
-				}
-			}
-			else if(Actions[i]->GetTypeId() == FEdGraphSchemaAction_K2Delegate::StaticGetTypeId())
+			if (Action->GetPersistentItemDefiningObject().IsPotentiallyEditable())
 			{
-				if(!FBlueprintEditorUtils::IsPaletteActionReadOnly(Actions[i], BlueprintEditorPtr.Pin()))
-				{
-					bIsReadOnly = false;
-					break;
-				}
+				bIsReadOnly = false;
+				break;
 			}
 		}
 	}
-	else if(InSelectedNode.Pin()->IsActionNode())
+	else if (InSelectedNode.Pin()->IsActionNode())
 	{
 		check( InSelectedNode.Pin()->Actions.Num() > 0 && InSelectedNode.Pin()->Actions[0].IsValid() );
 		bIsReadOnly = FBlueprintEditorUtils::IsPaletteActionReadOnly(InSelectedNode.Pin()->Actions[0], BlueprintEditorPtr.Pin());
@@ -840,7 +889,7 @@ TSharedRef<SWidget> SMyBlueprint::OnCreateWidgetForAction(FCreateWidgetForAction
 	return BlueprintEditorPtr.IsValid() ? SNew(SBlueprintPaletteItem, InCreateData, BlueprintEditorPtr.Pin()) : SNew(SBlueprintPaletteItem, InCreateData, GetBlueprintObj());
 }
 
-void SMyBlueprint::GetChildGraphs(UEdGraph* InEdGraph, int32 const SectionId, FGraphActionListBuilderBase& OutAllActions, FText ParentCategory)
+void SMyBlueprint::GetChildGraphs(UEdGraph* InEdGraph, int32 const SectionId, FGraphActionSort& SortList, FText ParentCategory)
 {
 	check(InEdGraph);
 
@@ -853,9 +902,8 @@ void SMyBlueprint::GetChildGraphs(UEdGraph* InEdGraph, int32 const SectionId, FG
 	const FText EdGraphDisplayName = EdGraphDisplayInfo.DisplayName;
 
 	// Grab children graphs
-	for (TArray<UEdGraph*>::TIterator It(InEdGraph->SubGraphs); It; ++It)
+	for (UEdGraph* Graph : InEdGraph->SubGraphs)
 	{
-		UEdGraph* Graph = *It;
 		check(Graph);
 
 		FGraphDisplayInfo ChildGraphDisplayInfo;
@@ -883,17 +931,17 @@ void SMyBlueprint::GetChildGraphs(UEdGraph* InEdGraph, int32 const SectionId, FG
 		TSharedPtr<FEdGraphSchemaAction_K2Graph> NewChildAction = MakeShareable(new FEdGraphSchemaAction_K2Graph(EEdGraphSchemaAction_K2Graph::Subgraph, Category, ChildDesc, ChildTooltip, 1, SectionId));
 		NewChildAction->FuncName = DisplayName;
 		NewChildAction->EdGraph = Graph;
-		OutAllActions.AddAction(NewChildAction);
+		SortList.AddAction(NewChildAction);
 		
-		GetChildGraphs(Graph, SectionId, OutAllActions, Category);
-		GetChildEvents(Graph, SectionId, OutAllActions, Category);
+		GetChildGraphs(Graph, SectionId, SortList, Category);
+		GetChildEvents(Graph, SectionId, SortList, Category);
 	}
 }
 
 struct FCreateEdGraphSchemaActionHelper
 {
 	template<class ActionType, class NodeType>
-	static void CreateAll(UEdGraph const* EdGraph, int32 SectionId, FGraphActionListBuilderBase& OutAllActions, const FText& ActionCategory)
+	static void CreateAll(UEdGraph const* EdGraph, int32 SectionId, FGraphActionSort& SortList, const FText& ActionCategory)
 	{
 		TArray<NodeType*> EventNodes;
 		EdGraph->GetNodesOfClass<NodeType>(EventNodes);
@@ -907,12 +955,12 @@ struct FCreateEdGraphSchemaActionHelper
 			TSharedPtr<ActionType> EventNodeAction = MakeShareable(new ActionType(ActionCategory, Description, Tooltip, 0));
 			EventNodeAction->NodeTemplate = EventNode;
 			EventNodeAction->SectionID = SectionId;
-			OutAllActions.AddAction(EventNodeAction);
+			SortList.AddAction(EventNodeAction);
 		}
 	}
 };
 
-void SMyBlueprint::GetChildEvents(UEdGraph const* InEdGraph, int32 const SectionId, FGraphActionListBuilderBase& OutAllActions, FText ParentCategory) const
+void SMyBlueprint::GetChildEvents(UEdGraph const* InEdGraph, int32 const SectionId, FGraphActionSort& SortList, FText ParentCategory) const
 {
 	if (!ensure(InEdGraph != NULL))
 	{
@@ -936,12 +984,12 @@ void SMyBlueprint::GetChildEvents(UEdGraph const* InEdGraph, int32 const Section
 		ActionCategory = EdGraphDisplayName;
 	}
 
-	FCreateEdGraphSchemaActionHelper::CreateAll<FEdGraphSchemaAction_K2Event, UK2Node_Event>(InEdGraph, SectionId, OutAllActions, ActionCategory);
-	FCreateEdGraphSchemaActionHelper::CreateAll<FEdGraphSchemaAction_K2InputAction, UK2Node_InputKey>(InEdGraph, SectionId, OutAllActions, ActionCategory);
-	FCreateEdGraphSchemaActionHelper::CreateAll<FEdGraphSchemaAction_K2InputAction, UK2Node_InputAction>(InEdGraph, SectionId, OutAllActions, ActionCategory);
+	FCreateEdGraphSchemaActionHelper::CreateAll<FEdGraphSchemaAction_K2Event, UK2Node_Event>(InEdGraph, SectionId, SortList, ActionCategory);
+	FCreateEdGraphSchemaActionHelper::CreateAll<FEdGraphSchemaAction_K2InputAction, UK2Node_InputKey>(InEdGraph, SectionId, SortList, ActionCategory);
+	FCreateEdGraphSchemaActionHelper::CreateAll<FEdGraphSchemaAction_K2InputAction, UK2Node_InputAction>(InEdGraph, SectionId, SortList, ActionCategory);
 }
 
-void SMyBlueprint::GetLocalVariables(FGraphActionListBuilderBase& OutAllActions) const
+void SMyBlueprint::GetLocalVariables(FGraphActionSort& SortList) const
 {
 	// We want to pull local variables from the top level function graphs
 	UEdGraph* TopLevelGraph = FBlueprintEditorUtils::GetTopLevelGraph(GetFocusedGraph());
@@ -963,7 +1011,7 @@ void SMyBlueprint::GetLocalVariables(FGraphActionListBuilderBase& OutAllActions)
 		const UEdGraphSchema_K2* K2Schema = GetDefault<UEdGraphSchema_K2>();
 		for (UK2Node_FunctionEntry* const FunctionEntry : FunctionEntryNodes)
 		{
-			for( const FBPVariableDescription& Variable : FunctionEntry->LocalVariables )
+			for (const FBPVariableDescription& Variable : FunctionEntry->LocalVariables)
 			{
 				FText Category = Variable.Category;
 				if (Variable.Category.EqualTo(K2Schema->VR_DefaultCategory))
@@ -976,7 +1024,7 @@ void SMyBlueprint::GetLocalVariables(FGraphActionListBuilderBase& OutAllActions)
 				{
 					TSharedPtr<FEdGraphSchemaAction_K2LocalVar> NewVarAction = MakeShareable(new FEdGraphSchemaAction_K2LocalVar(Category, FText::FromName(Variable.VarName), FText::GetEmpty(), 0, NodeSectionID::LOCAL_VARIABLE));
 					NewVarAction->SetVariableInfo(Variable.VarName, Func, Variable.VarType.PinCategory == K2Schema->PC_Boolean);
-					OutAllActions.AddAction(NewVarAction);
+					SortList.AddAction(NewVarAction);
 				}
 			}
 		}
@@ -1012,89 +1060,6 @@ void SMyBlueprint::CollectAllActions(FGraphActionListBuilderBase& OutAllActions)
 
 	bool bShowReplicatedOnly = IsShowingReplicatedVariablesOnly();
 
-	// Helper structure to aid category sorting
-	struct FGraphActionSort
-	{
-		public:
-			FGraphActionSort( TArray<FName>& BlueprintCategorySorting )
-				: bCategoriesModified( false )
-				, CategorySortIndices( BlueprintCategorySorting )
-			{
-				CategoryUsage.Init( 0, CategorySortIndices.Num() );
-			}
-
-			void AddAction( const FString& Category, TSharedPtr<FEdGraphSchemaAction> Action )
-			{
-				// Find root category
-				int32 RootCategoryDelim = Category.Find( TEXT( "|" ));
-				FName RootCategory = RootCategoryDelim == INDEX_NONE ?	*Category : *Category.Left( RootCategoryDelim );
-				// Get root sort index
-				const int32 SortIndex = GetSortIndex( RootCategory ) + Action->GetSectionID();
-
-				static const FName DefaultCategory( TEXT( "Default" ));
-				//if( RootCategory != DefaultCategory )
-				//{
-				//	// Strip sections out of non default categories to enable them to be grouped together
-				//	Action->SectionID = NodeSectionID::USER_SORTED;
-				//}
-				SortedActions.Add( SortIndex, Action );
-			}
-
-			void GetAllActions( FGraphActionListBuilderBase& OutActions )
-			{
-				SortedActions.KeySort( TLess<int32>() );
-
-				for( auto Iter : SortedActions )
-				{
-					OutActions.AddAction( Iter.Value );
-				}
-			}
-
-			void CleanupCategories()
-			{
-				// Scrub unused categories from the blueprint
-				if( bCategoriesModified )
-				{
-					for( int32 CategoryIdx = CategoryUsage.Num() - 1; CategoryIdx >= 0; CategoryIdx-- )
-					{
-						if( CategoryUsage[ CategoryIdx ] == 0 )
-						{
-							CategorySortIndices.RemoveAt( CategoryIdx );
-						}
-					}
-					bCategoriesModified = false;
-				}
-			}
-
-		private:
-			const int32 GetSortIndex( FName Category )
-			{
-				int32 SortIndex = CategorySortIndices.Find( Category );
-
-				if( SortIndex == INDEX_NONE )
-				{
-					bCategoriesModified = true;
-					SortIndex = CategorySortIndices.Add( Category );
-					CategoryUsage.Add( 0 );
-				}
-				CategoryUsage[ SortIndex ]++;
-				// Spread the sort values so we can fine tune sorting
-				SortIndex *= 1000;
-
-				return SortIndex + SortedActions.Num();
-			}
-
-		private:
-			/** Signals if the blueprint categories have been modified and require cleanup */
-			bool bCategoriesModified;
-			/** Tracks category usage to aid removal of unused categories */
-			TArray<int32> CategoryUsage;
-			/** Reference to the category sorting in the blueprint */
-			TArray<FName>& CategorySortIndices;
-			/** Map used to sort Graph actions */
-			TMultiMap<int32,TSharedPtr<FEdGraphSchemaAction>> SortedActions;
-	}; 
-
 	// Initialise action sorting instance
 	FGraphActionSort SortList( BlueprintObj->CategorySorting );
 	// List of names of functions we implement
@@ -1156,41 +1121,33 @@ void SMyBlueprint::CollectAllActions(FGraphActionListBuilderBase& OutAllActions)
 		}
 		else if (bShouldShowAsDelegate)
 		{
-			TSharedPtr<FEdGraphSchemaAction_K2Delegate> NewFuncAction;
+			TSharedPtr<FEdGraphSchemaAction_K2Delegate> NewDelegateAction;
 			// Delegate is visible in MyBlueprint when not-native or its category name is not empty.
 			if (Property->HasAllPropertyFlags(CPF_Edit) || !PropertyCategory.IsEmpty())
 			{
-				NewFuncAction = MakeShareable(new FEdGraphSchemaAction_K2Delegate(PropertyCategory, PropertyDesc, PropertyTooltip, 0, NodeSectionID::DELEGATE));
-				NewFuncAction->SetDelegateInfo(PropertyName, BlueprintObj->SkeletonGeneratedClass);
-				NewFuncAction->EdGraph = NULL;
-				SortList.AddAction( UserCategoryName, NewFuncAction );
+				NewDelegateAction = MakeShareable(new FEdGraphSchemaAction_K2Delegate(PropertyCategory, PropertyDesc, PropertyTooltip, 0, NodeSectionID::DELEGATE));
+				NewDelegateAction->SetVariableInfo(PropertyName, BlueprintObj->SkeletonGeneratedClass, false);
+				SortList.AddAction( UserCategoryName, NewDelegateAction );
 			}
 
 			UClass* OwnerClass = CastChecked<UClass>(Property->GetOuter());
 			UEdGraph* Graph = FBlueprintEditorUtils::GetDelegateSignatureGraphByName(BlueprintObj, PropertyName);
 			if (Graph && OwnerClass && (BlueprintObj == OwnerClass->ClassGeneratedBy))
 			{
-				if ( NewFuncAction.IsValid() )
+				if (NewDelegateAction.IsValid())
 				{
-					NewFuncAction->EdGraph = Graph;
+					NewDelegateAction->EdGraph = Graph;
 				}
 				ImplementedFunctionCache.Add(PropertyName);
 			}
 		}
 	}
 
-	// Add all the sorted variables and components
-	SortList.CleanupCategories();
-	SortList.GetAllActions( OutAllActions );
-
 	// Grab functions implemented by the blueprint
-	for (int32 i = 0; i < BlueprintObj->FunctionGraphs.Num(); i++)
+	for (UEdGraph* Graph : BlueprintObj->FunctionGraphs)
 	{
-		UEdGraph* Graph = BlueprintObj->FunctionGraphs[i];
 		check(Graph);
 
-		bool bIsConstructionScript = Graph->GetFName() == K2Schema->FN_UserConstructionScript;
-
 		FGraphDisplayInfo DisplayInfo;
 		Graph->GetSchema()->GetGraphDisplayInformation(*Graph, DisplayInfo);
 
@@ -1205,18 +1162,21 @@ void SMyBlueprint::CollectAllActions(FGraphActionListBuilderBase& OutAllActions)
 		}
 
 		//@TODO: Should be a bit more generic (or the AnimGraph shouldn't be stored as a FunctionGraph...)
+		const bool bIsConstructionScript = Graph->GetFName() == K2Schema->FN_UserConstructionScript;
 		int32 SectionID = Graph->IsA<UAnimationGraph>() ? NodeSectionID::GRAPH : NodeSectionID::FUNCTION;
 		TSharedPtr<FEdGraphSchemaAction_K2Graph> NewFuncAction = MakeShareable(new FEdGraphSchemaAction_K2Graph(EEdGraphSchemaAction_K2Graph::Function, FunctionCategory, DisplayInfo.PlainName, DisplayInfo.Tooltip, bIsConstructionScript ? 2 : 1, SectionID));
 		NewFuncAction->FuncName = Graph->GetFName();
 		NewFuncAction->EdGraph = Graph;
 
-		OutAllActions.AddAction(NewFuncAction);
+		const FString UserCategoryName = FEditorCategoryUtils::GetCategoryDisplayString(FunctionCategory.ToString());
+		SortList.AddAction(UserCategoryName, NewFuncAction);
 
-		GetChildGraphs(Graph, NewFuncAction->GetSectionID(), OutAllActions, FunctionCategory);
-		GetChildEvents(Graph, NewFuncAction->GetSectionID(), OutAllActions, FunctionCategory);
+		GetChildGraphs(Graph, NewFuncAction->GetSectionID(), SortList, FunctionCategory);
+		GetChildEvents(Graph, NewFuncAction->GetSectionID(), SortList, FunctionCategory);
 
 		ImplementedFunctionCache.Add(Graph->GetFName());
 	}
+
 	// Grab macros implemented by the blueprint
 	for (int32 i = 0; i < BlueprintObj->MacroGraphs.Num(); i++)
 	{
@@ -1233,10 +1193,12 @@ void SMyBlueprint::CollectAllActions(FGraphActionListBuilderBase& OutAllActions)
 		TSharedPtr<FEdGraphSchemaAction_K2Graph> NewMacroAction = MakeShareable(new FEdGraphSchemaAction_K2Graph(EEdGraphSchemaAction_K2Graph::Macro, MacroCategory, DisplayInfo.PlainName, DisplayInfo.Tooltip, 1, NodeSectionID::MACRO));
 		NewMacroAction->FuncName = MacroName;
 		NewMacroAction->EdGraph = Graph;
-		OutAllActions.AddAction(NewMacroAction);
 
-		GetChildGraphs(Graph, NewMacroAction->GetSectionID(), OutAllActions, MacroCategory);
-		GetChildEvents(Graph, NewMacroAction->GetSectionID(), OutAllActions, MacroCategory);
+		const FString UserCategoryName = FEditorCategoryUtils::GetCategoryDisplayString(MacroCategory.ToString());
+		SortList.AddAction(UserCategoryName, NewMacroAction);
+
+		GetChildGraphs(Graph, NewMacroAction->GetSectionID(), SortList, MacroCategory);
+		GetChildEvents(Graph, NewMacroAction->GetSectionID(), SortList, MacroCategory);
 
 		ImplementedFunctionCache.Add(MacroName);
 	}
@@ -1304,8 +1266,8 @@ void SMyBlueprint::CollectAllActions(FGraphActionListBuilderBase& OutAllActions)
 			NewFuncAction->EdGraph = Graph;
 			OutAllActions.AddAction(NewFuncAction);
 
-			GetChildGraphs(Graph, NewFuncAction->GetSectionID(), OutAllActions, FunctionCategory);
-			GetChildEvents(Graph, NewFuncAction->GetSectionID(), OutAllActions, FunctionCategory);
+			GetChildGraphs(Graph, NewFuncAction->GetSectionID(), SortList, FunctionCategory);
+			GetChildEvents(Graph, NewFuncAction->GetSectionID(), SortList, FunctionCategory);
 		}
 	}
 
@@ -1353,8 +1315,8 @@ void SMyBlueprint::CollectAllActions(FGraphActionListBuilderBase& OutAllActions)
 		NeUbergraphAction->EdGraph = Graph;
 		OutAllActions.AddAction(NeUbergraphAction);
 
-		GetChildGraphs(Graph, NeUbergraphAction->GetSectionID(), OutAllActions);
-		GetChildEvents(Graph, NeUbergraphAction->GetSectionID(), OutAllActions);
+		GetChildGraphs(Graph, NeUbergraphAction->GetSectionID(), SortList);
+		GetChildEvents(Graph, NeUbergraphAction->GetSectionID(), SortList);
 	}
 
 	// Grab intermediate pages
@@ -1371,14 +1333,18 @@ void SMyBlueprint::CollectAllActions(FGraphActionListBuilderBase& OutAllActions)
 		NewIntermediateAction->EdGraph = Graph;
 		OutAllActions.AddAction(NewIntermediateAction);
 
-		GetChildGraphs(Graph, NewIntermediateAction->GetSectionID(), OutAllActions);
-		GetChildEvents(Graph, NewIntermediateAction->GetSectionID(), OutAllActions);
+		GetChildGraphs(Graph, NewIntermediateAction->GetSectionID(), SortList);
+		GetChildEvents(Graph, NewIntermediateAction->GetSectionID(), SortList);
 	}
 
-	if(GetLocalActionsListVisibility().IsVisible())
+	if (GetLocalActionsListVisibility().IsVisible())
 	{
-		GetLocalVariables(OutAllActions);
+		GetLocalVariables(SortList);
 	}
+
+	// Add all the sorted variables, components, functions, etc...
+	SortList.CleanupCategories();
+	SortList.GetAllActions(OutAllActions);
 }
 
 void SMyBlueprint::CollectStaticSections(TArray<int32>& StaticSectionIDs)
@@ -1507,7 +1473,7 @@ FReply SMyBlueprint::OnActionDragged( const TArray< TSharedPtr<FEdGraphSchemaAct
 			{
 				if ((FuncAction->EdGraph != NULL) && GetBlueprintObj()->BlueprintType != BPTYPE_MacroLibrary)
 				{
-					return FReply::Handled().BeginDragDrop(FKismetMacroDragDropAction::New(FuncAction->FuncName, GetBlueprintObj(), FuncAction->EdGraph, AnalyticsDelegate));
+					return FReply::Handled().BeginDragDrop(FKismetMacroDragDropAction::New(InAction, FuncAction->FuncName, GetBlueprintObj(), FuncAction->EdGraph, AnalyticsDelegate));
 				}
 			}
 		}
@@ -1515,13 +1481,12 @@ FReply SMyBlueprint::OnActionDragged( const TArray< TSharedPtr<FEdGraphSchemaAct
 		{
 			FEdGraphSchemaAction_K2Delegate* DelegateAction = (FEdGraphSchemaAction_K2Delegate*)InAction.Get();
 			check(DelegateAction->GetDelegateName() != NAME_None);
-			UClass* VarClass = DelegateAction->GetDelegateClass();
-			if(VarClass != NULL)
+			if (UClass* VarClass = DelegateAction->GetDelegateClass())
 			{
 				const bool bIsAltDown = MouseEvent.IsAltDown();
 				const bool bIsCtrlDown = MouseEvent.IsLeftControlDown() || MouseEvent.IsRightControlDown();
 				
-				TSharedRef<FKismetVariableDragDropAction> DragOperation = FKismetDelegateDragDropAction::New(DelegateAction->GetDelegateName(), VarClass, AnalyticsDelegate);
+				TSharedRef<FKismetVariableDragDropAction> DragOperation = FKismetDelegateDragDropAction::New(InAction, DelegateAction->GetDelegateName(), VarClass, AnalyticsDelegate);
 				DragOperation->SetAltDrag(bIsAltDown);
 				DragOperation->SetCtrlDrag(bIsCtrlDown);
 				return FReply::Handled().BeginDragDrop(DragOperation);
@@ -1530,32 +1495,22 @@ FReply SMyBlueprint::OnActionDragged( const TArray< TSharedPtr<FEdGraphSchemaAct
 		else if( InAction->GetTypeId() == FEdGraphSchemaAction_K2LocalVar::StaticGetTypeId())
 		{
 			FEdGraphSchemaAction_K2LocalVar* VarAction = (FEdGraphSchemaAction_K2LocalVar*)InAction.Get();
-
-			UStruct* VariableScope = VarAction->GetVariableScope();
-			if(VariableScope != NULL)
+			if (UStruct* VariableScope = VarAction->GetVariableScope())
 			{
-				const bool bIsAltDown = MouseEvent.IsAltDown();
-				const bool bIsCtrlDown = MouseEvent.IsLeftControlDown() || MouseEvent.IsRightControlDown();
-
-				TSharedRef<FKismetVariableDragDropAction> DragOperation = FKismetVariableDragDropAction::New(VarAction->GetVariableName(), VariableScope, AnalyticsDelegate);
-				DragOperation->SetAltDrag(bIsAltDown);
-				DragOperation->SetCtrlDrag(bIsCtrlDown);
+				TSharedRef<FKismetVariableDragDropAction> DragOperation = FKismetVariableDragDropAction::New(InAction, VarAction->GetVariableName(), VariableScope, AnalyticsDelegate);
+				DragOperation->SetAltDrag(MouseEvent.IsAltDown());
+				DragOperation->SetCtrlDrag(MouseEvent.IsLeftControlDown() || MouseEvent.IsRightControlDown());
 				return FReply::Handled().BeginDragDrop(DragOperation);
 			}
 		}
 		else if(InAction->GetTypeId() == FEdGraphSchemaAction_K2Var::StaticGetTypeId())
 		{
 			FEdGraphSchemaAction_K2Var* VarAction = (FEdGraphSchemaAction_K2Var*)InAction.Get();
-
-			UClass* VarClass = VarAction->GetVariableClass();
-			if(VarClass != NULL)
+			if (UClass* VarClass = VarAction->GetVariableClass())
 			{
-				const bool bIsAltDown = MouseEvent.IsAltDown();
-				const bool bIsCtrlDown = MouseEvent.IsLeftControlDown() || MouseEvent.IsRightControlDown();
-				
-				TSharedRef<FKismetVariableDragDropAction> DragOperation = FKismetVariableDragDropAction::New(VarAction->GetVariableName(), VarClass, AnalyticsDelegate);
-				DragOperation->SetAltDrag(bIsAltDown);
-				DragOperation->SetCtrlDrag(bIsCtrlDown);
+				TSharedRef<FKismetVariableDragDropAction> DragOperation = FKismetVariableDragDropAction::New(InAction, VarAction->GetVariableName(), VarClass, AnalyticsDelegate);
+				DragOperation->SetAltDrag(MouseEvent.IsAltDown());
+				DragOperation->SetCtrlDrag(MouseEvent.IsLeftControlDown() || MouseEvent.IsRightControlDown());
 				return FReply::Handled().BeginDragDrop(DragOperation);
 			}
 		}
@@ -2162,19 +2117,20 @@ void SMyBlueprint::OnFindReference()
 	}
 	else if (FEdGraphSchemaAction_K2Var* VarAction = SelectionAsVar())
 	{
-		FGuid Guid = FBlueprintEditorUtils::FindMemberVariableGuidByName(Blueprint, VarAction->GetVariableName());
+		FString GuidTerm;
+		const FGuid Guid = FBlueprintEditorUtils::FindMemberVariableGuidByName(Blueprint, VarAction->GetVariableName());
 		if (Guid.IsValid())
 		{
-			FString VariableName = VarAction->GetVariableName().ToString();
-			SearchTerm = FString::Printf(TEXT("Nodes(VariableReference(MemberName=+\"%s\" && MemberGuid(A=%i && B=%i && C=%i && D=%i)))"), *VariableName, Guid.A, Guid.B, Guid.C, Guid.D);
-			bUseQuotes = false;
-		}
-		else
-		{
-			FString VariableName = VarAction->GetVariableName().ToString();
-			SearchTerm = FString::Printf(TEXT("Nodes(VariableReference(MemberName=+\"%s\"))"), *VariableName);
-			bUseQuotes = false;
+			GuidTerm = FString::Printf(TEXT(" && MemberGuid(A=%i && B=%i && C=%i && D=%i)"), Guid.A, Guid.B, Guid.C, Guid.D);
 		}
+
+		const FString VariableName = VarAction->GetVariableName().ToString();
+
+		// Search for both an explicit variable reference (finds get/sets of exactly that var, without including related-sounding variables)
+		// and a softer search for (VariableName) to capture bound component/widget event nodes which wouldn't otherwise show up
+		//@TODO: This logic is duplicated in SSCSEditor::OnFindReferences(), keep in sync
+		SearchTerm = FString::Printf(TEXT("Nodes(VariableReference(MemberName=+\"%s\"%s) || Name=\"(%s)\")"), *VariableName, *GuidTerm, *VariableName);
+		bUseQuotes = false;
 	}
 	else if (FEdGraphSchemaAction_K2LocalVar* LocalVarAction = SelectionAsLocalVar())
 	{
@@ -2776,7 +2732,7 @@ bool SMyBlueprint::CanRequestRenameOnActionNode() const
 	GraphActionMenu->GetSelectedActions(SelectedActions);
 
 	// If there is anything selected in the GraphActionMenu, check the item for if it can be renamed.
-	if(SelectedActions.Num())
+	if (SelectedActions.Num() || SelectionIsCategory())
 	{
 		return GraphActionMenu->CanRequestRenameOnActionNode();
 	}
@@ -2825,9 +2781,9 @@ bool SMyBlueprint::MoveCategoryBeforeCategory( const FText& InCategoryToMove, co
 	if( BlueprintObj )
 	{
 		// Find root categories
-		int32 RootCategoryDelim = CategoryToMoveString.Find( TEXT( "|" ));
+		int32 RootCategoryDelim = CategoryToMoveString.Find( TEXT( "|" ), ESearchCase::CaseSensitive );
 		FName CategoryToMove = RootCategoryDelim == INDEX_NONE ? *CategoryToMoveString : *CategoryToMoveString.Left( RootCategoryDelim );
-		RootCategoryDelim = TargetCategoryString.Find( TEXT( "|" ));
+		RootCategoryDelim = TargetCategoryString.Find( TEXT( "|" ), ESearchCase::CaseSensitive );
 		FName TargetCategory = RootCategoryDelim == INDEX_NONE ? *TargetCategoryString : *TargetCategoryString.Left( RootCategoryDelim );
 
 		TArray<FName>& CategorySort = BlueprintObj->CategorySorting;
diff --git a/Engine/Source/Editor/Kismet/Private/SMyBlueprint.h b/Engine/Source/Editor/Kismet/Private/SMyBlueprint.h
--- a/Engine/Source/Editor/Kismet/Private/SMyBlueprint.h
+++ b/Engine/Source/Editor/Kismet/Private/SMyBlueprint.h
@@ -24,6 +24,7 @@ class UUserDefinedEnum;
 class UUserDefinedStruct;
 struct FEdGraphSchemaAction_K2Struct;
 struct FGraphActionNode;
+struct FGraphActionSort;
 
 /* Enums to use when grouping the blueprint members in the list panel. The order here will determine the order in the list */
 namespace NodeSectionID
@@ -161,9 +162,9 @@ private:
 	/** Callback used to populate all actions list in SGraphActionMenu */
 	void CollectAllActions(FGraphActionListBuilderBase& OutAllActions);
 	void CollectStaticSections(TArray<int32>& StaticSectionIDs);
-	void GetChildGraphs(UEdGraph* EdGraph, int32 const SectionId, FGraphActionListBuilderBase& OutAllActions, FText ParentCategory = FText::GetEmpty());
-	void GetChildEvents(UEdGraph const* EdGraph, int32 const SectionId, FGraphActionListBuilderBase& OutAllActions, FText ParentCategory = FText::GetEmpty()) const;
-	void GetLocalVariables(FGraphActionListBuilderBase& OutAllActions) const;
+	void GetChildGraphs(UEdGraph* EdGraph, int32 const SectionId, FGraphActionSort& SortList, FText ParentCategory = FText::GetEmpty());
+	void GetChildEvents(UEdGraph const* EdGraph, int32 const SectionId, FGraphActionSort& SortList, FText ParentCategory = FText::GetEmpty()) const;
+	void GetLocalVariables(FGraphActionSort& SortList) const;
 	
 	/** Handles the visibility of the local action list */
 	EVisibility GetLocalActionsListVisibility() const;
diff --git a/Engine/Source/Editor/Kismet/Private/SSCSEditor.cpp b/Engine/Source/Editor/Kismet/Private/SSCSEditor.cpp
--- a/Engine/Source/Editor/Kismet/Private/SSCSEditor.cpp
+++ b/Engine/Source/Editor/Kismet/Private/SSCSEditor.cpp
@@ -1403,7 +1403,7 @@ void FSCSEditorTreeNodeRootActor::OnCompleteRename(const FText& InNewName)
 	if (Actor && Actor->IsActorLabelEditable() && !InNewName.ToString().Equals(Actor->GetActorLabel(), ESearchCase::CaseSensitive))
 	{
 		const FScopedTransaction Transaction(LOCTEXT("SCSEditorRenameActorTransaction", "Rename Actor"));
-		Actor->SetActorLabel(InNewName.ToString());
+		FActorLabelUtilities::RenameExistingActor(Actor, InNewName.ToString());
 	}
 }
 
@@ -3269,7 +3269,7 @@ FText SSCS_RowWidget_ActorRoot::GetActorDisplayText() const
 			{
 				FString Name;
 				UBlueprint* Blueprint = UBlueprint::GetBlueprintFromClass(DefaultActor->GetClass());
-				if(Blueprint != nullptr)
+				if(Blueprint != nullptr && SCSEditorPtr->GetEditorMode() != EComponentEditorMode::ActorInstance)
 				{
 					Blueprint->GetName(Name);
 				}
@@ -4035,7 +4035,11 @@ void SSCSEditor::OnFindReferences()
 		if (FoundAssetEditor.IsValid())
 		{
 			const FString VariableName = SelectedNodes[0]->GetVariableName().ToString();
-			const FString SearchTerm = FString::Printf(TEXT("Nodes(VariableReference(MemberName=+\"%s\"))"), *VariableName);
+
+			// Search for both an explicit variable reference (finds get/sets of exactly that var, without including related-sounding variables)
+			// and a softer search for (VariableName) to capture bound component/widget event nodes which wouldn't otherwise show up
+			//@TODO: This logic is duplicated in SMyBlueprint::OnFindReference(), keep in sync
+			const FString SearchTerm = FString::Printf(TEXT("Nodes(VariableReference(MemberName=+\"%s\") || Name=\"(%s)\")"), *VariableName, *VariableName);
 
 			TSharedRef<IBlueprintEditor> BlueprintEditor = StaticCastSharedRef<IBlueprintEditor>(FoundAssetEditor.ToSharedRef());
 			BlueprintEditor->SummonSearchUI(true, SearchTerm);
@@ -6162,7 +6166,7 @@ void SSCSEditor::OnOpenBlueprintEditor(bool bForceCodeEditing) const
 		{
 			if (bForceCodeEditing && (Blueprint->UbergraphPages.Num() > 0))
 			{
-				FKismetEditorUtilities::BringKismetToFocusAttentionOnObject(Blueprint->UbergraphPages[0]);
+				FKismetEditorUtilities::BringKismetToFocusAttentionOnObject(Blueprint->GetLastEditedUberGraph());
 			}
 			else
 			{
diff --git a/Engine/Source/Editor/Kismet/Private/UserDefinedEnumEditor.cpp b/Engine/Source/Editor/Kismet/Private/UserDefinedEnumEditor.cpp
--- a/Engine/Source/Editor/Kismet/Private/UserDefinedEnumEditor.cpp
+++ b/Engine/Source/Editor/Kismet/Private/UserDefinedEnumEditor.cpp
@@ -21,6 +21,7 @@
 #include "Widgets/Docking/SDockTab.h"
 #include "IDocumentation.h"
 #include "STextPropertyEditableTextBox.h"
+#include "ScopedTransaction.h"
 
 #define LOCTEXT_NAMESPACE "UserDefinedEnumEditor"
 
@@ -77,9 +78,8 @@ public:
 	virtual void SetText(const int32 InIndex, const FText& InText) override
 	{
 		check(InIndex == 0);
-		bCausedChange = true;
+		TGuardValue<bool> CausingChange(bCausedChange, true);
 		FEnumEditorUtils::SetEnumeratorDisplayName(TargetEnum, EnumeratorIndex, InText);
-		bCausedChange = false;
 	}
 
 	virtual bool IsValidText(const FText& InText, FText& OutErrorMsg) const override
@@ -129,6 +129,103 @@ private:
 	bool bCausedChange;
 };
 
+
+
+/** Allows STextPropertyEditableTextBox to edit the tooltip metadata for a user defined enum entry */
+class FEditableTextUserDefinedEnumTooltip : public IEditableTextProperty
+{
+public:
+	FEditableTextUserDefinedEnumTooltip(UUserDefinedEnum* InTargetEnum, const int32 InEnumeratorIndex)
+		: TargetEnum(InTargetEnum)
+		, EnumeratorIndex(InEnumeratorIndex)
+		, bCausedChange(false)
+	{
+	}
+
+	virtual bool IsMultiLineText() const override
+	{
+		return true;
+	}
+
+	virtual bool IsPassword() const override
+	{
+		return false;
+	}
+
+	virtual bool IsReadOnly() const override
+	{
+		return false;
+	}
+
+	virtual bool IsDefaultValue() const override
+	{
+		return false;
+	}
+
+	virtual FText GetToolTipText() const override
+	{
+		return FText::GetEmpty();
+	}
+
+	virtual int32 GetNumTexts() const override
+	{
+		return 1;
+	}
+
+	virtual FText GetText(const int32 InIndex) const override
+	{
+		check(InIndex == 0);
+		return TargetEnum->GetToolTipTextByIndex(EnumeratorIndex);
+	}
+
+	virtual void SetText(const int32 InIndex, const FText& InText) override
+	{
+		check(InIndex == 0);
+		TGuardValue<bool> CausingChange(bCausedChange, true);
+		//@TODO: Metadata is not transactional right now, so we cannot transact a tooltip edit
+		// const FScopedTransaction Transaction(NSLOCTEXT("EnumEditor", "SetEnumeratorTooltip", "Set Description"));
+		TargetEnum->Modify();
+		TargetEnum->SetMetaData(TEXT("ToolTip"), *InText.ToString(), EnumeratorIndex);
+	}
+
+	virtual bool IsValidText(const FText& InText, FText& OutErrorMsg) const override
+	{
+		return true;
+	}
+
+#if USE_STABLE_LOCALIZATION_KEYS
+	virtual void GetStableTextId(const int32 InIndex, const ETextPropertyEditAction InEditAction, const FString& InTextSource, const FString& InProposedNamespace, const FString& InProposedKey, FString& OutStableNamespace, FString& OutStableKey) const override
+	{
+		check(InIndex == 0);
+		return StaticStableTextId(TargetEnum, InEditAction, InTextSource, InProposedNamespace, InProposedKey, OutStableNamespace, OutStableKey);
+	}
+#endif // USE_STABLE_LOCALIZATION_KEYS
+
+	virtual void RequestRefresh() override
+	{
+	}
+
+	bool CausedChange() const
+	{
+		return bCausedChange;
+	}
+
+private:
+	/** The user defined enum being edited */
+	UUserDefinedEnum* TargetEnum;
+
+	/** Index of enumerator entry */
+	int32 EnumeratorIndex;
+
+	/** Set while we are invoking a change to the user defined enum */
+	bool bCausedChange;
+};
+
+
+
+
+
+
 void FUserDefinedEnumEditor::RegisterTabSpawners(const TSharedRef<class FTabManager>& InTabManager)
 {
 	WorkspaceMenuCategory = InTabManager->AddLocalWorkspaceMenuCategory(LOCTEXT("WorkspaceMenu_UserDefinedEnumEditor", "User-Defined Enum Editor"));
@@ -388,13 +485,15 @@ void FUserDefinedEnumLayout::GenerateChildContent( IDetailChildrenBuilder& Child
 
 bool FUserDefinedEnumIndexLayout::CausedChange() const
 {
-	return DisplayNameEditor.IsValid() && DisplayNameEditor->CausedChange();
+	return (DisplayNameEditor.IsValid() && DisplayNameEditor->CausedChange()) || (TooltipEditor.IsValid() && TooltipEditor->CausedChange());
 }
 
 void FUserDefinedEnumIndexLayout::GenerateHeaderRowContent( FDetailWidgetRow& NodeRow )
 {
 	DisplayNameEditor = MakeShared<FEditableTextUserDefinedEnum>(TargetEnum, EnumeratorIndex);
 
+	TooltipEditor = MakeShared<FEditableTextUserDefinedEnumTooltip>(TargetEnum, EnumeratorIndex);
+
 	const bool bIsEditable = !DisplayNameEditor->IsReadOnly();
 	const bool bIsMoveUpEnabled = (TargetEnum->NumEnums() != 1) && (EnumeratorIndex != 0) && bIsEditable;
 	const bool bIsMoveDownEnabled = (TargetEnum->NumEnums() != 1) && (EnumeratorIndex < TargetEnum->NumEnums() - 2) && bIsEditable;
@@ -408,13 +507,41 @@ void FUserDefinedEnumIndexLayout::GenerateHeaderRowContent( FDetailWidgetRow& No
 			SNew(SHorizontalBox)
 
 			+SHorizontalBox::Slot()
+			.VAlign(VAlign_Center)
+			.AutoWidth()
+			.Padding(0, 0, 4, 0)
+			[
+				SNew(STextBlock)
+				.Text(LOCTEXT("EnumDisplayNameLabel", "Display Name"))
+			]
+
+			+SHorizontalBox::Slot()
+			.VAlign(VAlign_Center)
+			.FillWidth(1.0f)
 			[
 				SNew(STextPropertyEditableTextBox, DisplayNameEditor.ToSharedRef())
 			]
 
+			+SHorizontalBox::Slot()
+			.VAlign(VAlign_Center)
+			.AutoWidth()
+			.Padding(4, 0, 4, 0)
+			[
+				SNew(STextBlock)
+				.Text(LOCTEXT("EnumTooltipLabel", "Description"))
+			]
+
+			+SHorizontalBox::Slot()
+			.VAlign(VAlign_Center)
+			.FillWidth(1.0f)
+			[
+				SNew(STextPropertyEditableTextBox, TooltipEditor.ToSharedRef())
+			]
+
 			+SHorizontalBox::Slot()
 			.AutoWidth()
 			.Padding(2, 0)
+			.VAlign(VAlign_Center)
 			[
 				SNew(SButton)
 				.ContentPadding(0)
@@ -428,6 +555,7 @@ void FUserDefinedEnumIndexLayout::GenerateHeaderRowContent( FDetailWidgetRow& No
 			+SHorizontalBox::Slot()
 			.AutoWidth()
 			.Padding(2, 0)
+			.VAlign(VAlign_Center)
 			[
 				SNew(SButton)
 				.ContentPadding(0)
@@ -442,6 +570,7 @@ void FUserDefinedEnumIndexLayout::GenerateHeaderRowContent( FDetailWidgetRow& No
 			+SHorizontalBox::Slot()
 			.AutoWidth()
 			.Padding(2, 0)
+			.VAlign(VAlign_Center)
 			[
 				ClearButton
 			]
diff --git a/Engine/Source/Editor/Kismet/Private/UserDefinedEnumEditor.h b/Engine/Source/Editor/Kismet/Private/UserDefinedEnumEditor.h
--- a/Engine/Source/Editor/Kismet/Private/UserDefinedEnumEditor.h
+++ b/Engine/Source/Editor/Kismet/Private/UserDefinedEnumEditor.h
@@ -19,6 +19,8 @@ class FUserDefinedEnumLayout;
 class IDetailChildrenBuilder;
 class IDetailLayoutBuilder;
 class SEditableTextBox;
+class FEditableTextUserDefinedEnum;
+class FEditableTextUserDefinedEnumTooltip;
 
 class KISMET_API FUserDefinedEnumEditor : public IUserDefinedEnumEditor
 {
@@ -172,5 +174,8 @@ private:
 	int32 EnumeratorIndex;
 
 	/** The editable text interface for the display name data */
-	TSharedPtr<class FEditableTextUserDefinedEnum> DisplayNameEditor;
+	TSharedPtr<FEditableTextUserDefinedEnum> DisplayNameEditor;
+
+	/** The editable text interface for the tooltip data */
+	TSharedPtr<FEditableTextUserDefinedEnumTooltip> TooltipEditor;
 };
diff --git a/Engine/Source/Editor/Kismet/Private/UserDefinedStructureEditor.cpp b/Engine/Source/Editor/Kismet/Private/UserDefinedStructureEditor.cpp
--- a/Engine/Source/Editor/Kismet/Private/UserDefinedStructureEditor.cpp
+++ b/Engine/Source/Editor/Kismet/Private/UserDefinedStructureEditor.cpp
@@ -283,7 +283,7 @@ void FUserDefinedStructureEditor::RegisterTabSpawners(const TSharedRef<class FTa
 	FAssetEditorToolkit::RegisterTabSpawners(InTabManager);
 
 	InTabManager->RegisterTabSpawner( MemberVariablesTabId, FOnSpawnTab::CreateSP(this, &FUserDefinedStructureEditor::SpawnStructureTab) )
-		.SetDisplayName( LOCTEXT("MemberVariablesEditor", "Member Variables") )
+		.SetDisplayName( LOCTEXT("MemberVariablesEditor", "Structure Editor") )
 		.SetGroup(WorkspaceMenuCategory.ToSharedRef())
 		.SetIcon(FSlateIcon(FEditorStyle::GetStyleSetName(), "Kismet.Tabs.Variables"));
 }
@@ -419,7 +419,7 @@ TSharedRef<SDockTab> FUserDefinedStructureEditor::SpawnStructureTab(const FSpawn
 
 	return SNew(SDockTab)
 		.Icon( FEditorStyle::GetBrush("GenericEditor.Tabs.Properties") )
-		.Label( LOCTEXT("UserDefinedStructureEditor", "Structure") )
+		.Label( LOCTEXT("UserDefinedStructureEditor", "Structure Editor") )
 		.TabColorScale( GetTabColorScale() )
 		[
 			Splitter
@@ -1110,7 +1110,7 @@ void FUserDefinedStructureDetails::CustomizeDetails(class IDetailLayoutBuilder&
 	{
 		UserDefinedStruct = CastChecked<UUserDefinedStruct>(Objects[0].Get());
 
-		IDetailCategoryBuilder& StructureCategory = DetailLayout.EditCategory("Structure", LOCTEXT("Structure", "Structure"));
+		IDetailCategoryBuilder& StructureCategory = DetailLayout.EditCategory("Structure", LOCTEXT("StructureCategory", "Structure"));
 		Layout = MakeShareable(new FUserDefinedStructureLayout(SharedThis(this)));
 		StructureCategory.AddCustomBuilder(Layout.ToSharedRef());
 	}
diff --git a/Engine/Source/Editor/Kismet/Public/BlueprintEditor.h b/Engine/Source/Editor/Kismet/Public/BlueprintEditor.h
--- a/Engine/Source/Editor/Kismet/Public/BlueprintEditor.h
+++ b/Engine/Source/Editor/Kismet/Public/BlueprintEditor.h
@@ -31,7 +31,6 @@ class IMessageLogListing;
 class INameValidatorInterface;
 class ISCSEditorCustomization;
 class SBlueprintPalette;
-class SBlueprintProfilerView;
 class SFindInBlueprints;
 class SKismetDebuggingView;
 class SKismetInspector;
@@ -40,6 +39,7 @@ class SReplaceNodeReferences;
 class SSCSEditor;
 class SSCSEditorViewport;
 class UEdGraph;
+class UEdGraphNode;
 class UUserDefinedEnum;
 class UUserDefinedStruct;
 struct Rect;
@@ -223,7 +223,6 @@ public:
 	TSharedRef<class SKismetInspector> GetInspector() const { return Inspector.ToSharedRef(); }
 	TSharedRef<class SKismetInspector> GetDefaultEditor() const { return DefaultEditor.ToSharedRef(); }
 	TSharedRef<class SKismetDebuggingView> GetDebuggingView() const { return DebuggingView.ToSharedRef(); }
-	TSharedRef<class SBlueprintProfilerView> GetBlueprintProfilerView() const { return BlueprintProfiler.ToSharedRef(); }
 	TSharedRef<class SBlueprintPalette> GetPalette() const { return Palette.ToSharedRef(); }
 	TSharedRef<class SWidget> GetCompilerResults() const { return CompilerResults.ToSharedRef(); }
 	TSharedRef<class SFindInBlueprints> GetFindResults() const { return FindResults.ToSharedRef(); }
@@ -282,6 +281,9 @@ public:
 	/** Get the currently selected set of nodes */
 	TSet<UObject*> GetSelectedNodes() const;
 
+	/** Returns the currently selected node if there is a single node selected (if there are multiple nodes selected or none selected, it will return nullptr) */
+	UEdGraphNode* GetSingleSelectedNode() const;
+
 	/** Save the current set of edited objects in the LastEditedObjects array so it will be opened next time we open K2 */
 	void SaveEditedObjectState();
 
@@ -321,18 +323,6 @@ public:
 	/** Handles opening the header file of native parent class */
 	void OnEditParentClassNativeCodeClicked();
 
-	/** Called to open native function definition of the current node selection in an IDE */
-	void GotoNativeFunctionDefinition();
-
-	/** Called to check if the current selection is a native function */
-	bool IsSelectionNativeFunction();
-
-	/** Called to open native variable declaration of the current node selection in an IDE */
-	void GotoNativeVariableDefinition();
-
-	/** Called to check if the current selection is a native variable */
-	bool IsSelectionNativeVariable();
-
 	/** Returns: "(<NativeParentClass>.h)" */
 	FText GetTextForNativeParentClassHeaderLink() const;
 
@@ -534,15 +524,6 @@ public:
 	/** Can generate native code for current blueprint */
 	bool CanGenerateNativeCode() const;
 
-	/** Returns if the blueprint profiler is enabled in editor settings. */
-	bool IsProfilerAvailable() const;
-
-	/** Returns if the blueprint profiler is currently active. */
-	bool IsProfilerActive() const;
-
-	/** Toggle blueprint profiler state */
-	void ToggleProfiler();
-
 	/**
 	 * Check to see if we can customize the SCS editor for the passed-in scene component
 	 * @param	InComponentToCustomize	The component to check to see if a customization exists
@@ -849,7 +830,7 @@ protected:
 	void OnListObjectsReferencedByBlueprint();
 	void OnRepairCorruptedBlueprint();
 
-	void OnNodeDoubleClicked(class UEdGraphNode* Node);
+	void OnNodeDoubleClicked(UEdGraphNode* Node);
 
 	virtual void OnEditTabClosed(TSharedRef<SDockTab> Tab);
 
@@ -1010,13 +991,10 @@ private:
 	/** Returns whether the currently focused graph is editable or not */
 	bool IsFocusedGraphEditable() const;
 
-	/** Called to check if native code browsing is available */
-	bool IsNativeCodeBrowsingAvailable() const;
-
-	/** Called when the user wants to jump to a node's graph definition */
+	/** Called when the user wants to jump to a node's definition */
 	void OnGoToDefinition();
 
-	/** Checks to see if it is possible to jump to the selected node's graph definition. */
+	/** Checks to see if it is possible to jump to the selected node's definition. */
 	bool CanGoToDefinition() const;
 
 	/** Open documentation for the selected node class */
@@ -1081,9 +1059,6 @@ protected:
 	/** Debugging window (watches, breakpoints, etc...) */
 	TSharedPtr<class SKismetDebuggingView> DebuggingView;
 
-	/** Performance Analsys View */
-	TSharedPtr<class SBlueprintProfilerView> BlueprintProfiler;
-
 	/** Palette of all classes with funcs/vars */
 	TSharedPtr<class SBlueprintPalette> Palette;
 
@@ -1148,13 +1123,13 @@ public:
 	TSharedPtr<SGraphEditor> OpenGraphAndBringToFront(UEdGraph* Graph);
 
 	//@TODO: To be moved/merged
-	TSharedPtr<SDockTab> OpenDocument(UObject* DocumentID, FDocumentTracker::EOpenDocumentCause Cause);
+	TSharedPtr<SDockTab> OpenDocument(const UObject* DocumentID, FDocumentTracker::EOpenDocumentCause Cause);
 
 	/** Finds the tab associated with the specified asset, and closes if it is open */
-	void CloseDocumentTab(UObject* DocumentID);
+	void CloseDocumentTab(const UObject* DocumentID);
 
 	// Finds any open tabs containing the specified document and adds them to the specified array; returns true if at least one is found
-	bool FindOpenTabsContainingDocument(UObject* DocumentID, /*inout*/ TArray< TSharedPtr<SDockTab> >& Results);
+	bool FindOpenTabsContainingDocument(const UObject* DocumentID, /*inout*/ TArray< TSharedPtr<SDockTab> >& Results);
 
 
 public:
@@ -1219,9 +1194,6 @@ private:
 	/** Handle to the registered OnActiveTabChanged delegate */
 	FDelegateHandle OnActiveTabChangedDelegateHandle;
 
-	/** Blueprint generated class instrumentation state */
-	bool bBlueprintHasInstrumentation;
-
 	// Allow derived editors to add command mappings 
 	virtual void OnCreateGraphEditorCommands(TSharedPtr<FUICommandList> GraphEditorCommandsList) {}
 };
diff --git a/Engine/Source/Editor/Kismet/Public/BlueprintEditorTabs.h b/Engine/Source/Editor/Kismet/Public/BlueprintEditorTabs.h
--- a/Engine/Source/Editor/Kismet/Public/BlueprintEditorTabs.h
+++ b/Engine/Source/Editor/Kismet/Public/BlueprintEditorTabs.h
@@ -18,7 +18,6 @@ struct KISMET_API FBlueprintEditorTabs
 	static const FName MyBlueprintID;
 	static const FName ReplaceNodeReferencesID;
 	static const FName UserDefinedStructureID;
-	static const FName BlueprintProfilerID;
 
 	// Document tab identifiers
 	static const FName GraphEditorID;
diff --git a/Engine/Source/Editor/Kismet/Public/SBlueprintEditorToolbar.h b/Engine/Source/Editor/Kismet/Public/SBlueprintEditorToolbar.h
--- a/Engine/Source/Editor/Kismet/Public/SBlueprintEditorToolbar.h
+++ b/Engine/Source/Editor/Kismet/Public/SBlueprintEditorToolbar.h
@@ -58,9 +58,6 @@ public:
 	TSharedPtr<FUICommandInfo> SaveOnCompile_Always;
 	TSharedPtr<FUICommandInfo> JumpToErrorNode;
 
-	/** Profiler options  */
-	TSharedPtr<FUICommandInfo> ToggleProfiler;
-
 	/** Switch between modes in the blueprint editor */
 	TSharedPtr<FUICommandInfo> SwitchToScriptingMode;
 	TSharedPtr<FUICommandInfo> SwitchToBlueprintDefaultsMode;
@@ -89,14 +86,10 @@ public:
 	void AddScriptingToolbar(TSharedPtr<FExtender> Extender);
 	void AddDebuggingToolbar(TSharedPtr<FExtender> Extender);
 	void AddComponentsToolbar(TSharedPtr<FExtender> Extender);
-	void AddProfilerToolbar(TSharedPtr<FExtender> Extender);
 
 	/** Returns the current status icon for the blueprint being edited */
 	FSlateIcon GetStatusImage() const;
 
-	/** Returns the current profiler status tool tip */
-	FText GetProfilerStatusTooltip() const;
-
 	/** Returns the current status as text for the blueprint being edited */
 	FText GetStatusTooltip() const;
 
@@ -111,7 +104,6 @@ private:
 	void FillScriptingToolbar(FToolBarBuilder& ToolbarBuilder);
 	void FillDebuggingToolbar(FToolBarBuilder& ToolbarBuilder);
 	void FillComponentsToolbar(FToolBarBuilder& ToolbarBuilder);
-	void FillProfilerToolbar(FToolBarBuilder& ToolbarBuilder);
 
 protected:
 	/** Pointer back to the blueprint editor tool that owns us */
diff --git a/Engine/Source/Editor/Kismet/Public/WorkflowOrientedApp/WorkflowUObjectDocuments.h b/Engine/Source/Editor/Kismet/Public/WorkflowOrientedApp/WorkflowUObjectDocuments.h
--- a/Engine/Source/Editor/Kismet/Public/WorkflowOrientedApp/WorkflowUObjectDocuments.h
+++ b/Engine/Source/Editor/Kismet/Public/WorkflowOrientedApp/WorkflowUObjectDocuments.h
@@ -19,9 +19,9 @@ struct FTabPayload_UObject : public FTabPayload
 {
 public:
 	// Create a new payload wrapper for a UObject
-	static TSharedRef<FTabPayload_UObject> Make(UObject* DocumentID)
+	static TSharedRef<FTabPayload_UObject> Make(const UObject* DocumentID)
 	{
-		return MakeShareable(new FTabPayload_UObject(DocumentID));
+		return MakeShareable(new FTabPayload_UObject(const_cast<UObject*>(DocumentID)));
 	}
 
 	// Helper method to get the payload object as a specific type
diff --git a/Engine/Source/Editor/KismetCompiler/Private/KismetCompiler.cpp b/Engine/Source/Editor/KismetCompiler/Private/KismetCompiler.cpp
--- a/Engine/Source/Editor/KismetCompiler/Private/KismetCompiler.cpp
+++ b/Engine/Source/Editor/KismetCompiler/Private/KismetCompiler.cpp
@@ -1858,18 +1858,7 @@ void FKismetCompilerContext::CompileFunction(FKismetFunctionContext& Context)
 					// debug data for the behind-the-scenes entry point, only for the user-visible ones.
 					bEmitDebuggingSite = false;
 				}
-				if (Context.bInstrumentScriptCode)
-				{
-					for (UEdGraphPin* Pin : Node->Pins)
-					{
-						if (Pin->Direction == EGPD_Input && Pin->PinType.PinCategory == UEdGraphSchema_K2::PC_Exec)
-						{
-							ExecPin = Pin;
-							bEmitDebuggingSite = Context.IsInstrumentationRequiredForPin(ExecPin);
-							break;
-						}
-					}
-				}
+
 				if (bEmitDebuggingSite)
 				{
 					FBlueprintCompiledStatement& Statement = Context.AppendStatementForNode(Node);
@@ -1878,11 +1867,6 @@ void FKismetCompilerContext::CompileFunction(FKismetFunctionContext& Context)
 					Statement.Comment = NodeComment;
 				}
 			}
-			else if (Context.IsInstrumentationRequiredForNode(Node))
-			{
-				FBlueprintCompiledStatement& Statement = Context.AppendStatementForNode(Node);
-				Statement.Type = KCST_InstrumentedPureNodeEntry;
-			}
 		}
 
 		// Let the node handlers try to compile it
@@ -1944,8 +1928,7 @@ void FKismetCompilerContext::CompileFunction(FKismetFunctionContext& Context)
 						for (UEdGraphPin* LinkedTo : Pin->LinkedTo)
 						{
 							UEdGraphNode* NodeUsingOutput = LinkedTo->GetOwningNode();
-							const bool bValidExpansionLink = !Context.bInstrumentScriptCode || Context.IsInstrumentationRequiredForPin(LinkedTo);
-							if (NodeUsingOutput != nullptr && bValidExpansionLink)
+							if (NodeUsingOutput != nullptr)
 							{
 								// Add this node, as well as other nodes this node depends on
 								TSet<UEdGraphNode*>& TargetNodesRequired = PureNodesNeeded.FindOrAdd(NodeUsingOutput);
@@ -1977,12 +1960,6 @@ void FKismetCompilerContext::CompileFunction(FKismetFunctionContext& Context)
 					OrderedInsertIntoArray(SortedPureNodes, SortKeyMap, *It);
 				}
 
-				if (Context.bInstrumentScriptCode)
-				{
-					FBlueprintCompiledStatement& PopState = Context.PrependStatementForNode(Node);
-					PopState.Type = KCST_InstrumentedStatePop;
-				}
-
 				// Inline their code
 				for (int32 i = 0; i < SortedPureNodes.Num(); ++i)
 				{
@@ -1990,12 +1967,6 @@ void FKismetCompilerContext::CompileFunction(FKismetFunctionContext& Context)
 
 					Context.CopyAndPrependStatements(Node, NodeToInline);
 				}
-
-				if (Context.bInstrumentScriptCode)
-				{
-					FBlueprintCompiledStatement& PushState = Context.PrependStatementForNode(Node);
-					PushState.Type = KCST_InstrumentedStatePush;
-				}
 			}
 
 			// Proceed to the next node
@@ -2285,6 +2256,27 @@ void FKismetCompilerContext::FinishCompilingClass(UClass* Class)
 	//@TODO: Not sure if doing this again is actually necessary
 	// It will be if locals get promoted to class scope during function compilation, but that should ideally happen during Precompile or similar
 	Class->Bind();
+	
+	// Ensure that function netflags equate to any super function in a parent BP prior to linking; it may have been changed by the user
+	// and won't be reflected in the child class until it is recompiled. Without this, UClass::Link() will assert if they are out of sync.
+	for (UField* Field = Class->Children; Field; Field = Field->Next)
+	{
+		UFunction* Function = dynamic_cast<UFunction*>(Field);
+		if (Function != nullptr)
+		{
+			UFunction* ParentFunction = Function->GetSuperFunction();
+			if (ParentFunction != nullptr)
+			{
+				const EFunctionFlags ParentNetFlags = (ParentFunction->FunctionFlags & FUNC_NetFuncFlags);
+				if (ParentNetFlags != (Function->FunctionFlags & FUNC_NetFuncFlags))
+				{
+					Function->FunctionFlags &= ~FUNC_NetFuncFlags;
+					Function->FunctionFlags |= ParentNetFlags;
+				}
+			}
+		}
+	}
+
 	Class->StaticLink(true);
 
 	// Create the default object for this class
@@ -2357,15 +2349,6 @@ void FKismetCompilerContext::CreatePinEventNodeForTimelineFunction(UK2Node_Timel
 	{
 		MovePinLinksToIntermediate(*UpdatePin, *UpdateOutput);
 	}
-
-	if (CompileOptions.bAddInstrumentation)
-	{
-		// Register the timeline event nodes to any composite instances the timeline is located in.
-		if (UEdGraphNode* TunnelInstance = MessageLog.GetIntermediateTunnelInstance(TimelineNode))
-		{
-			MessageLog.RegisterIntermediateTunnelNode(TimelineEventNode, TunnelInstance);
-		}
-	}
 }
 
 UK2Node_CallFunction* FKismetCompilerContext::CreateCallTimelineFunction(UK2Node_Timeline* TimelineNode, UEdGraph* SourceGraph, FName FunctionName, UEdGraphPin* TimelineVarPin, UEdGraphPin* TimelineFunctionPin)
@@ -2663,7 +2646,7 @@ void FKismetCompilerContext::CreateFunctionStubForEvent(UK2Node_Event* SrcEventN
 	ChildStubGraph->SetFlags(RF_Transient);
 	MessageLog.NotifyIntermediateObjectCreation(ChildStubGraph, SrcEventNode);
 
-	FKismetFunctionContext& StubContext = *new (FunctionList) FKismetFunctionContext(MessageLog, Schema, NewClass, Blueprint, CompileOptions.DoesRequireCppCodeGeneration(), false);
+	FKismetFunctionContext& StubContext = *new (FunctionList) FKismetFunctionContext(MessageLog, Schema, NewClass, Blueprint, CompileOptions.DoesRequireCppCodeGeneration());
 	StubContext.SourceGraph = ChildStubGraph;
 
 	// A stub graph has no visual representation and is thus not suited to be debugged via the debugger
@@ -2808,12 +2791,11 @@ void FKismetCompilerContext::MergeUbergraphPagesIn(UEdGraph* Ubergraph)
 	for (TArray<UEdGraph*>::TIterator It(Blueprint->UbergraphPages); It; ++It)
 	{
 		UEdGraph* SourceGraph = *It;
-		const bool bCreateTunnelBoundries = CompileOptions.IsInstrumentationActive();
 
 		if (CompileOptions.bSaveIntermediateProducts)
 		{
 			TArray<UEdGraphNode*> ClonedNodeList;
-			FEdGraphUtilities::CloneAndMergeGraphIn(Ubergraph, SourceGraph, MessageLog, /*bRequireSchemaMatch=*/ true, /*bIsCompiling*/ true, bCreateTunnelBoundries/*bCreateTunnelBoundries*/, &ClonedNodeList);
+			FEdGraphUtilities::CloneAndMergeGraphIn(Ubergraph, SourceGraph, MessageLog, /*bRequireSchemaMatch=*/ true, /*bIsCompiling*/ true, &ClonedNodeList);
 
 			// Create a comment block around the ubergrapgh contents before anything else got started
 			int32 OffsetX = 0;
@@ -2831,7 +2813,7 @@ void FKismetCompilerContext::MergeUbergraphPagesIn(UEdGraph* Ubergraph)
 		}
 		else
 		{
-			FEdGraphUtilities::CloneAndMergeGraphIn(Ubergraph, SourceGraph, MessageLog, /*bRequireSchemaMatch=*/ true, /*bIsCompiling*/ true, bCreateTunnelBoundries/*bCreateTunnelBoundries*/);
+			FEdGraphUtilities::CloneAndMergeGraphIn(Ubergraph, SourceGraph, MessageLog, /*bRequireSchemaMatch=*/ true, /*bIsCompiling*/ true);
 		}
 	}
 }
@@ -2867,46 +2849,7 @@ void FKismetCompilerContext::ExpansionStep(UEdGraph* Graph, bool bAllowUbergraph
 			UK2Node* Node = Cast<UK2Node>(Graph->Nodes[NodeIndex]);
 			if (Node)
 			{
-				if (!CompileOptions.bAddInstrumentation)
-				{
-					Node->ExpandNode(*this, Graph);
-				}
-				else
-				{
-					// Determine links before expansion
-					TMap<UEdGraphPin*, UEdGraphPin*> SourceNodeLinks;
-					DetermineNodeExecLinks(Node, SourceNodeLinks);
-					// track the current node count so we can figure out if expansion happened.
-					int32 ExpansionNodeIdx = Graph->Nodes.Num();
-					Node->ExpandNode(*this, Graph);
-					// Keep track of new expansion nodes
-					if (ExpansionNodeIdx < Graph->Nodes.Num())
-					{
-						bool bImpureExpansionNodeAdded = false;
-						TArray<UEdGraphNode*> ExpansionNodes;
-						for (; ExpansionNodeIdx < Graph->Nodes.Num(); ++ExpansionNodeIdx)
-						{
-							ExpansionNodes.Add(Graph->Nodes[ExpansionNodeIdx]);
-							MessageLog.RegisterExpansionNode(Graph->Nodes[ExpansionNodeIdx]);
-							// We don't want to create boundary nodes around a set of pure only expansion nodes.
-							bImpureExpansionNodeAdded |= !IsNodePure(Graph->Nodes[ExpansionNodeIdx]);
-							// Add the relevent pins to the trace suppression filter
-							for (UEdGraphPin* ExpansionPin : Graph->Nodes[ExpansionNodeIdx]->Pins)
-							{
-								for (UEdGraphPin* ExpansionPinLink : ExpansionPin->LinkedTo)
-								{
-									MessageLog.AddPinTraceFilter(ExpansionPin);
-								}
-							}
-						}
-						// Create boundaries around expsnion nodes so we can simulate pin traces correctly
-						if (bImpureExpansionNodeAdded)
-						{
-							UK2Node_TunnelBoundary::CreateBoundariesForExpansionNodes(Node, ExpansionNodes, SourceNodeLinks, MessageLog);
-							ExpansionNodeIdx = Graph->Nodes.Num();
-						}
-					}
-				}
+				Node->ExpandNode(*this, Graph);
 			}
 		}
 	}
@@ -3151,7 +3094,7 @@ void FKismetCompilerContext::CreateAndProcessUbergraph()
 
 		// Do some cursory validation (pin types match, inputs to outputs, pins never point to their parent node, etc...)
 		{
-			UbergraphContext = new (FunctionList)FKismetFunctionContext(MessageLog, Schema, NewClass, Blueprint, CompileOptions.DoesRequireCppCodeGeneration(), CompileOptions.IsInstrumentationActive());
+			UbergraphContext = new (FunctionList) FKismetFunctionContext(MessageLog, Schema, NewClass, Blueprint, CompileOptions.DoesRequireCppCodeGeneration());
 			UbergraphContext->SourceGraph = ConsolidatedEventGraph;
 			UbergraphContext->MarkAsEventGraph();
 			UbergraphContext->MarkAsInternalOrCppUseOnly();
@@ -3296,12 +3239,6 @@ void FKismetCompilerContext::ExpandTunnelsAndMacros(UEdGraph* SourceGraph)
 				MessageLog.RegisterIntermediateTunnelInstance(MacroInstanceNode, TunnelInstances);
 			}
 
-			if (CompileOptions.IsInstrumentationActive())
-			{
-				// Create tunnel Boundary signals
-				UK2Node_TunnelBoundary::CreateBoundaryNodesForTunnelInstance(MacroInstanceNode, ClonedGraph, MessageLog);
-			}
-
 			for (int32 I = 0; I < ClonedGraph->Nodes.Num(); ++I)
 			{
 				MacroGeneratedNodes.Add(ClonedGraph->Nodes[I], CurrentNode);
@@ -3531,7 +3468,7 @@ void FKismetCompilerContext::ProcessOneFunctionGraph(UEdGraph* SourceGraph, bool
 	if ((CompileOptions.CompileType == EKismetCompileType::SkeletonOnly) || ValidateGraphIsWellFormed(FunctionGraph))
 	{
 		const UEdGraphSchema_K2* FunctionGraphSchema = CastChecked<const UEdGraphSchema_K2>(FunctionGraph->GetSchema());
-		FKismetFunctionContext& Context = *new (FunctionList)FKismetFunctionContext(MessageLog, FunctionGraphSchema, NewClass, Blueprint, CompileOptions.DoesRequireCppCodeGeneration(), CompileOptions.IsInstrumentationActive());
+		FKismetFunctionContext& Context = *new (FunctionList) FKismetFunctionContext(MessageLog, FunctionGraphSchema, NewClass, Blueprint, CompileOptions.DoesRequireCppCodeGeneration());
 		Context.SourceGraph = FunctionGraph;
 
 		if(FBlueprintEditorUtils::IsDelegateSignatureGraph(SourceGraph))
@@ -3633,7 +3570,7 @@ void FKismetCompilerContext::CreateFunctionList()
 
 FKismetFunctionContext* FKismetCompilerContext::CreateFunctionContext()
 {
-	return new (FunctionList)FKismetFunctionContext(MessageLog, Schema, NewClass, Blueprint, CompileOptions.DoesRequireCppCodeGeneration(), CompileOptions.IsInstrumentationActive());
+	return new (FunctionList) FKismetFunctionContext(MessageLog, Schema, NewClass, Blueprint, CompileOptions.DoesRequireCppCodeGeneration());
 }
 
 /** Compile a blueprint into a class and a set of functions */
@@ -3749,13 +3686,6 @@ void FKismetCompilerContext::CompileClassLayout(EInternalCompilerFlags InternalF
 		}
 	}
 
-	// Declare class wants profiling data.
-	if (CompileOptions.IsInstrumentationActive())
-	{
-		TargetClass->bHasInstrumentation = true;
-		MessageLog.bTreatCompositeGraphsAsTunnels =	true;
-	}
-
 	// Early validation
 	if (CompileOptions.CompileType == EKismetCompileType::Full)
 	{
@@ -4354,7 +4284,13 @@ void FKismetCompilerContext::CompileFunctions(EInternalCompilerFlags InternalFla
 
 					UObject* SuperStruct = Struct->GetSuperStruct();
 					Ar << SuperStruct;
-					Ar << Struct->Children;
+
+					UField* ChildrenIter = Struct->Children;
+					while(ChildrenIter)
+					{
+						Ar << ChildrenIter;
+						ChildrenIter = ChildrenIter->Next;
+					}
 
 					if (UFunction* Function = Cast<UFunction>(Struct))
 					{
diff --git a/Engine/Source/Editor/KismetCompiler/Private/KismetCompilerMisc.cpp b/Engine/Source/Editor/KismetCompiler/Private/KismetCompilerMisc.cpp
--- a/Engine/Source/Editor/KismetCompiler/Private/KismetCompilerMisc.cpp
+++ b/Engine/Source/Editor/KismetCompiler/Private/KismetCompilerMisc.cpp
@@ -72,7 +72,7 @@ static bool IsTypeCompatibleWithProperty(UEdGraphPin* SourcePin, const FEdGraphP
 		UEnumProperty* EnumProperty = Cast<UEnumProperty>(TestProperty);
 		bTypeMismatch = (ByteProperty == nullptr) && (EnumProperty == nullptr || !EnumProperty->GetUnderlyingProperty()->IsA<UByteProperty>());
 	}
-	else if ((PinCategory == Schema->PC_Class) || (PinCategory == Schema->PC_AssetClass))
+	else if ((PinCategory == Schema->PC_Class) || (PinCategory == Schema->PC_SoftClass))
 	{
 		const UClass* ClassType = (PinSubCategory == Schema->PSC_Self) ? SelfClass : Cast<const UClass>(PinSubCategoryObject);
 
@@ -87,9 +87,9 @@ static bool IsTypeCompatibleWithProperty(UEdGraphPin* SourcePin, const FEdGraphP
 			{
 				MetaClass = ClassProperty->MetaClass;
 			}
-			else if (auto AssetClassProperty = Cast<UAssetClassProperty>(TestProperty))
+			else if (auto SoftClassProperty = Cast<USoftClassProperty>(TestProperty))
 			{
-				MetaClass = AssetClassProperty->MetaClass;
+				MetaClass = SoftClassProperty->MetaClass;
 			}
 
 			if (MetaClass != NULL)
@@ -102,7 +102,7 @@ static bool IsTypeCompatibleWithProperty(UEdGraphPin* SourcePin, const FEdGraphP
 				// It matches if it's an exact match or if the output class is more derived than the input class
 				bTypeMismatch = bSubtypeMismatch = !((OutputClass == InputClass) || (OutputClass->IsChildOf(InputClass)));
 
-				if ((PinCategory == Schema->PC_AssetClass) && (!TestProperty->IsA<UAssetClassProperty>()))
+				if ((PinCategory == Schema->PC_SoftClass) && (!TestProperty->IsA<USoftClassProperty>()))
 				{
 					bTypeMismatch = true;
 				}
@@ -137,7 +137,7 @@ static bool IsTypeCompatibleWithProperty(UEdGraphPin* SourcePin, const FEdGraphP
 			&& PropertyDelegate->SignatureFunction 
 			&& PropertyDelegate->SignatureFunction->IsSignatureCompatibleWith(SignatureFunction));
 	}
-	else if ((PinCategory == Schema->PC_Object) || (PinCategory == Schema->PC_Interface) || (PinCategory == Schema->PC_Asset))
+	else if ((PinCategory == Schema->PC_Object) || (PinCategory == Schema->PC_Interface) || (PinCategory == Schema->PC_SoftObject))
 	{
 		const UClass* ObjectType = (PinSubCategory == Schema->PSC_Self) ? SelfClass : Cast<const UClass>(PinSubCategoryObject);
 
@@ -184,7 +184,7 @@ static bool IsTypeCompatibleWithProperty(UEdGraphPin* SourcePin, const FEdGraphP
 				// It matches if it's an exact match or if the output class is more derived than the input class
 				bTypeMismatch = bSubtypeMismatch = !((OutputClass == InputClass) || (OutputClass->IsChildOf(InputClass)));
 
-				if ((PinCategory == Schema->PC_Asset) && (!TestProperty->IsA<UAssetObjectProperty>()))
+				if ((PinCategory == Schema->PC_SoftObject) && (!TestProperty->IsA<USoftObjectProperty>()))
 				{
 					bTypeMismatch = true;
 				}
@@ -863,7 +863,7 @@ UProperty* FKismetCompilerUtilities::CreatePrimitiveProperty(UObject* PropertySc
 
 	UProperty* NewProperty = nullptr;
 	//@TODO: Nasty string if-else tree
-	if ((PinCategory == Schema->PC_Object) || (PinCategory == Schema->PC_Interface) || (PinCategory == Schema->PC_Asset))
+	if ((PinCategory == Schema->PC_Object) || (PinCategory == Schema->PC_Interface) || (PinCategory == Schema->PC_SoftObject))
 	{
 		UClass* SubType = (PinSubCategory == Schema->PSC_Self) ? SelfClass : Cast<UClass>(PinSubCategoryObject);
 
@@ -891,9 +891,9 @@ UProperty* FKismetCompilerUtilities::CreatePrimitiveProperty(UObject* PropertySc
 			{
 				UObjectPropertyBase* NewPropertyObj = NULL;
 
-				if (PinCategory == Schema->PC_Asset)
+				if (PinCategory == Schema->PC_SoftObject)
 				{
-					NewPropertyObj = NewObject<UAssetObjectProperty>(PropertyScope, ValidatedPropertyName, ObjectFlags);
+					NewPropertyObj = NewObject<USoftObjectProperty>(PropertyScope, ValidatedPropertyName, ObjectFlags);
 				}
 				else if (bIsWeakPointer)
 				{
@@ -946,7 +946,7 @@ UProperty* FKismetCompilerUtilities::CreatePrimitiveProperty(UObject* PropertySc
 			}
 		}
 	}
-	else if ((PinCategory == Schema->PC_Class) || (PinCategory == Schema->PC_AssetClass))
+	else if ((PinCategory == Schema->PC_Class) || (PinCategory == Schema->PC_SoftClass))
 	{
 		UClass* SubType = Cast<UClass>(PinSubCategoryObject);
 
@@ -964,16 +964,16 @@ UProperty* FKismetCompilerUtilities::CreatePrimitiveProperty(UObject* PropertySc
 
 		if (SubType)
 		{
-			if (PinCategory == Schema->PC_AssetClass)
+			if (PinCategory == Schema->PC_SoftClass)
 			{
-				UAssetClassProperty* AssetClassProperty = NewObject<UAssetClassProperty>(PropertyScope, ValidatedPropertyName, ObjectFlags);
+				USoftClassProperty* SoftClassProperty = NewObject<USoftClassProperty>(PropertyScope, ValidatedPropertyName, ObjectFlags);
 				// we want to use this setter function instead of setting the 
 				// MetaClass member directly, because it properly handles  
 				// placeholder classes (classes that are stubbed in during load)
-				AssetClassProperty->SetMetaClass(SubType);
-				AssetClassProperty->PropertyClass = UClass::StaticClass();
-				AssetClassProperty->SetPropertyFlags(CPF_HasGetValueTypeHash);
-				NewProperty = AssetClassProperty;
+				SoftClassProperty->SetMetaClass(SubType);
+				SoftClassProperty->PropertyClass = UClass::StaticClass();
+				SoftClassProperty->SetPropertyFlags(CPF_HasGetValueTypeHash);
+				NewProperty = SoftClassProperty;
 			}
 			else
 			{
@@ -1559,7 +1559,7 @@ FBlueprintCompiledStatement& FNodeHandlingFunctor::GenerateSimpleThenGoto(FKisme
 		TargetNode = ThenExecPin->LinkedTo[0]->GetOwningNode();
 	}
 
-	if (Context.bCreateDebugData && (!Context.bInstrumentScriptCode || Context.IsInstrumentationRequiredForPin(ThenExecPin)))
+	if (Context.bCreateDebugData)
 	{
 		FBlueprintCompiledStatement& TraceStatement = Context.AppendStatementForNode(&Node);
 		TraceStatement.Type = Context.GetWireTraceType();
@@ -1704,7 +1704,7 @@ FString FNetNameMapping::MakeBaseName(const UAnimGraphNode_Base* Net)
 //////////////////////////////////////////////////////////////////////////
 // FKismetFunctionContext
 
-FKismetFunctionContext::FKismetFunctionContext(FCompilerResultsLog& InMessageLog, const UEdGraphSchema_K2* InSchema, UBlueprintGeneratedClass* InNewClass, UBlueprint* InBlueprint, bool bInGeneratingCpp, bool bInWantsInstrumentation)
+FKismetFunctionContext::FKismetFunctionContext(FCompilerResultsLog& InMessageLog, const UEdGraphSchema_K2* InSchema, UBlueprintGeneratedClass* InNewClass, UBlueprint* InBlueprint, bool bInGeneratingCpp)
 	: Blueprint(InBlueprint)
 	, SourceGraph(nullptr)
 	, EntryPoint(nullptr)
@@ -1718,7 +1718,6 @@ FKismetFunctionContext::FKismetFunctionContext(FCompilerResultsLog& InMessageLog
 	, bIsInterfaceStub(false)
 	, bIsConstFunction(false)
 	, bEnforceConstCorrectness(false)
-	, bInstrumentScriptCode(bInWantsInstrumentation)
 	// only need debug-data when running in the editor app:
 	, bCreateDebugData(GIsEditor && !IsRunningCommandlet())
 	, bIsSimpleStubGraphWithNoParams(false)
@@ -1884,10 +1883,8 @@ void FKismetFunctionContext::ResolveGotoFixups()
 			{
 				continue;
 			}
-			if (!bInstrumentScriptCode || IsInstrumentationRequiredForPin(GotoIt->Value))
-			{
-				InsertWireTrace(GotoIt.Key(), GotoIt.Value());
-			}
+
+			InsertWireTrace(GotoIt.Key(), GotoIt.Value());
 		}
 	}
 
diff --git a/Engine/Source/Editor/KismetCompiler/Private/KismetCompilerModule.cpp b/Engine/Source/Editor/KismetCompiler/Private/KismetCompilerModule.cpp
--- a/Engine/Source/Editor/KismetCompiler/Private/KismetCompilerModule.cpp
+++ b/Engine/Source/Editor/KismetCompiler/Private/KismetCompilerModule.cpp
@@ -133,20 +133,6 @@ void FKismet2CompilerModule::CompileBlueprintInner(class UBlueprint* Blueprint,
 			}
 		}
 
-		// we need to add instrumentation to any inherited blueprints.
-		if (Reinstancer.IsValid() && CompileOptions.IsInstrumentationActive())
-		{
-			TArray<UBlueprint*> InheritedBlueprints;
-			Blueprint->GetBlueprintHierarchyFromClass(Cast<UClass>(Blueprint->GeneratedClass), InheritedBlueprints);
-			for (auto CurrentBP : InheritedBlueprints)
-			{
-				if (CurrentBP != Blueprint)
-				{
-					Reinstancer->EnlistDependentBlueprintToRecompile(CurrentBP, true);
-				}
-			}
-		}
-
 		if (bRecompileDependencies)
 		{
 			Reinstancer->BlueprintWasRecompiled(Blueprint, CompileOptions.CompileType == EKismetCompileType::BytecodeOnly);
diff --git a/Engine/Source/Editor/KismetCompiler/Private/KismetCompilerVMBackend.cpp b/Engine/Source/Editor/KismetCompiler/Private/KismetCompilerVMBackend.cpp
--- a/Engine/Source/Editor/KismetCompiler/Private/KismetCompilerVMBackend.cpp
+++ b/Engine/Source/Editor/KismetCompiler/Private/KismetCompilerVMBackend.cpp
@@ -10,7 +10,7 @@
 #include "UObject/ObjectMacros.h"
 #include "Serialization/ArchiveUObject.h"
 #include "UObject/Class.h"
-#include "UObject/AssetPtr.h"
+#include "UObject/SoftObjectPtr.h"
 #include "UObject/Interface.h"
 #include "UObject/UnrealType.h"
 #include "UObject/TextProperty.h"
@@ -95,12 +95,12 @@ public:
 		return FArchive::operator<<(LazyObjectPtr);
 	}
 
-	FArchive& operator<<(FAssetPtr& AssetPtr) override
+	FArchive& operator<<(FSoftObjectPtr& Value) override
 	{
-		return FArchive::operator<<(AssetPtr);
+		return FArchive::operator<<(Value);
 	}
 
-	FArchive& operator<<(FStringAssetReference& Value) override
+	FArchive& operator<<(FSoftObjectPath& Value) override
 	{
 		return FArchiveUObject::operator<<(Value);
 	}
@@ -536,13 +536,13 @@ public:
 			return Type && (Type->PinCategory == UEdGraphSchema_K2::PC_Delegate);
 		}
 
-		static bool IsAsset(const FEdGraphPinType* Type, const UProperty* Property)
+		static bool IsSoftObject(const FEdGraphPinType* Type, const UProperty* Property)
 		{
 			if (Property)
 			{
-				return Property->IsA<UAssetObjectProperty>();
+				return Property->IsA<USoftObjectProperty>();
 			}
-			return Type && (Type->PinCategory == UEdGraphSchema_K2::PC_Asset);
+			return Type && (Type->PinCategory == UEdGraphSchema_K2::PC_SoftObject);
 		}
 
 		// Will handle Class properties as well
@@ -972,9 +972,9 @@ public:
 					Writer << FunctionName;
 				}
 			}
-			else if (FLiteralTypeHelper::IsAsset(&Term->Type, CoerceProperty))
+			else if (FLiteralTypeHelper::IsSoftObject(&Term->Type, CoerceProperty))
 			{
-				Writer << EX_AssetConst;
+				Writer << EX_SoftObjectConst;
 				EmitStringLiteral(Term->Name);
 			}
 			else if (FLiteralTypeHelper::IsObject(&Term->Type, CoerceProperty) || FLiteralTypeHelper::IsClass(&Term->Type, CoerceProperty))
@@ -1354,7 +1354,7 @@ public:
 		const bool bIsMulticastDelegate = Schema->PC_MCDelegate == DestinationExpression->Type.PinCategory;
 		const bool bIsBoolean = Schema->PC_Boolean == DestinationExpression->Type.PinCategory;
 		const bool bIsObj = (Schema->PC_Object == DestinationExpression->Type.PinCategory) || (Schema->PC_Class == DestinationExpression->Type.PinCategory);
-		const bool bIsAsset = Schema->PC_Asset == DestinationExpression->Type.PinCategory;
+		const bool bIsSoftObject = Schema->PC_SoftObject == DestinationExpression->Type.PinCategory;
 		const bool bIsWeakObjPtr = DestinationExpression->Type.bIsWeakPointer;
 
 		if (bIsContainer)
@@ -1375,7 +1375,7 @@ public:
 		{
 			Writer << EX_LetBool;
 		}
-		else if (bIsObj && !bIsAsset)
+		else if (bIsObj && !bIsSoftObject)
 		{
 			if( !bIsWeakObjPtr )
 			{
@@ -2158,21 +2158,6 @@ void FKismetCompilerVMBackend::ConstructFunction(FKismetFunctionContext& Functio
 	if (bIsUbergraph)
 	{
 		ScriptWriter.CopyStatementMapToUbergraphMap();
-
-		// Create a mapping between compile time generated events and what created them.
-		if (FunctionContext.bInstrumentScriptCode)
-		{
-			FBlueprintDebugData& DebugData = CompilerContext.NewClass->GetDebugData();
-			for (const TPair<UEdGraphPin*, UK2Node_Event*>& EventInfo : CompilerContext.SourcePinToExpansionEvent)
-			{
-				TArray<int32> PinOffsets;
-				DebugData.FindAllCodeLocationsFromSourcePin(EventInfo.Key, Function, PinOffsets);
-				for (const int32 ScriptOffset : PinOffsets)
-				{
-					DebugData.RegisterEntryPoint(ScriptOffset, EventInfo.Value->GetFunctionName());
-				}
-			}
-		}
 	}
 
 	// Make sure we didn't overflow the maximum bytecode size
diff --git a/Engine/Source/Editor/KismetCompiler/Public/KismetCompiledFunctionContext.h b/Engine/Source/Editor/KismetCompiler/Public/KismetCompiledFunctionContext.h
--- a/Engine/Source/Editor/KismetCompiler/Public/KismetCompiledFunctionContext.h
+++ b/Engine/Source/Editor/KismetCompiler/Public/KismetCompiledFunctionContext.h
@@ -86,7 +86,6 @@ public:
 	bool bIsInterfaceStub;
 	bool bIsConstFunction;
 	bool bEnforceConstCorrectness;
-	bool bInstrumentScriptCode;
 	bool bCreateDebugData;
 	bool bIsSimpleStubGraphWithNoParams;
 	uint32 NetFlags;
@@ -105,7 +104,7 @@ public:
 	//Does this function use requires FlowStack ?
 	bool bUseFlowStack;
 public:
-	FKismetFunctionContext(FCompilerResultsLog& InMessageLog, const UEdGraphSchema_K2* InSchema, UBlueprintGeneratedClass* InNewClass, UBlueprint* InBlueprint, bool bInGeneratingCpp, bool bInWantsInstrumentation);
+	FKismetFunctionContext(FCompilerResultsLog& InMessageLog, const UEdGraphSchema_K2* InSchema, UBlueprintGeneratedClass* InNewClass, UBlueprint* InBlueprint, bool bInGeneratingCpp);
 	
 	~FKismetFunctionContext();
 
@@ -180,42 +179,17 @@ public:
 
 	bool IsDebuggingOrInstrumentationRequired() const
 	{
-		return bInstrumentScriptCode || bCreateDebugData;
-	}
-
-	bool IsInstrumentationRequired() const
-	{
-		return bInstrumentScriptCode;
-	}
-
-	bool IsInstrumentationRequiredForNode(UEdGraphNode* Node) const
-	{
-		return bInstrumentScriptCode && !IsExpansionNode(Node);
-	}
-
-	bool IsInstrumentationRequiredForPin(UEdGraphPin* Pin) const
-	{
-		return bInstrumentScriptCode && !IsPinTraceSuppressed(Pin);
-	}
-
-	bool IsExpansionNode(const UEdGraphNode* Node) const
-	{
-		return MessageLog.IsExpansionNode(Node);
-	}
-
-	bool IsPinTraceSuppressed(const UEdGraphPin* Pin) const
-	{
-		return MessageLog.IsPinTraceSuppressed(Pin);
+		return bCreateDebugData;
 	}
 
 	EKismetCompiledStatementType GetWireTraceType() const
 	{
-		return bInstrumentScriptCode ? KCST_InstrumentedWireExit : KCST_WireTraceSite;
+		return KCST_WireTraceSite;
 	}
 
 	EKismetCompiledStatementType GetBreakpointType() const
 	{
-		return bInstrumentScriptCode ? KCST_InstrumentedWireEntry : KCST_DebugSite;
+		return KCST_DebugSite;
 	}
 
 	uint32 GetNetFlags() const
diff --git a/Engine/Source/Editor/KismetWidgets/Private/SPinTypeSelector.cpp b/Engine/Source/Editor/KismetWidgets/Private/SPinTypeSelector.cpp
--- a/Engine/Source/Editor/KismetWidgets/Private/SPinTypeSelector.cpp
+++ b/Engine/Source/Editor/KismetWidgets/Private/SPinTypeSelector.cpp
@@ -714,17 +714,17 @@ TSharedRef< SWidget > SPinTypeSelector::GetAllowedObjectTypes(FPinTypeTreeItem I
 		AllowedObjectReferenceTypes.Add(ObjectReferenceType);
 	}
 
-	if (PossibleObjectReferenceTypes & static_cast<uint8>(EObjectReferenceType::AssetID))
+	if (PossibleObjectReferenceTypes & static_cast<uint8>(EObjectReferenceType::SoftObject))
 	{
-		PinType.PinCategory = UEdGraphSchema_K2::PC_Asset;
+		PinType.PinCategory = UEdGraphSchema_K2::PC_SoftObject;
 		TSharedRef<SWidget> Widget = CreateObjectReferenceWidget(InItem, PinType, IconBrush, FText::Format(LOCTEXT("AssetTooltip", "Path to an instanced object of type \'{Typename}\' which may be in an unloaded state. Can be utilized to asynchronously load the object reference."), Args));
 		FObjectReferenceListItem ObjectReferenceType = MakeShareable(new FObjectReferenceType(InItem, Widget, PinType.PinCategory));
 		AllowedObjectReferenceTypes.Add(ObjectReferenceType);
 	}
 
-	if (PossibleObjectReferenceTypes & static_cast<uint8>(EObjectReferenceType::ClassAssetID))
+	if (PossibleObjectReferenceTypes & static_cast<uint8>(EObjectReferenceType::SoftClass))
 	{
-		PinType.PinCategory = UEdGraphSchema_K2::PC_AssetClass;
+		PinType.PinCategory = UEdGraphSchema_K2::PC_SoftClass;
 		TSharedRef<SWidget> Widget = CreateObjectReferenceWidget(InItem, PinType, IconBrush, FText::Format(LOCTEXT("ClassAssetTooltip", "Path to a class object of type \'{Typename}\' which may be in an unloaded state. Can be utilized to asynchronously load the class."), Args));
 		FObjectReferenceListItem ObjectReferenceType = MakeShareable(new FObjectReferenceType(InItem, Widget, PinType.PinCategory));
 		AllowedObjectReferenceTypes.Add(ObjectReferenceType);
@@ -1059,29 +1059,39 @@ bool SPinTypeSelector::GetChildrenMatchingSearch(const FText& InSearchText, cons
 
 FText SPinTypeSelector::GetToolTipForComboBoxType() const
 {
+	FText EditText;
 	if(IsEnabled())
 	{
 		if (bIsCompactSelector)
 		{
-			return LOCTEXT("CompactPinTypeSelector", "Left click to select the variable's pin type. Right click to toggle the type as an array.");
+			EditText = LOCTEXT("CompactPinTypeSelector", "Left click to select the variable's pin type. Right click to toggle the type as an array.");
 		}
 		else
 		{
-			return LOCTEXT("PinTypeSelector", "Select the variable's pin type.");
+			EditText = LOCTEXT("PinTypeSelector", "Select the variable's pin type.");
 		}
 	}
+	else
+	{
+		EditText = LOCTEXT("PinTypeSelector_Disabled", "Cannot edit variable type when they are inherited from parent.");
+	}
 
-	return LOCTEXT("PinTypeSelector_Disabled", "Cannot edit variable type when they are inherited from parent.");
+	return FText::Format(LOCTEXT("PrimaryTypeTwoLines", "{0}\nCurrent Type: {1}"), EditText, GetTypeDescription());
 }
 
 FText SPinTypeSelector::GetToolTipForComboBoxSecondaryType() const
 {
+	FText EditText;
 	if (IsEnabled())
 	{
-		return LOCTEXT("PinTypeValueSelector", "Select the map's value type.");
+		EditText = LOCTEXT("PinTypeValueSelector", "Select the map's value type.");
+	}
+	else
+	{
+		EditText = LOCTEXT("PinTypeSelector_ValueDisabled", "Cannot edit map value type when they are inherited from parent.");
 	}
 
-	return LOCTEXT("PinTypeSelector_ValueDisabled", "Cannot edit map value type when they are inherited from parent.");
+	return FText::Format(LOCTEXT("SecondaryTypeTwoLines", "{0}\nValue Type: {1}"), EditText, GetSecondaryTypeDescription());
 }
 
 FText SPinTypeSelector::GetToolTipForArrayWidget() const
@@ -1101,17 +1111,38 @@ FText SPinTypeSelector::GetToolTipForArrayWidget() const
 
 FText SPinTypeSelector::GetToolTipForContainerWidget() const
 {
-	if (IsEnabled())
+	if (TargetPinType.Get().PinCategory == Schema->PC_Exec)
 	{
 		// The entire widget may be enabled, but the container type button may be disabled because it is an "exec" pin.
-		if (TargetPinType.Get().PinCategory == Schema->PC_Exec)
+		return LOCTEXT("ContainerType_ExecDisabled", "Exec pins cannot be containers.");
+	}
+	else
+	{
+		FText EditText;
+		if (IsEnabled())
 		{
-			return LOCTEXT("ContainerType_ExecDisabled", "Exec pins cannot be containers.");
+			EditText = LOCTEXT("ContainerType", "Make this variable a container (array, set, or map) of selected type.");
+		}
+		else
+		{
+			EditText = LOCTEXT("ContainerType_Disabled", "Cannot edit variable type while the variable is placed in a graph or inherited from parent.");
 		}
-		return LOCTEXT("ContainerType", "Make this variable a container (array, set, or map) of selected type.");
-	}
 
-	return LOCTEXT("ContainerType_Disabled", "Cannot edit variable type while the variable is placed in a graph or inherited from parent.");
+		FText ContainerTypeText;
+		switch (TargetPinType.Get().ContainerType)
+		{
+		case EPinContainerType::Array:
+			ContainerTypeText = LOCTEXT("ContainerTypeTooltip_Array", "Array");
+			break;
+		case EPinContainerType::Set:
+			ContainerTypeText = LOCTEXT("ContainerTypeTooltip_Set", "Set");
+			break;
+		case EPinContainerType::Map:
+			ContainerTypeText = LOCTEXT("ContainerTypeTooltip_Map", "Map (Dictionary)");
+			break;
+		}
+		return ContainerTypeText.IsEmpty() ? EditText : FText::Format(LOCTEXT("ContainerTypeTwoLines", "{0}\nContainer Type: {1}"), EditText, ContainerTypeText);
+	}
 }
 
 FReply SPinTypeSelector::OnMouseButtonDown( const FGeometry& MyGeometry, const FPointerEvent& MouseEvent )
diff --git a/Engine/Source/Editor/MainFrame/Private/Frame/MainFrameActions.cpp b/Engine/Source/Editor/MainFrame/Private/Frame/MainFrameActions.cpp
--- a/Engine/Source/Editor/MainFrame/Private/Frame/MainFrameActions.cpp
+++ b/Engine/Source/Editor/MainFrame/Private/Frame/MainFrameActions.cpp
@@ -438,8 +438,6 @@ bool FMainFrameActionCallbacks::PackageBuildConfigurationIsChecked( EProjectPack
 void FMainFrameActionCallbacks::PackageProject( const FName InPlatformInfoName )
 {
 	GUnrealEd->CancelPlayingViaLauncher();
-	/*TArray<FString> Packages;
-	GUnrealEd->SaveWorldForPlay(Packages);*/
 	SaveAll();
 	
 	// does the project have any code?
diff --git a/Engine/Source/Editor/MovieSceneCaptureDialog/Private/MovieSceneCaptureDialogModule.cpp b/Engine/Source/Editor/MovieSceneCaptureDialog/Private/MovieSceneCaptureDialogModule.cpp
--- a/Engine/Source/Editor/MovieSceneCaptureDialog/Private/MovieSceneCaptureDialogModule.cpp
+++ b/Engine/Source/Editor/MovieSceneCaptureDialog/Private/MovieSceneCaptureDialogModule.cpp
@@ -713,33 +713,19 @@ class FMovieSceneCaptureDialogModule : public IMovieSceneCaptureDialogModule
 	FText OnStartCapture(UMovieSceneCapture* CaptureObject)
 	{
 		FString MapNameToLoad;
-		if( CaptureObject->Settings.bCreateTemporaryCopiesOfLevels )
-		{
-			TArray<FString> SavedMapNames;
-			GEditor->SaveWorldForPlay(SavedMapNames);
-
-			if (SavedMapNames.Num() == 0)
-			{
-				return LOCTEXT("CouldNotSaveMap", "Could not save map for movie capture.");
-			}
 
-			MapNameToLoad = SavedMapNames[ 0 ];
-		}
-		else
+		// Prompt the user to save their changes so that they'll be in the movie, since we're not saving temporary copies of the level.
+		bool bPromptUserToSave = true;
+		bool bSaveMapPackages = true;
+		bool bSaveContentPackages = true;
+		if( !FEditorFileUtils::SaveDirtyPackages( bPromptUserToSave, bSaveMapPackages, bSaveContentPackages ) )
 		{
-			// Prompt the user to save their changes so that they'll be in the movie, since we're not saving temporary copies of the level.
-			bool bPromptUserToSave = true;
-			bool bSaveMapPackages = true;
-			bool bSaveContentPackages = true;
-			if( !FEditorFileUtils::SaveDirtyPackages( bPromptUserToSave, bSaveMapPackages, bSaveContentPackages ) )
-			{
-				return LOCTEXT( "UserCancelled", "Capturing was cancelled from the save dialog." );
-			}
-
-			const FString WorldPackageName = GWorld->GetOutermost()->GetName();
-			MapNameToLoad = WorldPackageName;
+			return LOCTEXT( "UserCancelled", "Capturing was cancelled from the save dialog." );
 		}
 
+		const FString WorldPackageName = GWorld->GetOutermost()->GetName();
+		MapNameToLoad = WorldPackageName;
+
 		// Allow the game mode to be overridden
 		if( CaptureObject->Settings.GameModeOverride != nullptr )
 		{
@@ -823,15 +809,6 @@ class FMovieSceneCaptureDialogModule : public IMovieSceneCaptureDialogModule
 
 		FString EditorCommandLine = FString::Printf(TEXT("%s -MovieSceneCaptureManifest=\"%s\" -game -NoLoadingScreen -ForceRes -Windowed"), *MapNameToLoad, *Filename);
 
-		if( CaptureObject->Settings.bCreateTemporaryCopiesOfLevels )
-		{
-			// the PIEVIACONSOLE parameter tells UGameEngine to add the auto-save dir to the paths array and repopulate the package file cache
-			// this is needed in order to support streaming levels as the streaming level packages will be loaded only when needed (thus
-			// their package names need to be findable by the package file caching system)
-			// (we add to EditorCommandLine because the URL is ignored by WindowsTools)
-			EditorCommandLine.Append( TEXT( " -PIEVIACONSOLE" ) );
-		}
-
 		// Spit out any additional, user-supplied command line args
 		if (!CaptureObject->AdditionalCommandLineArguments.IsEmpty())
 		{
diff --git a/Engine/Source/Editor/MovieSceneTools/Public/K2Node_GetSequenceBinding.h b/Engine/Source/Editor/MovieSceneTools/Public/K2Node_GetSequenceBinding.h
--- a/Engine/Source/Editor/MovieSceneTools/Public/K2Node_GetSequenceBinding.h
+++ b/Engine/Source/Editor/MovieSceneTools/Public/K2Node_GetSequenceBinding.h
@@ -24,7 +24,7 @@ public:
 
 	/** The sequence from which to choose a binding identifier */
 	UPROPERTY(EditAnywhere, Category="Sequence", meta=(AllowedClasses="MovieSceneSequence"))
-	FStringAssetReference SourceSequence;
+	FSoftObjectPath SourceSequence;
 
 	/** The user-selected literal binding identifier from the sequence to use */
 	UPROPERTY()
diff --git a/Engine/Source/Editor/Persona/Private/Customization/BlendSampleDetails.cpp b/Engine/Source/Editor/Persona/Private/Customization/BlendSampleDetails.cpp
--- a/Engine/Source/Editor/Persona/Private/Customization/BlendSampleDetails.cpp
+++ b/Engine/Source/Editor/Persona/Private/Customization/BlendSampleDetails.cpp
@@ -6,7 +6,6 @@
 #include "UObject/Class.h"
 #include "IDetailsView.h"
 #include "EditorStyleSet.h"
-#include "Misc/StringAssetReference.h"
 #include "Widgets/DeclarativeSyntaxSupport.h"
 #include "PropertyHandle.h"
 #include "DetailLayoutBuilder.h"
@@ -215,9 +214,8 @@ bool FBlendSampleDetails::ShouldFilterAssetStatic(const FAssetData& AssetData, c
 	FString SkeletonName;
 	if (AssetData.GetTagValue(SkeletonTagName, SkeletonName))
 	{
-		FStringAssetReference StringAssetReference(SkeletonName);
 		// Check whether or not the skeletons match
-		if (StringAssetReference.ToString() == BlendSpaceBase->GetSkeleton()->GetPathName())
+		if (SkeletonName == BlendSpaceBase->GetSkeleton()->GetPathName())
 		{
 			// If so check if the additive animation tpye is compatible with the blend space
 			const FName AdditiveTypeTagName = GET_MEMBER_NAME_CHECKED(UAnimSequence, AdditiveAnimType);
@@ -250,9 +248,8 @@ bool FBlendSampleDetails::ShouldFilterAsset(const FAssetData& AssetData) const
 	FString SkeletonName;
 	if (AssetData.GetTagValue(SkeletonTagName, SkeletonName))
 	{
-		FStringAssetReference StringAssetReference(SkeletonName);
 		// Check whether or not the skeletons match
-		if (StringAssetReference.ToString() == BlendSpace->GetSkeleton()->GetPathName())
+		if (SkeletonName == BlendSpace->GetSkeleton()->GetPathName())
 		{
 			// If so check if the additive animation tpye is compatible with the blend space
 			const FName AdditiveTypeTagName = GET_MEMBER_NAME_CHECKED(UAnimSequence, AdditiveAnimType);
diff --git a/Engine/Source/Editor/Persona/Private/PersonaPreviewSceneDescription.h b/Engine/Source/Editor/Persona/Private/PersonaPreviewSceneDescription.h
--- a/Engine/Source/Editor/Persona/Private/PersonaPreviewSceneDescription.h
+++ b/Engine/Source/Editor/Persona/Private/PersonaPreviewSceneDescription.h
@@ -5,7 +5,7 @@
 #include "CoreMinimal.h"
 #include "UObject/ObjectMacros.h"
 #include "UObject/Object.h"
-#include "UObject/AssetPtr.h"
+#include "UObject/SoftObjectPtr.h"
 #include "PersonaPreviewSceneDescription.generated.h"
 
 class UAnimationAsset;
@@ -39,14 +39,14 @@ public:
 
 	/** The preview animation to use */
 	UPROPERTY(EditAnywhere, Category = "Animation", meta=(DisplayThumbnail=true))
-	TAssetPtr<UAnimationAsset> Animation;
+	TSoftObjectPtr<UAnimationAsset> Animation;
 
 	/** The preview mesh to use */
 	UPROPERTY(EditAnywhere, Category = "Mesh", meta=(DisplayThumbnail=true))
-	TAssetPtr<USkeletalMesh> PreviewMesh;
+	TSoftObjectPtr<USkeletalMesh> PreviewMesh;
 
 	UPROPERTY(EditAnywhere, Category = "Additional Meshes")
-	TAssetPtr<UDataAsset> AdditionalMeshes;
+	TSoftObjectPtr<UDataAsset> AdditionalMeshes;
 
 	UPROPERTY()
 	UPreviewMeshCollection* DefaultAdditionalMeshes;
diff --git a/Engine/Source/Editor/PhAT/Private/PhATSharedData.cpp b/Engine/Source/Editor/PhAT/Private/PhATSharedData.cpp
--- a/Engine/Source/Editor/PhAT/Private/PhATSharedData.cpp
+++ b/Engine/Source/Editor/PhAT/Private/PhATSharedData.cpp
@@ -86,24 +86,18 @@ FPhATSharedData::~FPhATSharedData()
 
 void FPhATSharedData::Initialize()
 {
-	EditorSkelComp = NULL;
+	EditorSkelComp = nullptr;
 	PhysicalAnimationComponent = nullptr;
 
-	USkeletalMesh * PreviewMesh =  NULL; 
-	FStringAssetReference PreviewMeshStringRef = PhysicsAsset->PreviewSkeletalMesh.ToStringReference();
-	// load it since now is the time to load
-	if (!PreviewMeshStringRef.ToString().IsEmpty())
-	{
-		PreviewMesh = Cast<USkeletalMesh>(StaticLoadObject(USkeletalMesh::StaticClass(), NULL, *PreviewMeshStringRef.ToString(), NULL, LOAD_None, NULL));
-	}
+	USkeletalMesh* PreviewMesh = PhysicsAsset->PreviewSkeletalMesh.LoadSynchronous();
 
-	if ( PreviewMesh == NULL)
+	if ( PreviewMesh == nullptr)
 	{
 		// Fall back to the default skeletal mesh in the EngineMeshes package.
 		// This is statically loaded as the package is likely not fully loaded
 		// (otherwise, it would have been found in the above iteration).
 		PreviewMesh = (USkeletalMesh*)StaticLoadObject(
-			USkeletalMesh::StaticClass(), NULL, TEXT("/Engine/EngineMeshes/SkeletalCube.SkeletalCube"), NULL, LOAD_None, NULL);
+			USkeletalMesh::StaticClass(), NULL, TEXT("/Engine/EngineMeshes/SkeletalCube.SkeletalCube"), nullptr, LOAD_None, nullptr);
 		check(PreviewMesh);
 
 		FMessageDialog::Open(EAppMsgType::Ok, FText::Format(
@@ -123,7 +117,7 @@ void FPhATSharedData::Initialize()
 	PhysicalAnimationComponent->SetSkeletalMeshComponent(EditorSkelComp);
 
 	// Create floor component
-	UStaticMesh* FloorMesh = LoadObject<UStaticMesh>(NULL, TEXT("/Engine/EditorMeshes/PhAT_FloorBox.PhAT_FloorBox"), NULL, LOAD_None, NULL);
+	UStaticMesh* FloorMesh = LoadObject<UStaticMesh>(nullptr, TEXT("/Engine/EditorMeshes/PhAT_FloorBox.PhAT_FloorBox"), NULL, LOAD_None, NULL);
 	check(FloorMesh);
 
 	EditorFloorComp = NewObject<UStaticMeshComponent>();
diff --git a/Engine/Source/Editor/PropertyEditor/Private/ItemPropertyNode.cpp b/Engine/Source/Editor/PropertyEditor/Private/ItemPropertyNode.cpp
--- a/Engine/Source/Editor/PropertyEditor/Private/ItemPropertyNode.cpp
+++ b/Engine/Source/Editor/PropertyEditor/Private/ItemPropertyNode.cpp
@@ -333,7 +333,7 @@ void FItemPropertyNode::InitChildNodes()
 				if ( FPropertySettings::Get().ExpandDistributions() == false)
 				{
 					// auto-expand distribution structs
-					if ( Cast<UObjectProperty>(StructMember) || Cast<UWeakObjectProperty>(StructMember) || Cast<ULazyObjectProperty>(StructMember) || Cast<UAssetObjectProperty>(StructMember) )
+					if ( Cast<UObjectProperty>(StructMember) || Cast<UWeakObjectProperty>(StructMember) || Cast<ULazyObjectProperty>(StructMember) || Cast<USoftObjectProperty>(StructMember) )
 					{
 						const FName StructName = StructProperty->Struct->GetFName();
 						if (StructName == NAME_RawDistributionFloat || StructName == NAME_RawDistributionVector)
diff --git a/Engine/Source/Editor/PropertyEditor/Private/PropertyEditorHelpers.cpp b/Engine/Source/Editor/PropertyEditor/Private/PropertyEditorHelpers.cpp
--- a/Engine/Source/Editor/PropertyEditor/Private/PropertyEditorHelpers.cpp
+++ b/Engine/Source/Editor/PropertyEditor/Private/PropertyEditorHelpers.cpp
@@ -619,38 +619,16 @@ namespace PropertyEditorHelpers
 		return (NodeProperty->IsA<UObjectPropertyBase>() || NodeProperty->IsA<UInterfaceProperty>()) && (!bUsingAssetPicker || !SPropertyEditorAsset::Supports(NodeProperty));
 	}
 	
-	static bool IsStringAssetReference( const UProperty* Property )
+	static bool IsSoftObjectPath( const UProperty* Property )
 	{
-		bool bIsStringAssetRef = false;
-
 		const UStructProperty* StructProp = Cast<const UStructProperty>( Property );
-		if( StructProp && StructProp->Struct )
-		{
-			FName StructName = StructProp->Struct->GetFName();
-
-			static const FName StringAssetRef("StringAssetReference");
-
-			bIsStringAssetRef = StructName == StringAssetRef;
-		}
-
-		return bIsStringAssetRef;
+		return StructProp && StructProp->Struct == TBaseStructure<FSoftObjectPath>::Get();
 	}
 
-	static bool IsStringClassReference( const UProperty* Property )
+	static bool IsSoftClassPath( const UProperty* Property )
 	{
-		bool bIsStringClassRef = false;
-
-		const UStructProperty* StructProp = Cast<const UStructProperty>( Property );
-		if( StructProp && StructProp->Struct )
-		{
-			FName StructName = StructProp->Struct->GetFName();
-
-			static const FName StringClassRef("StringClassReference");
-
-			bIsStringClassRef = StructName == StringClassRef;
-		}
-
-		return bIsStringClassRef;
+		const UStructProperty* StructProp = Cast<const UStructProperty>(Property);
+		return StructProp && StructProp->Struct == TBaseStructure<FSoftClassPath>::Get();
 	}
 
 	void GetRequiredPropertyButtons( TSharedRef<FPropertyNode> PropertyNode, TArray<EPropertyButton::Type>& OutRequiredButtons, bool bUsingAssetPicker )
@@ -708,7 +686,7 @@ namespace PropertyEditorHelpers
 					else
 					{
 						// ignore class properties
-						if( (Cast<const UClassProperty>( NodeProperty ) == NULL) && (Cast<const UAssetClassProperty>( NodeProperty ) == NULL) )
+						if( (Cast<const UClassProperty>( NodeProperty ) == NULL) && (Cast<const USoftClassProperty>( NodeProperty ) == NULL) )
 						{
 							UObjectPropertyBase* ObjectProperty = Cast<UObjectPropertyBase>( NodeProperty );
 
@@ -734,7 +712,7 @@ namespace PropertyEditorHelpers
 							}
 							
 							// Do not allow actor object properties to show the asset picker
-							if( ( ObjectProperty && !ObjectProperty->PropertyClass->IsChildOf( AActor::StaticClass() ) ) || IsStringAssetReference(NodeProperty) )
+							if( ( ObjectProperty && !ObjectProperty->PropertyClass->IsChildOf( AActor::StaticClass() ) ) || IsSoftObjectPath(NodeProperty) )
 							{
 								// add button for picking the asset from an asset picker
 								OutRequiredButtons.Add( EPropertyButton::PickAsset );
@@ -753,7 +731,7 @@ namespace PropertyEditorHelpers
 		// Handle a class property.
 
 		UClassProperty* ClassProp = Cast<UClassProperty>(NodeProperty);
-		if( ClassProp || IsStringClassReference(NodeProperty))
+		if( ClassProp || IsSoftClassPath(NodeProperty))
 		{
 			OutRequiredButtons.Add( EPropertyButton::Use );			
 			OutRequiredButtons.Add( EPropertyButton::Browse );
@@ -770,7 +748,7 @@ namespace PropertyEditorHelpers
 				OutRequiredButtons.Add( EPropertyButton::Clear );
 			}
 		}
-		else if (NodeProperty->IsA<UAssetClassProperty>() )
+		else if (NodeProperty->IsA<USoftClassProperty>() )
 		{
 			OutRequiredButtons.Add( EPropertyButton::Use );
 			
@@ -843,11 +821,11 @@ namespace PropertyEditorHelpers
 
 		UProperty* Property = PropertyNode->GetProperty();
 		UClassProperty* ClassProperty = Cast<UClassProperty>(Property);
-		UAssetClassProperty* AssetClassProperty = Cast<UAssetClassProperty>(Property);
+		USoftClassProperty* SoftClassProperty = Cast<USoftClassProperty>(Property);
 
-		if (ClassProperty || AssetClassProperty)
+		if (ClassProperty || SoftClassProperty)
 		{
-			UClass const* const SelectedClass = GEditor->GetFirstSelectedClass(ClassProperty ? ClassProperty->MetaClass : AssetClassProperty->MetaClass);
+			UClass const* const SelectedClass = GEditor->GetFirstSelectedClass(ClassProperty ? ClassProperty->MetaClass : SoftClassProperty->MetaClass);
 			if (SelectedClass != nullptr)
 			{
 				SelectionPathName = SelectedClass->GetPathName();
diff --git a/Engine/Source/Editor/PropertyEditor/Private/PropertyEditorModule.cpp b/Engine/Source/Editor/PropertyEditor/Private/PropertyEditorModule.cpp
--- a/Engine/Source/Editor/PropertyEditor/Private/PropertyEditorModule.cpp
+++ b/Engine/Source/Editor/PropertyEditor/Private/PropertyEditorModule.cpp
@@ -716,7 +716,7 @@ TSharedRef<class IStructureDetailsView> FPropertyEditorModule::CreateStructureDe
 			PropertyToTest = SetProperty ? SetProperty->ElementProp : PropertyToTest;
 			PropertyToTest = MapProperty ? MapProperty->ValueProp : PropertyToTest;
 
-			if( InStructureDetailsViewArgs.bShowClasses && (PropertyToTest->IsA<UClassProperty>() || PropertyToTest->IsA<UAssetClassProperty>()) )
+			if( InStructureDetailsViewArgs.bShowClasses && (PropertyToTest->IsA<UClassProperty>() || PropertyToTest->IsA<USoftClassProperty>()) )
 			{
 				return true;
 			}
@@ -732,7 +732,7 @@ TSharedRef<class IStructureDetailsView> FPropertyEditorModule::CreateStructureDe
 				if( InStructureDetailsViewArgs.bShowAssets )
 				{
 					// Is this an "asset" property?
-					if( PropertyToTest->IsA<UAssetObjectProperty>())
+					if( PropertyToTest->IsA<USoftObjectProperty>())
 					{
 						return true;
 					}
