Copying //UE4/Dev-Framework to //UE4/Dev-Main (Source: //UE4/Dev-Framework @ 3716594)

diff --git a/Engine/Source/Editor/Kismet/Private/BlueprintDetailsCustomization.cpp b/Engine/Source/Editor/Kismet/Private/BlueprintDetailsCustomization.cpp
--- a/Engine/Source/Editor/Kismet/Private/BlueprintDetailsCustomization.cpp
+++ b/Engine/Source/Editor/Kismet/Private/BlueprintDetailsCustomization.cpp
@@ -1526,16 +1526,14 @@ FText FBlueprintVarActionDetails::OnGetCategoryText() const
 	FName VarName = CachedVariableName;
 	if (VarName != NAME_None)
 	{
-		const UEdGraphSchema_K2* K2Schema = GetDefault<UEdGraphSchema_K2>();
-
 		if ( UBlueprint* OwnerBlueprint = GetPropertyOwnerBlueprint() )
 		{
 			FText Category = FBlueprintEditorUtils::GetBlueprintVariableCategory(OwnerBlueprint, VarName, GetLocalVariableScope(CachedVariableProperty.Get()));
 
 			// Older blueprints will have their name as the default category and whenever it is the same as the default category, display localized text
-			if ( Category.EqualTo(FText::FromString(OwnerBlueprint->GetName())) || Category.EqualTo(K2Schema->VR_DefaultCategory) )
+			if ( Category.EqualTo(FText::FromString(OwnerBlueprint->GetName())) || Category.EqualTo(UEdGraphSchema_K2::VR_DefaultCategory) )
 			{
-				return K2Schema->VR_DefaultCategory;
+				return UEdGraphSchema_K2::VR_DefaultCategory;
 			}
 			else
 			{
@@ -2865,12 +2863,10 @@ FReply FBlueprintGraphArgumentLayout::OnArgMoveDown()
 
 bool FBlueprintGraphArgumentLayout::ShouldPinBeReadOnly(bool bIsEditingPinType/* = false*/) const
 {
-	const UEdGraphSchema_K2* Schema = GetDefault<UEdGraphSchema_K2>();
-
 	if (TargetNode && ParamItemPtr.IsValid())
 	{
 		// Right now, we only care that the user is unable to edit the auto-generated "then" pin
-		if ((ParamItemPtr.Pin()->PinType.PinCategory == Schema->PC_Exec) && (!TargetNode->CanModifyExecutionWires()))
+		if ((ParamItemPtr.Pin()->PinType.PinCategory == UEdGraphSchema_K2::PC_Exec) && (!TargetNode->CanModifyExecutionWires()))
 		{
 			return true;
 		}
@@ -2901,7 +2897,7 @@ FText FBlueprintGraphArgumentLayout::OnGetArgNameText() const
 {
 	if (ParamItemPtr.IsValid())
 	{
-		return FText::FromString(ParamItemPtr.Pin()->PinName);
+		return FText::FromName(ParamItemPtr.Pin()->PinName);
 	}
 	return FText();
 }
@@ -2911,7 +2907,7 @@ FText FBlueprintGraphArgumentLayout::OnGetArgToolTipText() const
 	if (ParamItemPtr.IsValid())
 	{
 		FText PinTypeText = UEdGraphSchema_K2::TypeToText(ParamItemPtr.Pin()->PinType);
-		return FText::Format(LOCTEXT("BlueprintArgToolTipText", "Name: {0}\nType: {1}"), FText::FromString(ParamItemPtr.Pin()->PinName), PinTypeText);
+		return FText::Format(LOCTEXT("BlueprintArgToolTipText", "Name: {0}\nType: {1}"), FText::FromName(ParamItemPtr.Pin()->PinName), PinTypeText);
 	}
 	return FText::GetEmpty();
 }
@@ -2934,8 +2930,7 @@ void FBlueprintGraphArgumentLayout::OnArgNameChange(const FText& InNewText)
 	}
 	else
 	{
-		const FString& OldName = ParamItemPtr.Pin()->PinName;
-		bVerified = GraphActionDetailsPtr.Pin()->OnVerifyPinRename(TargetNode, OldName, InNewText.ToString(), ErrorMessage);
+		bVerified = GraphActionDetailsPtr.Pin()->OnVerifyPinRename(TargetNode, ParamItemPtr.Pin()->PinName, InNewText.ToString(), ErrorMessage);
 	}
 
 	if(!bVerified)
@@ -2952,9 +2947,9 @@ void FBlueprintGraphArgumentLayout::OnArgNameTextCommitted(const FText& NewText,
 {
 	if (!NewText.IsEmpty() && TargetNode && ParamItemPtr.IsValid() && GraphActionDetailsPtr.IsValid() && !ShouldPinBeReadOnly())
 	{
-		const FString OldName = ParamItemPtr.Pin()->PinName;
+		const FName OldName = ParamItemPtr.Pin()->PinName;
 		const FString& NewName = NewText.ToString();
-		if(OldName != NewName)
+		if (OldName.ToString() != NewName)
 		{
 			GraphActionDetailsPtr.Pin()->OnPinRenamed(TargetNode, OldName, NewName);
 		}
@@ -2999,7 +2994,7 @@ void FBlueprintGraphArgumentLayout::PinInfoChanged(const FEdGraphPinType& PinTyp
 {
 	if (ParamItemPtr.IsValid() && FBlueprintEditorUtils::IsPinTypeValid(PinType))
 	{
-		const FString PinName = ParamItemPtr.Pin()->PinName;
+		const FName PinName = ParamItemPtr.Pin()->PinName;
 		TSharedPtr<class FBaseBlueprintGraphActionDetails> GraphActionDetailsPinned = GraphActionDetailsPtr.Pin();
 		if (GraphActionDetailsPinned.IsValid())
 		{
@@ -3783,12 +3778,11 @@ FText FBlueprintDelegateActionDetails::OnGetCategoryText() const
 {
 	if (UMulticastDelegateProperty* DelegateProperty = GetDelegateProperty())
 	{
-		const UEdGraphSchema_K2* K2Schema = GetDefault<UEdGraphSchema_K2>();
 		FName DelegateName = DelegateProperty->GetFName();
 		FText Category = FBlueprintEditorUtils::GetBlueprintVariableCategory(GetBlueprintObj(), DelegateName, NULL);
 
 		// Older blueprints will have their name as the default category
-		if( Category.EqualTo(FText::FromString(GetBlueprintObj()->GetName())) || Category.EqualTo(K2Schema->VR_DefaultCategory) )
+		if( Category.EqualTo(FText::FromString(GetBlueprintObj()->GetName())) || Category.EqualTo(UEdGraphSchema_K2::VR_DefaultCategory) )
 		{
 			return LOCTEXT("DefaultCategory", "Default");
 		}
@@ -4001,7 +3995,7 @@ void FBlueprintDelegateActionDetails::OnFunctionSelected(TSharedPtr<FString> Fun
 				UProperty* FuncParam = *PropIt;
 				FEdGraphPinType TypeOut;
 				Schema->ConvertPropertyToPinType(FuncParam, TypeOut);
-				UEdGraphPin* EdGraphPin = FunctionEntryNode->CreateUserDefinedPin(FuncParam->GetName(), TypeOut, EGPD_Output);
+				UEdGraphPin* EdGraphPin = FunctionEntryNode->CreateUserDefinedPin(FuncParam->GetFName(), TypeOut, EGPD_Output);
 				ensure(EdGraphPin);
 			}
 
@@ -4056,15 +4050,15 @@ struct FPinRenamedHelper : public FBasePinChangeHelper
 	}
 };
 
-bool FBaseBlueprintGraphActionDetails::OnVerifyPinRename(UK2Node_EditablePinBase* InTargetNode, const FString& InOldName, const FString& InNewName, FText& OutErrorMessage)
+bool FBaseBlueprintGraphActionDetails::OnVerifyPinRename(UK2Node_EditablePinBase* InTargetNode, const FName InOldName, const FString& InNewName, FText& OutErrorMessage)
 {
 	// If the name is unchanged, allow the name
-	if(InOldName == InNewName)
+	if (InOldName.ToString() == InNewName)
 	{
 		return true;
 	}
 
-	if( InNewName.Len() > NAME_SIZE )
+	if (InNewName.Len() > NAME_SIZE)
 	{
 		OutErrorMessage = FText::Format( LOCTEXT("PinNameTooLong", "The name you entered is too long. Names must be less than {0} characters"), FText::AsNumber( NAME_SIZE ) );
 		return false;
@@ -4084,7 +4078,7 @@ bool FBaseBlueprintGraphActionDetails::OnVerifyPinRename(UK2Node_EditablePinBase
 	return true;
 }
 
-bool FBaseBlueprintGraphActionDetails::OnPinRenamed(UK2Node_EditablePinBase* TargetNode, const FString& OldName, const FString& NewName)
+bool FBaseBlueprintGraphActionDetails::OnPinRenamed(UK2Node_EditablePinBase* TargetNode, const FName OldName, const FString& NewName)
 {
 	// Before changing the name, verify the name
 	FText ErrorMessage;
@@ -4117,10 +4111,12 @@ bool FBaseBlueprintGraphActionDetails::OnPinRenamed(UK2Node_EditablePinBase* Tar
 		// GATHER 
 		PinRenamedHelper.Broadcast(GetBlueprintObj(), TargetNode, Graph);
 
+		const FName NewFName = *NewName;
+
 		// TEST
 		for (UK2Node* NodeToRename : PinRenamedHelper.NodesToRename)
 		{
-			if (ERenamePinResult::ERenamePinResult_NameCollision == NodeToRename->RenameUserDefinedPin(OldName, NewName, true))
+			if (ERenamePinResult::ERenamePinResult_NameCollision == NodeToRename->RenameUserDefinedPin(OldName, NewFName, true))
 			{
 				return false;
 			}
@@ -4129,7 +4125,7 @@ bool FBaseBlueprintGraphActionDetails::OnPinRenamed(UK2Node_EditablePinBase* Tar
 		// UPDATE
 		for (UK2Node* NodeToRename : PinRenamedHelper.NodesToRename)
 		{
-			NodeToRename->RenameUserDefinedPin(OldName, NewName, false);
+			NodeToRename->RenameUserDefinedPin(OldName, NewFName, false);
 		}
 
 		for (UK2Node_EditablePinBase* TerminalNode : TerminalNodes)
@@ -4140,7 +4136,7 @@ bool FBaseBlueprintGraphActionDetails::OnPinRenamed(UK2Node_EditablePinBase* Tar
 			});
 			if (UDPinPtr)
 			{
-				(*UDPinPtr)->PinName = NewName;
+				(*UDPinPtr)->PinName = NewFName;
 			}
 		}
 	}
@@ -4150,8 +4146,8 @@ bool FBaseBlueprintGraphActionDetails::OnPinRenamed(UK2Node_EditablePinBase* Tar
 void FBlueprintGraphActionDetails::SetEntryAndResultNodes()
 {
 	// Clear the entry and exit nodes to the graph
-	FunctionEntryNodePtr = NULL;
-	FunctionResultNodePtr = NULL;
+	FunctionEntryNodePtr = nullptr;
+	FunctionResultNodePtr = nullptr;
 
 	if (UEdGraph* NewTargetGraph = GetGraph())
 	{
@@ -4284,8 +4280,7 @@ FText FBlueprintGraphActionDetails::OnGetCategoryText() const
 {
 	if (FKismetUserDeclaredFunctionMetadata* Metadata = GetMetadataBlock())
 	{
-		const UEdGraphSchema_K2* K2Schema = GetDefault<UEdGraphSchema_K2>();
-		if( Metadata->Category.IsEmpty() || Metadata->Category.EqualTo(K2Schema->VR_DefaultCategory) )
+		if( Metadata->Category.IsEmpty() || Metadata->Category.EqualTo(UEdGraphSchema_K2::VR_DefaultCategory) )
 		{
 			return LOCTEXT("DefaultCategory", "Default");
 		}
@@ -4744,13 +4739,12 @@ FReply FBaseBlueprintGraphActionDetails::OnAddNewInputClicked()
 		FEdGraphPinType PinType = MyBlueprint.Pin()->GetLastFunctionPinTypeUsed();
 
 		// Make sure that if this is an exec node we are allowed one.
-		const UEdGraphSchema_K2* Schema = GetDefault<UEdGraphSchema_K2>();
-		if ((PinType.PinCategory == Schema->PC_Exec) && (!FunctionEntryNode->CanModifyExecutionWires()))
+		if ((PinType.PinCategory == UEdGraphSchema_K2::PC_Exec) && (!FunctionEntryNode->CanModifyExecutionWires()))
 		{
 			MyBlueprint.Pin()->ResetLastPinType();
 			PinType = MyBlueprint.Pin()->GetLastFunctionPinTypeUsed();
 		}
-		FString NewPinName = TEXT("NewParam");
+		const FName NewPinName = TEXT("NewParam");
 		if (FunctionEntryNode->CreateUserDefinedPin(NewPinName, PinType, EGPD_Output))
 		{
 			OnParamsChanged(FunctionEntryNode, true);
@@ -4819,14 +4813,13 @@ FReply FBlueprintGraphActionDetails::OnAddNewOutputClicked()
 		FEdGraphPinType PinType = MyBlueprint.Pin()->GetLastFunctionPinTypeUsed();
 		PinType.bIsReference = false;
 		// Make sure that if this is an exec node we are allowed one.
-		const UEdGraphSchema_K2* Schema = GetDefault<UEdGraphSchema_K2>();
-		if ((PinType.PinCategory == Schema->PC_Exec) && (!FunctionResultNode->CanModifyExecutionWires()))
+		if ((PinType.PinCategory == UEdGraphSchema_K2::PC_Exec) && (!FunctionResultNode->CanModifyExecutionWires()))
 		{
 			MyBlueprint.Pin()->ResetLastPinType();
 			PinType = MyBlueprint.Pin()->GetLastFunctionPinTypeUsed();
 		}
 
-		const FString NewPinName = FunctionResultNode->CreateUniquePinName(TEXT("NewParam"));
+		const FName NewPinName = FunctionResultNode->CreateUniquePinName(TEXT("NewParam"));
 		TArray<UK2Node_EditablePinBase*> TargetNodes = GatherAllResultNodes(FunctionResultNode);
 		bool bAllChanged = TargetNodes.Num() > 0;
 		for (UK2Node_EditablePinBase* Node : TargetNodes)
@@ -5679,14 +5672,12 @@ FText FBlueprintComponentDetails::OnGetVariableCategoryText() const
 	FName VarName = CachedNodePtr->GetVariableName();
 	if (VarName != NAME_None)
 	{
-		const UEdGraphSchema_K2* K2Schema = GetDefault<UEdGraphSchema_K2>();
-
 		FText Category = FBlueprintEditorUtils::GetBlueprintVariableCategory(GetBlueprintObj(), VarName, NULL);
 
 		// Older blueprints will have their name as the default category
 		if( Category.EqualTo(FText::FromString(GetBlueprintObj()->GetName())) )
 		{
-			return K2Schema->VR_DefaultCategory;
+			return UEdGraphSchema_K2::VR_DefaultCategory;
 		}
 		else
 		{
diff --git a/Engine/Source/Editor/Kismet/Private/BlueprintDetailsCustomization.h b/Engine/Source/Editor/Kismet/Private/BlueprintDetailsCustomization.h
--- a/Engine/Source/Editor/Kismet/Private/BlueprintDetailsCustomization.h
+++ b/Engine/Source/Editor/Kismet/Private/BlueprintDetailsCustomization.h
@@ -306,10 +306,10 @@ public:
 	void OnParamsChanged(UK2Node_EditablePinBase* TargetNode, bool bForceRefresh = false);
 
 	/** Checks if the pin rename can occur */
-	bool OnVerifyPinRename(UK2Node_EditablePinBase* InTargetNode, const FString& InOldName, const FString& InNewName, FText& OutErrorMessage);
+	bool OnVerifyPinRename(UK2Node_EditablePinBase* InTargetNode, const FName InOldName, const FString& InNewName, FText& OutErrorMessage);
 
 	/** Refreshes the graph and ensures the target node is up to date */
-	bool OnPinRenamed(UK2Node_EditablePinBase* TargetNode, const FString& OldName, const FString& NewName);
+	bool OnPinRenamed(UK2Node_EditablePinBase* TargetNode, const FName OldName, const FString& NewName);
 	
 	/** Gets the blueprint we're editing */
 	TWeakPtr<SMyBlueprint> GetMyBlueprint() const {return MyBlueprint.Pin();}
diff --git a/Engine/Source/Editor/Kismet/Private/BlueprintEditor.cpp b/Engine/Source/Editor/Kismet/Private/BlueprintEditor.cpp
--- a/Engine/Source/Editor/Kismet/Private/BlueprintEditor.cpp
+++ b/Engine/Source/Editor/Kismet/Private/BlueprintEditor.cpp
@@ -1497,8 +1497,7 @@ void FBlueprintEditor::EnsureBlueprintIsUpToDate(UBlueprint* BlueprintObj)
 		// If we should have a UCS but don't yet, make it
 		if(!FBlueprintEditorUtils::FindUserConstructionScript(BlueprintObj))
 		{
-			const UEdGraphSchema_K2* K2Schema = GetDefault<UEdGraphSchema_K2>();
-			UEdGraph* UCSGraph = FBlueprintEditorUtils::CreateNewGraph(BlueprintObj, K2Schema->FN_UserConstructionScript, UEdGraph::StaticClass(), UEdGraphSchema_K2::StaticClass());
+			UEdGraph* UCSGraph = FBlueprintEditorUtils::CreateNewGraph(BlueprintObj, UEdGraphSchema_K2::FN_UserConstructionScript, UEdGraph::StaticClass(), UEdGraphSchema_K2::StaticClass());
 			FBlueprintEditorUtils::AddFunctionGraph(BlueprintObj, UCSGraph, /*bIsUserCreated=*/ false, AActor::StaticClass());
 			UCSGraph->bAllowDeletion = false;
 		}
@@ -3741,9 +3740,9 @@ bool FBlueprintEditor::CanPromoteToVariable(bool bInToMemberVariable) const
 	{
 		if (UEdGraphPin* Pin = FocusedGraphEd->GetGraphPinForMenu())
 		{
-			if (!Pin->bOrphanedPin && (bInToMemberVariable || (!bInToMemberVariable && FBlueprintEditorUtils::DoesSupportLocalVariables(FocusedGraphEd->GetCurrentGraph()))))
+			if (!Pin->bOrphanedPin && (bInToMemberVariable || FBlueprintEditorUtils::DoesSupportLocalVariables(FocusedGraphEd->GetCurrentGraph())))
 			{
-				bCanPromote = K2Schema->CanPromotePinToVariable(*Pin);
+				bCanPromote = K2Schema->CanPromotePinToVariable(*Pin, bInToMemberVariable);
 			}
 		}
 	}
@@ -4555,12 +4554,12 @@ void FBlueprintEditor::OnCollapseSelectionToFunction()
 bool FBlueprintEditor::CanCollapseSelectionToFunction(TSet<class UEdGraphNode*>& InSelection) const
 {
 	bool bBadConnection = false;
-	UEdGraphPin* OutputConnection = NULL;
-	UEdGraphPin* InputConnection = NULL;
+	UEdGraphPin* OutputConnection = nullptr;
+	UEdGraphPin* InputConnection = nullptr;
 
 	// Create a function graph
 	UEdGraph* FunctionGraph = FBlueprintEditorUtils::CreateNewGraph(GetBlueprintObj(), FBlueprintEditorUtils::FindUniqueKismetName(GetBlueprintObj(), TEXT("TempGraph")), UEdGraph::StaticClass(), UEdGraphSchema_K2::StaticClass());
-	FBlueprintEditorUtils::AddFunctionGraph<UClass>(GetBlueprintObj(), FunctionGraph, /*bIsUserCreated=*/ true, NULL);
+	FBlueprintEditorUtils::AddFunctionGraph<UClass>(GetBlueprintObj(), FunctionGraph, /*bIsUserCreated=*/ true, nullptr);
 
 	const UEdGraphSchema_K2* K2Schema = GetDefault<UEdGraphSchema_K2>();
 
@@ -4593,7 +4592,7 @@ bool FBlueprintEditor::CanCollapseSelectionToFunction(TSet<class UEdGraphNode*>&
 					InterfaceTemplateNode = CustomEvent;
 					if(InputConnection)
 					{
-						InputConnection = EventExecPin->LinkedTo[0];
+						InputConnection = (EventExecPin->LinkedTo.Num() > 0) ? EventExecPin->LinkedTo[0] : nullptr;
 					}
 					continue;
 				}
@@ -4604,47 +4603,45 @@ bool FBlueprintEditor::CanCollapseSelectionToFunction(TSet<class UEdGraphNode*>&
 		}
 		else
 		{
-			for (int32 PinIndex = 0; PinIndex < Node->Pins.Num(); ++PinIndex)
+			for (UEdGraphPin* NodePin : Node->Pins)
 			{
-				if(Node->Pins[PinIndex]->PinType.PinCategory == K2Schema->PC_Exec)
+				if(NodePin->PinType.PinCategory == UEdGraphSchema_K2::PC_Exec)
 				{
-					if (Node->Pins[PinIndex]->LinkedTo.Num() == 0 && Node->Pins[PinIndex]->Direction == EGPD_Input)
+					if (NodePin->LinkedTo.Num() == 0 && NodePin->Direction == EGPD_Input)
 					{
-						EntryGatewayPins.Add(Node->Pins[PinIndex]);
+						EntryGatewayPins.Add(NodePin);
 					}
 					else
 					{
-						for (int32 LinkIndex = 0; LinkIndex < Node->Pins[PinIndex]->LinkedTo.Num(); ++LinkIndex)
+						for (UEdGraphPin* ConnectedPin : NodePin->LinkedTo)
 						{
-							if (!InSelection.Contains(Node->Pins[PinIndex]->LinkedTo[LinkIndex]->GetOwningNode()))
+							if (!InSelection.Contains(ConnectedPin->GetOwningNode()))
 							{
-								if (Node->Pins[PinIndex]->Direction == EGPD_Input)
+								if (NodePin->Direction == EGPD_Input)
 								{
 									// For input pins, there must be a single connection 
-									if (InputConnection == NULL || InputConnection == Node->Pins[PinIndex])
+									if ((InputConnection == nullptr) || (InputConnection == NodePin))
 									{
-										EntryGatewayPins.Add(Node->Pins[PinIndex]);
-										InputConnection = Node->Pins[PinIndex];
+										EntryGatewayPins.Add(NodePin);
+										InputConnection = NodePin;
 									}
 									else
 									{
 										// Check if the input connection was linked, report what node it is connected to
-										LogResults.Error(*LOCTEXT("TooManyPathsMultipleInput_Error", "Found too many input connections in selection! @@ is connected to @@, previously found @@ connected to @@").ToString(), Node, Node->Pins[PinIndex]->LinkedTo[LinkIndex]->GetOwningNode(), InputConnection->GetOwningNode(), InputConnection->LinkedTo[0]->GetOwningNode());
+										LogResults.Error(*LOCTEXT("TooManyPathsMultipleInput_Error", "Found too many input connections in selection! @@ is connected to @@, previously found @@ connected to @@").ToString(), Node, ConnectedPin->GetOwningNode(), InputConnection->GetOwningNode(), (InputConnection->LinkedTo.Num() > 0) ? InputConnection->LinkedTo[0]->GetOwningNode() : nullptr);
 										bBadConnection = true;
 									}
 								}
 								else
 								{
 									// For output pins, as long as they all connect to the same pin, we consider the selection valid for being made into a function
-									if (OutputConnection == NULL || OutputConnection == Node->Pins[PinIndex]->LinkedTo[LinkIndex])
+									if ((OutputConnection == nullptr) || (OutputConnection == ConnectedPin))
 									{
-										OutputConnection = Node->Pins[PinIndex]->LinkedTo[LinkIndex];
+										OutputConnection = ConnectedPin;
 									}
 									else
 									{
-										check(OutputConnection->LinkedTo.Num());
-
-										LogResults.Error(*LOCTEXT("TooManyPathsMultipleOutput_Error", "Found too many output connections in selection! @@ is connected to @@, previously found @@ connected to @@").ToString(), Node, Node->Pins[PinIndex]->LinkedTo[LinkIndex]->GetOwningNode(), OutputConnection->GetOwningNode(), OutputConnection->LinkedTo[0]->GetOwningNode());
+										LogResults.Error(*LOCTEXT("TooManyPathsMultipleOutput_Error", "Found too many output connections in selection! @@ is connected to @@, previously found @@ connected to @@").ToString(), Node, ConnectedPin->GetOwningNode(), OutputConnection->GetOwningNode(), (OutputConnection->LinkedTo.Num() > 0) ? OutputConnection->LinkedTo[0]->GetOwningNode() : nullptr);
 										bBadConnection = true;
 									}
 								}
@@ -4656,7 +4653,7 @@ bool FBlueprintEditor::CanCollapseSelectionToFunction(TSet<class UEdGraphNode*>&
 		}
 	}
 
-	if (!bBadConnection && InputConnection == nullptr && EntryGatewayPins.Num() > 1)
+	if (!bBadConnection && (InputConnection == nullptr) && (EntryGatewayPins.Num() > 1))
 	{
 		// Too many input gateway pins with no connections.
 		LogResults.Error(*LOCTEXT("AmbiguousEntryPaths_Error", "Multiple entry pin possibilities. Unable to convert to a function. Make sure that selection either has only 1 entry pin or exactly 1 entry pin has a connection.").ToString());
@@ -5617,7 +5614,7 @@ struct FUpdatePastedNodes
 		{
 			if (UK2Node_CallFunctionOnMember* CallOnMember = Cast<UK2Node_CallFunctionOnMember>(PastedNode))
 			{
-				if (UEdGraphPin* TargetInPin = CallOnMember->FindPin(K2Schema->PN_Self))
+				if (UEdGraphPin* TargetInPin = CallOnMember->FindPin(UEdGraphSchema_K2::PN_Self))
 				{
 					const UClass* TargetClass = Cast<const UClass>(TargetInPin->PinType.PinSubCategoryObject.Get());
 
@@ -5683,7 +5680,7 @@ private:
 		bool bResult = true;
 		for(UEdGraphPin* OldPin : OldNode->Pins)
 		{
-			if(OldPin && (OldPin->PinName != K2Schema->PN_Self))
+			if(OldPin && (OldPin->PinName != UEdGraphSchema_K2::PN_Self))
 			{
 				UEdGraphPin* NewPin = NewNode->FindPin(OldPin->PinName);
 				if (NewPin)
@@ -5691,13 +5688,13 @@ private:
 					if (!K2Schema->MovePinLinks(*OldPin, *NewPin).CanSafeConnect())
 					{
 						UE_LOG(LogBlueprint, Error, TEXT("FUpdatePastedNodes: Cannot connect pin '%s' node '%s'"),
-							*OldPin->PinName, *OldNode->GetName());
+							*OldPin->PinName.ToString(), *OldNode->GetName());
 						bResult = false;
 					}
 				}
 				else
 				{
-					UE_LOG(LogBlueprint, Error, TEXT("FUpdatePastedNodes: Cannot find pin '%s'"), *OldPin->PinName);
+					UE_LOG(LogBlueprint, Error, TEXT("FUpdatePastedNodes: Cannot find pin '%s'"), *OldPin->PinName.ToString());
 					bResult = false;
 				}
 			}
@@ -5760,7 +5757,7 @@ private:
 
 		if (NewTarget)
 		{
-			UEdGraphPin* SelfPin = NewCall->FindPinChecked(K2Schema->PN_Self);
+			UEdGraphPin* SelfPin = NewCall->FindPinChecked(UEdGraphSchema_K2::PN_Self);
 			if (!K2Schema->TryCreateConnection(SelfPin, NewTarget->GetValuePin()))
 			{
 				UE_LOG(LogBlueprint, Error, TEXT("FUpdatePastedNodes: Cannot connect new self."));
@@ -6349,7 +6346,7 @@ void FBlueprintEditor::ExtractEventTemplateForFunction(class UK2Node_CustomEvent
 
 	for(UEdGraphPin* Pin : InCustomEvent->Pins)
 	{
-		if(Pin->PinType.PinCategory == K2Schema->PC_Exec)
+		if(Pin->PinType.PinCategory == UEdGraphSchema_K2::PC_Exec)
 		{
 			TArray< UEdGraphPin* > PinLinkList = Pin->LinkedTo;
 			for( UEdGraphPin* PinLink : PinLinkList)
@@ -6364,7 +6361,7 @@ void FBlueprintEditor::ExtractEventTemplateForFunction(class UK2Node_CustomEvent
 				}
 			}
 		}
-		else if(Pin->PinType.PinCategory != K2Schema->PC_Delegate)
+		else if(Pin->PinType.PinCategory != UEdGraphSchema_K2::PC_Delegate)
 		{
 
 			TArray< UEdGraphPin* > PinLinkList = Pin->LinkedTo;
@@ -6376,14 +6373,14 @@ void FBlueprintEditor::ExtractEventTemplateForFunction(class UK2Node_CustomEvent
 					Pin->Modify();
 					PinLink->Modify();
 
-					FString PortName = Pin->PinName + TEXT("_Out");
+					const FName PortName = *FString::Printf(TEXT("%s_Out"), *Pin->PinName.ToString());
 					UEdGraphPin* RemotePortPin = InGatewayNode->FindPin(PortName);
 					// For nodes that are connected to the event but not collapsing into the graph, they need to create a pin on the result.
 					if(RemotePortPin == nullptr)
 					{
-						FString UniquePortName = InGatewayNode->CreateUniquePinName(PortName);
+						FName UniquePortName = InGatewayNode->CreateUniquePinName(PortName);
 
-						RemotePortPin = InGatewayNode->CreatePin(Pin->Direction, Pin->PinType,UniquePortName);
+						RemotePortPin = InGatewayNode->CreatePin(Pin->Direction, Pin->PinType, UniquePortName);
 						InResultNode->CreateUserDefinedPin(UniquePortName, Pin->PinType, EGPD_Input);
 					}
 					PinLink->BreakAllPinLinks();
@@ -6393,7 +6390,7 @@ void FBlueprintEditor::ExtractEventTemplateForFunction(class UK2Node_CustomEvent
 				{
 					InEntryNode->Modify();
 
-					FString UniquePortName = InGatewayNode->CreateUniquePinName(Pin->PinName);
+					const FName UniquePortName = InGatewayNode->CreateUniquePinName(Pin->PinName);
 					InEntryNode->CreateUserDefinedPin(UniquePortName, Pin->PinType, EGPD_Output);
 				}
 			}
@@ -6518,7 +6515,7 @@ void FBlueprintEditor::CollapseNodesIntoGraph(UEdGraphNode* InGatewayNode, UK2No
 				UEdGraphPin* RemotePortPin = NULL;
 
 				// Function graphs have a single exec path through them, so only one exec pin for input and another for output. In this fashion, they must not be handled by name.
-				if(InGatewayNode->GetClass() == UK2Node_CallFunction::StaticClass() && LocalPin->PinType.PinCategory == K2Schema->PC_Exec)
+				if(InGatewayNode->GetClass() == UK2Node_CallFunction::StaticClass() && LocalPin->PinType.PinCategory == UEdGraphSchema_K2::PC_Exec)
 				{
 					LocalPortPin = LocalPort->Pins[0];
 					RemotePortPin = K2Schema->FindExecutionPin(*InGatewayNode, (LocalPortPin->Direction == EGPD_Input)? EGPD_Output : EGPD_Input);
@@ -6540,7 +6537,7 @@ void FBlueprintEditor::CollapseNodesIntoGraph(UEdGraphNode* InGatewayNode, UK2No
 
 					if(LocalPin->LinkedTo[0]->GetOwningNode() != InEntryNode)
 					{
-						FString UniquePortName = InGatewayNode->CreateUniquePinName(LocalPin->PinName);
+						const FName UniquePortName = InGatewayNode->CreateUniquePinName(LocalPin->PinName);
 
 						if(!RemotePortPin && !LocalPortPin)
 						{
@@ -6924,7 +6921,7 @@ void FBlueprintEditor::ExpandNode(UEdGraphNode* InNodeToExpand, UEdGraph* InSour
 			for (int32 PinIndex = 0; PinIndex < Node->Pins.Num(); ++PinIndex)
 			{
 				// Only hookup output exec pins that do not have a connection
-				if(Node->Pins[PinIndex]->PinType.PinCategory == K2Schema->PC_Exec && Node->Pins[PinIndex]->Direction == EGPD_Output && Node->Pins[PinIndex]->LinkedTo.Num() == 0)
+				if(Node->Pins[PinIndex]->PinType.PinCategory == UEdGraphSchema_K2::PC_Exec && Node->Pins[PinIndex]->Direction == EGPD_Output && Node->Pins[PinIndex]->LinkedTo.Num() == 0)
 				{
 					Node->Pins[PinIndex]->MakeLinkTo(OutputExecPinReconnect);
 				}
@@ -7136,7 +7133,7 @@ void FBlueprintEditor::OnAddNewDelegate()
 	Blueprint->Modify();
 
 	FEdGraphPinType DelegateType;
-	DelegateType.PinCategory = K2Schema->PC_MCDelegate;
+	DelegateType.PinCategory = UEdGraphSchema_K2::PC_MCDelegate;
 	const bool bVarCreatedSuccess = FBlueprintEditorUtils::AddMemberVariable(Blueprint, Name, DelegateType);
 	if(!bVarCreatedSuccess)
 	{
diff --git a/Engine/Source/Editor/Kismet/Private/BlueprintEditorModule.cpp b/Engine/Source/Editor/Kismet/Private/BlueprintEditorModule.cpp
--- a/Engine/Source/Editor/Kismet/Private/BlueprintEditorModule.cpp
+++ b/Engine/Source/Editor/Kismet/Private/BlueprintEditorModule.cpp
@@ -30,6 +30,7 @@
 #include "EdGraphUtilities.h"
 #include "BlueprintGraphPanelPinFactory.h"
 #include "Widgets/Docking/SDockTab.h"
+#include "KismetCompiler.h"
 
 #define LOCTEXT_NAMESPACE "BlueprintEditor"
 
@@ -278,6 +279,8 @@ void FBlueprintEditorModule::ShutdownModule()
 	// Unregister internal SCS editor customizations
 	UnregisterSCSEditorCustomization("InstancedStaticMeshComponent");
 
+	BlueprintCompilers.Empty();
+
 	UEdGraphPin::ShutdownVerification();
 }
 
@@ -423,4 +426,53 @@ void FBlueprintEditorModule::PrepareAutoGeneratedDefaultEvents()
 	}
 }
 
+/** Blueprint Compiler map*/
+TMap<TSubclassOf<UBlueprint>, FOnGetBlueprintCompiler> FBlueprintEditorModule::BlueprintCompilers;
+
+void FBlueprintEditorModule::RegisterCompilerGetter(const TSubclassOf<UBlueprint> BlueprintType, FOnGetBlueprintCompiler InCompilerGetter)
+{
+	if (InCompilerGetter.IsBound())
+	{
+		FOnGetBlueprintCompiler* Getter = BlueprintCompilers.Find(BlueprintType);
+		if (Getter)
+		{
+			*Getter = InCompilerGetter;
+		}
+		else
+		{
+			BlueprintCompilers.Add(BlueprintType, InCompilerGetter);
+		}
+	}
+}
+
+void FBlueprintEditorModule::UnregisterCompilerGetter(const TSubclassOf<UBlueprint> BlueprintType)
+{
+	if (UObjectInitialized())
+	{
+		BlueprintCompilers.Remove(BlueprintType);
+	}
+}
+
+TSharedPtr<FKismetCompilerContext> FBlueprintEditorModule::GetCompiler(UBlueprint* BP, FCompilerResultsLog& InMessageLog, const FKismetCompilerOptions& InCompileOptions)
+{
+	if (BP != nullptr)
+	{
+		// Load required modules to make sure we have the required items we have
+		BP->LoadModulesRequiredForCompilation();
+
+		if (BlueprintCompilers.Num() > 0)
+		{
+			for (UClass* Class = BP->GetClass(); Class != UBlueprint::StaticClass(); Class = Class->GetSuperClass())
+			{
+				if (FOnGetBlueprintCompiler* CompilerGetter = BlueprintCompilers.Find(Class))
+				{
+					return CompilerGetter->Execute(BP, InMessageLog, InCompileOptions);
+				}
+			}
+		}
+	}
+	
+	return TSharedPtr<FKismetCompilerContext>(new FKismetCompilerContext(BP, InMessageLog, InCompileOptions, nullptr));
+}
+
 #undef LOCTEXT_NAMESPACE
diff --git a/Engine/Source/Editor/Kismet/Private/DetailsDiff.cpp b/Engine/Source/Editor/Kismet/Private/DetailsDiff.cpp
--- a/Engine/Source/Editor/Kismet/Private/DetailsDiff.cpp
+++ b/Engine/Source/Editor/Kismet/Private/DetailsDiff.cpp
@@ -4,7 +4,8 @@
 #include "Modules/ModuleManager.h"
 #include "IDetailsView.h"
 #include "PropertyEditorModule.h"
-//#include "SBlueprintDiff.h"
+#include "Algo/Copy.h"
+#include "Algo/Transform.h"
 
 FDetailsDiff::FDetailsDiff(const UObject* InObject, FOnDisplayedPropertiesChanged InOnDisplayedPropertiesChanged )
 	: OnDisplayedPropertiesChanged( InOnDisplayedPropertiesChanged )
@@ -57,10 +58,7 @@ void FDetailsDiff::HandlePropertiesChanged()
 TArray<FPropertySoftPath> FDetailsDiff::GetDisplayedProperties() const
 {
 	TArray<FPropertySoftPath> Ret;
-	for( const auto& Property : DifferingProperties )
-	{
-		Ret.Push( Property );
-	}
+	Algo::Copy(DifferingProperties, Ret);
 	return Ret;
 }
 
@@ -69,31 +67,21 @@ void FDetailsDiff::DiffAgainst(const FDetailsDiff& Newer, TArray< FSingleObjectD
 	TSharedPtr< class IDetailsView > OldDetailsView = DetailsView;
 	TSharedPtr< class IDetailsView > NewDetailsView = Newer.DetailsView;
 
-	const auto& OldSelectedObjects = OldDetailsView->GetSelectedObjects();
-	const auto& NewSelectedObjects = NewDetailsView->GetSelectedObjects();
+	const TArray<TWeakObjectPtr<UObject>>& OldSelectedObjects = OldDetailsView->GetSelectedObjects();
+	const TArray<TWeakObjectPtr<UObject>>& NewSelectedObjects = NewDetailsView->GetSelectedObjects();
 
-	check(OldSelectedObjects.Num() == 1);
-
-	auto OldProperties = OldDetailsView->GetPropertiesInOrderDisplayed();
-	auto NewProperties = NewDetailsView->GetPropertiesInOrderDisplayed();
+	const TArray<FPropertyPath> OldProperties = OldDetailsView->GetPropertiesInOrderDisplayed();
+	const TArray<FPropertyPath> NewProperties = NewDetailsView->GetPropertiesInOrderDisplayed();
 
 	TSet<FPropertySoftPath> OldPropertiesSet;
 	TSet<FPropertySoftPath> NewPropertiesSet;
 
-	auto ToSet = [](const TArray<FPropertyPath>& Array, TSet<FPropertySoftPath>& OutSet)
-	{
-		for (const auto& Entry : Array)
-		{
-			OutSet.Add(FPropertySoftPath(Entry));
-		}
-	};
-
-	ToSet(OldProperties, OldPropertiesSet);
-	ToSet(NewProperties, NewPropertiesSet);
+	Algo::Transform(OldProperties, OldPropertiesSet, [](const FPropertyPath& Entry) { return FPropertySoftPath(Entry); });
+	Algo::Transform(NewProperties, NewPropertiesSet, [](const FPropertyPath& Entry) { return FPropertySoftPath(Entry); });
 
 	// detect removed properties:
-	TSet<FPropertySoftPath> RemovedProperties = OldPropertiesSet.Difference(NewPropertiesSet);
-	for (const auto& RemovedProperty : RemovedProperties)
+	const TSet<FPropertySoftPath> RemovedProperties = OldPropertiesSet.Difference(NewPropertiesSet);
+	for (const FPropertySoftPath& RemovedProperty : RemovedProperties)
 	{
 		// @todo: (doc) label these as removed, rather than added to a
 		FSingleObjectDiffEntry Entry(RemovedProperty, EPropertyDiffType::PropertyAddedToA);
@@ -101,16 +89,16 @@ void FDetailsDiff::DiffAgainst(const FDetailsDiff& Newer, TArray< FSingleObjectD
 	}
 
 	// detect added properties:
-	TSet<FPropertySoftPath> AddedProperties = NewPropertiesSet.Difference(OldPropertiesSet);
-	for (const auto& AddedProperty : AddedProperties)
+	const TSet<FPropertySoftPath> AddedProperties = NewPropertiesSet.Difference(OldPropertiesSet);
+	for (const FPropertySoftPath& AddedProperty : AddedProperties)
 	{
 		FSingleObjectDiffEntry Entry(AddedProperty, EPropertyDiffType::PropertyAddedToB);
 		OutDifferences.Push(Entry);
 	}
 
 	// check for changed properties
-	TSet<FPropertySoftPath> CommonProperties = NewPropertiesSet.Intersect(OldPropertiesSet);
-	for (const auto& CommonProperty : CommonProperties)
+	const TSet<FPropertySoftPath> CommonProperties = NewPropertiesSet.Intersect(OldPropertiesSet);
+	for (const FPropertySoftPath& CommonProperty : CommonProperties)
 	{
 		// get value, diff:
 		check(NewSelectedObjects.Num() == 1);
@@ -121,9 +109,9 @@ void FDetailsDiff::DiffAgainst(const FDetailsDiff& Newer, TArray< FSingleObjectD
 
 		if (!DiffUtils::Identical(OldProperty, NewProperty, CommonProperty, DifferingSubProperties))
 		{
-			for (int32 DifferingIndex = 0; DifferingIndex < DifferingSubProperties.Num(); DifferingIndex++)
+			for (const FPropertySoftPath& DifferingSubProperty : DifferingSubProperties)
 			{
-				OutDifferences.Push(FSingleObjectDiffEntry(DifferingSubProperties[DifferingIndex], EPropertyDiffType::PropertyValueChanged));
+				OutDifferences.Push(FSingleObjectDiffEntry(DifferingSubProperty, EPropertyDiffType::PropertyValueChanged));
 			}
 		}
 	}
diff --git a/Engine/Source/Editor/Kismet/Private/DiffUtils.cpp b/Engine/Source/Editor/Kismet/Private/DiffUtils.cpp
--- a/Engine/Source/Editor/Kismet/Private/DiffUtils.cpp
+++ b/Engine/Source/Editor/Kismet/Private/DiffUtils.cpp
@@ -113,7 +113,7 @@ FPropertyPath FPropertySoftPath::ResolvePath(const UObject* Object) const
 	};
 
 	const void* ContainerAddress = Object;
-	const UStruct* ContainerStruct = Object->GetClass();
+	const UStruct* ContainerStruct = (Object ? Object->GetClass() : nullptr);
 
 	FPropertyPath Ret;
 	for( int32 I = 0; I < PropertyChain.Num(); ++I )
@@ -433,11 +433,7 @@ static void IdenticalHelper(const UProperty* AProperty, const UProperty* BProper
 
 		return;
 	}
-	else if (AProperty->Identical(AValue, BValue, PPF_DeepComparison))
-	{
-		return;
-	}
-
+	
 	const UStructProperty* APropAsStruct = Cast<UStructProperty>(AProperty);
 	const UArrayProperty* APropAsArray = Cast<UArrayProperty>(AProperty);
 	const USetProperty* APropAsSet = Cast<USetProperty>(AProperty);
@@ -463,7 +459,7 @@ static void IdenticalHelper(const UProperty* AProperty, const UProperty* BProper
 	else if (APropAsArray != nullptr)
 	{
 		const UArrayProperty* BPropAsArray = CastChecked<UArrayProperty>(BProperty);
-		if(BPropAsArray->Inner == APropAsArray->Inner)
+		if(BPropAsArray->Inner->GetClass() == APropAsArray->Inner->GetClass())
 		{
 			FScriptArrayHelper ArrayHelperA(APropAsArray, AValue);
 			FScriptArrayHelper ArrayHelperB(BPropAsArray, BValue);
@@ -488,7 +484,7 @@ static void IdenticalHelper(const UProperty* AProperty, const UProperty* BProper
 	else if(APropAsSet != nullptr)
 	{
 		const USetProperty* BPropAsSet = CastChecked<USetProperty>(BProperty);
-		if(BPropAsSet->ElementProp == APropAsSet->ElementProp)
+		if(BPropAsSet->ElementProp->GetClass() == APropAsSet->ElementProp->GetClass())
 		{
 			FScriptSetHelper SetHelperA(APropAsSet, AValue);
 			FScriptSetHelper SetHelperB(BPropAsSet, BValue);
@@ -527,7 +523,7 @@ static void IdenticalHelper(const UProperty* AProperty, const UProperty* BProper
 	else if(APropAsMap != nullptr)
 	{
 		const UMapProperty* BPropAsMap = CastChecked<UMapProperty>(BProperty);
-		if(APropAsMap->KeyProp == BPropAsMap->KeyProp && APropAsMap->ValueProp == BPropAsMap->ValueProp)
+		if(APropAsMap->KeyProp->GetClass() == BPropAsMap->KeyProp->GetClass() && APropAsMap->ValueProp->GetClass() == BPropAsMap->ValueProp->GetClass())
 		{
 			FScriptMapHelper MapHelperA(APropAsMap, AValue);
 			FScriptMapHelper MapHelperB(BPropAsMap, BValue);
@@ -564,6 +560,12 @@ static void IdenticalHelper(const UProperty* AProperty, const UProperty* BProper
 	}
 	else if(APropAsObject != nullptr)
 	{
+		// Past container check, do a normal identical check now before going into components
+		if (AProperty->Identical(AValue, BValue, PPF_DeepComparison))
+		{
+			return;
+		}
+
 		// dig into the objects if they are in the same package as our initial object:
 		const UObjectProperty* BPropAsObject = CastChecked<UObjectProperty>(BProperty);
 
@@ -582,12 +584,18 @@ static void IdenticalHelper(const UProperty* AProperty, const UProperty* BProper
 			}
 		}
 		else
-		{		
+		{
 			DifferingSubProperties.Push(RootPath);
 		}
 	}
 	else
 	{
+		// Passed all container tests that would check for nested properties being wrong
+		if (AProperty->Identical(AValue, BValue, PPF_DeepComparison))
+		{
+			return;
+		}
+
 		DifferingSubProperties.Push(RootPath);
 	}
 }
diff --git a/Engine/Source/Editor/Kismet/Private/FindInBlueprintManager.cpp b/Engine/Source/Editor/Kismet/Private/FindInBlueprintManager.cpp
--- a/Engine/Source/Editor/Kismet/Private/FindInBlueprintManager.cpp
+++ b/Engine/Source/Editor/Kismet/Private/FindInBlueprintManager.cpp
@@ -628,14 +628,14 @@ namespace BlueprintSearchMetaDataHelpers
 	{
 		// Only save strings that are not empty
 
-		if(!InPinType.PinCategory.IsEmpty())
+		if(!InPinType.PinCategory.IsNone())
 		{
-			InWriter->WriteValue(FFindInBlueprintSearchTags::FiB_PinCategory, InPinType.PinCategory);
+			InWriter->WriteValue(FFindInBlueprintSearchTags::FiB_PinCategory, InPinType.PinCategory.ToString());
 		}
 
-		if(!InPinType.PinSubCategory.IsEmpty())
+		if(!InPinType.PinSubCategory.IsNone())
 		{
-			InWriter->WriteValue(FFindInBlueprintSearchTags::FiB_PinSubCategory, InPinType.PinSubCategory);
+			InWriter->WriteValue(FFindInBlueprintSearchTags::FiB_PinSubCategory, InPinType.PinSubCategory.ToString());
 		}
 
 		if(InPinType.PinSubCategoryObject.IsValid())
@@ -1172,7 +1172,6 @@ public:
 	/** Enables the caching process */
 	void Start() { bIsStarted = true; }
 
-	/** FTickableEditorObject interface */
 	EActiveTimerReturnType Tick(double InCurrentTime, float InDeltaTime)
 	{
 		// Protect against Slate recursion if a modal dialog appears from loading/resaving an asset
diff --git a/Engine/Source/Editor/Kismet/Private/FindInBlueprints.cpp b/Engine/Source/Editor/Kismet/Private/FindInBlueprints.cpp
--- a/Engine/Source/Editor/Kismet/Private/FindInBlueprints.cpp
+++ b/Engine/Source/Editor/Kismet/Private/FindInBlueprints.cpp
@@ -65,14 +65,14 @@ bool FindInBlueprintsHelpers::IsTextEqualToString(const FText& InText, const FSt
 
 FString FindInBlueprintsHelpers::GetPinTypeAsString(const FEdGraphPinType& InPinType)
 {
-	FString Result = InPinType.PinCategory;
+	FString Result = InPinType.PinCategory.ToString();
 	if(UObject* SubCategoryObject = InPinType.PinSubCategoryObject.Get()) 
 	{
 		Result += FString(" '") + SubCategoryObject->GetName() + "'";
 	}
 	else
 	{
-		Result += FString(" '") + InPinType.PinSubCategory + "'";
+		Result += FString(" '") + InPinType.PinSubCategory.ToString() + "'";
 	}
 
 	return Result;
@@ -84,15 +84,15 @@ bool FindInBlueprintsHelpers::ParsePinType(FText InKey, FText InValue, FEdGraphP
 
 	if(InKey.CompareTo(FFindInBlueprintSearchTags::FiB_PinCategory) == 0)
 	{
-		InOutPinType.PinCategory = InValue.ToString();
+		InOutPinType.PinCategory = *InValue.ToString();
 	}
 	else if(InKey.CompareTo(FFindInBlueprintSearchTags::FiB_PinSubCategory) == 0)
 	{
-		InOutPinType.PinSubCategory = InValue.ToString();
+		InOutPinType.PinSubCategory = *InValue.ToString();
 	}
 	else if(InKey.CompareTo(FFindInBlueprintSearchTags::FiB_ObjectClass) == 0)
 	{
-		InOutPinType.PinSubCategory = InValue.ToString();
+		InOutPinType.PinSubCategory = *InValue.ToString();
 	}
 	else if(InKey.CompareTo(FFindInBlueprintSearchTags::FiB_IsArray) == 0)
 	{
@@ -375,17 +375,17 @@ FText FFindInBlueprintsPin::GetCategory() const
 
 void FFindInBlueprintsPin::FinalizeSearchData()
 {
-	if(!PinType.PinSubCategory.IsEmpty())
+	if(!PinType.PinSubCategory.IsNone())
 	{
-		PinType.PinSubCategoryObject = FindObject<UClass>(ANY_PACKAGE, *PinType.PinSubCategory, true);
+		PinType.PinSubCategoryObject = FindObject<UClass>(ANY_PACKAGE, *PinType.PinSubCategory.ToString(), true);
 		if(!PinType.PinSubCategoryObject.IsValid())
 		{
-			PinType.PinSubCategoryObject = FindObject<UScriptStruct>(UObject::StaticClass(), *PinType.PinSubCategory);
+			PinType.PinSubCategoryObject = FindObject<UScriptStruct>(UObject::StaticClass(), *PinType.PinSubCategory.ToString());
 		}
 
 		if (PinType.PinSubCategoryObject.IsValid())
 		{
-			PinType.PinSubCategory.Empty();
+			PinType.PinSubCategory = NAME_None;
 		}
 	}
 
@@ -494,17 +494,17 @@ FText FFindInBlueprintsProperty::GetCategory() const
 
 void FFindInBlueprintsProperty::FinalizeSearchData()
 {
-	if(!PinType.PinSubCategory.IsEmpty())
+	if(!PinType.PinSubCategory.IsNone())
 	{
-		PinType.PinSubCategoryObject = FindObject<UClass>(ANY_PACKAGE, *PinType.PinSubCategory, true);
+		PinType.PinSubCategoryObject = FindObject<UClass>(ANY_PACKAGE, *PinType.PinSubCategory.ToString(), true);
 		if(!PinType.PinSubCategoryObject.IsValid())
 		{
-			PinType.PinSubCategoryObject = FindObject<UScriptStruct>(UObject::StaticClass(), *PinType.PinSubCategory);
+			PinType.PinSubCategoryObject = FindObject<UScriptStruct>(UObject::StaticClass(), *PinType.PinSubCategory.ToString());
 		}
 
 		if (PinType.PinSubCategoryObject.IsValid())
 		{
-			PinType.PinSubCategory.Empty();
+			PinType.PinSubCategory = NAME_None;
 		}
 	}
 }
diff --git a/Engine/Source/Editor/Kismet/Private/FormatTextDetails.cpp b/Engine/Source/Editor/Kismet/Private/FormatTextDetails.cpp
--- a/Engine/Source/Editor/Kismet/Private/FormatTextDetails.cpp
+++ b/Engine/Source/Editor/Kismet/Private/FormatTextDetails.cpp
@@ -219,7 +219,7 @@ void FFormatTextArgumentLayout::OnArgumentNameCommitted(const FText& NewText, ET
 	if(IsValidArgumentName(NewText))
 	{
 		FScopeTrue ScopeTrue(bCausedChange);
-		TargetNode->SetArgumentName(ArgumentIndex, NewText.ToString());
+		TargetNode->SetArgumentName(ArgumentIndex, *NewText.ToString());
 	}
 	ArgumentNameWidget.Pin()->SetError(FString());
 }
@@ -231,7 +231,7 @@ void FFormatTextArgumentLayout::OnArgumentNameChanged(const FText& NewText)
 
 bool FFormatTextArgumentLayout::IsValidArgumentName(const FText& InNewText) const
 {
-	if(TargetNode->FindArgumentPin(InNewText.ToString()))
+	if(TargetNode->FindArgumentPin(*InNewText.ToString()))
 	{
 		ArgumentNameWidget.Pin()->SetError(LOCTEXT("UniqueName_Error", "Name must be unique."));
 		return false;
diff --git a/Engine/Source/Editor/Kismet/Private/SBlueprintActionMenu.cpp b/Engine/Source/Editor/Kismet/Private/SBlueprintActionMenu.cpp
--- a/Engine/Source/Editor/Kismet/Private/SBlueprintActionMenu.cpp
+++ b/Engine/Source/Editor/Kismet/Private/SBlueprintActionMenu.cpp
@@ -546,18 +546,18 @@ void SBlueprintActionMenu::TryInsertPromoteToVariable(FBlueprintActionContext co
 	const UEdGraphSchema_K2* K2Schema = Cast<const UEdGraphSchema_K2>(GraphObj->GetSchema());
 	if ((K2Schema != nullptr) && (MenuContext.Pins.Num() > 0))
 	{
-		if (K2Schema->CanPromotePinToVariable(*MenuContext.Pins[0]))
+		if (K2Schema->CanPromotePinToVariable(*MenuContext.Pins[0], false))
 		{
 			TSharedPtr<FBlueprintAction_PromoteVariable> PromoteAction = TSharedPtr<FBlueprintAction_PromoteVariable>(new FBlueprintAction_PromoteVariable(true));
 			PromoteAction->MyBlueprintEditor = EditorPtr;
-			OutAllActions.AddAction( PromoteAction );
+			OutAllActions.AddAction(PromoteAction);
+		}
 
-			if (MenuContext.Graphs.Num() == 1 && FBlueprintEditorUtils::DoesSupportLocalVariables(MenuContext.Graphs[0]))
-			{
-				TSharedPtr<FBlueprintAction_PromoteVariable> LocalPromoteAction = TSharedPtr<FBlueprintAction_PromoteVariable>(new FBlueprintAction_PromoteVariable(false));
-				LocalPromoteAction->MyBlueprintEditor = EditorPtr;
-				OutAllActions.AddAction( LocalPromoteAction );
-			}
+		if (MenuContext.Graphs.Num() == 1 && FBlueprintEditorUtils::DoesSupportLocalVariables(MenuContext.Graphs[0]) && K2Schema->CanPromotePinToVariable(*MenuContext.Pins[0], true))
+		{
+			TSharedPtr<FBlueprintAction_PromoteVariable> LocalPromoteAction = TSharedPtr<FBlueprintAction_PromoteVariable>(new FBlueprintAction_PromoteVariable(false));
+			LocalPromoteAction->MyBlueprintEditor = EditorPtr;
+			OutAllActions.AddAction( LocalPromoteAction );
 		}
 	}
 }
diff --git a/Engine/Source/Editor/Kismet/Private/SBlueprintEditorToolbar.cpp b/Engine/Source/Editor/Kismet/Private/SBlueprintEditorToolbar.cpp
--- a/Engine/Source/Editor/Kismet/Private/SBlueprintEditorToolbar.cpp
+++ b/Engine/Source/Editor/Kismet/Private/SBlueprintEditorToolbar.cpp
@@ -1,4 +1,4 @@
-// Copyright 1998-2017 Epic Games, Inc. All Rights Reserved.
+﻿// Copyright 1998-2017 Epic Games, Inc. All Rights Reserved.
 
 
 #include "SBlueprintEditorToolbar.h"
@@ -110,15 +110,11 @@ void FKismet2Menu::FillDeveloperMenu( FMenuBuilder& MenuBuilder )
 	}
 	MenuBuilder.EndSection();
 
-	static const FBoolConfigValueHelper NativeCodeGenerationTool(TEXT("Kismet"), TEXT("bNativeCodeGenerationTool"), GEngineIni);
-	if (NativeCodeGenerationTool)
+	MenuBuilder.BeginSection("GenerateNativeCode", LOCTEXT("Cpp", "C++"));
 	{
-		MenuBuilder.BeginSection("GenerateNativeCode", LOCTEXT("Cpp", "C++"));
-		{
-			MenuBuilder.AddMenuEntry(FBlueprintEditorCommands::Get().GenerateNativeCode);
-		}
-		MenuBuilder.EndSection();
+		MenuBuilder.AddMenuEntry(FBlueprintEditorCommands::Get().GenerateNativeCode);
 	}
+	MenuBuilder.EndSection();
 
 	if (false)
 	{
diff --git a/Engine/Source/Editor/Kismet/Private/SBlueprintPalette.cpp b/Engine/Source/Editor/Kismet/Private/SBlueprintPalette.cpp
--- a/Engine/Source/Editor/Kismet/Private/SBlueprintPalette.cpp
+++ b/Engine/Source/Editor/Kismet/Private/SBlueprintPalette.cpp
@@ -81,10 +81,9 @@ static FString GetVarType(UStruct* VarScope, FName VarName, bool bUseObjToolTip,
 {
 	FString VarDesc;
 
-	if(VarScope != NULL)
+	if (VarScope)
 	{
-		UProperty* Property = FindField<UProperty>(VarScope, VarName);
-		if(Property != NULL)
+		if (UProperty* Property = FindField<UProperty>(VarScope, VarName))
 		{
 			// If it is an object property, see if we can get a nice class description instead of just the name
 			UObjectProperty* ObjProp = Cast<UObjectProperty>(Property);
@@ -122,20 +121,20 @@ static FString GetVarType(UStruct* VarScope, FName VarName, bool bUseObjToolTip,
 static FString GetVarTooltip(UBlueprint* InBlueprint, UClass* VarClass, FName VarName)
 {
 	FString ResultTooltip;
-	if(VarClass != NULL)
+	if (VarClass)
 	{
-		UProperty* Property = FindField<UProperty>(VarClass, VarName);
-		if(Property != NULL)
+	
+		if (UProperty* Property = FindField<UProperty>(VarClass, VarName))
 		{
 			// discover if the variable property is a non blueprint user variable
 			UClass* SourceClass = Property->GetOwnerClass();
-			if( SourceClass && SourceClass->ClassGeneratedBy == NULL )
+			if( SourceClass && SourceClass->ClassGeneratedBy == nullptr )
 			{
 				ResultTooltip = Property->GetToolTipText().ToString();
 			}
 			else
 			{
-				FBlueprintEditorUtils::GetBlueprintVariableMetaData(InBlueprint, VarName, NULL, TEXT("tooltip"), ResultTooltip);
+				FBlueprintEditorUtils::GetBlueprintVariableMetaData(InBlueprint, VarName, nullptr, TEXT("tooltip"), ResultTooltip);
 			}
 		}
 	}
@@ -156,13 +155,13 @@ static FString GetVarTooltip(UBlueprint* InBlueprint, UClass* VarClass, FName Va
  */
 static void GetSubGraphIcon(FEdGraphSchemaAction_K2Graph const* const ActionIn, UBlueprint const* BlueprintIn, FSlateBrush const*& IconOut, FSlateColor& ColorOut, FText& ToolTipOut)
 {
-	check(BlueprintIn != NULL);
+	check(BlueprintIn);
 
 	switch (ActionIn->GraphType)
 	{
 	case EEdGraphSchemaAction_K2Graph::Graph:
 		{
-			if (ActionIn->EdGraph != NULL)
+			if (ActionIn->EdGraph)
 			{
 				IconOut = FBlueprintEditor::GetGlyphForGraph(ActionIn->EdGraph);
 			}
@@ -176,20 +175,20 @@ static void GetSubGraphIcon(FEdGraphSchemaAction_K2Graph const* const ActionIn,
 		break;
 	case EEdGraphSchemaAction_K2Graph::Subgraph:
 		{
-			if ( ActionIn->EdGraph != NULL && ActionIn->EdGraph->IsA(UAnimationStateMachineGraph::StaticClass()) )
+			if (Cast<UAnimationStateMachineGraph>(ActionIn->EdGraph))
 			{
 				IconOut = FEditorStyle::GetBrush(TEXT("GraphEditor.StateMachine_16x") );
 				ToolTipOut = LOCTEXT("AnimationStateMachineGraph_ToolTip", "Animation State Machine");
 			}
-			else if ( ActionIn->EdGraph != NULL && ActionIn->EdGraph->IsA(UAnimationStateGraph::StaticClass()) )
+			else if (Cast<UAnimationStateGraph>(ActionIn->EdGraph))
 			{
 				IconOut = FEditorStyle::GetBrush(TEXT("GraphEditor.State_16x") );
 				ToolTipOut = LOCTEXT("AnimationState_ToolTip", "Animation State");
 			}
-			else if ( ActionIn->EdGraph != NULL && ActionIn->EdGraph->IsA(UAnimationTransitionGraph::StaticClass()) )
+			else if (Cast<UAnimationTransitionGraph>(ActionIn->EdGraph))
 			{
-				UObject* EdGraphOuter = ActionIn->EdGraph->GetOuter();
-				if ( EdGraphOuter != NULL && EdGraphOuter->IsA(UAnimStateConduitNode::StaticClass()) )
+				UAnimStateConduitNode* EdGraphOuter = Cast<UAnimStateConduitNode>(ActionIn->EdGraph->GetOuter());
+				if (EdGraphOuter)
 				{
 					IconOut = FEditorStyle::GetBrush(TEXT("GraphEditor.Conduit_16x"));
 					ToolTipOut = LOCTEXT("ConduitGraph_ToolTip", "Conduit");
@@ -210,21 +209,21 @@ static void GetSubGraphIcon(FEdGraphSchemaAction_K2Graph const* const ActionIn,
 	case EEdGraphSchemaAction_K2Graph::Macro:
 		{
 			IconOut = FEditorStyle::GetBrush(TEXT("GraphEditor.Macro_16x"));
-			if ( ActionIn->EdGraph == NULL )
+			if ( ActionIn->EdGraph == nullptr )
 			{
 				ToolTipOut = LOCTEXT("PotentialOverride_Tooltip", "Potential Override");	
 			}
 			else
 			{
 				// Need to see if this is a function overriding something in the parent, or 
-				UFunction* OverrideFunc = FindField<UFunction>(BlueprintIn->ParentClass, ActionIn->FuncName);
-				if ( OverrideFunc == NULL )
+				;
+				if (UFunction* OverrideFunc = FindField<UFunction>(BlueprintIn->ParentClass, ActionIn->FuncName))
 				{
-					ToolTipOut = LOCTEXT("Macro_Tooltip", "Macro");
+					ToolTipOut = LOCTEXT("Override_Tooltip", "Override");
 				}
 				else 
 				{
-					ToolTipOut = LOCTEXT("Override_Tooltip", "Override");
+					ToolTipOut = LOCTEXT("Macro_Tooltip", "Macro");
 				}
 			}
 		}
@@ -248,7 +247,7 @@ static void GetSubGraphIcon(FEdGraphSchemaAction_K2Graph const* const ActionIn,
 		break;
 	case EEdGraphSchemaAction_K2Graph::Function:
 		{
-			if ( ActionIn->EdGraph == NULL )
+			if ( ActionIn->EdGraph == nullptr )
 			{
 				IconOut = FEditorStyle::GetBrush(TEXT("GraphEditor.PotentialOverrideFunction_16x"));
 				ToolTipOut = LOCTEXT("PotentialOverride_Tooltip", "Potential Override");	
@@ -259,25 +258,21 @@ static void GetSubGraphIcon(FEdGraphSchemaAction_K2Graph const* const ActionIn,
 				{
 					IconOut = FEditorStyle::GetBrush(TEXT("GraphEditor.Animation_16x"));
 				}
+				else if (UFunction* OverrideFunc = FindField<UFunction>(BlueprintIn->ParentClass, ActionIn->FuncName))
+				{
+					IconOut = FEditorStyle::GetBrush(TEXT("GraphEditor.OverrideFunction_16x"));
+					ToolTipOut = LOCTEXT("Override_Tooltip", "Override");
+				}
 				else
 				{
-					UFunction* OverrideFunc = FindField<UFunction>(BlueprintIn->ParentClass, ActionIn->FuncName);
-					if ( OverrideFunc == NULL )
+					IconOut = FEditorStyle::GetBrush(TEXT("GraphEditor.Function_16x"));
+					if (ActionIn->EdGraph->IsA(UAnimationGraph::StaticClass()))
 					{
-						IconOut = FEditorStyle::GetBrush(TEXT("GraphEditor.Function_16x"));
-						if ( ActionIn->EdGraph->IsA(UAnimationGraph::StaticClass()) )
-						{
-							ToolTipOut = LOCTEXT("AnimationGraph_Tooltip", "Animation Graph");
-						}
-						else
-						{
-							ToolTipOut = LOCTEXT("Function_Tooltip", "Function");
-						}
+						ToolTipOut = LOCTEXT("AnimationGraph_Tooltip", "Animation Graph");
 					}
 					else
 					{
-						IconOut = FEditorStyle::GetBrush(TEXT("GraphEditor.OverrideFunction_16x"));
-						ToolTipOut = LOCTEXT("Override_Tooltip", "Override");
+						ToolTipOut = LOCTEXT("Function_Tooltip", "Function");
 					}
 				}
 			}
@@ -461,9 +456,9 @@ private:
 		else
 		{
 			// Check to see if the name conflicts
-			for( auto Iter = AssetData.CreateConstIterator(); Iter; ++Iter)
+			for ( const FAssetData& AssetInfo : AssetData)
 			{
-				if(Iter->AssetName.ToString() == InNewText.ToString())
+				if(AssetInfo.AssetName.ToString() == InNewText.ToString())
 				{
 					OutErrorMessage = FText::FromString(TEXT("Asset name already in use!"));
 					return false;
@@ -488,7 +483,7 @@ private:
 				AssetToolsModule.Get().RenameAssets(AssetsAndNames);
 			}
 
-			auto MyBlueprint = BlueprintEditor->GetMyBlueprintWidget();
+			TSharedPtr<SMyBlueprint> MyBlueprint = BlueprintEditor->GetMyBlueprintWidget();
 			if (MyBlueprint.IsValid())
 			{
 				MyBlueprint->SelectItemByName(FName(*Object->GetPathName()));
@@ -511,9 +506,9 @@ public:
 		// Should never make it here with anything but an enum action
 		check(ActionPtr.Pin()->GetTypeId() == FEdGraphSchemaAction_K2Enum::StaticGetTypeId());
 
-		auto EnumAction = (FEdGraphSchemaAction_K2Enum*)ActionPtr.Pin().Get();
+		FEdGraphSchemaAction_K2Enum* EnumAction = (FEdGraphSchemaAction_K2Enum*)ActionPtr.Pin().Get();
 		
-		return VerifyNewAssetName(EnumAction ? EnumAction->Enum : NULL, InNewText, OutErrorMessage);
+		return VerifyNewAssetName(EnumAction ? EnumAction->Enum : nullptr, InNewText, OutErrorMessage);
 	}
 
 	/**
@@ -561,9 +556,9 @@ public:
 		// Should never make it here with anything but a struct action
 		check(ActionPtr.Pin()->GetTypeId() == FEdGraphSchemaAction_K2Struct::StaticGetTypeId());
 
-		auto Action = (FEdGraphSchemaAction_K2Struct*)ActionPtr.Pin().Get();
+		FEdGraphSchemaAction_K2Struct* Action = (FEdGraphSchemaAction_K2Struct*)ActionPtr.Pin().Get();
 
-		return VerifyNewAssetName(Action ? Action->Struct : NULL, InNewText, OutErrorMessage);
+		return VerifyNewAssetName(Action ? Action->Struct : nullptr, InNewText, OutErrorMessage);
 	}
 
 	/**
@@ -584,7 +579,7 @@ public:
 		FEdGraphSchemaAction_K2Event* EventAction = (FEdGraphSchemaAction_K2Event*)ActionPtr.Pin().Get();
 
 		UK2Node* AssociatedNode = EventAction->NodeTemplate;
-		if ((AssociatedNode != NULL) && AssociatedNode->bCanRenameNode)
+		if (AssociatedNode && AssociatedNode->bCanRenameNode)
 		{
 			TSharedPtr<INameValidatorInterface> NodeNameValidator = FNameValidatorFactory::MakeValidator(AssociatedNode);
 			bIsNameValid = (NodeNameValidator->IsValid(InNewText.ToString(), true) == EValidatorResult::Ok);
@@ -608,7 +603,7 @@ public:
 
 		FEdGraphSchemaAction_K2Struct* Action = (FEdGraphSchemaAction_K2Struct*)ActionPtr.Pin().Get();
 
-		CommitNewAssetName(Action ? Action->Struct : NULL, BlueprintEditorPtr.Pin().Get(), NewText);
+		CommitNewAssetName(Action ? Action->Struct : nullptr, BlueprintEditorPtr.Pin().Get(), NewText);
 	}
 
 	/**
@@ -624,7 +619,7 @@ public:
 		check(ActionPtr.Pin()->GetTypeId() == FEdGraphSchemaAction_K2Event::StaticGetTypeId());
 
 		FEdGraphSchemaAction_K2Event* EventAction = (FEdGraphSchemaAction_K2Event*)ActionPtr.Pin().Get();
-		if (EventAction->NodeTemplate != NULL)
+		if (EventAction->NodeTemplate)
 		{
 			EventAction->NodeTemplate->OnRenameNode(NewText.ToString());
 		}
@@ -649,7 +644,7 @@ public:
 		FEdGraphSchemaAction_K2TargetNode* TargetNodeAction = (FEdGraphSchemaAction_K2TargetNode*)ActionPtr.Pin().Get();
 
 		UK2Node* AssociatedNode = TargetNodeAction->NodeTemplate;
-		if ((AssociatedNode != NULL) && AssociatedNode->bCanRenameNode)
+		if (AssociatedNode && AssociatedNode->bCanRenameNode)
 		{
 			TSharedPtr<INameValidatorInterface> NodeNameValidator = FNameValidatorFactory::MakeValidator(AssociatedNode);
 			bIsNameValid = (NodeNameValidator->IsValid(InNewText.ToString(), true) == EValidatorResult::Ok);
@@ -670,7 +665,7 @@ public:
 		check(ActionPtr.Pin()->GetTypeId() == FEdGraphSchemaAction_K2TargetNode::StaticGetTypeId());
 
 		FEdGraphSchemaAction_K2TargetNode* TargetNodeAction = (FEdGraphSchemaAction_K2TargetNode*)ActionPtr.Pin().Get();
-		if (TargetNodeAction->NodeTemplate != NULL)
+		if (TargetNodeAction->NodeTemplate)
 		{
 			TargetNodeAction->NodeTemplate->OnRenameNode(NewText.ToString());
 		}
@@ -797,9 +792,9 @@ public:
 			UProperty* VariableProp = StaticCastSharedPtr<FEdGraphSchemaAction_K2Var>(PaletteAction)->GetProperty();
 			UObjectProperty* VariableObjProp = Cast<UObjectProperty>(VariableProp);
 
-			UStruct* VarSourceScope = (VariableProp != NULL) ? CastChecked<UStruct>(VariableProp->GetOuter()) : NULL;
+			UStruct* VarSourceScope = (VariableProp ? CastChecked<UStruct>(VariableProp->GetOuter()) : nullptr);
 			const bool bIsBlueprintVariable = (VarSourceScope == BlueprintObj->SkeletonGeneratedClass);
-			const bool bIsComponentVar = (VariableObjProp != NULL) && (VariableObjProp->PropertyClass != NULL) && (VariableObjProp->PropertyClass->IsChildOf(UActorComponent::StaticClass()));
+			const bool bIsComponentVar = (VariableObjProp && VariableObjProp->PropertyClass && VariableObjProp->PropertyClass->IsChildOf(UActorComponent::StaticClass()));
 			bShouldHaveAVisibilityToggle = bIsBlueprintVariable && (!bIsComponentVar || FBlueprintEditorUtils::IsVariableCreatedByBlueprint(BlueprintObj, VariableObjProp));
 		}
 
@@ -845,8 +840,7 @@ private:
 		if ( PaletteAction->GetTypeId() == FEdGraphSchemaAction_K2Var::StaticGetTypeId() )
 		{
 			TSharedPtr<FEdGraphSchemaAction_K2Var> VarAction = StaticCastSharedPtr<FEdGraphSchemaAction_K2Var>(PaletteAction);
-			UProperty* VariableProperty = VarAction->GetProperty();
-			if ( VariableProperty != NULL )
+			if (UProperty* VariableProperty = VarAction->GetProperty())
 			{
 				return VariableProperty->HasAnyPropertyFlags(CPF_DisableEditOnInstance) ? ECheckBoxState::Unchecked : ECheckBoxState::Checked;
 			}
@@ -912,7 +906,7 @@ private:
 			TSharedPtr<FEdGraphSchemaAction_K2Var> VarAction = StaticCastSharedPtr<FEdGraphSchemaAction_K2Var>(ActionPtr.Pin());
 
 			FString Result;
-			FBlueprintEditorUtils::GetBlueprintVariableMetaData(BlueprintObj, VarAction->GetVariableName(), NULL, TEXT("tooltip"), Result);
+			FBlueprintEditorUtils::GetBlueprintVariableMetaData(BlueprintObj, VarAction->GetVariableName(), nullptr, TEXT("tooltip"), Result);
 
 			if ( !Result.IsEmpty() )
 			{
@@ -944,7 +938,7 @@ private:
 			TSharedPtr<FEdGraphSchemaAction_K2Var> VarAction = StaticCastSharedPtr<FEdGraphSchemaAction_K2Var>(ActionPtr.Pin());
 
 			FString Result;
-			FBlueprintEditorUtils::GetBlueprintVariableMetaData(BlueprintObj, VarAction->GetVariableName(), NULL, TEXT("tooltip"), Result);
+			FBlueprintEditorUtils::GetBlueprintVariableMetaData(BlueprintObj, VarAction->GetVariableName(), nullptr, TEXT("tooltip"), Result);
 			if ( !Result.IsEmpty() )
 			{
 				ToolTipText = LOCTEXT("VariablePrivacy_is_public_Tooltip", "Variable is public and is editable on each instance of this Blueprint.");
@@ -1232,7 +1226,7 @@ bool SBlueprintPaletteItem::OnNameTextVerifyChanged(const FText& InNewText, FTex
 
 	FName OriginalName;
 
-	UStruct* ValidationScope = NULL;
+	UStruct* ValidationScope = nullptr;
 
 	// Check if certain action names are unchanged.
 	if (ActionPtr.Pin()->GetTypeId() == FEdGraphSchemaAction_K2Var::StaticGetTypeId())
@@ -1249,7 +1243,7 @@ bool SBlueprintPaletteItem::OnNameTextVerifyChanged(const FText& InNewText, FTex
 	}
 	else
 	{
-		UEdGraph* Graph = NULL;
+		UEdGraph* Graph = nullptr;
 
 		if(ActionPtr.Pin()->GetTypeId() == FEdGraphSchemaAction_K2Graph::StaticGetTypeId())
 		{
@@ -1269,9 +1263,9 @@ bool SBlueprintPaletteItem::OnNameTextVerifyChanged(const FText& InNewText, FTex
 	}
 
 	UBlueprint* BlueprintObj = BlueprintEditorPtr.Pin()->GetBlueprintObj();
-	check(BlueprintObj != NULL);
+	check(BlueprintObj);
 
-	if(BlueprintObj->SimpleConstructionScript != NULL)
+	if (BlueprintObj->SimpleConstructionScript)
 	{
 		for (USCS_Node* Node : BlueprintObj->SimpleConstructionScript->GetAllNodes())
 		{
@@ -1311,7 +1305,7 @@ void SBlueprintPaletteItem::OnNameTextCommitted(const FText& NewText, ETextCommi
 		FEdGraphSchemaAction_K2Graph* GraphAction = (FEdGraphSchemaAction_K2Graph*)ActionPtr.Pin().Get();
 
 		UEdGraph* Graph = GraphAction->EdGraph;
-		if ((Graph != NULL) && (Graph->bAllowDeletion || Graph->bAllowRenaming))
+		if (Graph && (Graph->bAllowDeletion || Graph->bAllowRenaming))
 		{
 			if (GraphAction->EdGraph)
 			{
@@ -1327,7 +1321,7 @@ void SBlueprintPaletteItem::OnNameTextCommitted(const FText& NewText, ETextCommi
 
 			// Make sure we aren't renaming the graph into something that already exists
 			UEdGraph* ExistingGraph = FindObject<UEdGraph>(Graph->GetOuter(), *NewNameString );
-			if (ExistingGraph == NULL || ExistingGraph == Graph)
+			if (ExistingGraph == nullptr || ExistingGraph == Graph)
 			{
 				const FScopedTransaction Transaction( LOCTEXT( "Rename Function", "Rename Function" ) );
 
@@ -1358,7 +1352,7 @@ void SBlueprintPaletteItem::OnNameTextCommitted(const FText& NewText, ETextCommi
 		FEdGraphSchemaAction_K2Delegate* DelegateAction = (FEdGraphSchemaAction_K2Delegate*)ActionPtr.Pin().Get();
 
 		UEdGraph* Graph = DelegateAction->EdGraph;
-		if ((Graph != NULL) && (Graph->bAllowDeletion || Graph->bAllowRenaming))
+		if (Graph && (Graph->bAllowDeletion || Graph->bAllowRenaming))
 		{
 			FGraphDisplayInfo DisplayInfo;
 			Graph->GetSchema()->GetGraphDisplayInformation(*Graph, DisplayInfo);
@@ -1371,7 +1365,7 @@ void SBlueprintPaletteItem::OnNameTextCommitted(const FText& NewText, ETextCommi
 
 			// Make sure we aren't renaming the graph into something that already exists
 			UEdGraph* ExistingGraph = FindObject<UEdGraph>(Graph->GetOuter(), *NewNameString );
-			if (ExistingGraph == NULL || ExistingGraph == Graph)
+			if (ExistingGraph == nullptr || ExistingGraph == Graph)
 			{
 				const FScopedTransaction Transaction( LOCTEXT( "Rename Delegate", "Rename Event Dispatcher" ) );
 				const FName OldName =  Graph->GetFName();
@@ -1397,12 +1391,11 @@ void SBlueprintPaletteItem::OnNameTextCommitted(const FText& NewText, ETextCommi
 
 		// Double check we're not renaming a timeline disguised as a variable
 		bool bIsTimeline = false;
-		UProperty* VariableProperty = VarAction->GetProperty();
-		if (VariableProperty != NULL)
+		if (UProperty* VariableProperty = VarAction->GetProperty())
 		{
 			// Don't allow removal of timeline properties - you need to remove the timeline node for that
 			UObjectProperty* ObjProperty = Cast<UObjectProperty>(VariableProperty);
-			if(ObjProperty != NULL && ObjProperty->PropertyClass == UTimelineComponent::StaticClass())
+			if (ObjProperty && ObjProperty->PropertyClass == UTimelineComponent::StaticClass())
 			{
 				bIsTimeline = true;
 			}
@@ -1488,7 +1481,7 @@ FText SBlueprintPaletteItem::GetToolTipText() const
 		else if (PaletteAction->GetTypeId() == FEdGraphSchemaAction_K2Graph::StaticGetTypeId())
 		{
 			FEdGraphSchemaAction_K2Graph* GraphAction = (FEdGraphSchemaAction_K2Graph*)PaletteAction.Get();
-			if (GraphAction->EdGraph != NULL && GraphAction->EdGraph->GetSchema())
+			if (GraphAction->EdGraph)
 			{
 				if (const UEdGraphSchema* GraphSchema = GraphAction->EdGraph->GetSchema())
 				{
@@ -1502,10 +1495,10 @@ FText SBlueprintPaletteItem::GetToolTipText() const
 		{
 			FEdGraphSchemaAction_K2Var* VarAction = (FEdGraphSchemaAction_K2Var*)PaletteAction.Get();
 			UClass* VarClass = VarAction->GetVariableClass();
-			if (bShowClassInTooltip && (VarClass != NULL))
+			if (bShowClassInTooltip && VarClass)
 			{
 				UBlueprint* BlueprintObj = UBlueprint::GetBlueprintFromClass(VarClass);
-				ClassDisplayName = (BlueprintObj != NULL) ? FText::FromName(BlueprintObj->GetFName()) : VarClass->GetDisplayNameText();
+				ClassDisplayName = (BlueprintObj ? FText::FromName(BlueprintObj->GetFName()) : VarClass->GetDisplayNameText());
 			}
 			else
 			{
@@ -1521,10 +1514,10 @@ FText SBlueprintPaletteItem::GetToolTipText() const
 			if(LocalVarAction->GetVariableScope())
 			{
 				UClass* VarClass = CastChecked<UClass>(LocalVarAction->GetVariableScope()->GetOuter());
-				if (bShowClassInTooltip && (VarClass != NULL))
+				if (bShowClassInTooltip && (VarClass != nullptr))
 				{
 					UBlueprint* BlueprintObj = UBlueprint::GetBlueprintFromClass(VarClass);
-					ClassDisplayName = (BlueprintObj != NULL) ? FText::FromName(BlueprintObj->GetFName()) : VarClass->GetDisplayNameText();
+					ClassDisplayName = (BlueprintObj ? FText::FromName(BlueprintObj->GetFName()) : VarClass->GetDisplayNameText());
 				}
 				else
 				{
@@ -1598,10 +1591,10 @@ TSharedPtr<SToolTip> SBlueprintPaletteItem::ConstructToolTipWidget() const
 		else if (PaletteAction->GetTypeId() == FEdGraphSchemaAction_K2Graph::StaticGetTypeId())
 		{
 			FEdGraphSchemaAction_K2Graph* GraphAction = (FEdGraphSchemaAction_K2Graph*)PaletteAction.Get();
-			if (GraphAction->EdGraph != NULL)
+			if (GraphAction->EdGraph)
 			{
 				FGraphDisplayInfo DisplayInfo;
-				if (auto GraphSchema = GraphAction->EdGraph->GetSchema())
+				if (const UEdGraphSchema* GraphSchema = GraphAction->EdGraph->GetSchema())
 				{
 					GraphSchema->GetGraphDisplayInformation(*(GraphAction->EdGraph), DisplayInfo);
 				}
@@ -1614,7 +1607,7 @@ TSharedPtr<SToolTip> SBlueprintPaletteItem::ConstructToolTipWidget() const
 		{
 			FEdGraphSchemaAction_K2Var* VarAction = (FEdGraphSchemaAction_K2Var*)PaletteAction.Get();
 			UClass* VarClass = VarAction->GetVariableClass();
-			if (!bShowClassInTooltip || VarClass == NULL)
+			if (!bShowClassInTooltip || VarClass == nullptr)
 			{
 				// Don't show big tooltip if we are showing class as well (means we are not in MyBlueprint)
 				DocExcerptRef.DocLink = TEXT("Shared/Editors/BlueprintEditor/GraphTypes");
@@ -1646,7 +1639,7 @@ TSharedPtr<SToolTip> SBlueprintPaletteItem::ConstructToolTipWidget() const
 	TAttribute<FText> TextAttribute;
 	TextAttribute.Bind(this, &SBlueprintPaletteItem::GetToolTipText);
 
-	TSharedRef< SToolTip > TooltipWidget = IDocumentation::Get()->CreateToolTip(TextAttribute, NULL, DocExcerptRef.DocLink, DocExcerptRef.DocExcerptName);
+	TSharedRef< SToolTip > TooltipWidget = IDocumentation::Get()->CreateToolTip(TextAttribute, nullptr, DocExcerptRef.DocLink, DocExcerptRef.DocExcerptName);
 
 	// English speakers have no real need to know this exists.
 	if ( (NodeTemplate != nullptr) && (FInternationalization::Get().GetCurrentCulture()->GetTwoLetterISOLanguageName() != TEXT("en")) )
diff --git a/Engine/Source/Editor/Kismet/Private/SMyBlueprint.cpp b/Engine/Source/Editor/Kismet/Private/SMyBlueprint.cpp
--- a/Engine/Source/Editor/Kismet/Private/SMyBlueprint.cpp
+++ b/Engine/Source/Editor/Kismet/Private/SMyBlueprint.cpp
@@ -573,7 +573,7 @@ void SMyBlueprint::OnCategoryNameCommitted(const FText& InNewText, ETextCommit::
 					// Don't allow changing the category of a graph who's parent is not the current Blueprint
 					if(GraphAction && !FBlueprintEditorUtils::IsPaletteActionReadOnly(Actions[i], BlueprintEditorPtr.Pin()) && FBlueprintEditorUtils::FindBlueprintForGraph(GraphAction->EdGraph) == GetBlueprintObj())
 					{
-						auto EntryNode = FBlueprintEditorUtils::GetEntryNode(GraphAction->EdGraph);
+						UK2Node_EditablePinBase* EntryNode = FBlueprintEditorUtils::GetEntryNode(GraphAction->EdGraph);
 						EntryNode->Modify();
 						if (UK2Node_FunctionEntry* FunctionEntryNode = Cast<UK2Node_FunctionEntry>(EntryNode))
 						{
@@ -945,10 +945,8 @@ struct FCreateEdGraphSchemaActionHelper
 	{
 		TArray<NodeType*> EventNodes;
 		EdGraph->GetNodesOfClass<NodeType>(EventNodes);
-		for (auto It = EventNodes.CreateConstIterator(); It; ++It)
+		for (NodeType* const EventNode : EventNodes)
 		{
-			NodeType* const EventNode = (*It);
-
 			const FText Tooltip = EventNode->GetTooltipText();
 			const FText Description = EventNode->GetNodeTitle(ENodeTitleType::EditableTitle);
 
@@ -1008,13 +1006,12 @@ void SMyBlueprint::GetLocalVariables(FGraphActionSort& SortList) const
 
 		// Search in all FunctionEntry nodes for their local variables
 		FText ActionCategory;
-		const UEdGraphSchema_K2* K2Schema = GetDefault<UEdGraphSchema_K2>();
 		for (UK2Node_FunctionEntry* const FunctionEntry : FunctionEntryNodes)
 		{
 			for (const FBPVariableDescription& Variable : FunctionEntry->LocalVariables)
 			{
 				FText Category = Variable.Category;
-				if (Variable.Category.EqualTo(K2Schema->VR_DefaultCategory))
+				if (Variable.Category.EqualTo(UEdGraphSchema_K2::VR_DefaultCategory))
 				{
 					Category = FText::GetEmpty();
 				}
@@ -1023,7 +1020,7 @@ void SMyBlueprint::GetLocalVariables(FGraphActionSort& SortList) const
 				if (Func)
 				{
 					TSharedPtr<FEdGraphSchemaAction_K2LocalVar> NewVarAction = MakeShareable(new FEdGraphSchemaAction_K2LocalVar(Category, FText::FromName(Variable.VarName), FText::GetEmpty(), 0, NodeSectionID::LOCAL_VARIABLE));
-					NewVarAction->SetVariableInfo(Variable.VarName, Func, Variable.VarType.PinCategory == K2Schema->PC_Boolean);
+					NewVarAction->SetVariableInfo(Variable.VarName, Func, Variable.VarType.PinCategory == UEdGraphSchema_K2::PC_Boolean);
 					SortList.AddAction(NewVarAction);
 				}
 			}
@@ -1097,7 +1094,7 @@ void SMyBlueprint::CollectAllActions(FGraphActionListBuilderBase& OutAllActions)
 		FText PropertyCategory = FObjectEditorUtils::GetCategoryText(Property);
 		const FString UserCategoryName = FEditorCategoryUtils::GetCategoryDisplayString( PropertyCategory.ToString() );
 
-		if (CategoryName.EqualTo(FText::FromString(BlueprintObj->GetName())) || CategoryName.EqualTo(K2Schema->VR_DefaultCategory))
+		if (CategoryName.EqualTo(FText::FromString(BlueprintObj->GetName())) || CategoryName.EqualTo(UEdGraphSchema_K2::VR_DefaultCategory))
 		{
 			CategoryName = FText::GetEmpty();		// default, so place in 'non' category
 			PropertyCategory = FText::GetEmpty();
@@ -1162,7 +1159,7 @@ void SMyBlueprint::CollectAllActions(FGraphActionListBuilderBase& OutAllActions)
 		}
 
 		//@TODO: Should be a bit more generic (or the AnimGraph shouldn't be stored as a FunctionGraph...)
-		const bool bIsConstructionScript = Graph->GetFName() == K2Schema->FN_UserConstructionScript;
+		const bool bIsConstructionScript = Graph->GetFName() == UEdGraphSchema_K2::FN_UserConstructionScript;
 		int32 SectionID = Graph->IsA<UAnimationGraph>() ? NodeSectionID::GRAPH : NodeSectionID::FUNCTION;
 		TSharedPtr<FEdGraphSchemaAction_K2Graph> NewFuncAction = MakeShareable(new FEdGraphSchemaAction_K2Graph(EEdGraphSchemaAction_K2Graph::Function, FunctionCategory, DisplayInfo.PlainName, DisplayInfo.Tooltip, bIsConstructionScript ? 2 : 1, SectionID));
 		NewFuncAction->FuncName = Graph->GetFName();
@@ -1452,11 +1449,11 @@ FReply SMyBlueprint::OnActionDragged( const TArray< TSharedPtr<FEdGraphSchemaAct
 				};
 
 				bool bIsBlueprintCallableFunction = false;
-				if (FuncAction->EdGraph != NULL)
+				if (FuncAction->EdGraph)
 				{
-					for (auto It = FuncAction->EdGraph->Nodes.CreateConstIterator(); It; ++It)
+					for (UEdGraphNode* GraphNode : FuncAction->EdGraph->Nodes)
 					{
-						if (UK2Node_FunctionEntry* Node = Cast<UK2Node_FunctionEntry>(*It))
+						if (UK2Node_FunctionEntry* Node = Cast<UK2Node_FunctionEntry>(GraphNode))
 						{
 							// See whether this node is a blueprint callable function
 							if (Node->GetFunctionFlags() & (FUNC_BlueprintCallable|FUNC_BlueprintPure))
@@ -1798,7 +1795,7 @@ FText SMyBlueprint::GetGraphCategory(UEdGraph* InGraph) const
 	FText ReturnCategory;
 
 	// Pull the category from the required metadata based on the types of nodes we can discover in the graph
-	auto EntryNode = FBlueprintEditorUtils::GetEntryNode(InGraph);
+	UK2Node_EditablePinBase* EntryNode = FBlueprintEditorUtils::GetEntryNode(InGraph);
 	if (UK2Node_FunctionEntry* FunctionEntryNode = Cast<UK2Node_FunctionEntry>(EntryNode))
 	{
 		ReturnCategory = FunctionEntryNode->MetaData.Category;
@@ -1809,8 +1806,7 @@ FText SMyBlueprint::GetGraphCategory(UEdGraph* InGraph) const
 	}
 
 	// Empty the category if it's default, we don't want to display the "default" category and items will just appear without a category
-	const UEdGraphSchema_K2* K2Schema = GetDefault<UEdGraphSchema_K2>();
-	if(ReturnCategory.EqualTo(K2Schema->VR_DefaultCategory))
+	if(ReturnCategory.EqualTo(UEdGraphSchema_K2::VR_DefaultCategory))
 	{
 		ReturnCategory = FText::GetEmpty();
 	}
@@ -1996,7 +1992,7 @@ void SMyBlueprint::OnFocusNode()
 	FEdGraphSchemaAction_K2InputAction* InputAction = SelectionAsInputAction();
 	if (EventAction || InputAction)
 	{
-		auto Node = EventAction ? EventAction->NodeTemplate : InputAction->NodeTemplate;
+		UK2Node* Node = EventAction ? EventAction->NodeTemplate : InputAction->NodeTemplate;
 		FKismetEditorUtilities::BringKismetToFocusAttentionOnObject(Node);
 	}
 }
@@ -2212,9 +2208,8 @@ void SMyBlueprint::OnDeleteGraph(UEdGraph* InGraph, EEdGraphSchemaAction_K2Graph
 			FBlueprintEditorUtils::GetAllNodesOfClass<UK2Node_Composite>(GetBlueprintObj(), AllCompositeNodes);
 
 			const bool bDontRecompile = true;
-			for (auto CompIt = AllCompositeNodes.CreateIterator(); CompIt; ++CompIt)
+			for (UK2Node_Composite* CompNode : AllCompositeNodes)
 			{
-				UK2Node_Composite* CompNode = *CompIt;
 				if (CompNode->BoundGraph == InGraph)
 				{
 					FBlueprintEditorUtils::RemoveNode(GetBlueprintObj(), CompNode, bDontRecompile);
@@ -2242,7 +2237,7 @@ void SMyBlueprint::OnDeleteGraph(UEdGraph* InGraph, EEdGraphSchemaAction_K2Graph
 
 UEdGraph* SMyBlueprint::GetFocusedGraph() const
 {
-	auto BlueprintEditorPtrPinned = BlueprintEditorPtr.Pin();
+	TSharedPtr<FBlueprintEditor> BlueprintEditorPtrPinned = BlueprintEditorPtr.Pin();
 	if( BlueprintEditorPtrPinned.IsValid() )
 	{
 		return BlueprintEditorPtrPinned->GetFocusedGraph();
@@ -2261,7 +2256,7 @@ void SMyBlueprint::OnObjectPropertyChanged(UObject* InObject, FPropertyChangedEv
 
 bool SMyBlueprint::IsEditingMode() const
 {
-	auto BlueprintEditorSPtr = BlueprintEditorPtr.Pin();
+	TSharedPtr<FBlueprintEditor> BlueprintEditorSPtr = BlueprintEditorPtr.Pin();
 	return BlueprintEditorSPtr.IsValid() && BlueprintEditorSPtr->InEditingMode();
 }
 
@@ -2669,13 +2664,11 @@ void SMyBlueprint::OnResetItemFilter()
 
 void SMyBlueprint::EnsureLastPinTypeValid()
 {
-	const UEdGraphSchema_K2* Schema = GetDefault<UEdGraphSchema_K2>();
-
 	LastPinType.bIsWeakPointer = false;
 	LastFunctionPinType.bIsWeakPointer = false;
 
-	const bool bLastPinTypeValid = (Schema->PC_Struct != LastPinType.PinCategory) || LastPinType.PinSubCategoryObject.IsValid();
-	const bool bLastFunctionPinTypeValid = (Schema->PC_Struct != LastFunctionPinType.PinCategory) || LastFunctionPinType.PinSubCategoryObject.IsValid();
+	const bool bLastPinTypeValid = (UEdGraphSchema_K2::PC_Struct != LastPinType.PinCategory) || LastPinType.PinSubCategoryObject.IsValid();
+	const bool bLastFunctionPinTypeValid = (UEdGraphSchema_K2::PC_Struct != LastFunctionPinType.PinCategory) || LastFunctionPinType.PinSubCategoryObject.IsValid();
 	const bool bConstType = LastPinType.bIsConst || LastFunctionPinType.bIsConst;
 	if (!bLastPinTypeValid || !bLastFunctionPinTypeValid || bConstType)
 	{
@@ -2685,10 +2678,8 @@ void SMyBlueprint::EnsureLastPinTypeValid()
 
 void SMyBlueprint::ResetLastPinType()
 {
-	const UEdGraphSchema_K2* Schema = GetDefault<UEdGraphSchema_K2>();
-
 	LastPinType.ResetToDefaults();
-	LastPinType.PinCategory = Schema->PC_Boolean;
+	LastPinType.PinCategory = UEdGraphSchema_K2::PC_Boolean;
 	LastFunctionPinType = LastPinType;
 }
 
diff --git a/Engine/Source/Editor/Kismet/Private/SReplaceNodeReferences.cpp b/Engine/Source/Editor/Kismet/Private/SReplaceNodeReferences.cpp
--- a/Engine/Source/Editor/Kismet/Private/SReplaceNodeReferences.cpp
+++ b/Engine/Source/Editor/Kismet/Private/SReplaceNodeReferences.cpp
@@ -330,7 +330,7 @@ void SReplaceNodeReferences::GatherAllAvailableBlueprintVariables(UClass* InTarg
 		const FString UserCategoryName = FEditorCategoryUtils::GetCategoryDisplayString( PropertyCategory.ToString() );
 
 		UEdGraphSchema_K2 const* K2Schema = GetDefault<UEdGraphSchema_K2>();
-		if (CategoryName.EqualTo(FText::FromString(PathObject->GetName())) || CategoryName.EqualTo(K2Schema->VR_DefaultCategory))
+		if (CategoryName.EqualTo(FText::FromString(PathObject->GetName())) || CategoryName.EqualTo(UEdGraphSchema_K2::VR_DefaultCategory))
 		{
 			CategoryName = FText::GetEmpty();		// default, so place in 'non' category
 			PropertyCategory = FText::GetEmpty();
diff --git a/Engine/Source/Editor/Kismet/Private/STimelineEditor.cpp b/Engine/Source/Editor/Kismet/Private/STimelineEditor.cpp
--- a/Engine/Source/Editor/Kismet/Private/STimelineEditor.cpp
+++ b/Engine/Source/Editor/Kismet/Private/STimelineEditor.cpp
@@ -1540,7 +1540,7 @@ bool STimelineEditor::OnVerifyTrackNameCommit(const FText& TrackName, FText& Out
 			{
 				UEdGraphPin* Pin = *PinIt;
 
-				if (Pin->PinName == TrackName.ToString())
+				if (Pin->PinName == RequestedName)
 				{
 					FFormatNamedArguments Args;
 					Args.Add(TEXT("TrackName"), TrackName);
@@ -1575,10 +1575,10 @@ void STimelineEditor::OnTrackNameCommitted( const FText& StringName, ETextCommit
 			{
 				UEdGraphPin* Pin = TimelineNode->Pins[PinIdx];
 			
-				if (Pin->PinName == TrackBase->TrackName.ToString())
+				if (Pin->PinName == TrackBase->TrackName)
 				{
 					Pin->Modify();
-					Pin->PinName = StringName.ToString();
+					Pin->PinName = RequestedName;
 					break;
 				}
 			}
diff --git a/Engine/Source/Editor/Kismet/Private/Tests/BlueprintEditorTests.cpp b/Engine/Source/Editor/Kismet/Private/Tests/BlueprintEditorTests.cpp
--- a/Engine/Source/Editor/Kismet/Private/Tests/BlueprintEditorTests.cpp
+++ b/Engine/Source/Editor/Kismet/Private/Tests/BlueprintEditorTests.cpp
@@ -384,7 +384,7 @@ namespace BlueprintEditorPromotionUtils
 	* @param NodeB - The second node to connect
 	* @param PinBName - The name of the pin on the second node
 	*/
-	static void ConnectGraphNodes(UEdGraphNode* NodeA, const FString& PinAName, UEdGraphNode* NodeB, const FString& PinBName)
+	static void ConnectGraphNodes(UEdGraphNode* NodeA, const FName& PinAName, UEdGraphNode* NodeB, const FName& PinBName)
 	{
 		const FScopedTransaction PropertyChanged(LOCTEXT("ConnectedNode", "Connected graph nodes"));
 		NodeA->GetGraph()->Modify();
@@ -398,7 +398,7 @@ namespace BlueprintEditorPromotionUtils
 		}
 		else
 		{
-			UE_LOG(LogBlueprintEditorPromotionTests, Error, TEXT("Could not connect pins %s and %s "), *PinAName, *PinBName);
+			UE_LOG(LogBlueprintEditorPromotionTests, Error, TEXT("Could not connect pins %s and %s "), *PinAName.ToString(), *PinBName.ToString());
 		}
 	}
 	
@@ -409,7 +409,7 @@ namespace BlueprintEditorPromotionUtils
 	* @param Node - The node that owns the pin
 	* @param PinName - The name of the pin to promote
 	*/
-	static void PromotePinToVariable(UBlueprint* InBlueprint, UEdGraphNode* Node, const FString& PinName)
+	static void PromotePinToVariable(UBlueprint* InBlueprint, UEdGraphNode* Node, const FName& PinName)
 	{
 		IAssetEditorInstance* OpenEditor = FAssetEditorManager::Get().FindEditorForAsset(InBlueprint, true);
 		FBlueprintEditor* CurrentBlueprintEditor = (FBlueprintEditor*)OpenEditor;
@@ -579,7 +579,7 @@ namespace BlueprintEditorPromotionUtils
 	* @param PinName - The name of the pin
 	* @param PinValue - The new default value
 	*/
-	static void SetPinDefaultValue(UEdGraphNode* Node, const FString& PinName, const FString& PinValue)
+	static void SetPinDefaultValue(UEdGraphNode* Node, const FName PinName, const FString& PinValue)
 	{
 		UEdGraphPin* Pin = Node->FindPin(PinName);
 		Pin->DefaultValue = PinValue;
@@ -592,7 +592,7 @@ namespace BlueprintEditorPromotionUtils
 	* @param PinName - The name of the pin
 	* @param PinObject - The new default object
 	*/
-	static void SetPinDefaultObject(UEdGraphNode* Node, const FString& PinName, UObject* PinObject)
+	static void SetPinDefaultObject(UEdGraphNode* Node, const FName PinName, UObject* PinObject)
 	{
 		UEdGraphPin* Pin = Node->FindPin(PinName);
 		Pin->DefaultObject = PinObject;
@@ -606,8 +606,7 @@ namespace BlueprintEditorPromotionUtils
 	*/
 	static void AddStringMemberValue(UBlueprint* InBlueprint, const FName& VariableName)
 	{
-		const UEdGraphSchema_K2* K2Schema = GetDefault<UEdGraphSchema_K2>();
-		FEdGraphPinType StringPinType(K2Schema->PC_String, FString(), nullptr, EPinContainerType::None, false, FEdGraphTerminalType());
+		FEdGraphPinType StringPinType(UEdGraphSchema_K2::PC_String, NAME_None, nullptr, EPinContainerType::None, false, FEdGraphTerminalType());
 		FBlueprintEditorUtils::AddMemberVariable(InBlueprint, VariableName, StringPinType);
 	}
 
@@ -1043,8 +1042,6 @@ namespace BlueprintEditorPromotionTestHelper
 		{
 			if (BlueprintObject)
 			{
-				const UEdGraphSchema_K2* K2Schema = GetDefault<UEdGraphSchema_K2>();
-
 				IAssetEditorInstance* AssetEditor = FAssetEditorManager::Get().FindEditorForAsset(BlueprintObject, true);
 				FBlueprintEditor* BlueprintEditor = (FBlueprintEditor*)AssetEditor;
 
@@ -1066,10 +1063,10 @@ namespace BlueprintEditorPromotionTestHelper
 				Test->TestNotNull(TEXT("Found entry node to connect Add Static Mesh to"), EntryNode);
 				if (EntryNode)
 				{
-					BlueprintEditorPromotionUtils::ConnectGraphNodes(AddMeshNode, K2Schema->PN_Execute, EntryNode, K2Schema->PN_Then);
+					BlueprintEditorPromotionUtils::ConnectGraphNodes(AddMeshNode, UEdGraphSchema_K2::PN_Execute, EntryNode, UEdGraphSchema_K2::PN_Then);
 
-					UEdGraphPin* EntryOutPin = EntryNode->FindPin(K2Schema->PN_Then);
-					UEdGraphPin* AddStaticMeshInPin = AddMeshNode->FindPin(K2Schema->PN_Execute);
+					UEdGraphPin* EntryOutPin = EntryNode->FindPin(UEdGraphSchema_K2::PN_Then);
+					UEdGraphPin* AddStaticMeshInPin = AddMeshNode->FindPin(UEdGraphSchema_K2::PN_Execute);
 
 					Test->TestTrue(TEXT("Connected entry node to Add Static Mesh node"), EntryOutPin->LinkedTo.Contains(AddStaticMeshInPin));
 
@@ -1111,8 +1108,7 @@ namespace BlueprintEditorPromotionTestHelper
 				IAssetEditorInstance* AssetEditor = FAssetEditorManager::Get().FindEditorForAsset(BlueprintObject, true);
 				FBlueprintEditor* BlueprintEditor = (FBlueprintEditor*)AssetEditor;
 
-				const UEdGraphSchema_K2* K2Schema = GetDefault<UEdGraphSchema_K2>();
-				BlueprintEditorPromotionUtils::PromotePinToVariable(BlueprintObject, AddMeshNode, K2Schema->PN_ReturnValue);
+				BlueprintEditorPromotionUtils::PromotePinToVariable(BlueprintObject, AddMeshNode, UEdGraphSchema_K2::PN_ReturnValue);
 
 				Test->AddInfo(TEXT("Promoted the return pin on the add mesh node to a variable"));
 
@@ -1177,8 +1173,7 @@ namespace BlueprintEditorPromotionTestHelper
 				PostBeginPlayEventNode = BlueprintEditorPromotionUtils::CreatePostBeginPlayEvent(BlueprintObject, EventGraph);
 				Test->TestNotNull(TEXT("Created EventBeginPlay node"), PostBeginPlayEventNode);
 				
-				const UEdGraphSchema_K2* K2Schema = GetDefault<UEdGraphSchema_K2>();
-				UEdGraphPin* PlayThenPin = PostBeginPlayEventNode->FindPin(K2Schema->PN_Then);
+				UEdGraphPin* PlayThenPin = PostBeginPlayEventNode->FindPin(UEdGraphSchema_K2::PN_Then);
 				
 				DelayNode = BlueprintEditorPromotionUtils::AddDelayNode(BlueprintObject, EventGraph, PlayThenPin);
 				Test->TestNotNull(TEXT("Created Delay node"), DelayNode);
@@ -1190,7 +1185,7 @@ namespace BlueprintEditorPromotionTestHelper
 				Test->TestTrue(TEXT("Redo adding Delay node succeeded"), EventGraph->Nodes.Num() > 0 && EventGraph->Nodes[EventGraph->Nodes.Num() - 1] == DelayNode);
 				
 				// Update Delay node's Duration pin with new default value
-				const FString DelayDurationPinName = TEXT("Duration");
+				const FName DelayDurationPinName = TEXT("Duration");
 				const FString NewDurationDefaultValue = TEXT("2.0");
 
 				BlueprintEditorPromotionUtils::SetPinDefaultValue(DelayNode, DelayDurationPinName, NewDurationDefaultValue);				
@@ -1256,18 +1251,16 @@ namespace BlueprintEditorPromotionTestHelper
 				PrintNode = BlueprintEditorPromotionUtils::AddPrintStringNode(BlueprintObject, EventGraph);
 				Test->TestNotNull(TEXT("Added Print String node"), PrintNode);
 				
-				const UEdGraphSchema_K2* K2Schema = GetDefault<UEdGraphSchema_K2>();
-
 				//Connect Get to printstring
-				UEdGraphPin* GetVarPin = GetNode->FindPin(BlueprintEditorPromotionUtils::BlueprintStringVariableName.ToString());
+				UEdGraphPin* GetVarPin = GetNode->FindPin(BlueprintEditorPromotionUtils::BlueprintStringVariableName);
 				UEdGraphPin* InStringPin = PrintNode->FindPin(TEXT("InString"));
 				GetVarPin->MakeLinkTo(InStringPin);
 				Test->TestTrue(TEXT("Connected string variable Get node to the Print String node"), GetVarPin->LinkedTo.Contains(InStringPin));
 
 
 				//Connect Delay to PrintString
-				UEdGraphPin* DelayExecPin = DelayNode->FindPin(K2Schema->PN_Then);
-				UEdGraphPin* PrintStringPin = PrintNode->FindPin(K2Schema->PN_Execute);
+				UEdGraphPin* DelayExecPin = DelayNode->FindPin(UEdGraphSchema_K2::PN_Then);
+				UEdGraphPin* PrintStringPin = PrintNode->FindPin(UEdGraphSchema_K2::PN_Execute);
 				DelayExecPin->MakeLinkTo(PrintStringPin);
 				Test->TestTrue(TEXT("Connected Delay nod to Print String node"), DelayExecPin->LinkedTo.Contains(PrintStringPin));
 
@@ -1279,7 +1272,7 @@ namespace BlueprintEditorPromotionTestHelper
 				Test->TestNotNull(TEXT("Added Set Static Mesh node"), SetStaticMeshNode);
 				
  				UEdGraphPin* GetExecPin = GetNode->FindPin(TEXT("MyMesh"));
-				UEdGraphPin* SetStaticMeshSelfPin = SetStaticMeshNode->FindPin(K2Schema->PN_Self);
+				UEdGraphPin* SetStaticMeshSelfPin = SetStaticMeshNode->FindPin(UEdGraphSchema_K2::PN_Self);
 				GetExecPin->MakeLinkTo(SetStaticMeshSelfPin);
 				Test->TestTrue(TEXT("Connected Get MyMesh node to Set Static Mesh node"), GetExecPin->LinkedTo.Contains(SetStaticMeshSelfPin));
 
@@ -1288,8 +1281,8 @@ namespace BlueprintEditorPromotionTestHelper
 				Test->TestEqual(*FString::Printf(TEXT("Set Static Mesh default mesh updated to %s"), *SecondBlueprintMesh->GetName()), Cast<UStaticMesh>(SetStaticMeshMeshPin->DefaultObject), SecondBlueprintMesh);
 
 				//Connect SetStaticMeshMesh to PrintString
-				UEdGraphPin* PrintStringThenPin = PrintNode->FindPin(K2Schema->PN_Then);
-				UEdGraphPin* SetStaticMeshExecPin = SetStaticMeshNode->FindPin(K2Schema->PN_Execute);
+				UEdGraphPin* PrintStringThenPin = PrintNode->FindPin(UEdGraphSchema_K2::PN_Then);
+				UEdGraphPin* SetStaticMeshExecPin = SetStaticMeshNode->FindPin(UEdGraphSchema_K2::PN_Execute);
 				PrintStringThenPin->MakeLinkTo(SetStaticMeshExecPin);
 				Test->TestTrue(TEXT("Connected Print String node to Set Static Mesh node"), PrintStringThenPin->LinkedTo.Contains(SetStaticMeshExecPin));
 			}
@@ -1331,15 +1324,13 @@ namespace BlueprintEditorPromotionTestHelper
 		{
 			if (BlueprintObject)
 			{
-				const UEdGraphSchema_K2* K2Schema = GetDefault<UEdGraphSchema_K2>();
-
 				CustomGraph = BlueprintEditorPromotionUtils::CreateNewFunctionGraph(BlueprintObject, TEXT("NewFunction"));
 				Test->TestNotNull(TEXT("Created new function graph"), CustomGraph);
 
 				AddParticleSystemNode = BlueprintEditorPromotionUtils::CreateAddComponentActionNode(BlueprintObject, CustomGraph, LoadedParticleSystem);
 				Test->TestNotNull(TEXT("Created Add Particle System node"), AddParticleSystemNode);
 
-				UEdGraphPin* ExecutePin = AddParticleSystemNode ? AddParticleSystemNode->FindPin(K2Schema->PN_Execute) : NULL;
+				UEdGraphPin* ExecutePin = AddParticleSystemNode ? AddParticleSystemNode->FindPin(UEdGraphSchema_K2::PN_Execute) : NULL;
 
 				//Find the input for the function graph
 				TArray<UK2Node_FunctionEntry*> EntryNodes;
@@ -1347,7 +1338,7 @@ namespace BlueprintEditorPromotionTestHelper
 				UEdGraphNode* EntryNode = EntryNodes.Num() > 0 ? EntryNodes[0] : NULL;
 				if (EntryNode && ExecutePin)
 				{
-					UEdGraphPin* EntryPin = EntryNode->FindPin(K2Schema->PN_Then);
+					UEdGraphPin* EntryPin = EntryNode->FindPin(UEdGraphSchema_K2::PN_Then);
 					EntryPin->MakeLinkTo(ExecutePin);
 					Test->TestTrue(TEXT("Connected Add Particle System node to entry node"), EntryPin->LinkedTo.Contains(ExecutePin));
 				}
@@ -1365,14 +1356,13 @@ namespace BlueprintEditorPromotionTestHelper
 		{
 			if (BlueprintObject)
 			{
-				const UEdGraphSchema_K2* K2Schema = GetDefault<UEdGraphSchema_K2>();
 				UEdGraph* EventGraph = FBlueprintEditorUtils::FindEventGraph(BlueprintObject);
-				UEdGraphPin* SetStaticMeshThenPin = SetStaticMeshNode->FindPin(K2Schema->PN_Then);
+				UEdGraphPin* SetStaticMeshThenPin = SetStaticMeshNode->FindPin(UEdGraphSchema_K2::PN_Then);
 				CallFunctionNode = BlueprintEditorPromotionUtils::AddCallFunctionGraphNode(BlueprintObject, EventGraph, TEXT("NewFunction"), SetStaticMeshThenPin);
 				Test->TestNotNull(TEXT("Created Call Function node"), CallFunctionNode);
 				if (CallFunctionNode)
 				{
-					Test->TestTrue(TEXT("Connected Set Static Mesh node to Call Function node"), SetStaticMeshThenPin->LinkedTo.Contains(CallFunctionNode->FindPin(K2Schema->PN_Execute)));
+					Test->TestTrue(TEXT("Connected Set Static Mesh node to Call Function node"), SetStaticMeshThenPin->LinkedTo.Contains(CallFunctionNode->FindPin(UEdGraphSchema_K2::PN_Execute)));
 				}
 
 				BlueprintEditorPromotionUtils::CompileBlueprint(BlueprintObject);
diff --git a/Engine/Source/Editor/Kismet/Private/UserDefinedStructureEditor.cpp b/Engine/Source/Editor/Kismet/Private/UserDefinedStructureEditor.cpp
--- a/Engine/Source/Editor/Kismet/Private/UserDefinedStructureEditor.cpp
+++ b/Engine/Source/Editor/Kismet/Private/UserDefinedStructureEditor.cpp
@@ -435,7 +435,7 @@ class FUserDefinedStructureLayout : public IDetailCustomNodeBuilder, public TSha
 public:
 	FUserDefinedStructureLayout(TWeakPtr<class FUserDefinedStructureDetails> InStructureDetails)
 		: StructureDetails(InStructureDetails)
-		, InitialPinType(GetDefault<UEdGraphSchema_K2>()->PC_Boolean, FString(), nullptr, EPinContainerType::None, false, FEdGraphTerminalType())
+		, InitialPinType(UEdGraphSchema_K2::PC_Boolean, NAME_None, nullptr, EPinContainerType::None, false, FEdGraphTerminalType())
 	{}
 
 	void OnChanged()
diff --git a/Engine/Source/Editor/Kismet/Public/BlueprintCompilationManager.h b/Engine/Source/Editor/Kismet/Public/BlueprintCompilationManager.h
--- a/Engine/Source/Editor/Kismet/Public/BlueprintCompilationManager.h
+++ b/Engine/Source/Editor/Kismet/Public/BlueprintCompilationManager.h
@@ -42,6 +42,11 @@ struct KISMET_API FBlueprintCompilationManager
 	 * will not load data while processing the compilation queue)
 	 */
 	static void FlushCompilationQueue(TArray<UObject*>* ObjLoaded = nullptr);
+	
+	/**
+	 * Flushes the compilation queue and finishes reinstancing
+	 */
+	static void FlushCompilationQueueAndReinstance();
 
 	/**
 	 * Immediately compiles the blueprint, no expectation that related blueprints be subsequently compiled.
@@ -55,8 +60,20 @@ struct KISMET_API FBlueprintCompilationManager
 	 */
 	static void NotifyBlueprintLoaded(UBlueprint* BPLoaded);
 
+	/**
+	 * Adds a blueprint to the compilation queue - useful for batch compilation
+	 */
+	static void QueueForCompilation(UBlueprint* BP);
+
 	/** Returns true when UBlueprint::GeneratedClass members are up to date */
 	static bool IsGeneratedClassLayoutReady();
+	
+	/** 
+	 * Returns the Default Value associated with ForClass::Property, if ForClass is currently 
+	 * being compiled this function can look at the old version of the CDO and read the default
+	 * value from there
+	 */
+	static bool GetDefaultValue(const UClass* ForClass, const UProperty* Property, FString& OutDefaultValueAsString);
 private:
 	FBlueprintCompilationManager();
 };
diff --git a/Engine/Source/Editor/Kismet/Public/BlueprintEditor.h b/Engine/Source/Editor/Kismet/Public/BlueprintEditor.h
--- a/Engine/Source/Editor/Kismet/Public/BlueprintEditor.h
+++ b/Engine/Source/Editor/Kismet/Public/BlueprintEditor.h
@@ -130,15 +130,16 @@ class KISMET_API FBlueprintEditor : public IBlueprintEditor, public FGCObject, p
 	/** A record of a warning generated by a disallowed pin connection attempt */
 	struct FDisallowedPinConnection
 	{
-		FString PinTypeCategoryA;
-		bool bPinIsArrayA;
-		bool bPinIsReferenceA;
-		bool bPinIsWeakPointerA;
-
-		FString PinTypeCategoryB;
-		bool bPinIsArrayB;
-		bool bPinIsReferenceB;
-		bool bPinIsWeakPointerB;
+		FName PinTypeCategoryA;
+		FName PinTypeCategoryB;
+
+		uint8 bPinIsArrayA:1;
+		uint8 bPinIsReferenceA:1;
+		uint8 bPinIsWeakPointerA:1;
+
+		uint8 bPinIsArrayB:1;
+		uint8 bPinIsReferenceB:1;
+		uint8 bPinIsWeakPointerB:1;
 	};
 
 public:
@@ -198,7 +199,7 @@ public:
 
 	//~ Begin FTickableEditorObject Interface
 	virtual void Tick(float DeltaTime) override;
-	virtual bool IsTickable() const override { return true; }
+	virtual ETickableTickType GetTickableTickType() const override { return ETickableTickType::Always; }
 	virtual TStatId GetStatId() const override;
 	//~ End FTickableEditorObject Interface
 
diff --git a/Engine/Source/Editor/Kismet/Public/BlueprintEditorModule.h b/Engine/Source/Editor/Kismet/Public/BlueprintEditorModule.h
--- a/Engine/Source/Editor/Kismet/Public/BlueprintEditorModule.h
+++ b/Engine/Source/Editor/Kismet/Public/BlueprintEditorModule.h
@@ -21,7 +21,7 @@ class FBlueprintEditor;
 class UUserDefinedEnum;
 class UUserDefinedStruct;
 class IDetailCustomization;
-
+class FKismetCompilerContext;
 
 /** Delegate used to customize variable display */
 DECLARE_DELEGATE_RetVal_OneParam(TSharedPtr<IDetailCustomization>, FOnGetVariableCustomizationInstance, TSharedPtr<IBlueprintEditor> /*BlueprintEditor*/);
@@ -90,6 +90,7 @@ public:
 };
 
 DECLARE_DELEGATE_RetVal_OneParam(TSharedRef<class ISCSEditorCustomization>, FSCSEditorCustomizationBuilder, TSharedRef< IBlueprintEditor > /* InBlueprintEditor */);
+DECLARE_DELEGATE_RetVal_ThreeParams(TSharedPtr<FKismetCompilerContext>, FOnGetBlueprintCompiler, UBlueprint*, FCompilerResultsLog&, const FKismetCompilerOptions&);
 
 /**
  * The blueprint editor module provides the blueprint editor application.
@@ -192,6 +193,24 @@ public:
 	DECLARE_EVENT_OneParam( FBlueprintEditorModule, FBlueprintEditorOpenedEvent, EBlueprintType );
 	FBlueprintEditorOpenedEvent& OnBlueprintEditorOpened() { return BlueprintEditorOpened; }
 
+	/**
+	 * Register custom compiler getter for the specific blueprint class (ex. UAnimBlueprint vs UBlueprint)
+	 * @param	UBlueprint class		The Blueprint derived class type 
+	 * @param	InCompilerGetter		The delegate for compiler getter
+	 */
+	KISMET_API static void RegisterCompilerGetter(const TSubclassOf<UBlueprint> BlueprintType, FOnGetBlueprintCompiler InCompilerGetter);
+
+	/**
+	 * Unregister a previously registered compiler getter for the specific blueprint class (ex. UAnimBlueprint vs UBlueprint)
+	 * @param	UBlueprint class		The Blueprint derived class type 
+	 */
+	KISMET_API static void UnregisterCompilerGetter(const TSubclassOf<UBlueprint> BlueprintType);
+
+	/** 
+	 * Return compiler for the specific blueprint instance
+	 */
+	static TSharedPtr<FKismetCompilerContext> GetCompiler(UBlueprint* BP, FCompilerResultsLog& InMessageLog, const FKismetCompilerOptions& InCompileOptions);
+
 	/** 
 	 * Exposes a way for other modules to fold in their own Blueprint editor 
 	 * commands (folded in with other BP editor commands, when the editor is 
@@ -233,4 +252,7 @@ private:
 
 	/** Reference to keep our custom configuration panel alive */
 	TSharedPtr<SWidget> ConfigurationPanel;
+
+	/** list of compiler getters for the blueprint types */
+	static TMap<TSubclassOf<UBlueprint>, FOnGetBlueprintCompiler> BlueprintCompilers;
 };
diff --git a/Engine/Source/Editor/KismetCompiler/KismetCompiler.Build.cs b/Engine/Source/Editor/KismetCompiler/KismetCompiler.Build.cs
--- a/Engine/Source/Editor/KismetCompiler/KismetCompiler.Build.cs
+++ b/Engine/Source/Editor/KismetCompiler/KismetCompiler.Build.cs
@@ -19,6 +19,7 @@ public class KismetCompiler : ModuleRules
 				"BlueprintGraph",
 				"AnimGraph",
                 "MessageLog",
+				"Kismet",
 			}
 			);
 	}
