Re-add line dropped in merge operation performed in 3721054. This would result in stale message log pointers, which would crash
#rnx

diff --git a/Engine/Source/Editor/Kismet/Private/BlueprintCompilationManager.cpp b/Engine/Source/Editor/Kismet/Private/BlueprintCompilationManager.cpp
--- a/Engine/Source/Editor/Kismet/Private/BlueprintCompilationManager.cpp
+++ b/Engine/Source/Editor/Kismet/Private/BlueprintCompilationManager.cpp
@@ -654,6 +654,7 @@ void FBlueprintCompilationManagerImpl::FlushCompilationQueueImpl(TArray<UObject*
 					}
 					
 					Data.BP->bBeingCompiled = false;
+					Data.BP->CurrentMessageLog = nullptr;
 					if(UPackage* Package = Data.BP->GetOutermost())
 					{
 						Package->SetDirtyFlag(Data.bPackageWasDirty);
