Copying //UE4/Dev-Framework to //UE4/Dev-Main (Source: //UE4/Dev-Framework @ 3967517)

diff --git a/Engine/Source/Editor/BlueprintGraph/Private/K2Node_Select.cpp b/Engine/Source/Editor/BlueprintGraph/Private/K2Node_Select.cpp
--- a/Engine/Source/Editor/BlueprintGraph/Private/K2Node_Select.cpp
+++ b/Engine/Source/Editor/BlueprintGraph/Private/K2Node_Select.cpp
@@ -19,16 +19,24 @@
 
 #define LOCTEXT_NAMESPACE "K2Node_Select"
 
+namespace
+{
+	FName NAME_bPickOption0(TEXT("bPickOption0"));
+	FName NAME_Index(TEXT("Index"));
+	FName NAME_Option_0(TEXT("Option 0"));
+	FName NAME_Option_1(TEXT("Option 1"));
+}
+
 //////////////////////////////////////////////////////////////////////////
 // FKCHandler_Select
 
-class FKCHandler_SelectRef : public FNodeHandlingFunctor
+class FKCHandler_Select : public FNodeHandlingFunctor
 {
 protected:
 	TMap<UEdGraphNode*, FBPTerminal*> DefaultTermMap;
 
 public:
-	FKCHandler_SelectRef(FKismetCompilerContext& InCompilerContext)
+	FKCHandler_Select(FKismetCompilerContext& InCompilerContext)
 		: FNodeHandlingFunctor(InCompilerContext)
 	{
 	}
@@ -143,190 +151,6 @@ public:
 	}
 };
 
-class FKCHandler_Select : public FNodeHandlingFunctor
-{
-protected:
-	TMap<UEdGraphNode*, FBPTerminal*> BoolTermMap;
-
-public:
-	FKCHandler_Select(FKismetCompilerContext& InCompilerContext)
-		: FNodeHandlingFunctor(InCompilerContext)
-	{
-	}
-
-	virtual void RegisterNets(FKismetFunctionContext& Context, UEdGraphNode* Node) override
-	{
-		FNodeHandlingFunctor::RegisterNets(Context, Node);
-
-		// Create the net for the return value manually as it's a special case Output Direction pin
-		UK2Node_Select* SelectNode = Cast<UK2Node_Select>(Node);
-		UEdGraphPin* ReturnPin = SelectNode->GetReturnValuePin();
-
-		FBPTerminal* Term = Context.CreateLocalTerminalFromPinAutoChooseScope(ReturnPin, Context.NetNameMap->MakeValidName(ReturnPin));
-		Context.NetMap.Add(SelectNode->GetReturnValuePin(), Term);
-
-		// Create a term to determine if the compare was successful or not
-		FBPTerminal* BoolTerm = Context.CreateLocalTerminal();
-		BoolTerm->Type.PinCategory = UEdGraphSchema_K2::PC_Boolean;
-		BoolTerm->Source = Node;
-		BoolTerm->Name = Context.NetNameMap->MakeValidName(Node) + TEXT("_CmpSuccess");
-		BoolTermMap.Add(Node, BoolTerm);
-	}
-
-	virtual void Compile(FKismetFunctionContext& Context, UEdGraphNode* Node) override
-	{
-		// Cast the node and get all the input pins
-		UK2Node_Select* SelectNode = Cast<UK2Node_Select>(Node);
-		TArray<UEdGraphPin*> OptionPins;
-		SelectNode->GetOptionPins(OptionPins);
-		UEdGraphPin* IndexPin = SelectNode->GetIndexPin();
-
-		// Get the kismet term for the (Condition or Index) that will determine which option to use
-		UEdGraphPin* PinToTry = FEdGraphUtilities::GetNetFromPin(IndexPin);
-		FBPTerminal** ConditionTerm = Context.NetMap.Find(PinToTry);
-
-		// Get the kismet term for the return value
-		UEdGraphPin* ReturnPin = SelectNode->GetReturnValuePin();
-		FBPTerminal** ReturnTerm = Context.NetMap.Find(ReturnPin);
-
-		// Don't proceed if there is no return value or there is no selection
-		if (ConditionTerm != NULL && ReturnTerm != NULL)
-		{
-			FName ConditionalFunctionName = "";
-			UClass* ConditionalFunctionClass = NULL;
-			SelectNode->GetConditionalFunction(ConditionalFunctionName, &ConditionalFunctionClass);
-			UFunction* ConditionFunction = FindField<UFunction>(ConditionalFunctionClass, ConditionalFunctionName);
-
-			// Find the local boolean for use in the equality call function below (BoolTerm = result of EqualEqual_IntInt or NotEqual_BoolBool)
-			FBPTerminal* BoolTerm = BoolTermMap.FindRef(SelectNode);
-
-			// We need to keep a pointer to the previous IfNot statement so it can be linked to the next conditional statement
-			FBlueprintCompiledStatement* PrevIfNotStatement = NULL;
-
-			// Keep an array of all the unconditional goto statements so we can clean up their jumps after the noop statement is created
-			TArray<FBlueprintCompiledStatement*> GotoStatementList;
-
-			// Loop through all the options
-			for (int32 OptionIdx = 0; OptionIdx < OptionPins.Num(); OptionIdx++)
-			{
-				// Create a CallFunction statement with the condition function from the Select class
-				FBlueprintCompiledStatement& Statement = Context.AppendStatementForNode(Node);
-				Statement.Type = KCST_CallFunction;
-				Statement.FunctionToCall = ConditionFunction;
-				Statement.FunctionContext = NULL;
-				Statement.bIsParentContext = false;
-				// BoolTerm will be the return value of the condition statement
-				Statement.LHS = BoolTerm;
-				// The condition passed into the Select node
-				Statement.RHS.Add(*ConditionTerm);
-				// Create a local int for use in the equality call function below (LiteralTerm = the right hand side of the EqualEqual_IntInt or NotEqual_BoolBool statement)
-				FBPTerminal* LiteralTerm = Context.CreateLocalTerminal(ETerminalSpecification::TS_Literal);
-				LiteralTerm->bIsLiteral = true;
-				LiteralTerm->Type.PinCategory = UEdGraphSchema_K2::PC_Int;
-
-				if (UEnum* NodeEnum = SelectNode->GetEnum())
-				{
-					int32 EnumValue = NodeEnum->GetValueByName(OptionPins[OptionIdx]->PinName);
-					LiteralTerm->Name = FString::Printf(TEXT("%d"), EnumValue);
-				}
-				else
-				{
-					LiteralTerm->Name = FString::Printf(TEXT("%d"), OptionIdx);
-				}
-				Statement.RHS.Add(LiteralTerm);
-				// If there is a previous IfNot statement, hook this one to that one for jumping
-				if (PrevIfNotStatement)
-				{
-					Statement.bIsJumpTarget = true;
-					PrevIfNotStatement->TargetLabel = &Statement;
-				}
-
-				// Create a GotoIfNot statement using the BoolTerm from above as the condition
-				FBlueprintCompiledStatement* IfNotStatement = &Context.AppendStatementForNode(Node);
-				IfNotStatement->Type = KCST_GotoIfNot;
-				IfNotStatement->LHS = BoolTerm;
-
-				// Create an assignment statement
-				FBlueprintCompiledStatement& AssignStatement = Context.AppendStatementForNode(Node);
-				AssignStatement.Type = KCST_Assignment;
-				AssignStatement.LHS = *ReturnTerm;
-				// Get the kismet term from the option pin
-				UEdGraphPin* OptionPinToTry = FEdGraphUtilities::GetNetFromPin(OptionPins[OptionIdx]);
-				FBPTerminal** OptionTerm = Context.NetMap.Find(OptionPinToTry);
-				if (!OptionTerm)
-				{
-					Context.MessageLog.Error(*LOCTEXT("Error_UnregisterOptionPin", "Unregister option pin @@").ToString(), OptionPins[OptionIdx]);
-					return;
-				}
-				AssignStatement.RHS.Add(*OptionTerm);
-
-				// Create an unconditional goto to exit the node
-				FBlueprintCompiledStatement& GotoStatement = Context.AppendStatementForNode(Node);
-				GotoStatement.Type = KCST_UnconditionalGoto;
-				GotoStatementList.Add(&GotoStatement);
-
-				// If this is the last IfNot statement, hook the jump to an error message
-				if (OptionIdx == OptionPins.Num() - 1)
-				{
-					// Create a CallFunction statement for doing a print string of our error message
-					FBlueprintCompiledStatement& PrintStatement = Context.AppendStatementForNode(Node);
-					PrintStatement.Type = KCST_CallFunction;
-					PrintStatement.bIsJumpTarget = true;
-					FName PrintStringFunctionName = "";
-					UClass* PrintStringFunctionClass = NULL;
-					SelectNode->GetPrintStringFunction(PrintStringFunctionName, &PrintStringFunctionClass);
-					UFunction* PrintFunction = FindField<UFunction>(PrintStringFunctionClass, PrintStringFunctionName);
-					PrintStatement.FunctionToCall = PrintFunction;
-					PrintStatement.FunctionContext = NULL;
-					PrintStatement.bIsParentContext = false;
-
-					// Create a local int for use in the equality call function below (LiteralTerm = the right hand side of the EqualEqual_IntInt or NotEqual_BoolBool statement)
-					FBPTerminal* LiteralStringTerm = Context.CreateLocalTerminal(ETerminalSpecification::TS_Literal);
-					LiteralStringTerm->bIsLiteral = true;
-					LiteralStringTerm->Type.PinCategory = UEdGraphSchema_K2::PC_String;
-
-					FString SelectionNodeType;
-					if (IndexPin)
-					{
-						if (UEnum* EnumObject = Cast<UEnum>(IndexPin->PinType.PinSubCategoryObject.Get()))
-						{
-							SelectionNodeType = EnumObject->GetName();
-						}
-						else
-						{
-							// Not an enum, so just use the basic type
-							SelectionNodeType = IndexPin->PinType.PinCategory.ToString();
-						}
-					}
-					const UEdGraph* OwningGraph = Context.MessageLog.FindSourceObjectTypeChecked<UEdGraph>( SelectNode->GetGraph() );
-					LiteralStringTerm->Name = FText::Format(
-						LOCTEXT("SelectNodeIndexWarningFmt", "Graph {0}: Selection Node of type {1} failed! Out of bounds indexing of the options. There are only {2} options available."),
-						(OwningGraph) ? FText::FromString(OwningGraph->GetFullName()) : LOCTEXT("SelectNodeIndexWarningNoGraph", "NONE"),
-						(IndexPin) ? FText::FromString(SelectionNodeType) : LOCTEXT("SelectNodeIndexWarningNoPin", "NONE"),
-						OptionPins.Num()
-					).ToString();
-					PrintStatement.RHS.Add(LiteralStringTerm);
-
-					// Hook the IfNot statement's jump target to this statement
-					IfNotStatement->TargetLabel = &PrintStatement;
-				}
-
-				PrevIfNotStatement = IfNotStatement;
-			}
-
-			// Create a noop to jump to so the unconditional goto statements can exit the node after successful assignment
-			FBlueprintCompiledStatement& NopStatement = Context.AppendStatementForNode(Node);
-			NopStatement.Type = KCST_Nop;
-			NopStatement.bIsJumpTarget = true;
-			// Loop through the unconditional goto statements and fix their jump targets
-			for (FBlueprintCompiledStatement* GotoStatement : GotoStatementList)
-			{
-				GotoStatement->TargetLabel = &NopStatement;
-			}
-		}
-	}
-};
-
 UK2Node_Select::UK2Node_Select(const FObjectInitializer& ObjectInitializer)
 	: Super(ObjectInitializer)
 {
@@ -353,8 +177,6 @@ void UK2Node_Select::AllocateDefaultPins()
 		NumOptionPins = EnumEntries.Num();
 	}
 
-	static const FBoolConfigValueHelper UseSelectRef(TEXT("Kismet"), TEXT("bUseSelectRef"), GEngineIni);
-
 	// Create the option pins
 	for (int32 Idx = 0; Idx < NumOptionPins; Idx++)
 	{
@@ -377,7 +199,7 @@ void UK2Node_Select::AllocateDefaultPins()
 
 		if (NewPin)
 		{
-			NewPin->bDisplayAsMutableRef = UseSelectRef;
+			NewPin->bDisplayAsMutableRef = true;
 			if (IndexPinType.PinCategory == UEdGraphSchema_K2::PC_Boolean)
 			{
 				NewPin->PinFriendlyName = (Idx == 0 ? GFalse : GTrue);
@@ -394,7 +216,7 @@ void UK2Node_Select::AllocateDefaultPins()
 
 	// Create the return value
 	UEdGraphPin* ReturnPin = CreatePin(EGPD_Output, UEdGraphSchema_K2::PC_Wildcard, UEdGraphSchema_K2::PN_ReturnValue);
-	ReturnPin->bDisplayAsMutableRef = UseSelectRef;
+	ReturnPin->bDisplayAsMutableRef = true;
 
 	Super::AllocateDefaultPins();
 }
@@ -435,6 +257,19 @@ FText UK2Node_Select::GetNodeTitle(ENodeTitleType::Type TitleType) const
 
 UK2Node::ERedirectType UK2Node_Select::DoPinsMatchForReconstruction(const UEdGraphPin* NewPin, int32 NewPinIndex, const UEdGraphPin* OldPin, int32 OldPinIndex) const
 {
+	if (bReconstructForPinTypeChange)
+	{
+		// If we're reconstructing for the purposes of changing the index pin type then we want to
+		// keep our connections based on the index of the option pin
+		if (NewPin != GetIndexPin() && NewPin != GetReturnValuePin())
+		{
+			if (NewPinIndex == OldPinIndex)
+			{
+				return UK2Node::ERedirectType_Name;
+			}
+		}
+	}
+
 	// Check to see if the new pin name matches the old pin name.
 	if (Enum && (NewPinIndex < NumOptionPins) && (NewPin->PinName != OldPin->PinName))
 	{
@@ -469,11 +304,11 @@ void UK2Node_Select::ReallocatePinsDuringReconstruction(TArray<UEdGraphPin*>& Ol
 	UEdGraphPin* OldReturnPin = nullptr;
 	for (UEdGraphPin* OldPin : OldPins)
 	{
-		if (OldPin->PinName == TEXT("bPickOption0"))
+		if (OldPin->PinName == NAME_bPickOption0)
 		{
 			OldConditionPin = OldPin;
 		}
-		else if (OldPin->PinName == TEXT("Index"))
+		else if (OldPin->PinName == NAME_Index)
 		{
 			OldIndexPin = OldPin;
 		}
@@ -512,8 +347,8 @@ void UK2Node_Select::ReallocatePinsDuringReconstruction(TArray<UEdGraphPin*>& Ol
 			PinConnectionListChanged(IndexPin);
 		}
 
-		UEdGraphPin* OptionPin0 = FindPin(TEXT("Option 0"));
-		UEdGraphPin* OptionPin1 = FindPin(TEXT("Option 1"));
+		UEdGraphPin* OptionPin0 = FindPin(NAME_Option_0);
+		UEdGraphPin* OptionPin1 = FindPin(NAME_Option_1);
 
 		for (UEdGraphPin* OldPin : OldPins)
 		{
@@ -560,6 +395,7 @@ void UK2Node_Select::PostReconstructNode()
 {
 	// After ReconstructNode we must be sure that no additional reconstruction is required
 	bReconstructNode = false;
+	bReconstructForPinTypeChange = false;
 
 	UEdGraphPin* ReturnPin = GetReturnValuePin();
 	const bool bFillTypeFromConnected = ReturnPin && (ReturnPin->PinType.PinCategory == UEdGraphSchema_K2::PC_Wildcard);
@@ -587,7 +423,7 @@ void UK2Node_Select::PostReconstructNode()
 		}
 
 		ReturnPin->PinType = PinType;
-		PinTypeChanged(ReturnPin);
+		OnPinTypeChanged(ReturnPin);
 	}
 
 	Super::PostReconstructNode();
@@ -610,7 +446,7 @@ void UK2Node_Select::NotifyPinConnectionListChanged(UEdGraphPin* Pin)
 			{
 				Pin->PinType = LinkPin->PinType;
 
-				PinTypeChanged(Pin);
+				OnPinTypeChanged(Pin);
 			}
 		}
 	}
@@ -622,24 +458,71 @@ void UK2Node_Select::NotifyPinConnectionListChanged(UEdGraphPin* Pin)
 		UEdGraphPin* ReturnPin = FindPin(UEdGraphSchema_K2::PN_ReturnValue);
 
 		// See if this pin is one of the wildcard pins
-		bool bIsWildcardPin = (Pin == ReturnPin || OptionPins.Find(Pin) != INDEX_NONE) && Pin->PinType.PinCategory == UEdGraphSchema_K2::PC_Wildcard;
+		const bool bIsWildcardPin = ((Pin->PinType.PinCategory == UEdGraphSchema_K2::PC_Wildcard) && ((Pin == ReturnPin) || (OptionPins.Find(Pin) != INDEX_NONE)));
 
-		// If the pin was one of the wildcards we have to handle it specially
-		if (bIsWildcardPin)
+		TFunction<bool(UEdGraphPin*)> PinInUse = [&PinInUse](UEdGraphPin* PinToConsider)
 		{
-			// If the pin is linked, make sure the other wildcard pins match
-			if (Pin->LinkedTo.Num() > 0)
+			bool bPinInUse = ((PinToConsider->LinkedTo.Num() > 0) || (PinToConsider->ParentPin != nullptr) || !PinToConsider->DoesDefaultValueMatchAutogenerated());
+			if (!bPinInUse)
 			{
-				UEdGraphPin* LinkPin = Pin->LinkedTo[0];
+				for (UEdGraphPin* SubPin : PinToConsider->SubPins)
+				{
+					bPinInUse = PinInUse(SubPin);
+					if (bPinInUse)
+					{
+						break;
+					}
+				}
+			}
+			return bPinInUse;
+		};
 
-				if (Pin->PinType != LinkPin->PinType)
+		bool bPinsInUse = PinInUse(ReturnPin);
+		
+		if (!bPinsInUse)
+		{
+			for (UEdGraphPin* OptionPin : OptionPins)
+			{
+				bPinsInUse = PinInUse(OptionPin);
+				if (bPinsInUse)
 				{
-					Pin->PinType = LinkPin->PinType;
+					break;
+				}
+			}
+		}
+
+		bool bPinTypeChanged = false;
 
-					PinTypeChanged(Pin);
+		if (bPinsInUse)
+		{
+			// If the pin was one of the wildcards we have to handle it specially
+			if (bIsWildcardPin)
+			{
+				// If the pin is linked, make sure the other wildcard pins match
+				if (Pin->LinkedTo.Num() > 0)
+				{
+					UEdGraphPin* LinkPin = Pin->LinkedTo[0];
+
+					if (Pin->PinType != LinkPin->PinType)
+					{
+						Pin->PinType = LinkPin->PinType;
+						bPinTypeChanged = true;
+					}
 				}
 			}
 		}
+		else
+		{
+			bPinTypeChanged = true;
+			Pin->PinType.PinCategory = UEdGraphSchema_K2::PC_Wildcard;
+			Pin->PinType.PinSubCategory = NAME_None;
+			Pin->PinType.PinSubCategoryObject = nullptr;
+		}
+
+		if (bPinTypeChanged)
+		{
+			OnPinTypeChanged(Pin);
+		}
 	}
 }
 
@@ -719,8 +602,10 @@ void UK2Node_Select::GetPrintStringFunction(FName& FunctionName, UClass** Functi
 	*FunctionClass = UKismetSystemLibrary::StaticClass();
 }
 
-void UK2Node_Select::AddOptionPinToNode()
+void UK2Node_Select::AddInputPin()
 {
+	Modify();
+
 	const UEdGraphSchema_K2* K2Schema = GetDefault<UEdGraphSchema_K2>();
 
 	// Increment the pin count
@@ -794,7 +679,7 @@ void UK2Node_Select::NodeConnectionListChanged()
 	}
 }
 
-bool UK2Node_Select::CanAddOptionPinToNode() const
+bool UK2Node_Select::CanAddPin() const
 {
 	if (IndexPinType.PinCategory == UEdGraphSchema_K2::PC_Byte &&
 		IndexPinType.PinSubCategoryObject.IsValid() &&
@@ -827,10 +712,11 @@ bool UK2Node_Select::CanRemoveOptionPinToNode() const
 
 void UK2Node_Select::ChangePinType(UEdGraphPin* Pin)
 {
-	PinTypeChanged(Pin);
+	OnPinTypeChanged(Pin);
 
 	if (bReconstructNode)
 	{
+		bReconstructForPinTypeChange = true;
 		ReconstructNode();
 	}
 
@@ -876,6 +762,12 @@ bool UK2Node_Select::CanChangePinType(UEdGraphPin* Pin) const
 }
 
 void UK2Node_Select::PinTypeChanged(UEdGraphPin* Pin)
+{
+	bReconstructForPinTypeChange = true;
+	OnPinTypeChanged(Pin);
+}
+
+void UK2Node_Select::OnPinTypeChanged(UEdGraphPin* Pin)
 {
 	const UEdGraphSchema_K2* Schema = GetDefault<UEdGraphSchema_K2>();
 
@@ -888,7 +780,7 @@ void UK2Node_Select::PinTypeChanged(UEdGraphPin* Pin)
 			// Since it is an interactive action we want the pins to go away regardless of the new type
 			for (UEdGraphPin* PinToDiscard : Pins)
 			{
-				PinToDiscard->bSavePinIfOrphaned = false;
+				PinToDiscard->SetSavePinIfOrphaned(false);
 			}
 
 			if (IndexPinType.PinSubCategoryObject.IsValid())
@@ -968,10 +860,10 @@ void UK2Node_Select::PostPasteNode()
 		FString OldDefaultValue = IndexPin->DefaultValue;
 
 		// Corrects data in the index pin that is not valid after pasting
-		PinTypeChanged(IndexPin);
+		OnPinTypeChanged(IndexPin);
 
 		// Restore the default value of the index pin
-		IndexPin->DefaultValue = OldDefaultValue;
+		IndexPin->DefaultValue = MoveTemp(OldDefaultValue);
 	}
 }
 
@@ -994,10 +886,7 @@ bool UK2Node_Select::IsConnectionDisallowed(const UEdGraphPin* MyPin, const UEdG
 
 FNodeHandlingFunctor* UK2Node_Select::CreateNodeHandler(FKismetCompilerContext& CompilerContext) const
 {
-	static const FBoolConfigValueHelper UseSelectRef(TEXT("Kismet"), TEXT("bUseSelectRef"), GEngineIni);
-	return UseSelectRef
-		? static_cast<FNodeHandlingFunctor*>(new FKCHandler_SelectRef(CompilerContext))
-		: static_cast<FNodeHandlingFunctor*>(new FKCHandler_Select(CompilerContext));
+	return static_cast<FNodeHandlingFunctor*>(new FKCHandler_Select(CompilerContext));
 }
 
 void UK2Node_Select::GetMenuActions(FBlueprintActionDatabaseRegistrar& ActionRegistrar) const
@@ -1029,13 +918,6 @@ void UK2Node_Select::ExpandNode(class FKismetCompilerContext& CompilerContext, U
 {
 	Super::ExpandNode(CompilerContext, SourceGraph);
 
-	static const FBoolConfigValueHelper UseSelectRef(TEXT("Kismet"), TEXT("bUseSelectRef"), GEngineIni);
-	if (!UseSelectRef)
-	{
-		return;
-	}
-
-	bool bSuccess = false;
 	const UEdGraphSchema_K2* Schema = CompilerContext.GetSchema();
 	for (UEdGraphPin* Pin : Pins)
 	{
diff --git a/Engine/Source/Editor/BlueprintGraph/Private/K2Node_SetFieldsInStruct.cpp b/Engine/Source/Editor/BlueprintGraph/Private/K2Node_SetFieldsInStruct.cpp
--- a/Engine/Source/Editor/BlueprintGraph/Private/K2Node_SetFieldsInStruct.cpp
+++ b/Engine/Source/Editor/BlueprintGraph/Private/K2Node_SetFieldsInStruct.cpp
@@ -235,7 +235,7 @@ void UK2Node_SetFieldsInStruct::RemoveFieldPins(UEdGraphPin* Pin, EPinsToRemove
 			{
 				bWasChanged = true;
 				OptionalProperty.bShowPin = false;
-				Pin->bSavePinIfOrphaned = false;
+				Pin->SetSavePinIfOrphaned(false);
 			}
 		}
 
diff --git a/Engine/Source/Editor/BlueprintGraph/Private/K2Node_SpawnActorFromClass.cpp b/Engine/Source/Editor/BlueprintGraph/Private/K2Node_SpawnActorFromClass.cpp
--- a/Engine/Source/Editor/BlueprintGraph/Private/K2Node_SpawnActorFromClass.cpp
+++ b/Engine/Source/Editor/BlueprintGraph/Private/K2Node_SpawnActorFromClass.cpp
@@ -447,17 +447,4 @@ void UK2Node_SpawnActorFromClass::ExpandNode(class FKismetCompilerContext& Compi
 	SpawnNode->BreakAllNodeLinks();
 }
 
-bool UK2Node_SpawnActorFromClass::HasExternalDependencies(TArray<class UStruct*>* OptionalOutput) const
-{
-	UClass* SourceClass = GetClassToSpawn();
-	const UBlueprint* SourceBlueprint = GetBlueprint();
-	const bool bResult = (SourceClass != NULL) && (SourceClass->ClassGeneratedBy != SourceBlueprint);
-	if (bResult && OptionalOutput)
-	{
-		OptionalOutput->AddUnique(SourceClass);
-	}
-	const bool bSuperResult = Super::HasExternalDependencies(OptionalOutput);
-	return bSuperResult || bResult;
-}
-
 #undef LOCTEXT_NAMESPACE
diff --git a/Engine/Source/Editor/BlueprintGraph/Public/BlueprintEditorSettings.h b/Engine/Source/Editor/BlueprintGraph/Public/BlueprintEditorSettings.h
--- a/Engine/Source/Editor/BlueprintGraph/Public/BlueprintEditorSettings.h
+++ b/Engine/Source/Editor/BlueprintGraph/Public/BlueprintEditorSettings.h
@@ -5,6 +5,7 @@
 #include "CoreMinimal.h"
 #include "UObject/ObjectMacros.h"
 #include "UObject/Object.h"
+#include "Engine/Blueprint.h"
 #include "BlueprintEditorSettings.generated.h"
 
 UENUM()
@@ -89,6 +90,22 @@ public:
 	UPROPERTY(EditAnywhere, config, Category=Workflow)
 	bool bNavigateToNativeFunctionsFromCallNodes;
 
+	/** Blueprint bookmark database */
+	UPROPERTY(config)
+	TMap<FGuid, FEditedDocumentInfo> Bookmarks;
+
+	/** Blueprint bookmark nodes (for display) */
+	UPROPERTY(config)
+	TArray<FBPEditorBookmarkNode> BookmarkNodes;
+
+	/** If enabled, comment nodes will be included in the tree view display in the Bookmarks tab. */
+	UPROPERTY(config)
+	bool bIncludeCommentNodesInBookmarksTab;
+
+	/** If enabled, only the bookmarks for the current document will be shown in the Bookmarks tab. */
+	UPROPERTY(config)
+	bool bShowBookmarksForCurrentDocumentOnlyInTab;
+
 // Compiler Settings
 public:
 	/** Determines when to save Blueprints post-compile */
diff --git a/Engine/Source/Editor/DetailCustomizations/DetailCustomizations.Build.cs b/Engine/Source/Editor/DetailCustomizations/DetailCustomizations.Build.cs
--- a/Engine/Source/Editor/DetailCustomizations/DetailCustomizations.Build.cs
+++ b/Engine/Source/Editor/DetailCustomizations/DetailCustomizations.Build.cs
@@ -10,17 +10,17 @@ public class DetailCustomizations : ModuleRules
 
 		PrivateDependencyModuleNames.AddRange(
 			new string[] {
-                "AppFramework",
+				"AppFramework",
 				"Core",
 				"CoreUObject",
 				"ApplicationCore",
 				"DesktopWidgets",
 				"Engine",
 				"Landscape",
-                "InputCore",
+				"InputCore",
 				"Slate",
 				"SlateCore",
-                "EditorStyle",
+				"EditorStyle",
 				"UnrealEd",
 				"EditorWidgets",
 				"KismetWidgets",
@@ -31,47 +31,48 @@ public class DetailCustomizations : ModuleRules
                 "MovieScene",
                 "TimeManagement",
 				"SharedSettingsWidgets",
-                "ContentBrowser",
+				"ContentBrowser",
 				"BlueprintGraph",
-                "GraphEditor",
+				"GraphEditor",
 				"AnimGraph",
-                "PropertyEditor",
-                "LevelEditor",
+				"PropertyEditor",
+				"LevelEditor",
 				"DesktopPlatform",
 				"ClassViewer",
 				"TargetPlatform",
 				"ExternalImagePicker",
 				"MoviePlayer",
 				"SourceControl",
-                "InternationalizationSettings",
+				"InternationalizationSettings",
 				"SourceCodeAccess",
 				"RHI",
-                "HardwareTargeting",
-                "AIModule", 
+				"HardwareTargeting",
+				"NavigationSystem",
+				"AIModule", 
 				"ConfigEditor",
-                "Persona",
-                "CinematicCamera",
+				"Persona",
+				"CinematicCamera",
 				"ComponentVisualizers",
-                "SkeletonEditor",
-                "LevelSequence",
-                "AdvancedPreviewScene",
-                "AudioSettingsEditor",
+				"SkeletonEditor",
+				"LevelSequence",
+				"AdvancedPreviewScene",
+				"AudioSettingsEditor",
 				"HeadMountedDisplay",
             }
 		);
 
         PrivateIncludePathModuleNames.AddRange(
 			new string[] {
-                "Engine",
-                "Media",
+				"Engine",
+				"Media",
 				"Landscape",
 				"LandscapeEditor",
 				"PropertyEditor",
 				"GameProjectGeneration",
-                "ComponentVisualizers",
+				"ComponentVisualizers",
 				"GraphEditor",
-                "MeshMergeUtilities",
-                "MeshReductionInterface",
+				"MeshMergeUtilities",
+				"MeshReductionInterface",
             }
 		);
 
@@ -79,8 +80,8 @@ public class DetailCustomizations : ModuleRules
 			new string[] {
 				"Layers",
 				"GameProjectGeneration",
-                "MeshMergeUtilities",
-                "MeshReductionInterface",
+				"MeshMergeUtilities",
+				"MeshReductionInterface",
             }
 		);
 	}
diff --git a/Engine/Source/Editor/DetailCustomizations/Private/NavAgentSelectorCustomization.cpp b/Engine/Source/Editor/DetailCustomizations/Private/NavAgentSelectorCustomization.cpp
--- a/Engine/Source/Editor/DetailCustomizations/Private/NavAgentSelectorCustomization.cpp
+++ b/Engine/Source/Editor/DetailCustomizations/Private/NavAgentSelectorCustomization.cpp
@@ -3,7 +3,7 @@
 #include "NavAgentSelectorCustomization.h"
 #include "Widgets/Text/STextBlock.h"
 #include "Engine/Engine.h"
-#include "AI/Navigation/NavigationSystem.h"
+#include "NavigationSystem.h"
 #include "IDetailChildrenBuilder.h"
 #include "DetailWidgetRow.h"
 
@@ -39,9 +39,14 @@ void FNavAgentSelectorCustomization::CustomizeChildren(TSharedRef<class IPropert
 	StructPropertyHandle->GetNumChildren(NumChildren);
 
 	FString AgentPrefix("bSupportsAgent");
-	const UNavigationSystem* NavSysCDO = (*GEngine->NavigationSystemClass != nullptr)
-		? GetDefault<UNavigationSystem>(GEngine->NavigationSystemClass)
-		: GetDefault<UNavigationSystem>();
+	const UNavigationSystemV1* NavSysCDO = (*GEngine->NavigationSystemClass != nullptr)
+		? GetDefault<UNavigationSystemV1>(GEngine->NavigationSystemClass)
+		: GetDefault<UNavigationSystemV1>();
+
+	if (NavSysCDO == nullptr)
+	{
+		return;
+	}
 
 	const int32 NumAgents = FMath::Min(NavSysCDO->GetSupportedAgents().Num(), 16);
 
@@ -80,9 +85,14 @@ void FNavAgentSelectorCustomization::CustomizeChildren(TSharedRef<class IPropert
 
 void FNavAgentSelectorCustomization::OnAgentStateChanged()
 {
-	const UNavigationSystem* NavSysCDO = (*GEngine->NavigationSystemClass != nullptr)
-		? GetDefault<UNavigationSystem>(GEngine->NavigationSystemClass)
-		: GetDefault<UNavigationSystem>();
+	const UNavigationSystemV1* NavSysCDO = (*GEngine->NavigationSystemClass != nullptr)
+		? GetDefault<UNavigationSystemV1>(GEngine->NavigationSystemClass)
+		: GetDefault<UNavigationSystemV1>();
+
+	if (NavSysCDO == nullptr)
+	{
+		return;
+	}
 
 	const int32 NumAgents = FMath::Min(NavSysCDO->GetSupportedAgents().Num(), 16);
 
diff --git a/Engine/Source/Editor/EditorStyle/Private/SlateEditorStyle.cpp b/Engine/Source/Editor/EditorStyle/Private/SlateEditorStyle.cpp
--- a/Engine/Source/Editor/EditorStyle/Private/SlateEditorStyle.cpp
+++ b/Engine/Source/Editor/EditorStyle/Private/SlateEditorStyle.cpp
@@ -1474,6 +1474,31 @@ void FSlateEditorStyle::FStyle::SetupGeneralStyles()
 
 		Set( "GraphBreadcrumb.BrowseBack", new IMAGE_BRUSH( "Icons/icon_BlueprintBrowserL_24x", Icon24x24) );
 		Set( "GraphBreadcrumb.BrowseForward", new IMAGE_BRUSH( "Icons/icon_BlueprintBrowserR_24x", Icon24x24) );
+
+		const FComboButtonStyle FastJumpComboBoxComboButton = FComboButtonStyle()
+			.SetButtonStyle(GetWidgetStyle<FButtonStyle>("GraphBreadcrumbButton"));
+		Set("GraphBreadcrumbFastJumpComboBoxStyle", FComboBoxStyle()
+			.SetComboButtonStyle(FastJumpComboBoxComboButton));
+	}
+
+	// Graph bookmark button
+	{
+		Set("GraphBookmarkButton", FButtonStyle()
+			.SetNormal(FSlateNoResource())
+			.SetPressed(BOX_BRUSH("Common/Button_Pressed", 8.0f / 32.0f, SelectionColor_Pressed))
+			.SetHovered(BOX_BRUSH("Common/Button_Hovered", 8.0f / 32.0f, SelectionColor))
+		);
+
+		Set("GraphBookmarkButtonText", FTextBlockStyle(NormalText)
+			.SetFont(ICON_FONT(16))
+			.SetColorAndOpacity(FLinearColor(0.5f, 0.5f, 0.5f))
+			.SetShadowOffset(FVector2D::ZeroVector)
+		);
+
+		Set("GraphBookmarkMenuImage.Button_Add", new IMAGE_BRUSH("Icons/PlusSymbol_12x", Icon12x12));
+		Set("GraphBookmarkMenuText.EmptyListItem", FTextBlockStyle(NormalText)
+			.SetFont(DEFAULT_FONT("Fonts/Roboto-Italic", 9))
+			.SetColorAndOpacity(FSlateColor::UseSubduedForeground()));
 	}
 #endif // WITH_EDITOR
 
@@ -5394,6 +5419,7 @@ void FSlateEditorStyle::FStyle::SetupPersonaStyle()
 			Set( "Kismet.Tabs.Palette", new IMAGE_BRUSH( "/Icons/levels_16x", Icon16x16 ) );
 			Set( "Kismet.Tabs.CompilerResults", new IMAGE_BRUSH( "Icons/icon_tab_OutputLog_16x", Icon16x16 ) );
 			Set( "Kismet.Tabs.FindResults", new IMAGE_BRUSH( "/Icons/icon_Genericfinder_16x", Icon16x16 ) );
+			Set( "Kismet.Tabs.Bookmarks", new IMAGE_BRUSH( "/Icons/icon_BlueprintEditor_Bookmarks_16x", Icon16x16 ) );
 			Set( "Kismet.Tabs.Components", new IMAGE_BRUSH( "/Icons/icon_BlueprintEditor_Components_16x", Icon16x16 ) );
 			Set( "Kismet.Tabs.BlueprintDefaults", new IMAGE_BRUSH( "Icons/icon_BlueprintEditor_Defaults_40x", Icon16x16 ) );
 		}
@@ -5555,6 +5581,13 @@ void FSlateEditorStyle::FStyle::SetupPersonaStyle()
 		Set("FindResults.LockButton_Unlocked", new IMAGE_BRUSH("Icons/padlock_unlocked_16x", Icon16x16));
 	}
 
+	//Bookmarks
+	{
+		Set("Bookmarks.AddFolderButtonIcon", new IMAGE_BRUSH("Icons/icon_AddFolder_16x", Icon16x16));
+		Set("Bookmarks.TreeViewItemFont", DEFAULT_FONT("Fonts/Roboto-Regular", 10));
+		Set("Bookmarks.TreeViewRootItemFont", DEFAULT_FONT("Fonts/Roboto-Regular", 11));
+	}
+
 	//Blueprint Diff
 	{
 		Set("BlueprintDif.HasGraph", new IMAGE_BRUSH("/Icons/blueprint_Dif_has_graph_8x", Icon8x8));
diff --git a/Engine/Source/Editor/GraphEditor/Private/KismetNodes/SGraphNodeK2Base.cpp b/Engine/Source/Editor/GraphEditor/Private/KismetNodes/SGraphNodeK2Base.cpp
--- a/Engine/Source/Editor/GraphEditor/Private/KismetNodes/SGraphNodeK2Base.cpp
+++ b/Engine/Source/Editor/GraphEditor/Private/KismetNodes/SGraphNodeK2Base.cpp
@@ -32,7 +32,7 @@
 const FLinearColor SGraphNodeK2Base::BreakpointHitColor(0.7f, 0.0f, 0.0f);
 const FLinearColor SGraphNodeK2Base::LatentBubbleColor(1.f, 0.5f, 0.25f);
 const FLinearColor SGraphNodeK2Base::TimelineBubbleColor(0.7f, 0.5f, 0.5f);
-const FLinearColor SGraphNodeK2Base::PinnedWatchColor(0.7f, 0.5f, 0.5f);
+const FLinearColor SGraphNodeK2Base::PinnedWatchColor(0.35f, 0.25f, 0.25f);
 
 void SGraphNodeK2Base::UpdateStandardNode()
 {
diff --git a/Engine/Source/Editor/GraphEditor/Private/KismetNodes/SGraphNodeK2Sequence.cpp b/Engine/Source/Editor/GraphEditor/Private/KismetNodes/SGraphNodeK2Sequence.cpp
--- a/Engine/Source/Editor/GraphEditor/Private/KismetNodes/SGraphNodeK2Sequence.cpp
+++ b/Engine/Source/Editor/GraphEditor/Private/KismetNodes/SGraphNodeK2Sequence.cpp
@@ -31,6 +31,7 @@ void SGraphNodeK2Sequence::CreateOutputSideAddButton(TSharedPtr<SVerticalBox> Ou
 	OutputBox->AddSlot()
 	.AutoHeight()
 	.VAlign(VAlign_Center)
+	.HAlign(HAlign_Right)
 	.Padding(AddPinPadding)
 	[
 		AddPinButton
@@ -52,3 +53,10 @@ FReply SGraphNodeK2Sequence::OnAddPin()
 	
 	return FReply::Handled();
 }
+
+EVisibility SGraphNodeK2Sequence::IsAddPinButtonVisible() const
+{
+	IK2Node_AddPinInterface* AddPinNode = Cast<IK2Node_AddPinInterface>(GraphNode);
+	ensure(AddPinNode);
+	return ((AddPinNode && AddPinNode->CanAddPin()) ? EVisibility::Visible : EVisibility::Collapsed);
+}
\ No newline at end of file
diff --git a/Engine/Source/Editor/GraphEditor/Private/KismetNodes/SGraphNodeK2Sequence.h b/Engine/Source/Editor/GraphEditor/Private/KismetNodes/SGraphNodeK2Sequence.h
--- a/Engine/Source/Editor/GraphEditor/Private/KismetNodes/SGraphNodeK2Sequence.h
+++ b/Engine/Source/Editor/GraphEditor/Private/KismetNodes/SGraphNodeK2Sequence.h
@@ -22,5 +22,6 @@ protected:
 	// SGraphNode interface
 	virtual void CreateOutputSideAddButton(TSharedPtr<SVerticalBox> OutputBox) override;
 	virtual FReply OnAddPin() override;
+	virtual EVisibility IsAddPinButtonVisible() const override;
 	// End of SGraphNode interface
 };
diff --git a/Engine/Source/Editor/GraphEditor/Private/KismetPins/SGraphPinIndex.cpp b/Engine/Source/Editor/GraphEditor/Private/KismetPins/SGraphPinIndex.cpp
--- a/Engine/Source/Editor/GraphEditor/Private/KismetPins/SGraphPinIndex.cpp
+++ b/Engine/Source/Editor/GraphEditor/Private/KismetPins/SGraphPinIndex.cpp
@@ -35,9 +35,10 @@ void SGraphPinIndex::OnTypeChanged(const FEdGraphPinType& PinType)
 		GraphPinObj->Modify();
 		GraphPinObj->PinType = PinType;
 		// Let the node know that one of its' pins had their pin type changed
-		if (GraphPinObj->GetOwningNode())
+		if (UEdGraphNode* OwningNode = GraphPinObj->GetOwningNode())
 		{
-			GraphPinObj->GetOwningNode()->PinTypeChanged(GraphPinObj);
+			OwningNode->PinTypeChanged(GraphPinObj);
+			OwningNode->ReconstructNode();
 		}
 	}
 }
diff --git a/Engine/Source/Editor/GraphEditor/Private/KismetPins/SGraphPinKey.cpp b/Engine/Source/Editor/GraphEditor/Private/KismetPins/SGraphPinKey.cpp
--- a/Engine/Source/Editor/GraphEditor/Private/KismetPins/SGraphPinKey.cpp
+++ b/Engine/Source/Editor/GraphEditor/Private/KismetPins/SGraphPinKey.cpp
@@ -12,11 +12,19 @@ void SGraphPinKey::Construct(const FArguments& InArgs, UEdGraphPin* InGraphPinOb
 
 	SelectedKey = FKey(*InGraphPinObj->GetDefaultAsString());
 
-	if (!SelectedKey.IsValid() && InGraphPinObj->Direction == EEdGraphPinDirection::EGPD_Input)
+	if (InGraphPinObj->Direction == EEdGraphPinDirection::EGPD_Input)
 	{
-		// Ensure first valid key is always set by default
-		SelectedKey = KeyList[0];
-		InGraphPinObj->GetSchema()->TrySetDefaultValue(*InGraphPinObj, SelectedKey.ToString());
+		// None is a valid key
+		if (SelectedKey.GetFName() == NAME_None)
+		{
+			SelectedKey = EKeys::Invalid;
+		}
+		else if (!SelectedKey.IsValid())
+		{
+			// Ensure first valid key is always set by default
+			SelectedKey = KeyList[0];
+			InGraphPinObj->GetSchema()->TrySetDefaultValue(*InGraphPinObj, SelectedKey.ToString());
+		}
 	}
 
 	SGraphPin::Construct(SGraphPin::FArguments(), InGraphPinObj);
diff --git a/Engine/Source/Editor/GraphEditor/Private/SGraphEditorImpl.cpp b/Engine/Source/Editor/GraphEditor/Private/SGraphEditorImpl.cpp
--- a/Engine/Source/Editor/GraphEditor/Private/SGraphEditorImpl.cpp
+++ b/Engine/Source/Editor/GraphEditor/Private/SGraphEditorImpl.cpp
@@ -646,11 +646,11 @@ TSharedPtr<SWidget> SGraphEditorImpl::GetTitleBar() const
 	return TitleBar;
 }
 
-void SGraphEditorImpl::SetViewLocation( const FVector2D& Location, float ZoomAmount ) 
+void SGraphEditorImpl::SetViewLocation( const FVector2D& Location, float ZoomAmount, const FGuid& BookmarkId ) 
 {
 	if( GraphPanel.IsValid() &&  EdGraphObj && (!IsLocked() || !GraphPanel->HasDeferredObjectFocus()))
 	{
-		GraphPanel->RestoreViewSettings(Location, ZoomAmount);
+		GraphPanel->RestoreViewSettings(Location, ZoomAmount, BookmarkId);
 	}
 }
 
@@ -663,6 +663,14 @@ void SGraphEditorImpl::GetViewLocation( FVector2D& Location, float& ZoomAmount )
 	}
 }
 
+void SGraphEditorImpl::GetViewBookmark( FGuid& BookmarkId )
+{
+	if (GraphPanel.IsValid())
+	{
+		BookmarkId = GraphPanel->GetViewBookmarkId();
+	}
+}
+
 void SGraphEditorImpl::LockToGraphEditor( TWeakPtr<SGraphEditor> Other ) 
 {
 	if( !LockedGraphs.Contains(Other) )
diff --git a/Engine/Source/Editor/GraphEditor/Private/SGraphEditorImpl.h b/Engine/Source/Editor/GraphEditor/Private/SGraphEditorImpl.h
--- a/Engine/Source/Editor/GraphEditor/Private/SGraphEditorImpl.h
+++ b/Engine/Source/Editor/GraphEditor/Private/SGraphEditorImpl.h
@@ -146,8 +146,9 @@ public:
 	virtual bool GetBoundsForNode( const UEdGraphNode* InNode, class FSlateRect& Rect, float Padding) const override;
 	virtual void NotifyGraphChanged() override;
 	virtual TSharedPtr<SWidget> GetTitleBar() const override;
-	virtual void SetViewLocation(const FVector2D& Location, float ZoomAmount) override;
+	virtual void SetViewLocation(const FVector2D& Location, float ZoomAmount, const FGuid& BookmarkId = FGuid()) override;
 	virtual void GetViewLocation(FVector2D& Location, float& ZoomAmount) override;
+	virtual void GetViewBookmark(FGuid& BookmarkId) override;
 	virtual void LockToGraphEditor(TWeakPtr<SGraphEditor> Other) override;
 	virtual void UnlockFromGraphEditor(TWeakPtr<SGraphEditor> Other) override;
 	virtual void AddNotification ( FNotificationInfo& Info, bool bSuccess ) override;
diff --git a/Engine/Source/Editor/GraphEditor/Private/SNodePanel.cpp b/Engine/Source/Editor/GraphEditor/Private/SNodePanel.cpp
--- a/Engine/Source/Editor/GraphEditor/Private/SNodePanel.cpp
+++ b/Engine/Source/Editor/GraphEditor/Private/SNodePanel.cpp
@@ -499,7 +499,14 @@ void SNodePanel::Tick( const FGeometry& AllottedGeometry, const double InCurrent
 
 	PopulateVisibleChildren(AllottedGeometry);
 
-	OldZoomAmount = GetZoomAmount();
+	// Reset the current bookmark if the location and/or zoom level has been changed.
+	const float CurZoomAmount = GetZoomAmount();
+	if (CurrentBookmarkGuid.IsValid() && (OldViewOffset != ViewOffset || OldZoomAmount != CurZoomAmount))
+	{
+		CurrentBookmarkGuid.Invalidate();
+	}
+
+	OldZoomAmount = CurZoomAmount;
 	OldViewOffset = ViewOffset;
 
 	SPanel::Tick( AllottedGeometry, InCurrentTime, InDeltaTime );
@@ -1182,7 +1189,7 @@ bool SNodePanel::Contains(UObject* Node) const
 	return NodeToWidgetLookup.Find(Node) != nullptr;
 }
 
-void SNodePanel::RestoreViewSettings(const FVector2D& InViewOffset, float InZoomAmount)
+void SNodePanel::RestoreViewSettings(const FVector2D& InViewOffset, float InZoomAmount, const FGuid& InBookmarkGuid)
 {
 	ViewOffset = InViewOffset;
 
@@ -1206,6 +1213,9 @@ void SNodePanel::RestoreViewSettings(const FVector2D& InViewOffset, float InZoom
 	// This is so our locked window isn't forced to update according to this movement.
 	OldViewOffset = ViewOffset;
 	OldZoomAmount = GetZoomAmount();
+
+	// Update the current bookmark ID.
+	CurrentBookmarkGuid = InBookmarkGuid;
 }
 
 float SNodePanel::GetSnapGridSize()
@@ -1678,6 +1688,9 @@ bool SNodePanel::HasDeferredObjectFocus() const
 void SNodePanel::PostChangedZoom()
 {
 	CurrentLOD = ZoomLevels->GetLOD(ZoomLevel);
+
+	// Invalidate the current bookmark.
+	CurrentBookmarkGuid.Invalidate();
 }
 
 void SNodePanel::RequestZoomToFit()
diff --git a/Engine/Source/Editor/GraphEditor/Public/SNodePanel.h b/Engine/Source/Editor/GraphEditor/Public/SNodePanel.h
--- a/Engine/Source/Editor/GraphEditor/Public/SNodePanel.h
+++ b/Engine/Source/Editor/GraphEditor/Public/SNodePanel.h
@@ -4,6 +4,7 @@
 
 #include "CoreMinimal.h"
 #include "Misc/Attribute.h"
+#include "Misc/Guid.h"
 #include "Layout/Visibility.h"
 #include "Layout/SlateRect.h"
 #include "Layout/Geometry.h"
@@ -654,11 +655,14 @@ public:
 	/** @return the view offset in graph space */
 	FVector2D GetViewOffset() const;
 
+	/** @return the current view bookmark ID */
+	const FGuid& GetViewBookmarkId() const { return CurrentBookmarkGuid; }
+
 	/** Given a coordinate in panel space (i.e. panel widget space), return the same coordinate in graph space while taking zoom and panning into account */
 	FVector2D PanelCoordToGraphCoord(const FVector2D& PanelSpaceCoordinate) const;
 
-	/** Restore the graph panel to the supplied view offset/zoom */
-	void RestoreViewSettings(const FVector2D& InViewOffset, float InZoomAmount);
+	/** Restore the graph panel to the supplied view offset/zoom. Also, optionally set the current bookmark ID. */
+	void RestoreViewSettings(const FVector2D& InViewOffset, float InZoomAmount, const FGuid& InBookmarkGuid = FGuid());
 
 	/** Get the grid snap size */
 	static float GetSnapGridSize();
@@ -958,4 +962,7 @@ private:
 	/** Zoom target rectangle */
 	FVector2D ZoomTargetTopLeft;
 	FVector2D ZoomTargetBottomRight;
+
+	/** Current view bookmark info */
+	FGuid CurrentBookmarkGuid;
 };
diff --git a/Engine/Source/Editor/Kismet/Private/BlueprintDetailsCustomization.cpp b/Engine/Source/Editor/Kismet/Private/BlueprintDetailsCustomization.cpp
--- a/Engine/Source/Editor/Kismet/Private/BlueprintDetailsCustomization.cpp
+++ b/Engine/Source/Editor/Kismet/Private/BlueprintDetailsCustomization.cpp
@@ -3913,16 +3913,47 @@ void FBlueprintDelegateActionDetails::CustomizeDetails( IDetailLayoutBuilder& De
 		TSharedRef<FBlueprintGraphArgumentGroupLayout> InputArgumentGroup = MakeShareable(new FBlueprintGraphArgumentGroupLayout(SharedThis(this), FunctionEntryNode));
 		InputsCategory.AddCustomBuilder(InputArgumentGroup);
 
-		InputsCategory.AddCustomRow( LOCTEXT("FunctionNewInputArg", "New") )
+		TSharedRef<SHorizontalBox> InputsHeaderContentWidget = SNew(SHorizontalBox);
+		TWeakPtr<SWidget> WeakInputsHeaderWidget = InputsHeaderContentWidget;
+		InputsHeaderContentWidget->AddSlot()
 		[
-			SNew(SBox)
+			SNew(SHorizontalBox)
+		];
+		InputsHeaderContentWidget->AddSlot()
+		.AutoWidth()
+		[
+			SNew(SButton)
+			.ButtonStyle(FEditorStyle::Get(), "RoundButton")
+			.ForegroundColor(FEditorStyle::GetSlateColor("DefaultForeground"))
+			.ContentPadding(FMargin(2, 0))
+			.OnClicked(this, &FBlueprintDelegateActionDetails::OnAddNewInputClicked)
 			.HAlign(HAlign_Right)
+			.ToolTipText(LOCTEXT("DelegateNewOutputArgTooltip", "Create a new input argument"))
+			.VAlign(VAlign_Center)
+			.AddMetaData<FTagMetaData>(FTagMetaData(TEXT("DelegateNewInputArg")))
 			[
-				SNew(SButton)
-				.Text(LOCTEXT("FunctionNewInputArg", "New"))
-				.OnClicked(this, &FBlueprintDelegateActionDetails::OnAddNewInputClicked)
+				SNew(SHorizontalBox)
+				+SHorizontalBox::Slot()
+				.AutoWidth()
+				.Padding(FMargin(0, 1))
+				[
+					SNew(SImage)
+					.Image(FEditorStyle::GetBrush("Plus"))
+				]
+				+SHorizontalBox::Slot()
+				.VAlign(VAlign_Center)
+				.AutoWidth()
+				.Padding(FMargin(2, 0, 0, 0))
+				[
+					SNew(STextBlock)
+					.Font(IDetailLayoutBuilder::GetDetailFontBold())
+					.Text(LOCTEXT("DelegateNewParameterInputArg", "New Parameter"))
+					.Visibility(this, &FBlueprintDelegateActionDetails::OnGetSectionTextVisibility, WeakInputsHeaderWidget)
+					.ShadowOffset(FVector2D(1, 1))
+				]
 			]
 		];
+		InputsCategory.HeaderContent(InputsHeaderContentWidget);
 
 		CollectAvailibleSignatures();
 
@@ -4024,6 +4055,21 @@ void FBaseBlueprintGraphActionDetails::OnParamsChanged(UK2Node_EditablePinBase*
 	}
 }
 
+EVisibility FBlueprintDelegateActionDetails::OnGetSectionTextVisibility(TWeakPtr<SWidget> RowWidget) const
+{
+	bool ShowText = RowWidget.Pin()->IsHovered();
+
+	// If the row is currently hovered, or a menu is being displayed for a button, keep the button expanded.
+	if (ShowText)
+	{
+		return EVisibility::SelfHitTestInvisible;
+	}
+	else
+	{
+		return EVisibility::Collapsed;
+	}
+}
+
 struct FPinRenamedHelper : public FBasePinChangeHelper
 {
 	TSet<UBlueprint*> ModifiedBlueprints;
diff --git a/Engine/Source/Editor/Kismet/Private/BlueprintDetailsCustomization.h b/Engine/Source/Editor/Kismet/Private/BlueprintDetailsCustomization.h
--- a/Engine/Source/Editor/Kismet/Private/BlueprintDetailsCustomization.h
+++ b/Engine/Source/Editor/Kismet/Private/BlueprintDetailsCustomization.h
@@ -390,6 +390,7 @@ private:
 	void CollectAvailibleSignatures();
 	void OnFunctionSelected(TSharedPtr<FString> FunctionItemData, ESelectInfo::Type SelectInfo);
 	bool IsBlueprintProperty() const;
+	EVisibility OnGetSectionTextVisibility(TWeakPtr<SWidget> RowWidget) const;
 
 private:
 
diff --git a/Engine/Source/Editor/Kismet/Private/BlueprintEditor.cpp b/Engine/Source/Editor/Kismet/Private/BlueprintEditor.cpp
--- a/Engine/Source/Editor/Kismet/Private/BlueprintEditor.cpp
+++ b/Engine/Source/Editor/Kismet/Private/BlueprintEditor.cpp
@@ -107,6 +107,7 @@
 #include "SSCSEditorViewport.h"
 #include "SKismetInspector.h"
 #include "SBlueprintPalette.h"
+#include "SBlueprintBookmarks.h"
 #include "SBlueprintActionMenu.h"
 #include "SMyBlueprint.h"
 #include "SReplaceNodeReferences.h"
@@ -754,6 +755,11 @@ void FBlueprintEditor::RefreshEditors(ERefreshBlueprintEditorReason::Type Reason
 		UpdateSCSPreview();
 	}
 
+	if (BookmarksWidget.IsValid())
+	{
+		BookmarksWidget->RefreshBookmarksTree();
+	}
+
 	// Note: There is an optimization inside of ShowDetailsForSingleObject() that skips the refresh if the object being selected is the same as the previous object.
 	// The SKismetInspector class is shared between both Defaults mode and Components mode, but in Defaults mode the object selected is always going to be the CDO. Given
 	// that the selection does not really change, we force it to refresh and skip the optimization. Otherwise, some things may not work correctly in Defaults mode. For
@@ -2246,6 +2252,10 @@ void FBlueprintEditor::CreateDefaultTabContents(const TArray<UBlueprint*>& InBlu
 		this->Palette = 
 			SNew(SBlueprintPalette, SharedThis(this))
 				.IsEnabled(this, &FBlueprintEditor::IsFocusedGraphEditable);
+
+		this->BookmarksWidget =
+			SNew(SBlueprintBookmarks)
+				.EditorContext(SharedThis(this));
 	}
 
 	if (IsEditingSingleBlueprint())
@@ -2619,6 +2629,11 @@ void FBlueprintEditor::FindInBlueprint_Clicked()
 
 void FBlueprintEditor::ReparentBlueprint_Clicked()
 {
+	if (!ReparentBlueprint_IsVisible())
+	{
+		return;
+	}
+
 	TArray<UBlueprint*> Blueprints;
 	for (int32 i = 0; i < GetEditingObjects().Num(); ++i)
 	{
@@ -4044,7 +4059,7 @@ void FBlueprintEditor::OnAddOptionPin()
 			const FScopedTransaction Transaction( LOCTEXT("AddOptionPin", "Add Option Pin") );
 			SeqNode->Modify();
 
-			SeqNode->AddOptionPinToNode();
+			SeqNode->AddInputPin();
 
 			const UEdGraphSchema* Schema = SeqNode->GetSchema();
 			Schema->ReconstructNode(*SeqNode);
@@ -4068,7 +4083,7 @@ bool FBlueprintEditor::CanAddOptionPin() const
 	{
 		UK2Node_Select* SeqNode = Cast<UK2Node_Select>(*It);
 		// There's a bad node so return false
-		if (SeqNode == NULL || !SeqNode->CanAddOptionPinToNode())
+		if (SeqNode == nullptr || !SeqNode->CanAddPin())
 		{
 			return false;
 		}
@@ -4196,9 +4211,11 @@ void FBlueprintEditor::OnChangePinTypeFinished(const FEdGraphPinType& PinType, U
 {
 	if (FBlueprintEditorUtils::IsPinTypeValid(PinType))
 	{
+		UEdGraphNode* OwningNode = InSelectedPin->GetOwningNode();
+		OwningNode->Modify();
 		InSelectedPin->PinType = PinType;
 		if (UK2Node_Select* SelectNode = Cast<UK2Node_Select>(InSelectedPin->GetOwningNode()))
-			{
+		{
 			SelectNode->ChangePinType(InSelectedPin);
 		}
 	}
@@ -6938,6 +6955,43 @@ void FBlueprintEditor::RequestSaveEditedObjectState()
 	bRequestedSavingOpenDocumentState = true;
 }
 
+void FBlueprintEditor::GetBoundsForNode(const UEdGraphNode* InNode, class FSlateRect& OutRect, float InPadding) const
+{
+	if (FocusedGraphEdPtr.IsValid())
+	{
+		FocusedGraphEdPtr.Pin()->GetBoundsForNode(InNode, OutRect, InPadding);
+	}
+}
+
+void FBlueprintEditor::GetViewBookmark(FGuid& BookmarkId)
+{
+	BookmarkId.Invalidate();
+
+	if (FocusedGraphEdPtr.IsValid())
+	{
+		FocusedGraphEdPtr.Pin()->GetViewBookmark(BookmarkId);
+	}
+}
+
+void FBlueprintEditor::GetViewLocation(FVector2D& Location, float& ZoomAmount)
+{
+	Location = FVector2D::ZeroVector;
+	ZoomAmount = 0.0f;
+
+	if (FocusedGraphEdPtr.IsValid())
+	{
+		FocusedGraphEdPtr.Pin()->GetViewLocation(Location, ZoomAmount);
+	}
+}
+
+void FBlueprintEditor::SetViewLocation(const FVector2D& Location, float ZoomAmount, const FGuid& BookmarkId)
+{
+	if (FocusedGraphEdPtr.IsValid())
+	{
+		FocusedGraphEdPtr.Pin()->SetViewLocation(Location, ZoomAmount, BookmarkId);
+	}
+}
+
 void FBlueprintEditor::Tick(float DeltaTime)
 {
 	// Create or update the Blueprint actor instance in the preview scene
@@ -7113,6 +7167,11 @@ void FBlueprintEditor::OnAddNewLocalVariable()
 
 void FBlueprintEditor::OnAddNewDelegate()
 {
+	if (!AddNewDelegateIsVisible())
+	{
+		return;
+	}
+
 	const UEdGraphSchema_K2* K2Schema = GetDefault<UEdGraphSchema_K2>();
 	check(NULL != K2Schema);
 	UBlueprint* const Blueprint = GetBlueprintObj();
@@ -7156,6 +7215,11 @@ void FBlueprintEditor::OnAddNewDelegate()
 
 void FBlueprintEditor::NewDocument_OnClicked(ECreatedDocumentType GraphType)
 {
+	if (!NewDocument_IsVisibleForType(GraphType))
+	{
+		return;
+	}
+
 	FText DocumentNameText;
 	bool bResetMyBlueprintFilter = false;
 
@@ -7677,7 +7741,7 @@ void FBlueprintEditor::RestoreEditedObjectState()
 
 	for (int32 i = 0; i < Blueprint->LastEditedDocuments.Num(); i++)
 	{
-		if (UObject* Obj = Blueprint->LastEditedDocuments[i].EditedObject)
+		if (UObject* Obj = Blueprint->LastEditedDocuments[i].EditedObjectPath.ResolveObject())
 		{
 			if(UEdGraph* Graph = Cast<UEdGraph>(Obj))
 			{
@@ -8186,6 +8250,137 @@ void FBlueprintEditor::SelectGraphActionItemByName(const FName& ItemName, ESelec
 	}
 }
 
+FBPEditorBookmarkNode* FBlueprintEditor::AddBookmark(const FText& DisplayName, const FEditedDocumentInfo& BookmarkInfo, bool bSharedBookmark)
+{
+	FBPEditorBookmarkNode* NewNode = nullptr;
+
+	if (bSharedBookmark)
+	{
+		if (UBlueprint* Blueprint = GetBlueprintObj())
+		{
+			NewNode = new(Blueprint->BookmarkNodes) FBPEditorBookmarkNode;
+			NewNode->NodeGuid = FGuid::NewGuid();
+			NewNode->DisplayName = DisplayName;
+
+			Blueprint->Modify();
+			Blueprint->Bookmarks.Add(NewNode->NodeGuid, BookmarkInfo);
+		}
+	}
+	else if(UBlueprintEditorSettings* LocalSettings = GetMutableDefault<UBlueprintEditorSettings>())
+	{
+		NewNode = new(LocalSettings->BookmarkNodes) FBPEditorBookmarkNode;
+		NewNode->NodeGuid = FGuid::NewGuid();
+		NewNode->DisplayName = DisplayName;
+
+		LocalSettings->Bookmarks.Add(NewNode->NodeGuid, BookmarkInfo);
+		LocalSettings->SaveConfig();
+	}
+
+	if (NewNode && BookmarksWidget.IsValid())
+	{
+		BookmarksWidget->RefreshBookmarksTree();
+	}
+
+	return NewNode;
+}
+
+void FBlueprintEditor::RenameBookmark(const FGuid& BookmarkNodeId, const FText& NewName)
+{
+	bool bFoundSharedBookmark = false;
+	if (UBlueprint* Blueprint = GetBlueprintObj())
+	{
+		for (FBPEditorBookmarkNode& BookmarkNode : Blueprint->BookmarkNodes)
+		{
+			if (BookmarkNode.NodeGuid == BookmarkNodeId)
+			{
+				Blueprint->Modify();
+				BookmarkNode.DisplayName = NewName;
+
+				bFoundSharedBookmark = true;
+				break;
+			}
+		}
+	}
+
+	if (!bFoundSharedBookmark)
+	{
+		UBlueprintEditorSettings* LocalSettings = GetMutableDefault<UBlueprintEditorSettings>();
+		for (FBPEditorBookmarkNode& BookmarkNode : LocalSettings->BookmarkNodes)
+		{
+			if (BookmarkNode.NodeGuid == BookmarkNodeId)
+			{
+				BookmarkNode.DisplayName = NewName;
+				LocalSettings->SaveConfig();
+
+				break;
+			}
+		}
+	}
+
+	if (BookmarksWidget.IsValid())
+	{
+		BookmarksWidget->RefreshBookmarksTree();
+	}
+}
+
+void FBlueprintEditor::RemoveBookmark(const FGuid& BookmarkNodeId, bool bRefreshUI)
+{
+	bool bFoundSharedBookmark = false;
+	if (UBlueprint* Blueprint = GetBlueprintObj())
+	{
+		for (int32 i = 0; i < Blueprint->BookmarkNodes.Num(); ++i)
+		{
+			const FBPEditorBookmarkNode& BookmarkNode = Blueprint->BookmarkNodes[i];
+			if (BookmarkNode.NodeGuid == BookmarkNodeId)
+			{
+				Blueprint->Modify();
+				Blueprint->BookmarkNodes.RemoveAtSwap(i);
+				FEditedDocumentInfo BookmarkInfo = Blueprint->Bookmarks.FindAndRemoveChecked(BookmarkNodeId);
+
+				FGuid CurrentBookmarkId;
+				GetViewBookmark(CurrentBookmarkId);
+				if (CurrentBookmarkId == BookmarkNodeId)
+				{
+					SetViewLocation(BookmarkInfo.SavedViewOffset, BookmarkInfo.SavedZoomAmount);
+				}
+
+				bFoundSharedBookmark = true;
+				break;
+			}
+		}
+	}
+
+	if (!bFoundSharedBookmark)
+	{
+		UBlueprintEditorSettings* LocalSettings = GetMutableDefault<UBlueprintEditorSettings>();
+		for (int32 i = 0; i < LocalSettings->BookmarkNodes.Num(); ++i)
+		{
+			const FBPEditorBookmarkNode& BookmarkNode = LocalSettings->BookmarkNodes[i];
+			if (BookmarkNode.NodeGuid == BookmarkNodeId)
+			{
+				LocalSettings->BookmarkNodes.RemoveAtSwap(i);
+				LocalSettings->SaveConfig();
+
+				FEditedDocumentInfo BookmarkInfo = LocalSettings->Bookmarks.FindAndRemoveChecked(BookmarkNodeId);
+
+				FGuid CurrentBookmarkId;
+				GetViewBookmark(CurrentBookmarkId);
+				if (CurrentBookmarkId == BookmarkNodeId)
+				{
+					SetViewLocation(BookmarkInfo.SavedViewOffset, BookmarkInfo.SavedZoomAmount);
+				}
+
+				break;
+			}
+		}
+	}
+
+	if (bRefreshUI && BookmarksWidget.IsValid())
+	{
+		BookmarksWidget->RefreshBookmarksTree();
+	}
+}
+
 /////////////////////////////////////////////////////
 /////////////////////////////////////////////////////
 /////////////////////////////////////////////////////
diff --git a/Engine/Source/Editor/Kismet/Private/BlueprintEditorModes.cpp b/Engine/Source/Editor/Kismet/Private/BlueprintEditorModes.cpp
--- a/Engine/Source/Editor/Kismet/Private/BlueprintEditorModes.cpp
+++ b/Engine/Source/Editor/Kismet/Private/BlueprintEditorModes.cpp
@@ -19,6 +19,7 @@
 #include "BlueprintEditorTabs.h"
 #include "BlueprintEditorTabFactories.h"
 #include "BlueprintEditorSharedTabFactories.h"
+#include "BlueprintEditorSettings.h"
 
 #define LOCTEXT_NAMESPACE "BlueprintEditor"
 
@@ -61,6 +62,7 @@ TSharedPtr<FTabManager::FLayout> GetDefaltEditorLayout(TSharedPtr<class FBluepri
 					->SetSizeCoefficient( 0.20f )
 					->AddTab( FBlueprintEditorTabs::CompilerResultsID, ETabState::ClosedTab )
 					->AddTab( FBlueprintEditorTabs::FindResultsID, ETabState::ClosedTab )
+					->AddTab( FBlueprintEditorTabs::BookmarksID, ETabState::ClosedTab )
 				)
 			)
 			->Split
@@ -99,6 +101,10 @@ FBlueprintEditorApplicationMode::FBlueprintEditorApplicationMode(TSharedPtr<clas
 	}
 	BlueprintEditorTabFactories.RegisterFactory(MakeShareable(new FCompilerResultsSummoner(InBlueprintEditor)));
 	BlueprintEditorTabFactories.RegisterFactory(MakeShareable(new FFindResultsSummoner(InBlueprintEditor)));
+	if (GetDefault<UEditorExperimentalSettings>()->bEnableBlueprintGraphBookmarks)
+	{
+		BlueprintEditorTabFactories.RegisterFactory(MakeShareable(new FBookmarksSummoner(InBlueprintEditor)));
+	}
 	
 	if( bRegisterViewport )
 	{
@@ -388,6 +394,7 @@ FBlueprintInterfaceApplicationMode::FBlueprintInterfaceApplicationMode(TSharedPt
 						->SetSizeCoefficient( 0.20f )
 						->AddTab( FBlueprintEditorTabs::CompilerResultsID, ETabState::ClosedTab )
 						->AddTab( FBlueprintEditorTabs::FindResultsID, ETabState::ClosedTab )
+						->AddTab( FBlueprintEditorTabs::BookmarksID, ETabState::ClosedTab )
 					)
 				)
 				->Split
@@ -491,6 +498,7 @@ FBlueprintMacroApplicationMode::FBlueprintMacroApplicationMode(TSharedPtr<class
 						FTabManager::NewStack()
 						->SetSizeCoefficient( 0.20f )
 						->AddTab( FBlueprintEditorTabs::FindResultsID, ETabState::ClosedTab )
+						->AddTab( FBlueprintEditorTabs::BookmarksID, ETabState::ClosedTab )
 					)
 				)
 				->Split
@@ -566,6 +574,10 @@ FBlueprintEditorUnifiedMode::FBlueprintEditorUnifiedMode(TSharedPtr<class FBluep
 	}
 	BlueprintEditorTabFactories.RegisterFactory(MakeShareable(new FCompilerResultsSummoner(InBlueprintEditor)));
 	BlueprintEditorTabFactories.RegisterFactory(MakeShareable(new FFindResultsSummoner(InBlueprintEditor)));
+	if (GetDefault<UEditorExperimentalSettings>()->bEnableBlueprintGraphBookmarks)
+	{
+		BlueprintEditorTabFactories.RegisterFactory(MakeShareable(new FBookmarksSummoner(InBlueprintEditor)));
+	}
 	
 	if( bRegisterViewport )
 	{
@@ -625,6 +637,7 @@ FBlueprintEditorUnifiedMode::FBlueprintEditorUnifiedMode(TSharedPtr<class FBluep
 						->SetSizeCoefficient( 0.20f )
 						->AddTab( FBlueprintEditorTabs::CompilerResultsID, ETabState::ClosedTab )
 						->AddTab( FBlueprintEditorTabs::FindResultsID, ETabState::ClosedTab )
+						->AddTab( FBlueprintEditorTabs::BookmarksID, ETabState::ClosedTab )
 					)
 				)
 				->Split
@@ -685,6 +698,7 @@ FBlueprintEditorUnifiedMode::FBlueprintEditorUnifiedMode(TSharedPtr<class FBluep
 						->SetSizeCoefficient( 0.20f )
 						->AddTab( FBlueprintEditorTabs::CompilerResultsID, ETabState::ClosedTab )
 						->AddTab( FBlueprintEditorTabs::FindResultsID, ETabState::ClosedTab )
+						->AddTab( FBlueprintEditorTabs::BookmarksID, ETabState::ClosedTab )
 					)
 				)
 				->Split
diff --git a/Engine/Source/Editor/Kismet/Private/BlueprintEditorModule.cpp b/Engine/Source/Editor/Kismet/Private/BlueprintEditorModule.cpp
--- a/Engine/Source/Editor/Kismet/Private/BlueprintEditorModule.cpp
+++ b/Engine/Source/Editor/Kismet/Private/BlueprintEditorModule.cpp
@@ -7,12 +7,13 @@
 #include "Framework/Application/SlateApplication.h"
 #include "EditorStyleSet.h"
 #include "EditorUndoClient.h"
-#include "Kismet2/KismetEditorUtilities.h"
 #include "Logging/TokenizedMessage.h"
 #include "Misc/ConfigCacheIni.h"
 #include "UObject/UObjectHash.h"
 #include "Serialization/ArchiveReplaceObjectRef.h"
 #include "BlueprintEditor.h"
+#include "Kismet2/KismetEditorUtilities.h"
+#include "Kismet2/KismetDebugUtilities.h"
 #include "Kismet2/BlueprintEditorUtils.h"
 //#include "BlueprintWarningsConfigurationPanel.h"
 #include "Editor/WorkspaceMenuStructure/Public/WorkspaceMenuStructure.h"
@@ -32,6 +33,7 @@
 #include "Widgets/Docking/SDockTab.h"
 #include "KismetCompiler.h"
 #include "CallStackViewer.h"
+#include "WatchPointViewer.h"
 
 #define LOCTEXT_NAMESPACE "BlueprintEditor"
 
@@ -205,6 +207,8 @@ void FBlueprintEditorModule::StartupModule()
 
 	CallStackViewer::RegisterTabSpawner();
 
+	WatchViewer::RegisterTabSpawner();
+
 	// Have to check GIsEditor because right now editor modules can be loaded by the game
 	// Once LoadModule is guaranteed to return NULL for editor modules in game, this can be removed
 	// Without this check, loading the level editor in the game will crash
diff --git a/Engine/Source/Editor/Kismet/Private/BlueprintEditorTabFactories.cpp b/Engine/Source/Editor/Kismet/Private/BlueprintEditorTabFactories.cpp
--- a/Engine/Source/Editor/Kismet/Private/BlueprintEditorTabFactories.cpp
+++ b/Engine/Source/Editor/Kismet/Private/BlueprintEditorTabFactories.cpp
@@ -15,6 +15,7 @@
 #include "SSCSEditor.h"
 #include "SSCSEditorViewport.h"
 #include "SBlueprintPalette.h"
+#include "SBlueprintBookmarks.h"
 #include "FindInBlueprints.h"
 #include "Kismet2/BlueprintEditorUtils.h"
 #include "SMyBlueprint.h"
@@ -311,6 +312,25 @@ TSharedRef<SWidget> FPaletteSummoner::CreateTabBody(const FWorkflowTabSpawnInfo&
 	return BlueprintEditorPtr->GetPalette();
 }
 
+FBookmarksSummoner::FBookmarksSummoner(TSharedPtr<class FAssetEditorToolkit> InHostingApp)
+	: FWorkflowTabFactory(FBlueprintEditorTabs::BookmarksID, InHostingApp)
+{
+	TabLabel = LOCTEXT("BookmarksTabTitle", "Bookmarks");
+	TabIcon = FSlateIcon(FEditorStyle::GetStyleSetName(), "Kismet.Tabs.Bookmarks");
+
+	bIsSingleton = true;
+
+	ViewMenuDescription = LOCTEXT("BookmarksView", "Bookmarks");
+	ViewMenuTooltip = LOCTEXT("BookmarksView_ToolTip", "Show bookmarks associated with this Blueprint");
+}
+
+TSharedRef<SWidget> FBookmarksSummoner::CreateTabBody(const FWorkflowTabSpawnInfo& Info) const
+{
+	TSharedPtr<FBlueprintEditor> BlueprintEditorPtr = StaticCastSharedPtr<FBlueprintEditor>(HostingApp.Pin());
+
+	return BlueprintEditorPtr->GetBookmarksWidget();
+}
+
 FMyBlueprintSummoner::FMyBlueprintSummoner(TSharedPtr<class FAssetEditorToolkit> InHostingApp)
 	: FWorkflowTabFactory(FBlueprintEditorTabs::MyBlueprintID, InHostingApp)
 {
@@ -420,6 +440,7 @@ void FGraphTabHistory::SaveHistory()
 	{
 		check(GraphEditor.IsValid());
 		GraphEditor.Pin()->GetViewLocation(SavedLocation, SavedZoomAmount);
+		GraphEditor.Pin()->GetViewBookmark(SavedBookmarkId);
 	}
 }
 
@@ -428,7 +449,7 @@ void FGraphTabHistory::RestoreHistory()
 	if (IsHistoryValid())
 	{
 		check(GraphEditor.IsValid());
-		GraphEditor.Pin()->SetViewLocation(SavedLocation, SavedZoomAmount);
+		GraphEditor.Pin()->SetViewLocation(SavedLocation, SavedZoomAmount, SavedBookmarkId);
 	}
 }
 
diff --git a/Engine/Source/Editor/Kismet/Private/BlueprintEditorTabFactories.h b/Engine/Source/Editor/Kismet/Private/BlueprintEditorTabFactories.h
--- a/Engine/Source/Editor/Kismet/Private/BlueprintEditorTabFactories.h
+++ b/Engine/Source/Editor/Kismet/Private/BlueprintEditorTabFactories.h
@@ -118,6 +118,8 @@ private:
 	FVector2D SavedLocation;
 	/** Saved zoom the graph editor was at when this history node was last visited */
 	float SavedZoomAmount;
+	/** Saved bookmark ID the graph editor was at when this history node was last visited */
+	FGuid SavedBookmarkId;
 };
 
 /////////////////////////////////////////////////////
@@ -231,6 +233,22 @@ public:
 	}
 };
 
+//////////////////////////////////////////////////////////////////////////
+// FBookmarksSummoner
+
+struct FBookmarksSummoner : public FWorkflowTabFactory
+{
+public:
+	FBookmarksSummoner(TSharedPtr<class FAssetEditorToolkit> InHostingApp);
+
+	virtual TSharedRef<SWidget> CreateTabBody(const FWorkflowTabSpawnInfo& Info) const override;
+
+	virtual FText GetTabToolTipText(const FWorkflowTabSpawnInfo& Info) const override
+	{
+		return LOCTEXT("BookmarksTooltip", "The Bookmarks tab allows you to create, edit and select bookmarks associated with this Blueprint.");
+	}
+};
+
 /////////////////////////////////////////////////////
 // FMyBlueprintSummoner
 
diff --git a/Engine/Source/Editor/Kismet/Private/BlueprintEditorTabs.cpp b/Engine/Source/Editor/Kismet/Private/BlueprintEditorTabs.cpp
--- a/Engine/Source/Editor/Kismet/Private/BlueprintEditorTabs.cpp
+++ b/Engine/Source/Editor/Kismet/Private/BlueprintEditorTabs.cpp
@@ -6,6 +6,7 @@ const FName FBlueprintEditorTabs::DetailsID( TEXT( "Inspector" ) );
 const FName FBlueprintEditorTabs::DefaultEditorID( TEXT( "DefaultEditor" ) );
 const FName FBlueprintEditorTabs::DebugID( TEXT( "Debug" ) );
 const FName FBlueprintEditorTabs::PaletteID( TEXT( "PaletteList" ) );
+const FName FBlueprintEditorTabs::BookmarksID( TEXT( "BookmarkList" ) );
 const FName FBlueprintEditorTabs::CompilerResultsID( TEXT( "CompilerResults" ) );
 const FName FBlueprintEditorTabs::FindResultsID( TEXT( "FindResults" ) );
 const FName FBlueprintEditorTabs::ConstructionScriptEditorID( TEXT( "ConstructionScriptEditor" ) );
diff --git a/Engine/Source/Editor/Kismet/Private/SGraphTitleBar.cpp b/Engine/Source/Editor/Kismet/Private/SGraphTitleBar.cpp
--- a/Engine/Source/Editor/Kismet/Private/SGraphTitleBar.cpp
+++ b/Engine/Source/Editor/Kismet/Private/SGraphTitleBar.cpp
@@ -2,10 +2,12 @@
 
 
 #include "SGraphTitleBar.h"
+#include "SGraphTitleBarAddNewBookmark.h"
 #include "Fonts/SlateFontInfo.h"
 #include "Misc/Paths.h"
 #include "Widgets/SBoxPanel.h"
 #include "Widgets/Layout/SBorder.h"
+#include "Widgets/Layout/SSeparator.h"
 #include "Widgets/Images/SImage.h"
 #include "Widgets/Text/STextBlock.h"
 #include "Styling/CoreStyle.h"
@@ -97,21 +99,35 @@ void SGraphTitleBar::Construct( const FArguments& InArgs )
 		.HAlign(HAlign_Fill)
 		.AddMetaData<FTagMetaData>(FTagMetaData(TEXT("EventGraphTitleBar")))
 		[
-			SNew(SVerticalBox)
+			SNew(SHorizontalBox)
+			+ SHorizontalBox::Slot()
+			.AutoWidth()
+			[
+				SNew(SGraphTitleBarAddNewBookmark)
+				.EditorPtr(Kismet2Ptr)
+			]
+			+ SHorizontalBox::Slot()
+			.AutoWidth()
+			[
+				SNew(SSeparator)
+				.Orientation(Orient_Vertical)
+			]
+			+ SHorizontalBox::Slot()
+			.AutoWidth()
+			[
+				InArgs._HistoryNavigationWidget.ToSharedRef()
+			]
+			+ SHorizontalBox::Slot()
+			.AutoWidth()
+			[
+				SNew(SSeparator)
+				.Orientation(Orient_Vertical)
+			]
 			// Title text/icon
-			+SVerticalBox::Slot()
-			.AutoHeight()
+			+SHorizontalBox::Slot()
 			[
 				SNew(SHorizontalBox)
-
-				+SHorizontalBox::Slot()
-				.AutoWidth()
-				[
-					InArgs._HistoryNavigationWidget.ToSharedRef()
-				]
-
 				+SHorizontalBox::Slot()
-				.HAlign(HAlign_Center)
 				.FillWidth(1.f)
 				[
 					SNew(SHorizontalBox)
diff --git a/Engine/Source/Editor/Kismet/Private/SGraphTitleBar.h b/Engine/Source/Editor/Kismet/Private/SGraphTitleBar.h
--- a/Engine/Source/Editor/Kismet/Private/SGraphTitleBar.h
+++ b/Engine/Source/Editor/Kismet/Private/SGraphTitleBar.h
@@ -60,7 +60,6 @@ protected:
 	/** Blueprint title being displayed for toolbar */
 	FText BlueprintTitle;
 
-
 protected:
 	/** Get the icon to use */
 	const FSlateBrush* GetTypeGlyph() const;
@@ -82,4 +81,7 @@ protected:
 
 	/** Helper method used to show blueprint title in breadcrumbs */
 	FText GetBlueprintTitle() const;
+
+	/** Helper method used to create the bookmark selector widget */
+	TSharedRef<SWidget> CreateBookmarkSelectionWidget();
 };
diff --git a/Engine/Source/Editor/Kismet/Private/WatchPointViewer.cpp b/Engine/Source/Editor/Kismet/Private/WatchPointViewer.cpp
new file mode 100644
--- /dev/null
+++ b/Engine/Source/Editor/Kismet/Private/WatchPointViewer.cpp
@@ -0,0 +1,632 @@
+// Copyright 1998-2017 Epic Games, Inc. All Rights Reserved.
+
+#include "WatchPointViewer.h"
+
+#include "Blueprint/WidgetBlueprintGeneratedClass.h"
+#include "EditorStyleSet.h"
+#include "Framework/Commands/GenericCommands.h"
+#include "Framework/Docking/TabManager.h"
+#include "Framework/MultiBox/MultiBoxBuilder.h"
+#include "HAL/PlatformApplicationMisc.h"
+#include "K2Node_Event.h"
+#include "Kismet2/KismetEditorUtilities.h"
+#include "Kismet2/KismetDebugUtilities.h"
+#include "GraphEditorActions.h"
+#include "EdGraphSchema_K2.h"
+
+#include "Editor.h"
+
+#include "Widgets/Docking/SDockTab.h"
+#include "Widgets/Images/SImage.h"
+#include "Widgets/Layout/SBox.h"
+#include "Widgets/Layout/SSpacer.h"
+#include "Widgets/Views/STreeView.h"
+#include "Widgets/Views/STableRow.h"
+#include "Widgets/Input/SHyperlink.h"
+
+#include "WorkspaceMenuStructure.h"
+#include "WorkspaceMenuStructureModule.h"
+
+#include "KismetNodes/KismetNodeInfoContext.h"
+#include "Stats/Stats.h"
+
+#define LOCTEXT_NAMESPACE "WatchPointViewer"
+
+struct FWatchRow
+{
+	FWatchRow(
+		UBlueprint* InBP,
+		const UEdGraphNode* InNode,
+		const UEdGraphPin* InPin,
+		FText InBlueprintName,
+		FText InGraphName,
+		FText InNodeName,
+		FText InDisplayName,
+		FText InValue,
+		FText InType
+	)
+		: BP(InBP)
+		, Node(InNode)
+		, Pin(InPin)
+		, BlueprintName(MoveTemp(InBlueprintName))
+		, GraphName(MoveTemp(InGraphName))
+		, NodeName(MoveTemp(InNodeName))
+		, DisplayName(MoveTemp(InDisplayName))
+		, Value(MoveTemp(InValue))
+		, Type(MoveTemp(InType))
+	{
+		ObjectBeingDebugged = BP ? BP->GetObjectBeingDebugged() : nullptr;
+		ObjectBeingDebuggedName = ObjectBeingDebugged ? FText::FromString(ObjectBeingDebugged->GetName()) : LOCTEXT("Unknown Object", "Unknown object");
+	}
+
+	FWatchRow(
+		UBlueprint* InBP,
+		const UEdGraphNode* InNode,
+		const UEdGraphPin* InPin,
+		FText InBlueprintName,
+		FText InGraphName,
+		FText InNodeName,
+		FDebugInfo Info
+	)
+		: BP(InBP)
+		, Node(InNode)
+		, Pin(InPin)
+		, BlueprintName(MoveTemp(InBlueprintName))
+		, GraphName(MoveTemp(InGraphName))
+		, NodeName(MoveTemp(InNodeName))
+		, DisplayName(MoveTemp(Info.DisplayName))
+		, Value(MoveTemp(Info.Value))
+		, Type(MoveTemp(Info.Type))
+	{
+		BlueprintName = MoveTemp(InBlueprintName);
+		ObjectBeingDebugged = BP ? BP->GetObjectBeingDebugged() : nullptr;
+		ObjectBeingDebuggedName = ObjectBeingDebugged ? FText::FromString(ObjectBeingDebugged->GetName()) : LOCTEXT("Unknown Object", "Unknown object");
+		for (const FDebugInfo& ChildInfo : Info.Children)
+		{
+			Children.Add(MakeShared<FWatchRow>(InBP, InNode, InPin, BlueprintName, GraphName, NodeName, ChildInfo));
+		}
+	}
+
+	// this can't be const because we store watches in the blueprint
+	UBlueprint* BP;
+	const UEdGraphNode* Node;
+	const UEdGraphPin* Pin;
+	// this can't be const because SelectActor takes a non-const actor
+	UObject* ObjectBeingDebugged;
+	
+	FText BlueprintName;
+	FText ObjectBeingDebuggedName;
+	FText GraphName;
+	FText NodeName;
+	FText DisplayName;
+	FText Value;
+	FText Type;
+
+	TArray<TSharedRef<FWatchRow>> Children;
+
+	// used for copying entries in the watch viewer
+	FText GetTextForEntry() const
+	{
+		FFormatNamedArguments Args;
+		Args.Add(TEXT("ObjectName"), FText::FromString(ObjectBeingDebugged ? ObjectBeingDebugged->GetName() : TEXT("")));
+		Args.Add(TEXT("BlueprintName"), BlueprintName);
+		Args.Add(TEXT("GraphName"), GraphName);
+		Args.Add(TEXT("NodeName"), NodeName);
+		Args.Add(TEXT("DisplayName"), DisplayName);
+		Args.Add(TEXT("Type"), Type);
+		Args.Add(TEXT("Value"), Value);
+		return FText::Format(LOCTEXT("WatchEntry", "{ObjectName}({BlueprintName}) {GraphName} {NodeName} {DisplayName}({Type}): {Value}"), Args);
+	}
+};
+
+DECLARE_MULTICAST_DELEGATE_OneParam(FOnDisplayedWatchWindowChanged, TArray<TSharedRef<FWatchRow>>*);
+FOnDisplayedWatchWindowChanged WatchListSubscribers;
+
+/**
+* Widget that visualizes the contents of a FWatchRow.
+*/
+class SWatchTreeWidgetItem : public SMultiColumnTableRow<TSharedRef<FWatchRow>>
+{
+public:
+
+	SLATE_BEGIN_ARGS(SWatchTreeWidgetItem)
+		: _WatchToVisualize()
+	{ }
+
+	SLATE_ARGUMENT(TSharedPtr<FWatchRow>, WatchToVisualize)
+
+	SLATE_END_ARGS()
+
+public:
+
+	/**
+	* Construct child widgets that comprise this widget.
+	*
+	* @param InArgs Declaration from which to construct this widget.
+	*/
+	void Construct(const FArguments& InArgs, class SWatchViewer* InOwner, const TSharedRef<STableViewBase>& InOwnerTableView);
+
+public:
+
+	// SMultiColumnTableRow overrides
+	virtual TSharedRef<SWidget> GenerateWidgetForColumn(const FName& ColumnName) override;
+
+protected:
+	FText GetDebuggedObjectName() const
+	{
+		return WatchRow->ObjectBeingDebuggedName;
+	}
+
+	FText GetBlueprintName() const
+	{
+		return WatchRow->BlueprintName;
+	}
+
+	FText GetGraphName() const
+	{
+		return WatchRow->GraphName;
+	}
+
+	FText GetNodeName() const
+	{
+		return WatchRow->NodeName;
+	}
+
+	FText GetVariableName() const
+	{
+		return WatchRow->DisplayName;
+	}
+
+	FText GetValue() const
+	{
+		return WatchRow->Value;
+	}
+
+	FText GetType() const
+	{
+		return WatchRow->Type;
+	}
+
+	void HandleHyperlinkDebuggedObjectNavigate() const
+	{
+		if (AActor* Actor = Cast<AActor>(WatchRow->ObjectBeingDebugged))
+		{
+			GEditor->SelectActor(Actor, true, true, true);
+		}
+	}
+
+	void HandleHyperlinkNodeNavigate() const
+	{
+		if (WatchRow.IsValid() && WatchRow->Node)
+		{
+			FKismetEditorUtilities::BringKismetToFocusAttentionOnObject(WatchRow->Node);
+		}
+	}
+
+private:
+
+	/** The info about the widget that we are visualizing. */
+	TSharedPtr<FWatchRow> WatchRow;
+
+	SWatchViewer* Owner;
+};
+
+typedef STreeView<TSharedRef<FWatchRow>> SWatchTree;
+
+class SWatchViewer : public SCompoundWidget
+{
+	friend class FBlueprintEditor;
+public:
+	SLATE_BEGIN_ARGS(SWatchViewer){}
+	SLATE_END_ARGS()
+
+	SWatchViewer()
+	{
+		FKismetDebugUtilities::WatchedPinsListChangedEvent.AddRaw(this, &SWatchViewer::HandleWatchedPinsChanged);
+	}
+	~SWatchViewer()
+	{
+		FKismetDebugUtilities::WatchedPinsListChangedEvent.RemoveAll(this);
+	}
+
+	void Construct(const FArguments& InArgs, TArray<TSharedRef<FWatchRow>>* InWatchSource);
+	TSharedRef<ITableRow> HandleGenerateRow(TSharedRef<FWatchRow> InWatchRow, const TSharedRef<STableViewBase>& OwnerTable);
+	void HandleGetChildren(TSharedRef<FWatchRow> InWatchRow, TArray<TSharedRef<FWatchRow>>& OutChildren);
+	void HandleWatchedPinsChanged(UBlueprint* BlueprintObj);
+	void UpdateWatches(TArray<TSharedRef<FWatchRow>>* WatchValues);
+	void CopySelectedRows() const;
+	void StopWatchingPin() const;
+
+	TSharedPtr<SWatchTree> WatchTreeWidget;
+	TArray<TSharedRef<FWatchRow>>* WatchSource;
+
+	TSharedPtr< FUICommandList > CommandList;
+};
+
+void SWatchViewer::Construct(const FArguments& InArgs, TArray<TSharedRef<FWatchRow>>* InWatchSource)
+{
+	CommandList = MakeShareable( new FUICommandList );
+	CommandList->MapAction( 
+		FGenericCommands::Get().Copy,
+		FExecuteAction::CreateSP( this, &SWatchViewer::CopySelectedRows ),
+		// we need to override the default 'can execute' because we want to be available during debugging:
+		FCanExecuteAction::CreateStatic( [](){ return true; } )
+	);
+
+	CommandList->MapAction(
+		FGraphEditorCommands::Get().StopWatchingPin,
+		FExecuteAction::CreateSP(this, &SWatchViewer::StopWatchingPin),
+		FCanExecuteAction::CreateStatic([]() { return true; })
+	);
+
+	WatchSource = InWatchSource;
+
+	const auto ContextMenuOpened = [](TWeakPtr<FUICommandList> InCommandList, TWeakPtr<SWatchViewer> ControlOwnerWeak) -> TSharedPtr<SWidget>
+	{
+		const bool CloseAfterSelection = true;
+		FMenuBuilder MenuBuilder( CloseAfterSelection, InCommandList.Pin() );
+		MenuBuilder.AddMenuEntry(FGraphEditorCommands::Get().StopWatchingPin);
+		MenuBuilder.AddMenuEntry(FGenericCommands::Get().Copy);
+		return MenuBuilder.MakeWidget();
+	};
+
+	const auto EmptyWarningVisibility = [](TWeakPtr<SWatchViewer> ControlOwnerWeak) -> EVisibility
+	{
+		TSharedPtr<SWatchViewer> ControlOwner = ControlOwnerWeak.Pin();
+		if (ControlOwner.IsValid() &&
+			ControlOwner->WatchSource &&
+			ControlOwner->WatchSource->Num() > 0)
+		{
+			return EVisibility::Hidden;
+		}
+		return EVisibility::Visible;
+	};
+
+	const auto WatchViewIsEnabled = [](TWeakPtr<SWatchViewer> ControlOwnerWeak) -> bool
+	{
+		TSharedPtr<SWatchViewer> ControlOwner = ControlOwnerWeak.Pin();
+		if (ControlOwner.IsValid() &&
+			ControlOwner->WatchSource &&
+			ControlOwner->WatchSource->Num() > 0)
+		{
+			return true;
+		}
+		return false;
+	};
+
+	// Cast due to TSharedFromThis inheritance issues:
+	TSharedRef<SWatchViewer> SelfTyped = StaticCastSharedRef<SWatchViewer>(AsShared());
+	TWeakPtr<SWatchViewer> SelfWeak = SelfTyped;
+	TWeakPtr<FUICommandList> CommandListWeak = CommandList;
+
+	ChildSlot
+	[
+		SNew(SBorder)
+		.Padding(4)
+		.BorderImage( FEditorStyle::GetBrush("ToolPanel.GroupBorder") )
+		[
+			SNew(SOverlay)
+			+SOverlay::Slot()
+			[
+				SAssignNew(WatchTreeWidget, SWatchTree)
+				.ItemHeight(25.0f)
+				.TreeItemsSource(WatchSource)
+				.OnGenerateRow(this, &SWatchViewer::HandleGenerateRow)
+				.OnGetChildren(this, &SWatchViewer::HandleGetChildren)
+				.OnContextMenuOpening(FOnContextMenuOpening::CreateStatic(ContextMenuOpened, CommandListWeak, SelfWeak))
+				.IsEnabled(
+					TAttribute<bool>::Create(
+						TAttribute<bool>::FGetter::CreateStatic(WatchViewIsEnabled, SelfWeak)
+					)
+				)
+				.HeaderRow
+				(
+					SNew(SHeaderRow)
+					+SHeaderRow::Column(TEXT("ObjectName"))
+					.FillWidth(.2f)
+					.VAlignHeader(VAlign_Center)
+					.HeaderContent()
+					[
+						SNew(STextBlock)
+						.Text(LOCTEXT("ObjectName", "Object Name"))
+						.ToolTipText(LOCTEXT("ObjectNameTooltip", "Name of the object being debugged"))
+					]
+					+ SHeaderRow::Column(TEXT("GraphName"))
+					.FillWidth(.2f)
+					.VAlignHeader(VAlign_Center)
+					.HeaderContent()
+					[
+						SNew(STextBlock)
+						.Text(LOCTEXT("GraphName", "Graph Name"))
+						.ToolTipText(LOCTEXT("GraphNameTooltip", "Name of the source blueprint graph for this variable"))
+					]
+					+ SHeaderRow::Column(TEXT("NodeName"))
+					.FillWidth(.3f)
+					.VAlignHeader(VAlign_Center)
+					.HeaderContent()
+					[
+						SNew(STextBlock)
+						.Text(LOCTEXT("NodeName", "Node Name"))
+						.ToolTipText(LOCTEXT("NodeNameTooltip", "Name of the source blueprint graph node for this variable"))
+					]
+					+ SHeaderRow::Column(TEXT("VariableName"))
+					.FillWidth(.3f)
+					.VAlignHeader(VAlign_Center)
+					.HeaderContent()
+					[
+						SNew(STextBlock)
+						.Text(LOCTEXT("VariableName", "Variable Name"))
+						.ToolTipText(LOCTEXT("VariabelNameTooltip", "Name of the variable"))
+					]
+					+ SHeaderRow::Column(TEXT("Value"))
+					.FillWidth(.8f)
+					.VAlignHeader(VAlign_Center)
+					.HeaderContent()
+					[
+						SNew(STextBlock)
+						.Text(LOCTEXT("Value", "Value"))
+						.ToolTipText(LOCTEXT("ValueTooltip", "Current value of this variable"))
+					]
+				)
+			]
+			+SOverlay::Slot()
+			.Padding(32.f)
+			[
+				SNew(STextBlock)
+				.Text(
+					LOCTEXT("NoWatches", "No watches to display")
+				)
+				.Justification(ETextJustify::Center)
+				.Visibility(
+					TAttribute<EVisibility>::Create(
+						TAttribute<EVisibility>::FGetter::CreateStatic(EmptyWarningVisibility, SelfWeak)
+					)
+				)
+			]
+		]
+	];
+
+	WatchListSubscribers.AddSP(StaticCastSharedRef<SWatchViewer>(AsShared()), &SWatchViewer::UpdateWatches);
+}
+
+TSharedRef<ITableRow> SWatchViewer::HandleGenerateRow(TSharedRef<FWatchRow> InWatchRow, const TSharedRef<STableViewBase>& OwnerTable)
+{
+	return SNew(SWatchTreeWidgetItem, this, OwnerTable)
+		.WatchToVisualize(InWatchRow);
+}
+
+void SWatchViewer::HandleGetChildren(TSharedRef<FWatchRow> InWatchRow, TArray<TSharedRef<FWatchRow>>& OutChildren)
+{
+	OutChildren = InWatchRow->Children;
+}
+
+void SWatchViewer::HandleWatchedPinsChanged(UBlueprint* BlueprintObj)
+{
+	WatchViewer::UpdateDisplayedWatches(BlueprintObj);
+}
+
+void SWatchViewer::UpdateWatches(TArray<TSharedRef<FWatchRow>>* Watches)
+{
+	WatchTreeWidget->RequestTreeRefresh();
+}
+
+void SWatchViewer::CopySelectedRows() const
+{
+	FString StringToCopy;
+
+	// We want to copy in the order displayed, not the order selected, so iterate the list and build up the string:
+	if (WatchSource)
+	{
+		for (const TSharedRef<FWatchRow>& Item : *WatchSource)
+		{
+			if (WatchTreeWidget->IsItemSelected(Item))
+			{
+				StringToCopy.Append(Item->GetTextForEntry().ToString());
+				StringToCopy.Append(TEXT("\r\n"));
+			}
+		}
+	}
+
+	if (!StringToCopy.IsEmpty())
+	{
+		FPlatformApplicationMisc::ClipboardCopy(*StringToCopy);
+	}
+}
+
+void SWatchTreeWidgetItem::Construct(const FArguments& InArgs, SWatchViewer* InOwner, const TSharedRef<STableViewBase>& InOwnerTableView)
+{
+	this->WatchRow = InArgs._WatchToVisualize;
+	Owner = InOwner;
+
+	SMultiColumnTableRow< TSharedRef<FWatchRow> >::Construct(SMultiColumnTableRow< TSharedRef<FWatchRow> >::FArguments().Padding(1), InOwnerTableView);
+}
+
+void SWatchViewer::StopWatchingPin() const
+{
+	TArray<TSharedRef<FWatchRow>> SelectedRows = WatchTreeWidget->GetSelectedItems();
+	for (TSharedRef<FWatchRow>& Row : SelectedRows)
+	{
+		FKismetDebugUtilities::TogglePinWatch(Row->BP, Row->Pin);
+	}
+}
+
+TSharedRef<SWidget> SWatchTreeWidgetItem::GenerateWidgetForColumn(const FName& ColumnName)
+{
+	const static FName NAME_ObjectName(TEXT("ObjectName"));
+	const static FName NAME_GraphName(TEXT("GraphName"));
+	const static FName NAME_NodeName(TEXT("NodeName"));
+	const static FName NAME_VariableName(TEXT("VariableName"));
+	const static FName NAME_Value(TEXT("Value"));
+
+	if (ColumnName == NAME_ObjectName)
+	{
+		return SNew(SBox)
+			.HAlign(HAlign_Left)
+			.VAlign(VAlign_Center)
+			.Padding(FMargin(2.0f, 0.0f))
+			[
+				SNew(SHyperlink)
+				.Text(this, &SWatchTreeWidgetItem::GetDebuggedObjectName)
+				.ToolTipText(this, &SWatchTreeWidgetItem::GetBlueprintName)
+				.OnNavigate(this, &SWatchTreeWidgetItem::HandleHyperlinkDebuggedObjectNavigate)
+			];
+	}
+	else if (ColumnName == NAME_GraphName)
+	{
+		return SNew(SBox)
+			.HAlign(HAlign_Left)
+			.VAlign(VAlign_Center)
+			.Padding(FMargin(2.0f, 0.0f))
+			[
+				SNew(STextBlock)
+				.Text(this, &SWatchTreeWidgetItem::GetGraphName)
+			];
+	}
+	else if (ColumnName == NAME_NodeName)
+	{
+		FString Comment; 
+		if (WatchRow->Node->NodeComment.Len() > 0)
+		{
+			Comment = TEXT("\n\n");
+			Comment.Append(WatchRow->Node->NodeComment);
+		}
+		FText TooltipText = FText::Format(LOCTEXT("NodeTooltip", "Find the {0} node in the blueprint graph.{1}"), GetNodeName(), FText::FromString(Comment));
+		return SNew(SBox)
+			.HAlign(HAlign_Left)
+			.VAlign(VAlign_Center)
+			.Padding(FMargin(2.0f, 0.0f))
+			[
+				SNew(SHyperlink)
+				.Text(this, &SWatchTreeWidgetItem::GetNodeName)
+				.ToolTipText(TooltipText)
+				.OnNavigate(this, &SWatchTreeWidgetItem::HandleHyperlinkNodeNavigate)
+			];
+	}
+	else if (ColumnName == NAME_VariableName)
+	{
+		return SNew(SHorizontalBox)
+			+ SHorizontalBox::Slot()
+			.AutoWidth()
+			.VAlign(VAlign_Center)
+			[
+				SNew(SExpanderArrow, SharedThis(this))
+			]
+		+ SHorizontalBox::Slot()
+			.AutoWidth()
+			.Padding(2.0f, 0.0f)
+			.VAlign(VAlign_Center)
+			[
+				SNew(STextBlock)
+				.Text(this, &SWatchTreeWidgetItem::GetVariableName)
+				.ToolTipText(this, &SWatchTreeWidgetItem::GetType)
+			];
+	}
+	else if (ColumnName == NAME_Value)
+	{
+		return SNew(SBox)
+			.HAlign(HAlign_Left)
+			.VAlign(VAlign_Center)
+			.Padding(FMargin(2.0f, 0.0f))
+			[
+				SNew(STextBlock)
+				.Text(this, &SWatchTreeWidgetItem::GetValue)
+			];
+	}
+	else
+	{
+		return SNullWidget::NullWidget;
+	}
+}
+
+// Proxy array of the watches. This allows us to manually refresh UI state when changes are made:
+TArray<TSharedRef<FWatchRow>> Private_WatchSource;
+
+void WatchViewer::UpdateDisplayedWatches(UBlueprint* BlueprintObj)
+{
+	if (!ensure(BlueprintObj))
+	{
+		return;
+	}
+	Private_WatchSource.Reset();
+
+	FText BlueprintName = FText::FromString(BlueprintObj->GetName());
+
+	for (const FEdGraphPinReference& PinRef : BlueprintObj->WatchedPins)
+	{
+		UEdGraphPin* Pin = PinRef.Get();
+
+		FText GraphName = FText::FromString(Pin->GetOwningNode()->GetGraph()->GetName());
+		FText NodeName = Pin->GetOwningNode()->GetNodeTitle(ENodeTitleType::ListView);
+
+		FDebugInfo DebugInfo;
+		const FKismetDebugUtilities::EWatchTextResult WatchStatus = FKismetDebugUtilities::GetDebugInfo(DebugInfo, BlueprintObj, BlueprintObj->GetObjectBeingDebugged(), Pin);
+
+		if (WatchStatus != FKismetDebugUtilities::EWTR_Valid)
+		{
+			const UEdGraphSchema* Schema = Pin->GetOwningNode()->GetSchema();
+			DebugInfo.DisplayName = Schema->GetPinDisplayName(Pin);
+			DebugInfo.Type = UEdGraphSchema_K2::TypeToText(Pin->PinType);
+
+			switch (WatchStatus)
+			{
+			case FKismetDebugUtilities::EWTR_NotInScope:
+				DebugInfo.Value = LOCTEXT("NotInScope", "(not in scope)");
+				break;
+
+			case FKismetDebugUtilities::EWTR_NoProperty:
+				DebugInfo.Value = LOCTEXT("NoDebugData", "(no debug data)");
+				break;
+
+			case FKismetDebugUtilities::EWTR_NoDebugObject:
+				DebugInfo.Value = LOCTEXT("NoDebugObject", "(no debug object)");
+				break;
+
+			default:
+				// do nothing
+				break;
+			}
+		}
+
+		Private_WatchSource.Add(
+			MakeShared<FWatchRow>(
+				BlueprintObj,
+				Pin->GetOwningNode(),
+				Pin,
+				BlueprintName,
+				GraphName,
+				NodeName,
+				DebugInfo
+				)
+		);
+	}
+
+	// Notify subscribers:
+	WatchListSubscribers.Broadcast(&Private_WatchSource);
+}
+
+void WatchViewer::RegisterTabSpawner()
+{
+	const auto SpawnWatchViewTab = []( const FSpawnTabArgs& Args )
+	{
+		return SNew(SDockTab)
+			.TabRole( ETabRole::NomadTab )
+			.Label( LOCTEXT("TabTitle", "Watch Window") )
+			[
+				SNew(SBorder)
+				.BorderImage( FEditorStyle::GetBrush("Docking.Tab.ContentAreaBrush") )
+				[
+					SNew(SWatchViewer, &Private_WatchSource)
+				]
+			];
+	};
+	
+	const IWorkspaceMenuStructure& MenuStructure = WorkspaceMenu::GetMenuStructure();
+
+	const FName TabName = TEXT("WatchViewer");
+	FGlobalTabmanager::Get()->RegisterNomadTabSpawner( TabName, FOnSpawnTab::CreateStatic(SpawnWatchViewTab) )
+		.SetDisplayName( LOCTEXT("TabTitle", "Watch Window") )
+		.SetTooltipText( LOCTEXT("TooltipText", "Open the watch window tab.") )
+		.SetGroup( MenuStructure.GetDeveloperToolsDebugCategory() );
+}
+
+#undef LOCTEXT_NAMESPACE
diff --git a/Engine/Source/Editor/Kismet/Public/BlueprintEditor.h b/Engine/Source/Editor/Kismet/Public/BlueprintEditor.h
--- a/Engine/Source/Editor/Kismet/Public/BlueprintEditor.h
+++ b/Engine/Source/Editor/Kismet/Public/BlueprintEditor.h
@@ -31,6 +31,7 @@ class IMessageLogListing;
 class INameValidatorInterface;
 class ISCSEditorCustomization;
 class SBlueprintPalette;
+class SBlueprintBookmarks;
 class SFindInBlueprints;
 class SKismetDebuggingView;
 class SKismetInspector;
@@ -225,6 +226,7 @@ public:
 	TSharedRef<class SKismetInspector> GetDefaultEditor() const { return DefaultEditor.ToSharedRef(); }
 	TSharedRef<class SKismetDebuggingView> GetDebuggingView() const { return DebuggingView.ToSharedRef(); }
 	TSharedRef<class SBlueprintPalette> GetPalette() const { return Palette.ToSharedRef(); }
+	TSharedRef<class SBlueprintBookmarks> GetBookmarksWidget() const { return BookmarksWidget.ToSharedRef(); }
 	TSharedRef<class SWidget> GetCompilerResults() const { return CompilerResults.ToSharedRef(); }
 	TSharedRef<class SFindInBlueprints> GetFindResults() const { return FindResults.ToSharedRef(); }
 
@@ -295,6 +297,18 @@ public:
 	// This is used to delay it by one frame when triggered by a tab being closed, so it can finish closing before remembering the new state
 	void RequestSaveEditedObjectState();
 
+	/** Get the visible bounds of the given graph node */
+	void GetBoundsForNode(const UEdGraphNode* InNode, class FSlateRect& OutRect, float InPadding) const;
+
+	/** Gets the focused graph current view bookmark ID */
+	void GetViewBookmark(FGuid& BookmarkId);
+
+	/** Gets the focused graph view location/zoom amount */
+	void GetViewLocation(FVector2D& Location, float& ZoomAmount);
+
+	/** Sets the focused graph view location/zoom amount */
+	void SetViewLocation(const FVector2D& Location, float ZoomAmount, const FGuid& BookmarkId = FGuid());
+
 	/** Returns whether a graph is editable or not */
 	virtual bool IsEditable(UEdGraph* InGraph) const;
 	/** Determines if the graph is ReadOnly, this differs from editable in that it is never expected to be edited and is in a read-only state */
@@ -558,6 +572,15 @@ public:
 	/** Handle when the debug object is changed in the UI */
 	virtual void HandleSetObjectBeingDebugged(UObject* InObject) {}
 
+	/** Adds a new bookmark node to the Blueprint that's being edited. */
+	FBPEditorBookmarkNode* AddBookmark(const FText& DisplayName, const FEditedDocumentInfo& BookmarkInfo, bool bSharedBookmark = false);
+
+	/** Renames the bookmark node with the given ID. */
+	void RenameBookmark(const FGuid& BookmarkNodeId, const FText& NewName);
+
+	/** Removes the bookmark node with the given ID. */
+	void RemoveBookmark(const FGuid& BookmarkNodeId, bool bRefreshUI = true);
+
 protected:
 	virtual void AppendExtraCompilerResults(TSharedPtr<class IMessageLogListing> ResultsListing);
 
@@ -1063,6 +1086,9 @@ protected:
 	/** Palette of all classes with funcs/vars */
 	TSharedPtr<class SBlueprintPalette> Palette;
 
+	/** Blueprint bookmark editing widget */
+	TSharedPtr<class SBlueprintBookmarks> BookmarksWidget;
+
 	/** All of this blueprints' functions and variables */
 	TSharedPtr<class SMyBlueprint> MyBlueprintWidget;
 	
diff --git a/Engine/Source/Editor/Kismet/Public/BlueprintEditorTabs.h b/Engine/Source/Editor/Kismet/Public/BlueprintEditorTabs.h
--- a/Engine/Source/Editor/Kismet/Public/BlueprintEditorTabs.h
+++ b/Engine/Source/Editor/Kismet/Public/BlueprintEditorTabs.h
@@ -11,6 +11,7 @@ struct KISMET_API FBlueprintEditorTabs
 	static const FName DefaultEditorID;
 	static const FName DebugID;
 	static const FName PaletteID;
+	static const FName BookmarksID;
 	static const FName CompilerResultsID;
 	static const FName FindResultsID;
 	static const FName ConstructionScriptEditorID;
diff --git a/Engine/Source/Editor/Kismet/Public/WatchPointViewer.h b/Engine/Source/Editor/Kismet/Public/WatchPointViewer.h
new file mode 100644
--- /dev/null
+++ b/Engine/Source/Editor/Kismet/Public/WatchPointViewer.h
@@ -0,0 +1,13 @@
+// Copyright 1998-2017 Epic Games, Inc. All Rights Reserved.
+
+#pragma once
+
+#include "Widgets/SCompoundWidget.h"
+#include "WorkflowOrientedApp/WorkflowTabFactory.h"
+
+
+namespace WatchViewer
+{
+	void KISMET_API UpdateDisplayedWatches(class UBlueprint* BlueprintObj);
+	void RegisterTabSpawner();
+}
diff --git a/Engine/Source/Editor/KismetCompiler/KismetCompiler.Build.cs b/Engine/Source/Editor/KismetCompiler/KismetCompiler.Build.cs
--- a/Engine/Source/Editor/KismetCompiler/KismetCompiler.Build.cs
+++ b/Engine/Source/Editor/KismetCompiler/KismetCompiler.Build.cs
@@ -15,7 +15,6 @@ public class KismetCompiler : ModuleRules
 				"UnrealEd",
 				"MovieScene",
 				"MovieSceneTools",
-				"FunctionalTesting",
 				"BlueprintGraph",
 				"AnimGraph",
                 "MessageLog",
diff --git a/Engine/Source/Editor/KismetCompiler/Private/KismetCompiler.cpp b/Engine/Source/Editor/KismetCompiler/Private/KismetCompiler.cpp
--- a/Engine/Source/Editor/KismetCompiler/Private/KismetCompiler.cpp
+++ b/Engine/Source/Editor/KismetCompiler/Private/KismetCompiler.cpp
@@ -690,7 +690,16 @@ void FKismetCompilerContext::CreateClassVariablesFromBlueprint()
 		const FString TimelineVariableName = UTimelineTemplate::TimelineTemplateNameToVariableName(Timeline->GetFName());
 		if (UProperty* TimelineProperty = CreateVariable(*TimelineVariableName, TimelinePinType))
 		{
-			TimelineProperty->SetMetaData( TEXT("Category"), *Blueprint->GetName() );
+			FString CategoryName;
+			if (Timeline->FindMetaDataEntryIndexForKey(TEXT("Category")) != INDEX_NONE)
+			{
+				CategoryName = Timeline->GetMetaData(TEXT("Category"));
+			}
+			else
+			{
+				CategoryName = Blueprint->GetName();
+			}
+			TimelineProperty->SetMetaData(TEXT("Category"), *CategoryName);
 			TimelineProperty->SetPropertyFlags(CPF_BlueprintVisible);
 
 			TimelineToMemberVariableMap.Add(Timeline, TimelineProperty);
diff --git a/Engine/Source/Editor/LandscapeEditor/Private/LandscapeEdModeMirrorTool.cpp b/Engine/Source/Editor/LandscapeEditor/Private/LandscapeEdModeMirrorTool.cpp
--- a/Engine/Source/Editor/LandscapeEditor/Private/LandscapeEdModeMirrorTool.cpp
+++ b/Engine/Source/Editor/LandscapeEditor/Private/LandscapeEdModeMirrorTool.cpp
@@ -3,7 +3,7 @@
 #include "CoreMinimal.h"
 #include "InputCoreTypes.h"
 #include "Materials/MaterialInterface.h"
-#include "AI/Navigation/NavigationSystem.h"
+#include "AI/NavigationSystemBase.h"
 #include "Materials/MaterialInstanceDynamic.h"
 #include "UnrealWidget.h"
 #include "EditorModeManager.h"
@@ -19,7 +19,6 @@
 #include "LandscapeRender.h"
 #include "LandscapeHeightfieldCollisionComponent.h"
 //#include "LandscapeDataAccess.h"
-//#include "AI/Navigation/NavigationSystem.h"
 
 #define LOCTEXT_NAMESPACE "Landscape"
 
@@ -640,8 +639,6 @@ public:
 		TSet<ULandscapeComponent*> Components;
 		if (LandscapeEdit.GetComponentsInRegion(DestMinX, DestMinY, DestMaxX, DestMaxY, &Components) && Components.Num() > 0)
 		{
-			UNavigationSystem* NavSys = UNavigationSystem::GetCurrent(*begin(Components));
-
 			for (ULandscapeComponent* Component : Components)
 			{
 				// Recreate collision for modified components and update the navmesh
@@ -649,10 +646,7 @@ public:
 				if (CollisionComponent)
 				{
 					CollisionComponent->RecreateCollision();
-					if (NavSys)
-					{
-						NavSys->UpdateComponentInNavOctree(*CollisionComponent);
-					}
+					FNavigationSystem::UpdateComponentData(*CollisionComponent);
 				}
 			}
 
diff --git a/Engine/Source/Editor/LandscapeEditor/Private/LandscapeEdModeRampTool.cpp b/Engine/Source/Editor/LandscapeEditor/Private/LandscapeEdModeRampTool.cpp
--- a/Engine/Source/Editor/LandscapeEditor/Private/LandscapeEdModeRampTool.cpp
+++ b/Engine/Source/Editor/LandscapeEditor/Private/LandscapeEdModeRampTool.cpp
@@ -3,7 +3,7 @@
 #include "CoreMinimal.h"
 #include "InputCoreTypes.h"
 #include "HitProxies.h"
-#include "AI/Navigation/NavigationSystem.h"
+#include "AI/NavigationSystemBase.h"
 #include "Editor/UnrealEdEngine.h"
 #include "Engine/Texture2D.h"
 #include "UnrealWidget.h"
@@ -580,11 +580,7 @@ public:
 					if (CollisionComponent)
 					{
 						CollisionComponent->RecreateCollision();
-						UNavigationSystem* NavSys = UNavigationSystem::GetCurrent(Component);
-						if (NavSys)
-						{
-							NavSys->UpdateComponentInNavOctree(*CollisionComponent);
-						}
+						FNavigationSystem::UpdateComponentData(*CollisionComponent);
 					}
 				}
 			}
diff --git a/Engine/Source/Editor/LandscapeEditor/Private/LandscapeEdModeTools.h b/Engine/Source/Editor/LandscapeEditor/Private/LandscapeEdModeTools.h
--- a/Engine/Source/Editor/LandscapeEditor/Private/LandscapeEdModeTools.h
+++ b/Engine/Source/Editor/LandscapeEditor/Private/LandscapeEdModeTools.h
@@ -14,7 +14,7 @@
 #include "LandscapeHeightfieldCollisionComponent.h"
 #include "InstancedFoliageActor.h"
 #include "VREditorInteractor.h"
-#include "AI/Navigation/NavigationSystem.h"
+#include "AI/NavigationSystemBase.h"
 
 // VR Editor
 
@@ -710,7 +710,7 @@ struct FHeightmapAccessor
 			if (CollisionComponent)
 			{
 				CollisionComponent->RecreateCollision();
-				UNavigationSystem::UpdateComponentInNavOctree(*CollisionComponent);
+				FNavigationSystem::UpdateComponentData(*CollisionComponent);
 			}
 		}
 	}
@@ -935,11 +935,7 @@ struct FAlphamapAccessor
 				// We need to trigger navigation mesh build, in case user have painted holes on a landscape
 				if (LayerInfo == ALandscapeProxy::VisibilityLayer)
 				{
-					UNavigationSystem* NavSys = UNavigationSystem::GetCurrent(Component);
-					if (NavSys)
-					{
-						NavSys->UpdateComponentInNavOctree(*CollisionComponent);
-					}
+					FNavigationSystem::UpdateComponentData(*CollisionComponent);
 				}
 			}
 		}
@@ -1041,11 +1037,7 @@ struct FFullWeightmapAccessor
 				// We need to trigger navigation mesh build, in case user have painted holes on a landscape
 				if (LandscapeInfo->GetLayerInfoIndex(ALandscapeProxy::VisibilityLayer) != INDEX_NONE)
 				{
-					UNavigationSystem* NavSys = UNavigationSystem::GetCurrent(Component);
-					if (NavSys)
-					{
-						NavSys->UpdateComponentInNavOctree(*CollisionComponent);
-					}
+					FNavigationSystem::UpdateComponentData(*CollisionComponent);
 				}
 			}
 		}
diff --git a/Engine/Source/Editor/LevelEditor/LevelEditor.Build.cs b/Engine/Source/Editor/LevelEditor/LevelEditor.Build.cs
--- a/Engine/Source/Editor/LevelEditor/LevelEditor.Build.cs
+++ b/Engine/Source/Editor/LevelEditor/LevelEditor.Build.cs
@@ -65,7 +65,6 @@ public class LevelEditor : ModuleRules
 				"HierarchicalLODUtilities",
 				"MaterialShaderQualitySettings",
 				"PixelInspectorModule",
-				"FunctionalTesting",
 				"CommonMenuExtensions"
 			}
 		);
diff --git a/Engine/Source/Editor/ProjectSettingsViewer/Private/ProjectSettingsViewerModule.cpp b/Engine/Source/Editor/ProjectSettingsViewer/Private/ProjectSettingsViewerModule.cpp
--- a/Engine/Source/Editor/ProjectSettingsViewer/Private/ProjectSettingsViewerModule.cpp
+++ b/Engine/Source/Editor/ProjectSettingsViewer/Private/ProjectSettingsViewerModule.cpp
@@ -11,7 +11,8 @@
 #include "Widgets/SNullWidget.h"
 #include "Widgets/DeclarativeSyntaxSupport.h"
 #include "Engine/GameViewportClient.h"
-#include "AI/Navigation/NavigationSystem.h"
+#include "AI/NavigationSystemBase.h"
+#include "NavigationSystem.h"
 #include "Engine/Engine.h"
 #include "GameFramework/InputSettings.h"
 #include "Textures/SlateIcon.h"
@@ -28,7 +29,7 @@
 #include "ISettingsViewer.h"
 #include "Widgets/Docking/SDockTab.h"
 
-#include "AI/Navigation/RecastNavMesh.h"
+#include "NavMesh/RecastNavMesh.h"
 #include "Navigation/CrowdManager.h"
 
 
@@ -140,9 +141,9 @@ protected:
 		);
 
 		// navigation system's class can be game specific so we need to find appropriate CDO
-		UNavigationSystem* NavSysCDO = (*GEngine->NavigationSystemClass != nullptr)
-			? GetMutableDefault<UNavigationSystem>(GEngine->NavigationSystemClass)
-			: GetMutableDefault<UNavigationSystem>();
+		UNavigationSystemBase* NavSysCDO = (*GEngine->NavigationSystemClass != nullptr)
+			? GetMutableDefault<UNavigationSystemBase>(GEngine->NavigationSystemClass)
+			: GetMutableDefault<UNavigationSystemBase>();
 		SettingsModule.RegisterSettings("Project", "Engine", "NavigationSystem",
 			LOCTEXT("NavigationSystemSettingsName", "Navigation System"),
 			LOCTEXT("NavigationSystemSettingsDescription", "Settings for the navigation system."),
diff --git a/Engine/Source/Editor/ProjectSettingsViewer/ProjectSettingsViewer.Build.cs b/Engine/Source/Editor/ProjectSettingsViewer/ProjectSettingsViewer.Build.cs
--- a/Engine/Source/Editor/ProjectSettingsViewer/ProjectSettingsViewer.Build.cs
+++ b/Engine/Source/Editor/ProjectSettingsViewer/ProjectSettingsViewer.Build.cs
@@ -24,9 +24,10 @@ namespace UnrealBuildTool.Rules
 					"SlateCore",
 					"UnrealEd",
 					"MoviePlayer",
-                    "AIModule",
+					"NavigationSystem",
+					"AIModule",
 					"ProjectTargetPlatformEditor",
-                    "EditorStyle",
+					"EditorStyle",
 				}
 			);
 
diff --git a/Engine/Source/Editor/PropertyEditor/Private/Presentation/PropertyEditor/PropertyEditor.cpp b/Engine/Source/Editor/PropertyEditor/Private/Presentation/PropertyEditor/PropertyEditor.cpp
--- a/Engine/Source/Editor/PropertyEditor/Private/Presentation/PropertyEditor/PropertyEditor.cpp
+++ b/Engine/Source/Editor/PropertyEditor/Private/Presentation/PropertyEditor/PropertyEditor.cpp
@@ -456,18 +456,6 @@ void FPropertyEditor::SetEditConditionState( bool bShouldEnable )
 	PropertyNode->NotifyPostChange( ChangeEvent, PropertyUtilities->GetNotifyHook() );
 }
 
-void FPropertyEditor::ResetToDefault()
-{
-	// This action must be deferred until next tick so that we avoid accessing invalid data before we have a chance to tick
-	PropertyUtilities->EnqueueDeferredAction( FSimpleDelegate::CreateSP( this, &FPropertyEditor::OnResetToDefault ) );
-}
-
-void FPropertyEditor::CustomResetToDefault(const FResetToDefaultOverride& OnCustomResetToDefault)
-{
-	// This action must be deferred until next tick so that we avoid accessing invalid data before we have a chance to tick
-	PropertyUtilities->EnqueueDeferredAction(FSimpleDelegate::CreateLambda([this, OnCustomResetToDefault](){ this->OnCustomResetToDefault(OnCustomResetToDefault); }));
-}
-
 void FPropertyEditor::OnGetClassesForAssetPicker( TArray<const UClass*>& OutClasses )
 {
 	UProperty* NodeProperty = GetPropertyNode()->GetProperty();
@@ -513,25 +501,6 @@ void FPropertyEditor::OnGetActorFiltersForSceneOutliner( TSharedPtr<SceneOutline
 	OutFilters->AddFilterPredicate( SceneOutliner::FActorFilterPredicate::CreateStatic( &Local::IsFilteredActor, AsShared() ) );
 }
 
-void FPropertyEditor::OnResetToDefault()
-{
-	PropertyNode->ResetToDefault( PropertyUtilities->GetNotifyHook() );
-}
-
-void FPropertyEditor::OnCustomResetToDefault(const FResetToDefaultOverride& OnCustomResetToDefault)
-{
-	if (OnCustomResetToDefault.OnResetToDefaultClicked().IsBound())
-	{
-		PropertyNode->NotifyPreChange( PropertyNode->GetProperty(), PropertyUtilities->GetNotifyHook() );
-
-		OnCustomResetToDefault.OnResetToDefaultClicked().Execute(GetPropertyHandle());
-
-		// Call PostEditchange on all the objects
-		FPropertyChangedEvent ChangeEvent( PropertyNode->GetProperty() );
-		PropertyNode->NotifyPostChange( ChangeEvent, PropertyUtilities->GetNotifyHook() );
-	}
-}
-
 bool FPropertyEditor::IsPropertyEditingEnabled() const
 {
 	return ( PropertyUtilities->IsPropertyEditingEnabled() ) && 
@@ -563,27 +532,6 @@ bool FPropertyEditor::SupportsEditConditionToggle() const
 	return SupportsEditConditionToggle( PropertyNode->GetProperty() );
 }
 
-bool FPropertyEditor::IsResetToDefaultAvailable() const
-{
-	UProperty* Property = PropertyNode->GetProperty();
-
-	// Should not be able to reset fixed size arrays
-	const bool bFixedSized = Property && Property->PropertyFlags & CPF_EditFixedSize;
-	const bool bCanResetToDefault = !(Property && Property->PropertyFlags & CPF_Config);
-
-	return Property && bCanResetToDefault && !bFixedSized && PropertyHandle->DiffersFromDefault();
-}
-
-bool FPropertyEditor::ValueDiffersFromDefault() const
-{
-	return PropertyHandle->DiffersFromDefault();
-}
-
-FText FPropertyEditor::GetResetToDefaultLabel() const
-{
-	return PropertyNode->GetResetToDefaultLabel();
-}
-
 void FPropertyEditor::AddPropertyEditorChild( const TSharedRef<FPropertyEditor>& Child )
 {
 	ChildPropertyEditors.Add( Child );
