Copying //UE4/Dev-Framework to //UE4/Dev-Main (Source: //UE4/Dev-Framework @ 4006870)

diff --git a/Engine/Source/Editor/ClassViewer/Private/UnloadedBlueprintData.cpp b/Engine/Source/Editor/ClassViewer/Private/UnloadedBlueprintData.cpp
--- a/Engine/Source/Editor/ClassViewer/Private/UnloadedBlueprintData.cpp
+++ b/Engine/Source/Editor/ClassViewer/Private/UnloadedBlueprintData.cpp
@@ -31,9 +31,9 @@ bool FUnloadedBlueprintData::IsChildOf(const UClass* InClass) const
 	TSharedPtr< FClassViewerNode > CurrentNode = ClassViewerNode.Pin()->ParentNode.Pin();
 
 	// Keep going through parents till you find an invalid.
-	while(CurrentNode.IsValid())
+	while (CurrentNode.IsValid())
 	{
-		if(*CurrentNode->GetClassName() == InClass->GetName())
+		if (CurrentNode->Class.Get() == InClass)
 		{
 			return true;
 		}
diff --git a/Engine/Source/Editor/ContentBrowser/Private/FrontendFilters.cpp b/Engine/Source/Editor/ContentBrowser/Private/FrontendFilters.cpp
--- a/Engine/Source/Editor/ContentBrowser/Private/FrontendFilters.cpp
+++ b/Engine/Source/Editor/ContentBrowser/Private/FrontendFilters.cpp
@@ -710,7 +710,7 @@ void FFrontendFilter_Modified::OnPackageDirtyStateUpdated(UPackage* Package)
 
 bool FFrontendFilter_ReplicatedBlueprint::PassesFilter(FAssetFilterType InItem) const
 {
-	const int32 NumReplicatedProperties = InItem.GetTagValueRef<int32>("NumReplicatedProperties");
+	const int32 NumReplicatedProperties = InItem.GetTagValueRef<int32>(FBlueprintTags::NumReplicatedProperties);
 	return NumReplicatedProperties > 0;
 }
 
diff --git a/Engine/Source/Editor/FoliageEdit/Private/FoliageTypeObjectCustomization.cpp b/Engine/Source/Editor/FoliageEdit/Private/FoliageTypeObjectCustomization.cpp
--- a/Engine/Source/Editor/FoliageEdit/Private/FoliageTypeObjectCustomization.cpp
+++ b/Engine/Source/Editor/FoliageEdit/Private/FoliageTypeObjectCustomization.cpp
@@ -42,7 +42,7 @@ void FFoliageTypeObjectCustomization::CustomizeChildren(TSharedRef<IPropertyHand
 bool FFoliageTypeObjectCustomization::OnShouldFilterAsset(const FAssetData& AssetData) const
 {
 	// If the asset is a BP class that doesn't inherit from UFoliageType, hide it
-	const FString ParentClassName = AssetData.GetTagValueRef<FString>("ParentClass");
+	const FString ParentClassName = AssetData.GetTagValueRef<FString>(FBlueprintTags::ParentClassPath);
 	if (!ParentClassName.IsEmpty() && !ParentClassName.Contains(TEXT("FoliageType_InstancedStaticMesh")))
 	{
 		return true;
diff --git a/Engine/Source/Editor/GraphEditor/Private/KismetPins/SGraphPinClass.cpp b/Engine/Source/Editor/GraphEditor/Private/KismetPins/SGraphPinClass.cpp
--- a/Engine/Source/Editor/GraphEditor/Private/KismetPins/SGraphPinClass.cpp
+++ b/Engine/Source/Editor/GraphEditor/Private/KismetPins/SGraphPinClass.cpp
@@ -125,7 +125,13 @@ FOnClicked SGraphPinClass::GetOnUseButtonDelegate()
 
 void SGraphPinClass::OnPickedNewClass(UClass* ChosenClass)
 {
-	if(GraphPinObj->DefaultObject != ChosenClass)
+	FString NewPath;
+	if (ChosenClass)
+	{
+		NewPath = ChosenClass->GetPathName();
+	}
+
+	if(GraphPinObj->GetDefaultAsString() != NewPath)
 	{
 		const FScopedTransaction Transaction( NSLOCTEXT("GraphEditor", "ChangeClassPinValue", "Change Class Pin Value" ) );
 		GraphPinObj->Modify();
diff --git a/Engine/Source/Editor/IntroTutorials/Private/STutorialsBrowser.cpp b/Engine/Source/Editor/IntroTutorials/Private/STutorialsBrowser.cpp
--- a/Engine/Source/Editor/IntroTutorials/Private/STutorialsBrowser.cpp
+++ b/Engine/Source/Editor/IntroTutorials/Private/STutorialsBrowser.cpp
@@ -751,8 +751,8 @@ void STutorialsBrowser::RebuildTutorials(TSharedPtr<FTutorialListEntry_Category>
 	FARFilter Filter;
 	Filter.ClassNames.Add(UBlueprint::StaticClass()->GetFName());
 	Filter.bRecursiveClasses = true;
-	Filter.TagsAndValues.Add(TEXT("NativeParentClass"), FString::Printf(TEXT("%s'%s'"), *UClass::StaticClass()->GetName(), *UEditorTutorial::StaticClass()->GetPathName()));
-	Filter.TagsAndValues.Add(TEXT("ParentClass"), FString::Printf(TEXT("%s'%s'"), *UClass::StaticClass()->GetName(), *UEditorTutorial::StaticClass()->GetPathName()));
+	Filter.TagsAndValues.Add(FBlueprintTags::NativeParentClassPath, FString::Printf(TEXT("%s'%s'"), *UClass::StaticClass()->GetName(), *UEditorTutorial::StaticClass()->GetPathName()));
+	Filter.TagsAndValues.Add(FBlueprintTags::ParentClassPath, FString::Printf(TEXT("%s'%s'"), *UClass::StaticClass()->GetName(), *UEditorTutorial::StaticClass()->GetPathName()));
 
 	TArray<FAssetData> AssetData;
 	AssetRegistry.Get().GetAssets(Filter, AssetData);
diff --git a/Engine/Source/Editor/Kismet/Private/BlueprintCompilationManager.cpp b/Engine/Source/Editor/Kismet/Private/BlueprintCompilationManager.cpp
--- a/Engine/Source/Editor/Kismet/Private/BlueprintCompilationManager.cpp
+++ b/Engine/Source/Editor/Kismet/Private/BlueprintCompilationManager.cpp
@@ -789,18 +789,6 @@ void FBlueprintCompilationManagerImpl::FlushCompilationQueueImpl(TArray<UObject*
 			if(BP->GeneratedClass && BP->GeneratedClass->GetSuperClass()->HasAnyClassFlags(CLASS_NewerVersionExists))
 			{
 				BP->GeneratedClass->SetSuperStruct(BP->GeneratedClass->GetSuperClass()->GetAuthoritativeClass());
-
-				if(CompilerData.ShouldResetClassMembers())
-				{
-					BP->GeneratedClass->Children = NULL;
-					BP->GeneratedClass->Script.Empty();
-					BP->GeneratedClass->MinAlignment = 0;
-					BP->GeneratedClass->RefLink = NULL;
-					BP->GeneratedClass->PropertyLink = NULL;
-					BP->GeneratedClass->DestructorLink = NULL;
-					BP->GeneratedClass->ScriptObjectReferences.Empty();
-					BP->GeneratedClass->PropertyLink = NULL;
-				}
 			}
 		}
 
@@ -827,6 +815,16 @@ void FBlueprintCompilationManagerImpl::FlushCompilationQueueImpl(TArray<UObject*
 		
 					FKismetCompilerContext& CompilerContext = *(CompilerData.Compiler);
 					CompilerContext.CompileClassLayout(EInternalCompilerFlags::PostponeLocalsGenerationUntilPhaseTwo);
+
+					// We immediately relink children so that iterative compilation logic has an easier time:
+					TArray<UClass*> ClassesToRelink;
+					GetDerivedClasses(BP->GeneratedClass, ClassesToRelink, false);
+					for (UClass* ChildClass : ClassesToRelink)
+					{
+						ChildClass->Bind();
+						ChildClass->StaticLink();
+						ensure(ChildClass->ClassDefaultObject == nullptr);
+					}
 				}
 				else
 				{
diff --git a/Engine/Source/Editor/Kismet/Private/FiBSearchInstance.cpp b/Engine/Source/Editor/Kismet/Private/FiBSearchInstance.cpp
--- a/Engine/Source/Editor/Kismet/Private/FiBSearchInstance.cpp
+++ b/Engine/Source/Editor/Kismet/Private/FiBSearchInstance.cpp
@@ -106,7 +106,7 @@ public:
 
 	}
 
-	FFiBContextHelper(TWeakPtr< FImaginaryFiBData > InContext)
+	FFiBContextHelper(FImaginaryFiBDataWeakPtr InContext)
 		: Context(InContext)
 	{
 	}
@@ -132,7 +132,7 @@ public:
 	}
 
 	/** Context that is actually being tested */
-	TWeakPtr< FImaginaryFiBData > Context;
+	FImaginaryFiBDataWeakPtr Context;
 
 	/** Modified in a const function callback, this is a going list of all search components that matched the expression */
 	mutable TMultiMap< const FImaginaryFiBData*, FComponentUniqueDisplay > MatchingSearchComponents;
@@ -141,7 +141,7 @@ public:
 ////////////////////////
 // FFiBSearchInstance
 
-FSearchResult FFiBSearchInstance::StartSearchQuery(const FString& InSearchString, TSharedPtr<FImaginaryBlueprint> InImaginaryBlueprintRoot)
+FSearchResult FFiBSearchInstance::StartSearchQuery(const FString& InSearchString, FImaginaryFiBDataSharedPtr InImaginaryBlueprintRoot)
 {
 	PendingSearchables.Add(InImaginaryBlueprintRoot);
 	DoSearchQuery(InSearchString);
@@ -149,13 +149,13 @@ FSearchResult FFiBSearchInstance::StartSearchQuery(const FString& InSearchString
 	return GetSearchResults(InImaginaryBlueprintRoot);
 }
 
-void FFiBSearchInstance::MakeSearchQuery(const FString& InSearchString, TSharedPtr<FImaginaryBlueprint> InImaginaryBlueprintRoot)
+void FFiBSearchInstance::MakeSearchQuery(const FString& InSearchString, FImaginaryFiBDataSharedPtr InImaginaryBlueprintRoot)
 {
 	PendingSearchables.Add(InImaginaryBlueprintRoot);
 	DoSearchQuery(InSearchString);
 }
 
-FSearchResult FFiBSearchInstance::GetSearchResults(TSharedPtr<FImaginaryBlueprint> InImaginaryBlueprintRoot)
+FSearchResult FFiBSearchInstance::GetSearchResults(FImaginaryFiBDataSharedPtr InImaginaryBlueprintRoot)
 {
 	FSearchResult CachedSearchResult = FImaginaryFiBData::CreateSearchTree(nullptr, InImaginaryBlueprintRoot, MatchesSearchQuery, MatchingSearchComponents);
 	return MatchesSearchQuery.Num() > 0? CachedSearchResult : nullptr;
@@ -183,7 +183,7 @@ bool FFiBSearchInstance::DoSearchQuery(const FString& InSearchString, bool bInCo
 	for (int SearchableIdx = 0; SearchableIdx < PendingSearchables.Num(); ++SearchableIdx)
 	{
 		CurrentSearchable = PendingSearchables[SearchableIdx];
-		TSharedPtr< FImaginaryFiBData > CurrentSearchablePinned = CurrentSearchable.Pin();
+		FImaginaryFiBDataSharedPtr CurrentSearchablePinned = CurrentSearchable.Pin();
 		CurrentSearchablePinned->ParseAllChildData();
 		if (ExpressionEvaluator->TestTextFilter(*CurrentSearchablePinned.Get()))
 		{
@@ -193,8 +193,8 @@ bool FFiBSearchInstance::DoSearchQuery(const FString& InSearchString, bool bInCo
 		if (bInComplete || CurrentSearchablePinned->IsCategory())
 		{
 			// Any children that are not treated as a TagAndValue Category should be added for independent searching
-			TArray< TSharedPtr< FImaginaryFiBData > > Children = CurrentSearchablePinned->GetAllParsedChildData();
-			for (TSharedPtr< FImaginaryFiBData > Child : Children)
+			TArray<FImaginaryFiBDataSharedPtr> Children = CurrentSearchablePinned->GetAllParsedChildData();
+			for (FImaginaryFiBDataSharedPtr Child : Children)
 			{
 				if (!Child->IsTagAndValueCategory())
 				{
@@ -208,7 +208,7 @@ bool FFiBSearchInstance::DoSearchQuery(const FString& InSearchString, bool bInCo
 	return MatchesSearchQuery.Num() > 0;
 }
 
-void FFiBSearchInstance::CreateFilteredResultsListFromTree(ESearchQueryFilter InSearchQueryFilter, TArray< TSharedPtr<FImaginaryFiBData> >& InOutValidSearchResults)
+void FFiBSearchInstance::CreateFilteredResultsListFromTree(ESearchQueryFilter InSearchQueryFilter, TArray<FImaginaryFiBDataSharedPtr>& InOutValidSearchResults)
 {
 	for (const FImaginaryFiBData* ImaginaryDataPtr : MatchesSearchQuery)
 	{
@@ -220,9 +220,9 @@ void FFiBSearchInstance::CreateFilteredResultsListFromTree(ESearchQueryFilter In
 	}
 }
 
-void FFiBSearchInstance::BuildFunctionTargets(TSharedPtr<FImaginaryFiBData> InRootData, ESearchQueryFilter InSearchQueryFilter, TArray< TWeakPtr< FImaginaryFiBData > >& OutTargetPendingSearchables)
+void FFiBSearchInstance::BuildFunctionTargets(FImaginaryFiBDataSharedPtr InRootData, ESearchQueryFilter InSearchQueryFilter, TArray<FImaginaryFiBDataWeakPtr>& OutTargetPendingSearchables)
 {
-	for (TSharedPtr<FImaginaryFiBData> ChildData : InRootData->GetAllParsedChildData())
+	for (FImaginaryFiBDataSharedPtr ChildData : InRootData->GetAllParsedChildData())
 	{
 		if (!ChildData->IsCategory() && ChildData->IsCompatibleWithFilter(InSearchQueryFilter))
 		{
@@ -236,9 +236,9 @@ void FFiBSearchInstance::BuildFunctionTargets(TSharedPtr<FImaginaryFiBData> InRo
 	}
 }
 
-void FFiBSearchInstance::BuildFunctionTargetsByName(TSharedPtr<FImaginaryFiBData> InRootData, FString InTagName, TArray< TWeakPtr< FImaginaryFiBData > >& OutTargetPendingSearchables)
+void FFiBSearchInstance::BuildFunctionTargetsByName(FImaginaryFiBDataSharedPtr InRootData, FString InTagName, TArray<FImaginaryFiBDataWeakPtr>& OutTargetPendingSearchables)
 {
-	for (TSharedPtr<FImaginaryFiBData> ChildData : InRootData->GetAllParsedChildData())
+	for (FImaginaryFiBDataSharedPtr ChildData : InRootData->GetAllParsedChildData())
 	{
 		if (ChildData->IsCategory())
 		{
@@ -262,7 +262,7 @@ bool FFiBSearchInstance::OnFilterFunction(const FTextFilterString& A, ESearchQue
 		TSharedPtr< FFiBSearchInstance > SubSearchInstance(new FFiBSearchInstance);
 		bool bSearchSuccess = false;
 
-		TSharedPtr< FImaginaryFiBData > CurrentSearchablePinned = CurrentSearchable.Pin();
+		FImaginaryFiBDataSharedPtr CurrentSearchablePinned = CurrentSearchable.Pin();
 		if (CurrentSearchablePinned->CanCallFilter(InSearchQueryFilter))
 		{
 			CurrentSearchablePinned->ParseAllChildData();
@@ -304,7 +304,7 @@ bool FFiBSearchInstance::OnFilterDefaultFunction(const FTextFilterString& InFunc
 		TSharedPtr< FFiBSearchInstance > SubSearchInstance(new FFiBSearchInstance);
 		bool bSearchSuccess = false;
 
-		TSharedPtr< FImaginaryFiBData > CurrentSearchablePinned = CurrentSearchable.Pin();
+		FImaginaryFiBDataSharedPtr CurrentSearchablePinned = CurrentSearchable.Pin();
 		CurrentSearchablePinned->ParseAllChildData();
 		BuildFunctionTargetsByName(CurrentSearchablePinned, InFunctionName.AsString(), SubSearchInstance->PendingSearchables);
 
diff --git a/Engine/Source/Editor/Kismet/Private/FiBSearchInstance.h b/Engine/Source/Editor/Kismet/Private/FiBSearchInstance.h
--- a/Engine/Source/Editor/Kismet/Private/FiBSearchInstance.h
+++ b/Engine/Source/Editor/Kismet/Private/FiBSearchInstance.h
@@ -61,7 +61,7 @@ public:
 	 * @param InImaginaryBlueprintRoot		The imaginary Blueprint to search through
 	 * @return								Search result shared pointer, can be used for display in the search results window
 	 */
-	FSearchResult StartSearchQuery(const FString& InSearchString, TSharedPtr<FImaginaryBlueprint> InImaginaryBlueprintRoot);
+	FSearchResult StartSearchQuery(const FString& InSearchString, FImaginaryFiBDataSharedPtr InImaginaryBlueprintRoot);
 
 	/**
 	 * Starts a search query given a string and an imaginary Blueprint
@@ -69,7 +69,7 @@ public:
 	 * @param InSearchString				The string to search using
 	 * @param InImaginaryBlueprintRoot		The imaginary Blueprint to search through
 	 */
-	void MakeSearchQuery(const FString& InSearchString, TSharedPtr<FImaginaryBlueprint> InImaginaryBlueprintRoot);
+	void MakeSearchQuery(const FString& InSearchString, FImaginaryFiBDataSharedPtr InImaginaryBlueprintRoot);
 	
 	/**
 	 * Runs a search query on any pending imaginary data
@@ -87,7 +87,7 @@ public:
 	 * @param InSearchQueryFilter				A filter to decide if the imaginary item is compatible with the function
 	 * @param OutTargetPendingSearchables		List of pending imaginary items that will be sub-searched, this gets filled out by the function
 	 */
-	void BuildFunctionTargets(TSharedPtr<FImaginaryFiBData> InRootData, ESearchQueryFilter InSearchQueryFilter, TArray< TWeakPtr< FImaginaryFiBData > >& OutTargetPendingSearchables);
+	void BuildFunctionTargets(FImaginaryFiBDataSharedPtr InRootData, ESearchQueryFilter InSearchQueryFilter, TArray<FImaginaryFiBDataWeakPtr>& OutTargetPendingSearchables);
 
 	/**
 	 * Builds a list of imaginary items, using their names, that can be targeted by a function
@@ -96,7 +96,7 @@ public:
 	 * @param InTagName							The name of objects to find
 	 * @param OutTargetPendingSearchables		List of pending imaginary items that will be sub-searched, this gets filled out by the function
 	 */
-	void BuildFunctionTargetsByName(TSharedPtr<FImaginaryFiBData> InRootData, FString InTagName, TArray< TWeakPtr< FImaginaryFiBData > >& OutTargetPendingSearchables);
+	void BuildFunctionTargetsByName(FImaginaryFiBDataSharedPtr InRootData, FString InTagName, TArray<FImaginaryFiBDataWeakPtr>& OutTargetPendingSearchables);
 
 	/**
 	 * Callback when a function is called by the evaluator
@@ -117,13 +117,13 @@ public:
 	bool OnFilterDefaultFunction(const FTextFilterString& InFunctionName, const FTextFilterString& InFunctionParams);
 
 	/** Builds a list of search results in their imaginary data form, filtered by an object type */
-	void CreateFilteredResultsListFromTree(ESearchQueryFilter InSearchQueryFilter, TArray< TSharedPtr<FImaginaryFiBData> >& InOutValidSearchResults);
+	void CreateFilteredResultsListFromTree(ESearchQueryFilter InSearchQueryFilter, TArray<FImaginaryFiBDataSharedPtr>& InOutValidSearchResults);
 
 	/** Helper function to return search results given an imaginary Blueprint root */
-	FSearchResult GetSearchResults(TSharedPtr<FImaginaryBlueprint> InImaginaryBlueprintRoot);
+	FSearchResult GetSearchResults(FImaginaryFiBDataSharedPtr InImaginaryBlueprintRoot);
 public:
-	/** Current item being searched in the Imaginary Blueprint */
-	TWeakPtr< FImaginaryFiBData > CurrentSearchable;
+	/** Current item being searched in the Imaginary Blueprint. Must be declared as thread-safe because this is a shared object. */
+	FImaginaryFiBDataWeakPtr CurrentSearchable;
 
 	/** A going list of all imaginary items that match the search query */
 	TArray< const FImaginaryFiBData* > MatchesSearchQuery;
@@ -131,8 +131,8 @@ public:
 	/** A mapping of items and their components that match the search query */
 	TMultiMap< const FImaginaryFiBData*, FComponentUniqueDisplay > MatchingSearchComponents;
 
-	/** A list of imaginary items that still need to be searched */
-	TArray< TWeakPtr< FImaginaryFiBData > > PendingSearchables;
+	/** A list of imaginary items that still need to be searched. Must be declared as thread-safe because it points to a shared object. */
+	TArray<FImaginaryFiBDataWeakPtr> PendingSearchables;
 
 	/** When a function returns on an item, this is the list of items that matched the sub-search query */
 	TArray< const FImaginaryFiBData* > LastFunctionResultMatchesSearchQuery;
diff --git a/Engine/Source/Editor/Kismet/Private/FindInBlueprintManager.cpp b/Engine/Source/Editor/Kismet/Private/FindInBlueprintManager.cpp
--- a/Engine/Source/Editor/Kismet/Private/FindInBlueprintManager.cpp
+++ b/Engine/Source/Editor/Kismet/Private/FindInBlueprintManager.cpp
@@ -221,7 +221,7 @@ float FStreamSearch::GetPercentComplete() const
 	return FFindInBlueprintSearchManager::Get().GetPercentComplete(this);
 }
 
-void FStreamSearch::GetFilteredImaginaryResults(TArray<TSharedPtr<class FImaginaryFiBData>>& OutFilteredImaginaryResults)
+void FStreamSearch::GetFilteredImaginaryResults(TArray<FImaginaryFiBDataSharedPtr>& OutFilteredImaginaryResults)
 {
 	OutFilteredImaginaryResults = MoveTemp(FilteredImaginaryResults);
 }
@@ -1477,7 +1477,7 @@ void FFindInBlueprintSearchManager::OnAssetAdded(const FAssetData& InAssetData)
 			{
 				ExtractUnloadedFiBData(InAssetData, *FiBSearchData, false);
 			}
-			else if(const FString* FiBVersionedSearchData = InAssetData.TagsAndValues.Find("FiBData"))
+			else if(const FString* FiBVersionedSearchData = InAssetData.TagsAndValues.Find(FBlueprintTags::FindInBlueprintsData))
 			{
 				if (FiBVersionedSearchData->Len() == 0)
 				{
@@ -1509,9 +1509,9 @@ void FFindInBlueprintSearchManager::ExtractUnloadedFiBData(const FAssetData& InA
 	FSearchData NewSearchData;
 
 	NewSearchData.BlueprintPath = InAssetData.ObjectPath;
-	InAssetData.GetTagValue("ParentClass", NewSearchData.ParentClass);
+	InAssetData.GetTagValue(FBlueprintTags::ParentClassPath, NewSearchData.ParentClass);
 
-	const FString ImplementedInterfaces = InAssetData.GetTagValueRef<FString>("ImplementedInterfaces");
+	const FString ImplementedInterfaces = InAssetData.GetTagValueRef<FString>(FBlueprintTags::ImplementedInterfaces);
 	if(!ImplementedInterfaces.IsEmpty())
 	{
 		FString FullInterface;
@@ -1808,15 +1808,10 @@ bool FFindInBlueprintSearchManager::ContinueSearchQuery(const FStreamSearch* InS
 		ActiveSearchCounter.Increment();
 	}
 
-	int32* SearchIdxPtr = NULL;
+	// Must lock this behind a critical section to ensure that no other thread is accessing it at the same time
+	FScopeLock ScopeLock(&SafeQueryModifyCriticalSection);
 
-	{
-		// Must lock this behind a critical section to ensure that no other thread is accessing it at the same time
-		FScopeLock ScopeLock(&SafeQueryModifyCriticalSection);
-		SearchIdxPtr = ActiveSearchQueries.Find(InSearchOriginator);
-	}
-
-	if(SearchIdxPtr)
+	if(int32* SearchIdxPtr = ActiveSearchQueries.Find(InSearchOriginator))
 	{
 		int32& SearchIdx = *SearchIdxPtr;
 		while(SearchIdx < SearchArray.Num())
@@ -1844,11 +1839,7 @@ bool FFindInBlueprintSearchManager::ContinueSearchQuery(const FStreamSearch* InS
 		}
 	}
 
-	{
-		// Must lock this behind a critical section to ensure that no other thread is accessing it at the same time
-		FScopeLock ScopeLock(&SafeQueryModifyCriticalSection);
-		ActiveSearchQueries.Remove(InSearchOriginator);
-	}
+	ActiveSearchQueries.Remove(InSearchOriginator);
 	ActiveSearchCounter.Decrement();
 
 	return false;
diff --git a/Engine/Source/Editor/Kismet/Private/FindInBlueprints.cpp b/Engine/Source/Editor/Kismet/Private/FindInBlueprints.cpp
--- a/Engine/Source/Editor/Kismet/Private/FindInBlueprints.cpp
+++ b/Engine/Source/Editor/Kismet/Private/FindInBlueprints.cpp
@@ -956,7 +956,7 @@ EActiveTimerReturnType SFindInBlueprints::UpdateSearchResults( double InCurrentT
 
 			StreamSearch->EnsureCompletion();
 
-			TArray<TSharedPtr<class FImaginaryFiBData>> ImaginaryResults;
+			TArray<FImaginaryFiBDataSharedPtr> ImaginaryResults;
 			if (OnSearchComplete.IsBound())
 			{
 				// Pull out the filtered imaginary results if there is a callback to pass them to
@@ -1061,7 +1061,7 @@ void SFindInBlueprints::MakeSearchQuery(FString InSearchString, bool bInIsFindWi
 			{
 				Interfaces.Add(InterfaceDesc.Interface->GetPathName());
 			}
-			TSharedPtr< FImaginaryBlueprint> ImaginaryBlueprint(new FImaginaryBlueprint(Blueprint->GetName(), Blueprint->GetPathName(), ParentClass, Interfaces, FFindInBlueprintSearchManager::Get().QuerySingleBlueprint(Blueprint)));
+			FImaginaryFiBDataSharedPtr ImaginaryBlueprint(new FImaginaryBlueprint(Blueprint->GetName(), Blueprint->GetPathName(), ParentClass, Interfaces, FFindInBlueprintSearchManager::Get().QuerySingleBlueprint(Blueprint)));
 			TSharedPtr< FFiBSearchInstance > SearchInstance(new FFiBSearchInstance);
 			FSearchResult SearchResult = RootSearchResult = SearchInstance->StartSearchQuery(SearchValue, ImaginaryBlueprint);
 
diff --git a/Engine/Source/Editor/Kismet/Private/ImaginaryBlueprintData.cpp b/Engine/Source/Editor/Kismet/Private/ImaginaryBlueprintData.cpp
--- a/Engine/Source/Editor/Kismet/Private/ImaginaryBlueprintData.cpp
+++ b/Engine/Source/Editor/Kismet/Private/ImaginaryBlueprintData.cpp
@@ -31,13 +31,13 @@ bool FComponentUniqueDisplay::operator==(const FComponentUniqueDisplay& Other)
 
 FCriticalSection FImaginaryFiBData::ParseChildDataCriticalSection;
 
-FImaginaryFiBData::FImaginaryFiBData(TWeakPtr<FImaginaryFiBData> InOuter) 
+FImaginaryFiBData::FImaginaryFiBData(FImaginaryFiBDataWeakPtr InOuter)
 	: LookupTablePtr(nullptr)
 	, Outer(InOuter)
 {
 }
 
-FImaginaryFiBData::FImaginaryFiBData(TWeakPtr<FImaginaryFiBData> InOuter, TSharedPtr< FJsonObject > InUnparsedJsonObject, TMap<int32, FText>* InLookupTablePtr)
+FImaginaryFiBData::FImaginaryFiBData(FImaginaryFiBDataWeakPtr InOuter, TSharedPtr< FJsonObject > InUnparsedJsonObject, TMap<int32, FText>* InLookupTablePtr)
 	: UnparsedJsonObject(InUnparsedJsonObject)
 	, LookupTablePtr(InLookupTablePtr)
 	, Outer(InOuter)
@@ -60,16 +60,16 @@ FSearchResult FImaginaryFiBData::CreateSearchResult(FSearchResult InParent) cons
 	return ReturnSearchResult;
 }
 
-FSearchResult FImaginaryFiBData::CreateSearchTree(FSearchResult InParentSearchResult, TWeakPtr< FImaginaryFiBData > InCurrentPointer, TArray< const FImaginaryFiBData* >& InValidSearchResults, TMultiMap< const FImaginaryFiBData*, FComponentUniqueDisplay >& InMatchingSearchComponents)
+FSearchResult FImaginaryFiBData::CreateSearchTree(FSearchResult InParentSearchResult, FImaginaryFiBDataWeakPtr InCurrentPointer, TArray< const FImaginaryFiBData* >& InValidSearchResults, TMultiMap< const FImaginaryFiBData*, FComponentUniqueDisplay >& InMatchingSearchComponents)
 {
-	TSharedPtr<FImaginaryFiBData> CurrentDataPtr = InCurrentPointer.Pin();
+	FImaginaryFiBDataSharedPtr CurrentDataPtr = InCurrentPointer.Pin();
 	if (FImaginaryFiBData* CurrentData = CurrentDataPtr.Get())
 	{
 		FSearchResult CurrentSearchResult = CurrentData->CreateSearchResult(InParentSearchResult);
 		bool bValidSearchResults = false;
 
 		// Check all children first, to see if they are valid in the search results
-		for (TSharedPtr<FImaginaryFiBData> ChildData : CurrentData->ParsedChildData)
+		for (FImaginaryFiBDataSharedPtr ChildData : CurrentData->ParsedChildData)
 		{
 			FSearchResult Result = CreateSearchTree(CurrentSearchResult, ChildData, InValidSearchResults, InMatchingSearchComponents);
 			if (Result.IsValid())
@@ -135,7 +135,7 @@ void FImaginaryFiBData::ParseAllChildData_Internal(ESearchableValueStatus InSear
 
 			if (MetaDataField.IsValid())
 			{
-				TSharedPtr<FFiBMetaData> MetaDataFiBInfo = MakeShareable(new FFiBMetaData(AsShared(), MetaDataField, LookupTablePtr));
+				TSharedPtr<FFiBMetaData, ESPMode::ThreadSafe> MetaDataFiBInfo = MakeShareable(new FFiBMetaData(AsShared(), MetaDataField, LookupTablePtr));
 				MetaDataFiBInfo->ParseAllChildData_Internal();
 
 				if (MetaDataFiBInfo->IsHidden() && MetaDataFiBInfo->IsExplicit())
@@ -188,7 +188,7 @@ void FImaginaryFiBData::ParseJsonValue(FText InKey, FText InDisplayKey, TSharedP
 	}
 	else if( InJsonValue->Type == EJson::Array)
 	{
-		TSharedPtr< FCategorySectionHelper > ArrayCategory = MakeShareable(new FCategorySectionHelper(AsShared(), LookupTablePtr, InKey, true));
+		TSharedPtr< FCategorySectionHelper, ESPMode::ThreadSafe > ArrayCategory = MakeShareable(new FCategorySectionHelper(AsShared(), LookupTablePtr, InKey, true));
 		ParsedChildData.Add(ArrayCategory);
 		TArray<TSharedPtr< FJsonValue > > ArrayList = InJsonValue->AsArray();
 		for( int32 ArrayIdx = 0; ArrayIdx < ArrayList.Num(); ++ArrayIdx)
@@ -199,7 +199,7 @@ void FImaginaryFiBData::ParseJsonValue(FText InKey, FText InDisplayKey, TSharedP
 	}
 	else if (InJsonValue->Type == EJson::Object)
 	{
-		TSharedPtr< FCategorySectionHelper > SubObjectCategory = MakeShareable(new FCategorySectionHelper(AsShared(), InJsonValue->AsObject(), LookupTablePtr, InDisplayKey, bIsInArray));
+		TSharedPtr< FCategorySectionHelper, ESPMode::ThreadSafe > SubObjectCategory = MakeShareable(new FCategorySectionHelper(AsShared(), InJsonValue->AsObject(), LookupTablePtr, InDisplayKey, bIsInArray));
 		SubObjectCategory->ParseAllChildData_Internal(SearchabilityStatus);
 		ParsedChildData.Add(SubObjectCategory);
 	}
@@ -240,7 +240,7 @@ bool FImaginaryFiBData::TestBasicStringExpression(const FTextFilterString& InVal
 		}
 	}
 	// Any children that are treated as a TagAndValue Category should be added for independent searching
-	for (const TSharedPtr< FImaginaryFiBData > Child : ParsedChildData)
+	for (const FImaginaryFiBDataSharedPtr Child : ParsedChildData)
 	{
 		if (Child->IsTagAndValueCategory())
 		{
@@ -276,7 +276,7 @@ bool FImaginaryFiBData::TestComplexExpression(const FName& InKey, const FTextFil
 	}
 
 	// Any children that are treated as a TagAndValue Category should be added for independent searching
-	for (const TSharedPtr< FImaginaryFiBData > Child : ParsedChildData)
+	for (const FImaginaryFiBDataSharedPtr Child : ParsedChildData)
 	{
 		if (Child->IsTagAndValueCategory())
 		{
@@ -294,7 +294,7 @@ UObject* FImaginaryFiBData::GetObject(UBlueprint* InBlueprint) const
 ///////////////////////////
 // FFiBMetaData
 
-FFiBMetaData::FFiBMetaData(TWeakPtr<FImaginaryFiBData> InOuter, TSharedPtr< FJsonObject > InUnparsedJsonObject, TMap<int32, FText>* InLookupTablePtr)
+FFiBMetaData::FFiBMetaData(FImaginaryFiBDataWeakPtr InOuter, TSharedPtr< FJsonObject > InUnparsedJsonObject, TMap<int32, FText>* InLookupTablePtr)
 	: FImaginaryFiBData(InOuter, InUnparsedJsonObject, InLookupTablePtr)
 	, bIsHidden(false)
 	, bIsExplicit(false)
@@ -322,21 +322,21 @@ bool FFiBMetaData::TrySpecialHandleJsonValue(FText InKey, TSharedPtr< FJsonValue
 ///////////////////////////
 // FCategorySectionHelper
 
-FCategorySectionHelper::FCategorySectionHelper(TWeakPtr<FImaginaryFiBData> InOuter, TMap<int32, FText>* InLookupTablePtr, FText InCategoryName, bool bInTagAndValueCategory)
+FCategorySectionHelper::FCategorySectionHelper(FImaginaryFiBDataWeakPtr InOuter, TMap<int32, FText>* InLookupTablePtr, FText InCategoryName, bool bInTagAndValueCategory)
 	: FImaginaryFiBData(InOuter, nullptr, InLookupTablePtr)
 	, CategoryName(InCategoryName)
 	, bIsTagAndValue(bInTagAndValueCategory)
 {
 }
 
-FCategorySectionHelper::FCategorySectionHelper(TWeakPtr<FImaginaryFiBData> InOuter, TSharedPtr< FJsonObject > InUnparsedJsonObject, TMap<int32, FText>* InLookupTablePtr, FText InCategoryName, bool bInTagAndValueCategory)
+FCategorySectionHelper::FCategorySectionHelper(FImaginaryFiBDataWeakPtr InOuter, TSharedPtr< FJsonObject > InUnparsedJsonObject, TMap<int32, FText>* InLookupTablePtr, FText InCategoryName, bool bInTagAndValueCategory)
 	: FImaginaryFiBData(InOuter, InUnparsedJsonObject, InLookupTablePtr)
 	, CategoryName(InCategoryName)
 	, bIsTagAndValue(bInTagAndValueCategory)
 {
 }
 
-FCategorySectionHelper::FCategorySectionHelper(TWeakPtr<FImaginaryFiBData> InOuter, TSharedPtr< FJsonObject > InUnparsedJsonObject, TMap<int32, FText>* InLookupTablePtr, FText InCategoryName, bool bInTagAndValueCategory, FCategorySectionHelperCallback InSpecialHandlingCallback)
+FCategorySectionHelper::FCategorySectionHelper(FImaginaryFiBDataWeakPtr InOuter, TSharedPtr< FJsonObject > InUnparsedJsonObject, TMap<int32, FText>* InLookupTablePtr, FText InCategoryName, bool bInTagAndValueCategory, FCategorySectionHelperCallback InSpecialHandlingCallback)
 	: FImaginaryFiBData(InOuter, InUnparsedJsonObject, InLookupTablePtr)
 	, SpecialHandlingCallback(InSpecialHandlingCallback)
 	, CategoryName(InCategoryName)
@@ -403,7 +403,7 @@ FImaginaryBlueprint::FImaginaryBlueprint(FString InBlueprintName, FString InBlue
 	ParsedTagsAndValues.Add(FindInBlueprintsHelpers::FSimpleFTextKeyStorage(FFindInBlueprintSearchTags::FiB_Path), FSearchableValueInfo(FFindInBlueprintSearchTags::FiB_Path, FText::FromString(InBlueprintPath), ESearchableValueStatus::ExplicitySearchable));
 	ParsedTagsAndValues.Add(FindInBlueprintsHelpers::FSimpleFTextKeyStorage(FFindInBlueprintSearchTags::FiB_ParentClass), FSearchableValueInfo(FFindInBlueprintSearchTags::FiB_ParentClass, FText::FromString(InBlueprintParentClass), ESearchableValueStatus::ExplicitySearchable));
 
-	TSharedPtr< FCategorySectionHelper > InterfaceCategory = MakeShareable(new FCategorySectionHelper(nullptr, &LookupTable, FFindInBlueprintSearchTags::FiB_Interfaces, true));
+	TSharedPtr< FCategorySectionHelper, ESPMode::ThreadSafe > InterfaceCategory = MakeShareable(new FCategorySectionHelper(nullptr, &LookupTable, FFindInBlueprintSearchTags::FiB_Interfaces, true));
 	for( int32 InterfaceIdx = 0; InterfaceIdx < InInterfaces.Num(); ++InterfaceIdx)
 	{
 		FString& Interface = InInterfaces[InterfaceIdx];
@@ -507,7 +507,7 @@ void FImaginaryBlueprint::ParseGraph( TSharedPtr< FJsonValue > InJsonValue, FStr
 	}
 }
 
-void FImaginaryBlueprint::ParseComponents(TSharedPtr< FJsonObject > InJsonObject, TArray< TSharedPtr<FImaginaryFiBData> >& OutParsedChildData)
+void FImaginaryBlueprint::ParseComponents(TSharedPtr< FJsonObject > InJsonObject, TArray<FImaginaryFiBDataSharedPtr>& OutParsedChildData)
 {
 	// Pulls out all properties (variables) for this Blueprint
 	TArray<TSharedPtr< FJsonValue > > ComponentList = InJsonObject->GetArrayField(FFindInBlueprintSearchTags::FiB_Components.ToString());
@@ -520,7 +520,7 @@ void FImaginaryBlueprint::ParseComponents(TSharedPtr< FJsonObject > InJsonObject
 //////////////////////////
 // FImaginaryGraph
 
-FImaginaryGraph::FImaginaryGraph(TWeakPtr<FImaginaryFiBData> InOuter, TSharedPtr< FJsonObject > InUnparsedJsonObject, TMap<int32, FText>* InLookupTablePtr, EGraphType InGraphType)
+FImaginaryGraph::FImaginaryGraph(FImaginaryFiBDataWeakPtr InOuter, TSharedPtr< FJsonObject > InUnparsedJsonObject, TMap<int32, FText>* InLookupTablePtr, EGraphType InGraphType)
 	: FImaginaryFiBData(InOuter, InUnparsedJsonObject, InLookupTablePtr)
 	, GraphType(InGraphType)
 {
@@ -590,7 +590,7 @@ bool FImaginaryGraph::TrySpecialHandleJsonValue(FText InKey, TSharedPtr< FJsonVa
 //////////////////////////////////////
 // FImaginaryGraphNode
 
-FImaginaryGraphNode::FImaginaryGraphNode(TWeakPtr<FImaginaryFiBData> InOuter, TSharedPtr< FJsonObject > InUnparsedJsonObject, TMap<int32, FText>* InLookupTablePtr)
+FImaginaryGraphNode::FImaginaryGraphNode(FImaginaryFiBDataWeakPtr InOuter, TSharedPtr< FJsonObject > InUnparsedJsonObject, TMap<int32, FText>* InLookupTablePtr)
 	: FImaginaryFiBData(InOuter, InUnparsedJsonObject, InLookupTablePtr)
 {
 }
@@ -681,7 +681,7 @@ void FImaginaryGraphNode::ParseAllChildData_Internal(ESearchableValueStatus InSe
 ///////////////////////////////////////////
 // FImaginaryProperty
 
-FImaginaryProperty::FImaginaryProperty(TWeakPtr<FImaginaryFiBData> InOuter, TSharedPtr< FJsonObject > InUnparsedJsonObject, TMap<int32, FText>* InLookupTablePtr)
+FImaginaryProperty::FImaginaryProperty(FImaginaryFiBDataWeakPtr InOuter, TSharedPtr< FJsonObject > InUnparsedJsonObject, TMap<int32, FText>* InLookupTablePtr)
 	: FImaginaryFiBData(InOuter, InUnparsedJsonObject, InLookupTablePtr)
 {
 }
@@ -723,7 +723,7 @@ ESearchableValueStatus FImaginaryProperty::GetSearchabilityStatus(FString InKey)
 //////////////////////////////
 // FImaginaryComponent
 
-FImaginaryComponent::FImaginaryComponent(TWeakPtr<FImaginaryFiBData> InOuter, TSharedPtr< FJsonObject > InUnparsedJsonObject, TMap<int32, FText>* InLookupTablePtr)
+FImaginaryComponent::FImaginaryComponent(FImaginaryFiBDataWeakPtr InOuter, TSharedPtr< FJsonObject > InUnparsedJsonObject, TMap<int32, FText>* InLookupTablePtr)
 	: FImaginaryProperty(InOuter, InUnparsedJsonObject, InLookupTablePtr)
 {
 }
@@ -736,7 +736,7 @@ bool FImaginaryComponent::IsCompatibleWithFilter(ESearchQueryFilter InSearchQuer
 //////////////////////////////
 // FImaginaryPin
 
-FImaginaryPin::FImaginaryPin(TWeakPtr<FImaginaryFiBData> InOuter, TSharedPtr< FJsonObject > InUnparsedJsonObject, TMap<int32, FText>* InLookupTablePtr, FString InSchemaName)
+FImaginaryPin::FImaginaryPin(FImaginaryFiBDataWeakPtr InOuter, TSharedPtr< FJsonObject > InUnparsedJsonObject, TMap<int32, FText>* InLookupTablePtr, FString InSchemaName)
 	: FImaginaryFiBData(InOuter, InUnparsedJsonObject, InLookupTablePtr)
 	, SchemaName(InSchemaName)
 {
diff --git a/Engine/Source/Editor/Kismet/Private/SBlueprintBookmarks.cpp b/Engine/Source/Editor/Kismet/Private/SBlueprintBookmarks.cpp
--- a/Engine/Source/Editor/Kismet/Private/SBlueprintBookmarks.cpp
+++ b/Engine/Source/Editor/Kismet/Private/SBlueprintBookmarks.cpp
@@ -258,7 +258,9 @@ void SBlueprintBookmarks::OnDeleteSelectedTreeViewItems()
 	{
 		for(FTreeViewItemPtr SelectedItem : TreeViewWidget->GetSelectedItems())
 		{
-			BlueprintEditor->RemoveBookmark(SelectedItem->BookmarkNode.NodeGuid, false);
+			// The item contains a reference to the node in the source array, so use a local copy of the GUID since this will use RemoveSwap().
+			const FGuid NodeGuid = SelectedItem->BookmarkNode.NodeGuid;
+			BlueprintEditor->RemoveBookmark(NodeGuid, false);
 		}
 
 		RefreshBookmarksTree();
diff --git a/Engine/Source/Editor/Kismet/Private/SReplaceNodeReferences.cpp b/Engine/Source/Editor/Kismet/Private/SReplaceNodeReferences.cpp
--- a/Engine/Source/Editor/Kismet/Private/SReplaceNodeReferences.cpp
+++ b/Engine/Source/Editor/Kismet/Private/SReplaceNodeReferences.cpp
@@ -8,9 +8,6 @@
 #include "Engine/MemberReference.h"
 #include "EdGraphSchema_K2.h"
 #include "K2Node_Variable.h"
-
-#include "FindInBlueprintManager.h"
-#include "FindInBlueprints.h"
 #include "Kismet2/BlueprintEditorUtils.h"
 #include "ImaginaryBlueprintData.h"
 #include "ObjectEditorUtils.h"
@@ -448,7 +445,7 @@ void SReplaceNodeReferences::OnSubmitSearchQuery(bool bFindAndReplace)
 	FindInBlueprints->MakeSearchQuery(SearchTerm, false, ESearchQueryFilter::NodesFilter, EFiBVersion::FIB_VER_VARIABLE_REFERENCE, OnSearchComplete);
 }
 
-void SReplaceNodeReferences::FindAllReplacementsComplete(TArray<TSharedPtr<class FImaginaryFiBData>>& InRawDataList)
+void SReplaceNodeReferences::FindAllReplacementsComplete(TArray<FImaginaryFiBDataSharedPtr>& InRawDataList)
 {
 	if (SelectedTargetReferenceItem.IsValid())
 	{
@@ -460,7 +457,7 @@ void SReplaceNodeReferences::FindAllReplacementsComplete(TArray<TSharedPtr<class
 			BlueprintEditor.Pin()->GetBlueprintObj()->Modify();
 
 			TArray< UBlueprint* > BlueprintsModified;
-			for (TSharedPtr<class FImaginaryFiBData> ImaginaryData : InRawDataList)
+			for (FImaginaryFiBDataSharedPtr ImaginaryData : InRawDataList)
 			{
 				BlueprintsModified.AddUnique(ImaginaryData->GetBlueprint());
 				UObject* Node = ImaginaryData->GetObject(ImaginaryData->GetBlueprint());
diff --git a/Engine/Source/Editor/Kismet/Private/SReplaceNodeReferences.h b/Engine/Source/Editor/Kismet/Private/SReplaceNodeReferences.h
--- a/Engine/Source/Editor/Kismet/Private/SReplaceNodeReferences.h
+++ b/Engine/Source/Editor/Kismet/Private/SReplaceNodeReferences.h
@@ -14,6 +14,7 @@
 #include "Widgets/Views/STableRow.h"
 #include "Widgets/Views/STreeView.h"
 #include "BlueprintEditor.h"
+#include "FindInBlueprints.h"
 
 class FImaginaryFiBData;
 class SComboButton;
@@ -118,7 +119,7 @@ protected:
 	FReply OnFindAndReplaceAll();
 
 	/** Callback when the search for "Find and Replace All" is complete so that the replacements can begin */
-	void FindAllReplacementsComplete(TArray<TSharedPtr<class FImaginaryFiBData>>& InRawDataList);
+	void FindAllReplacementsComplete(TArray<FImaginaryFiBDataSharedPtr>& InRawDataList);
 
 	/** Recursively gathers all available Blueprint Variable references to replace with */
 	void GatherAllAvailableBlueprintVariables(UClass* InTargetClass);
diff --git a/Engine/Source/Editor/Kismet/Public/FindInBlueprintManager.h b/Engine/Source/Editor/Kismet/Public/FindInBlueprintManager.h
--- a/Engine/Source/Editor/Kismet/Public/FindInBlueprintManager.h
+++ b/Engine/Source/Editor/Kismet/Public/FindInBlueprintManager.h
@@ -18,6 +18,10 @@ class FImaginaryFiBData;
 class FSpawnTabArgs;
 class SFindInBlueprints;
 
+// Shared pointers to cached imaginary data (must be declared as thread-safe).
+typedef TWeakPtr<FImaginaryFiBData, ESPMode::ThreadSafe> FImaginaryFiBDataWeakPtr;
+typedef TSharedPtr<FImaginaryFiBData, ESPMode::ThreadSafe> FImaginaryFiBDataSharedPtr;
+
 #define MAX_GLOBAL_FIND_RESULTS 4
 
 /**
@@ -119,7 +123,7 @@ struct FSearchData
 	bool bMarkedForDeletion;
 
 	/** Cached ImaginaryBlueprint data for the searchable content, prevents having to re-parse every search */
-	TSharedPtr< class FImaginaryBlueprint > ImaginaryBlueprint;
+	FImaginaryFiBDataSharedPtr ImaginaryBlueprint;
 
 	/** Version of the data */
 	int32 Version;
@@ -299,7 +303,7 @@ public:
 	}
 
 	/** Returns the FilteredImaginaryResults from the search query, these results have been filtered by the ImaginaryDataFilter. */
-	void GetFilteredImaginaryResults(TArray<TSharedPtr<class FImaginaryFiBData>>& OutFilteredImaginaryResults);
+	void GetFilteredImaginaryResults(TArray<FImaginaryFiBDataSharedPtr>& OutFilteredImaginaryResults);
 
 public:
 	/** Thread to run the cleanup FRunnable on */
@@ -326,8 +330,8 @@ public:
 	/** A going count of all Blueprints below the MinimiumVersionRequirement */
 	int32 BlueprintCountBelowVersion;
 
-	/** Filtered (ImaginaryDataFilter) list of imaginary data results that met the search requirements */
-	TArray<TSharedPtr<class FImaginaryFiBData>> FilteredImaginaryResults;
+	/** Filtered (ImaginaryDataFilter) list of imaginary data results that met the search requirements. Must be declared as thread-safe since imaginary data is a shared resource. */
+	TArray<FImaginaryFiBDataSharedPtr> FilteredImaginaryResults;
 
 	/** Filter to limit the FilteredImaginaryResults to */
 	enum ESearchQueryFilter ImaginaryDataFilter;
diff --git a/Engine/Source/Editor/Kismet/Public/FindInBlueprints.h b/Engine/Source/Editor/Kismet/Public/FindInBlueprints.h
--- a/Engine/Source/Editor/Kismet/Public/FindInBlueprints.h
+++ b/Engine/Source/Editor/Kismet/Public/FindInBlueprints.h
@@ -25,7 +25,7 @@ class FUICommandList;
 
 typedef STreeView<FSearchResult>  STreeViewType;
 
-DECLARE_DELEGATE_OneParam(FOnSearchComplete, TArray<TSharedPtr<class FImaginaryFiBData>>&);
+DECLARE_DELEGATE_OneParam(FOnSearchComplete, TArray<FImaginaryFiBDataSharedPtr>&);
 
 /** Some utility functions to help with Find-in-Blueprint functionality */
 namespace FindInBlueprintsHelpers
diff --git a/Engine/Source/Editor/Kismet/Public/ImaginaryBlueprintData.h b/Engine/Source/Editor/Kismet/Public/ImaginaryBlueprintData.h
--- a/Engine/Source/Editor/Kismet/Public/ImaginaryBlueprintData.h
+++ b/Engine/Source/Editor/Kismet/Public/ImaginaryBlueprintData.h
@@ -100,11 +100,11 @@ struct FComponentUniqueDisplay
 	FSearchResult SearchResult;
 };
 
-class FImaginaryFiBData : public ITextFilterExpressionContext, public TSharedFromThis<FImaginaryFiBData>
+class FImaginaryFiBData : public ITextFilterExpressionContext, public TSharedFromThis<FImaginaryFiBData, ESPMode::ThreadSafe>
 {
 public:
-	FImaginaryFiBData(TWeakPtr<FImaginaryFiBData> InOuter);
-	FImaginaryFiBData(TWeakPtr<FImaginaryFiBData> InOuter, TSharedPtr< FJsonObject > InUnparsedJsonObject, TMap<int32, FText>* InLookupTablePtr);
+	FImaginaryFiBData(FImaginaryFiBDataWeakPtr InOuter);
+	FImaginaryFiBData(FImaginaryFiBDataWeakPtr InOuter, TSharedPtr< FJsonObject > InUnparsedJsonObject, TMap<int32, FText>* InLookupTablePtr);
 
 	/** ITextFilterExpressionContext Interface */
 	// We don't actually use these overrides, see FFiBContextHelper for how we call the alternate functions. These will assert if they are accidentally called.
@@ -150,7 +150,7 @@ public:
 	FSearchResult CreateSearchResult(FSearchResult InParent) const;
 
 	/** Accessor for the parsed child data for this item */
-	const TArray< TSharedPtr<FImaginaryFiBData> >& GetAllParsedChildData() const
+	const TArray<FImaginaryFiBDataSharedPtr>& GetAllParsedChildData() const
 	{
 		return ParsedChildData;
 	}
@@ -162,13 +162,13 @@ public:
 	}
 
 	/** Returns the Outer of this Imaginary data that directly owns it */
-	TWeakPtr< FImaginaryFiBData > GetOuter() const
+	FImaginaryFiBDataWeakPtr GetOuter() const
 	{
 		return Outer;
 	}
 
 	/** Builds a SearchTree ready to be displayed in the Find-in-Blueprints window */
-	static FSearchResult CreateSearchTree(FSearchResult InParentSearchResult, TWeakPtr< FImaginaryFiBData > InCurrentPointer, TArray< const FImaginaryFiBData* >& InValidSearchResults, TMultiMap< const FImaginaryFiBData*, FComponentUniqueDisplay >& InMatchingSearchComponents);
+	static FSearchResult CreateSearchTree(FSearchResult InParentSearchResult, FImaginaryFiBDataWeakPtr InCurrentPointer, TArray< const FImaginaryFiBData* >& InValidSearchResults, TMultiMap< const FImaginaryFiBData*, FComponentUniqueDisplay >& InMatchingSearchComponents);
 
 protected:
 	/**
@@ -196,11 +196,11 @@ protected:
 	virtual void ParseAllChildData_Internal(ESearchableValueStatus InSearchabilityOverride = ESearchableValueStatus::Searchable);
 
 protected:
-	/** The unparsed Json object representing this item. Auto-cleared after parsing */
+	/** The unparsed Json object representing this item. Auto-cleared after parsing. Does not need to be declared as thread-safe because it's only accessed when parsing and that is already a critical section. */
 	TSharedPtr< FJsonObject > UnparsedJsonObject;
 
-	/** All parsed child data for this item */
-	TArray< TSharedPtr<FImaginaryFiBData> > ParsedChildData;
+	/** All parsed child data for this item. Must be declared as thread-safe because it may be accessed on different threads. */
+	TArray<FImaginaryFiBDataSharedPtr> ParsedChildData;
 
 	/** A mapping of tags to their values and searchability status */
 	TMultiMap< FindInBlueprintsHelpers::FSimpleFTextKeyStorage, FSearchableValueInfo > ParsedTagsAndValues;
@@ -208,17 +208,17 @@ protected:
 	/** Pointer to the lookup table to decompressed the Json strings back into fully formed FTexts */
 	TMap<int32, FText>* LookupTablePtr;
 
-	/** Outer of this object that owns it, used for climbing up the hierarchy */
-	TWeakPtr<FImaginaryFiBData> Outer;
+	/** Outer of this object that owns it, used for climbing up the hierarchy. Must be declared as thread-safe because it may be accessed on different threads. */
+	FImaginaryFiBDataWeakPtr Outer;
 
-	/** Allows for thread-safe parsing of the imaginary data. Only a single Imaginary data can be parsed at a time */
+	/** Allows for thread-safe parsing of the imaginary data. Only a single Imaginary data can be parsed at a time. */
 	static FCriticalSection ParseChildDataCriticalSection;
 };
 
 class FFiBMetaData : public FImaginaryFiBData
 {
 public:
-	FFiBMetaData(TWeakPtr<FImaginaryFiBData> InOuter, TSharedPtr< FJsonObject > InUnparsedJsonObject, TMap<int32, FText>* InLookupTablePtr);
+	FFiBMetaData(FImaginaryFiBDataWeakPtr InOuter, TSharedPtr< FJsonObject > InUnparsedJsonObject, TMap<int32, FText>* InLookupTablePtr);
 	
 	/** FImaginaryFiBData Interface */
 	virtual bool TrySpecialHandleJsonValue(FText InKey, TSharedPtr< FJsonValue > InJsonValue) override;
@@ -253,11 +253,11 @@ class FCategorySectionHelper : public FImaginaryFiBData
 {
 public:
 	/** Callback declaration for handling special parsing of the items in the category */
-	DECLARE_DELEGATE_TwoParams(FCategorySectionHelperCallback, TSharedPtr< FJsonObject >, TArray< TSharedPtr<FImaginaryFiBData> >&);
+	DECLARE_DELEGATE_TwoParams(FCategorySectionHelperCallback, TSharedPtr< FJsonObject >, TArray<FImaginaryFiBDataSharedPtr>&);
 
-	FCategorySectionHelper(TWeakPtr<FImaginaryFiBData> InOuter, TMap<int32, FText>* InLookupTablePtr, FText InCategoryName, bool bInTagAndValueCategory);
-	FCategorySectionHelper(TWeakPtr<FImaginaryFiBData> InOuter, TSharedPtr< FJsonObject > InUnparsedJsonObject, TMap<int32, FText>* InLookupTablePtr, FText InCategoryName, bool bInTagAndValueCategory);
-	FCategorySectionHelper(TWeakPtr<FImaginaryFiBData> InOuter, TSharedPtr< FJsonObject > InUnparsedJsonObject, TMap<int32, FText>* InLookupTablePtr, FText InCategoryName, bool bInTagAndValueCategory, FCategorySectionHelperCallback InSpecialHandlingCallback);
+	FCategorySectionHelper(FImaginaryFiBDataWeakPtr InOuter, TMap<int32, FText>* InLookupTablePtr, FText InCategoryName, bool bInTagAndValueCategory);
+	FCategorySectionHelper(FImaginaryFiBDataWeakPtr InOuter, TSharedPtr< FJsonObject > InUnparsedJsonObject, TMap<int32, FText>* InLookupTablePtr, FText InCategoryName, bool bInTagAndValueCategory);
+	FCategorySectionHelper(FImaginaryFiBDataWeakPtr InOuter, TSharedPtr< FJsonObject > InUnparsedJsonObject, TMap<int32, FText>* InLookupTablePtr, FText InCategoryName, bool bInTagAndValueCategory, FCategorySectionHelperCallback InSpecialHandlingCallback);
 
 	/** FImaginaryFiBData Interface */
 	virtual void ParseAllChildData_Internal(ESearchableValueStatus InSearchabilityOverride/* = ESearchableValueStatus::ESearchableValueStatus::Searchable*/) override;
@@ -310,7 +310,7 @@ protected:
 	void ParseGraph( TSharedPtr< FJsonValue > InJsonValue, FString InCategoryTitle, EGraphType InGraphType );
 	
 	/** Callback to specially parse an array of Json Objects representing components */
-	void ParseComponents(TSharedPtr< FJsonObject > InJsonObject, TArray< TSharedPtr<FImaginaryFiBData> >& OutParsedChildData);
+	void ParseComponents(TSharedPtr< FJsonObject > InJsonObject, TArray<FImaginaryFiBDataSharedPtr>& OutParsedChildData);
 
 	/** Parses a raw string of Json to a Json object hierarchy */
 	void ParseToJson(bool bInIsVersioned);
@@ -330,7 +330,7 @@ protected:
 class FImaginaryGraph : public FImaginaryFiBData
 {
 public:
-	FImaginaryGraph(TWeakPtr<FImaginaryFiBData> InOuter, TSharedPtr< FJsonObject > InUnparsedJsonObject, TMap<int32, FText>* InLookupTablePtr, EGraphType InGraphType);
+	FImaginaryGraph(FImaginaryFiBDataWeakPtr InOuter, TSharedPtr< FJsonObject > InUnparsedJsonObject, TMap<int32, FText>* InLookupTablePtr, EGraphType InGraphType);
 
 	/** FImaginaryFiBData Interface */
 	virtual bool IsCompatibleWithFilter(ESearchQueryFilter InSearchQueryFilter) const override;
@@ -353,7 +353,7 @@ protected:
 class FImaginaryGraphNode : public FImaginaryFiBData
 {
 public:
-	FImaginaryGraphNode(TWeakPtr<FImaginaryFiBData> InOuter, TSharedPtr< FJsonObject > InUnparsedJsonObject, TMap<int32, FText>* InLookupTablePtr);
+	FImaginaryGraphNode(FImaginaryFiBDataWeakPtr InOuter, TSharedPtr< FJsonObject > InUnparsedJsonObject, TMap<int32, FText>* InLookupTablePtr);
 
 	/** FImaginaryFiBData Interface */
 	virtual bool IsCompatibleWithFilter(ESearchQueryFilter InSearchQueryFilter) const override;
@@ -377,7 +377,7 @@ protected:
 class FImaginaryProperty : public FImaginaryFiBData
 {
 public:
-	FImaginaryProperty(TWeakPtr<FImaginaryFiBData> InOuter, TSharedPtr< FJsonObject > InUnparsedJsonObject, TMap<int32, FText>* InLookupTablePtr);
+	FImaginaryProperty(FImaginaryFiBDataWeakPtr InOuter, TSharedPtr< FJsonObject > InUnparsedJsonObject, TMap<int32, FText>* InLookupTablePtr);
 
 	/** FImaginaryFiBData Interface */
 	virtual bool IsCompatibleWithFilter(ESearchQueryFilter InSearchQueryFilter) const override;
@@ -394,7 +394,7 @@ protected:
 class FImaginaryComponent : public FImaginaryProperty
 {
 public:
-	FImaginaryComponent(TWeakPtr<FImaginaryFiBData> InOuter, TSharedPtr< FJsonObject > InUnparsedJsonObject, TMap<int32, FText>* InLookupTablePtr);
+	FImaginaryComponent(FImaginaryFiBDataWeakPtr InOuter, TSharedPtr< FJsonObject > InUnparsedJsonObject, TMap<int32, FText>* InLookupTablePtr);
 
 	/** FImaginaryFiBData Interface */
 	virtual bool IsCompatibleWithFilter(ESearchQueryFilter InSearchQueryFilter) const override;
@@ -405,7 +405,7 @@ public:
 class FImaginaryPin : public FImaginaryFiBData
 {
 public:
-	FImaginaryPin(TWeakPtr<FImaginaryFiBData> InOuter, TSharedPtr< FJsonObject > InUnparsedJsonObject, TMap<int32, FText>* InLookupTablePtr, FString InSchemaName);
+	FImaginaryPin(FImaginaryFiBDataWeakPtr InOuter, TSharedPtr< FJsonObject > InUnparsedJsonObject, TMap<int32, FText>* InLookupTablePtr, FString InSchemaName);
 
 	/** FImaginaryFiBData Interface */
 	virtual bool IsCompatibleWithFilter(ESearchQueryFilter InSearchQueryFilter) const override;
diff --git a/Engine/Source/Editor/LevelEditor/Private/LevelEditorCreateActorMenu.cpp b/Engine/Source/Editor/LevelEditor/Private/LevelEditorCreateActorMenu.cpp
--- a/Engine/Source/Editor/LevelEditor/Private/LevelEditorCreateActorMenu.cpp
+++ b/Engine/Source/Editor/LevelEditor/Private/LevelEditorCreateActorMenu.cpp
@@ -240,12 +240,9 @@ static void GetContentBrowserSelectionFactoryMenuEntries( FAssetData& TargetAsse
 	{
 		// For blueprints, attempt to determine placeability from its tag information
 
-		const FName NativeParentClassTag = TEXT("NativeParentClass");
-		const FName ClassFlagsTag = TEXT("ClassFlags");
-
 		FString TagValue;
 
-		if ( TargetAssetData.GetTagValue( NativeParentClassTag, TagValue ) && !TagValue.IsEmpty() )
+		if ( TargetAssetData.GetTagValue( FBlueprintTags::NativeParentClassPath, TagValue ) && !TagValue.IsEmpty() )
 		{
 			// If the native parent class can't be placed, neither can the blueprint
 
@@ -256,7 +253,7 @@ static void GetContentBrowserSelectionFactoryMenuEntries( FAssetData& TargetAsse
 			bPlaceable = AssetSelectionUtils::IsClassPlaceable( NativeParentClass );
 		}
 		
-		if ( bPlaceable && TargetAssetData.GetTagValue( ClassFlagsTag, TagValue ) && !TagValue.IsEmpty() )
+		if ( bPlaceable && TargetAssetData.GetTagValue( FBlueprintTags::ClassFlags, TagValue ) && !TagValue.IsEmpty() )
 		{
 			// Check to see if this class is placeable from its class flags
 
diff --git a/Engine/Source/Editor/LevelEditor/Private/SLevelViewport.cpp b/Engine/Source/Editor/LevelEditor/Private/SLevelViewport.cpp
--- a/Engine/Source/Editor/LevelEditor/Private/SLevelViewport.cpp
+++ b/Engine/Source/Editor/LevelEditor/Private/SLevelViewport.cpp
@@ -567,10 +567,6 @@ bool SLevelViewport::HandleDragObjects(const FGeometry& MyGeometry, const FDragD
 			new(SelectedAssetDatas)FAssetData(ClassOperation->ClassesToDrop[DroppedAssetIdx].Get());
 		}
 	}
-	else if (Operation->IsOfType<FUnloadedClassDragDropOp>())
-	{
-		bValidDrag = true;
-	}
 	else if (Operation->IsOfType<FExportTextDragDropOp>())
 	{
 		bValidDrag = true;
@@ -708,56 +704,6 @@ bool SLevelViewport::HandlePlaceDraggedObjects(const FGeometry& MyGeometry, cons
 
 		bValidDrop = true;
 	}
-	else if (Operation->IsOfType<FUnloadedClassDragDropOp>())
-	{
-		TSharedPtr<FUnloadedClassDragDropOp> DragDropOp = StaticCastSharedPtr<FUnloadedClassDragDropOp>( Operation );
-
-		DroppedObjects.Empty();
-
-		// Check if the asset is loaded, used to see if the context menu should be available
-		bAllAssetWereLoaded = true;
-
-		TArray< FClassPackageData >& AssetArray = *(DragDropOp->AssetsToDrop.Get());
-		for (int32 DroppedAssetIdx = 0; DroppedAssetIdx < AssetArray.Num(); ++DroppedAssetIdx)
-		{
-			bValidDrop = true;
-
-			FString& AssetName = AssetArray[DroppedAssetIdx].AssetName;
-
-			// Check to see if the asset can be found, otherwise load it.
-			UObject* Object = FindObject<UObject>(NULL, *AssetName);
-			if(Object == NULL)
-			{
-				// Check to see if the dropped asset was a blueprint
-				const FString& PackageName = AssetArray[DroppedAssetIdx].GeneratedPackageName;
-				Object = FindObject<UObject>(NULL, *FString::Printf(TEXT("%s.%s"), *PackageName, *AssetName));
-
-				if ( Object == NULL )
-				{
-					// Load the package.
-					GWarn->BeginSlowTask( LOCTEXT("OnDrop_FullyLoadPackage", "Fully Loading Package For Drop"), true, false );
-					UPackage* Package = LoadPackage(NULL, *PackageName, LOAD_NoRedirects );
-					if (Package)
-					{
-						Package->FullyLoad();
-					}
-					GWarn->EndSlowTask();
-
-					Object = FindObject<UObject>(Package, *AssetName);
-				}
-			}
-
-			// Check again if it has been loaded, if not, mark that all were not loaded and move on.
-			if(Object)
-			{
-				DroppedObjects.Add(Object);
-			}
-			else
-			{	
-				bAllAssetWereLoaded = false;
-			}
-		}
-	}
 	else if (Operation->IsOfType<FAssetDragDropOp>())
 	{
 		bValidDrop = true;
diff --git a/Engine/Source/Editor/MovieSceneTools/Private/TrackEditors/CameraShakeTrackEditor.cpp b/Engine/Source/Editor/MovieSceneTools/Private/TrackEditors/CameraShakeTrackEditor.cpp
--- a/Engine/Source/Editor/MovieSceneTools/Private/TrackEditors/CameraShakeTrackEditor.cpp
+++ b/Engine/Source/Editor/MovieSceneTools/Private/TrackEditors/CameraShakeTrackEditor.cpp
@@ -25,7 +25,6 @@
 
 #define LOCTEXT_NAMESPACE "FCameraShakeTrackEditor"
 
-static const FName ParentClassTagName(TEXT("ParentClass"));
 static const FString CameraShakeClassPath(TEXT("Class'/Script/Engine.CameraShake'"));
 
 
@@ -143,7 +142,7 @@ void FCameraShakeTrackEditor::BuildObjectBindingTrackMenu(FMenuBuilder& MenuBuil
 			Filter.ClassNames.Add(UBlueprint::StaticClass()->GetFName());
 
 			TMultiMap<FName, FString> TagsAndValuesFilter;
-			TagsAndValuesFilter.Add(ParentClassTagName, CameraShakeClassPath);
+			TagsAndValuesFilter.Add(FBlueprintTags::ParentClassPath, CameraShakeClassPath);
 			Filter.TagsAndValues = TagsAndValuesFilter;
 
 			AssetRegistryModule.Get().GetAssets(Filter, AssetDataList);
@@ -174,7 +173,7 @@ void FCameraShakeTrackEditor::AddCameraShakeSubMenu(FMenuBuilder& MenuBuilder, F
 		AssetPickerConfig.bAllowNullSelection = false;
 		AssetPickerConfig.InitialAssetViewType = EAssetViewType::List;
 		AssetPickerConfig.Filter.ClassNames.Add(UBlueprint::StaticClass()->GetFName());
-		AssetPickerConfig.Filter.TagsAndValues.Add(ParentClassTagName, CameraShakeClassPath);
+		AssetPickerConfig.Filter.TagsAndValues.Add(FBlueprintTags::ParentClassPath, CameraShakeClassPath);
 	}
 
 	FContentBrowserModule& ContentBrowserModule = FModuleManager::Get().LoadModuleChecked<FContentBrowserModule>(TEXT("ContentBrowser"));
diff --git a/Engine/Source/Editor/Persona/Private/PersonaModule.cpp b/Engine/Source/Editor/Persona/Private/PersonaModule.cpp
--- a/Engine/Source/Editor/Persona/Private/PersonaModule.cpp
+++ b/Engine/Source/Editor/Persona/Private/PersonaModule.cpp
@@ -104,12 +104,11 @@ void FPersonaModule::StartupModule()
 		TArray<FAssetData> AssetData;
 		AssetRegistryModule.Get().GetAssetsByClass(UBlueprint::StaticClass()->GetFName(), AssetData);
 
-		const FName BPParentClassName( TEXT( "ParentClass" ) );
 		const FString BPAnimNotify( TEXT("Class'/Script/Engine.AnimNotify'" ));
 
 		for (int32 AssetIndex = 0; AssetIndex < AssetData.Num(); ++AssetIndex)
 		{
-			FString TagValue = AssetData[ AssetIndex ].GetTagValueRef<FString>(BPParentClassName);
+			FString TagValue = AssetData[ AssetIndex ].GetTagValueRef<FString>(FBlueprintTags::ParentClassPath);
 			if (TagValue == BPAnimNotify)
 			{
 				FString BlueprintPath = AssetData[AssetIndex].ObjectPath.ToString();
diff --git a/Engine/Source/Editor/Persona/Private/SAnimNotifyPanel.cpp b/Engine/Source/Editor/Persona/Private/SAnimNotifyPanel.cpp
--- a/Engine/Source/Editor/Persona/Private/SAnimNotifyPanel.cpp
+++ b/Engine/Source/Editor/Persona/Private/SAnimNotifyPanel.cpp
@@ -3781,7 +3781,7 @@ void SAnimNotifyTrack::GetNotifyMenuData(TArray<FAssetData>& NotifyAssetData, TA
 		// this functionality is only available in native class
 		// so we don't have to call BP function but just call native on the check of validity
 		FString NativeParentClassName;
-		if (NotifyData.GetTagValue("NativeParentClass", NativeParentClassName))
+		if (NotifyData.GetTagValue(FBlueprintTags::NativeParentClassPath, NativeParentClassName))
 		{
 			UObject* Outer = nullptr;
 			ResolveName(Outer, NativeParentClassName, false, false);
@@ -4826,8 +4826,6 @@ void SAnimNotifyPanel::OnGetNotifyBlueprintData(TArray<FAssetData>& OutNotifyDat
 	// Collect a full list of assets with the specified class
 	AssetRegistryModule.Get().GetAssetsByClass(UBlueprint::StaticClass()->GetFName(), AssetDataList);
 
-	static const FName BPParentClassName(TEXT("ParentClass"));
-	static const FName BPGenClassName(TEXT("GeneratedClass"));
 
 	int32 BeginClassCount = InOutAllowedClassNames->Num();
 	int32 CurrentClassCount = -1;
@@ -4839,12 +4837,12 @@ void SAnimNotifyPanel::OnGetNotifyBlueprintData(TArray<FAssetData>& OutNotifyDat
 		for(int32 AssetIndex = 0; AssetIndex < AssetDataList.Num(); ++AssetIndex)
 		{
 			FAssetData& AssetData = AssetDataList[AssetIndex];
-			FString TagValue = AssetData.GetTagValueRef<FString>(BPParentClassName);
+			FString TagValue = AssetData.GetTagValueRef<FString>(FBlueprintTags::ParentClassPath);
 
 			if(InOutAllowedClassNames->Contains(TagValue))
 			{
-				FString GenClass = AssetData.GetTagValueRef<FString>(BPGenClassName);
-				const uint32 ClassFlags = AssetData.GetTagValueRef<uint32>("ClassFlags");
+				FString GenClass = AssetData.GetTagValueRef<FString>(FBlueprintTags::GeneratedClassPath);
+				const uint32 ClassFlags = AssetData.GetTagValueRef<uint32>(FBlueprintTags::ClassFlags);
 				if (ClassFlags & CLASS_Abstract)
 				{
 					continue;
diff --git a/Engine/Source/Editor/PropertyEditor/Private/UserInterface/PropertyEditor/SPropertyEditorClass.cpp b/Engine/Source/Editor/PropertyEditor/Private/UserInterface/PropertyEditor/SPropertyEditorClass.cpp
--- a/Engine/Source/Editor/PropertyEditor/Private/UserInterface/PropertyEditor/SPropertyEditorClass.cpp
+++ b/Engine/Source/Editor/PropertyEditor/Private/UserInterface/PropertyEditor/SPropertyEditorClass.cpp
@@ -7,6 +7,7 @@
 #include "Widgets/Layout/SBox.h"
 
 #include "DragAndDrop/ClassDragDropOp.h"
+#include "DragAndDrop/AssetDragDropOp.h"
 #include "ClassViewerModule.h"
 #include "ClassViewerFilter.h"
 
@@ -266,57 +267,57 @@ FReply SPropertyEditorClass::OnDrop(const FGeometry& MyGeometry, const FDragDrop
 	if (ClassOperation.IsValid())
 	{
 		// We can only drop one item into the combo box, so drop the first one.
-		FString AssetName = ClassOperation->ClassesToDrop[0]->GetName();
+		FString ClassPath = ClassOperation->ClassesToDrop[0]->GetPathName();
 
 		// Set the property, it will be verified as valid.
-		SendToObjects(AssetName);
+		SendToObjects(ClassPath);
 
 		return FReply::Handled();
 	}
 	
-	TSharedPtr<FUnloadedClassDragDropOp> UnloadedClassOp = DragDropEvent.GetOperationAs<FUnloadedClassDragDropOp>();
+	TSharedPtr<FAssetDragDropOp> UnloadedClassOp = DragDropEvent.GetOperationAs<FAssetDragDropOp>();
 	if (UnloadedClassOp.IsValid())
 	{
-		// Check if the asset is loaded, used to see if the context menu should be available
 		bool bAllAssetWereLoaded = true;
 
-		TArray<FClassPackageData>& AssetArray = *(UnloadedClassOp->AssetsToDrop.Get());
+		FString AssetPath;
 
-		// We can only drop one item into the combo box, so drop the first one.
-		FString& AssetName = AssetArray[0].AssetName;
+		// Find the class/blueprint path
+		if (UnloadedClassOp->HasAssets())
+		{
+			AssetPath = UnloadedClassOp->GetAssets()[0].ObjectPath.ToString();
+		}
+		else if (UnloadedClassOp->HasAssetPaths())
+		{
+			AssetPath = UnloadedClassOp->GetAssetPaths()[0];
+		}
 
 		// Check to see if the asset can be found, otherwise load it.
-		UObject* Object = FindObject<UObject>(NULL, *AssetName);
-		if(Object == NULL)
+		UObject* Object = FindObject<UObject>(nullptr, *AssetPath);
+		if(Object == nullptr)
 		{
-			// Check to see if the dropped asset was a blueprint
-			const FString& PackageName = AssetArray[0].GeneratedPackageName;
-			Object = FindObject<UObject>(NULL, *FString::Printf(TEXT("%s.%s"), *PackageName, *AssetName));
+			// Load the package.
+			GWarn->BeginSlowTask(LOCTEXT("OnDrop_LoadPackage", "Fully Loading Package For Drop"), true, false);
 
-			if(Object == NULL)
-			{
-				// Load the package.
-				GWarn->BeginSlowTask(LOCTEXT("OnDrop_LoadPackage", "Fully Loading Package For Drop"), true, false);
-				UPackage* Package = LoadPackage(NULL, *PackageName, LOAD_NoRedirects );
-				if(Package)
-				{
-					Package->FullyLoad();
-				}
-				GWarn->EndSlowTask();
-
-				Object = FindObject<UObject>(Package, *AssetName);
-			}
+			Object = LoadObject<UObject>(nullptr, *AssetPath);
 
-			if(Object->IsA(UBlueprint::StaticClass()))
+			GWarn->EndSlowTask();
+		}
+
+		if (UClass* Class = Cast<UClass>(Object))
+		{
+			// This was pointing to a class directly
+			SendToObjects(Class->GetPathName());
+		}
+		else if (UBlueprint* Blueprint = Cast<UBlueprint>(Object))
+		{
+			if (Blueprint->GeneratedClass)
 			{
-				// Get the default object from the generated class.
-				Object = Cast<UBlueprint>(Object)->GeneratedClass->GetDefaultObject();
+				// This was pointing to a blueprint, get generated class
+				SendToObjects(Blueprint->GeneratedClass->GetPathName());
 			}
 		}
 
-		// Set the property, it will be verified as valid.
-		SendToObjects(AssetName);
-
 		return FReply::Handled();
 	}
 
diff --git a/Engine/Source/Editor/Sequencer/Private/SSequencer.cpp b/Engine/Source/Editor/Sequencer/Private/SSequencer.cpp
--- a/Engine/Source/Editor/Sequencer/Private/SSequencer.cpp
+++ b/Engine/Source/Editor/Sequencer/Private/SSequencer.cpp
@@ -1565,7 +1565,6 @@ FReply SSequencer::OnDragOver( const FGeometry& MyGeometry, const FDragDropEvent
 	if (Operation.IsValid() && (
 		Operation->IsOfType<FAssetDragDropOp>() ||
 		Operation->IsOfType<FClassDragDropOp>() ||
-		Operation->IsOfType<FUnloadedClassDragDropOp>() ||
 		Operation->IsOfType<FActorDragDropGraphEdOp>() ) )
 	{
 		bIsDragSupported = true;
@@ -1608,13 +1607,6 @@ FReply SSequencer::OnDrop( const FGeometry& MyGeometry, const FDragDropEvent& Dr
 			OnClassesDropped( *DragDropOp );
 			bWasDropHandled = true;
 		}
-		else if( Operation->IsOfType<FUnloadedClassDragDropOp>() )
-		{
-			const auto& DragDropOp = StaticCastSharedPtr<FUnloadedClassDragDropOp>( Operation );
-
-			OnUnloadedClassesDropped( *DragDropOp );
-			bWasDropHandled = true;
-		}
 		else if( Operation->IsOfType<FActorDragDropGraphEdOp>() )
 		{
 			const auto& DragDropOp = StaticCastSharedPtr<FActorDragDropGraphEdOp>( Operation );
@@ -1654,6 +1646,21 @@ void SSequencer::OnAssetsDropped( const FAssetDragDropOp& DragDropOp )
 	bool bObjectAdded = false;
 	TArray< UObject* > DroppedObjects;
 	bool bAllAssetsWereLoaded = true;
+	bool bNeedsLoad = false;
+
+	for (const FAssetData& AssetData : DragDropOp.GetAssets())
+	{
+		if (!AssetData.IsAssetLoaded())
+		{
+			bNeedsLoad = true;
+			break;
+		}
+	}
+
+	if (bNeedsLoad)
+	{
+		GWarn->BeginSlowTask(LOCTEXT("OnDrop_FullyLoadPackage", "Fully Loading Package For Drop"), true, false);
+	}
 
 	for (const FAssetData& AssetData : DragDropOp.GetAssets())
 	{
@@ -1669,6 +1676,11 @@ void SSequencer::OnAssetsDropped( const FAssetDragDropOp& DragDropOp )
 		}
 	}
 
+	if (bNeedsLoad)
+	{
+		GWarn->EndSlowTask();
+	}
+
 	const TSet< TSharedRef<FSequencerDisplayNode> >& SelectedNodes = SequencerPtr.Pin()->GetSelection().GetSelectedOutlinerNodes();
 	FGuid TargetObjectGuid;
 	// if exactly one object node is selected, we have a target object guid
@@ -1737,59 +1749,11 @@ void SSequencer::OnClassesDropped( const FClassDragDropOp& DragDropOp )
 	}
 }
 
-
-void SSequencer::OnUnloadedClassesDropped( const FUnloadedClassDragDropOp& DragDropOp )
-{
-	FSequencer& SequencerRef = *SequencerPtr.Pin();
-	for( auto ClassDataIter = DragDropOp.AssetsToDrop->CreateConstIterator(); ClassDataIter; ++ClassDataIter )
-	{
-		auto& ClassData = *ClassDataIter;
-
-		// Check to see if the asset can be found, otherwise load it.
-		UObject* Object = FindObject<UObject>( nullptr, *ClassData.AssetName );
-		if( Object == nullptr )
-		{
-			Object = FindObject<UObject>(nullptr, *FString::Printf(TEXT("%s.%s"), *ClassData.GeneratedPackageName, *ClassData.AssetName));
-		}
-
-		if( Object == nullptr )
-		{
-			// Load the package.
-			GWarn->BeginSlowTask( LOCTEXT("OnDrop_FullyLoadPackage", "Fully Loading Package For Drop"), true, false );
-			UPackage* Package = LoadPackage(nullptr, *ClassData.GeneratedPackageName, LOAD_NoRedirects );
-			if( Package != nullptr )
-			{
-				Package->FullyLoad();
-			}
-			GWarn->EndSlowTask();
-
-			Object = FindObject<UObject>(Package, *ClassData.AssetName);
-		}
-
-		if( Object != nullptr )
-		{
-			// Check to see if the dropped asset was a blueprint
-			if(Object->IsA(UBlueprint::StaticClass()))
-			{
-				// Get the default object from the generated class.
-				Object = Cast<UBlueprint>(Object)->GeneratedClass->GetDefaultObject();
-			}
-		}
-
-		if( Object != nullptr )
-		{
-			FGuid NewGuid = SequencerRef.MakeNewSpawnable( *Object );
-		}
-	}
-}
-
-
 void SSequencer::OnActorsDropped( FActorDragDropGraphEdOp& DragDropOp )
 {
 	SequencerPtr.Pin()->OnActorsDropped( DragDropOp.Actors );
 }
 
-
 void SSequencer::OnCrumbClicked(const FSequencerBreadcrumb& Item)
 {
 	if (Item.BreadcrumbType != FSequencerBreadcrumb::ShotType)
diff --git a/Engine/Source/Editor/Sequencer/Private/SSequencer.h b/Engine/Source/Editor/Sequencer/Private/SSequencer.h
--- a/Engine/Source/Editor/Sequencer/Private/SSequencer.h
+++ b/Engine/Source/Editor/Sequencer/Private/SSequencer.h
@@ -21,7 +21,6 @@ class FAssetDragDropOp;
 class FClassDragDropOp;
 class FMovieSceneClipboard;
 class FSequencerTimeSliderController;
-class FUnloadedClassDragDropOp;
 class FVirtualTrackArea;
 class ISequencerEditTool;
 class SSequencerGotoBox;
@@ -376,14 +375,7 @@ private:
 	 * @param	DragDropOp	Information about the class(es) that were dropped
 	 */
 	void OnClassesDropped(const FClassDragDropOp& DragDropOp);
-	
-	/**
-	 * Called when one or more unloaded classes are dropped into the widget
-	 *
-	 * @param	DragDropOp	Information about the unloaded class(es) that were dropped
-	 */
-	void OnUnloadedClassesDropped(const FUnloadedClassDragDropOp& DragDropOp);
-	
+		
 	/**
 	 * Called when one or more actors are dropped into the widget
 	 *
diff --git a/Engine/Source/Editor/SkeletalMeshEditor/Private/SControlRigMappingWindow.cpp b/Engine/Source/Editor/SkeletalMeshEditor/Private/SControlRigMappingWindow.cpp
--- a/Engine/Source/Editor/SkeletalMeshEditor/Private/SControlRigMappingWindow.cpp
+++ b/Engine/Source/Editor/SkeletalMeshEditor/Private/SControlRigMappingWindow.cpp
@@ -314,7 +314,7 @@ void SControlRigMappingWindow::OnAssetSelectedFromMeshPicker(const FAssetData& A
 bool SControlRigMappingWindow::OnShouldFilterAnimAsset(const FAssetData& AssetData) const
 {
 	FString ParentClassName;
-	if (AssetData.GetTagValue(FName(TEXT("NativeParentClass")), ParentClassName))
+	if (AssetData.GetTagValue(FBlueprintTags::NativeParentClassPath, ParentClassName))
 	{
 		if (ParentClassName.IsEmpty() == false)
 		{
diff --git a/Engine/Source/Editor/UnrealEd/Private/AssetDeleteModel.cpp b/Engine/Source/Editor/UnrealEd/Private/AssetDeleteModel.cpp
--- a/Engine/Source/Editor/UnrealEd/Private/AssetDeleteModel.cpp
+++ b/Engine/Source/Editor/UnrealEd/Private/AssetDeleteModel.cpp
@@ -356,9 +356,8 @@ bool FAssetDeleteModel::CanReplaceReferencesWith( const FAssetData& InAssetData
 	if (FirstPendingDelete->IsChildOf(UBlueprint::StaticClass()) && AssetDataClass->IsChildOf(UBlueprint::StaticClass()))
 	{
 		// Get BP native parent classes
-		static const FName ParentClassName("ParentClass");
 		UClass* OriginalBPParentClass = CastChecked<UBlueprint>(PendingDeletes[0]->GetObject())->ParentClass;
-		const FString BPClassNameToTest = InAssetData.GetTagValueRef<FString>(ParentClassName);
+		const FString BPClassNameToTest = InAssetData.GetTagValueRef<FString>(FBlueprintTags::ParentClassPath);
 
 		if (!BPClassNameToTest.IsEmpty())
 		{
diff --git a/Engine/Source/Editor/UnrealEd/Private/ClassIconFinder.cpp b/Engine/Source/Editor/UnrealEd/Private/ClassIconFinder.cpp
--- a/Engine/Source/Editor/UnrealEd/Private/ClassIconFinder.cpp
+++ b/Engine/Source/Editor/UnrealEd/Private/ClassIconFinder.cpp
@@ -121,9 +121,6 @@ const UClass* FClassIconFinder::GetIconClassForAssetData(const FAssetData& InAss
 	
 	if ( AssetClass == UBlueprint::StaticClass() )
 	{
-		static const FName NativeParentClassTag = "NativeParentClass";
-		static const FName ParentClassTag = "ParentClass";
-
 		if ( bOutIsClassType )
 		{
 			*bOutIsClassType = true;
@@ -131,9 +128,9 @@ const UClass* FClassIconFinder::GetIconClassForAssetData(const FAssetData& InAss
 
 		// We need to use the asset data to get the parent class as the blueprint may not be loaded
 		FString ParentClassName;
-		if ( !InAssetData.GetTagValue(NativeParentClassTag, ParentClassName) )
+		if ( !InAssetData.GetTagValue(FBlueprintTags::NativeParentClassPath, ParentClassName) )
 		{
-			InAssetData.GetTagValue(ParentClassTag, ParentClassName);
+			InAssetData.GetTagValue(FBlueprintTags::ParentClassPath, ParentClassName);
 		}
 		if ( !ParentClassName.IsEmpty() )
 		{
diff --git a/Engine/Source/Editor/UnrealEd/Private/Factories/ActorFactory.cpp b/Engine/Source/Editor/UnrealEd/Private/Factories/ActorFactory.cpp
--- a/Engine/Source/Editor/UnrealEd/Private/Factories/ActorFactory.cpp
+++ b/Engine/Source/Editor/UnrealEd/Private/Factories/ActorFactory.cpp
@@ -1347,7 +1347,7 @@ bool UActorFactoryBlueprint::CanCreateActorFrom( const FAssetData& AssetData, FT
 		return false;
 	}
 
-	const FString ParentClassPath = AssetData.GetTagValueRef<FString>( "ParentClass" );
+	const FString ParentClassPath = AssetData.GetTagValueRef<FString>(FBlueprintTags::ParentClassPath);
 	if ( ParentClassPath.IsEmpty() )
 	{
 		OutErrorMsg = NSLOCTEXT("CanCreateActor", "NoBlueprint", "No Blueprint was specified, or the specified Blueprint needs to be compiled.");
@@ -1393,7 +1393,7 @@ AActor* UActorFactoryBlueprint::GetDefaultActor( const FAssetData& AssetData )
 		return NULL;
 	}
 
-	const FString GeneratedClassPath = AssetData.GetTagValueRef<FString>("GeneratedClass");
+	const FString GeneratedClassPath = AssetData.GetTagValueRef<FString>(FBlueprintTags::GeneratedClassPath);
 	if ( GeneratedClassPath.IsEmpty() )
 	{
 		return NULL;
diff --git a/Engine/Source/Editor/UnrealEd/Private/Kismet2/KismetReinstanceUtilities.cpp b/Engine/Source/Editor/UnrealEd/Private/Kismet2/KismetReinstanceUtilities.cpp
--- a/Engine/Source/Editor/UnrealEd/Private/Kismet2/KismetReinstanceUtilities.cpp
+++ b/Engine/Source/Editor/UnrealEd/Private/Kismet2/KismetReinstanceUtilities.cpp
@@ -95,7 +95,10 @@ struct FReplaceReferenceHelper
 
 			for (UObject* Obj : Targets)
 			{
-				if (!ObjectsToReplace.Contains(Obj)) // Don't bother trying to fix old objects, this would break them
+				// Make sure we don't update properties in old objects, as they
+				// may take ownership of objects referenced in new objects (e.g.
+				// delete components owned by new actors)
+				if (!ObjectsToReplace.Contains(Obj))
 				{
 					// The class for finding and replacing weak references.
 					// We can't relay on "standard" weak references replacement as
@@ -1852,8 +1855,8 @@ void FBlueprintCompileReinstancer::ReplaceInstancesOfClass_Inner(TMap<UClass*, U
 	}
 
 	USelection* SelectedActors = nullptr;
+	TArray<UObject*> ObjectsReplaced;
 	bool bSelectionChanged = false;
-	TArray<UObject*> ObjectsToReplace;
 	const bool bLogConversions = false; // for debugging
 
 	// Map of old objects to new objects
@@ -1904,6 +1907,8 @@ void FBlueprintCompileReinstancer::ReplaceInstancesOfClass_Inner(TMap<UClass*, U
 	};
 
 	{
+		TArray<UObject*> ObjectsToReplace;
+
 		BP_SCOPED_COMPILER_EVENT_STAT(EKismetReinstancerStats_ReplaceInstancesOfClass);
 		if(GEditor && GEditor->GetSelectedActors())
 		{
@@ -1927,12 +1932,13 @@ void FBlueprintCompileReinstancer::ReplaceInstancesOfClass_Inner(TMap<UClass*, U
 #endif
 			{
 				const bool bIncludeDerivedClasses = false;
-				ObjectsToReplace.Empty();
+				ObjectsToReplace.Reset();
 				GetObjectsOfClass(OldClass, ObjectsToReplace, bIncludeDerivedClasses);
 				// Then fix 'real' (non archetype) instances of the class
 				for (int32 OldObjIndex = 0; OldObjIndex < ObjectsToReplace.Num(); ++OldObjIndex)
 				{
 					UObject* OldObject = ObjectsToReplace[OldObjIndex];
+					
 					AActor* OldActor = Cast<AActor>(OldObject);
 
 					// Skip archetype instances, EXCEPT for component templates and child actor templates
@@ -1950,7 +1956,8 @@ void FBlueprintCompileReinstancer::ReplaceInstancesOfClass_Inner(TMap<UClass*, U
 						UObject* NewUObject = nullptr;
 						ReplaceObjectHelper(OldObject, OldClass, NewUObject, NewClass, OldToNewInstanceMap, OldToNewNameMap, OldObjIndex, ObjectsToReplace, PotentialEditorsForRefreshing, OwnersToRerunConstructionScript, &FDirectAttachChildrenAccessor::Get, bIsComponent, bArchetypesAreUpToDate);
 						UpdateObjectBeingDebugged(OldObject, NewUObject);
-						
+						ObjectsReplaced.Add(OldObject);
+
 						if (bLogConversions)
 						{
 							UE_LOG(LogBlueprint, Log, TEXT("Converted instance '%s' to '%s'"), *GetPathNameSafe(OldObject), *GetPathNameSafe(NewUObject));
@@ -1971,7 +1978,7 @@ void FBlueprintCompileReinstancer::ReplaceInstancesOfClass_Inner(TMap<UClass*, U
 
 			{
 				const bool bIncludeDerivedClasses = false;
-				ObjectsToReplace.Empty();
+				ObjectsToReplace.Reset();
 				GetObjectsOfClass(OldClass, ObjectsToReplace, bIncludeDerivedClasses);
 
 				// store old attachment data before we mess with components, etc:
@@ -2016,7 +2023,8 @@ void FBlueprintCompileReinstancer::ReplaceInstancesOfClass_Inner(TMap<UClass*, U
 							ReplaceObjectHelper(OldObject, OldClass, NewUObject, NewClass, OldToNewInstanceMap, OldToNewNameMap, OldObjIndex, ObjectsToReplace, PotentialEditorsForRefreshing, OwnersToRerunConstructionScript, &FDirectAttachChildrenAccessor::Get, false, bArchetypesAreUpToDate);
 						}
 						UpdateObjectBeingDebugged(OldObject, NewUObject);
-						
+						ObjectsReplaced.Add(OldObject);
+
 						if (bLogConversions)
 						{
 							UE_LOG(LogBlueprint, Log, TEXT("Converted instance '%s' to '%s'"), *GetPathNameSafe(OldObject), *GetPathNameSafe(NewUObject));
@@ -2059,7 +2067,7 @@ void FBlueprintCompileReinstancer::ReplaceInstancesOfClass_Inner(TMap<UClass*, U
 
 			if (bClassObjectReplaced)
 			{
-				FReplaceReferenceHelper::IncludeClass(OldClass, NewClass, OldToNewInstanceMap, SourceObjects, ObjectsToReplace);
+				FReplaceReferenceHelper::IncludeClass(OldClass, NewClass, OldToNewInstanceMap, SourceObjects, ObjectsReplaced);
 			}
 		}
 	}
@@ -2080,16 +2088,16 @@ void FBlueprintCompileReinstancer::ReplaceInstancesOfClass_Inner(TMap<UClass*, U
 		// FArchiveReplaceObjectRef to run construction-scripts).
 		for (FActorReplacementHelper& ReplacementActor : ReplacementActors)
 		{
-			ReplacementActor.Finalize(ObjectRemappingHelper.ReplacedObjects, ObjectsThatShouldUseOldStuff, ObjectsToReplace, ReinstancedObjectsWeakReferenceMap);
+			ReplacementActor.Finalize(ObjectRemappingHelper.ReplacedObjects, ObjectsThatShouldUseOldStuff, ObjectsReplaced, ReinstancedObjectsWeakReferenceMap);
 		}
 
 		for (FActorReplacementHelper& ReplacementActor : ReplacementActors)
 		{
-			ReplacementActor.ApplyAttachments(ObjectRemappingHelper.ReplacedObjects, ObjectsThatShouldUseOldStuff, ObjectsToReplace, ReinstancedObjectsWeakReferenceMap);
+			ReplacementActor.ApplyAttachments(ObjectRemappingHelper.ReplacedObjects, ObjectsThatShouldUseOldStuff, ObjectsReplaced, ReinstancedObjectsWeakReferenceMap);
 		}
 	}
 
-	FReplaceReferenceHelper::FindAndReplaceReferences(SourceObjects, ObjectsThatShouldUseOldStuff, ObjectsToReplace, OldToNewInstanceMap, ReinstancedObjectsWeakReferenceMap);
+	FReplaceReferenceHelper::FindAndReplaceReferences(SourceObjects, ObjectsThatShouldUseOldStuff, ObjectsReplaced, OldToNewInstanceMap, ReinstancedObjectsWeakReferenceMap);
 
 	if(SelectedActors)
 	{
diff --git a/Engine/Source/Editor/UnrealEd/Public/DragAndDrop/ClassDragDropOp.h b/Engine/Source/Editor/UnrealEd/Public/DragAndDrop/ClassDragDropOp.h
--- a/Engine/Source/Editor/UnrealEd/Public/DragAndDrop/ClassDragDropOp.h
+++ b/Engine/Source/Editor/UnrealEd/Public/DragAndDrop/ClassDragDropOp.h
@@ -76,51 +76,3 @@ protected:
 	/** Adding hint text to all drag drop operations that drop locations could set, and the decorator could report. */
 	FText HintText;
 };
-
-struct FClassPackageData
-{
-	FString AssetName;
-	FString GeneratedPackageName;
-
-	FClassPackageData(const FString& InAssetName, const FString& InGeneratedPackageName)
-	{
-		AssetName = InAssetName;
-		GeneratedPackageName = InGeneratedPackageName;
-	}
-};
-
-class FUnloadedClassDragDropOp : public FDragDropOperation
-{
-public:
-	DRAG_DROP_OPERATOR_TYPE(FUnloadedClassDragDropOp, FDragDropOperation)
-
-	/** The assets to be dropped. */
-	TSharedPtr< TArray< FClassPackageData > >	AssetsToDrop;
-	
-	/** The widget decorator to use */
-	virtual TSharedPtr<SWidget> GetDefaultDecorator() const override
-	{
-		// Create hover widget
-		return SNew(SBorder)
-			.BorderImage(FEditorStyle::GetBrush("Graph.ConnectorFeedback.Border"))
-			.Content()
-			[			
-				SNew(SHorizontalBox)
-				+ SHorizontalBox::Slot()
-				.AutoWidth()
-				[
-					SNew(STextBlock) 
-					.Text( FText::FromString((*AssetsToDrop.Get())[0].AssetName) )
-				]
-			];
-	}
-
-	static TSharedRef<FUnloadedClassDragDropOp> New(FClassPackageData AssetToDrop)
-	{
-		TSharedRef<FUnloadedClassDragDropOp> Operation = MakeShareable(new FUnloadedClassDragDropOp);
-		Operation->AssetsToDrop = MakeShareable(new TArray<FClassPackageData>);
-		Operation->AssetsToDrop->Add(AssetToDrop);
-		Operation->Construct();
-		return Operation;
-	}
-};
diff --git a/Engine/Source/Programs/UnrealBuildTool/Platform/Windows/VCToolChain.cs b/Engine/Source/Programs/UnrealBuildTool/Platform/Windows/VCToolChain.cs
--- a/Engine/Source/Programs/UnrealBuildTool/Platform/Windows/VCToolChain.cs
+++ b/Engine/Source/Programs/UnrealBuildTool/Platform/Windows/VCToolChain.cs
@@ -1644,7 +1644,8 @@ namespace UnrealBuildTool
 			FileReference ResponseFileName = GetResponseFileName(LinkEnvironment, OutputFile);
 			if (!ProjectFileGenerator.bGenerateProjectFiles)
 			{
-				ResponseFile.Create(ResponseFileName, Arguments);
+				FileItem ResponseFile = FileItem.CreateIntermediateTextFile(ResponseFileName, String.Join(Environment.NewLine, Arguments));
+				PrerequisiteItems.Add(ResponseFile);
 			}
 
 			// Create an action that invokes the linker.
diff --git a/Engine/Source/Runtime/AIModule/Classes/BehaviorTree/BTFunctionLibrary.h b/Engine/Source/Runtime/AIModule/Classes/BehaviorTree/BTFunctionLibrary.h
--- a/Engine/Source/Runtime/AIModule/Classes/BehaviorTree/BTFunctionLibrary.h
+++ b/Engine/Source/Runtime/AIModule/Classes/BehaviorTree/BTFunctionLibrary.h
@@ -101,14 +101,14 @@ public:
 	static void SetBlackboardValueAsVector(UBTNode* NodeOwner, const FBlackboardKeySelector& Key, FVector Value);
 
 	/** (DEPRECATED) Use ClearBlackboardValue instead */
-	UFUNCTION(BlueprintCallable, Category=BehaviorTree, Meta=(HidePin="NodeOwner", DefaultToSelf="NodeOwner", DeprecatedFunction, DeprecationMessage="Use ClearBlackboardValue instead."))
+	UFUNCTION(BlueprintCallable, Category="AI|BehaviorTree", Meta=(HidePin="NodeOwner", DefaultToSelf="NodeOwner", DeprecatedFunction, DeprecationMessage="Use ClearBlackboardValue instead."))
 	static void ClearBlackboardValueAsVector(UBTNode* NodeOwner, const FBlackboardKeySelector& Key);
 
-	UFUNCTION(BlueprintCallable, Category = "AI|BehaviorTree", Meta = (HidePin = "NodeOwner", DefaultToSelf = "NodeOwner"))
+	UFUNCTION(BlueprintCallable, Category="AI|BehaviorTree", Meta = (HidePin = "NodeOwner", DefaultToSelf = "NodeOwner"))
 	static void SetBlackboardValueAsRotator(UBTNode* NodeOwner, const FBlackboardKeySelector& Key, FRotator Value);
 
 	/** Resets indicated value to "not set" value, based on values type */
-	UFUNCTION(BlueprintCallable, Category = BehaviorTree, Meta = (HidePin = "NodeOwner", DefaultToSelf = "NodeOwner"))
+	UFUNCTION(BlueprintCallable, Category="AI|BehaviorTree", Meta = (HidePin = "NodeOwner", DefaultToSelf = "NodeOwner"))
 	static void ClearBlackboardValue(UBTNode* NodeOwner, const FBlackboardKeySelector& Key);
 
 	/** Initialize variables marked as "instance memory" and set owning actor for blackboard operations */
diff --git a/Engine/Source/Runtime/AIModule/Private/BehaviorTree/Tasks/BTTask_RotateToFaceBBEntry.cpp b/Engine/Source/Runtime/AIModule/Private/BehaviorTree/Tasks/BTTask_RotateToFaceBBEntry.cpp
--- a/Engine/Source/Runtime/AIModule/Private/BehaviorTree/Tasks/BTTask_RotateToFaceBBEntry.cpp
+++ b/Engine/Source/Runtime/AIModule/Private/BehaviorTree/Tasks/BTTask_RotateToFaceBBEntry.cpp
@@ -39,7 +39,9 @@ namespace
 {
 	FORCEINLINE_DEBUGGABLE float CalculateAngleDifferenceDot(const FVector& VectorA, const FVector& VectorB)
 	{
-		return VectorA.CosineAngle2D(VectorB);
+		return (VectorA.IsNearlyZero() || VectorB.IsNearlyZero())
+            ? 1.f
+            : VectorA.CosineAngle2D(VectorB);
 	}
 }
 
@@ -70,7 +72,7 @@ EBTNodeResult::Type UBTTask_RotateToFaceBBEntry::ExecuteTask(UBehaviorTreeCompon
 		if (ActorValue != NULL)
 		{
 			const float AngleDifference = CalculateAngleDifferenceDot(Pawn->GetActorForwardVector()
-				, (ActorValue->GetActorLocation() - PawnLocation).GetSafeNormal2D());
+				, (ActorValue->GetActorLocation() - PawnLocation));
 			
 			if (AngleDifference >= PrecisionDot)
 			{
@@ -92,7 +94,7 @@ EBTNodeResult::Type UBTTask_RotateToFaceBBEntry::ExecuteTask(UBehaviorTreeCompon
 		if (FAISystem::IsValidLocation(KeyValue))
 		{
 			const float AngleDifference = CalculateAngleDifferenceDot(Pawn->GetActorForwardVector()
-				, (KeyValue - PawnLocation).GetSafeNormal2D());
+				, (KeyValue - PawnLocation));
 
 			if (AngleDifference >= PrecisionDot)
 			{
@@ -201,7 +203,7 @@ void UBTTask_RotateToFaceBBEntry::DescribeRuntimeValues(const UBehaviorTreeCompo
 
 		if (FocalPoint != FAISystem::InvalidLocation)
 		{
-			const float CurrentAngleRadians = CalculateAngleDifferenceDot(PawnDirection, (FocalPoint - AIController->GetPawn()->GetActorLocation()).GetSafeNormal2D());
+			const float CurrentAngleRadians = CalculateAngleDifferenceDot(PawnDirection, (FocalPoint - AIController->GetPawn()->GetActorLocation()));
 			Values.Add(FString::Printf(TEXT("Current angle: %.2f"), FMath::RadiansToDegrees(FMath::Acos(CurrentAngleRadians))));
 		}
 		else
diff --git a/Engine/Source/Runtime/AIModule/Public/BlueprintNodeHelpers.h b/Engine/Source/Runtime/AIModule/Public/BlueprintNodeHelpers.h
--- a/Engine/Source/Runtime/AIModule/Public/BlueprintNodeHelpers.h
+++ b/Engine/Source/Runtime/AIModule/Public/BlueprintNodeHelpers.h
@@ -13,23 +13,23 @@ class UBTNode;
 
 namespace BlueprintNodeHelpers
 {
-	FString CollectPropertyDescription(const UObject* Ob, const UClass* StopAtClass, const TArray<UProperty*>& PropertyData);
-	void CollectPropertyData(const UObject* Ob, const UClass* StopAtClass, TArray<UProperty*>& PropertyData);
-	uint16 GetPropertiesMemorySize(const TArray<UProperty*>& PropertyData);
+	AIMODULE_API FString CollectPropertyDescription(const UObject* Ob, const UClass* StopAtClass, const TArray<UProperty*>& PropertyData);
+	AIMODULE_API void CollectPropertyData(const UObject* Ob, const UClass* StopAtClass, TArray<UProperty*>& PropertyData);
+	AIMODULE_API uint16 GetPropertiesMemorySize(const TArray<UProperty*>& PropertyData);
 
-	void CollectBlackboardSelectors(const UObject* Ob, const UClass* StopAtClass, TArray<FName>& KeyNames);
-	void ResolveBlackboardSelectors(UObject& Ob, const UClass& StopAtClass, const UBlackboardData& BlackboardAsset);
-	bool HasAnyBlackboardSelectors(const UObject* Ob, const UClass* StopAtClass);
+	AIMODULE_API void CollectBlackboardSelectors(const UObject* Ob, const UClass* StopAtClass, TArray<FName>& KeyNames);
+	AIMODULE_API void ResolveBlackboardSelectors(UObject& Ob, const UClass& StopAtClass, const UBlackboardData& BlackboardAsset);
+	AIMODULE_API bool HasAnyBlackboardSelectors(const UObject* Ob, const UClass* StopAtClass);
 
 	AIMODULE_API FString DescribeProperty(const UProperty* Prop, const uint8* PropertyAddr);
-	void DescribeRuntimeValues(const UObject* Ob, const TArray<UProperty*>& PropertyData, TArray<FString>& RuntimeValues);
+	AIMODULE_API void DescribeRuntimeValues(const UObject* Ob, const TArray<UProperty*>& PropertyData, TArray<FString>& RuntimeValues);
 
-	void CopyPropertiesToContext(const TArray<UProperty*>& PropertyData, uint8* ObjectMemory, uint8* ContextMemory);
-	void CopyPropertiesFromContext(const TArray<UProperty*>& PropertyData, uint8* ObjectMemory, uint8* ContextMemory);
+	AIMODULE_API void CopyPropertiesToContext(const TArray<UProperty*>& PropertyData, uint8* ObjectMemory, uint8* ContextMemory);
+	AIMODULE_API void CopyPropertiesFromContext(const TArray<UProperty*>& PropertyData, uint8* ObjectMemory, uint8* ContextMemory);
 
-	bool FindNodeOwner(AActor* OwningActor, UBTNode* Node, UBehaviorTreeComponent*& OwningComp, int32& OwningInstanceIdx);
+	AIMODULE_API bool FindNodeOwner(AActor* OwningActor, UBTNode* Node, UBehaviorTreeComponent*& OwningComp, int32& OwningInstanceIdx);
 
-	void AbortLatentActions(UActorComponent& OwnerOb, const UObject& Ob);
+	AIMODULE_API void AbortLatentActions(UActorComponent& OwnerOb, const UObject& Ob);
 
 	FORCEINLINE bool HasBlueprintFunction(FName FuncName, const UObject& Object, const UClass& StopAtClass)
 	{
diff --git a/Engine/Source/Runtime/AssetRegistry/Private/AssetRegistry.cpp b/Engine/Source/Runtime/AssetRegistry/Private/AssetRegistry.cpp
--- a/Engine/Source/Runtime/AssetRegistry/Private/AssetRegistry.cpp
+++ b/Engine/Source/Runtime/AssetRegistry/Private/AssetRegistry.cpp
@@ -10,6 +10,7 @@
 #include "UObject/UObjectIterator.h"
 #include "UObject/MetaData.h"
 #include "UObject/CoreRedirects.h"
+#include "Blueprint/BlueprintSupport.h"
 #include "AssetRegistryPrivate.h"
 #include "ARFilter.h"
 #include "DependsNode.h"
@@ -1480,8 +1481,8 @@ void UAssetRegistryImpl::Serialize(FArchive& Ar)
 				// Populate the class map if adding blueprint
 				if (ClassGeneratorNames.Contains(AssetData->AssetClass))
 				{
-					const FString GeneratedClass = AssetData->GetTagValueRef<FString>("GeneratedClass");
-					const FString ParentClass = AssetData->GetTagValueRef<FString>("ParentClass");
+					const FString GeneratedClass = AssetData->GetTagValueRef<FString>(FBlueprintTags::GeneratedClassPath);
+					const FString ParentClass = AssetData->GetTagValueRef<FString>(FBlueprintTags::ParentClassPath);
 					if (!GeneratedClass.IsEmpty() && !ParentClass.IsEmpty())
 					{
 						const FName GeneratedClassFName = *ExportTextPathToObjectName(GeneratedClass);
@@ -2002,8 +2003,8 @@ void UAssetRegistryImpl::AddAssetData(FAssetData* AssetData)
 	// Populate the class map if adding blueprint
 	if (ClassGeneratorNames.Contains(AssetData->AssetClass))
 	{
-		const FString GeneratedClass = AssetData->GetTagValueRef<FString>("GeneratedClass");
-		const FString ParentClass = AssetData->GetTagValueRef<FString>("ParentClass");
+		const FString GeneratedClass = AssetData->GetTagValueRef<FString>(FBlueprintTags::GeneratedClassPath);
+		const FString ParentClass = AssetData->GetTagValueRef<FString>(FBlueprintTags::ParentClassPath);
 		if (!GeneratedClass.IsEmpty() && !ParentClass.IsEmpty())
 		{
 			const FName GeneratedClassFName = *ExportTextPathToObjectName(GeneratedClass);
@@ -2018,15 +2019,15 @@ void UAssetRegistryImpl::UpdateAssetData(FAssetData* AssetData, const FAssetData
 	// Update the class map if updating a blueprint
 	if (ClassGeneratorNames.Contains(AssetData->AssetClass))
 	{
-		const FString OldGeneratedClass = AssetData->GetTagValueRef<FString>("GeneratedClass");
+		const FString OldGeneratedClass = AssetData->GetTagValueRef<FString>(FBlueprintTags::GeneratedClassPath);
 		if (!OldGeneratedClass.IsEmpty())
 		{
 			const FName OldGeneratedClassFName = *ExportTextPathToObjectName(OldGeneratedClass);
 			CachedInheritanceMap.Remove(OldGeneratedClassFName);
 		}
 
-		const FString NewGeneratedClass = NewAssetData.GetTagValueRef<FString>("GeneratedClass");
-		const FString NewParentClass = NewAssetData.GetTagValueRef<FString>("ParentClass");
+		const FString NewGeneratedClass = NewAssetData.GetTagValueRef<FString>(FBlueprintTags::GeneratedClassPath);
+		const FString NewParentClass = NewAssetData.GetTagValueRef<FString>(FBlueprintTags::ParentClassPath);
 		if (!NewGeneratedClass.IsEmpty() && !NewParentClass.IsEmpty())
 		{
 			const FName NewGeneratedClassFName = *ExportTextPathToObjectName(*NewGeneratedClass);
@@ -2050,7 +2051,7 @@ bool UAssetRegistryImpl::RemoveAssetData(FAssetData* AssetData)
 		// Remove from the class map if removing a blueprint
 		if (ClassGeneratorNames.Contains(AssetData->AssetClass))
 		{
-			const FString OldGeneratedClass = AssetData->GetTagValueRef<FString>("GeneratedClass");
+			const FString OldGeneratedClass = AssetData->GetTagValueRef<FString>(FBlueprintTags::GeneratedClassPath);
 			if (!OldGeneratedClass.IsEmpty())
 			{
 				const FName OldGeneratedClassFName = *ExportTextPathToObjectName(OldGeneratedClass);
diff --git a/Engine/Source/Runtime/AssetRegistry/Public/AssetData.h b/Engine/Source/Runtime/AssetRegistry/Public/AssetData.h
--- a/Engine/Source/Runtime/AssetRegistry/Public/AssetData.h
+++ b/Engine/Source/Runtime/AssetRegistry/Public/AssetData.h
@@ -166,7 +166,7 @@ public:
 		return ObjectPath != NAME_None;
 	}
 
-	/** Returns true if this asset was found in a UAsset file */
+	/** Returns true if this is the primary asset in a package, true for maps and assets but false for secondary objects like class redirectors */
 	bool IsUAsset() const
 	{
 		return FPackageName::GetLongPackageAssetName(PackageName.ToString()) == AssetName.ToString();
diff --git a/Engine/Source/Runtime/AssetRegistry/Public/AssetRegistryHelpers.h b/Engine/Source/Runtime/AssetRegistry/Public/AssetRegistryHelpers.h
--- a/Engine/Source/Runtime/AssetRegistry/Public/AssetRegistryHelpers.h
+++ b/Engine/Source/Runtime/AssetRegistry/Public/AssetRegistryHelpers.h
@@ -42,7 +42,7 @@ public:
 	UFUNCTION(BlueprintPure, Category = "Asset Data")
 	static bool IsValid(const FAssetData& InAssetData);
 
-	/** Returns true if this asset was found in a UAsset file */
+	/** Returns true if this is the primary asset in a package, true for maps and assets but false for secondary objects like class redirectors */
 	UFUNCTION(BlueprintPure, Category = "Asset Data")
 	static bool IsUAsset(const FAssetData& InAssetData);
 
diff --git a/Engine/Source/Runtime/Core/Private/Misc/CoreGlobals.cpp b/Engine/Source/Runtime/Core/Private/Misc/CoreGlobals.cpp
--- a/Engine/Source/Runtime/Core/Private/Misc/CoreGlobals.cpp
+++ b/Engine/Source/Runtime/Core/Private/Misc/CoreGlobals.cpp
@@ -142,7 +142,6 @@ bool					GIsSlowTask						= false;					/* Whether there is a slow task in progre
 bool					GSlowTaskOccurred				= false;					/* Whether a slow task began last tick*/
 bool					GIsRequestingExit				= false;					/* Indicates that MainLoop() should be exited at the end of the current iteration */
 /** Archive for serializing arbitrary data to and from memory												*/
-FReloadObjectArc*		GMemoryArchive					= NULL;
 bool					GAreScreenMessagesEnabled		= true;						/* Whether onscreen warnings/messages are enabled */
 bool					GScreenMessagesRestoreState		= false;					/* Used to restore state after a screenshot */
 int32					GIsDumpingMovie					= 0;						/* Whether we are dumping screenshots (!= 0), exposed as console variable r.DumpingMovie */
diff --git a/Engine/Source/Runtime/Core/Public/CoreGlobals.h b/Engine/Source/Runtime/Core/Public/CoreGlobals.h
--- a/Engine/Source/Runtime/Core/Public/CoreGlobals.h
+++ b/Engine/Source/Runtime/Core/Public/CoreGlobals.h
@@ -13,7 +13,6 @@ class FConfigCacheIni;
 class FFixedUObjectArray;
 class FOutputDeviceConsole;
 class FOutputDeviceRedirector;
-class FReloadObjectArc;
 class ITransaction;
 
 CORE_API DECLARE_LOG_CATEGORY_EXTERN(LogHAL, Log, All);
@@ -204,9 +203,6 @@ extern CORE_API bool GSlowTaskOccurred;
 extern CORE_API bool GIsGuarded;
 extern CORE_API bool GIsRequestingExit;
 
-/** Archive for serializing arbitrary data to and from memory						*/
-extern CORE_API FReloadObjectArc* GMemoryArchive;
-
 /**
 *	Global value indicating on-screen warnings/message should be displayed.
 *	Disabled via console command "DISABLEALLSCREENMESSAGES"
diff --git a/Engine/Source/Runtime/CoreUObject/Private/Blueprint/BlueprintSupport.cpp b/Engine/Source/Runtime/CoreUObject/Private/Blueprint/BlueprintSupport.cpp
--- a/Engine/Source/Runtime/CoreUObject/Private/Blueprint/BlueprintSupport.cpp
+++ b/Engine/Source/Runtime/CoreUObject/Private/Blueprint/BlueprintSupport.cpp
@@ -31,6 +31,19 @@ DEFINE_LOG_CATEGORY_STATIC(LogBlueprintSupport, Log, All);
 // Flag to enable the new BlueprintCompilationManager:
 COREUOBJECT_API bool GBlueprintUseCompilationManager = false;
 
+const FName FBlueprintTags::GeneratedClassPath(TEXT("GeneratedClass"));
+const FName FBlueprintTags::ParentClassPath(TEXT("ParentClass"));
+const FName FBlueprintTags::NativeParentClassPath(TEXT("NativeParentClass"));
+const FName FBlueprintTags::ClassFlags(TEXT("ClassFlags"));
+const FName FBlueprintTags::BlueprintType(TEXT("BlueprintType"));
+const FName FBlueprintTags::BlueprintDescription(TEXT("BlueprintDescription"));
+const FName FBlueprintTags::IsDataOnly(TEXT("IsDataOnly"));
+const FName FBlueprintTags::ImplementedInterfaces(TEXT("ImplementedInterfaces"));
+const FName FBlueprintTags::FindInBlueprintsData(TEXT("FiBData"));
+const FName FBlueprintTags::NumReplicatedProperties(TEXT("NumReplicatedProperties"));
+const FName FBlueprintTags::NumNativeComponents(TEXT("NativeComponents"));
+const FName FBlueprintTags::NumBlueprintComponents(TEXT("BlueprintComponents"));
+
 /**
  * Defined in BlueprintSupport.cpp
  * Duplicates all fields of a class in depth-first order. It makes sure that everything contained
diff --git a/Engine/Source/Runtime/CoreUObject/Private/Serialization/ObjectWriter.cpp b/Engine/Source/Runtime/CoreUObject/Private/Serialization/ObjectWriter.cpp
--- a/Engine/Source/Runtime/CoreUObject/Private/Serialization/ObjectWriter.cpp
+++ b/Engine/Source/Runtime/CoreUObject/Private/Serialization/ObjectWriter.cpp
@@ -32,7 +32,6 @@ FArchive& FObjectWriter::operator<<(FLazyObjectPtr& Value)
 
 FArchive& FObjectWriter::operator<<(FSoftObjectPtr& Value)
 {
-	Value.ResetWeakPtr();
 	return *this << Value.GetUniqueID();
 }
 
diff --git a/Engine/Source/Runtime/CoreUObject/Private/UObject/SoftObjectPath.cpp b/Engine/Source/Runtime/CoreUObject/Private/UObject/SoftObjectPath.cpp
--- a/Engine/Source/Runtime/CoreUObject/Private/UObject/SoftObjectPath.cpp
+++ b/Engine/Source/Runtime/CoreUObject/Private/UObject/SoftObjectPath.cpp
@@ -442,11 +442,12 @@ bool FSoftObjectPath::FixupForPIE()
 
 bool FSoftObjectPath::FixupCoreRedirects()
 {
-	FCoreRedirectObjectName OldName = FCoreRedirectObjectName(ToString());
+	FString OldString = ToString();
+	FCoreRedirectObjectName OldName = FCoreRedirectObjectName(OldString);
 	FCoreRedirectObjectName NewName = FCoreRedirects::GetRedirectedName(ECoreRedirectFlags::Type_Object, OldName);
 
 	// This also might be a class
-	if (OldName == NewName)
+	if (OldName == NewName && OldString.StartsWith(TEXT("/Script/")))
 	{
 		NewName = FCoreRedirects::GetRedirectedName(ECoreRedirectFlags::Type_Class, OldName);
 	}
diff --git a/Engine/Source/Runtime/CoreUObject/Public/Blueprint/BlueprintSupport.h b/Engine/Source/Runtime/CoreUObject/Public/Blueprint/BlueprintSupport.h
--- a/Engine/Source/Runtime/CoreUObject/Public/Blueprint/BlueprintSupport.h
+++ b/Engine/Source/Runtime/CoreUObject/Public/Blueprint/BlueprintSupport.h
@@ -10,6 +10,38 @@ class UDynamicClass;
 struct FCompilerNativizationOptions;
 class ITargetPlatform;
 
+/**
+ * List of asset registry tags used by blueprints. These are here so they can be used by both the asset registry and blueprint code
+ * These need to be kept in sync with UBlueprint::GetAssetRegistryTags, and any changes there will require resaving content
+ */
+struct COREUOBJECT_API FBlueprintTags
+{
+	/** Full path in export form ClassType'/PackagePath/PackageName.ClassName' of generated blueprint class */
+	static const FName GeneratedClassPath;
+	/** Full path in export form ClassType'/PackagePath/PackageName.ClassName' of the immediate parent, may be a blueprint or native class */
+	static const FName ParentClassPath;
+	/** Full path in export form Class'/Script/ModuleName.ClassName' of the first found parent native class */
+	static const FName NativeParentClassPath;
+	/** Integer representing bitfield EClassFlags */
+	static const FName ClassFlags;
+	/** String representing enum EBlueprintType */
+	static const FName BlueprintType;
+	/** String with user-entered description of blueprint */
+	static const FName BlueprintDescription;
+	/** String set to True/False, set if this is a data only blueprint */
+	static const FName IsDataOnly;
+	/** List of implemented interfaces, must be converted to FBPInterfaceDescription */
+	static const FName ImplementedInterfaces;
+	/** Very large string used to store find in blueprint data for the editor */
+	static const FName FindInBlueprintsData;
+	/** Number of replicated properties */
+	static const FName NumReplicatedProperties;
+	/** Number of native components */
+	static const FName NumNativeComponents;
+	/** Number of blueprint components */
+	static const FName NumBlueprintComponents;
+};
+
 struct FBlueprintWarningDeclaration
 {
 	FBlueprintWarningDeclaration(FName InWarningIdentifier, FText InWarningDescription)
diff --git a/Engine/Source/Runtime/CoreUObject/Public/CoreUObject.h b/Engine/Source/Runtime/CoreUObject/Public/CoreUObject.h
--- a/Engine/Source/Runtime/CoreUObject/Public/CoreUObject.h
+++ b/Engine/Source/Runtime/CoreUObject/Public/CoreUObject.h
@@ -58,8 +58,6 @@ MONOLITHIC_HEADER_BOILERPLATE()
 #include "Serialization/ObjectAndNameAsStringProxyArchive.h"
 #include "Serialization/ObjectWriter.h"
 #include "Serialization/ObjectReader.h"
-#include "Serialization/ReloadObjectArc.h"
-#include "Serialization/ArchiveReplaceArchetype.h"
 #include "Serialization/ArchiveShowReferences.h"
 #include "Serialization/FindReferencersArchive.h"
 #include "Serialization/FindObjectReferencers.h"
