Copying //UE4/Dev-Framework to //UE4/Dev-Main (Source: //UE4/Dev-Framework @ 4058146)

diff --git a/Engine/Source/Editor/Kismet/Private/BlueprintDetailsCustomization.cpp b/Engine/Source/Editor/Kismet/Private/BlueprintDetailsCustomization.cpp
--- a/Engine/Source/Editor/Kismet/Private/BlueprintDetailsCustomization.cpp
+++ b/Engine/Source/Editor/Kismet/Private/BlueprintDetailsCustomization.cpp
@@ -1168,7 +1168,7 @@ void FBlueprintVarActionDetails::OnVarNameChanged(const FText& InNewText)
 	}
 	else if(ValidatorResult == EValidatorResult::LocallyInUse)
 	{
-		VarNameEditableTextBox->SetError(LOCTEXT("ConflictsWithProperty", "Conflicts with another another local variable or function parameter!"));
+		VarNameEditableTextBox->SetError(LOCTEXT("ConflictsWithProperty", "Conflicts with another local variable or function parameter!"));
 	}
 	else
 	{
@@ -3004,9 +3004,7 @@ void FBlueprintGraphArgumentLayout::OnRefCheckStateChanged(ECheckBoxState InStat
 {
 	FEdGraphPinType PinType = OnGetPinInfo();
 	PinType.bIsReference = (InState == ECheckBoxState::Checked)? true : false;
-	// Note: Container types are implicitly passed by reference. For custom event nodes, the reference flag is essentially
-	//  treated as being redundant on container inputs, but we also need to implicitly set the 'const' flag to avoid a compiler note.
-	PinType.bIsConst = (PinType.IsContainer() || PinType.bIsReference) && TargetNode && TargetNode->IsA<UK2Node_CustomEvent>();
+	
 	PinInfoChanged(PinType);
 }
 
@@ -3030,7 +3028,7 @@ void FBlueprintGraphArgumentLayout::PinInfoChanged(const FEdGraphPinType& PinTyp
 				{
 					if (Node)
 					{
-						TSharedPtr<FUserPinInfo>* UDPinPtr = Node->UserDefinedPins.FindByPredicate([&](TSharedPtr<FUserPinInfo>& UDPin)
+						TSharedPtr<FUserPinInfo>* UDPinPtr = Node->UserDefinedPins.FindByPredicate([PinName](TSharedPtr<FUserPinInfo>& UDPin)
 						{
 							return UDPin.IsValid() && (UDPin->PinName == PinName);
 						});
@@ -3038,8 +3036,14 @@ void FBlueprintGraphArgumentLayout::PinInfoChanged(const FEdGraphPinType& PinTyp
 						{
 							(*UDPinPtr)->PinType = PinType;
 
-							// Container types are implicitly passed by reference. For custom event nodes, since they are inputs, also implicitly treat them as 'const' so that they don't result in a compiler note.
-							(*UDPinPtr)->PinType.bIsConst = PinType.IsContainer() && Node->IsA<UK2Node_CustomEvent>();
+							// Inputs flagged as pass-by-reference will also be flagged as 'const' here to conform to the expected native C++
+							// declaration of 'const Type&' for input reference parameters on functions with no outputs (i.e. events). Array
+							// types are also flagged as 'const' here since they will always be implicitly passed by reference, regardless of
+							// the checkbox setting. See UEditablePinBase::PostLoad() for more details.
+							if(!PinType.bIsConst && Node->ShouldUseConstRefParams())
+							{
+								(*UDPinPtr)->PinType.bIsConst = PinType.IsArray() || PinType.bIsReference;
+							}
 
 							// Reset default value, it probably doesn't match
 							(*UDPinPtr)->PinDefaultValue.Reset();
@@ -4137,7 +4141,7 @@ bool FBaseBlueprintGraphActionDetails::OnVerifyPinRename(UK2Node_EditablePinBase
 		const UProperty* ExistingProperty = FindField<const UProperty>(FoundFunction, *InNewName);
 		if (ExistingProperty)
 		{
-			OutErrorMessage = LOCTEXT("ConflictsWithProperty", "Conflicts with another another local variable or function parameter!");
+			OutErrorMessage = LOCTEXT("ConflictsWithProperty", "Conflicts with another local variable or function parameter!");
 			return false;
 		}
 	}
diff --git a/Engine/Source/Editor/Kismet/Private/BlueprintEditor.cpp b/Engine/Source/Editor/Kismet/Private/BlueprintEditor.cpp
--- a/Engine/Source/Editor/Kismet/Private/BlueprintEditor.cpp
+++ b/Engine/Source/Editor/Kismet/Private/BlueprintEditor.cpp
@@ -116,6 +116,7 @@
 // Debugging
 #include "Debugging/SKismetDebuggingView.h"
 #include "Debugging/KismetDebugCommands.h"
+#include "WatchPointViewer.h"
 // End of debugging
 
 // Misc diagnostics
@@ -1690,6 +1691,7 @@ void FBlueprintEditor::LoadLibrariesFromAssetRegistry()
 						if (BlueprintLibPtr )
 						{
 							StandardLibraries.AddUnique(BlueprintLibPtr);
+							WatchViewer::UpdateWatchListFromBlueprint(BlueprintLibPtr);
 						}
 					}
 				}
@@ -1748,6 +1750,7 @@ void FBlueprintEditor::InitBlueprintEditor(
 	const TSharedRef<FTabManager::FLayout> DummyLayout = FTabManager::NewLayout("NullLayout")->AddArea(FTabManager::NewPrimaryArea());
 	const bool bCreateDefaultStandaloneMenu = true;
 	const bool bCreateDefaultToolbar = true;
+	const FName BlueprintEditorAppName = FName(TEXT("BlueprintEditorApp"));
 	InitAssetEditor(Mode, InitToolkitHost, BlueprintEditorAppName, DummyLayout, bCreateDefaultStandaloneMenu, bCreateDefaultToolbar, Objects);
 	
 	CommonInitialization(InBlueprints);
diff --git a/Engine/Source/Editor/Kismet/Private/BlueprintEditorModes.cpp b/Engine/Source/Editor/Kismet/Private/BlueprintEditorModes.cpp
--- a/Engine/Source/Editor/Kismet/Private/BlueprintEditorModes.cpp
+++ b/Engine/Source/Editor/Kismet/Private/BlueprintEditorModes.cpp
@@ -101,10 +101,7 @@ FBlueprintEditorApplicationMode::FBlueprintEditorApplicationMode(TSharedPtr<clas
 	}
 	BlueprintEditorTabFactories.RegisterFactory(MakeShareable(new FCompilerResultsSummoner(InBlueprintEditor)));
 	BlueprintEditorTabFactories.RegisterFactory(MakeShareable(new FFindResultsSummoner(InBlueprintEditor)));
-	if (GetDefault<UEditorExperimentalSettings>()->bEnableBlueprintGraphBookmarks)
-	{
-		BlueprintEditorTabFactories.RegisterFactory(MakeShareable(new FBookmarksSummoner(InBlueprintEditor)));
-	}
+	BlueprintEditorTabFactories.RegisterFactory(MakeShareable(new FBookmarksSummoner(InBlueprintEditor)));
 	
 	if( bRegisterViewport )
 	{
@@ -175,6 +172,7 @@ FBlueprintDefaultsApplicationMode::FBlueprintDefaultsApplicationMode(TSharedPtr<
 	
 	BlueprintDefaultsTabFactories.RegisterFactory(MakeShareable(new FDefaultsEditorSummoner(InBlueprintEditor)));
 	BlueprintDefaultsTabFactories.RegisterFactory(MakeShareable(new FFindResultsSummoner(InBlueprintEditor)));
+	BlueprintDefaultsTabFactories.RegisterFactory(MakeShareable(new FMyBlueprintSummoner(InBlueprintEditor)));
 
 	TabLayout = FTabManager::NewLayout( "Standalone_BlueprintDefaults_Layout_v6" )
 		->AddArea
@@ -574,10 +572,7 @@ FBlueprintEditorUnifiedMode::FBlueprintEditorUnifiedMode(TSharedPtr<class FBluep
 	}
 	BlueprintEditorTabFactories.RegisterFactory(MakeShareable(new FCompilerResultsSummoner(InBlueprintEditor)));
 	BlueprintEditorTabFactories.RegisterFactory(MakeShareable(new FFindResultsSummoner(InBlueprintEditor)));
-	if (GetDefault<UEditorExperimentalSettings>()->bEnableBlueprintGraphBookmarks)
-	{
-		BlueprintEditorTabFactories.RegisterFactory(MakeShareable(new FBookmarksSummoner(InBlueprintEditor)));
-	}
+	BlueprintEditorTabFactories.RegisterFactory(MakeShareable(new FBookmarksSummoner(InBlueprintEditor)));
 	
 	if( bRegisterViewport )
 	{
diff --git a/Engine/Source/Editor/Kismet/Private/BlueprintEditorModule.cpp b/Engine/Source/Editor/Kismet/Private/BlueprintEditorModule.cpp
--- a/Engine/Source/Editor/Kismet/Private/BlueprintEditorModule.cpp
+++ b/Engine/Source/Editor/Kismet/Private/BlueprintEditorModule.cpp
@@ -2,10 +2,10 @@
 
 
 #include "BlueprintEditorModule.h"
+
+#include "BlueprintDebugger.h"
 #include "Editor.h"
 #include "Modules/ModuleManager.h"
-#include "Framework/Application/SlateApplication.h"
-#include "EditorStyleSet.h"
 #include "EditorUndoClient.h"
 #include "Logging/TokenizedMessage.h"
 #include "Misc/ConfigCacheIni.h"
@@ -15,9 +15,6 @@
 #include "Kismet2/KismetEditorUtilities.h"
 #include "Kismet2/KismetDebugUtilities.h"
 #include "Kismet2/BlueprintEditorUtils.h"
-//#include "BlueprintWarningsConfigurationPanel.h"
-#include "Editor/WorkspaceMenuStructure/Public/WorkspaceMenuStructure.h"
-#include "Editor/WorkspaceMenuStructure/Public/WorkspaceMenuStructureModule.h"
 #include "Editor/LevelEditor/Public/LevelEditor.h"
 #include "UserDefinedEnumEditor.h"
 #include "MessageLogInitializationOptions.h"
@@ -30,39 +27,13 @@
 #include "UserDefinedStructureEditor.h"
 #include "EdGraphUtilities.h"
 #include "BlueprintGraphPanelPinFactory.h"
-#include "Widgets/Docking/SDockTab.h"
-#include "KismetCompiler.h"
-#include "CallStackViewer.h"
 #include "WatchPointViewer.h"
+#include "KismetCompiler.h"
 
 #define LOCTEXT_NAMESPACE "BlueprintEditor"
 
 IMPLEMENT_MODULE( FBlueprintEditorModule, Kismet );
 
-const FName BlueprintEditorAppName = FName(TEXT("BlueprintEditorApp"));
-const FName DebuggerAppName = FName(TEXT("DebuggerApp"));
-
-//////////////////////////////////////////////////////////////////////////
-// SDebuggerApp
-
-#include "Debugging/SKismetDebuggingView.h"
-
-static bool CanCloseBlueprintDebugger()
-{
-	return !GIntraFrameDebuggingGameThread;
-}
-
-TSharedRef<SDockTab> CreateBluprintDebuggerTab( const FSpawnTabArgs& Args )
-{
-	return SNew(SDockTab)
-		.TabRole( ETabRole::NomadTab )
-		.OnCanCloseTab( SDockTab::FCanCloseTab::CreateStatic(&CanCloseBlueprintDebugger) )
-		.Label( NSLOCTEXT("BlueprintDebugger", "TabTitle", "Blueprint Debugger") )
-		[
-			SNew(SKismetDebuggingView)
-		];
-}
-
 //////////////////////////////////////////////////////////////////////////
 // FBlueprintEditorModule
 
@@ -197,17 +168,7 @@ void FBlueprintEditorModule::StartupModule()
 	MenuExtensibilityManager = MakeShareable(new FExtensibilityManager);
 	SharedBlueprintEditorCommands = MakeShareable(new FUICommandList);
 
-	const IWorkspaceMenuStructure& MenuStructure = WorkspaceMenu::GetMenuStructure();
-
-	FGlobalTabmanager::Get()->RegisterNomadTabSpawner( DebuggerAppName, FOnSpawnTab::CreateStatic(&CreateBluprintDebuggerTab) )
-		.SetDisplayName( NSLOCTEXT("BlueprintDebugger", "TabTitle", "Blueprint Debugger") )
-		.SetTooltipText( NSLOCTEXT("BlueprintDebugger", "TooltipText", "Open the Blueprint Debugger tab.") )
-		.SetGroup( MenuStructure.GetDeveloperToolsDebugCategory() )
-		.SetIcon(FSlateIcon(FEditorStyle::GetStyleSetName(), "BlueprintDebugger.TabIcon"));
-
-	CallStackViewer::RegisterTabSpawner();
-
-	WatchViewer::RegisterTabSpawner();
+	BlueprintDebugger = MakeUnique<FBlueprintDebugger>();
 
 	// Have to check GIsEditor because right now editor modules can be loaded by the game
 	// Once LoadModule is guaranteed to return NULL for editor modules in game, this can be removed
@@ -266,11 +227,6 @@ void FBlueprintEditorModule::ShutdownModule()
 	SharedBlueprintEditorCommands.Reset();
 	MenuExtensibilityManager.Reset();
 
-	if (FSlateApplication::IsInitialized())
-	{
-		FGlobalTabmanager::Get()->UnregisterNomadTabSpawner( DebuggerAppName );
-	}
-	
 	// Remove level viewport context menu extenders
 	if ( FModuleManager::Get().IsModuleLoaded( "LevelEditor" ) )
 	{
@@ -303,6 +259,8 @@ TSharedRef<IBlueprintEditor> FBlueprintEditorModule::CreateBlueprintEditor(const
 		NewBlueprintEditor->RegisterSCSEditorCustomization(It->Key, It->Value.Execute(NewBlueprintEditor));
 	}
 	
+	WatchViewer::UpdateWatchListFromBlueprint(Blueprint);
+
 	EBlueprintType const BPType = Blueprint ? (EBlueprintType)Blueprint->BlueprintType : BPTYPE_Normal;
 	BlueprintEditorOpened.Broadcast(BPType);
 
@@ -320,6 +278,11 @@ TSharedRef<IBlueprintEditor> FBlueprintEditorModule::CreateBlueprintEditor(const
 		NewBlueprintEditor->RegisterSCSEditorCustomization(It->Key, It->Value.Execute(NewBlueprintEditor));
 	}
 
+	for (UBlueprint* Blueprint : BlueprintsToEdit)
+	{
+		WatchViewer::UpdateWatchListFromBlueprint(Blueprint);
+	}
+
 	EBlueprintType const BPType = ( (BlueprintsToEdit.Num() > 0) && (BlueprintsToEdit[0] != NULL) ) 
 		? (EBlueprintType) BlueprintsToEdit[0]->BlueprintType
 		: BPTYPE_Normal;
diff --git a/Engine/Source/Editor/Kismet/Private/CallStackViewer.cpp b/Engine/Source/Editor/Kismet/Private/CallStackViewer.cpp
--- a/Engine/Source/Editor/Kismet/Private/CallStackViewer.cpp
+++ b/Engine/Source/Editor/Kismet/Private/CallStackViewer.cpp
@@ -22,6 +22,10 @@
 #include "WorkspaceMenuStructure.h"
 #include "WorkspaceMenuStructureModule.h"
 
+// So that we can poll the running state:
+#include "Editor/UnrealEdEngine.h"
+extern UNREALED_API UUnrealEdEngine* GUnrealEd;
+
 #define LOCTEXT_NAMESPACE "CallStackViewer"
 
 enum class ECallstackLanguages : uint8
@@ -297,7 +301,7 @@ void SCallStackViewer::Construct(const FArguments& InArgs, TArray<TSharedRef<FCa
 	const auto EmptyWarningVisibility = [](TWeakPtr<SCallStackViewer> ControlOwnerWeak) -> EVisibility
 	{
 		TSharedPtr<SCallStackViewer> ControlOwner = ControlOwnerWeak.Pin();
-		if(ControlOwner.IsValid() &&
+		if( ControlOwner.IsValid() &&
 			ControlOwner->CallStackSource &&
 			ControlOwner->CallStackSource->Num() > 0)
 		{
@@ -306,6 +310,25 @@ void SCallStackViewer::Construct(const FArguments& InArgs, TArray<TSharedRef<FCa
 		return EVisibility::Visible;
 	};
 
+	const auto StaleCallStackWarning = [](TWeakPtr<SCallStackViewer> ControlOwnerWeak)->FText
+	{
+		TSharedPtr<SCallStackViewer> ControlOwner = ControlOwnerWeak.Pin();
+		if (ControlOwner.IsValid() &&
+			ControlOwner->CallStackSource &&
+			ControlOwner->CallStackSource->Num() > 0)
+		{
+			if (!GUnrealEd->PlayWorld)
+			{
+				return LOCTEXT("SessionOver", "Warning: The session has ended and therefore the displayed call stack is out of date");
+			}
+			if (!GUnrealEd->PlayWorld->bDebugPauseExecution)
+			{
+				return LOCTEXT("CurrentlyRunning", "Warning: The game is currently running - the callstack will be updated when excution is paused");
+			}
+		}
+		return FText();
+	};
+
 	const auto CallStackViewIsEnabled = [](TWeakPtr<SCallStackViewer> ControlOwnerWeak) -> bool
 	{
 		TSharedPtr<SCallStackViewer> ControlOwner = ControlOwnerWeak.Pin();
@@ -332,31 +355,47 @@ void SCallStackViewer::Construct(const FArguments& InArgs, TArray<TSharedRef<FCa
 			SNew(SOverlay)
 			+SOverlay::Slot()
 			[
-				SAssignNew(CallStackTreeWidget, SCallStackTree)
-				.ItemHeight(25.0f)
-				.TreeItemsSource(CallStackSource)
-				.OnGenerateRow(SCallStackTree::FOnGenerateRow::CreateStatic(RowGenerator, SelfWeak))
-				.OnGetChildren(SCallStackTree::FOnGetChildren::CreateStatic(ChildrenAccessor))
-				.OnMouseButtonDoubleClick(SCallStackTree::FOnMouseButtonClick::CreateSP(SelfTyped, &SCallStackViewer::JumpToEntry))
-				.OnContextMenuOpening(FOnContextMenuOpening::CreateStatic(ContextMenuOpened, CommandListWeak, SelfWeak))
-				.IsEnabled(
-					TAttribute<bool>::Create(
-						TAttribute<bool>::FGetter::CreateStatic(CallStackViewIsEnabled, SelfWeak)
+				SNew(SVerticalBox)
+				+SVerticalBox::Slot()
+				.AutoHeight()
+				[
+					SAssignNew(CallStackTreeWidget, SCallStackTree)
+					.ItemHeight(25.0f)
+					.TreeItemsSource(CallStackSource)
+					.OnGenerateRow(SCallStackTree::FOnGenerateRow::CreateStatic(RowGenerator, SelfWeak))
+					.OnGetChildren(SCallStackTree::FOnGetChildren::CreateStatic(ChildrenAccessor))
+					.OnMouseButtonDoubleClick(SCallStackTree::FOnMouseButtonClick::CreateSP(SelfTyped, &SCallStackViewer::JumpToEntry))
+					.OnContextMenuOpening(FOnContextMenuOpening::CreateStatic(ContextMenuOpened, CommandListWeak, SelfWeak))
+					.IsEnabled(
+						TAttribute<bool>::Create(
+							TAttribute<bool>::FGetter::CreateStatic(CallStackViewIsEnabled, SelfWeak)
+						)
 					)
-				)
-				.HeaderRow
-				(
-					SNew(SHeaderRow)
-					+SHeaderRow::Column(TEXT("ProgramCounter"))
-					.DefaultLabel(LOCTEXT("ProgramCounterLabel", ""))
-					.FixedWidth(16.f)
-					+SHeaderRow::Column(TEXT("FunctionName"))
-					.FillWidth(.8f)
-					.DefaultLabel(LOCTEXT("FunctionName", "Function Name"))
-					+SHeaderRow::Column(TEXT("Language"))
-					.DefaultLabel(LOCTEXT("Language", "Language"))
-					.FillWidth(.15f)
-				)
+					.HeaderRow
+					(
+						SNew(SHeaderRow)
+						+SHeaderRow::Column(TEXT("ProgramCounter"))
+						.DefaultLabel(LOCTEXT("ProgramCounterLabel", ""))
+						.FixedWidth(16.f)
+						+SHeaderRow::Column(TEXT("FunctionName"))
+						.FillWidth(.8f)
+						.DefaultLabel(LOCTEXT("FunctionName", "Function Name"))
+						+SHeaderRow::Column(TEXT("Language"))
+						.DefaultLabel(LOCTEXT("Language", "Language"))
+						.FillWidth(.15f)
+					)
+				]
+				+ SVerticalBox::Slot()
+				[
+					SNew(STextBlock)
+					.Text(
+						TAttribute<FText>::Create(
+							TAttribute<FText>::FGetter::CreateStatic(StaleCallStackWarning, SelfWeak)
+						)
+					)
+					.ColorAndOpacity(FLinearColor::Yellow)
+					.Justification(ETextJustify::Center)
+				]
 			]
 			+SOverlay::Slot()
 			.Padding(32.f)
@@ -538,12 +577,18 @@ void CallStackViewer::UpdateDisplayedCallstack(const TArray<const FFrame*>& Scri
 	ListSubscribers.Broadcast(&Private_CallStackSource);
 }
 
-void CallStackViewer::RegisterTabSpawner()
+FName CallStackViewer::GetTabName()
+{
+	const FName TabName = TEXT("CallStackViewer");
+	return TabName;
+}
+
+void CallStackViewer::RegisterTabSpawner(FTabManager& TabManager)
 {
 	const auto SpawnCallStackViewTab = []( const FSpawnTabArgs& Args )
 	{
 		return SNew(SDockTab)
-			.TabRole( ETabRole::NomadTab )
+			.TabRole( ETabRole::PanelTab )
 			.Label( LOCTEXT("TabTitle", "Call Stack") )
 			[
 				SNew(SBorder)
@@ -554,13 +599,9 @@ void CallStackViewer::RegisterTabSpawner()
 			];
 	};
 	
-	const IWorkspaceMenuStructure& MenuStructure = WorkspaceMenu::GetMenuStructure();
-
-	const FName TabName = TEXT("CallStackViewer");
-	FGlobalTabmanager::Get()->RegisterNomadTabSpawner( TabName, FOnSpawnTab::CreateStatic(SpawnCallStackViewTab) )
+	TabManager.RegisterTabSpawner(CallStackViewer::GetTabName(), FOnSpawnTab::CreateStatic(SpawnCallStackViewTab) )
 		.SetDisplayName( LOCTEXT("TabTitle", "Call Stack") )
-		.SetTooltipText( LOCTEXT("TooltipText", "Open the Call Stack tab.") )
-		.SetGroup( MenuStructure.GetDeveloperToolsDebugCategory() );
+		.SetTooltipText( LOCTEXT("TooltipText", "Open the Call Stack tab") );
 }
 
 #undef LOCTEXT_NAMESPACE
diff --git a/Engine/Source/Editor/Kismet/Private/Debugging/SKismetDebuggingView.cpp b/Engine/Source/Editor/Kismet/Private/Debugging/SKismetDebuggingView.cpp
--- a/Engine/Source/Editor/Kismet/Private/Debugging/SKismetDebuggingView.cpp
+++ b/Engine/Source/Editor/Kismet/Private/Debugging/SKismetDebuggingView.cpp
@@ -1085,7 +1085,7 @@ void SKismetDebuggingView::Construct(const FArguments& InArgs)
 			.Text( this, &SKismetDebuggingView::GetTopText )
 		]
 		+SVerticalBox::Slot()
-		. FillHeight(1.0f)
+		.AutoHeight()
 		[
 			SAssignNew( DebugTreeView, STreeView< FDebugTreeItemPtr > )
 			.TreeItemsSource( &RootTreeItems )
diff --git a/Engine/Source/Editor/Kismet/Private/DiffUtils.cpp b/Engine/Source/Editor/Kismet/Private/DiffUtils.cpp
--- a/Engine/Source/Editor/Kismet/Private/DiffUtils.cpp
+++ b/Engine/Source/Editor/Kismet/Private/DiffUtils.cpp
@@ -445,7 +445,7 @@ static void IdenticalHelper(const UProperty* AProperty, const UProperty* BProper
 	if (APropAsStruct != nullptr)
 	{
 		const UStructProperty* BPropAsStruct = CastChecked<UStructProperty>(BProperty);
-		if (APropAsStruct->Struct->StructFlags & STRUCT_IdenticalNative || BPropAsStruct->Struct != APropAsStruct->Struct)
+		if (APropAsStruct->Struct->StructFlags & STRUCT_IdenticalNative && BPropAsStruct->Struct != APropAsStruct->Struct)
 		{
 			// If the struct uses CPP identical tests, then we can't dig into it, and we already know it's not identical from the test when we started
 			DifferingSubProperties.Push(RootPath);
diff --git a/Engine/Source/Editor/Kismet/Private/SBlueprintBookmarks.cpp b/Engine/Source/Editor/Kismet/Private/SBlueprintBookmarks.cpp
--- a/Engine/Source/Editor/Kismet/Private/SBlueprintBookmarks.cpp
+++ b/Engine/Source/Editor/Kismet/Private/SBlueprintBookmarks.cpp
@@ -69,7 +69,10 @@ void SBlueprintBookmarks::Construct(const FArguments& InArgs)
 		NAME_None,
 		EUserInterfaceActionType::ToggleButton);
 
-	SAssignNew(SearchBoxWidget, SSearchBox);
+	SAssignNew(SearchBoxWidget, SSearchBox)
+	.SelectAllTextWhenFocused(true)
+	.OnTextCommitted(this, &SBlueprintBookmarks::OnFilterTextCommitted)
+	.OnTextChanged(this, &SBlueprintBookmarks::OnFilterTextCommitted, ETextCommit::Default);
 
 	SAssignNew(TreeViewWidget, STreeView<FTreeViewItemPtr>)
 	.TreeItemsSource(&TreeViewRootItems)
@@ -251,6 +254,13 @@ void SBlueprintBookmarks::STreeItemRow::OnNameTextCommitted(const FText& InNewNa
 	}
 }
 
+void SBlueprintBookmarks::OnFilterTextCommitted(const FText& InText, ETextCommit::Type CommitType)
+{
+	FilterText = InText;
+
+	RefreshBookmarksTree();
+}
+
 void SBlueprintBookmarks::OnDeleteSelectedTreeViewItems()
 {
 	TSharedPtr<FBlueprintEditor> BlueprintEditor = EditorContext.Pin();
@@ -405,17 +415,26 @@ void SBlueprintBookmarks::RefreshBookmarksTree()
 		const UEdGraph* FocusedGraph = BlueprintEditor->GetFocusedGraph();
 		const bool bCurrentGraphOnly = BlueprintEditorSettings->bShowBookmarksForCurrentDocumentOnlyInTab;
 
+		const FText& LocalFilterText = FilterText;
+		auto IsNodeFilteredOutLambda = [LocalFilterText](const FBPEditorBookmarkNode& InNode)
+		{
+			return !LocalFilterText.IsEmpty() && !InNode.DisplayName.ToString().Contains(LocalFilterText.ToString());
+		};
+
 		for (FBPEditorBookmarkNode& BookmarkNode : BlueprintEditorSettings->BookmarkNodes)
 		{
-			if (const FEditedDocumentInfo* BookmarkInfo = BlueprintEditorSettings->Bookmarks.Find(BookmarkNode.NodeGuid))
+			if (!IsNodeFilteredOutLambda(BookmarkNode))
 			{
-				if (const UEdGraph* GraphContext = Cast<UEdGraph>(BookmarkInfo->EditedObjectPath.ResolveObject()))
+				if (const FEditedDocumentInfo* BookmarkInfo = BlueprintEditorSettings->Bookmarks.Find(BookmarkNode.NodeGuid))
 				{
-					if (Blueprint == FBlueprintEditorUtils::FindBlueprintForGraph(GraphContext))
+					if (const UEdGraph* GraphContext = Cast<UEdGraph>(BookmarkInfo->EditedObjectPath.ResolveObject()))
 					{
-						if (!bCurrentGraphOnly || (GraphContext == FocusedGraph))
+						if (Blueprint == FBlueprintEditorUtils::FindBlueprintForGraph(GraphContext))
 						{
-							BookmarksTreeViewRoot->Children.Add(MakeShareable(new FTreeViewItem(ETreeViewNodeType::LocalBookmark, BookmarkNode, BookmarkInfo)));
+							if (!bCurrentGraphOnly || (GraphContext == FocusedGraph))
+							{
+								BookmarksTreeViewRoot->Children.Add(MakeShareable(new FTreeViewItem(ETreeViewNodeType::LocalBookmark, BookmarkNode, BookmarkInfo)));
+							}
 						}
 					}
 				}
@@ -424,10 +443,13 @@ void SBlueprintBookmarks::RefreshBookmarksTree()
 
 		for (FBPEditorBookmarkNode& BookmarkNode : Blueprint->BookmarkNodes)
 		{
-			const FEditedDocumentInfo* BookmarkInfo = Blueprint->Bookmarks.Find(BookmarkNode.NodeGuid);
-			if (!bCurrentGraphOnly || (BookmarkInfo && BookmarkInfo->EditedObjectPath.ResolveObject() == FocusedGraph))
+			if (!IsNodeFilteredOutLambda(BookmarkNode))
 			{
-				BookmarksTreeViewRoot->Children.Add(MakeShareable(new FTreeViewItem(ETreeViewNodeType::SharedBookmark, BookmarkNode, BookmarkInfo)));
+				const FEditedDocumentInfo* BookmarkInfo = Blueprint->Bookmarks.Find(BookmarkNode.NodeGuid);
+				if (!bCurrentGraphOnly || (BookmarkInfo && BookmarkInfo->EditedObjectPath.ResolveObject() == FocusedGraph))
+				{
+					BookmarksTreeViewRoot->Children.Add(MakeShareable(new FTreeViewItem(ETreeViewNodeType::SharedBookmark, BookmarkNode, BookmarkInfo)));
+				}
 			}
 		}
 
@@ -459,7 +481,10 @@ void SBlueprintBookmarks::RefreshBookmarksTree()
 					FEditedDocumentInfo& CommentInfo = CommentNodeInfo.Add(CommentNode.NodeGuid);
 					CommentInfo.EditedObjectPath = GraphCommentNode;
 					
-					CommentsTreeViewRoot->Children.Add(MakeShareable(new FTreeViewItem(ETreeViewNodeType::Comment, CommentNode, &CommentInfo)));
+					if (!IsNodeFilteredOutLambda(CommentNode))
+					{
+						CommentsTreeViewRoot->Children.Add(MakeShareable(new FTreeViewItem(ETreeViewNodeType::Comment, CommentNode, &CommentInfo)));
+					}
 				}
 			}
 
diff --git a/Engine/Source/Editor/Kismet/Private/SBlueprintBookmarks.h b/Engine/Source/Editor/Kismet/Private/SBlueprintBookmarks.h
--- a/Engine/Source/Editor/Kismet/Private/SBlueprintBookmarks.h
+++ b/Engine/Source/Editor/Kismet/Private/SBlueprintBookmarks.h
@@ -92,6 +92,8 @@ protected:
 		TWeakPtr<FBlueprintEditor> EditorContext;
 	};
 
+	void OnFilterTextCommitted(const FText& InText, ETextCommit::Type CommitType);
+
 	void OnDeleteSelectedTreeViewItems();
 	bool CanDeleteSelectedTreeViewItems() const;
 	void OnRenameSelectedTreeViewItems();
@@ -122,4 +124,6 @@ private:
 	TSharedPtr<FUICommandList> CommandList;
 	TSharedPtr<SSearchBox> SearchBoxWidget;
 	TSharedPtr<STreeView<FTreeViewItemPtr> > TreeViewWidget;
+
+	FText FilterText;
 };
diff --git a/Engine/Source/Editor/Kismet/Private/SGraphTitleBarAddNewBookmark.cpp b/Engine/Source/Editor/Kismet/Private/SGraphTitleBarAddNewBookmark.cpp
--- a/Engine/Source/Editor/Kismet/Private/SGraphTitleBarAddNewBookmark.cpp
+++ b/Engine/Source/Editor/Kismet/Private/SGraphTitleBarAddNewBookmark.cpp
@@ -2,6 +2,7 @@
 
 #include "SGraphTitleBarAddNewBookmark.h"
 #include "BlueprintEditor.h"
+#include "BlueprintEditorSettings.h"
 #include "Kismet2/BlueprintEditorUtils.h"
 #include "EditorStyleSet.h"
 #include "Framework/Application/SlateApplication.h"
@@ -63,7 +64,11 @@ void SGraphTitleBarAddNewBookmark::Construct(const FArguments& InArgs)
 						SAssignNew(NameEntryWidget, SEditableTextBox)
 						.SelectAllTextWhenFocused(true)
 						.OnTextCommitted(this, &SGraphTitleBarAddNewBookmark::OnNameTextCommitted)
-						.Text(this, &SGraphTitleBarAddNewBookmark::GetDefaultNameText)
+						.OnTextChanged(this, &SGraphTitleBarAddNewBookmark::OnNameTextCommitted, ETextCommit::Default)
+						.Text_Lambda([this]() -> FText
+						{
+							return CurrentNameText;
+						})
 					]
 				]
 				+ SVerticalBox::Slot()
@@ -114,6 +119,7 @@ void SGraphTitleBarAddNewBookmark::Construct(const FArguments& InArgs)
 							.TextStyle(FEditorStyle::Get(), "FlatButton.DefaultTextStyle")
 							.ButtonStyle(FEditorStyle::Get(), "FlatButton.Success")
 							.OnClicked(this, &SGraphTitleBarAddNewBookmark::OnAddButtonClicked)
+							.IsEnabled(this, &SGraphTitleBarAddNewBookmark::IsAddButtonEnabled)
 						]
 					]
 				]
@@ -153,13 +159,23 @@ FText SGraphTitleBarAddNewBookmark::GetDefaultNameText() const
 
 		if (CurrentViewBookmarkId.IsValid())
 		{
-			for (const FBPEditorBookmarkNode& BookmarkNode : Blueprint->BookmarkNodes)
+			const FGuid& TargetNodeId = CurrentViewBookmarkId;
+			auto FindCurrentBookmarkNodeLambda = [TargetNodeId](const TArray<FBPEditorBookmarkNode>& InBookmarkNodesToSearch)
 			{
-				if (BookmarkNode.NodeGuid == CurrentViewBookmarkId)
-				{
-					ResultText = BookmarkNode.DisplayName;
-					break;
-				}
+				return InBookmarkNodesToSearch.FindByPredicate([TargetNodeId](const FBPEditorBookmarkNode& InNode) { return InNode.NodeGuid == TargetNodeId; });
+			};
+
+			// Check for a shared bookmark
+			const FBPEditorBookmarkNode* BookmarkNode = FindCurrentBookmarkNodeLambda(Blueprint->BookmarkNodes);
+			if (BookmarkNode == nullptr)
+			{
+				// No shared bookmark found, so check the local user settings
+				BookmarkNode = FindCurrentBookmarkNodeLambda(GetDefault<UBlueprintEditorSettings>()->BookmarkNodes);
+			}
+
+			if (ensure(BookmarkNode != nullptr))
+			{
+				ResultText = BookmarkNode->DisplayName;
 			}
 		}
 		else
@@ -220,7 +236,6 @@ FText SGraphTitleBarAddNewBookmark::GetAddButtonLabel() const
 
 void SGraphTitleBarAddNewBookmark::OnComboBoxOpened()
 {
-	CurrentNameText = FText::GetEmpty();
 	CurrentViewBookmarkId.Invalidate();
 
 	TSharedPtr<FBlueprintEditor> EditorContext = EditorContextPtr.Pin();
@@ -228,6 +243,19 @@ void SGraphTitleBarAddNewBookmark::OnComboBoxOpened()
 	{
 		EditorContext->GetViewBookmark(CurrentViewBookmarkId);
 	}
+
+	CurrentNameText = GetDefaultNameText();
+	if (CurrentViewBookmarkId.IsValid())
+	{
+		OriginalNameText = CurrentNameText;
+	}
+	else
+	{
+		OriginalNameText = FText::GetEmpty();
+	}
+
+	check(NameEntryWidget.IsValid());
+	NameEntryWidget->SetText(CurrentNameText);
 }
 
 FReply SGraphTitleBarAddNewBookmark::OnAddButtonClicked()
@@ -256,6 +284,11 @@ FReply SGraphTitleBarAddNewBookmark::OnAddButtonClicked()
 	return FReply::Handled();
 }
 
+bool SGraphTitleBarAddNewBookmark::IsAddButtonEnabled() const
+{
+	return !CurrentNameText.IsEmpty() && CurrentNameText.CompareTo(OriginalNameText) != 0;
+}
+
 FReply SGraphTitleBarAddNewBookmark::OnRemoveButtonClicked()
 {
 	TSharedPtr<FBlueprintEditor> EditorContext = EditorContextPtr.Pin();
diff --git a/Engine/Source/Editor/Kismet/Private/SGraphTitleBarAddNewBookmark.h b/Engine/Source/Editor/Kismet/Private/SGraphTitleBarAddNewBookmark.h
--- a/Engine/Source/Editor/Kismet/Private/SGraphTitleBarAddNewBookmark.h
+++ b/Engine/Source/Editor/Kismet/Private/SGraphTitleBarAddNewBookmark.h
@@ -26,6 +26,7 @@ protected:
 
 	void OnComboBoxOpened();
 	FReply OnAddButtonClicked();
+	bool IsAddButtonEnabled() const;
 	FReply OnRemoveButtonClicked();
 	void OnNameTextCommitted(const FText& InText, ETextCommit::Type CommitType);
 
@@ -34,5 +35,6 @@ private:
 	TSharedPtr<class SEditableTextBox> NameEntryWidget;
 
 	FText CurrentNameText;
+	FText OriginalNameText;
 	FGuid CurrentViewBookmarkId;
 };
diff --git a/Engine/Source/Editor/Kismet/Private/SSCSEditor.cpp b/Engine/Source/Editor/Kismet/Private/SSCSEditor.cpp
--- a/Engine/Source/Editor/Kismet/Private/SSCSEditor.cpp
+++ b/Engine/Source/Editor/Kismet/Private/SSCSEditor.cpp
@@ -785,7 +785,7 @@ FSCSEditorTreeNodePtrType FSCSEditorTreeNode::FactoryNodeFromComponent(UActorCom
 	{
 		if (InComponent->CreationMethod == EComponentCreationMethod::Instance)
 		{
-			return MakeShareable(new FSCSEditorTreeNodeInstanceAddedComponent(Owner, InComponent->GetFName()));
+			return MakeShareable(new FSCSEditorTreeNodeInstanceAddedComponent(Owner, InComponent));
 		}
 		else
 		{
@@ -1072,22 +1072,12 @@ UActorComponent* FSCSEditorTreeNodeInstancedInheritedComponent::GetEditableCompo
 //////////////////////////////////////////////////////////////////////////
 // FSCSEditorTreeNodeInstanceAddedComponent
 
-FSCSEditorTreeNodeInstanceAddedComponent::FSCSEditorTreeNodeInstanceAddedComponent(AActor* Owner, FName InComponentName)
+FSCSEditorTreeNodeInstanceAddedComponent::FSCSEditorTreeNodeInstanceAddedComponent(AActor* Owner, UActorComponent* InComponentTemplate)
 {
-	InstancedComponentName = InComponentName;
-	check(InstancedComponentName != NAME_None);	// ...otherwise IsRootActor() can return a false positive.
-
+	check(InComponentTemplate);
+	InstancedComponentName = InComponentTemplate->GetFName();
 	InstancedComponentOwnerPtr = Owner;
-
-	SetComponentTemplate(nullptr);
-	for (UActorComponent* ComponentInstance : Owner->GetComponents())
-	{
-		if (ComponentInstance && ComponentInstance->GetFName() == InstancedComponentName)
-		{
-			SetComponentTemplate(ComponentInstance);
-			break;
-		}
-	}
+	SetComponentTemplate(InComponentTemplate);
 }
 
 bool FSCSEditorTreeNodeInstanceAddedComponent::IsRootComponent() const
@@ -4578,43 +4568,47 @@ void SSCSEditor::UpdateTree(bool bRegenerateTreeNodes)
 			if (AActor* ActorInstance = GetActorContext())
 			{
 				// Get the full set of instanced components
-				TInlineComponentArray<UActorComponent*> Components;
-				ActorInstance->GetComponents(Components);
+				TSet<UActorComponent*> ComponentsToAdd(ActorInstance->GetComponents());
 
 				// Add the root component first (it may not be the first one)
 				USceneComponent* RootComponent = ActorInstance->GetRootComponent();
 				if(RootComponent != nullptr)
 				{
-					Components.Remove(RootComponent);
-					AddTreeNodeFromComponent(RootComponent);
+					ComponentsToAdd.Remove(RootComponent);
+
+					// Recursively add any instanced children that are already attached through the root, and keep track of added
+					// instances. This will be a faster path than the loop below, because we create new parent tree nodes as we go.
+					FSCSEditorTreeNodePtrType NewParentNode = AddTreeNodeFromComponent(RootComponent);
+					AddInstancedTreeNodesRecursive(RootComponent, NewParentNode, ComponentsToAdd);
 				}
-				
-				// Now add the rest of the instanced scene component hierarchy (excluding editor-only instances and nested DSOs attached to BP-constructed instances, which are not mutable)
-				for(auto CompIter = Components.CreateIterator(); CompIter; ++CompIter)
+
+				// Sort components by type (always put scene components first in the tree)
+				ComponentsToAdd.Sort([](const UActorComponent& A, const UActorComponent& /* B */)
 				{
-					USceneComponent* SceneComp = Cast<USceneComponent>(*CompIter);
-					USceneComponent* ParentSceneComp = SceneComp != nullptr ? SceneComp->GetAttachParent() : nullptr;
-					if(SceneComp != nullptr && !SceneComp->IsEditorOnly()
-						&& (SceneComp->CreationMethod != EComponentCreationMethod::UserConstructionScript || !GetDefault<UBlueprintEditorSettings>()->bHideConstructionScriptComponentsInDetailsView)
-						&& (ParentSceneComp == nullptr || !ParentSceneComp->IsCreatedByConstructionScript() || !SceneComp->HasAnyFlags(RF_DefaultSubObject)))
-					{
-						AddTreeNodeFromComponent(SceneComp);
-					}
-				}
+					return A.IsA<USceneComponent>();
+				});
 
-				// Add all non-scene component instances to the root set first
-				for(auto CompIter = Components.CreateIterator(); CompIter; ++CompIter)
+				// Now add any remaining instanced owned components not already added above. This will first add any
+				// unattached scene components followed by any instanced non-scene components owned by the Actor instance.
+				for (UActorComponent* ActorComp : ComponentsToAdd)
 				{
-					UActorComponent* ActorComp = *CompIter;
-					if (!ActorComp->IsA<USceneComponent>() && !ActorComp->IsEditorOnly()
-						&& (ActorComp->CreationMethod != EComponentCreationMethod::UserConstructionScript || !GetDefault<UBlueprintEditorSettings>()->bHideConstructionScriptComponentsInDetailsView))
+					USceneComponent* SceneComp = Cast<USceneComponent>(ActorComp);
+					USceneComponent* ParentSceneComp = SceneComp != nullptr ? SceneComp->GetAttachParent() : nullptr;
+					if (ShouldAddInstancedActorComponent(ActorComp, ParentSceneComp))
 					{
-						if (!bHasAddedSceneAndBehaviorComponentSeparator)
+						if (SceneComp != nullptr)
 						{
-							bHasAddedSceneAndBehaviorComponentSeparator = true;
-							RootNodes.Add(MakeShareable(new FSCSEditorTreeNode(FSCSEditorTreeNode::SeparatorNode)));
+							AddTreeNodeFromComponent(SceneComp);
+						}
+						else
+						{
+							if (!bHasAddedSceneAndBehaviorComponentSeparator)
+							{
+								bHasAddedSceneAndBehaviorComponentSeparator = true;
+								RootNodes.Add(MakeShareable(new FSCSEditorTreeNode(FSCSEditorTreeNode::SeparatorNode)));
+							}
+							AddRootComponentTreeNode(ActorComp);
 						}
-						AddRootComponentTreeNode(ActorComp);
 					}
 				}
 			}
@@ -4622,7 +4616,7 @@ void SSCSEditor::UpdateTree(bool bRegenerateTreeNodes)
 
 		// Restore the previous expansion state on the new tree nodes
 		TArray<FSCSEditorTreeNodePtrType> CollapsedTreeNodeArray = CollapsedTreeNodes.Array();
-		for(int i = 0; i < CollapsedTreeNodeArray.Num(); ++i)
+		for(int32 i = 0; i < CollapsedTreeNodeArray.Num(); ++i)
 		{
 			// Look for a component match in the new hierarchy; if found, mark it as collapsed to match the previous setting
 			FSCSEditorTreeNodePtrType NodeToExpandPtr = FindTreeNode(CollapsedTreeNodeArray[i]->GetComponentTemplate());
@@ -4678,6 +4672,34 @@ void SSCSEditor::UpdateTree(bool bRegenerateTreeNodes)
 	SCSTreeWidget->RequestTreeRefresh();
 }
 
+void SSCSEditor::AddInstancedTreeNodesRecursive(USceneComponent* Component, FSCSEditorTreeNodePtrType TreeNode, TSet<UActorComponent*>& ComponentsToAdd)
+{
+	if (Component != nullptr)
+	{
+		for (USceneComponent* ChildComponent : Component->GetAttachChildren())
+		{
+			if (ComponentsToAdd.Contains(ChildComponent)
+				&& ShouldAddInstancedActorComponent(ChildComponent, Component)
+				&& ChildComponent->GetOwner() == Component->GetOwner())
+			{
+				ComponentsToAdd.Remove(ChildComponent);
+
+				FSCSEditorTreeNodePtrType NewParentNode = AddTreeNodeFromComponent(ChildComponent, TreeNode);
+				AddInstancedTreeNodesRecursive(ChildComponent, NewParentNode, ComponentsToAdd);
+			}
+		}
+	}
+}
+
+bool SSCSEditor::ShouldAddInstancedActorComponent(UActorComponent* ActorComp, USceneComponent* ParentSceneComp) const
+{
+	// Exclude nested DSOs attached to BP-constructed instances, which are not mutable.
+	return (ActorComp != nullptr
+		&& (ActorComp->CreationMethod == EComponentCreationMethod::Instance || !ActorComp->IsEditorOnly())
+		&& (ActorComp->CreationMethod != EComponentCreationMethod::UserConstructionScript || !GetDefault<UBlueprintEditorSettings>()->bHideConstructionScriptComponentsInDetailsView)
+		&& (ParentSceneComp == nullptr || !ParentSceneComp->IsCreatedByConstructionScript() || !ActorComp->HasAnyFlags(RF_DefaultSubObject)));
+}
+
 void SSCSEditor::DumpTree()
 {
 	/* Example:
@@ -5878,7 +5900,7 @@ FSCSEditorTreeNodePtrType SSCSEditor::AddTreeNode(USCS_Node* InSCSNode, FSCSEdit
 	return NewNodePtr;
 }
 
-FSCSEditorTreeNodePtrType SSCSEditor::AddTreeNodeFromComponent(USceneComponent* InSceneComponent)
+FSCSEditorTreeNodePtrType SSCSEditor::AddTreeNodeFromComponent(USceneComponent* InSceneComponent, FSCSEditorTreeNodePtrType InParentTreeNode)
 {
 	FSCSEditorTreeNodePtrType NewNodePtr;
 
@@ -5890,17 +5912,25 @@ FSCSEditorTreeNodePtrType SSCSEditor::AddTreeNodeFromComponent(USceneComponent*
 		&& (EditorMode != EComponentEditorMode::ActorInstance || InSceneComponent->GetAttachParent()->GetOwner() == GetActorContext()))
 	{
 		// Attempt to find the parent node in the current tree
-		FSCSEditorTreeNodePtrType ParentNodePtr = FindTreeNode(InSceneComponent->GetAttachParent());
-		if(!ParentNodePtr.IsValid())
+		FSCSEditorTreeNodePtrType ParentNodePtr;
+		if (InParentTreeNode.IsValid())
+		{
+			ParentNodePtr = InParentTreeNode;
+		}
+		else
 		{
-			// If the actual attach parent wasn't found, attempt to find its archetype.
-			// This handles the BP editor case where we might add UCS component nodes taken
-			// from the preview actor instance, which are not themselves template objects.
-			ParentNodePtr = FindTreeNode(Cast<USceneComponent>(InSceneComponent->GetAttachParent()->GetArchetype()));
+			ParentNodePtr = FindTreeNode(InSceneComponent->GetAttachParent());
 			if(!ParentNodePtr.IsValid())
 			{
-				// Recursively add the parent node to the tree if it does not exist yet
-				ParentNodePtr = AddTreeNodeFromComponent(InSceneComponent->GetAttachParent());
+				// If the actual attach parent wasn't found, attempt to find its archetype.
+				// This handles the BP editor case where we might add UCS component nodes taken
+				// from the preview actor instance, which are not themselves template objects.
+				ParentNodePtr = FindTreeNode(Cast<USceneComponent>(InSceneComponent->GetAttachParent()->GetArchetype()));
+				if(!ParentNodePtr.IsValid())
+				{
+					// Recursively add the parent node to the tree if it does not exist yet
+					ParentNodePtr = AddTreeNodeFromComponent(InSceneComponent->GetAttachParent());
+				}
 			}
 		}
 
diff --git a/Engine/Source/Editor/Kismet/Private/Tests/BlueprintEditorTests.cpp b/Engine/Source/Editor/Kismet/Private/Tests/BlueprintEditorTests.cpp
--- a/Engine/Source/Editor/Kismet/Private/Tests/BlueprintEditorTests.cpp
+++ b/Engine/Source/Editor/Kismet/Private/Tests/BlueprintEditorTests.cpp
@@ -253,7 +253,7 @@ namespace BlueprintEditorPromotionUtils
 		check(OldParent);
 
 		//Remove the new root from its old parent and 
-		OldParent->ChildNodes.Remove(NewRoot);
+		OldParent->RemoveChildNode(NewRoot);
 		NewRoot->Modify();
 		NewRoot->AttachToName = NAME_None;
 
diff --git a/Engine/Source/Editor/Kismet/Private/WatchPointViewer.cpp b/Engine/Source/Editor/Kismet/Private/WatchPointViewer.cpp
--- a/Engine/Source/Editor/Kismet/Private/WatchPointViewer.cpp
+++ b/Engine/Source/Editor/Kismet/Private/WatchPointViewer.cpp
@@ -1,4 +1,4 @@
-// Copyright 1998-2017 Epic Games, Inc. All Rights Reserved.
+﻿// Copyright 1998-2018 Epic Games, Inc. All Rights Reserved.
 
 #include "WatchPointViewer.h"
 
@@ -13,6 +13,8 @@
 #include "Kismet2/KismetDebugUtilities.h"
 #include "GraphEditorActions.h"
 #include "EdGraphSchema_K2.h"
+#include "UnrealEdGlobals.h"
+#include "Editor/UnrealEdEngine.h"
 
 #include "Editor.h"
 
@@ -38,6 +40,7 @@ struct FWatchRow
 		UBlueprint* InBP,
 		const UEdGraphNode* InNode,
 		const UEdGraphPin* InPin,
+		UObject* InObjectBeingDebugged,
 		FText InBlueprintName,
 		FText InGraphName,
 		FText InNodeName,
@@ -48,6 +51,7 @@ struct FWatchRow
 		: BP(InBP)
 		, Node(InNode)
 		, Pin(InPin)
+		, ObjectBeingDebugged(InObjectBeingDebugged)
 		, BlueprintName(MoveTemp(InBlueprintName))
 		, GraphName(MoveTemp(InGraphName))
 		, NodeName(MoveTemp(InNodeName))
@@ -55,14 +59,14 @@ struct FWatchRow
 		, Value(MoveTemp(InValue))
 		, Type(MoveTemp(InType))
 	{
-		ObjectBeingDebugged = BP ? BP->GetObjectBeingDebugged() : nullptr;
-		ObjectBeingDebuggedName = ObjectBeingDebugged ? FText::FromString(ObjectBeingDebugged->GetName()) : LOCTEXT("Unknown Object", "Unknown object");
+		SetObjectBeingDebuggedName();
 	}
 
 	FWatchRow(
 		UBlueprint* InBP,
 		const UEdGraphNode* InNode,
 		const UEdGraphPin* InPin,
+		UObject* InObjectBeingDebugged,
 		FText InBlueprintName,
 		FText InGraphName,
 		FText InNodeName,
@@ -71,6 +75,7 @@ struct FWatchRow
 		: BP(InBP)
 		, Node(InNode)
 		, Pin(InPin)
+		, ObjectBeingDebugged(InObjectBeingDebugged)
 		, BlueprintName(MoveTemp(InBlueprintName))
 		, GraphName(MoveTemp(InGraphName))
 		, NodeName(MoveTemp(InNodeName))
@@ -78,12 +83,11 @@ struct FWatchRow
 		, Value(MoveTemp(Info.Value))
 		, Type(MoveTemp(Info.Type))
 	{
-		BlueprintName = MoveTemp(InBlueprintName);
-		ObjectBeingDebugged = BP ? BP->GetObjectBeingDebugged() : nullptr;
-		ObjectBeingDebuggedName = ObjectBeingDebugged ? FText::FromString(ObjectBeingDebugged->GetName()) : LOCTEXT("Unknown Object", "Unknown object");
+		SetObjectBeingDebuggedName();
+
 		for (const FDebugInfo& ChildInfo : Info.Children)
 		{
-			Children.Add(MakeShared<FWatchRow>(InBP, InNode, InPin, BlueprintName, GraphName, NodeName, ChildInfo));
+			Children.Add(MakeShared<FWatchRow>(InBP, InNode, InPin, InObjectBeingDebugged, BlueprintName, GraphName, NodeName, ChildInfo));
 		}
 	}
 
@@ -117,6 +121,27 @@ struct FWatchRow
 		Args.Add(TEXT("Value"), Value);
 		return FText::Format(LOCTEXT("WatchEntry", "{ObjectName}({BlueprintName}) {GraphName} {NodeName} {DisplayName}({Type}): {Value}"), Args);
 	}
+
+private:
+	void SetObjectBeingDebuggedName()
+	{
+		if (ObjectBeingDebugged != nullptr)
+		{
+			AActor* ActorBeingDebugged = Cast<AActor>(ObjectBeingDebugged);
+			if (ActorBeingDebugged)
+			{
+				ObjectBeingDebuggedName = FText::AsCultureInvariant(ActorBeingDebugged->GetActorLabel());
+			}
+			else
+			{
+				ObjectBeingDebuggedName = FText::FromName(ObjectBeingDebugged->GetFName());
+			}
+		}
+		else
+		{
+			ObjectBeingDebuggedName = BlueprintName;
+		}
+	}
 };
 
 DECLARE_MULTICAST_DELEGATE_OneParam(FOnDisplayedWatchWindowChanged, TArray<TSharedRef<FWatchRow>>*);
@@ -189,12 +214,36 @@ protected:
 
 	void HandleHyperlinkDebuggedObjectNavigate() const
 	{
-		if (AActor* Actor = Cast<AActor>(WatchRow->ObjectBeingDebugged))
+		if (AActor* Actor = Cast<AActor>(WatchRow.IsValid() ? WatchRow->ObjectBeingDebugged : nullptr))
 		{
+			// unselect whatever was selected
+			GEditor->SelectNone(false, false, false);
+
+			// select the actor we care about
 			GEditor->SelectActor(Actor, true, true, true);
 		}
 	}
 
+	EVisibility DisplayDebuggedObjectAsHyperlink() const
+	{
+		if ( AActor* Actor = Cast<AActor>(WatchRow.IsValid() ? WatchRow->ObjectBeingDebugged : nullptr))
+		{
+			return EVisibility::Visible;
+		}
+
+		return EVisibility::Collapsed;
+	}
+
+	EVisibility DisplayDebuggedObjectAsText() const
+	{
+		if (AActor* Actor = Cast<AActor>(WatchRow.IsValid() ? WatchRow->ObjectBeingDebugged : nullptr))
+		{
+			return EVisibility::Collapsed;
+		}
+
+		return EVisibility::Visible;
+	}
+
 	void HandleHyperlinkNodeNavigate() const
 	{
 		if (WatchRow.IsValid() && WatchRow->Node)
@@ -223,16 +272,22 @@ public:
 	SWatchViewer()
 	{
 		FKismetDebugUtilities::WatchedPinsListChangedEvent.AddRaw(this, &SWatchViewer::HandleWatchedPinsChanged);
+		FEditorDelegates::ResumePIE.AddRaw(this, &SWatchViewer::HandleResumePIE);
+		FEditorDelegates::EndPIE.AddRaw(this, &SWatchViewer::HandleEndPIE);
 	}
 	~SWatchViewer()
 	{
 		FKismetDebugUtilities::WatchedPinsListChangedEvent.RemoveAll(this);
+		FEditorDelegates::ResumePIE.RemoveAll(this);
+		FEditorDelegates::EndPIE.RemoveAll(this);
 	}
 
 	void Construct(const FArguments& InArgs, TArray<TSharedRef<FWatchRow>>* InWatchSource);
 	TSharedRef<ITableRow> HandleGenerateRow(TSharedRef<FWatchRow> InWatchRow, const TSharedRef<STableViewBase>& OwnerTable);
 	void HandleGetChildren(TSharedRef<FWatchRow> InWatchRow, TArray<TSharedRef<FWatchRow>>& OutChildren);
 	void HandleWatchedPinsChanged(UBlueprint* BlueprintObj);
+	void HandleResumePIE(bool);
+	void HandleEndPIE(bool);
 	void UpdateWatches(TArray<TSharedRef<FWatchRow>>* WatchValues);
 	void CopySelectedRows() const;
 	void StopWatchingPin() const;
@@ -241,6 +296,9 @@ public:
 	TArray<TSharedRef<FWatchRow>>* WatchSource;
 
 	TSharedPtr< FUICommandList > CommandList;
+
+private:
+	void CopySelectedRowsHelper(const TArray<TSharedRef<FWatchRow>>& RowSource, FString& StringToCopy) const;
 };
 
 void SWatchViewer::Construct(const FArguments& InArgs, TArray<TSharedRef<FWatchRow>>* InWatchSource)
@@ -325,48 +383,28 @@ void SWatchViewer::Construct(const FArguments& InArgs, TArray<TSharedRef<FWatchR
 					+SHeaderRow::Column(TEXT("ObjectName"))
 					.FillWidth(.2f)
 					.VAlignHeader(VAlign_Center)
-					.HeaderContent()
-					[
-						SNew(STextBlock)
-						.Text(LOCTEXT("ObjectName", "Object Name"))
-						.ToolTipText(LOCTEXT("ObjectNameTooltip", "Name of the object being debugged"))
-					]
+					.DefaultLabel(LOCTEXT("ObjectName", "Object Name"))
+					.DefaultTooltip(LOCTEXT("ObjectNameTooltip", "Name of the object instance being debugged or the blueprint if there is no object being debugged"))
 					+ SHeaderRow::Column(TEXT("GraphName"))
 					.FillWidth(.2f)
 					.VAlignHeader(VAlign_Center)
-					.HeaderContent()
-					[
-						SNew(STextBlock)
-						.Text(LOCTEXT("GraphName", "Graph Name"))
-						.ToolTipText(LOCTEXT("GraphNameTooltip", "Name of the source blueprint graph for this variable"))
-					]
+					.DefaultLabel(LOCTEXT("GraphName", "Graph Name"))
+					.DefaultTooltip(LOCTEXT("GraphNameTooltip", "Name of the source blueprint graph for this variable"))
 					+ SHeaderRow::Column(TEXT("NodeName"))
 					.FillWidth(.3f)
 					.VAlignHeader(VAlign_Center)
-					.HeaderContent()
-					[
-						SNew(STextBlock)
-						.Text(LOCTEXT("NodeName", "Node Name"))
-						.ToolTipText(LOCTEXT("NodeNameTooltip", "Name of the source blueprint graph node for this variable"))
-					]
+					.DefaultLabel(LOCTEXT("NodeName", "Node Name"))
+					.DefaultTooltip(LOCTEXT("NodeNameTooltip", "Name of the source blueprint graph node for this variable"))
 					+ SHeaderRow::Column(TEXT("VariableName"))
 					.FillWidth(.3f)
 					.VAlignHeader(VAlign_Center)
-					.HeaderContent()
-					[
-						SNew(STextBlock)
-						.Text(LOCTEXT("VariableName", "Variable Name"))
-						.ToolTipText(LOCTEXT("VariabelNameTooltip", "Name of the variable"))
-					]
+					.DefaultLabel(LOCTEXT("VariableName", "Variable Name"))
+					.DefaultTooltip(LOCTEXT("VariabelNameTooltip", "Name of the variable"))
 					+ SHeaderRow::Column(TEXT("Value"))
 					.FillWidth(.8f)
 					.VAlignHeader(VAlign_Center)
-					.HeaderContent()
-					[
-						SNew(STextBlock)
-						.Text(LOCTEXT("Value", "Value"))
-						.ToolTipText(LOCTEXT("ValueTooltip", "Current value of this variable"))
-					]
+					.DefaultLabel(LOCTEXT("Value", "Value"))
+					.DefaultTooltip(LOCTEXT("ValueTooltip", "Current value of this variable"))
 				)
 			]
 			+SOverlay::Slot()
@@ -402,12 +440,39 @@ void SWatchViewer::HandleGetChildren(TSharedRef<FWatchRow> InWatchRow, TArray<TS
 
 void SWatchViewer::HandleWatchedPinsChanged(UBlueprint* BlueprintObj)
 {
-	WatchViewer::UpdateDisplayedWatches(BlueprintObj);
+	WatchViewer::UpdateWatchListFromBlueprint(BlueprintObj);
+}
+
+void SWatchViewer::HandleResumePIE(bool)
+{
+	// swap to displaying the unpaused watches
+	WatchViewer::ContinueExecution();
+}
+
+void SWatchViewer::HandleEndPIE(bool)
+{
+	// show the unpaused watches in case we stopped PIE while at a breakpoint
+	WatchViewer::ContinueExecution();
 }
 
 void SWatchViewer::UpdateWatches(TArray<TSharedRef<FWatchRow>>* Watches)
 {
-	WatchTreeWidget->RequestTreeRefresh();
+	WatchSource = Watches;
+	WatchTreeWidget->SetTreeItemsSource(Watches);
+}
+
+void SWatchViewer::CopySelectedRowsHelper(const TArray<TSharedRef<FWatchRow>>& RowSource, FString& StringToCopy) const
+{
+	for (const TSharedRef<FWatchRow>& Item : RowSource)
+	{
+		if (WatchTreeWidget->IsItemSelected(Item))
+		{
+			StringToCopy.Append(Item->GetTextForEntry().ToString());
+			StringToCopy.Append(LINE_TERMINATOR);
+		}
+
+		CopySelectedRowsHelper(Item->Children, StringToCopy);
+	}
 }
 
 void SWatchViewer::CopySelectedRows() const
@@ -417,14 +482,7 @@ void SWatchViewer::CopySelectedRows() const
 	// We want to copy in the order displayed, not the order selected, so iterate the list and build up the string:
 	if (WatchSource)
 	{
-		for (const TSharedRef<FWatchRow>& Item : *WatchSource)
-		{
-			if (WatchTreeWidget->IsItemSelected(Item))
-			{
-				StringToCopy.Append(Item->GetTextForEntry().ToString());
-				StringToCopy.Append(TEXT("\r\n"));
-			}
-		}
+		CopySelectedRowsHelper(*WatchSource, StringToCopy);
 	}
 
 	if (!StringToCopy.IsEmpty())
@@ -433,14 +491,6 @@ void SWatchViewer::CopySelectedRows() const
 	}
 }
 
-void SWatchTreeWidgetItem::Construct(const FArguments& InArgs, SWatchViewer* InOwner, const TSharedRef<STableViewBase>& InOwnerTableView)
-{
-	this->WatchRow = InArgs._WatchToVisualize;
-	Owner = InOwner;
-
-	SMultiColumnTableRow< TSharedRef<FWatchRow> >::Construct(SMultiColumnTableRow< TSharedRef<FWatchRow> >::FArguments().Padding(1), InOwnerTableView);
-}
-
 void SWatchViewer::StopWatchingPin() const
 {
 	TArray<TSharedRef<FWatchRow>> SelectedRows = WatchTreeWidget->GetSelectedItems();
@@ -450,6 +500,14 @@ void SWatchViewer::StopWatchingPin() const
 	}
 }
 
+void SWatchTreeWidgetItem::Construct(const FArguments& InArgs, SWatchViewer* InOwner, const TSharedRef<STableViewBase>& InOwnerTableView)
+{
+	this->WatchRow = InArgs._WatchToVisualize;
+	Owner = InOwner;
+
+	SMultiColumnTableRow< TSharedRef<FWatchRow> >::Construct(SMultiColumnTableRow< TSharedRef<FWatchRow> >::FArguments().Padding(1), InOwnerTableView);
+}
+
 TSharedRef<SWidget> SWatchTreeWidgetItem::GenerateWidgetForColumn(const FName& ColumnName)
 {
 	const static FName NAME_ObjectName(TEXT("ObjectName"));
@@ -463,12 +521,26 @@ TSharedRef<SWidget> SWatchTreeWidgetItem::GenerateWidgetForColumn(const FName& C
 		return SNew(SBox)
 			.HAlign(HAlign_Left)
 			.VAlign(VAlign_Center)
-			.Padding(FMargin(2.0f, 0.0f))
+			.Padding(FMargin(2.0f, 1.0f))
 			[
-				SNew(SHyperlink)
-				.Text(this, &SWatchTreeWidgetItem::GetDebuggedObjectName)
-				.ToolTipText(this, &SWatchTreeWidgetItem::GetBlueprintName)
-				.OnNavigate(this, &SWatchTreeWidgetItem::HandleHyperlinkDebuggedObjectNavigate)
+				SNew(SHorizontalBox)
+				+ SHorizontalBox::Slot()
+				.AutoWidth()
+				[
+					SNew(SHyperlink)
+					.Text(this, &SWatchTreeWidgetItem::GetDebuggedObjectName)
+					.ToolTipText(this, &SWatchTreeWidgetItem::GetBlueprintName)
+					.OnNavigate(this, &SWatchTreeWidgetItem::HandleHyperlinkDebuggedObjectNavigate)
+					.Visibility(this, &SWatchTreeWidgetItem::DisplayDebuggedObjectAsHyperlink)
+				]
+				+ SHorizontalBox::Slot()
+				.AutoWidth()
+				[
+					SNew(STextBlock)
+					.Text(this, &SWatchTreeWidgetItem::GetDebuggedObjectName)
+					.ToolTipText(this, &SWatchTreeWidgetItem::GetBlueprintName)
+					.Visibility(this, &SWatchTreeWidgetItem::DisplayDebuggedObjectAsText)
+				]
 			];
 	}
 	else if (ColumnName == NAME_GraphName)
@@ -476,7 +548,7 @@ TSharedRef<SWidget> SWatchTreeWidgetItem::GenerateWidgetForColumn(const FName& C
 		return SNew(SBox)
 			.HAlign(HAlign_Left)
 			.VAlign(VAlign_Center)
-			.Padding(FMargin(2.0f, 0.0f))
+			.Padding(FMargin(2.0f, 1.0f))
 			[
 				SNew(STextBlock)
 				.Text(this, &SWatchTreeWidgetItem::GetGraphName)
@@ -494,7 +566,7 @@ TSharedRef<SWidget> SWatchTreeWidgetItem::GenerateWidgetForColumn(const FName& C
 		return SNew(SBox)
 			.HAlign(HAlign_Left)
 			.VAlign(VAlign_Center)
-			.Padding(FMargin(2.0f, 0.0f))
+			.Padding(FMargin(2.0f, 1.0f))
 			[
 				SNew(SHyperlink)
 				.Text(this, &SWatchTreeWidgetItem::GetNodeName)
@@ -513,7 +585,7 @@ TSharedRef<SWidget> SWatchTreeWidgetItem::GenerateWidgetForColumn(const FName& C
 			]
 		+ SHorizontalBox::Slot()
 			.AutoWidth()
-			.Padding(2.0f, 0.0f)
+			.Padding(2.0f, 1.0f)
 			.VAlign(VAlign_Center)
 			[
 				SNew(STextBlock)
@@ -526,7 +598,7 @@ TSharedRef<SWidget> SWatchTreeWidgetItem::GenerateWidgetForColumn(const FName& C
 		return SNew(SBox)
 			.HAlign(HAlign_Left)
 			.VAlign(VAlign_Center)
-			.Padding(FMargin(2.0f, 0.0f))
+			.Padding(FMargin(2.0f, 1.0f))
 			[
 				SNew(STextBlock)
 				.Text(this, &SWatchTreeWidgetItem::GetValue)
@@ -540,77 +612,238 @@ TSharedRef<SWidget> SWatchTreeWidgetItem::GenerateWidgetForColumn(const FName& C
 
 // Proxy array of the watches. This allows us to manually refresh UI state when changes are made:
 TArray<TSharedRef<FWatchRow>> Private_WatchSource;
+TArray<TSharedRef<FWatchRow>> Private_InstanceWatchSource;
 
-void WatchViewer::UpdateDisplayedWatches(UBlueprint* BlueprintObj)
+TArray<UBlueprint*> WatchedBlueprints;
+TArray<UObject*> BlueprintStackInstances;
+
+bool bIsExecutionPaused = false;
+
+void UpdateNonInstancedWatchDisplay()
 {
-	if (!ensure(BlueprintObj))
-	{
-		return;
-	}
 	Private_WatchSource.Reset();
 
-	FText BlueprintName = FText::FromString(BlueprintObj->GetName());
-
-	for (const FEdGraphPinReference& PinRef : BlueprintObj->WatchedPins)
+	for (UBlueprint* BlueprintObj : WatchedBlueprints)
 	{
-		UEdGraphPin* Pin = PinRef.Get();
+		FText BlueprintName = FText::FromString(BlueprintObj->GetName());
 
-		FText GraphName = FText::FromString(Pin->GetOwningNode()->GetGraph()->GetName());
-		FText NodeName = Pin->GetOwningNode()->GetNodeTitle(ENodeTitleType::ListView);
+		for (const FEdGraphPinReference& PinRef : BlueprintObj->WatchedPins)
+		{
+			UEdGraphPin* Pin = PinRef.Get();
 
-		FDebugInfo DebugInfo;
-		const FKismetDebugUtilities::EWatchTextResult WatchStatus = FKismetDebugUtilities::GetDebugInfo(DebugInfo, BlueprintObj, BlueprintObj->GetObjectBeingDebugged(), Pin);
+			FText GraphName = FText::FromString(Pin->GetOwningNode()->GetGraph()->GetName());
+			FText NodeName = Pin->GetOwningNode()->GetNodeTitle(ENodeTitleType::ListView);
 
-		if (WatchStatus != FKismetDebugUtilities::EWTR_Valid)
-		{
 			const UEdGraphSchema* Schema = Pin->GetOwningNode()->GetSchema();
+
+			FDebugInfo DebugInfo;
 			DebugInfo.DisplayName = Schema->GetPinDisplayName(Pin);
 			DebugInfo.Type = UEdGraphSchema_K2::TypeToText(Pin->PinType);
+			DebugInfo.Value = LOCTEXT("ExecutionNotPaused", "(execution not paused)");
+
+			Private_WatchSource.Add(
+				MakeShared<FWatchRow>(
+					BlueprintObj,
+					Pin->GetOwningNode(),
+					Pin,
+					nullptr,
+					BlueprintName,
+					GraphName,
+					NodeName,
+					DebugInfo
+					)
+			);
+		}
+	}
+}
+
+void UpdateInstancedWatchDisplay()
+{
+	Private_InstanceWatchSource.Reset();
+
+	for (UObject* BlueprintInstance : BlueprintStackInstances)
+	{
+		UClass* Class = BlueprintInstance->GetClass();
+		UBlueprint* BlueprintObj = (Class ? Cast<UBlueprint>(Class->ClassGeneratedBy) : nullptr);
+		if (BlueprintObj == nullptr)
+		{
+			continue;
+		}
+
+		FText BlueprintName = FText::FromString(BlueprintObj->GetName());
+
+		// Don't show info for the CDO
+		if (BlueprintInstance->IsDefaultSubobject())
+		{
+			continue;
+		}
+
+		// Don't show info if this instance is pending kill
+		if (BlueprintInstance->IsPendingKill())
+		{
+			continue;
+		}
+
+		// Don't show info if this instance isn't in the current world
+		UObject* ObjOuter = BlueprintInstance;
+		UWorld* ObjWorld = nullptr;
+		static bool bUseNewWorldCode = false;
+		do		// Run through at least once in case the TestObject is a UGameInstance
+		{
+			UGameInstance* ObjGameInstance = Cast<UGameInstance>(ObjOuter);
 
-			switch (WatchStatus)
+			ObjOuter = ObjOuter->GetOuter();
+			ObjWorld = ObjGameInstance ? ObjGameInstance->GetWorld() : Cast<UWorld>(ObjOuter);
+		} while (ObjWorld == nullptr && ObjOuter != nullptr);
+
+		if (ObjWorld)
+		{
+			// Make check on owning level (not streaming level)
+			if (ObjWorld->PersistentLevel && ObjWorld->PersistentLevel->OwningWorld)
+			{
+				ObjWorld = ObjWorld->PersistentLevel->OwningWorld;
+			}
+
+			if (ObjWorld->WorldType != EWorldType::PIE && !((ObjWorld->WorldType == EWorldType::Editor) && (GUnrealEd->GetPIEViewport() == nullptr)))
 			{
-			case FKismetDebugUtilities::EWTR_NotInScope:
-				DebugInfo.Value = LOCTEXT("NotInScope", "(not in scope)");
-				break;
+				continue;
+			}
+		}
 
-			case FKismetDebugUtilities::EWTR_NoProperty:
-				DebugInfo.Value = LOCTEXT("NoDebugData", "(no debug data)");
-				break;
+		// We have a valid instance, iterate over all the watched pins and create rows for them
+		for (const FEdGraphPinReference& PinRef : BlueprintObj->WatchedPins)
+		{
+			UEdGraphPin* Pin = PinRef.Get();
 
-			case FKismetDebugUtilities::EWTR_NoDebugObject:
-				DebugInfo.Value = LOCTEXT("NoDebugObject", "(no debug object)");
-				break;
+			FText GraphName = FText::FromString(Pin->GetOwningNode()->GetGraph()->GetName());
+			FText NodeName = Pin->GetOwningNode()->GetNodeTitle(ENodeTitleType::ListView);
 
-			default:
-				// do nothing
-				break;
+			FDebugInfo DebugInfo;
+			const FKismetDebugUtilities::EWatchTextResult WatchStatus = FKismetDebugUtilities::GetDebugInfo(DebugInfo, BlueprintObj, BlueprintInstance, Pin);
+
+			if (WatchStatus != FKismetDebugUtilities::EWTR_Valid)
+			{
+				const UEdGraphSchema* Schema = Pin->GetOwningNode()->GetSchema();
+				DebugInfo.DisplayName = Schema->GetPinDisplayName(Pin);
+				DebugInfo.Type = UEdGraphSchema_K2::TypeToText(Pin->PinType);
+
+				switch (WatchStatus)
+				{
+				case FKismetDebugUtilities::EWTR_NotInScope:
+					DebugInfo.Value = LOCTEXT("NotInScope", "(not in scope)");
+					break;
+
+				case FKismetDebugUtilities::EWTR_NoProperty:
+					DebugInfo.Value = LOCTEXT("NoDebugData", "(no debug data)");
+					break;
+
+				case FKismetDebugUtilities::EWTR_NoDebugObject:
+					DebugInfo.Value = LOCTEXT("NoDebugObject", "(no debug object)");
+					break;
+
+				default:
+					// do nothing
+					break;
+				}
 			}
+
+			Private_InstanceWatchSource.Add(
+				MakeShared<FWatchRow>(
+					BlueprintObj,
+					Pin->GetOwningNode(),
+					Pin,
+					BlueprintInstance,
+					BlueprintName,
+					GraphName,
+					NodeName,
+					DebugInfo
+					)
+			);
 		}
+	}
 
-		Private_WatchSource.Add(
-			MakeShared<FWatchRow>(
-				BlueprintObj,
-				Pin->GetOwningNode(),
-				Pin,
-				BlueprintName,
-				GraphName,
-				NodeName,
-				DebugInfo
-				)
-		);
+	// Notify subscribers:
+	WatchListSubscribers.Broadcast(&Private_InstanceWatchSource);
+}
+
+void WatchViewer::UpdateDisplayedWatches(const TArray<const FFrame*>& ScriptStack)
+{
+	bIsExecutionPaused = true;
+	BlueprintStackInstances.Reset();
+
+	for (const FFrame* Frame : ScriptStack)
+	{
+		if (Frame == nullptr)
+		{
+			continue;
+		}
+
+		if (Frame->Object)
+		{
+			BlueprintStackInstances.AddUnique(Frame->Object);
+		}
+	}
+
+	UpdateInstancedWatchDisplay();
+}
+
+void WatchViewer::ContinueExecution()
+{
+	bIsExecutionPaused = false;
+	// Notify subscribers:
+	WatchListSubscribers.Broadcast(&Private_WatchSource);
+}
+
+FName WatchViewer::GetTabName()
+{
+	const FName TabName = TEXT("WatchViewer");
+	return TabName;
+}
+
+void WatchViewer::UpdateWatchListFromBlueprint(UBlueprint* BlueprintObj)
+{
+	if (!ensure(BlueprintObj))
+	{
+		return;
+	}
+
+	if (BlueprintObj->WatchedPins.Num() == 0)
+	{
+		// if this blueprint doesn't have any watched pins and we aren't watching it already then there is nothing to do
+		int32 FoundIdx = WatchedBlueprints.Find(BlueprintObj);
+		if (FoundIdx == INDEX_NONE)
+		{
+			return;
+		}
+
+		// since we're not watching any pins anymore we should remove it from the watched list
+		WatchedBlueprints.RemoveAt(FoundIdx);
+	}
+	else
+	{
+		// make sure the blueprint is in our list
+		WatchedBlueprints.AddUnique(BlueprintObj);
+	}
+
+	// something changed so we need to update the lists shown in the UI
+	UpdateNonInstancedWatchDisplay();
+
+	if (bIsExecutionPaused)
+	{
+		UpdateInstancedWatchDisplay();
 	}
 
 	// Notify subscribers:
 	WatchListSubscribers.Broadcast(&Private_WatchSource);
 }
 
-void WatchViewer::RegisterTabSpawner()
+void WatchViewer::RegisterTabSpawner(FTabManager& TabManager)
 {
 	const auto SpawnWatchViewTab = []( const FSpawnTabArgs& Args )
 	{
 		return SNew(SDockTab)
 			.TabRole( ETabRole::NomadTab )
-			.Label( LOCTEXT("TabTitle", "Watch Window") )
+			.Label( LOCTEXT("TabTitle", "Watches") )
 			[
 				SNew(SBorder)
 				.BorderImage( FEditorStyle::GetBrush("Docking.Tab.ContentAreaBrush") )
@@ -620,13 +853,9 @@ void WatchViewer::RegisterTabSpawner()
 			];
 	};
 	
-	const IWorkspaceMenuStructure& MenuStructure = WorkspaceMenu::GetMenuStructure();
-
-	const FName TabName = TEXT("WatchViewer");
-	FGlobalTabmanager::Get()->RegisterNomadTabSpawner( TabName, FOnSpawnTab::CreateStatic(SpawnWatchViewTab) )
+	TabManager.RegisterTabSpawner( WatchViewer::GetTabName(), FOnSpawnTab::CreateStatic(SpawnWatchViewTab) )
 		.SetDisplayName( LOCTEXT("TabTitle", "Watch Window") )
-		.SetTooltipText( LOCTEXT("TooltipText", "Open the watch window tab.") )
-		.SetGroup( MenuStructure.GetDeveloperToolsDebugCategory() );
+		.SetTooltipText( LOCTEXT("TooltipText", "Open the watch window tab") );
 }
 
 #undef LOCTEXT_NAMESPACE
diff --git a/Engine/Source/Editor/Kismet/Public/BlueprintEditorModule.h b/Engine/Source/Editor/Kismet/Public/BlueprintEditorModule.h
--- a/Engine/Source/Editor/Kismet/Public/BlueprintEditorModule.h
+++ b/Engine/Source/Editor/Kismet/Public/BlueprintEditorModule.h
@@ -14,14 +14,13 @@ class UUserDefinedEnum;
 class UUserDefinedStruct;
 struct Rect;
 
-extern const FName BlueprintEditorAppName;
-
 class IBlueprintEditor;
 class FBlueprintEditor;
 class UUserDefinedEnum;
 class UUserDefinedStruct;
 class IDetailCustomization;
 class FKismetCompilerContext;
+struct FBlueprintDebugger;
 
 /** Delegate used to customize variable display */
 DECLARE_DELEGATE_RetVal_OneParam(TSharedPtr<IDetailCustomization>, FOnGetVariableCustomizationInstance, TSharedPtr<IBlueprintEditor> /*BlueprintEditor*/);
@@ -233,4 +232,7 @@ private:
 
 	/** Reference to keep our custom configuration panel alive */
 	TSharedPtr<SWidget> ConfigurationPanel;
+
+	/** Blueprint debugger state - refactor into SBlueprintDebugger if needed */
+	TUniquePtr<FBlueprintDebugger> BlueprintDebugger;
 };
diff --git a/Engine/Source/Editor/Kismet/Public/CallStackViewer.h b/Engine/Source/Editor/Kismet/Public/CallStackViewer.h
--- a/Engine/Source/Editor/Kismet/Public/CallStackViewer.h
+++ b/Engine/Source/Editor/Kismet/Public/CallStackViewer.h
@@ -5,8 +5,12 @@
 #include "Widgets/SCompoundWidget.h"
 #include "WorkflowOrientedApp/WorkflowTabFactory.h"
 
+class FTabManager;
+struct FFrame;
+
 namespace CallStackViewer
 {
 	void KISMET_API UpdateDisplayedCallstack(const TArray<const FFrame*>& ScriptStack);
-	void RegisterTabSpawner();
+	FName GetTabName();
+	void RegisterTabSpawner(FTabManager& TabManager);
 }
diff --git a/Engine/Source/Editor/Kismet/Public/SSCSEditor.h b/Engine/Source/Editor/Kismet/Public/SSCSEditor.h
--- a/Engine/Source/Editor/Kismet/Public/SSCSEditor.h
+++ b/Engine/Source/Editor/Kismet/Public/SSCSEditor.h
@@ -404,7 +404,7 @@ public:
 	*
 	* @param InComponentTemplate The component template represented by this object.
 	*/
-	FSCSEditorTreeNodeInstanceAddedComponent(AActor* Owner, FName InComponentName);
+	FSCSEditorTreeNodeInstanceAddedComponent(AActor* Owner, UActorComponent* InComponentTemplate);
 
 	// FSCSEditorTreeNode public interface
 	virtual bool IsNative() const override { return false; }
@@ -1030,7 +1030,13 @@ protected:
 	FSCSEditorTreeNodePtrType AddTreeNode(USCS_Node* InSCSNode, FSCSEditorTreeNodePtrType InParentNodePtr, const bool bIsInheritedSCS);
 
 	/** Helper method to add a tree node for the given scene component */
-	FSCSEditorTreeNodePtrType AddTreeNodeFromComponent(USceneComponent* InSceneComponent);
+	FSCSEditorTreeNodePtrType AddTreeNodeFromComponent(USceneComponent* InSceneComponent, FSCSEditorTreeNodePtrType InParentTreeNode = FSCSEditorTreeNodePtrType());
+	
+	/** Helper method to recursively add tree nodes for an already-attached instanced scene component hierarchy */
+	void AddInstancedTreeNodesRecursive(USceneComponent* Component, FSCSEditorTreeNodePtrType TreeNode, TSet<UActorComponent*>& ComponentsToAdd);
+
+	/** Helper method to determine whether or not an instanced actor component should be added to the tree */
+	bool ShouldAddInstancedActorComponent(UActorComponent* ActorComp, USceneComponent* ParentSceneComp = nullptr) const;
 
 	/** Helper method to recursively find a tree node for the given SCS node starting at the given tree node */
 	FSCSEditorTreeNodePtrType FindTreeNode(const USCS_Node* InSCSNode, FSCSEditorTreeNodePtrType InStartNodePtr = FSCSEditorTreeNodePtrType()) const;
diff --git a/Engine/Source/Editor/Kismet/Public/WatchPointViewer.h b/Engine/Source/Editor/Kismet/Public/WatchPointViewer.h
--- a/Engine/Source/Editor/Kismet/Public/WatchPointViewer.h
+++ b/Engine/Source/Editor/Kismet/Public/WatchPointViewer.h
@@ -5,9 +5,17 @@
 #include "Widgets/SCompoundWidget.h"
 #include "WorkflowOrientedApp/WorkflowTabFactory.h"
 
-
 namespace WatchViewer
 {
-	void KISMET_API UpdateDisplayedWatches(class UBlueprint* BlueprintObj);
-	void RegisterTabSpawner();
+	// called when we pause execution to update the displayed watch values
+	void KISMET_API UpdateDisplayedWatches(const TArray<const FFrame*>& ScriptStack);
+
+	// called when we unpause execution and set watch values back to the blueprint versions
+	void KISMET_API ContinueExecution();
+
+	// called when we are adding or changing watches from BlueprintObj
+	void KISMET_API UpdateWatchListFromBlueprint(class UBlueprint* BlueprintObj);
+
+	FName GetTabName();
+	void RegisterTabSpawner(FTabManager& TabManager);
 }
diff --git a/Engine/Source/Editor/KismetCompiler/Private/KismetCompilerMisc.cpp b/Engine/Source/Editor/KismetCompiler/Private/KismetCompilerMisc.cpp
--- a/Engine/Source/Editor/KismetCompiler/Private/KismetCompilerMisc.cpp
+++ b/Engine/Source/Editor/KismetCompiler/Private/KismetCompilerMisc.cpp
@@ -571,7 +571,7 @@ UProperty* FKismetCompilerUtilities::FindPropertyInScope(UStruct* Scope, UEdGrap
 	// Couldn't find the name
 	if (!FKismetCompilerUtilities::IsMissingMemberPotentiallyLoading(Cast<UBlueprint>(SelfClass->ClassGeneratedBy), InitialScope))
 	{
-		MessageLog.Error(*LOCTEXT("PropertyNotFound_Error", "The property associated with @@ could not be found").ToString(), Pin);
+		MessageLog.Error(*FText::Format(LOCTEXT("PropertyNotFound_Error", "The property associated with @@ could not be found in '{0}'"), FText::FromString(SelfClass->GetPathName())).ToString(), Pin);
 	}
 	return NULL;
 }
diff --git a/Engine/Source/Editor/PropertyEditor/Private/PropertyHandleImpl.cpp b/Engine/Source/Editor/PropertyEditor/Private/PropertyHandleImpl.cpp
--- a/Engine/Source/Editor/PropertyEditor/Private/PropertyHandleImpl.cpp
+++ b/Engine/Source/Editor/PropertyEditor/Private/PropertyHandleImpl.cpp
@@ -292,26 +292,6 @@ FString FPropertyValueImpl::GetPropertyValueArray() const
 	return String;
 }
 
-bool FPropertyValueImpl::SendTextToObjectProperty( const FString& Text, EPropertyValueSetFlags::Type Flags )
-{
-	TSharedPtr<FPropertyNode> PropertyNodePin = PropertyNode.Pin();
-	if( PropertyNodePin.IsValid() )
-	{
-		FComplexPropertyNode* ParentNode = PropertyNodePin->FindComplexParent();
-
-		// If more than one object is selected, an empty field indicates their values for this property differ.
-		// Don't send it to the objects value in this case (if we did, they would all get set to None which isn't good).
-		if ((!ParentNode || ParentNode->GetInstancesNum() > 1) && !Text.Len())
-		{
-			return false;
-		}
-
-		return ImportText( Text, PropertyNodePin.Get(), Flags ) != FPropertyAccess::Fail;
-
-	}
-
-	return false;
-}
 
 void FPropertyValueImpl::GenerateArrayIndexMapToObjectNode( TMap<FString,int32>& OutArrayIndexMap, FPropertyNode* PropertyNode )
 {
@@ -855,163 +835,6 @@ FPropertyAccess::Result FPropertyValueImpl::SetValueAsString( const FString& InV
 	return Result;
 }
 
-bool FPropertyValueImpl::SetObject( const UObject* NewObject, EPropertyValueSetFlags::Type Flags )
-{
-	TSharedPtr<FPropertyNode> PropertyNodePin = PropertyNode.Pin();
-	if( PropertyNodePin.IsValid() )
-	{
-		FString ObjectPathName = NewObject ? NewObject->GetPathName() : TEXT("None");
-		bool bResult = SendTextToObjectProperty( ObjectPathName, Flags );
-
-		// @todo PropertyEditing:  This fails when it should not.  some UProperties remove the object class from the name and some dont. 
-		// So comparing the path name to one that adds the object class name will always fail even though the value was set correctly
-#if 0
-		if( bResult )
-		{
-			UProperty* NodeProperty = PropertyNodePin->GetProperty();
-
-			const FString CompareName = NewObject ? FString::Printf( TEXT("%s'%s'"), *NewObject->GetClass()->GetName(), *ObjectPathName ) : TEXT("None");
-
-			// Read values back and report any failures.
-			TArray<FObjectBaseAddress> ObjectsThatWereModified;
-			FObjectPropertyNode* ObjectNode = PropertyNodePin->FindObjectItemParent();
-			
-			if( ObjectNode )
-			{
-				bool bAllObjectPropertyValuesMatch = true;
-				for ( TPropObjectIterator Itor( ObjectNode->ObjectIterator() ) ; Itor ; ++Itor )
-				{
-					TWeakObjectPtr<UObject>	Object = *Itor;
-					uint8*		Addr = PropertyNodePin->GetValueBaseAddress( (uint8*) Object.Get() );
-
-					FString PropertyValue;
-					NodeProperty->ExportText_Direct(PropertyValue, Addr, Addr, nullptr, 0 );
-					if ( PropertyValue != CompareName )
-					{
-						bAllObjectPropertyValuesMatch = false;
-						break;
-					}
-				}
-
-				// Warn that some object assignments failed.
-				if ( !bAllObjectPropertyValuesMatch )
-				{
-					bResult = false;
-				}
-			}
-			else
-			{
-				bResult = false;
-			}
-		}
-#endif
-		return bResult;
-	}
-
-	return false;
-}
-
-FPropertyAccess::Result FPropertyValueImpl::OnUseSelected()
-{
-	FPropertyAccess::Result Res = FPropertyAccess::Fail;
-	TSharedPtr<FPropertyNode> PropertyNodePin = PropertyNode.Pin();
-	if( PropertyNodePin.IsValid() )
-	{
-		UProperty* NodeProperty = PropertyNodePin->GetProperty();
-
-		UObjectPropertyBase* ObjProp = Cast<UObjectPropertyBase>( NodeProperty );
-		UInterfaceProperty* IntProp = Cast<UInterfaceProperty>( NodeProperty );
-		UClassProperty* ClassProp = Cast<UClassProperty>( NodeProperty );
-		USoftClassProperty* SoftClassProperty = Cast<USoftClassProperty>( NodeProperty );
-		UClass* const InterfaceThatMustBeImplemented = ObjProp ? ObjProp->GetOwnerProperty()->GetClassMetaData(TEXT("MustImplement")) : nullptr;
-
-		if(ClassProp || SoftClassProperty)
-		{
-			FEditorDelegates::LoadSelectedAssetsIfNeeded.Broadcast();
-
-			const UClass* const SelectedClass = GEditor->GetFirstSelectedClass(ClassProp ? ClassProp->MetaClass : SoftClassProperty->MetaClass);
-			if(SelectedClass)
-			{
-				if (!InterfaceThatMustBeImplemented || SelectedClass->ImplementsInterface(InterfaceThatMustBeImplemented))
-				{
-					FString const ClassPathName = SelectedClass->GetPathName();
-
-					TArray<FText> RestrictReasons;
-					if (PropertyNodePin->IsRestricted(ClassPathName, RestrictReasons))
-					{
-						check(RestrictReasons.Num() > 0);
-						FMessageDialog::Open(EAppMsgType::Ok, RestrictReasons[0]);
-					}
-					else 
-					{
-						
-						Res = SetValueAsString(ClassPathName, EPropertyValueSetFlags::DefaultFlags);
-					}
-				}
-			}
-		}
-		else
-		{
-			FEditorDelegates::LoadSelectedAssetsIfNeeded.Broadcast();
-
-			UClass* ObjPropClass = UObject::StaticClass();
-			if( ObjProp )
-			{
-				ObjPropClass = ObjProp->PropertyClass;
-			}
-			else if( IntProp ) 
-			{
-				ObjPropClass = IntProp->InterfaceClass;
-			}
-
-			bool const bMustBeLevelActor = ObjProp ? ObjProp->GetOwnerProperty()->GetBoolMetaData(TEXT("MustBeLevelActor")) : false;
-
-			// Find best appropriate selected object
-			UObject* SelectedObject = nullptr;
-
-			if (bMustBeLevelActor)
-			{
-				// @todo: ensure at compile time that MustBeLevelActor flag is only set on actor properties
-
-				// looking only for level actors here
-				USelection* const SelectedSet = GEditor->GetSelectedActors();
-				SelectedObject = SelectedSet->GetTop( ObjPropClass, InterfaceThatMustBeImplemented);
-			}
-			else
-			{
-				// normal behavior, where actor classes will look for level actors and 
-				USelection* const SelectedSet = GEditor->GetSelectedSet( ObjPropClass );
-				SelectedObject = SelectedSet->GetTop( ObjPropClass, InterfaceThatMustBeImplemented );
-			}
-
-			if( SelectedObject )
-			{
-				FString const ObjPathName = SelectedObject->GetPathName();
-
-				TArray<FText> RestrictReasons;
-				if (PropertyNodePin->IsRestricted(ObjPathName, RestrictReasons))
-				{
-					check(RestrictReasons.Num() > 0);
-					FMessageDialog::Open(EAppMsgType::Ok, RestrictReasons[0]);
-				}
-				else if (!SetObject(SelectedObject, EPropertyValueSetFlags::DefaultFlags))
-				{
-					// Warn that some object assignments failed.
-					FMessageDialog::Open( EAppMsgType::Ok, FText::Format(
-						NSLOCTEXT("UnrealEd", "ObjectAssignmentsFailed", "Failed to assign {0} to the {1} property, see log for details."),
-						FText::FromString(SelectedObject->GetPathName()), PropertyNodePin->GetDisplayName()) );
-				}
-				else
-				{
-					Res = FPropertyAccess::Success;
-				}
-			}
-		}
-	}
-
-	return Res;
-}
-
 bool FPropertyValueImpl::IsPropertyTypeOf( UClass* ClassType ) const
 {
 	TSharedPtr<FPropertyNode> PropertyNodePin = PropertyNode.Pin();
@@ -3077,7 +2900,8 @@ void FPropertyHandleBase::NotifyFinishedChangingProperties()
 
 FPropertyAccess::Result FPropertyHandleBase::SetObjectValueFromSelection()
 {
-	return Implementation->OnUseSelected();
+	// Only implemented by Object handles
+	return FPropertyAccess::Result::Fail;
 }
 
 void FPropertyHandleBase::AddRestriction( TSharedRef<const FPropertyRestriction> Restriction )
@@ -3763,14 +3587,13 @@ FPropertyAccess::Result FPropertyHandleObject::SetValue( const UObject* const& N
 {
 	const TSharedPtr<FPropertyNode>& PropertyNode = Implementation->GetPropertyNode();
 
-	bool bResult = false;
 	if (!PropertyNode->HasNodeFlags(EPropertyNodeFlags::EditInlineNew))
 	{
 		FString ObjectPathName = NewValue ? NewValue->GetPathName() : TEXT("None");
-		bResult = Implementation->SendTextToObjectProperty(ObjectPathName, Flags);
+		return SetValueFromFormattedString(ObjectPathName, Flags);
 	}
 
-	return bResult ? FPropertyAccess::Success : FPropertyAccess::Fail;
+	return FPropertyAccess::Fail;
 }
 
 FPropertyAccess::Result FPropertyHandleObject::GetValue(FAssetData& OutValue) const
@@ -3790,7 +3613,6 @@ FPropertyAccess::Result FPropertyHandleObject::SetValue(const FAssetData& NewVal
 {
 	const TSharedPtr<FPropertyNode>& PropertyNode = Implementation->GetPropertyNode();
 
-	bool bResult = false;
 	if (!PropertyNode->HasNodeFlags(EPropertyNodeFlags::EditInlineNew))
 	{
 		if (!PropertyNode->GetProperty()->IsA(USoftObjectProperty::StaticClass()))
@@ -3800,42 +3622,124 @@ FPropertyAccess::Result FPropertyHandleObject::SetValue(const FAssetData& NewVal
 		}
 
 		FString ObjectPathName = NewValue.IsValid() ? NewValue.ObjectPath.ToString() : TEXT("None");
-		bResult = Implementation->SendTextToObjectProperty(ObjectPathName, Flags);
+		return SetValueFromFormattedString(ObjectPathName, Flags);
 	}
 
-	return bResult ? FPropertyAccess::Success : FPropertyAccess::Fail;
+	return FPropertyAccess::Fail;
 }
 
 FPropertyAccess::Result FPropertyHandleObject::SetValueFromFormattedString(const FString& InValue, EPropertyValueSetFlags::Type Flags)
 {
-	UProperty* Property = GetProperty();
-	const FString EmptyString = FString();
-	const FString& AllowedClassesString = Property ? Property->GetMetaData("AllowedClasses") : EmptyString;
+	// We need to do all of the type validation up front, to correctly support soft objects
+	TSharedPtr<FPropertyNode> PropertyNodePin = Implementation->GetPropertyNode();
+	UProperty* NodeProperty = PropertyNodePin.IsValid() ? PropertyNodePin->GetProperty() : nullptr;
+	const TCHAR* ObjectBuffer = *InValue;
+	UObject* QualifiedObject = nullptr;
 
-	if (Property && !AllowedClassesString.IsEmpty())
+	// This will attempt to load the object if it is not in memory. We purposefully pass in null as owner to avoid issues with cross level references
+	if (NodeProperty && UObjectPropertyBase::ParseObjectPropertyValue(NodeProperty, nullptr, UObject::StaticClass(), 0, ObjectBuffer, QualifiedObject))
 	{
-		const TCHAR* ObjectBuffer = *InValue;
-		UObject* QualifiedObject = nullptr;
-
-		// Check to see if the object we're attempting to import has a class that's allowed for the property. If not, bail early.
-		if (UObjectPropertyBase::ParseObjectPropertyValue(Property, Property->GetOuter(), UObject::StaticClass(), 0, ObjectBuffer, QualifiedObject))
+		if (QualifiedObject)
 		{
+			UObjectPropertyBase* ObjectProperty = Cast<UObjectPropertyBase>(NodeProperty);
+			UInterfaceProperty* InterfaceProperty = Cast<UInterfaceProperty>(NodeProperty);
+			UClassProperty* ClassProperty = Cast<UClassProperty>(NodeProperty);
+			USoftClassProperty* SoftClassProperty = Cast<USoftClassProperty>(NodeProperty);
+			USoftObjectProperty* SoftObjectProperty = Cast<USoftObjectProperty>(NodeProperty);
+
+			// Figure out what classes are required
+			UClass* InterfaceThatMustBeImplemented = nullptr;
+			UClass* RequiredClass = nullptr;
+			if (ClassProperty)
+			{
+				RequiredClass = ClassProperty->MetaClass;
+			}
+			else if (SoftClassProperty)
+			{
+				RequiredClass = SoftClassProperty->MetaClass;
+			}
+			else if (ObjectProperty)
+			{
+				RequiredClass = ObjectProperty->PropertyClass;
+				InterfaceThatMustBeImplemented = ObjectProperty->GetOwnerProperty()->GetClassMetaData(TEXT("MustImplement"));
+			}
+			else if (InterfaceProperty)
+			{
+				InterfaceThatMustBeImplemented = InterfaceProperty->InterfaceClass;
+			}
+			else
+			{
+				return FPropertyAccess::Fail;
+			}
+
+			// Figure out what class to check against
+			UClass* QualifiedClass = Cast<UClass>(QualifiedObject);
+			if (!QualifiedClass)
+			{
+				QualifiedClass = QualifiedObject->GetClass();
+			}
+
+			const FString& AllowedClassesString = NodeProperty->GetMetaData("AllowedClasses");
 			TArray<FString> AllowedClassNames;
 			AllowedClassesString.ParseIntoArray(AllowedClassNames, TEXT(","), true);
 			bool bSupportedObject = false;
 
-			for (auto ClassNameIt = AllowedClassNames.CreateConstIterator(); ClassNameIt; ++ClassNameIt)
+			// Check AllowedClasses metadata
+			if (AllowedClassNames.Num() > 0)
 			{
-				const FString& ClassName = *ClassNameIt;
-				UClass* AllowedClass = FindObject<UClass>(ANY_PACKAGE, *ClassName);
-				const bool bIsInterface = AllowedClass && AllowedClass->HasAnyClassFlags(CLASS_Interface);
-				
-				// Check if the object is an allowed class type this property supports
-				// Note: QualifieObject may be null if we're clearing the value.  Allow clears to pass through without a supported class
-				if (AllowedClass && (!QualifiedObject || QualifiedObject->IsA(AllowedClass) || (bIsInterface && QualifiedObject->GetClass()->ImplementsInterface(AllowedClass))))
+				for (const FString& ClassName : AllowedClassNames)
 				{
-					bSupportedObject = true;
-					break;
+					UClass* AllowedClass = FindObject<UClass>(ANY_PACKAGE, *ClassName);
+					const bool bIsInterface = AllowedClass && AllowedClass->HasAnyClassFlags(CLASS_Interface);
+
+					// Check if the object is an allowed class type this property supports
+					if ((AllowedClass && QualifiedObject->IsA(AllowedClass)) || (bIsInterface && QualifiedObject->GetClass()->ImplementsInterface(AllowedClass)))
+					{
+						bSupportedObject = true;
+						break;
+					}
+				}
+			}
+			else
+			{
+				bSupportedObject = true;
+			}
+
+			// Check required class
+			if (RequiredClass && !QualifiedClass->IsChildOf(RequiredClass))
+			{
+				bSupportedObject = false;
+			}
+
+			// Check required interface
+			if (InterfaceThatMustBeImplemented && !QualifiedClass->ImplementsInterface(InterfaceThatMustBeImplemented))
+			{
+				bSupportedObject = false;
+			}
+
+			// Check node restrictions
+			TArray<FText> RestrictReasons;
+			if (PropertyNodePin->IsRestricted(QualifiedObject->GetPathName(), RestrictReasons))
+			{
+				bSupportedObject = false;
+			}
+
+			// Check level actor
+			bool const bMustBeLevelActor = ObjectProperty ? ObjectProperty->GetOwnerProperty()->GetBoolMetaData(TEXT("MustBeLevelActor")) : false;
+			if (bMustBeLevelActor)
+			{
+				if (AActor* Actor = Cast<AActor>(QualifiedObject))
+				{
+					if (!Actor->GetLevel())
+					{
+						// Not in a level
+						bSupportedObject = false;
+					}
+				}
+				else
+				{
+					// Not an actor
+					bSupportedObject = false;
 				}
 			}
 
@@ -3844,14 +3748,117 @@ FPropertyAccess::Result FPropertyHandleObject::SetValueFromFormattedString(const
 				return FPropertyAccess::Fail;
 			}
 		}
+
+		// Parsing passed but object is null, we want to set it to null explicitly
+	}
+	else
+	{
+		// Failed parsing, it's either invalid format or a nonexistent object
+		return FPropertyAccess::Fail;
+	}
+
+	return FPropertyHandleBase::SetValueFromFormattedString(InValue, Flags);
+}
+
+FPropertyAccess::Result FPropertyHandleObject::SetObjectValueFromSelection()
+{
+	FPropertyAccess::Result Res = FPropertyAccess::Fail;
+	TSharedPtr<FPropertyNode> PropertyNodePin = Implementation->GetPropertyNode();
+	if (PropertyNodePin.IsValid())
+	{
+		UProperty* NodeProperty = PropertyNodePin->GetProperty();
+
+		UObjectPropertyBase* ObjProp = Cast<UObjectPropertyBase>(NodeProperty);
+		UInterfaceProperty* IntProp = Cast<UInterfaceProperty>(NodeProperty);
+		UClassProperty* ClassProp = Cast<UClassProperty>(NodeProperty);
+		USoftClassProperty* SoftClassProperty = Cast<USoftClassProperty>(NodeProperty);
+		UClass* const InterfaceThatMustBeImplemented = ObjProp ? ObjProp->GetOwnerProperty()->GetClassMetaData(TEXT("MustImplement")) : nullptr;
+
+		if (ClassProp || SoftClassProperty)
+		{
+			FEditorDelegates::LoadSelectedAssetsIfNeeded.Broadcast();
+
+			const UClass* const SelectedClass = GEditor->GetFirstSelectedClass(ClassProp ? ClassProp->MetaClass : SoftClassProperty->MetaClass);
+			if (SelectedClass)
+			{
+				if (!InterfaceThatMustBeImplemented || SelectedClass->ImplementsInterface(InterfaceThatMustBeImplemented))
+				{
+					FString const ClassPathName = SelectedClass->GetPathName();
+
+					TArray<FText> RestrictReasons;
+					if (PropertyNodePin->IsRestricted(ClassPathName, RestrictReasons))
+					{
+						check(RestrictReasons.Num() > 0);
+						FMessageDialog::Open(EAppMsgType::Ok, RestrictReasons[0]);
+					}
+					else
+					{
+
+						Res = SetValueFromFormattedString(ClassPathName, EPropertyValueSetFlags::DefaultFlags);
+					}
+				}
+			}
+		}
 		else
 		{
-			// Not an object, so bail.
-			return FPropertyAccess::Fail;
+			FEditorDelegates::LoadSelectedAssetsIfNeeded.Broadcast();
+
+			UClass* ObjPropClass = UObject::StaticClass();
+			if (ObjProp)
+			{
+				ObjPropClass = ObjProp->PropertyClass;
+			}
+			else if (IntProp)
+			{
+				ObjPropClass = IntProp->InterfaceClass;
+			}
+
+			bool const bMustBeLevelActor = ObjProp ? ObjProp->GetOwnerProperty()->GetBoolMetaData(TEXT("MustBeLevelActor")) : false;
+
+			// Find best appropriate selected object
+			UObject* SelectedObject = nullptr;
+
+			if (bMustBeLevelActor)
+			{
+				// @todo: ensure at compile time that MustBeLevelActor flag is only set on actor properties
+
+				// looking only for level actors here
+				USelection* const SelectedSet = GEditor->GetSelectedActors();
+				SelectedObject = SelectedSet->GetTop(ObjPropClass, InterfaceThatMustBeImplemented);
+			}
+			else
+			{
+				// normal behavior, where actor classes will look for level actors and 
+				USelection* const SelectedSet = GEditor->GetSelectedSet(ObjPropClass);
+				SelectedObject = SelectedSet->GetTop(ObjPropClass, InterfaceThatMustBeImplemented);
+			}
+
+			if (SelectedObject)
+			{
+				FString const ObjPathName = SelectedObject->GetPathName();
+
+				TArray<FText> RestrictReasons;
+				if (PropertyNodePin->IsRestricted(ObjPathName, RestrictReasons))
+				{
+					check(RestrictReasons.Num() > 0);
+					FMessageDialog::Open(EAppMsgType::Ok, RestrictReasons[0]);
+				}
+				else if (SetValueFromFormattedString(SelectedObject->GetPathName(), EPropertyValueSetFlags::DefaultFlags) != FPropertyAccess::Success)
+				{
+					// Warn that some object assignments failed.
+					FMessageDialog::Open(EAppMsgType::Ok, FText::Format(
+						NSLOCTEXT("UnrealEd", "ObjectAssignmentsFailed", "Failed to assign {0} to the {1} property, see log for details."),
+						FText::FromString(SelectedObject->GetPathName()), PropertyNodePin->GetDisplayName()));
+				}
+				else
+				{
+					Res = FPropertyAccess::Success;
+				}
+			}
 		}
 	}
 
-	return FPropertyHandleBase::SetValueFromFormattedString(InValue, Flags);
+	return Res;
 }
 
 // Vector
diff --git a/Engine/Source/Editor/PropertyEditor/Private/PropertyHandleImpl.h b/Engine/Source/Editor/PropertyEditor/Private/PropertyHandleImpl.h
--- a/Engine/Source/Editor/PropertyEditor/Private/PropertyHandleImpl.h
+++ b/Engine/Source/Editor/PropertyEditor/Private/PropertyHandleImpl.h
@@ -45,19 +45,6 @@ public:
 	 */
 	FPropertyValueImpl( TSharedPtr<FPropertyNode> InPropertyNode, FNotifyHook* InNotifyHook, TSharedPtr<IPropertyUtilities> InPropertyUtilities );
 
-	/**
-	 * Sets an object property to point to the new object
-	 * 
-	 * @param NewObject	The new object value
-	 */
-	bool SetObject( const UObject* NewObject, EPropertyValueSetFlags::Type Flags);
-
-	/**
-	 * Sets the value of an object property to the selected object in the content browser
-	 * @return Whether or not the call was successful
-	 */
-	FPropertyAccess::Result OnUseSelected();
-
 	/**
 	 * Recurse up to the next object node, adding all array indices into a map according to their property name
 	 * @param ArrayIndexMap - for the current object, what properties use which array offsets
@@ -174,14 +161,6 @@ public:
 	 */
 	void SetOnRebuildChildren( const FSimpleDelegate& InOnRebuildChildren );
 
-	/**
-	 * Sends a formatted string to an object property if safe to do so
-	 *
-	 * @param Text	The text to send
-	 * @return true if the text could be set
-	 */
-	bool SendTextToObjectProperty( const FString& Text, EPropertyValueSetFlags::Type Flags ); 
-
 	/**
 	 * Get the value of a property as a formatted string.
 	 * Each UProperty has a specific string format that it sets
@@ -607,6 +586,7 @@ public:
 	virtual FPropertyAccess::Result GetValue( FAssetData& OutValue ) const override;
 	virtual FPropertyAccess::Result SetValue( const FAssetData& InValue, EPropertyValueSetFlags::Type Flags = EPropertyValueSetFlags::DefaultFlags ) override;
 	virtual FPropertyAccess::Result SetValueFromFormattedString(const FString& InValue, EPropertyValueSetFlags::Type Flags = EPropertyValueSetFlags::DefaultFlags) override;
+	virtual FPropertyAccess::Result SetObjectValueFromSelection() override;
 };
 
 class FPropertyHandleVector : public FPropertyHandleBase
diff --git a/Engine/Source/Editor/UnrealEd/Classes/Settings/EditorExperimentalSettings.h b/Engine/Source/Editor/UnrealEd/Classes/Settings/EditorExperimentalSettings.h
--- a/Engine/Source/Editor/UnrealEd/Classes/Settings/EditorExperimentalSettings.h
+++ b/Engine/Source/Editor/UnrealEd/Classes/Settings/EditorExperimentalSettings.h
@@ -63,10 +63,6 @@ public:
 	UPROPERTY(EditAnywhere, config, Category = Blueprints, meta = (DisplayName = "Find and Replace All References Tool"))
 	bool bEnableFindAndReplaceReferences;
 
-	/** Bookmarks facilitate quick-jump navigation to saved locations in any Blueprint Graph context in the Blueprint Editor. */
-	UPROPERTY(EditAnywhere, config, Category = Blueprints)
-	bool bEnableBlueprintGraphBookmarks;
-
 	/** Should arrows indicating data/execution flow be drawn halfway along wires? */
 	UPROPERTY(/*EditAnywhere - deprecated (moved into UBlueprintEditorSettings), */config/*, Category=Blueprints, meta=(DisplayName="Draw midpoint arrows in Blueprints")*/)
 	bool bDrawMidpointArrowsInBlueprints;
diff --git a/Engine/Source/Editor/UnrealEd/Classes/Settings/ProjectPackagingSettings.h b/Engine/Source/Editor/UnrealEd/Classes/Settings/ProjectPackagingSettings.h
--- a/Engine/Source/Editor/UnrealEd/Classes/Settings/ProjectPackagingSettings.h
+++ b/Engine/Source/Editor/UnrealEd/Classes/Settings/ProjectPackagingSettings.h
@@ -283,20 +283,18 @@ public:
 
 	/**
 	 * Directories containing .uasset files that should always be cooked regardless of whether they're referenced by anything in your project
-	 * Note: These paths are relative to your project Content directory
+	 * These paths are stored relative to the project root so they can start with /game, /engine, or /pluginname
 	 */
-	UPROPERTY(config, EditAnywhere, Category=Packaging, AdvancedDisplay, meta=(DisplayName="Additional Asset Directories to Cook", RelativeToGameContentDir))
+	UPROPERTY(config, EditAnywhere, Category=Packaging, AdvancedDisplay, meta=(DisplayName="Additional Asset Directories to Cook", LongPackageName))
 	TArray<FDirectoryPath> DirectoriesToAlwaysCook;
-	
 
 	/**
-	* Directories containing .uasset files that should always be cooked regardless of whether they're referenced by anything in your project
-	* Note: These paths are relative to your project Content directory
-	*/
-	UPROPERTY(config, EditAnywhere, Category = Packaging, AdvancedDisplay, meta = (DisplayName = "Directories to never cook", RelativeToGameContentDir))
+	 * Directories containing .uasset files that should never be cooked even if they are referenced by your project
+	 * These paths are stored relative to the project root so they can start with /game, /engine, or /pluginname
+	 */
+	UPROPERTY(config, EditAnywhere, Category = Packaging, AdvancedDisplay, meta = (DisplayName = "Directories to never cook", LongPackageName))
 	TArray<FDirectoryPath> DirectoriesToNeverCook;
 
-
 	/**
 	 * Directories containing files that should always be added to the .pak file (if using a .pak file; otherwise they're copied as individual files)
 	 * This is used to stage additional files that you manually load via the UFS (Unreal File System) file IO API
@@ -359,4 +357,7 @@ public:
 private:
 	/** Returns the index of the specified Blueprint in the exclusive nativization list (otherwise INDEX_NONE) */
 	int32 FindBlueprintInNativizationList(const UBlueprint* InBlueprint) const;
+
+	/** Fix up cooking paths after they've been edited or laoded */
+	void FixCookingPaths();
 };
diff --git a/Engine/Source/Editor/UnrealEd/Private/Commandlets/CookCommandlet.cpp b/Engine/Source/Editor/UnrealEd/Private/Commandlets/CookCommandlet.cpp
--- a/Engine/Source/Editor/UnrealEd/Private/Commandlets/CookCommandlet.cpp
+++ b/Engine/Source/Editor/UnrealEd/Private/Commandlets/CookCommandlet.cpp
@@ -724,15 +724,24 @@ bool UCookCommandlet::CookByTheBook( const TArray<ITargetPlatform*>& Platforms,
 		CmdLineCultEntries += GetSwitchValueElements(TEXT("COOKCULTURES"));
 	}
 
-	// Also append any cookdirs from the project ini files; these dirs are relative to the game content directory
+	// Also append any cookdirs from the project ini files; these dirs are relative to the game content directory or start with a / root
 	{
 		const FString AbsoluteGameContentDir = FPaths::ConvertRelativePathToFull(FPaths::ProjectContentDir());
 		const UProjectPackagingSettings* const PackagingSettings = GetDefault<UProjectPackagingSettings>();
-		for (const auto& DirToCook : PackagingSettings->DirectoriesToAlwaysCook)
+		for (const FDirectoryPath& DirToCook : PackagingSettings->DirectoriesToAlwaysCook)
 		{
-			CmdLineDirEntries.Add(AbsoluteGameContentDir / DirToCook.Path);
+			if (DirToCook.Path.StartsWith(TEXT("/"), ESearchCase::CaseSensitive))
+			{
+				// If this starts with /, this includes a root like /engine
+				FString RelativePath = FPackageName::LongPackageNameToFilename(DirToCook.Path / TEXT(""));
+				CmdLineDirEntries.Add(FPaths::ConvertRelativePathToFull(RelativePath));
+			}
+			else
+			{
+				// This is relative to /game
+				CmdLineDirEntries.Add(AbsoluteGameContentDir / DirToCook.Path);
+			}
 		}
-
 	}
 
 	CookOnTheFlyServer->Initialize(ECookMode::CookByTheBook, CookFlags, OutputDirectoryOverride);
diff --git a/Engine/Source/Editor/UnrealEd/Private/Commandlets/GenerateDistillFileSetsCommandlet.cpp b/Engine/Source/Editor/UnrealEd/Private/Commandlets/GenerateDistillFileSetsCommandlet.cpp
--- a/Engine/Source/Editor/UnrealEd/Private/Commandlets/GenerateDistillFileSetsCommandlet.cpp
+++ b/Engine/Source/Editor/UnrealEd/Private/Commandlets/GenerateDistillFileSetsCommandlet.cpp
@@ -282,7 +282,19 @@ int32 UGenerateDistillFileSetsCommandlet::Main( const FString& InParams )
 	const FString AbsoluteGameContentDir = FPaths::ConvertRelativePathToFull(FPaths::ProjectContentDir());
 	for (const auto& DirToCook : PackagingSettings->DirectoriesToAlwaysCook)
 	{
-		FString DirectoryPath = AbsoluteGameContentDir / DirToCook.Path;
+		FString DirectoryPath;
+		if (DirToCook.Path.StartsWith(TEXT("/"), ESearchCase::CaseSensitive))
+		{
+			// If this starts with /, this includes a root like /engine
+			FString RelativePath = FPackageName::LongPackageNameToFilename(DirToCook.Path / TEXT(""));
+			DirectoryPath = FPaths::ConvertRelativePathToFull(RelativePath);
+		}
+		else
+		{
+			// This is relative to /game
+			DirectoryPath = AbsoluteGameContentDir / DirToCook.Path;
+		}
+
 		UE_LOG(LogGenerateDistillFileSetsCommandlet, Log, TEXT( "Examining directory to always cook: %s..." ), *DirToCook.Path );
 
 		TArray<FString> Files;
diff --git a/Engine/Source/Editor/UnrealEd/Private/ComponentTypeRegistry.cpp b/Engine/Source/Editor/UnrealEd/Private/ComponentTypeRegistry.cpp
--- a/Engine/Source/Editor/UnrealEd/Private/ComponentTypeRegistry.cpp
+++ b/Engine/Source/Editor/UnrealEd/Private/ComponentTypeRegistry.cpp
@@ -460,7 +460,9 @@ FComponentTypeRegistry::FComponentTypeRegistry()
 	: Data(nullptr)
 {
 	Data = new FComponentTypeRegistryData();
-	Data->ForceRefreshComponentList();
+
+	// This will load the assets on next tick. It's not safe to do right now because we could be deep in a stack
+	Data->Invalidate();
 
 	IHotReloadInterface& HotReloadSupport = FModuleManager::LoadModuleChecked<IHotReloadInterface>("HotReload");
 	HotReloadSupport.OnHotReload().AddRaw(this, &FComponentTypeRegistry::OnProjectHotReloaded);
diff --git a/Engine/Source/Editor/UnrealEd/Private/CookOnTheFlyServer.cpp b/Engine/Source/Editor/UnrealEd/Private/CookOnTheFlyServer.cpp
--- a/Engine/Source/Editor/UnrealEd/Private/CookOnTheFlyServer.cpp
+++ b/Engine/Source/Editor/UnrealEd/Private/CookOnTheFlyServer.cpp
@@ -4911,17 +4911,26 @@ void UCookOnTheFlyServer::CollectFilesToCook(TArray<FName>& FilesInPath, const T
 			}
 		}
 
-		// Also append any cookdirs from the project ini files; these dirs are relative to the game content directory
+		// Also append any cookdirs from the project ini files; these dirs are relative to the game content directory or start with a / root
 		{
 			const FString AbsoluteGameContentDir = FPaths::ConvertRelativePathToFull(FPaths::ProjectContentDir());
 			for (const FDirectoryPath& DirToCook : PackagingSettings->DirectoriesToAlwaysCook)
 			{
 				UE_LOG(LogCook, Verbose, TEXT("Loading directory to always cook %s"), *DirToCook.Path);
-				CookDirectories.Add(AbsoluteGameContentDir / DirToCook.Path);
+
+				if (DirToCook.Path.StartsWith(TEXT("/"), ESearchCase::CaseSensitive))
+				{
+					// If this starts with /, this includes a root like /engine
+					FString RelativePath = FPackageName::LongPackageNameToFilename(DirToCook.Path / TEXT(""));
+					CookDirectories.Add(FPaths::ConvertRelativePathToFull(RelativePath));
+				}
+				else
+				{
+					// This is relative to /game
+					CookDirectories.Add(AbsoluteGameContentDir / DirToCook.Path);
+				}
 			}
 		}
-
-
 	}
 
 	if (!(FilesToCookFlags & ECookByTheBookOptions::NoGameAlwaysCookPackages))
@@ -5868,7 +5877,18 @@ void UCookOnTheFlyServer::StartCookByTheBook( const FCookByTheBookStartupOptions
 
 		for (const FDirectoryPath& DirToNotCook : PackagingSettings->DirectoriesToNeverCook)
 		{
-			NeverCookDirectories.Add(AbsoluteGameContentDir / DirToNotCook.Path);
+			if (DirToNotCook.Path.StartsWith(TEXT("/"), ESearchCase::CaseSensitive))
+			{
+				// If this starts with /, this includes a root like /engine
+				FString RelativePath = FPackageName::LongPackageNameToFilename(DirToNotCook.Path / TEXT(""));
+				NeverCookDirectories.Add(FPaths::ConvertRelativePathToFull(RelativePath));
+			}
+			else
+			{
+				// This is relative to /game
+				NeverCookDirectories.Add(AbsoluteGameContentDir / DirToNotCook.Path);
+			}
+
 		}
 
 		for (const FString& NeverCookDirectory : NeverCookDirectories)
diff --git a/Engine/Source/Editor/UnrealEd/Private/EditorTransaction.cpp b/Engine/Source/Editor/UnrealEd/Private/EditorTransaction.cpp
--- a/Engine/Source/Editor/UnrealEd/Private/EditorTransaction.cpp
+++ b/Engine/Source/Editor/UnrealEd/Private/EditorTransaction.cpp
@@ -231,28 +231,21 @@ void FTransaction::FObjectRecord::Load(FTransaction* Owner)
 	{
 		bRestored = true;
 
-		if( CustomChange.IsValid() )
+		if (CustomChange.IsValid())
 		{
-			// @todo mesheditor debug
-			//GWarn->Logf( TEXT( "---------- Undoing Custom Change ----------" ) );
-			//CustomChange->PrintToLog( *GWarn );
-			
 			TUniquePtr<FChange> InvertedChange = CustomChange->Execute( Object.Get() );
-
-			// @todo mesheditor debug
-			//GWarn->Logf( TEXT( "-----(Here's what the Redo looks like)-----" ) );
-			//if( InvertedChange.IsValid() )
-			//{
-			//	InvertedChange->PrintToLog( *GWarn );
-			//}
-			//GWarn->Logf( TEXT( "-------------------------------------------" ) );
-
-			this->CustomChange = MoveTemp( InvertedChange );
+			CustomChange = MoveTemp( InvertedChange );
 		}
 		else
 		{
-			FReader Reader(Owner, SerializedObject, bWantsBinarySerialization);
-			SerializeContents(Reader, Oper);
+			// When objects are created outside the transaction system we can end up
+			// finding them but not having any data for them, so don't serialize 
+			// when that happens:
+			if (SerializedObject.Data.Num() > 0)
+			{
+				FReader Reader(Owner, SerializedObject, bWantsBinarySerialization);
+				SerializeContents(Reader, Oper);
+			}
 			SerializedObject.Swap(SerializedObjectFlip);
 		}
 		Oper *= -1;
