Fixed a crash when a Blueprint with variable watches was destroyed

diff --git a/Engine/Source/Editor/Kismet/Private/WatchPointViewer.cpp b/Engine/Source/Editor/Kismet/Private/WatchPointViewer.cpp
--- a/Engine/Source/Editor/Kismet/Private/WatchPointViewer.cpp
+++ b/Engine/Source/Editor/Kismet/Private/WatchPointViewer.cpp
@@ -800,30 +800,30 @@ FName WatchViewer::GetTabName()
 	return TabName;
 }
 
-void WatchViewer::UpdateWatchListFromBlueprint(UBlueprint* BlueprintObj)
+void UpdateWatchListFromBlueprintImpl(UBlueprint* BlueprintObj, const bool bShouldWatch)
 {
 	if (!ensure(BlueprintObj))
 	{
 		return;
 	}
 
-	if (BlueprintObj->WatchedPins.Num() == 0)
+	if (bShouldWatch)
 	{
-		// if this blueprint doesn't have any watched pins and we aren't watching it already then there is nothing to do
+		// make sure the blueprint is in our list
+		WatchedBlueprints.AddUnique(BlueprintObj);
+	}
+	else
+	{
+		// if this blueprint shouldn't be watched and we aren't watching it already then there is nothing to do
 		int32 FoundIdx = WatchedBlueprints.Find(BlueprintObj);
 		if (FoundIdx == INDEX_NONE)
 		{
 			return;
 		}
 
-		// since we're not watching any pins anymore we should remove it from the watched list
+		// since we're not watching the blueprint anymore we should remove it from the watched list
 		WatchedBlueprints.RemoveAt(FoundIdx);
 	}
-	else
-	{
-		// make sure the blueprint is in our list
-		WatchedBlueprints.AddUnique(BlueprintObj);
-	}
 
 	// something changed so we need to update the lists shown in the UI
 	UpdateNonInstancedWatchDisplay();
@@ -837,6 +837,16 @@ void WatchViewer::UpdateWatchListFromBlueprint(UBlueprint* BlueprintObj)
 	WatchListSubscribers.Broadcast(&Private_WatchSource);
 }
 
+void WatchViewer::UpdateWatchListFromBlueprint(UBlueprint* BlueprintObj)
+{
+	UpdateWatchListFromBlueprintImpl(BlueprintObj, BlueprintObj && BlueprintObj->WatchedPins.Num() > 0);
+}
+
+void WatchViewer::ClearWatchListFromBlueprint(UBlueprint* BlueprintObj)
+{
+	UpdateWatchListFromBlueprintImpl(BlueprintObj, false);
+}
+
 void WatchViewer::RegisterTabSpawner(FTabManager& TabManager)
 {
 	const auto SpawnWatchViewTab = []( const FSpawnTabArgs& Args )
diff --git a/Engine/Source/Editor/Kismet/Public/WatchPointViewer.h b/Engine/Source/Editor/Kismet/Public/WatchPointViewer.h
--- a/Engine/Source/Editor/Kismet/Public/WatchPointViewer.h
+++ b/Engine/Source/Editor/Kismet/Public/WatchPointViewer.h
@@ -16,6 +16,9 @@ namespace WatchViewer
 	// called when we are adding or changing watches from BlueprintObj
 	void KISMET_API UpdateWatchListFromBlueprint(class UBlueprint* BlueprintObj);
 
+	// called when a BlueprintObj should no longer be watched
+	void KISMET_API ClearWatchListFromBlueprint(class UBlueprint* BlueprintObj);
+
 	FName GetTabName();
 	void RegisterTabSpawner(FTabManager& TabManager);
 }
diff --git a/Engine/Source/Runtime/Engine/Private/Blueprint.cpp b/Engine/Source/Runtime/Engine/Private/Blueprint.cpp
--- a/Engine/Source/Runtime/Engine/Private/Blueprint.cpp
+++ b/Engine/Source/Runtime/Engine/Private/Blueprint.cpp
@@ -26,6 +26,7 @@
 #include "Kismet2/KismetEditorUtilities.h"
 #include "Kismet2/CompilerResultsLog.h"
 #include "Kismet2/StructureEditorUtils.h"
+#include "WatchPointViewer.h"
 #include "FindInBlueprintManager.h"
 #include "CookerSettings.h"
 #include "Editor.h"
@@ -1384,6 +1385,7 @@ void UBlueprint::BeginDestroy()
 	Super::BeginDestroy();
 
 	FBlueprintEditorUtils::RemoveAllLocalBookmarks(this);
+	WatchViewer::ClearWatchListFromBlueprint(this);
 }
 
 #endif // WITH_EDITOR
