Mirror 4081153 from Dev-Framework
Watch Point Viewer should not be a nomad tab, is owned by the Blueprint Debugger tab manager

diff --git a/Engine/Source/Editor/Kismet/Private/WatchPointViewer.cpp b/Engine/Source/Editor/Kismet/Private/WatchPointViewer.cpp
--- a/Engine/Source/Editor/Kismet/Private/WatchPointViewer.cpp
+++ b/Engine/Source/Editor/Kismet/Private/WatchPointViewer.cpp
@@ -1,4 +1,4 @@
-﻿// Copyright 1998-2018 Epic Games, Inc. All Rights Reserved.
+// Copyright 1998-2018 Epic Games, Inc. All Rights Reserved.
 
 #include "WatchPointViewer.h"
 
@@ -852,7 +852,7 @@ void WatchViewer::RegisterTabSpawner(FTabManager& TabManager)
 	const auto SpawnWatchViewTab = []( const FSpawnTabArgs& Args )
 	{
 		return SNew(SDockTab)
-			.TabRole( ETabRole::NomadTab )
+			.TabRole( ETabRole::PanelTab )
 			.Label( LOCTEXT("TabTitle", "Watches") )
 			[
 				SNew(SBorder)
