Mirror 4090523 from Dev-Framework
Allow user to close Blueprint Debugger while at a breakpoint, origin of this behavior was 1148085 and the original issue no longer seems to happen. Also fix mismatch between tab name and context menu entry

diff --git a/Engine/Source/Editor/Kismet/Private/BlueprintDebugger.cpp b/Engine/Source/Editor/Kismet/Private/BlueprintDebugger.cpp
--- a/Engine/Source/Editor/Kismet/Private/BlueprintDebugger.cpp
+++ b/Engine/Source/Editor/Kismet/Private/BlueprintDebugger.cpp
@@ -17,11 +17,6 @@
 
 #define LOCTEXT_NAMESPACE "BlueprintDebugger"
 
-static bool CanCloseBlueprintDebugger()
-{
-	return !GIntraFrameDebuggingGameThread;
-}
-
 struct FBlueprintDebuggerCommands : public TCommands<FBlueprintDebuggerCommands>
 {
 	FBlueprintDebuggerCommands()
@@ -47,7 +42,7 @@ void FBlueprintDebuggerCommands::RegisterCommands()
 {
 	UI_COMMAND(ShowCallStackViewer, "Call Stack", "Toggles visibility of the Call Stack window", EUserInterfaceActionType::Check, FInputChord());
 	UI_COMMAND(ShowWatchViewer, "Watches", "Toggles visibility of the Watches window", EUserInterfaceActionType::Check, FInputChord());
-	UI_COMMAND(ShowExecutionTrace, "Execution Trace", "Toggles visibility of the Execution Trace window", EUserInterfaceActionType::Check, FInputChord());
+	UI_COMMAND(ShowExecutionTrace, "Execution Flow", "Toggles visibility of the Execution Flow window", EUserInterfaceActionType::Check, FInputChord());
 }
 
 struct FBlueprintDebuggerImpl
@@ -105,7 +100,6 @@ TSharedRef<SDockTab> FBlueprintDebuggerImpl::CreateBluprintDebuggerTab(const FSp
 {
 	const TSharedRef<SDockTab> NomadTab = SNew(SDockTab)
 		.TabRole(ETabRole::NomadTab)
-		.OnCanCloseTab(SDockTab::FCanCloseTab::CreateStatic(&CanCloseBlueprintDebugger))
 		.Label(NSLOCTEXT("BlueprintDebugger", "TabTitle", "Blueprint Debugger"));
 
 	if (!DebuggingToolsTabManager.IsValid())
@@ -141,7 +135,6 @@ TSharedRef<SDockTab> FBlueprintDebuggerImpl::CreateBluprintDebuggerTab(const FSp
 				{
 				return SNew(SDockTab)
 					.TabRole(ETabRole::PanelTab)
-					.OnCanCloseTab(SDockTab::FCanCloseTab::CreateStatic(&CanCloseBlueprintDebugger))
 					.Label(NSLOCTEXT("BlueprintExecutionFlow", "TabTitle", "Execution Flow"))
 					[
 						SNew(SKismetDebuggingView)
