Mirror 4090824 from Dev-Framework
Add abillity to insert pins before or after a target pin for Sequence nodes via the context menu

diff --git a/Engine/Source/Editor/BlueprintGraph/Classes/K2Node.h b/Engine/Source/Editor/BlueprintGraph/Classes/K2Node.h
--- a/Engine/Source/Editor/BlueprintGraph/Classes/K2Node.h
+++ b/Engine/Source/Editor/BlueprintGraph/Classes/K2Node.h
@@ -248,6 +248,9 @@ class UK2Node : public UEdGraphNode
 	/** Return whether the node's properties display in the blueprint details panel */
 	virtual bool ShouldShowNodeProperties() const { return false; }
 
+	/** Return whether the node's execution pins should support the remove execution pin action */
+	virtual bool CanEverInsertExecutionPin() const { return false; }
+
 	/** Return whether the node's execution pins should support the remove execution pin action */
 	virtual bool CanEverRemoveExecutionPin() const { return false; }
 
diff --git a/Engine/Source/Editor/BlueprintGraph/Classes/K2Node_ExecutionSequence.h b/Engine/Source/Editor/BlueprintGraph/Classes/K2Node_ExecutionSequence.h
--- a/Engine/Source/Editor/BlueprintGraph/Classes/K2Node_ExecutionSequence.h
+++ b/Engine/Source/Editor/BlueprintGraph/Classes/K2Node_ExecutionSequence.h
@@ -8,10 +8,11 @@
 #include "Textures/SlateIcon.h"
 #include "K2Node.h"
 #include "K2Node_AddPinInterface.h"
+#include "EdGraph/EdGraphPin.h"
+
 #include "K2Node_ExecutionSequence.generated.h"
 
 class FBlueprintActionDatabaseRegistrar;
-class UEdGraphPin;
 
 UCLASS(MinimalAPI)
 class UK2Node_ExecutionSequence : public UK2Node, public IK2Node_AddPinInterface
@@ -31,6 +32,7 @@ class UK2Node_ExecutionSequence : public UK2Node, public IK2Node_AddPinInterface
 	virtual class FNodeHandlingFunctor* CreateNodeHandler(class FKismetCompilerContext& CompilerContext) const override;
 	virtual void GetMenuActions(FBlueprintActionDatabaseRegistrar& ActionRegistrar) const override;
 	virtual FText GetMenuCategory() const override;
+	virtual bool CanEverInsertExecutionPin() const override { return true; }
 	virtual bool CanEverRemoveExecutionPin() const override { return true; }
 	//~ End UK2Node Interface
 
@@ -49,6 +51,13 @@ class UK2Node_ExecutionSequence : public UK2Node, public IK2Node_AddPinInterface
 	DEPRECATED(4.17, "Use AddInputPin instead.")
 	BLUEPRINTGRAPH_API void AddPinToExecutionNode() { AddInputPin(); }
 
+	/**
+	 * Inserts a new execution pin, before the specified execution pin, into an execution node
+	 *
+	 * @param	PinToInsertBefore	The pin to insert a new pin before on the node
+	 */
+	BLUEPRINTGRAPH_API void InsertPinIntoExecutionNode(UEdGraphPin* PinToInsertBefore, EPinInsertPosition Position);
+
 	/**
 	 * Removes the specified execution pin from an execution node
 	 *
diff --git a/Engine/Source/Editor/BlueprintGraph/Private/EdGraphSchema_K2.cpp b/Engine/Source/Editor/BlueprintGraph/Private/EdGraphSchema_K2.cpp
--- a/Engine/Source/Editor/BlueprintGraph/Private/EdGraphSchema_K2.cpp
+++ b/Engine/Source/Editor/BlueprintGraph/Private/EdGraphSchema_K2.cpp
@@ -1490,9 +1490,15 @@ void UEdGraphSchema_K2::GetContextMenuActions(const UEdGraph* CurrentGraph, cons
 					MenuBuilder->AddMenuEntry( FGraphEditorCommands::Get().RecombineStructPin );
 				}
 	
-				// Conditionally add the execution path pin removal if this is an execution branching node
+				// Conditionally add the execution path pin options if this is an execution branching node
 				if( InGraphPin->Direction == EGPD_Output && InGraphPin->GetOwningNode())
 				{
+					if (CastChecked<UK2Node>(InGraphPin->GetOwningNode())->CanEverInsertExecutionPin())
+					{
+						MenuBuilder->AddMenuEntry(FGraphEditorCommands::Get().InsertExecutionPinBefore);
+						MenuBuilder->AddMenuEntry(FGraphEditorCommands::Get().InsertExecutionPinAfter);
+					}
+
 					if (CastChecked<UK2Node>(InGraphPin->GetOwningNode())->CanEverRemoveExecutionPin())
 					{
 						MenuBuilder->AddMenuEntry( FGraphEditorCommands::Get().RemoveExecutionPin );
diff --git a/Engine/Source/Editor/BlueprintGraph/Private/K2Node_ExecutionSequence.cpp b/Engine/Source/Editor/BlueprintGraph/Private/K2Node_ExecutionSequence.cpp
--- a/Engine/Source/Editor/BlueprintGraph/Private/K2Node_ExecutionSequence.cpp
+++ b/Engine/Source/Editor/BlueprintGraph/Private/K2Node_ExecutionSequence.cpp
@@ -191,6 +191,34 @@ void UK2Node_ExecutionSequence::AddInputPin()
 	CreatePin(EGPD_Output, UEdGraphSchema_K2::PC_Exec, GetUniquePinName());
 }
 
+void UK2Node_ExecutionSequence::InsertPinIntoExecutionNode(UEdGraphPin* PinToInsertBefore, EPinInsertPosition Position)
+{
+	Modify();
+
+	int32 DesiredPinIndex = Pins.Find(PinToInsertBefore);
+	if (DesiredPinIndex != INDEX_NONE)
+	{
+		if (Position == EPinInsertPosition::After)
+		{
+			DesiredPinIndex = DesiredPinIndex + 1;
+		}
+
+		FCreatePinParams Params;
+		Params.Index = DesiredPinIndex;
+		CreatePin(EGPD_Output, UEdGraphSchema_K2::PC_Exec, GetUniquePinName(), Params);
+
+		// refresh names on the pin list:
+		for (int32 Idx = 0; Idx < Pins.Num(); ++Idx)
+		{
+			UEdGraphPin* PotentialPin = Pins[Idx];
+			if (UEdGraphSchema_K2::IsExecPin(*PotentialPin) && (PotentialPin->Direction == EGPD_Output))
+			{
+				PotentialPin->PinName = GetPinNameGivenIndex(Idx);
+			}
+		}
+	}
+}
+
 void UK2Node_ExecutionSequence::RemovePinFromExecutionNode(UEdGraphPin* TargetPin) 
 {
 	UK2Node_ExecutionSequence* OwningSeq = Cast<UK2Node_ExecutionSequence>( TargetPin->GetOwningNode() );
diff --git a/Engine/Source/Editor/GraphEditor/Private/GraphEditorActions.cpp b/Engine/Source/Editor/GraphEditor/Private/GraphEditorActions.cpp
--- a/Engine/Source/Editor/GraphEditor/Private/GraphEditorActions.cpp
+++ b/Engine/Source/Editor/GraphEditor/Private/GraphEditorActions.cpp
@@ -10,6 +10,8 @@ void FGraphEditorCommandsImpl::RegisterCommands()
 	UI_COMMAND( BreakNodeLinks, "Break Link(s)", "Breaks links", EUserInterfaceActionType::Button, FInputChord() )
 	
 	UI_COMMAND( AddExecutionPin, "Add execution pin", "Adds another execution output pin to an execution sequence or switch node", EUserInterfaceActionType::Button, FInputChord() )
+	UI_COMMAND( InsertExecutionPinBefore, "Insert execution pin before", "Adds another execution output pin before this one, to an execution sequence node", EUserInterfaceActionType::Button, FInputChord() )
+	UI_COMMAND( InsertExecutionPinAfter, "Insert execution pin after", "Adds another execution output pin after this one, to an execution sequence node", EUserInterfaceActionType::Button, FInputChord() )
 	UI_COMMAND( RemoveExecutionPin, "Remove execution pin", "Removes an execution output pin from an execution sequence or switch node", EUserInterfaceActionType::Button, FInputChord() )
 
 	UI_COMMAND( RemoveThisStructVarPin, "Remove this struct variable pin", "Removes the selected input pin", EUserInterfaceActionType::Button, FInputChord() )
diff --git a/Engine/Source/Editor/GraphEditor/Private/SGraphEditorImpl.cpp b/Engine/Source/Editor/GraphEditor/Private/SGraphEditorImpl.cpp
--- a/Engine/Source/Editor/GraphEditor/Private/SGraphEditorImpl.cpp
+++ b/Engine/Source/Editor/GraphEditor/Private/SGraphEditorImpl.cpp
@@ -200,6 +200,11 @@ void SGraphEditorImpl::StraightenConnections(UEdGraphPin* SourcePin, UEdGraphPin
 	GraphPanel->StraightenConnections(SourcePin, PinToAlign);
 }
 
+void SGraphEditorImpl::RefreshNode(UEdGraphNode& Node)
+{
+	GraphPanel->RefreshNode(Node);
+}
+
 void SGraphEditorImpl::Construct( const FArguments& InArgs )
 {
 	Commands = MakeShareable( new FUICommandList() );
diff --git a/Engine/Source/Editor/GraphEditor/Private/SGraphEditorImpl.h b/Engine/Source/Editor/GraphEditor/Private/SGraphEditorImpl.h
--- a/Engine/Source/Editor/GraphEditor/Private/SGraphEditorImpl.h
+++ b/Engine/Source/Editor/GraphEditor/Private/SGraphEditorImpl.h
@@ -156,6 +156,7 @@ public:
 	virtual TSharedRef<FActiveTimerHandle> RegisterActiveTimer(float TickPeriod, FWidgetActiveTimerDelegate TickFunction) override;
 	virtual void StraightenConnections() override;
 	virtual void StraightenConnections(UEdGraphPin* SourcePin, UEdGraphPin* PinToAlign) override;
+	virtual void RefreshNode(UEdGraphNode& Node) override;
 	virtual void CaptureKeyboard() override;
 	virtual void SetNodeFactory(const TSharedRef<class FGraphNodeFactory>& NewNodeFactory) override;
 	// End of SGraphEditor interface
diff --git a/Engine/Source/Editor/GraphEditor/Private/SGraphPanel.cpp b/Engine/Source/Editor/GraphEditor/Private/SGraphPanel.cpp
--- a/Engine/Source/Editor/GraphEditor/Private/SGraphPanel.cpp
+++ b/Engine/Source/Editor/GraphEditor/Private/SGraphPanel.cpp
@@ -1336,6 +1336,15 @@ void SGraphPanel::StraightenConnections(UEdGraphPin* SourcePin, UEdGraphPin* Pin
 	Aligner.Process();
 }
 
+void SGraphPanel::RefreshNode(UEdGraphNode& Node)
+{
+	TSharedPtr<SGraphNode> GraphNode = GetNodeWidgetFromGuid(Node.NodeGuid);
+	if (GraphNode.IsValid())
+	{
+		GraphNode->UpdateGraphNode();
+	}
+}
+
 const TSharedRef<SGraphNode> SGraphPanel::GetChild(int32 ChildIndex)
 {
 	return StaticCastSharedRef<SGraphNode>(Children[ChildIndex]);
diff --git a/Engine/Source/Editor/GraphEditor/Public/GraphEditorActions.h b/Engine/Source/Editor/GraphEditor/Public/GraphEditorActions.h
--- a/Engine/Source/Editor/GraphEditor/Public/GraphEditorActions.h
+++ b/Engine/Source/Editor/GraphEditor/Public/GraphEditorActions.h
@@ -26,6 +26,8 @@ public:
 
 	// Execution sequence specific commands
 	TSharedPtr< FUICommandInfo > AddExecutionPin;
+	TSharedPtr< FUICommandInfo > InsertExecutionPinBefore;
+	TSharedPtr< FUICommandInfo > InsertExecutionPinAfter;
 	TSharedPtr< FUICommandInfo > RemoveExecutionPin;
 
 	// SetFieldsInStruct specific commands
diff --git a/Engine/Source/Editor/GraphEditor/Public/SGraphPanel.h b/Engine/Source/Editor/GraphEditor/Public/SGraphPanel.h
--- a/Engine/Source/Editor/GraphEditor/Public/SGraphPanel.h
+++ b/Engine/Source/Editor/GraphEditor/Public/SGraphPanel.h
@@ -173,6 +173,9 @@ public:
 	/** Straighten any connections attached to the specified pin, optionally limiting to the specified pin to align */
 	void StraightenConnections(UEdGraphPin* SourcePin, UEdGraphPin* PinToAlign = nullptr);
 
+	/** Refresh the visual state of a single node */
+	void RefreshNode(UEdGraphNode& Node);
+
 	/** When the graph panel needs to be dynamically refreshing for animations, this function is registered to tick and invalidate the UI. */
 	EActiveTimerReturnType InvalidatePerTick(double InCurrentTime, float InDeltaTime);
 
diff --git a/Engine/Source/Editor/Kismet/Private/BlueprintEditor.cpp b/Engine/Source/Editor/Kismet/Private/BlueprintEditor.cpp
--- a/Engine/Source/Editor/Kismet/Private/BlueprintEditor.cpp
+++ b/Engine/Source/Editor/Kismet/Private/BlueprintEditor.cpp
@@ -1004,6 +1004,16 @@ TSharedRef<SGraphEditor> FBlueprintEditor::CreateGraphEditorWidget(TSharedRef<FT
 				FCanExecuteAction::CreateSP( this, &FBlueprintEditor::CanAddExecutionPin )
 				);
 
+			GraphEditorCommands->MapAction( FGraphEditorCommands::Get().InsertExecutionPinBefore,
+				FExecuteAction::CreateSP( this, &FBlueprintEditor::OnInsertExecutionPinBefore ),
+				FCanExecuteAction::CreateSP( this, &FBlueprintEditor::CanInsertExecutionPin )
+				);
+
+			GraphEditorCommands->MapAction(FGraphEditorCommands::Get().InsertExecutionPinAfter,
+				FExecuteAction::CreateSP(this, &FBlueprintEditor::OnInsertExecutionPinAfter),
+				FCanExecuteAction::CreateSP(this, &FBlueprintEditor::CanInsertExecutionPin)
+			);
+
 			GraphEditorCommands->MapAction( FGraphEditorCommands::Get().RemoveExecutionPin,
 				FExecuteAction::CreateSP( this, &FBlueprintEditor::OnRemoveExecutionPin ),
 				FCanExecuteAction::CreateSP( this, &FBlueprintEditor::CanRemoveExecutionPin )
@@ -3887,6 +3897,64 @@ bool FBlueprintEditor::CanAddExecutionPin() const
 	return true;
 }
 
+void FBlueprintEditor::OnInsertExecutionPinBefore()
+{
+	OnInsertExecutionPin(EPinInsertPosition::Before);
+}
+
+void FBlueprintEditor::OnInsertExecutionPinAfter()
+{
+	OnInsertExecutionPin(EPinInsertPosition::After);
+}
+
+void FBlueprintEditor::OnInsertExecutionPin(EPinInsertPosition Position)
+{
+	TSharedPtr<SGraphEditor> FocusedGraphEd = FocusedGraphEdPtr.Pin();
+	if (FocusedGraphEd.IsValid())
+	{
+		const FScopedTransaction Transaction(LOCTEXT("InsertExecutionPinBefore", "Insert Execution Pin Before"));
+
+		UEdGraphPin* SelectedPin = FocusedGraphEd->GetGraphPinForMenu();
+		if (SelectedPin)
+		{
+			UEdGraphNode* OwningNode = SelectedPin->GetOwningNode();
+
+			if (OwningNode)
+			{
+				if (UK2Node_ExecutionSequence* SeqNode = Cast<UK2Node_ExecutionSequence>(OwningNode))
+				{
+					SeqNode->InsertPinIntoExecutionNode(SelectedPin, Position);
+					FocusedGraphEd->RefreshNode(*OwningNode);
+
+					if (UBlueprint* BP = SeqNode->GetBlueprint())
+					{
+						FBlueprintEditorUtils::MarkBlueprintAsModified(BP);
+					}
+				}
+			}
+		}
+	}
+}
+
+bool FBlueprintEditor::CanInsertExecutionPin() const
+{
+	// We likely don't need to validate here, as we validated on menu population,
+	// but better to grey out the option if it is somehow created but will
+	// not execute correctly
+	TSharedPtr<SGraphEditor> FocusedGraphEd = FocusedGraphEdPtr.Pin();
+	if (FocusedGraphEd.IsValid())
+	{
+		UEdGraphPin* SelectedPin = FocusedGraphEd->GetGraphPinForMenu();
+		if (SelectedPin)
+		{
+			UEdGraphNode* OwningNode = SelectedPin->GetOwningNode();
+			return Cast<UK2Node_ExecutionSequence>(OwningNode) != nullptr;
+		}
+	}
+
+	return false;
+}
+
 void  FBlueprintEditor::OnRemoveExecutionPin()
 {
 	TSharedPtr<SGraphEditor> FocusedGraphEd = FocusedGraphEdPtr.Pin();
diff --git a/Engine/Source/Editor/Kismet/Public/BlueprintEditor.h b/Engine/Source/Editor/Kismet/Public/BlueprintEditor.h
--- a/Engine/Source/Editor/Kismet/Public/BlueprintEditor.h
+++ b/Engine/Source/Editor/Kismet/Public/BlueprintEditor.h
@@ -710,6 +710,11 @@ protected:
 	void OnAddExecutionPin();
 	bool CanAddExecutionPin() const;
 
+	void OnInsertExecutionPinBefore();
+	void OnInsertExecutionPinAfter();
+	void OnInsertExecutionPin(EPinInsertPosition Position);
+	bool CanInsertExecutionPin() const;
+
 	void OnRemoveExecutionPin();
 	bool CanRemoveExecutionPin() const;
 
diff --git a/Engine/Source/Editor/UnrealEd/Public/GraphEditor.h b/Engine/Source/Editor/UnrealEd/Public/GraphEditor.h
--- a/Engine/Source/Editor/UnrealEd/Public/GraphEditor.h
+++ b/Engine/Source/Editor/UnrealEd/Public/GraphEditor.h
@@ -395,6 +395,14 @@ public:
 		}
 	}
 
+	virtual void RefreshNode(UEdGraphNode& Node)
+	{
+		if (Implementation.IsValid())
+		{
+			return Implementation->RefreshNode(Node);
+		}
+	}
+
 	// Invoked to let this widget know that the GraphEditor module has been reloaded
 	UNREALED_API void OnModuleReloaded();
 
diff --git a/Engine/Source/Runtime/Engine/Classes/EdGraph/EdGraphPin.h b/Engine/Source/Runtime/Engine/Classes/EdGraph/EdGraphPin.h
--- a/Engine/Source/Runtime/Engine/Classes/EdGraph/EdGraphPin.h
+++ b/Engine/Source/Runtime/Engine/Classes/EdGraph/EdGraphPin.h
@@ -15,6 +15,12 @@
 class UEdGraphPin;
 enum class EPinResolveType : uint8;
 
+enum class EPinInsertPosition
+{
+	Before,
+	After,
+};
+
 USTRUCT()
 struct FSimpleMemberReference
 {
