Fix for dark metal (missing reflections) in Blueprint Editor. Duplicate of changelist 4099649.

diff --git a/Engine/Source/Editor/Kismet/Private/BlueprintEditor.cpp b/Engine/Source/Editor/Kismet/Private/BlueprintEditor.cpp
--- a/Engine/Source/Editor/Kismet/Private/BlueprintEditor.cpp
+++ b/Engine/Source/Editor/Kismet/Private/BlueprintEditor.cpp
@@ -7086,6 +7086,8 @@ void FBlueprintEditor::SetViewLocation(const FVector2D& Location, float ZoomAmou
 
 void FBlueprintEditor::Tick(float DeltaTime)
 {
+	PreviewScene.UpdateCaptureContents();
+
 	// Create or update the Blueprint actor instance in the preview scene
 	if ( GetPreviewActor() == nullptr )
 	{
