Explicitly close floating Blueprint Debugger controls when the master tab is closed
#rnx

diff --git a/Engine/Source/Editor/Kismet/Private/BlueprintDebugger.cpp b/Engine/Source/Editor/Kismet/Private/BlueprintDebugger.cpp
--- a/Engine/Source/Editor/Kismet/Private/BlueprintDebugger.cpp
+++ b/Engine/Source/Editor/Kismet/Private/BlueprintDebugger.cpp
@@ -105,6 +105,7 @@ TSharedRef<SDockTab> FBlueprintDebuggerImpl::CreateBluprintDebuggerTab(const FSp
 	if (!DebuggingToolsTabManager.IsValid())
 	{
 		DebuggingToolsTabManager = FGlobalTabmanager::Get()->NewTabManager(NomadTab);
+		// on persist layout will handle saving layout if the editor is shut down:
 		DebuggingToolsTabManager->SetOnPersistLayout(
 			FTabManager::FOnPersistLayout::CreateStatic(
 				[](const TSharedRef<FTabManager::FLayout>& InLayout)
@@ -126,6 +127,23 @@ TSharedRef<SDockTab> FBlueprintDebuggerImpl::CreateBluprintDebuggerTab(const FSp
 	const FName CallStackTabName = CallStackViewer::GetTabName();
 	const FName WatchViewerTabName = WatchViewer::GetTabName();
 
+	TWeakPtr<FTabManager> DebuggingToolsTabManagerWeak = DebuggingToolsTabManager;
+	// On tab close will save the layout if the debugging window itself is closed,
+	// this handler also cleans up any floating debugging controls. If we don't close
+	// all areas we need to add some logic to the tab manager to reuse existing tabs:
+	NomadTab->SetOnTabClosed(SDockTab::FOnTabClosedCallback::CreateStatic(
+		[](TSharedRef<SDockTab> Self, TWeakPtr<FTabManager> TabManager)
+		{
+			TSharedPtr<FTabManager> OwningTabManager = TabManager.Pin();
+			if (OwningTabManager.IsValid())
+			{
+				FLayoutSaveRestore::SaveToConfig(GEditorLayoutIni, OwningTabManager->PersistLayout());
+				OwningTabManager->CloseAllAreas();
+			}
+		}
+		, DebuggingToolsTabManagerWeak
+	));
+
 	if (!BlueprintDebuggerLayout.IsValid())
 	{
 		DebuggingToolsTabManager->RegisterTabSpawner(
