The Blueprint editor Bookmarks tab now updates when switching to a new graph tab, when showing only the current graph context.

- Mirrored from //UE4/Dev-Framework (4112217).

#rnx

diff --git a/Engine/Source/Editor/Kismet/Private/BlueprintEditor.cpp b/Engine/Source/Editor/Kismet/Private/BlueprintEditor.cpp
--- a/Engine/Source/Editor/Kismet/Private/BlueprintEditor.cpp
+++ b/Engine/Source/Editor/Kismet/Private/BlueprintEditor.cpp
@@ -2998,6 +2998,13 @@ void FBlueprintEditor::OnGraphEditorFocused(const TSharedRef<SGraphEditor>& InGr
 			}
 		}
 	}
+
+	// If the bookmarks view is active, check whether or not we're restricting the view to the current graph. If we are, update the tree to reflect the focused graph context.
+	if (BookmarksWidget.IsValid()
+		&& GetDefault<UBlueprintEditorSettings>()->bShowBookmarksForCurrentDocumentOnlyInTab)
+	{
+		BookmarksWidget->RefreshBookmarksTree();
+	}
 }
 
 void FBlueprintEditor::OnGraphEditorBackgrounded(const TSharedRef<SGraphEditor>& InGraphEditor)
