Add missing Bookmarks command in macro and interface editors.

- Mirrored from //UE4/Dev-Framework (4113177).

#rnx

diff --git a/Engine/Source/Editor/Kismet/Private/BlueprintEditorModes.cpp b/Engine/Source/Editor/Kismet/Private/BlueprintEditorModes.cpp
--- a/Engine/Source/Editor/Kismet/Private/BlueprintEditorModes.cpp
+++ b/Engine/Source/Editor/Kismet/Private/BlueprintEditorModes.cpp
@@ -359,6 +359,7 @@ FBlueprintInterfaceApplicationMode::FBlueprintInterfaceApplicationMode(TSharedPt
 		BlueprintInterfaceTabFactories.RegisterFactory(MakeShareable(new FReplaceNodeReferencesSummoner(InBlueprintEditor)));
 	}
 	BlueprintInterfaceTabFactories.RegisterFactory(MakeShareable(new FCompilerResultsSummoner(InBlueprintEditor)));
+	BlueprintInterfaceTabFactories.RegisterFactory(MakeShareable(new FBookmarksSummoner(InBlueprintEditor)));
 	BlueprintInterfaceTabFactories.RegisterFactory(MakeShareable(new FFindResultsSummoner(InBlueprintEditor)));
 	BlueprintInterfaceTabFactories.RegisterFactory(MakeShareable(new FSelectionDetailsSummoner(InBlueprintEditor)));
 
@@ -464,6 +465,7 @@ FBlueprintMacroApplicationMode::FBlueprintMacroApplicationMode(TSharedPtr<class
 		BlueprintMacroTabFactories.RegisterFactory(MakeShareable(new FReplaceNodeReferencesSummoner(InBlueprintEditor)));
 	}
 	BlueprintMacroTabFactories.RegisterFactory(MakeShareable(new FPaletteSummoner(InBlueprintEditor)));
+	BlueprintMacroTabFactories.RegisterFactory(MakeShareable(new FBookmarksSummoner(InBlueprintEditor)));
 	BlueprintMacroTabFactories.RegisterFactory(MakeShareable(new FFindResultsSummoner(InBlueprintEditor)));
 	BlueprintMacroTabFactories.RegisterFactory(MakeShareable(new FSelectionDetailsSummoner(InBlueprintEditor)));
 
