Fix crash when Pin has been deleted already
Use MoveTemp in a few places with construction of FWatchRows
#rnx

diff --git a/Engine/Source/Editor/Kismet/Private/WatchPointViewer.cpp b/Engine/Source/Editor/Kismet/Private/WatchPointViewer.cpp
--- a/Engine/Source/Editor/Kismet/Private/WatchPointViewer.cpp
+++ b/Engine/Source/Editor/Kismet/Private/WatchPointViewer.cpp
@@ -85,9 +85,9 @@ struct FWatchRow
 	{
 		SetObjectBeingDebuggedName();
 
-		for (const FDebugInfo& ChildInfo : Info.Children)
+		for (FDebugInfo& ChildInfo : Info.Children)
 		{
-			Children.Add(MakeShared<FWatchRow>(InBP, InNode, InPin, InObjectBeingDebugged, BlueprintName, GraphName, NodeName, ChildInfo));
+			Children.Add(MakeShared<FWatchRow>(InBP, InNode, InPin, InObjectBeingDebugged, BlueprintName, GraphName, NodeName, MoveTemp(ChildInfo)));
 		}
 	}
 
@@ -629,30 +629,31 @@ void UpdateNonInstancedWatchDisplay()
 
 		for (const FEdGraphPinReference& PinRef : BlueprintObj->WatchedPins)
 		{
-			UEdGraphPin* Pin = PinRef.Get();
-
-			FText GraphName = FText::FromString(Pin->GetOwningNode()->GetGraph()->GetName());
-			FText NodeName = Pin->GetOwningNode()->GetNodeTitle(ENodeTitleType::ListView);
-
-			const UEdGraphSchema* Schema = Pin->GetOwningNode()->GetSchema();
+			if (UEdGraphPin* Pin = PinRef.Get())
+			{
+				FText GraphName = FText::FromString(Pin->GetOwningNode()->GetGraph()->GetName());
+				FText NodeName = Pin->GetOwningNode()->GetNodeTitle(ENodeTitleType::ListView);
 
-			FDebugInfo DebugInfo;
-			DebugInfo.DisplayName = Schema->GetPinDisplayName(Pin);
-			DebugInfo.Type = UEdGraphSchema_K2::TypeToText(Pin->PinType);
-			DebugInfo.Value = LOCTEXT("ExecutionNotPaused", "(execution not paused)");
+				const UEdGraphSchema* Schema = Pin->GetOwningNode()->GetSchema();
 
-			Private_WatchSource.Add(
-				MakeShared<FWatchRow>(
-					BlueprintObj,
-					Pin->GetOwningNode(),
-					Pin,
-					nullptr,
-					BlueprintName,
-					GraphName,
-					NodeName,
-					DebugInfo
-					)
-			);
+				FDebugInfo DebugInfo;
+				DebugInfo.DisplayName = Schema->GetPinDisplayName(Pin);
+				DebugInfo.Type = UEdGraphSchema_K2::TypeToText(Pin->PinType);
+				DebugInfo.Value = LOCTEXT("ExecutionNotPaused", "(execution not paused)");
+
+				Private_WatchSource.Add(
+					MakeShared<FWatchRow>(
+						BlueprintObj,
+						Pin->GetOwningNode(),
+						Pin,
+						nullptr,
+						BlueprintName,
+						MoveTemp(GraphName),
+						MoveTemp(NodeName),
+						MoveTemp(DebugInfo)
+						)
+				);
+			}
 		}
 	}
 }
