Unhide debug filter for level blueprints - it is useful when a user wants to filter based on the world

diff --git a/Engine/Source/Editor/Kismet/Private/SBlueprintEditorSelectedDebugObjectWidget.cpp b/Engine/Source/Editor/Kismet/Private/SBlueprintEditorSelectedDebugObjectWidget.cpp
--- a/Engine/Source/Editor/Kismet/Private/SBlueprintEditorSelectedDebugObjectWidget.cpp
+++ b/Engine/Source/Editor/Kismet/Private/SBlueprintEditorSelectedDebugObjectWidget.cpp
@@ -78,10 +78,6 @@ void SBlueprintEditorSelectedDebugObjectWidget::Construct(const FArguments& InAr
 			BrowseButton
 		];
 
-
-
-
-
 	ChildSlot
 		[
 			SNew(SLevelOfDetailBranchNode)
@@ -105,7 +101,6 @@ void SBlueprintEditorSelectedDebugObjectWidget::Construct(const FArguments& InAr
 			.HighDetail()
 				[
 					SNew(SVerticalBox)
-					.Visibility(this, &SBlueprintEditorSelectedDebugObjectWidget::ShouldShowDebugObjectPicker)
 					+ SVerticalBox::Slot()
 					.AutoHeight()
 					.VAlign(VAlign_Bottom)
@@ -542,12 +537,6 @@ void SBlueprintEditorSelectedDebugObjectWidget::GenerateDebugObjectNames(bool bR
 	}
 }
 
-EVisibility SBlueprintEditorSelectedDebugObjectWidget::ShouldShowDebugObjectPicker() const
-{
-	check(GetBlueprintObj());
-	return FBlueprintEditorUtils::IsLevelScriptBlueprint(GetBlueprintObj()) ? EVisibility::Collapsed : EVisibility::Visible;
-}
-
 TSharedPtr<FString> SBlueprintEditorSelectedDebugObjectWidget::GetDebugObjectName() const
 {
 	check(GetBlueprintObj());
diff --git a/Engine/Source/Editor/Kismet/Private/SBlueprintEditorSelectedDebugObjectWidget.h b/Engine/Source/Editor/Kismet/Private/SBlueprintEditorSelectedDebugObjectWidget.h
--- a/Engine/Source/Editor/Kismet/Private/SBlueprintEditorSelectedDebugObjectWidget.h
+++ b/Engine/Source/Editor/Kismet/Private/SBlueprintEditorSelectedDebugObjectWidget.h
@@ -36,9 +36,6 @@ private:
 	/** Refresh the widget. **/
 	void OnRefresh();
 
-	/** Returns true if we're not working on the level script */
-	EVisibility ShouldShowDebugObjectPicker() const;
-
 	/** Returns the name of the current debug actor */
 	TSharedPtr<FString> GetDebugObjectName() const;
 
