Mirror CL 4070790 from Dev-Framework

fix crash in watch window when blueprints that are being watched get deleted.
Add support to the watch window for renaming assets that are being watched.
Fix a bug that made the non-instanced watches show up if the watch window was opened while execution was paused at a breakpoint.

#rnx

diff --git a/Engine/Source/Editor/Kismet/Private/WatchPointViewer.cpp b/Engine/Source/Editor/Kismet/Private/WatchPointViewer.cpp
--- a/Engine/Source/Editor/Kismet/Private/WatchPointViewer.cpp
+++ b/Engine/Source/Editor/Kismet/Private/WatchPointViewer.cpp
@@ -15,6 +15,7 @@
 #include "EdGraphSchema_K2.h"
 #include "UnrealEdGlobals.h"
 #include "Editor/UnrealEdEngine.h"
+#include "AssetRegistryModule.h"
 
 #include "Editor.h"
 
@@ -34,6 +35,15 @@
 
 #define LOCTEXT_NAMESPACE "WatchPointViewer"
 
+namespace
+{
+	// Returns true if the blueprint execution is currently paused; false otherwise
+	bool IsPaused()
+	{
+		return GUnrealEd && GUnrealEd->PlayWorld && GUnrealEd->PlayWorld->bDebugPauseExecution;
+	}
+};
+
 struct FWatchRow
 {
 	FWatchRow(
@@ -60,6 +70,9 @@ struct FWatchRow
 		, Type(MoveTemp(InType))
 	{
 		SetObjectBeingDebuggedName();
+
+		UPackage* Package = Cast<UPackage>(BP ? BP->GetOuter() : nullptr);
+		BlueprintPackageName = Package ? Package->GetFName() : FName();
 	}
 
 	FWatchRow(
@@ -85,6 +98,9 @@ struct FWatchRow
 	{
 		SetObjectBeingDebuggedName();
 
+		UPackage* Package = Cast<UPackage>(BP ? BP->GetOuter() : nullptr);
+		BlueprintPackageName = Package ? Package->GetFName() : FName();
+
 		for (FDebugInfo& ChildInfo : Info.Children)
 		{
 			Children.Add(MakeShared<FWatchRow>(InBP, InNode, InPin, InObjectBeingDebugged, BlueprintName, GraphName, NodeName, MoveTemp(ChildInfo)));
@@ -105,6 +121,7 @@ struct FWatchRow
 	FText DisplayName;
 	FText Value;
 	FText Type;
+	FName BlueprintPackageName;
 
 	TArray<TSharedRef<FWatchRow>> Children;
 
@@ -274,12 +291,23 @@ public:
 		FKismetDebugUtilities::WatchedPinsListChangedEvent.AddRaw(this, &SWatchViewer::HandleWatchedPinsChanged);
 		FEditorDelegates::ResumePIE.AddRaw(this, &SWatchViewer::HandleResumePIE);
 		FEditorDelegates::EndPIE.AddRaw(this, &SWatchViewer::HandleEndPIE);
+
+		FAssetRegistryModule& AssetRegistryModule = FModuleManager::LoadModuleChecked<FAssetRegistryModule>(TEXT("AssetRegistry"));
+		AssetRegistryModule.Get().OnAssetRemoved().AddRaw(this, &SWatchViewer::HandleAssetRemoved);
+		AssetRegistryModule.Get().OnAssetRenamed().AddRaw(this, &SWatchViewer::HandleAssetRenamed);
 	}
 	~SWatchViewer()
 	{
 		FKismetDebugUtilities::WatchedPinsListChangedEvent.RemoveAll(this);
 		FEditorDelegates::ResumePIE.RemoveAll(this);
 		FEditorDelegates::EndPIE.RemoveAll(this);
+
+		if (FModuleManager::Get().IsModuleLoaded(TEXT("AssetRegistry")))
+		{
+			FAssetRegistryModule& AssetRegistryModule = FModuleManager::GetModuleChecked<FAssetRegistryModule>(TEXT("AssetRegistry"));
+			AssetRegistryModule.Get().OnAssetRemoved().RemoveAll(this);
+			AssetRegistryModule.Get().OnAssetRenamed().RemoveAll(this);
+		}
 	}
 
 	void Construct(const FArguments& InArgs, TArray<TSharedRef<FWatchRow>>* InWatchSource);
@@ -288,6 +316,8 @@ public:
 	void HandleWatchedPinsChanged(UBlueprint* BlueprintObj);
 	void HandleResumePIE(bool);
 	void HandleEndPIE(bool);
+	void HandleAssetRemoved(const FAssetData& InAssetData);
+	void HandleAssetRenamed(const FAssetData& InAssetData, const FString& InOldName);
 	void UpdateWatches(TArray<TSharedRef<FWatchRow>>* WatchValues);
 	void CopySelectedRows() const;
 	void StopWatchingPin() const;
@@ -455,6 +485,16 @@ void SWatchViewer::HandleEndPIE(bool)
 	WatchViewer::ContinueExecution();
 }
 
+void SWatchViewer::HandleAssetRemoved(const FAssetData& InAssetData)
+{
+	WatchViewer::RemoveWatchesForAsset(InAssetData);
+}
+
+void SWatchViewer::HandleAssetRenamed(const FAssetData& InAssetData, const FString& InOldName)
+{
+	WatchViewer::OnRenameAsset(InAssetData, InOldName);
+}
+
 void SWatchViewer::UpdateWatches(TArray<TSharedRef<FWatchRow>>* Watches)
 {
 	WatchSource = Watches;
@@ -769,7 +809,6 @@ void UpdateInstancedWatchDisplay()
 
 void WatchViewer::UpdateDisplayedWatches(const TArray<const FFrame*>& ScriptStack)
 {
-	bIsExecutionPaused = true;
 	BlueprintStackInstances.Reset();
 
 	for (const FFrame* Frame : ScriptStack)
@@ -790,7 +829,6 @@ void WatchViewer::UpdateDisplayedWatches(const TArray<const FFrame*>& ScriptStac
 
 void WatchViewer::ContinueExecution()
 {
-	bIsExecutionPaused = false;
 	// Notify subscribers:
 	WatchListSubscribers.Broadcast(&Private_WatchSource);
 }
@@ -829,7 +867,7 @@ void UpdateWatchListFromBlueprintImpl(UBlueprint* BlueprintObj, const bool bShou
 	// something changed so we need to update the lists shown in the UI
 	UpdateNonInstancedWatchDisplay();
 
-	if (bIsExecutionPaused)
+	if (IsPaused())
 	{
 		UpdateInstancedWatchDisplay();
 	}
@@ -838,6 +876,80 @@ void UpdateWatchListFromBlueprintImpl(UBlueprint* BlueprintObj, const bool bShou
 	WatchListSubscribers.Broadcast(&Private_WatchSource);
 }
 
+void WatchViewer::RemoveWatchesForBlueprint(class UBlueprint* BlueprintObj)
+{
+	if (!ensure(BlueprintObj))
+	{
+		return;
+	}
+
+	int32 FoundIdx = WatchedBlueprints.Find(BlueprintObj);
+	if (FoundIdx == INDEX_NONE)
+	{
+		return;
+	}
+
+	// since we're not watching any pins anymore we should remove it from the watched list
+	WatchedBlueprints.RemoveAt(FoundIdx);
+
+	// something changed so we need to update the lists shown in the UI
+	UpdateNonInstancedWatchDisplay();
+
+	if (IsPaused())
+	{
+		UpdateInstancedWatchDisplay();
+	}
+
+	// Notify subscribers
+	WatchListSubscribers.Broadcast(&Private_WatchSource);
+}
+
+void WatchViewer::RemoveWatchesForAsset(const struct FAssetData& AssetData)
+{
+	for (TSharedRef<FWatchRow> WatchRow : Private_WatchSource)
+	{
+		if (AssetData.PackageName == WatchRow->BlueprintPackageName && FText::FromName(AssetData.AssetName).EqualTo(WatchRow->BlueprintName))
+		{
+			RemoveWatchesForBlueprint(WatchRow->BP);
+			break;
+		}
+	}
+}
+
+void WatchViewer::OnRenameAsset(const struct FAssetData& AssetData, const FString& OldAssetName)
+{
+	FString OldPackageName;
+	FString OldBPName;
+
+	if (OldAssetName.Split(".", &OldPackageName, &OldBPName))
+	{
+		bool bUpdated = false;
+
+		for (TSharedRef<FWatchRow> WatchRow : Private_WatchSource)
+		{
+			if (OldPackageName == WatchRow->BlueprintPackageName.ToString() && FText::FromString(OldBPName).EqualTo(WatchRow->BlueprintName))
+			{
+				WatchRow->BlueprintName = FText::FromName(AssetData.AssetName);
+				bUpdated = true;
+			}
+		}
+	
+		if (bUpdated)
+		{
+			// something changed so we need to update the lists shown in the UI
+			UpdateNonInstancedWatchDisplay();
+
+			if (IsPaused())
+			{
+				UpdateInstancedWatchDisplay();
+			}
+
+			// Notify subscribers if necessary
+			WatchListSubscribers.Broadcast(&Private_WatchSource);
+		}
+	}
+}
+
 void WatchViewer::UpdateWatchListFromBlueprint(UBlueprint* BlueprintObj)
 {
 	UpdateWatchListFromBlueprintImpl(BlueprintObj, BlueprintObj && BlueprintObj->WatchedPins.Num() > 0);
@@ -852,6 +964,12 @@ void WatchViewer::RegisterTabSpawner(FTabManager& TabManager)
 {
 	const auto SpawnWatchViewTab = []( const FSpawnTabArgs& Args )
 	{
+		TArray<TSharedRef<FWatchRow>>* Source = &Private_WatchSource;
+		if (IsPaused())
+		{
+			Source = &Private_InstanceWatchSource;
+		}
+
 		return SNew(SDockTab)
 			.TabRole( ETabRole::PanelTab )
 			.Label( LOCTEXT("TabTitle", "Watches") )
@@ -859,7 +977,7 @@ void WatchViewer::RegisterTabSpawner(FTabManager& TabManager)
 				SNew(SBorder)
 				.BorderImage( FEditorStyle::GetBrush("Docking.Tab.ContentAreaBrush") )
 				[
-					SNew(SWatchViewer, &Private_WatchSource)
+					SNew(SWatchViewer, Source)
 				]
 			];
 	};
diff --git a/Engine/Source/Editor/Kismet/Public/WatchPointViewer.h b/Engine/Source/Editor/Kismet/Public/WatchPointViewer.h
--- a/Engine/Source/Editor/Kismet/Public/WatchPointViewer.h
+++ b/Engine/Source/Editor/Kismet/Public/WatchPointViewer.h
@@ -5,6 +5,9 @@
 #include "Widgets/SCompoundWidget.h"
 #include "WorkflowOrientedApp/WorkflowTabFactory.h"
 
+struct FAssetData;
+class UBlueprint;
+
 namespace WatchViewer
 {
 	// called when we pause execution to update the displayed watch values
@@ -14,7 +17,18 @@ namespace WatchViewer
 	void KISMET_API ContinueExecution();
 
 	// called when we are adding or changing watches from BlueprintObj
-	void KISMET_API UpdateWatchListFromBlueprint(class UBlueprint* BlueprintObj);
+	void KISMET_API UpdateWatchListFromBlueprint(UBlueprint* BlueprintObj);
+
+	// called when we want to remove watches in the watch window from a blueprint
+	// does NOT remove watches from the pins in the blueprint object
+	void KISMET_API RemoveWatchesForBlueprint(UBlueprint* BlueprintObj);
+
+	// called when we want to remove watches in the watch window from a blueprint
+	// does NOT remove watches from the pins in the blueprint object
+	void KISMET_API RemoveWatchesForAsset(const FAssetData& AssetData);
+
+	// called when an asset is renamed; updates the watches on the asset
+	void KISMET_API OnRenameAsset(const FAssetData& AssetData, const FString& OldAssetName);
 
 	// called when a BlueprintObj should no longer be watched
 	void KISMET_API ClearWatchListFromBlueprint(class UBlueprint* BlueprintObj);
