Fix crash when an interface that does not pass the filter is chosen

diff --git a/Engine/Source/Editor/Kismet/Private/BlueprintDetailsCustomization.cpp b/Engine/Source/Editor/Kismet/Private/BlueprintDetailsCustomization.cpp
--- a/Engine/Source/Editor/Kismet/Private/BlueprintDetailsCustomization.cpp
+++ b/Engine/Source/Editor/Kismet/Private/BlueprintDetailsCustomization.cpp
@@ -5096,12 +5096,15 @@ void FBlueprintInterfaceLayout::OnClassPicked(UClass* PickedClass)
 		AddInterfaceComboButton->SetIsOpen(false);
 	}
 
-	UBlueprint* Blueprint = GlobalOptionsDetailsPtr.Pin()->GetBlueprintObj();
-	check(Blueprint);
+	if (PickedClass)
+	{
+		UBlueprint* Blueprint = GlobalOptionsDetailsPtr.Pin()->GetBlueprintObj();
+		check(Blueprint);
 
-	FBlueprintEditorUtils::ImplementNewInterface( Blueprint, PickedClass->GetFName() );
+		FBlueprintEditorUtils::ImplementNewInterface(Blueprint, PickedClass->GetFName());
 
-	RegenerateChildrenDelegate.ExecuteIfBound();
+		RegenerateChildrenDelegate.ExecuteIfBound();
+	}
 
 	OnRefreshInDetailsView();
 }
