Mirrors CL 4133774.

Make the watch window refresh the list of breakpoints when it is opened to make sure it is current. Also moved some helper functions into an anonymous namespace.

#rnx

diff --git a/Engine/Source/Editor/Kismet/Private/WatchPointViewer.cpp b/Engine/Source/Editor/Kismet/Private/WatchPointViewer.cpp
--- a/Engine/Source/Editor/Kismet/Private/WatchPointViewer.cpp
+++ b/Engine/Source/Editor/Kismet/Private/WatchPointViewer.cpp
@@ -37,132 +37,333 @@
 
 namespace
 {
+	struct FWatchRow
+	{
+		FWatchRow(
+			UBlueprint* InBP,
+			const UEdGraphNode* InNode,
+			const UEdGraphPin* InPin,
+			UObject* InObjectBeingDebugged,
+			FText InBlueprintName,
+			FText InGraphName,
+			FText InNodeName,
+			FText InDisplayName,
+			FText InValue,
+			FText InType
+		)
+			: BP(InBP)
+			, Node(InNode)
+			, Pin(InPin)
+			, ObjectBeingDebugged(InObjectBeingDebugged)
+			, BlueprintName(MoveTemp(InBlueprintName))
+			, GraphName(MoveTemp(InGraphName))
+			, NodeName(MoveTemp(InNodeName))
+			, DisplayName(MoveTemp(InDisplayName))
+			, Value(MoveTemp(InValue))
+			, Type(MoveTemp(InType))
+		{
+			SetObjectBeingDebuggedName();
+
+			UPackage* Package = Cast<UPackage>(BP ? BP->GetOuter() : nullptr);
+			BlueprintPackageName = Package ? Package->GetFName() : FName();
+		}
+
+		FWatchRow(
+			UBlueprint* InBP,
+			const UEdGraphNode* InNode,
+			const UEdGraphPin* InPin,
+			UObject* InObjectBeingDebugged,
+			FText InBlueprintName,
+			FText InGraphName,
+			FText InNodeName,
+			FDebugInfo Info
+		)
+			: BP(InBP)
+			, Node(InNode)
+			, Pin(InPin)
+			, ObjectBeingDebugged(InObjectBeingDebugged)
+			, BlueprintName(MoveTemp(InBlueprintName))
+			, GraphName(MoveTemp(InGraphName))
+			, NodeName(MoveTemp(InNodeName))
+			, DisplayName(MoveTemp(Info.DisplayName))
+			, Value(MoveTemp(Info.Value))
+			, Type(MoveTemp(Info.Type))
+		{
+			SetObjectBeingDebuggedName();
+
+			UPackage* Package = Cast<UPackage>(BP ? BP->GetOuter() : nullptr);
+			BlueprintPackageName = Package ? Package->GetFName() : FName();
+
+			for (FDebugInfo& ChildInfo : Info.Children)
+			{
+				Children.Add(MakeShared<FWatchRow>(InBP, InNode, InPin, InObjectBeingDebugged, BlueprintName, GraphName, NodeName, MoveTemp(ChildInfo)));
+			}
+		}
+
+		// this can't be const because we store watches in the blueprint
+		UBlueprint* BP;
+		const UEdGraphNode* Node;
+		const UEdGraphPin* Pin;
+		// this can't be const because SelectActor takes a non-const actor
+		UObject* ObjectBeingDebugged;
+
+		FText BlueprintName;
+		FText ObjectBeingDebuggedName;
+		FText GraphName;
+		FText NodeName;
+		FText DisplayName;
+		FText Value;
+		FText Type;
+		FName BlueprintPackageName;
+
+		TArray<TSharedRef<FWatchRow>> Children;
+
+		// used for copying entries in the watch viewer
+		FText GetTextForEntry() const
+		{
+			FFormatNamedArguments Args;
+			Args.Add(TEXT("ObjectName"), FText::FromString(ObjectBeingDebugged ? ObjectBeingDebugged->GetName() : TEXT("")));
+			Args.Add(TEXT("BlueprintName"), BlueprintName);
+			Args.Add(TEXT("GraphName"), GraphName);
+			Args.Add(TEXT("NodeName"), NodeName);
+			Args.Add(TEXT("DisplayName"), DisplayName);
+			Args.Add(TEXT("Type"), Type);
+			Args.Add(TEXT("Value"), Value);
+			return FText::Format(LOCTEXT("WatchEntry", "{ObjectName}({BlueprintName}) {GraphName} {NodeName} {DisplayName}({Type}): {Value}"), Args);
+		}
+
+	private:
+		void SetObjectBeingDebuggedName()
+		{
+			if (ObjectBeingDebugged != nullptr)
+			{
+				AActor* ActorBeingDebugged = Cast<AActor>(ObjectBeingDebugged);
+				if (ActorBeingDebugged)
+				{
+					ObjectBeingDebuggedName = FText::AsCultureInvariant(ActorBeingDebugged->GetActorLabel());
+				}
+				else
+				{
+					ObjectBeingDebuggedName = FText::FromName(ObjectBeingDebugged->GetFName());
+				}
+			}
+			else
+			{
+				ObjectBeingDebuggedName = BlueprintName;
+			}
+		}
+	};
+
+	DECLARE_MULTICAST_DELEGATE_OneParam(FOnDisplayedWatchWindowChanged, TArray<TSharedRef<FWatchRow>>*);
+	FOnDisplayedWatchWindowChanged WatchListSubscribers;
+
+	// Proxy array of the watches. This allows us to manually refresh UI state when changes are made:
+	TArray<TSharedRef<FWatchRow>> Private_WatchSource;
+	TArray<TSharedRef<FWatchRow>> Private_InstanceWatchSource;
+
+	TArray<UBlueprint*> WatchedBlueprints;
+	TArray<UObject*> BlueprintStackInstances;
+
 	// Returns true if the blueprint execution is currently paused; false otherwise
 	bool IsPaused()
 	{
 		return GUnrealEd && GUnrealEd->PlayWorld && GUnrealEd->PlayWorld->bDebugPauseExecution;
 	}
-};
 
-struct FWatchRow
-{
-	FWatchRow(
-		UBlueprint* InBP,
-		const UEdGraphNode* InNode,
-		const UEdGraphPin* InPin,
-		UObject* InObjectBeingDebugged,
-		FText InBlueprintName,
-		FText InGraphName,
-		FText InNodeName,
-		FText InDisplayName,
-		FText InValue,
-		FText InType
-	)
-		: BP(InBP)
-		, Node(InNode)
-		, Pin(InPin)
-		, ObjectBeingDebugged(InObjectBeingDebugged)
-		, BlueprintName(MoveTemp(InBlueprintName))
-		, GraphName(MoveTemp(InGraphName))
-		, NodeName(MoveTemp(InNodeName))
-		, DisplayName(MoveTemp(InDisplayName))
-		, Value(MoveTemp(InValue))
-		, Type(MoveTemp(InType))
-	{
-		SetObjectBeingDebuggedName();
-
-		UPackage* Package = Cast<UPackage>(BP ? BP->GetOuter() : nullptr);
-		BlueprintPackageName = Package ? Package->GetFName() : FName();
-	}
-
-	FWatchRow(
-		UBlueprint* InBP,
-		const UEdGraphNode* InNode,
-		const UEdGraphPin* InPin,
-		UObject* InObjectBeingDebugged,
-		FText InBlueprintName,
-		FText InGraphName,
-		FText InNodeName,
-		FDebugInfo Info
-	)
-		: BP(InBP)
-		, Node(InNode)
-		, Pin(InPin)
-		, ObjectBeingDebugged(InObjectBeingDebugged)
-		, BlueprintName(MoveTemp(InBlueprintName))
-		, GraphName(MoveTemp(InGraphName))
-		, NodeName(MoveTemp(InNodeName))
-		, DisplayName(MoveTemp(Info.DisplayName))
-		, Value(MoveTemp(Info.Value))
-		, Type(MoveTemp(Info.Type))
-	{
-		SetObjectBeingDebuggedName();
-
-		UPackage* Package = Cast<UPackage>(BP ? BP->GetOuter() : nullptr);
-		BlueprintPackageName = Package ? Package->GetFName() : FName();
-
-		for (FDebugInfo& ChildInfo : Info.Children)
+	void UpdateNonInstancedWatchDisplay()
+	{
+		Private_WatchSource.Reset();
+
+		for (UBlueprint* BlueprintObj : WatchedBlueprints)
 		{
-			Children.Add(MakeShared<FWatchRow>(InBP, InNode, InPin, InObjectBeingDebugged, BlueprintName, GraphName, NodeName, MoveTemp(ChildInfo)));
-		}
-	}
+			FText BlueprintName = FText::FromString(BlueprintObj->GetName());
 
-	// this can't be const because we store watches in the blueprint
-	UBlueprint* BP;
-	const UEdGraphNode* Node;
-	const UEdGraphPin* Pin;
-	// this can't be const because SelectActor takes a non-const actor
-	UObject* ObjectBeingDebugged;
-	
-	FText BlueprintName;
-	FText ObjectBeingDebuggedName;
-	FText GraphName;
-	FText NodeName;
-	FText DisplayName;
-	FText Value;
-	FText Type;
-	FName BlueprintPackageName;
-
-	TArray<TSharedRef<FWatchRow>> Children;
-
-	// used for copying entries in the watch viewer
-	FText GetTextForEntry() const
-	{
-		FFormatNamedArguments Args;
-		Args.Add(TEXT("ObjectName"), FText::FromString(ObjectBeingDebugged ? ObjectBeingDebugged->GetName() : TEXT("")));
-		Args.Add(TEXT("BlueprintName"), BlueprintName);
-		Args.Add(TEXT("GraphName"), GraphName);
-		Args.Add(TEXT("NodeName"), NodeName);
-		Args.Add(TEXT("DisplayName"), DisplayName);
-		Args.Add(TEXT("Type"), Type);
-		Args.Add(TEXT("Value"), Value);
-		return FText::Format(LOCTEXT("WatchEntry", "{ObjectName}({BlueprintName}) {GraphName} {NodeName} {DisplayName}({Type}): {Value}"), Args);
+			for (const FEdGraphPinReference& PinRef : BlueprintObj->WatchedPins)
+			{
+				if (UEdGraphPin* Pin = PinRef.Get())
+				{
+					FText GraphName = FText::FromString(Pin->GetOwningNode()->GetGraph()->GetName());
+					FText NodeName = Pin->GetOwningNode()->GetNodeTitle(ENodeTitleType::ListView);
+
+					const UEdGraphSchema* Schema = Pin->GetOwningNode()->GetSchema();
+
+					FDebugInfo DebugInfo;
+					DebugInfo.DisplayName = Schema->GetPinDisplayName(Pin);
+					DebugInfo.Type = UEdGraphSchema_K2::TypeToText(Pin->PinType);
+					DebugInfo.Value = LOCTEXT("ExecutionNotPaused", "(execution not paused)");
+
+					Private_WatchSource.Add(
+						MakeShared<FWatchRow>(
+							BlueprintObj,
+							Pin->GetOwningNode(),
+							Pin,
+							nullptr,
+							BlueprintName,
+							MoveTemp(GraphName),
+							MoveTemp(NodeName),
+							MoveTemp(DebugInfo)
+						)
+					);
+				}
+			}
+		}
 	}
 
-private:
-	void SetObjectBeingDebuggedName()
+	void UpdateInstancedWatchDisplay()
 	{
-		if (ObjectBeingDebugged != nullptr)
+		Private_InstanceWatchSource.Reset();
+
+		for (UObject* BlueprintInstance : BlueprintStackInstances)
 		{
-			AActor* ActorBeingDebugged = Cast<AActor>(ObjectBeingDebugged);
-			if (ActorBeingDebugged)
+			UClass* Class = BlueprintInstance->GetClass();
+			UBlueprint* BlueprintObj = (Class ? Cast<UBlueprint>(Class->ClassGeneratedBy) : nullptr);
+			if (BlueprintObj == nullptr)
 			{
-				ObjectBeingDebuggedName = FText::AsCultureInvariant(ActorBeingDebugged->GetActorLabel());
+				continue;
 			}
-			else
+
+			FText BlueprintName = FText::FromString(BlueprintObj->GetName());
+
+			// Don't show info for the CDO
+			if (BlueprintInstance->IsDefaultSubobject())
+			{
+				continue;
+			}
+
+			// Don't show info if this instance is pending kill
+			if (BlueprintInstance->IsPendingKill())
 			{
-				ObjectBeingDebuggedName = FText::FromName(ObjectBeingDebugged->GetFName());
+				continue;
 			}
+
+			// Don't show info if this instance isn't in the current world
+			UObject* ObjOuter = BlueprintInstance;
+			UWorld* ObjWorld = nullptr;
+			static bool bUseNewWorldCode = false;
+			do		// Run through at least once in case the TestObject is a UGameInstance
+			{
+				UGameInstance* ObjGameInstance = Cast<UGameInstance>(ObjOuter);
+
+				ObjOuter = ObjOuter->GetOuter();
+				ObjWorld = ObjGameInstance ? ObjGameInstance->GetWorld() : Cast<UWorld>(ObjOuter);
+			} while (ObjWorld == nullptr && ObjOuter != nullptr);
+
+			if (ObjWorld)
+			{
+				// Make check on owning level (not streaming level)
+				if (ObjWorld->PersistentLevel && ObjWorld->PersistentLevel->OwningWorld)
+				{
+					ObjWorld = ObjWorld->PersistentLevel->OwningWorld;
+				}
+
+				if (ObjWorld->WorldType != EWorldType::PIE && !((ObjWorld->WorldType == EWorldType::Editor) && (GUnrealEd->GetPIEViewport() == nullptr)))
+				{
+					continue;
+				}
+			}
+
+			// We have a valid instance, iterate over all the watched pins and create rows for them
+			for (const FEdGraphPinReference& PinRef : BlueprintObj->WatchedPins)
+			{
+				UEdGraphPin* Pin = PinRef.Get();
+
+				FText GraphName = FText::FromString(Pin->GetOwningNode()->GetGraph()->GetName());
+				FText NodeName = Pin->GetOwningNode()->GetNodeTitle(ENodeTitleType::ListView);
+
+				FDebugInfo DebugInfo;
+				const FKismetDebugUtilities::EWatchTextResult WatchStatus = FKismetDebugUtilities::GetDebugInfo(DebugInfo, BlueprintObj, BlueprintInstance, Pin);
+
+				if (WatchStatus != FKismetDebugUtilities::EWTR_Valid)
+				{
+					const UEdGraphSchema* Schema = Pin->GetOwningNode()->GetSchema();
+					DebugInfo.DisplayName = Schema->GetPinDisplayName(Pin);
+					DebugInfo.Type = UEdGraphSchema_K2::TypeToText(Pin->PinType);
+
+					switch (WatchStatus)
+					{
+					case FKismetDebugUtilities::EWTR_NotInScope:
+						DebugInfo.Value = LOCTEXT("NotInScope", "(not in scope)");
+						break;
+
+					case FKismetDebugUtilities::EWTR_NoProperty:
+						DebugInfo.Value = LOCTEXT("NoDebugData", "(no debug data)");
+						break;
+
+					case FKismetDebugUtilities::EWTR_NoDebugObject:
+						DebugInfo.Value = LOCTEXT("NoDebugObject", "(no debug object)");
+						break;
+
+					default:
+						// do nothing
+						break;
+					}
+				}
+
+				Private_InstanceWatchSource.Add(
+					MakeShared<FWatchRow>(
+						BlueprintObj,
+						Pin->GetOwningNode(),
+						Pin,
+						BlueprintInstance,
+						BlueprintName,
+						GraphName,
+						NodeName,
+						DebugInfo
+						)
+				);
+			}
+		}
+
+		// Notify subscribers:
+		WatchListSubscribers.Broadcast(&Private_InstanceWatchSource);
+	}
+
+	void UpdateWatchListFromBlueprintImpl(UBlueprint* BlueprintObj, const bool bShouldWatch)
+	{
+		if (!ensure(BlueprintObj))
+		{
+			return;
+		}
+
+		if (bShouldWatch)
+		{
+			// make sure the blueprint is in our list
+			WatchedBlueprints.AddUnique(BlueprintObj);
 		}
 		else
 		{
-			ObjectBeingDebuggedName = BlueprintName;
+			// if this blueprint shouldn't be watched and we aren't watching it already then there is nothing to do
+			int32 FoundIdx = WatchedBlueprints.Find(BlueprintObj);
+			if (FoundIdx == INDEX_NONE)
+			{
+				return;
+			}
+
+			// since we're not watching the blueprint anymore we should remove it from the watched list
+			WatchedBlueprints.RemoveAt(FoundIdx);
 		}
+
+		// something changed so we need to update the lists shown in the UI
+		UpdateNonInstancedWatchDisplay();
+
+		if (IsPaused())
+		{
+			UpdateInstancedWatchDisplay();
+		}
+
+		// Notify subscribers:
+		WatchListSubscribers.Broadcast(&Private_WatchSource);
 	}
-};
 
-DECLARE_MULTICAST_DELEGATE_OneParam(FOnDisplayedWatchWindowChanged, TArray<TSharedRef<FWatchRow>>*);
-FOnDisplayedWatchWindowChanged WatchListSubscribers;
+	// Updates all of the watches from the currently watched blueprints
+	void UpdateAllBlueprintWatches()
+	{
+		for (UBlueprint* Blueprint : WatchedBlueprints)
+		{
+			UpdateWatchListFromBlueprintImpl(Blueprint, true);
+		}
+	}
+};
 
 /**
 * Widget that visualizes the contents of a FWatchRow.
@@ -288,6 +489,9 @@ public:
 
 	SWatchViewer()
 	{
+		// make sure we have the latest information about the watches on loaded blueprints
+		UpdateAllBlueprintWatches();
+
 		FKismetDebugUtilities::WatchedPinsListChangedEvent.AddRaw(this, &SWatchViewer::HandleWatchedPinsChanged);
 		FEditorDelegates::ResumePIE.AddRaw(this, &SWatchViewer::HandleResumePIE);
 		FEditorDelegates::EndPIE.AddRaw(this, &SWatchViewer::HandleEndPIE);
@@ -650,163 +854,6 @@ TSharedRef<SWidget> SWatchTreeWidgetItem::GenerateWidgetForColumn(const FName& C
 	}
 }
 
-// Proxy array of the watches. This allows us to manually refresh UI state when changes are made:
-TArray<TSharedRef<FWatchRow>> Private_WatchSource;
-TArray<TSharedRef<FWatchRow>> Private_InstanceWatchSource;
-
-TArray<UBlueprint*> WatchedBlueprints;
-TArray<UObject*> BlueprintStackInstances;
-
-bool bIsExecutionPaused = false;
-
-void UpdateNonInstancedWatchDisplay()
-{
-	Private_WatchSource.Reset();
-
-	for (UBlueprint* BlueprintObj : WatchedBlueprints)
-	{
-		FText BlueprintName = FText::FromString(BlueprintObj->GetName());
-
-		for (const FEdGraphPinReference& PinRef : BlueprintObj->WatchedPins)
-		{
-			if (UEdGraphPin* Pin = PinRef.Get())
-			{
-				FText GraphName = FText::FromString(Pin->GetOwningNode()->GetGraph()->GetName());
-				FText NodeName = Pin->GetOwningNode()->GetNodeTitle(ENodeTitleType::ListView);
-
-				const UEdGraphSchema* Schema = Pin->GetOwningNode()->GetSchema();
-
-				FDebugInfo DebugInfo;
-				DebugInfo.DisplayName = Schema->GetPinDisplayName(Pin);
-				DebugInfo.Type = UEdGraphSchema_K2::TypeToText(Pin->PinType);
-				DebugInfo.Value = LOCTEXT("ExecutionNotPaused", "(execution not paused)");
-
-				Private_WatchSource.Add(
-					MakeShared<FWatchRow>(
-						BlueprintObj,
-						Pin->GetOwningNode(),
-						Pin,
-						nullptr,
-						BlueprintName,
-						MoveTemp(GraphName),
-						MoveTemp(NodeName),
-						MoveTemp(DebugInfo)
-						)
-				);
-			}
-		}
-	}
-}
-
-void UpdateInstancedWatchDisplay()
-{
-	Private_InstanceWatchSource.Reset();
-
-	for (UObject* BlueprintInstance : BlueprintStackInstances)
-	{
-		UClass* Class = BlueprintInstance->GetClass();
-		UBlueprint* BlueprintObj = (Class ? Cast<UBlueprint>(Class->ClassGeneratedBy) : nullptr);
-		if (BlueprintObj == nullptr)
-		{
-			continue;
-		}
-
-		FText BlueprintName = FText::FromString(BlueprintObj->GetName());
-
-		// Don't show info for the CDO
-		if (BlueprintInstance->IsDefaultSubobject())
-		{
-			continue;
-		}
-
-		// Don't show info if this instance is pending kill
-		if (BlueprintInstance->IsPendingKill())
-		{
-			continue;
-		}
-
-		// Don't show info if this instance isn't in the current world
-		UObject* ObjOuter = BlueprintInstance;
-		UWorld* ObjWorld = nullptr;
-		static bool bUseNewWorldCode = false;
-		do		// Run through at least once in case the TestObject is a UGameInstance
-		{
-			UGameInstance* ObjGameInstance = Cast<UGameInstance>(ObjOuter);
-
-			ObjOuter = ObjOuter->GetOuter();
-			ObjWorld = ObjGameInstance ? ObjGameInstance->GetWorld() : Cast<UWorld>(ObjOuter);
-		} while (ObjWorld == nullptr && ObjOuter != nullptr);
-
-		if (ObjWorld)
-		{
-			// Make check on owning level (not streaming level)
-			if (ObjWorld->PersistentLevel && ObjWorld->PersistentLevel->OwningWorld)
-			{
-				ObjWorld = ObjWorld->PersistentLevel->OwningWorld;
-			}
-
-			if (ObjWorld->WorldType != EWorldType::PIE && !((ObjWorld->WorldType == EWorldType::Editor) && (GUnrealEd->GetPIEViewport() == nullptr)))
-			{
-				continue;
-			}
-		}
-
-		// We have a valid instance, iterate over all the watched pins and create rows for them
-		for (const FEdGraphPinReference& PinRef : BlueprintObj->WatchedPins)
-		{
-			UEdGraphPin* Pin = PinRef.Get();
-
-			FText GraphName = FText::FromString(Pin->GetOwningNode()->GetGraph()->GetName());
-			FText NodeName = Pin->GetOwningNode()->GetNodeTitle(ENodeTitleType::ListView);
-
-			FDebugInfo DebugInfo;
-			const FKismetDebugUtilities::EWatchTextResult WatchStatus = FKismetDebugUtilities::GetDebugInfo(DebugInfo, BlueprintObj, BlueprintInstance, Pin);
-
-			if (WatchStatus != FKismetDebugUtilities::EWTR_Valid)
-			{
-				const UEdGraphSchema* Schema = Pin->GetOwningNode()->GetSchema();
-				DebugInfo.DisplayName = Schema->GetPinDisplayName(Pin);
-				DebugInfo.Type = UEdGraphSchema_K2::TypeToText(Pin->PinType);
-
-				switch (WatchStatus)
-				{
-				case FKismetDebugUtilities::EWTR_NotInScope:
-					DebugInfo.Value = LOCTEXT("NotInScope", "(not in scope)");
-					break;
-
-				case FKismetDebugUtilities::EWTR_NoProperty:
-					DebugInfo.Value = LOCTEXT("NoDebugData", "(no debug data)");
-					break;
-
-				case FKismetDebugUtilities::EWTR_NoDebugObject:
-					DebugInfo.Value = LOCTEXT("NoDebugObject", "(no debug object)");
-					break;
-
-				default:
-					// do nothing
-					break;
-				}
-			}
-
-			Private_InstanceWatchSource.Add(
-				MakeShared<FWatchRow>(
-					BlueprintObj,
-					Pin->GetOwningNode(),
-					Pin,
-					BlueprintInstance,
-					BlueprintName,
-					GraphName,
-					NodeName,
-					DebugInfo
-					)
-			);
-		}
-	}
-
-	// Notify subscribers:
-	WatchListSubscribers.Broadcast(&Private_InstanceWatchSource);
-}
-
 void WatchViewer::UpdateDisplayedWatches(const TArray<const FFrame*>& ScriptStack)
 {
 	BlueprintStackInstances.Reset();
@@ -839,43 +886,6 @@ FName WatchViewer::GetTabName()
 	return TabName;
 }
 
-void UpdateWatchListFromBlueprintImpl(UBlueprint* BlueprintObj, const bool bShouldWatch)
-{
-	if (!ensure(BlueprintObj))
-	{
-		return;
-	}
-
-	if (bShouldWatch)
-	{
-		// make sure the blueprint is in our list
-		WatchedBlueprints.AddUnique(BlueprintObj);
-	}
-	else
-	{
-		// if this blueprint shouldn't be watched and we aren't watching it already then there is nothing to do
-		int32 FoundIdx = WatchedBlueprints.Find(BlueprintObj);
-		if (FoundIdx == INDEX_NONE)
-		{
-			return;
-		}
-
-		// since we're not watching the blueprint anymore we should remove it from the watched list
-		WatchedBlueprints.RemoveAt(FoundIdx);
-	}
-
-	// something changed so we need to update the lists shown in the UI
-	UpdateNonInstancedWatchDisplay();
-
-	if (IsPaused())
-	{
-		UpdateInstancedWatchDisplay();
-	}
-
-	// Notify subscribers:
-	WatchListSubscribers.Broadcast(&Private_WatchSource);
-}
-
 void WatchViewer::RemoveWatchesForBlueprint(class UBlueprint* BlueprintObj)
 {
 	if (!ensure(BlueprintObj))
@@ -952,7 +962,7 @@ void WatchViewer::OnRenameAsset(const struct FAssetData& AssetData, const FStrin
 
 void WatchViewer::UpdateWatchListFromBlueprint(UBlueprint* BlueprintObj)
 {
-	UpdateWatchListFromBlueprintImpl(BlueprintObj, BlueprintObj && BlueprintObj->WatchedPins.Num() > 0);
+	UpdateWatchListFromBlueprintImpl(BlueprintObj, true);
 }
 
 void WatchViewer::ClearWatchListFromBlueprint(UBlueprint* BlueprintObj)
