AssetEditorManager now has a OnAssetEditorRequestClose event that "quasi editors" can bind to

diff --git a/Engine/Source/Developer/AssetTools/Private/AssetTypeActions/AssetTypeActions_AnimBlueprint.cpp b/Engine/Source/Developer/AssetTools/Private/AssetTypeActions/AssetTypeActions_AnimBlueprint.cpp
--- a/Engine/Source/Developer/AssetTools/Private/AssetTypeActions/AssetTypeActions_AnimBlueprint.cpp
+++ b/Engine/Source/Developer/AssetTools/Private/AssetTypeActions/AssetTypeActions_AnimBlueprint.cpp
@@ -4,7 +4,6 @@
 #include "Framework/MultiBox/MultiBoxBuilder.h"
 #include "Widgets/Layout/SBorder.h"
 #include "Misc/MessageDialog.h"
-#include "Framework/Application/SlateApplication.h"
 #include "Widgets/Images/SImage.h"
 #include "EditorStyleSet.h"
 #include "Animation/AnimInstance.h"
@@ -172,27 +171,7 @@ void FAssetTypeActions_AnimBlueprint::PerformAssetDiff(UObject* Asset1, UObject*
 		WindowTitle = FText::Format(LOCTEXT("AnimationBlueprintDiff", "{0} - Animation Blueprint Diff"), FText::FromString(NewBlueprint->GetName()));
 	}
 
-	const TSharedPtr<SWindow> Window = SNew(SWindow)
-		.Title(WindowTitle)
-		.ClientSize(FVector2D(1000, 800));
-
-	Window->SetContent(SNew(SBlueprintDiff)
-		.BlueprintOld(OldBlueprint)
-		.BlueprintNew(NewBlueprint)
-		.OldRevision(OldRevision)
-		.NewRevision(NewRevision)
-		.ShowAssetNames(!bIsSingleAsset));
-
-	// Make this window a child of the modal window if we've been spawned while one is active.
-	TSharedPtr<SWindow> ActiveModal = FSlateApplication::Get().GetActiveModalWindow();
-	if (ActiveModal.IsValid())
-	{
-		FSlateApplication::Get().AddWindowAsNativeChild(Window.ToSharedRef(), ActiveModal.ToSharedRef());
-	}
-	else
-	{
-		FSlateApplication::Get().AddWindow(Window.ToSharedRef());
-	}
+	SBlueprintDiff::CreateDiffWindow(WindowTitle, OldBlueprint, NewBlueprint, OldRevision, NewRevision);
 }
 
 void FAssetTypeActions_AnimBlueprint::ExecuteFindSkeleton(TArray<TWeakObjectPtr<UAnimBlueprint>> Objects)
diff --git a/Engine/Source/Developer/AssetTools/Private/AssetTypeActions/AssetTypeActions_Blueprint.cpp b/Engine/Source/Developer/AssetTools/Private/AssetTypeActions/AssetTypeActions_Blueprint.cpp
--- a/Engine/Source/Developer/AssetTools/Private/AssetTypeActions/AssetTypeActions_Blueprint.cpp
+++ b/Engine/Source/Developer/AssetTools/Private/AssetTypeActions/AssetTypeActions_Blueprint.cpp
@@ -5,7 +5,6 @@
 #include "GameFramework/Actor.h"
 #include "Misc/PackageName.h"
 #include "Misc/MessageDialog.h"
-#include "Framework/Application/SlateApplication.h"
 #include "EditorStyleSet.h"
 #include "Factories/BlueprintFactory.h"
 #include "ThumbnailRendering/SceneThumbnailInfo.h"
@@ -247,27 +246,7 @@ void FAssetTypeActions_Blueprint::PerformAssetDiff(UObject* OldAsset, UObject* N
 		WindowTitle = FText::Format(LOCTEXT("Blueprint Diff", "{0} - Blueprint Diff"), FText::FromString(NewBlueprint->GetName()));
 	}
 
-	const TSharedPtr<SWindow> Window = SNew(SWindow)
-		.Title(WindowTitle)
-		.ClientSize(FVector2D(1000,800));
-
-	Window->SetContent(SNew(SBlueprintDiff)
-					  .BlueprintOld(OldBlueprint)
-					  .BlueprintNew(NewBlueprint)
-					  .OldRevision(OldRevision)
-					  .NewRevision(NewRevision)
-					  .ShowAssetNames(!bIsSingleAsset) );
-
-	// Make this window a child of the modal window if we've been spawned while one is active.
-	TSharedPtr<SWindow> ActiveModal = FSlateApplication::Get().GetActiveModalWindow();
-	if ( ActiveModal.IsValid() )
-	{
-		FSlateApplication::Get().AddWindowAsNativeChild( Window.ToSharedRef(), ActiveModal.ToSharedRef() );
-	}
-	else
-	{
-		FSlateApplication::Get().AddWindow( Window.ToSharedRef() );
-	}
+	SBlueprintDiff::CreateDiffWindow(WindowTitle, OldBlueprint, NewBlueprint, OldRevision, NewRevision);
 }
 
 UThumbnailInfo* FAssetTypeActions_Blueprint::GetThumbnailInfo(UObject* Asset) const
diff --git a/Engine/Source/Editor/Kismet/Private/SBlueprintDiff.cpp b/Engine/Source/Editor/Kismet/Private/SBlueprintDiff.cpp
--- a/Engine/Source/Editor/Kismet/Private/SBlueprintDiff.cpp
+++ b/Engine/Source/Editor/Kismet/Private/SBlueprintDiff.cpp
@@ -24,6 +24,8 @@
 #include "Framework/Commands/GenericCommands.h"
 #include "WidgetBlueprint.h"
 #include "HAL/PlatformApplicationMisc.h"
+#include "Framework/Application/SlateApplication.h"
+#include "Toolkits/AssetEditorManager.h"
 
 #define LOCTEXT_NAMESPACE "SBlueprintDif"
 
@@ -697,8 +699,12 @@ void SBlueprintDiff::Construct( const FArguments& InArgs)
 
 	bLockViews = true;
 
-	TAttribute<FName> GetActiveMode(this, &SBlueprintDiff::GetCurrentMode);
-	FOnModeChangeRequested SetActiveMode = FOnModeChangeRequested::CreateRaw(this, &SBlueprintDiff::SetCurrentMode);
+	if (InArgs._ParentWindow.IsValid())
+	{
+		WeakParentWindow = InArgs._ParentWindow;
+
+		AssetEditorCloseDelegate = FAssetEditorManager::Get().OnAssetEditorRequestClose().AddSP(this, &SBlueprintDiff::OnCloseAssetEditor);
+	}
 
 	FToolBarBuilder ToolbarBuilder(TSharedPtr< const FUICommandList >(), FMultiBoxCustomization::None);
 	ToolbarBuilder.AddToolBarButton(
@@ -813,6 +819,28 @@ void SBlueprintDiff::Construct( const FArguments& InArgs)
 }
 END_SLATE_FUNCTION_BUILD_OPTIMIZATION
 
+SBlueprintDiff::~SBlueprintDiff()
+{
+	if (AssetEditorCloseDelegate.IsValid())
+	{
+		FAssetEditorManager::Get().OnAssetEditorRequestClose().Remove(AssetEditorCloseDelegate);
+	}
+}
+
+void SBlueprintDiff::OnCloseAssetEditor(UObject* Asset, EAssetEditorCloseReason CloseReason)
+{
+	if (PanelOld.Blueprint == Asset || PanelNew.Blueprint == Asset || CloseReason == EAssetEditorCloseReason::CloseAllAssetEditors)
+	{
+		// Tell our window to close and set our selves to collapsed to try and stop it from ticking
+		SetVisibility(EVisibility::Collapsed);
+
+		if (WeakParentWindow.IsValid())
+		{
+			WeakParentWindow.Pin()->RequestDestroyWindow();
+		}
+	}
+}
+
 TSharedRef<ITableRow> SBlueprintDiff::OnGenerateRow( FGraphToDiff ParamItem, const TSharedRef<STableViewBase>& OwnerTable )
 {
 	return	SNew( STableRow< FGraphToDiff >, OwnerTable )
@@ -871,6 +899,38 @@ TSharedRef<SWidget> SBlueprintDiff::DefaultEmptyPanel()
 		];
 }
 
+TSharedPtr<SWindow> SBlueprintDiff::CreateDiffWindow(FText WindowTitle, UBlueprint* OldBlueprint, UBlueprint* NewBlueprint, const struct FRevisionInfo& OldRevision, const struct FRevisionInfo& NewRevision)
+{
+	// sometimes we're comparing different revisions of one single asset (other 
+	// times we're comparing two completely separate assets altogether)
+	bool bIsSingleAsset = (NewBlueprint->GetName() == OldBlueprint->GetName());
+
+	TSharedPtr<SWindow> Window = SNew(SWindow)
+		.Title(WindowTitle)
+		.ClientSize(FVector2D(1000, 800));
+
+	Window->SetContent(SNew(SBlueprintDiff)
+		.BlueprintOld(OldBlueprint)
+		.BlueprintNew(NewBlueprint)
+		.OldRevision(OldRevision)
+		.NewRevision(NewRevision)
+		.ShowAssetNames(!bIsSingleAsset)
+		.ParentWindow(Window));
+
+	// Make this window a child of the modal window if we've been spawned while one is active.
+	TSharedPtr<SWindow> ActiveModal = FSlateApplication::Get().GetActiveModalWindow();
+	if (ActiveModal.IsValid())
+	{
+		FSlateApplication::Get().AddWindowAsNativeChild(Window.ToSharedRef(), ActiveModal.ToSharedRef());
+	}
+	else
+	{
+		FSlateApplication::Get().AddWindow(Window.ToSharedRef());
+	}
+
+	return Window;
+}
+
 void SBlueprintDiff::NextDiff()
 {
 	DiffTreeView::HighlightNextDifference(DifferencesTreeView.ToSharedRef(), RealDifferences, MasterDifferencesList);
diff --git a/Engine/Source/Editor/Kismet/Public/SBlueprintDiff.h b/Engine/Source/Editor/Kismet/Public/SBlueprintDiff.h
--- a/Engine/Source/Editor/Kismet/Public/SBlueprintDiff.h
+++ b/Engine/Source/Editor/Kismet/Public/SBlueprintDiff.h
@@ -3,7 +3,7 @@
 
 #include "CoreMinimal.h"
 #include "Widgets/DeclarativeSyntaxSupport.h"
-#include "Widgets/SWidget.h"
+#include "Widgets/SWindow.h"
 #include "Widgets/SCompoundWidget.h"
 #include "Textures/SlateIcon.h"
 #include "Widgets/Views/STableViewBase.h"
@@ -21,6 +21,7 @@ class IDiffControl;
 class SMyBlueprint;
 class UEdGraph;
 struct FListItemGraphToDiff;
+enum class EAssetEditorCloseReason : uint8;
 
 struct FMatchFName
 {
@@ -124,9 +125,11 @@ public:
 			SLATE_ARGUMENT( struct FRevisionInfo, OldRevision )
 			SLATE_ARGUMENT( struct FRevisionInfo, NewRevision )
 			SLATE_ARGUMENT( bool, ShowAssetNames )
+			SLATE_ARGUMENT(TSharedPtr<SWindow>, ParentWindow)
 	SLATE_END_ARGS()
 
 	void Construct(const FArguments& InArgs);
+	virtual ~SBlueprintDiff();
 
 	/** Called when a new Graph is clicked on by user */
 	void OnGraphChanged(struct FListItemGraphToDiff* Diff);
@@ -137,6 +140,9 @@ public:
 	/** Helper function for generating an empty widget */
 	static TSharedRef<SWidget> DefaultEmptyPanel();
 
+	/** Helper function to create a window that holds a diff widget */
+	static TSharedPtr<SWindow> CreateDiffWindow(FText WindowTitle, UBlueprint* OldBlueprint, UBlueprint* NewBlueprint, const struct FRevisionInfo& OldRevision, const struct FRevisionInfo& NewRevision);
+
 protected:
 	/** Called when user clicks button to go to next difference */
 	void NextDiff();
@@ -200,6 +206,9 @@ protected:
 	/** Checks if a graph is valid for diffing */
 	bool IsGraphDiffNeeded(class UEdGraph* InGraph) const;
 
+	/** Called when editor may need to be closed */
+	void OnCloseAssetEditor(UObject* Asset, EAssetEditorCloseReason CloseReason);
+
 	struct FDiffControl
 	{
 		FDiffControl()
@@ -252,6 +261,11 @@ protected:
 	FDiffControl GraphPanel;
 	FDiffControl DefaultsPanel;
 	FDiffControl ComponentsPanel;
+
+	/** A pointer to the window holding this */
+	TWeakPtr<SWindow> WeakParentWindow;
+
+	FDelegateHandle AssetEditorCloseDelegate;
 };
 
 
diff --git a/Engine/Source/Editor/UMGEditor/Private/AssetTypeActions_WidgetBlueprint.cpp b/Engine/Source/Editor/UMGEditor/Private/AssetTypeActions_WidgetBlueprint.cpp
--- a/Engine/Source/Editor/UMGEditor/Private/AssetTypeActions_WidgetBlueprint.cpp
+++ b/Engine/Source/Editor/UMGEditor/Private/AssetTypeActions_WidgetBlueprint.cpp
@@ -3,10 +3,7 @@
 #include "AssetTypeActions_WidgetBlueprint.h"
 #include "Misc/MessageDialog.h"
 #include "Engine/Blueprint.h"
-#include "Framework/Application/SlateApplication.h"
-
 #include "WidgetBlueprint.h"
-
 #include "SBlueprintDiff.h"
 #include "WidgetBlueprintEditor.h"
 
@@ -68,27 +65,7 @@ void FAssetTypeActions_WidgetBlueprint::PerformAssetDiff(UObject* Asset1, UObjec
 		WindowTitle = FText::Format(LOCTEXT("WidgetBlueprintDiff", "{0} - Widget Blueprint Diff"), FText::FromString(NewBlueprint->GetName()));
 	}
 
-	const TSharedPtr<SWindow> Window = SNew(SWindow)
-		.Title(WindowTitle)
-		.ClientSize(FVector2D(1000, 800));
-
-	Window->SetContent(SNew(SBlueprintDiff)
-		.BlueprintOld(OldBlueprint)
-		.BlueprintNew(NewBlueprint)
-		.OldRevision(OldRevision)
-		.NewRevision(NewRevision)
-		.ShowAssetNames(!bIsSingleAsset));
-
-	// Make this window a child of the modal window if we've been spawned while one is active.
-	TSharedPtr<SWindow> ActiveModal = FSlateApplication::Get().GetActiveModalWindow();
-	if (ActiveModal.IsValid())
-	{
-		FSlateApplication::Get().AddWindowAsNativeChild(Window.ToSharedRef(), ActiveModal.ToSharedRef());
-	}
-	else
-	{
-		FSlateApplication::Get().AddWindow(Window.ToSharedRef());
-	}
+	SBlueprintDiff::CreateDiffWindow(WindowTitle, OldBlueprint, NewBlueprint, OldRevision, NewRevision);
 }
 
 #undef LOCTEXT_NAMESPACE
diff --git a/Engine/Source/Editor/UnrealEd/Private/Toolkits/AssetEditorManager.cpp b/Engine/Source/Editor/UnrealEd/Private/Toolkits/AssetEditorManager.cpp
--- a/Engine/Source/Editor/UnrealEd/Private/Toolkits/AssetEditorManager.cpp
+++ b/Engine/Source/Editor/UnrealEd/Private/Toolkits/AssetEditorManager.cpp
@@ -168,6 +168,8 @@ void FAssetEditorManager::CloseAllEditorsForAsset(UObject* Asset)
 	{
 		EditorIter->CloseWindow();
 	}
+
+	AssetEditorRequestCloseEvent.Broadcast(Asset, EAssetEditorCloseReason::CloseAllEditorsForAsset);
 }
 
 void FAssetEditorManager::RemoveAssetFromAllEditors(UObject* Asset)
@@ -178,6 +180,8 @@ void FAssetEditorManager::RemoveAssetFromAllEditors(UObject* Asset)
 	{
 		EditorIter->RemoveEditingAsset(Asset);
 	}
+
+	AssetEditorRequestCloseEvent.Broadcast(Asset, EAssetEditorCloseReason::RemoveAssetFromAllEditors);
 }
 
 
@@ -192,6 +196,8 @@ void FAssetEditorManager::CloseOtherEditors( UObject* Asset, IAssetEditorInstanc
 			Editor->CloseWindow();
 		}
 	}
+
+	AssetEditorRequestCloseEvent.Broadcast(Asset, EAssetEditorCloseReason::CloseOtherEditors);
 }
 
 
@@ -286,6 +292,8 @@ bool FAssetEditorManager::CloseAllAssetEditors()
 		}
 	}
 
+	AssetEditorRequestCloseEvent.Broadcast(nullptr, EAssetEditorCloseReason::CloseAllAssetEditors);
+
 	return bAllEditorsClosed;
 }
 
diff --git a/Engine/Source/Editor/UnrealEd/Public/Toolkits/AssetEditorManager.h b/Engine/Source/Editor/UnrealEd/Public/Toolkits/AssetEditorManager.h
--- a/Engine/Source/Editor/UnrealEd/Public/Toolkits/AssetEditorManager.h
+++ b/Engine/Source/Editor/UnrealEd/Public/Toolkits/AssetEditorManager.h
@@ -32,6 +32,14 @@ public:
 	virtual void RemoveEditingAsset(UObject* Asset) = 0;
 };
 
+/** The way that editors were requested to close */
+enum class EAssetEditorCloseReason : uint8
+{
+	CloseAllEditorsForAsset,
+	CloseOtherEditors,
+	RemoveAssetFromAllEditors,
+	CloseAllAssetEditors,
+};
 
 /**
  * Implements a manager for Editor windows that are currently open and the assets they are editing.
@@ -83,6 +91,10 @@ public:
 	/** Remove given asset from all open editors */
 	void RemoveAssetFromAllEditors(UObject* Asset);
 
+	/** Event called when CloseAllEditorsForAsset/RemoveAssetFromAllEditors is called */
+	DECLARE_EVENT_TwoParams(FAssetEditorManager, FAssetEditorRequestCloseEvent, UObject*, EAssetEditorCloseReason);
+	virtual FAssetEditorRequestCloseEvent& OnAssetEditorRequestClose() { return AssetEditorRequestCloseEvent; }
+
 	/** Get all assets currently being tracked with open editors */
 	TArray<UObject*> GetAllEditedAssets();
 
@@ -90,6 +102,7 @@ public:
 	void NotifyAssetOpened(UObject* Asset, IAssetEditorInstance* Instance);
 	void NotifyAssetsOpened( const TArray< UObject* >& Assets, IAssetEditorInstance* Instance);
 
+	/** Called when an asset has been opened in an editor */
 	DECLARE_EVENT_TwoParams(FAssetEditorManager, FOnAssetOpenedInEditorEvent, UObject*, IAssetEditorInstance*);
 	virtual FOnAssetOpenedInEditorEvent& OnAssetOpenedInEditor() { return AssetOpenedInEditorEvent; }
 
@@ -196,6 +209,10 @@ private:
 	/** Holds a delegate to be invoked when the widget ticks. */
 	FTickerDelegate TickDelegate;
 
+	/** Call to request closing editors for an asset */
+	FAssetEditorRequestCloseEvent AssetEditorRequestCloseEvent;
+
+	/** Called when an asset has been opened in an editor */
 	FOnAssetOpenedInEditorEvent AssetOpenedInEditorEvent;
 
 	/** Multicast delegate executed when an asset editor is requested to be opened */
