Mirrors CL 4136761.

Fix a crash that can occur when there are blueprint watches that get hit when the blueprint editor is not open.

#rnx

diff --git a/Engine/Source/Editor/Kismet/Private/WatchPointViewer.cpp b/Engine/Source/Editor/Kismet/Private/WatchPointViewer.cpp
--- a/Engine/Source/Editor/Kismet/Private/WatchPointViewer.cpp
+++ b/Engine/Source/Editor/Kismet/Private/WatchPointViewer.cpp
@@ -162,7 +162,6 @@
 	TArray<TSharedRef<FWatchRow>> Private_InstanceWatchSource;
 
 	TArray<UBlueprint*> WatchedBlueprints;
-	TArray<UObject*> BlueprintStackInstances;
 
 	// Returns true if the blueprint execution is currently paused; false otherwise
 	bool IsPaused()
@@ -209,115 +208,6 @@
 		}
 	}
 
-	void UpdateInstancedWatchDisplay()
-	{
-		Private_InstanceWatchSource.Reset();
-
-		for (UObject* BlueprintInstance : BlueprintStackInstances)
-		{
-			UClass* Class = BlueprintInstance->GetClass();
-			UBlueprint* BlueprintObj = (Class ? Cast<UBlueprint>(Class->ClassGeneratedBy) : nullptr);
-			if (BlueprintObj == nullptr)
-			{
-				continue;
-			}
-
-			FText BlueprintName = FText::FromString(BlueprintObj->GetName());
-
-			// Don't show info for the CDO
-			if (BlueprintInstance->IsDefaultSubobject())
-			{
-				continue;
-			}
-
-			// Don't show info if this instance is pending kill
-			if (BlueprintInstance->IsPendingKill())
-			{
-				continue;
-			}
-
-			// Don't show info if this instance isn't in the current world
-			UObject* ObjOuter = BlueprintInstance;
-			UWorld* ObjWorld = nullptr;
-			static bool bUseNewWorldCode = false;
-			do		// Run through at least once in case the TestObject is a UGameInstance
-			{
-				UGameInstance* ObjGameInstance = Cast<UGameInstance>(ObjOuter);
-
-				ObjOuter = ObjOuter->GetOuter();
-				ObjWorld = ObjGameInstance ? ObjGameInstance->GetWorld() : Cast<UWorld>(ObjOuter);
-			} while (ObjWorld == nullptr && ObjOuter != nullptr);
-
-			if (ObjWorld)
-			{
-				// Make check on owning level (not streaming level)
-				if (ObjWorld->PersistentLevel && ObjWorld->PersistentLevel->OwningWorld)
-				{
-					ObjWorld = ObjWorld->PersistentLevel->OwningWorld;
-				}
-
-				if (ObjWorld->WorldType != EWorldType::PIE && !((ObjWorld->WorldType == EWorldType::Editor) && (GUnrealEd->GetPIEViewport() == nullptr)))
-				{
-					continue;
-				}
-			}
-
-			// We have a valid instance, iterate over all the watched pins and create rows for them
-			for (const FEdGraphPinReference& PinRef : BlueprintObj->WatchedPins)
-			{
-				UEdGraphPin* Pin = PinRef.Get();
-
-				FText GraphName = FText::FromString(Pin->GetOwningNode()->GetGraph()->GetName());
-				FText NodeName = Pin->GetOwningNode()->GetNodeTitle(ENodeTitleType::ListView);
-
-				FDebugInfo DebugInfo;
-				const FKismetDebugUtilities::EWatchTextResult WatchStatus = FKismetDebugUtilities::GetDebugInfo(DebugInfo, BlueprintObj, BlueprintInstance, Pin);
-
-				if (WatchStatus != FKismetDebugUtilities::EWTR_Valid)
-				{
-					const UEdGraphSchema* Schema = Pin->GetOwningNode()->GetSchema();
-					DebugInfo.DisplayName = Schema->GetPinDisplayName(Pin);
-					DebugInfo.Type = UEdGraphSchema_K2::TypeToText(Pin->PinType);
-
-					switch (WatchStatus)
-					{
-					case FKismetDebugUtilities::EWTR_NotInScope:
-						DebugInfo.Value = LOCTEXT("NotInScope", "(not in scope)");
-						break;
-
-					case FKismetDebugUtilities::EWTR_NoProperty:
-						DebugInfo.Value = LOCTEXT("NoDebugData", "(no debug data)");
-						break;
-
-					case FKismetDebugUtilities::EWTR_NoDebugObject:
-						DebugInfo.Value = LOCTEXT("NoDebugObject", "(no debug object)");
-						break;
-
-					default:
-						// do nothing
-						break;
-					}
-				}
-
-				Private_InstanceWatchSource.Add(
-					MakeShared<FWatchRow>(
-						BlueprintObj,
-						Pin->GetOwningNode(),
-						Pin,
-						BlueprintInstance,
-						BlueprintName,
-						GraphName,
-						NodeName,
-						DebugInfo
-						)
-				);
-			}
-		}
-
-		// Notify subscribers:
-		WatchListSubscribers.Broadcast(&Private_InstanceWatchSource);
-	}
-
 	void UpdateWatchListFromBlueprintImpl(UBlueprint* BlueprintObj, const bool bShouldWatch)
 	{
 		if (!ensure(BlueprintObj))
@@ -348,7 +238,7 @@
 
 		if (IsPaused())
 		{
-			UpdateInstancedWatchDisplay();
+			WatchViewer::UpdateInstancedWatchDisplay();
 		}
 
 		// Notify subscribers:
@@ -854,24 +744,116 @@
 	}
 }
 
-void WatchViewer::UpdateDisplayedWatches(const TArray<const FFrame*>& ScriptStack)
+void WatchViewer::UpdateInstancedWatchDisplay()
 {
-	BlueprintStackInstances.Reset();
-
-	for (const FFrame* Frame : ScriptStack)
 	{
-		if (Frame == nullptr)
+		const TArray<const FFrame*>& ScriptStack = FBlueprintExceptionTracker::Get().ScriptStack;
+
+		for (const FFrame* ScriptFrame : ScriptStack)
 		{
-			continue;
+			UObject* BlueprintInstance = ScriptFrame ? ScriptFrame->Object : nullptr;
+			UClass* Class = BlueprintInstance ? BlueprintInstance->GetClass() : nullptr;
+			UBlueprint* BlueprintObj = (Class ? Cast<UBlueprint>(Class->ClassGeneratedBy) : nullptr);
+			if (BlueprintObj == nullptr)
+			{
+				continue;
+			}
+
+			FText BlueprintName = FText::FromString(BlueprintObj->GetName());
+
+			// Don't show info for the CDO
+			if (BlueprintInstance->IsDefaultSubobject())
+			{
+				continue;
+			}
+
+			// Don't show info if this instance is pending kill
+			if (BlueprintInstance->IsPendingKill())
+			{
+				continue;
+			}
+
+			// Don't show info if this instance isn't in the current world
+			UObject* ObjOuter = BlueprintInstance;
+			UWorld* ObjWorld = nullptr;
+			static bool bUseNewWorldCode = false;
+			do		// Run through at least once in case the TestObject is a UGameInstance
+			{
+				UGameInstance* ObjGameInstance = Cast<UGameInstance>(ObjOuter);
+
+				ObjOuter = ObjOuter->GetOuter();
+				ObjWorld = ObjGameInstance ? ObjGameInstance->GetWorld() : Cast<UWorld>(ObjOuter);
+			} while (ObjWorld == nullptr && ObjOuter != nullptr);
+
+			if (ObjWorld)
+			{
+				// Make check on owning level (not streaming level)
+				if (ObjWorld->PersistentLevel && ObjWorld->PersistentLevel->OwningWorld)
+				{
+					ObjWorld = ObjWorld->PersistentLevel->OwningWorld;
+				}
+
+				if (ObjWorld->WorldType != EWorldType::PIE && !((ObjWorld->WorldType == EWorldType::Editor) && (GUnrealEd->GetPIEViewport() == nullptr)))
+				{
+					continue;
+				}
+			}
+
+			// We have a valid instance, iterate over all the watched pins and create rows for them
+			for (const FEdGraphPinReference& PinRef : BlueprintObj->WatchedPins)
+			{
+				UEdGraphPin* Pin = PinRef.Get();
+
+				FText GraphName = FText::FromString(Pin->GetOwningNode()->GetGraph()->GetName());
+				FText NodeName = Pin->GetOwningNode()->GetNodeTitle(ENodeTitleType::ListView);
+
+				FDebugInfo DebugInfo;
+				const FKismetDebugUtilities::EWatchTextResult WatchStatus = FKismetDebugUtilities::GetDebugInfo(DebugInfo, BlueprintObj, BlueprintInstance, Pin);
+
+				if (WatchStatus != FKismetDebugUtilities::EWTR_Valid)
+				{
+					const UEdGraphSchema* Schema = Pin->GetOwningNode()->GetSchema();
+					DebugInfo.DisplayName = Schema->GetPinDisplayName(Pin);
+					DebugInfo.Type = UEdGraphSchema_K2::TypeToText(Pin->PinType);
+
+					switch (WatchStatus)
+					{
+					case FKismetDebugUtilities::EWTR_NotInScope:
+						DebugInfo.Value = LOCTEXT("NotInScope", "(not in scope)");
+						break;
+
+					case FKismetDebugUtilities::EWTR_NoProperty:
+						DebugInfo.Value = LOCTEXT("NoDebugData", "(no debug data)");
+						break;
+
+					case FKismetDebugUtilities::EWTR_NoDebugObject:
+						DebugInfo.Value = LOCTEXT("NoDebugObject", "(no debug object)");
+						break;
+
+					default:
+						// do nothing
+						break;
+					}
+				}
+
+				Private_InstanceWatchSource.Add(
+					MakeShared<FWatchRow>(
+						BlueprintObj,
+						Pin->GetOwningNode(),
+						Pin,
+						BlueprintInstance,
+						BlueprintName,
+						GraphName,
+						NodeName,
+						DebugInfo
+						)
+				);
+			}
 		}
 
-		if (Frame->Object)
-		{
-			BlueprintStackInstances.AddUnique(Frame->Object);
-		}
+		// Notify subscribers:
+		WatchListSubscribers.Broadcast(&Private_InstanceWatchSource);
 	}
-
-	UpdateInstancedWatchDisplay();
 }
 
 void WatchViewer::ContinueExecution()
@@ -907,7 +889,7 @@
 
 	if (IsPaused())
 	{
-		UpdateInstancedWatchDisplay();
+		WatchViewer::UpdateInstancedWatchDisplay();
 	}
 
 	// Notify subscribers
@@ -951,7 +933,7 @@
 
 			if (IsPaused())
 			{
-				UpdateInstancedWatchDisplay();
+				WatchViewer::UpdateInstancedWatchDisplay();
 			}
 
 			// Notify subscribers if necessary
diff --git a/Engine/Source/Editor/Kismet/Public/WatchPointViewer.h b/Engine/Source/Editor/Kismet/Public/WatchPointViewer.h
--- a/Engine/Source/Editor/Kismet/Public/WatchPointViewer.h
+++ b/Engine/Source/Editor/Kismet/Public/WatchPointViewer.h
@@ -10,8 +10,8 @@ class UBlueprint;
 
 namespace WatchViewer
 {
-	// called when we pause execution to update the displayed watch values
-	void KISMET_API UpdateDisplayedWatches(const TArray<const FFrame*>& ScriptStack);
+	// updates the instanced watch values, these are only valid while execution is paused
+	void KISMET_API UpdateInstancedWatchDisplay();
 
 	// called when we unpause execution and set watch values back to the blueprint versions
 	void KISMET_API ContinueExecution();
diff --git a/Engine/Source/Editor/UnrealEd/Private/Kismet2/KismetDebugUtilities.cpp b/Engine/Source/Editor/UnrealEd/Private/Kismet2/KismetDebugUtilities.cpp
--- a/Engine/Source/Editor/UnrealEd/Private/Kismet2/KismetDebugUtilities.cpp
+++ b/Engine/Source/Editor/UnrealEd/Private/Kismet2/KismetDebugUtilities.cpp
@@ -696,7 +696,7 @@ void FKismetDebugUtilities::AttemptToBreakExecution(UBlueprint* BlueprintObj, co
 		Data.LastExceptionMessage = Info.GetDescription();
 		FKismetEditorUtilities::BringKismetToFocusAttentionOnObject(NodeStoppedAt);
 		CallStackViewer::UpdateDisplayedCallstack(ScriptStack);
-		WatchViewer::UpdateDisplayedWatches(ScriptStack);
+		WatchViewer::UpdateInstancedWatchDisplay();
 		FSlateApplication::Get().EnterDebuggingMode();
 	}
 }
