Use weak object pointers instead of raw pointers in the watchpoint viewer.

diff --git a/Engine/Source/Editor/Kismet/Private/WatchPointViewer.cpp b/Engine/Source/Editor/Kismet/Private/WatchPointViewer.cpp
--- a/Engine/Source/Editor/Kismet/Private/WatchPointViewer.cpp
+++ b/Engine/Source/Editor/Kismet/Private/WatchPointViewer.cpp
@@ -40,7 +40,7 @@ namespace
 	struct FWatchRow
 	{
 		FWatchRow(
-			UBlueprint* InBP,
+			TWeakObjectPtr<UBlueprint> InBP,
 			const UEdGraphNode* InNode,
 			const UEdGraphPin* InPin,
 			UObject* InObjectBeingDebugged,
@@ -64,12 +64,12 @@ namespace
 		{
 			SetObjectBeingDebuggedName();
 
-			UPackage* Package = Cast<UPackage>(BP ? BP->GetOuter() : nullptr);
+			UPackage* Package = Cast<UPackage>(BP.IsValid() ? BP->GetOuter() : nullptr);
 			BlueprintPackageName = Package ? Package->GetFName() : FName();
 		}
 
 		FWatchRow(
-			UBlueprint* InBP,
+			TWeakObjectPtr<UBlueprint> InBP,
 			const UEdGraphNode* InNode,
 			const UEdGraphPin* InPin,
 			UObject* InObjectBeingDebugged,
@@ -91,7 +91,7 @@ namespace
 		{
 			SetObjectBeingDebuggedName();
 
-			UPackage* Package = Cast<UPackage>(BP ? BP->GetOuter() : nullptr);
+			UPackage* Package = Cast<UPackage>(BP.IsValid() ? BP->GetOuter() : nullptr);
 			BlueprintPackageName = Package ? Package->GetFName() : FName();
 
 			for (FDebugInfo& ChildInfo : Info.Children)
@@ -101,7 +101,7 @@ namespace
 		}
 
 		// this can't be const because we store watches in the blueprint
-		UBlueprint* BP;
+		TWeakObjectPtr<UBlueprint> BP;
 		const UEdGraphNode* Node;
 		const UEdGraphPin* Pin;
 		// this can't be const because SelectActor takes a non-const actor
@@ -161,7 +161,7 @@ namespace
 	TArray<TSharedRef<FWatchRow>> Private_WatchSource;
 	TArray<TSharedRef<FWatchRow>> Private_InstanceWatchSource;
 
-	TArray<UBlueprint*> WatchedBlueprints;
+	TArray<TWeakObjectPtr<UBlueprint>> WatchedBlueprints;
 
 	// Returns true if the blueprint execution is currently paused; false otherwise
 	bool IsPaused()
@@ -173,8 +173,12 @@ namespace
 	{
 		Private_WatchSource.Reset();
 
-		for (UBlueprint* BlueprintObj : WatchedBlueprints)
+		for (TWeakObjectPtr<UBlueprint> BlueprintObj : WatchedBlueprints)
 		{
+			if (!BlueprintObj.IsValid())
+			{
+				continue;
+			}
 			FText BlueprintName = FText::FromString(BlueprintObj->GetName());
 
 			for (const FEdGraphPinReference& PinRef : BlueprintObj->WatchedPins)
@@ -208,9 +212,9 @@ namespace
 		}
 	}
 
-	void UpdateWatchListFromBlueprintImpl(UBlueprint* BlueprintObj, const bool bShouldWatch)
+	void UpdateWatchListFromBlueprintImpl(TWeakObjectPtr<UBlueprint> BlueprintObj, const bool bShouldWatch)
 	{
-		if (!ensure(BlueprintObj))
+		if (!ensure(BlueprintObj.IsValid()))
 		{
 			return;
 		}
@@ -248,7 +252,7 @@ namespace
 	// Updates all of the watches from the currently watched blueprints
 	void UpdateAllBlueprintWatches()
 	{
-		for (UBlueprint* Blueprint : WatchedBlueprints)
+		for (TWeakObjectPtr<UBlueprint> Blueprint : WatchedBlueprints)
 		{
 			UpdateWatchListFromBlueprintImpl(Blueprint, true);
 		}
@@ -630,7 +634,7 @@ void SWatchViewer::StopWatchingPin() const
 	TArray<TSharedRef<FWatchRow>> SelectedRows = WatchTreeWidget->GetSelectedItems();
 	for (TSharedRef<FWatchRow>& Row : SelectedRows)
 	{
-		FKismetDebugUtilities::TogglePinWatch(Row->BP, Row->Pin);
+		FKismetDebugUtilities::TogglePinWatch(Row->BP.Get(), Row->Pin);
 	}
 }
 
@@ -868,9 +872,9 @@ FName WatchViewer::GetTabName()
 	return TabName;
 }
 
-void WatchViewer::RemoveWatchesForBlueprint(class UBlueprint* BlueprintObj)
+void WatchViewer::RemoveWatchesForBlueprint(TWeakObjectPtr<UBlueprint> BlueprintObj)
 {
-	if (!ensure(BlueprintObj))
+	if (!ensure(BlueprintObj.IsValid()))
 	{
 		return;
 	}
@@ -942,12 +946,12 @@ void WatchViewer::OnRenameAsset(const struct FAssetData& AssetData, const FStrin
 	}
 }
 
-void WatchViewer::UpdateWatchListFromBlueprint(UBlueprint* BlueprintObj)
+void WatchViewer::UpdateWatchListFromBlueprint(TWeakObjectPtr<UBlueprint> BlueprintObj)
 {
 	UpdateWatchListFromBlueprintImpl(BlueprintObj, true);
 }
 
-void WatchViewer::ClearWatchListFromBlueprint(UBlueprint* BlueprintObj)
+void WatchViewer::ClearWatchListFromBlueprint(TWeakObjectPtr<UBlueprint> BlueprintObj)
 {
 	UpdateWatchListFromBlueprintImpl(BlueprintObj, false);
 }
diff --git a/Engine/Source/Editor/Kismet/Public/WatchPointViewer.h b/Engine/Source/Editor/Kismet/Public/WatchPointViewer.h
--- a/Engine/Source/Editor/Kismet/Public/WatchPointViewer.h
+++ b/Engine/Source/Editor/Kismet/Public/WatchPointViewer.h
@@ -17,11 +17,11 @@ namespace WatchViewer
 	void KISMET_API ContinueExecution();
 
 	// called when we are adding or changing watches from BlueprintObj
-	void KISMET_API UpdateWatchListFromBlueprint(UBlueprint* BlueprintObj);
+	void KISMET_API UpdateWatchListFromBlueprint(TWeakObjectPtr<UBlueprint> BlueprintObj);
 
 	// called when we want to remove watches in the watch window from a blueprint
 	// does NOT remove watches from the pins in the blueprint object
-	void KISMET_API RemoveWatchesForBlueprint(UBlueprint* BlueprintObj);
+	void KISMET_API RemoveWatchesForBlueprint(TWeakObjectPtr<UBlueprint> BlueprintObj);
 
 	// called when we want to remove watches in the watch window from a blueprint
 	// does NOT remove watches from the pins in the blueprint object
@@ -31,7 +31,7 @@ namespace WatchViewer
 	void KISMET_API OnRenameAsset(const FAssetData& AssetData, const FString& OldAssetName);
 
 	// called when a BlueprintObj should no longer be watched
-	void KISMET_API ClearWatchListFromBlueprint(class UBlueprint* BlueprintObj);
+	void KISMET_API ClearWatchListFromBlueprint(TWeakObjectPtr<UBlueprint> BlueprintObj);
 
 	FName GetTabName();
 	void RegisterTabSpawner(FTabManager& TabManager);
