Change how we handle removing watches in the watchpoint viewer to correctly handle weak pointers that are invalid (marked pending kill, etc).

diff --git a/Engine/Source/Editor/Kismet/Private/WatchPointViewer.cpp b/Engine/Source/Editor/Kismet/Private/WatchPointViewer.cpp
--- a/Engine/Source/Editor/Kismet/Private/WatchPointViewer.cpp
+++ b/Engine/Source/Editor/Kismet/Private/WatchPointViewer.cpp
@@ -214,10 +214,6 @@ namespace
 
 	void UpdateWatchListFromBlueprintImpl(TWeakObjectPtr<UBlueprint> BlueprintObj, const bool bShouldWatch)
 	{
-		if (!ensure(BlueprintObj.IsValid()))
-		{
-			return;
-		}
 
 		if (bShouldWatch)
 		{
@@ -230,11 +226,29 @@ namespace
 			int32 FoundIdx = WatchedBlueprints.Find(BlueprintObj);
 			if (FoundIdx == INDEX_NONE)
 			{
-				return;
+				// if we didn't find it, it could be because BlueprintObj is no longer valid
+				// in this case the pointer in WatchedBlueprints would also be invalid
+				bool bRemovedBP = false;
+				for (int32 Idx = 0; Idx < WatchedBlueprints.Num(); ++Idx)
+				{
+					if (!WatchedBlueprints[Idx].IsValid())
+					{
+						bRemovedBP = true;
+						WatchedBlueprints.RemoveAt(Idx);
+						--Idx;
+					}
+				}
+				
+				if (!bRemovedBP)
+				{
+					return;
+				}
+			}
+			else
+			{
+				// since we're not watching the blueprint anymore we should remove it from the watched list
+				WatchedBlueprints.RemoveAt(FoundIdx);
 			}
-
-			// since we're not watching the blueprint anymore we should remove it from the watched list
-			WatchedBlueprints.RemoveAt(FoundIdx);
 		}
 
 		// something changed so we need to update the lists shown in the UI
