Add support for regenerating skeleton of individual blueprints to compilation manager

diff --git a/Engine/Source/Editor/Kismet/Private/BlueprintCompilationManager.cpp b/Engine/Source/Editor/Kismet/Private/BlueprintCompilationManager.cpp
--- a/Engine/Source/Editor/Kismet/Private/BlueprintCompilationManager.cpp
+++ b/Engine/Source/Editor/Kismet/Private/BlueprintCompilationManager.cpp
@@ -166,11 +166,17 @@ void FBlueprintCompilationManagerImpl::CompileSynchronouslyImpl(const FBPCompile
 	Request.BPToCompile->bQueuedForCompilation = true;
 
 	const bool bIsRegeneratingOnLoad		= (Request.CompileOptions & EBlueprintCompileOptions::IsRegeneratingOnLoad		) != EBlueprintCompileOptions::None;
-	const bool bSkipGarbageCollection		= (Request.CompileOptions & EBlueprintCompileOptions::SkipGarbageCollection		) != EBlueprintCompileOptions::None;
+	const bool bRegenerateSkeletonOnly		= (Request.CompileOptions & EBlueprintCompileOptions::RegenerateSkeletonOnly	) != EBlueprintCompileOptions::None;
+	const bool bSkipGarbageCollection		= (Request.CompileOptions & EBlueprintCompileOptions::SkipGarbageCollection		) != EBlueprintCompileOptions::None
+		|| bRegenerateSkeletonOnly;
 	const bool bBatchCompile				= (Request.CompileOptions & EBlueprintCompileOptions::BatchCompile				) != EBlueprintCompileOptions::None;
 	const bool bSkipReinstancing			= (Request.CompileOptions & EBlueprintCompileOptions::SkipReinstancing			) != EBlueprintCompileOptions::None;
 	const bool bSkipSaving					= (Request.CompileOptions & EBlueprintCompileOptions::SkipSave					) != EBlueprintCompileOptions::None;
 
+	ensure(!bIsRegeneratingOnLoad); // unexpected code path, compile on load handled with different function call
+	ensure(!bSkipReinstancing); // This is an internal option, should not go through CompileSynchronouslyImpl
+	ensure(QueuedRequests.Num() == 0);
+
 	// Wipe the PreCompile log, any generated messages are now irrelevant
 	Request.BPToCompile->PreCompileLog.Reset();
 
@@ -178,19 +184,14 @@ void FBlueprintCompilationManagerImpl::CompileSynchronouslyImpl(const FBPCompile
 	Request.BPToCompile->bDisplayCompilePIEWarning = true;
 	
 	// Do not want to run this code without the editor present nor when running commandlets.
-	if (GEditor && GIsEditor)
+	// We do not want to regenerate a search Guid during loads, nothing has changed in the Blueprint
+	// and it is cached elsewhere.
+	// We would like to regenerated it when a skeleton changes, but it is too expensive:
+	if (GEditor && GIsEditor && !bIsRegeneratingOnLoad && !bRegenerateSkeletonOnly)
 	{
-		// We do not want to regenerate a search Guid during loads, nothing has changed in the Blueprint and it is cached elsewhere
-		if (!bIsRegeneratingOnLoad)
-		{
-			FFindInBlueprintSearchManager::Get().AddOrUpdateBlueprintSearchMetadata(Request.BPToCompile);
-		}
+		FFindInBlueprintSearchManager::Get().AddOrUpdateBlueprintSearchMetadata(Request.BPToCompile);
 	}
 
-	ensure(!bIsRegeneratingOnLoad); // unexpected code path, compile on load handled with different function call
-	ensure(!bSkipReinstancing); // This is an internal option, should not go through CompileSynchronouslyImpl
-
-	ensure(QueuedRequests.Num() == 0);
 	QueuedRequests.Add(Request);
 	// We suppress normal compilation broadcasts because the old code path 
 	// did this after GC and we want to match the old behavior:
@@ -200,7 +201,7 @@ void FBlueprintCompilationManagerImpl::CompileSynchronouslyImpl(const FBPCompile
 	FlushCompilationQueueImpl(nullptr, bSuppressBroadcastCompiled, &CompiledBlueprints, &SkeletonCompiledBlueprints);
 	FlushReinstancingQueueImpl();
 	
-	if (FBlueprintEditorUtils::IsLevelScriptBlueprint(Request.BPToCompile))
+	if (FBlueprintEditorUtils::IsLevelScriptBlueprint(Request.BPToCompile) && !bRegenerateSkeletonOnly)
 	{
 		// When the Blueprint is recompiled, then update the bound events for level scripting
 		ULevelScriptBlueprint* LevelScriptBP = CastChecked<ULevelScriptBlueprint>(Request.BPToCompile);
@@ -213,7 +214,7 @@ void FBlueprintCompilationManagerImpl::CompileSynchronouslyImpl(const FBPCompile
 		}
 	}
 
-	if ( GEditor )
+	if ( GEditor && !bRegenerateSkeletonOnly)
 	{
 		GEditor->BroadcastBlueprintReinstanced();
 	}
@@ -224,13 +225,20 @@ void FBlueprintCompilationManagerImpl::CompileSynchronouslyImpl(const FBPCompile
 	{
 		CollectGarbage(GARBAGE_COLLECTION_KEEPFLAGS);
 	}
-	
-	for(UBlueprint* BP : SkeletonCompiledBlueprints)
+
+	if (!bRegenerateSkeletonOnly)
 	{
-		BP->BroadcastChanged();
+		for(UBlueprint* BP : SkeletonCompiledBlueprints)
+		{
+			BP->BroadcastChanged();
+		}
+	}
+	else
+	{
+		ensure(SkeletonCompiledBlueprints.Num() == 1);
 	}
 
-	if (!bBatchCompile)
+	if (!bBatchCompile && !bRegenerateSkeletonOnly)
 	{
 		for(UBlueprint* BP : SkeletonCompiledBlueprints)
 		{
@@ -243,7 +251,7 @@ void FBlueprintCompilationManagerImpl::CompileSynchronouslyImpl(const FBPCompile
 		}
 	}
 
-	if (CompiledBlueprintsToSave.Num() > 0)
+	if (CompiledBlueprintsToSave.Num() > 0 && !bRegenerateSkeletonOnly)
 	{
 		if (!bSkipSaving)
 		{
@@ -303,17 +311,16 @@ struct FCompilerData
 	}
 
 	bool IsSkeletonOnly() const { return JobType == ECompilationManagerJobType::SkeletonOnly; }
-	bool ShouldResetClassMembers() const { return JobType != ECompilationManagerJobType::RelinkOnly; }
 	bool ShouldSetTemporaryBlueprintFlags() const { return JobType != ECompilationManagerJobType::RelinkOnly; }
 	bool ShouldValidateVariableNames() const { return JobType == ECompilationManagerJobType::Normal; }
 	bool ShouldRegenerateSkeleton() const { return JobType != ECompilationManagerJobType::RelinkOnly; }
 	bool ShouldMarkUpToDateAfterSkeletonStage() const { return IsSkeletonOnly(); }
 	bool ShouldReconstructNodes() const { return JobType == ECompilationManagerJobType::Normal; }
 	bool ShouldSkipReinstancerCreation() const { return (IsSkeletonOnly() && BP->ParentClass->IsNative()); }
+	bool ShouldInitiateReinstancing() const { return JobType == ECompilationManagerJobType::Normal || BP->bIsRegeneratingOnLoad; }
 	bool ShouldCompileClassLayout() const { return JobType == ECompilationManagerJobType::Normal; }
 	bool ShouldCompileClassFunctions() const { return JobType == ECompilationManagerJobType::Normal; }
 	bool ShouldRegisterCompilerResults() const { return JobType == ECompilationManagerJobType::Normal; }
-	bool ShouldRelinkAfterSkippingCompile() const { return JobType == ECompilationManagerJobType::RelinkOnly; }
 	bool ShouldSkipIfDependenciesAreUnchanged() const { return InternalOptions.CompileType == EKismetCompileType::BytecodeOnly || JobType == ECompilationManagerJobType::RelinkOnly; }
 
 	UBlueprint* BP;
@@ -349,14 +356,19 @@ void FBlueprintCompilationManagerImpl::FlushCompilationQueueImpl(TArray<UObject*
 
 		// STAGE I: Add any related blueprints that were not compiled, then add any children so that they will be relinked:
 		TArray<UBlueprint*> BlueprintsToRecompile;
-		for(const FBPCompileRequest& ComplieJob : QueuedRequests)
+		for(const FBPCompileRequest& CompileJob : QueuedRequests)
 		{
+			if ((CompileJob.CompileOptions & EBlueprintCompileOptions::RegenerateSkeletonOnly) != EBlueprintCompileOptions::None)
+			{
+				continue;
+			}
+
 			// Add any dependent blueprints for a bytecode compile, this is needed because we 
 			// have no way to keep bytecode safe when a function is renamed or parameters are
 			// added or removed. Below (Stage VIII) we skip further compilation for blueprints 
 			// that are being bytecode compiled, but their dependencies have not changed:
 			TArray<UBlueprint*> DependentBlueprints;
-			FBlueprintEditorUtils::GetDependentBlueprints(ComplieJob.BPToCompile, DependentBlueprints);
+			FBlueprintEditorUtils::GetDependentBlueprints(CompileJob.BPToCompile, DependentBlueprints);
 			for(UBlueprint* DependentBlueprint : DependentBlueprints)
 			{
 				if(!IsQueuedForCompilation(DependentBlueprint))
@@ -441,7 +453,13 @@ void FBlueprintCompilationManagerImpl::FlushCompilationQueueImpl(TArray<UObject*
 			}
 			else
 			{
-				CurrentlyCompilingBPs.Add(FCompilerData(QueuedBP, ECompilationManagerJobType::Normal, QueuedJob.ClientResultsLog, QueuedJob.CompileOptions, false));
+				ECompilationManagerJobType JobType = ECompilationManagerJobType::Normal;
+				if ((QueuedJob.CompileOptions & EBlueprintCompileOptions::RegenerateSkeletonOnly) != EBlueprintCompileOptions::None)
+				{
+					JobType = ECompilationManagerJobType::SkeletonOnly;
+				}
+
+				CurrentlyCompilingBPs.Add(FCompilerData(QueuedBP, JobType, QueuedJob.ClientResultsLog, QueuedJob.CompileOptions, false));
 				BlueprintsToRecompile.Add(QueuedBP);
 			}
 		}
@@ -767,6 +785,32 @@ void FBlueprintCompilationManagerImpl::FlushCompilationQueueImpl(TArray<UObject*
 					continue;
 				}
 
+				// no need to reinstance skeleton or relink jobs that are not in a hierarchy that has had reinstancing initiated:
+				bool bRequiresReinstance = CompilerData.ShouldInitiateReinstancing();
+				if (!bRequiresReinstance)
+				{
+					UClass* Iter = CompilerData.BP->GeneratedClass;
+					if (!Iter)
+					{
+						bRequiresReinstance = true;
+					}
+					while (Iter)
+					{
+						if (Iter->HasAnyClassFlags(CLASS_NewerVersionExists))
+						{
+							bRequiresReinstance = true;
+							break;
+						}
+
+						Iter = Iter->GetSuperClass();
+					}
+				}
+
+				if (!bRequiresReinstance)
+				{
+					continue;
+				}
+
 				UBlueprint* BP = CompilerData.BP;
 
 				if(BP->GeneratedClass)
@@ -957,9 +1001,6 @@ void FBlueprintCompilationManagerImpl::FlushCompilationQueueImpl(TArray<UObject*
 				{
 					BP->SimpleConstructionScript->FixupRootNodeParentReferences();
 				}
-
-				TArray<UBlueprint*> DependentBPs;
-				FBlueprintEditorUtils::GetDependentBlueprints(BP, DependentBPs);
 				
 				UBlueprint::ValidateGeneratedClass(BP->GeneratedClass);
 			}
@@ -2166,7 +2207,7 @@ void FBlueprintCompilationManager::CompileSynchronously(const FBPCompileRequest&
 		BPCMImpl->CompileSynchronouslyImpl(Request);
 	}
 }
-	
+
 void FBlueprintCompilationManager::NotifyBlueprintLoaded(UBlueprint* BPLoaded)
 {
 	// Blueprints can be loaded before editor modules are on line:
diff --git a/Engine/Source/Editor/UnrealEd/Public/Kismet2/KismetEditorUtilities.h b/Engine/Source/Editor/UnrealEd/Public/Kismet2/KismetEditorUtilities.h
--- a/Engine/Source/Editor/UnrealEd/Public/Kismet2/KismetEditorUtilities.h
+++ b/Engine/Source/Editor/UnrealEd/Public/Kismet2/KismetEditorUtilities.h
@@ -46,6 +46,8 @@ enum class EBlueprintCompileOptions
 	SkipSave = 0x20,
 	/** Skips creating a reinstancer and running reinstancing routines - useful if calling code is performing reinstancing */
 	SkipReinstancing = 0x40,
+	/** Simply regenerates the skeleton class */
+	RegenerateSkeletonOnly = 0x80,
 };
 
 ENUM_CLASS_FLAGS(EBlueprintCompileOptions)
