Support for regenerating skeleton when class or CDO is missing (e.g. when missing parent type or new BP) to compilation manager

diff --git a/Engine/Source/Editor/Kismet/Private/BlueprintCompilationManager.cpp b/Engine/Source/Editor/Kismet/Private/BlueprintCompilationManager.cpp
--- a/Engine/Source/Editor/Kismet/Private/BlueprintCompilationManager.cpp
+++ b/Engine/Source/Editor/Kismet/Private/BlueprintCompilationManager.cpp
@@ -289,6 +289,7 @@ struct FCompilerData
 		JobType = InJobType;
 		UPackage* Package = BP->GetOutermost();
 		bPackageWasDirty = Package ? Package->IsDirty() : false;
+		OriginalBPStatus = BP->Status;
 
 		ActiveResultsLog = InResultsLogOverride;
 		if(InResultsLogOverride == nullptr)
@@ -316,7 +317,7 @@ struct FCompilerData
 	bool ShouldRegenerateSkeleton() const { return JobType != ECompilationManagerJobType::RelinkOnly; }
 	bool ShouldMarkUpToDateAfterSkeletonStage() const { return IsSkeletonOnly(); }
 	bool ShouldReconstructNodes() const { return JobType == ECompilationManagerJobType::Normal; }
-	bool ShouldSkipReinstancerCreation() const { return (IsSkeletonOnly() && BP->ParentClass->IsNative()); }
+	bool ShouldSkipReinstancerCreation() const { return (IsSkeletonOnly() && (!BP->ParentClass || BP->ParentClass->IsNative())); }
 	bool ShouldInitiateReinstancing() const { return JobType == ECompilationManagerJobType::Normal || BP->bIsRegeneratingOnLoad; }
 	bool ShouldCompileClassLayout() const { return JobType == ECompilationManagerJobType::Normal; }
 	bool ShouldCompileClassFunctions() const { return JobType == ECompilationManagerJobType::Normal; }
@@ -332,6 +333,7 @@ struct FCompilerData
 
 	ECompilationManagerJobType JobType;
 	bool bPackageWasDirty;
+	EBlueprintStatus OriginalBPStatus;
 };
 
 struct FReinstancingJob
@@ -398,7 +400,9 @@ void FBlueprintCompilationManagerImpl::FlushCompilationQueueImpl(TArray<UObject*
 			FBPCompileRequest& QueuedJob = QueuedRequests[I];
 			UBlueprint* QueuedBP = QueuedJob.BPToCompile;
 
-			ensure(!QueuedBP->GeneratedClass || !(QueuedBP->GeneratedClass->ClassDefaultObject->HasAnyFlags(RF_NeedLoad)));
+			ensure(!QueuedBP->GeneratedClass ||
+				!QueuedBP->GeneratedClass->ClassDefaultObject ||
+				!(QueuedBP->GeneratedClass->ClassDefaultObject->HasAnyFlags(RF_NeedLoad)));
 			bool bDefaultComponentMustBeAdded = false;
 			bool bHasPendingUberGraphFrame = false;
 
@@ -642,9 +646,12 @@ void FBlueprintCompilationManagerImpl::FlushCompilationQueueImpl(TArray<UObject*
 				if(CompilerData.ShouldMarkUpToDateAfterSkeletonStage())
 				{
 					// Flag data only blueprints as being up-to-date
-					BP->Status = BS_UpToDate;
+					BP->Status = BP->bHasBeenRegenerated ? CompilerData.OriginalBPStatus : BS_UpToDate;
 					BP->bHasBeenRegenerated = true;
-					BP->GeneratedClass->ClearFunctionMapsCaches();
+					if (BP->GeneratedClass)
+					{
+						BP->GeneratedClass->ClearFunctionMapsCaches();
+					}
 				}
 			}
 
@@ -817,12 +824,16 @@ void FBlueprintCompilationManagerImpl::FlushCompilationQueueImpl(TArray<UObject*
 				{
 					OldCDOs.Add(BP, BP->GeneratedClass->ClassDefaultObject);
 				}
-				CompilerData.Reinstancer = TSharedPtr<FBlueprintCompileReinstancer>( 
-					new FBlueprintCompileReinstancer(
-						BP->GeneratedClass, 
-						EBlueprintCompileReinstancerFlags::AutoInferSaveOnCompile|EBlueprintCompileReinstancerFlags::AvoidCDODuplication
-					)
-				);
+
+				if (BP->GeneratedClass && BP->GeneratedClass->ClassDefaultObject)
+				{
+					CompilerData.Reinstancer = TSharedPtr<FBlueprintCompileReinstancer>(
+						new FBlueprintCompileReinstancer(
+							BP->GeneratedClass,
+							EBlueprintCompileReinstancerFlags::AutoInferSaveOnCompile | EBlueprintCompileReinstancerFlags::AvoidCDODuplication
+						)
+					);
+				}
 			}
 		}
 
