Compilation manager now used to generate class when initially creating a blueprint

diff --git a/Engine/Source/Editor/Kismet/Private/BlueprintCompilationManager.cpp b/Engine/Source/Editor/Kismet/Private/BlueprintCompilationManager.cpp
--- a/Engine/Source/Editor/Kismet/Private/BlueprintCompilationManager.cpp
+++ b/Engine/Source/Editor/Kismet/Private/BlueprintCompilationManager.cpp
@@ -417,7 +417,7 @@ void FBlueprintCompilationManagerImpl::FlushCompilationQueueImpl(TArray<UObject*
 				bHasPendingUberGraphFrame = BPGC->UberGraphFramePointerProperty || BPGC->UberGraphFunction;
 			}
 			
-			if( FBlueprintEditorUtils::IsDataOnlyBlueprint(QueuedBP) && !QueuedBP->bHasBeenRegenerated && !bDefaultComponentMustBeAdded && !bHasPendingUberGraphFrame)
+			if( FBlueprintEditorUtils::IsDataOnlyBlueprint(QueuedBP) && !QueuedBP->bHasBeenRegenerated && QueuedBP->GetLinker() && !bDefaultComponentMustBeAdded && !bHasPendingUberGraphFrame)
 			{
 				const UClass* ParentClass = QueuedBP->ParentClass;
 				if (ParentClass && ParentClass->HasAllClassFlags(CLASS_Native))
@@ -825,15 +825,12 @@ void FBlueprintCompilationManagerImpl::FlushCompilationQueueImpl(TArray<UObject*
 					OldCDOs.Add(BP, BP->GeneratedClass->ClassDefaultObject);
 				}
 
-				if (BP->GeneratedClass && BP->GeneratedClass->ClassDefaultObject)
-				{
-					CompilerData.Reinstancer = TSharedPtr<FBlueprintCompileReinstancer>(
-						new FBlueprintCompileReinstancer(
-							BP->GeneratedClass,
-							EBlueprintCompileReinstancerFlags::AutoInferSaveOnCompile | EBlueprintCompileReinstancerFlags::AvoidCDODuplication
-						)
-					);
-				}
+				CompilerData.Reinstancer = TSharedPtr<FBlueprintCompileReinstancer>(
+					new FBlueprintCompileReinstancer(
+						BP->GeneratedClass,
+						EBlueprintCompileReinstancerFlags::AutoInferSaveOnCompile | EBlueprintCompileReinstancerFlags::AvoidCDODuplication
+					)
+				);
 			}
 		}
 
@@ -1210,12 +1207,13 @@ void FBlueprintCompilationManagerImpl::ReinstanceBatch(TArray<FReinstancingJob>&
 	{
 		const TSharedPtr<FBlueprintCompileReinstancer>& CurrentReinstancer = ReinstancingJob.Reinstancer;
 		UClass* OldClass = CurrentReinstancer->DuplicatedClass;
-		InOutOldToNewClassMap.Add(CurrentReinstancer->DuplicatedClass, CurrentReinstancer->ClassToReinstance);
 		if(!OldClass)
 		{
 			continue;
 		}
 
+		InOutOldToNewClassMap.Add(OldClass, CurrentReinstancer->ClassToReinstance);
+
 		if(!HasChildren(OldClass))
 		{
 			continue;
@@ -1294,7 +1292,10 @@ void FBlueprintCompilationManagerImpl::ReinstanceBatch(TArray<FReinstancingJob>&
 	for (const FReinstancingJob& ReinstancingJob : Reinstancers)
 	{
 		const TSharedPtr<FBlueprintCompileReinstancer>& CurrentReinstancer = ReinstancingJob.Reinstancer;
-		InOutOldToNewClassMap.Add( CurrentReinstancer->DuplicatedClass, CurrentReinstancer->ClassToReinstance );
+		if (CurrentReinstancer->DuplicatedClass)
+		{
+			InOutOldToNewClassMap.Add(CurrentReinstancer->DuplicatedClass, CurrentReinstancer->ClassToReinstance);
+		}
 			
 		UBlueprint* CompiledBlueprint = UBlueprint::GetBlueprintFromClass(CurrentReinstancer->ClassToReinstance);
 		CurrentReinstancer->UpdateBytecodeReferences();
@@ -1344,44 +1345,48 @@ void FBlueprintCompilationManagerImpl::ReinstanceBatch(TArray<FReinstancingJob>&
 	for (const FReinstancingJob& ReinstancingJob : Reinstancers)
 	{
 		const TSharedPtr<FBlueprintCompileReinstancer>& CurrentReinstancer = ReinstancingJob.Reinstancer;
-		UObject* OldCDO = CurrentReinstancer->DuplicatedClass->ClassDefaultObject;
-		if(OldCDO)
+		UObject* OldCDO = nullptr;
+		if (CurrentReinstancer->DuplicatedClass)
 		{
-			UObject* NewCDO = CurrentReinstancer->ClassToReinstance->GetDefaultObject(true);
-			FBlueprintCompileReinstancer::CopyPropertiesForUnrelatedObjects(OldCDO, NewCDO, true);
-
-			if(ReinstancingJob.Compiler.IsValid())
+			OldCDO = CurrentReinstancer->DuplicatedClass->ClassDefaultObject;
+			if (OldCDO)
 			{
-				ReinstancingJob.Compiler->PropagateValuesToCDO(NewCDO, OldCDO);
+				UObject* NewCDO = CurrentReinstancer->ClassToReinstance->GetDefaultObject(true);
+				FBlueprintCompileReinstancer::CopyPropertiesForUnrelatedObjects(OldCDO, NewCDO, true);
+
+				if (ReinstancingJob.Compiler.IsValid())
+				{
+					ReinstancingJob.Compiler->PropagateValuesToCDO(NewCDO, OldCDO);
+				}
 			}
+		}
 
-			if (UBlueprintGeneratedClass* BPGClass = CastChecked<UBlueprintGeneratedClass>(CurrentReinstancer->ClassToReinstance))
-			{
-				BPGClass->UpdateCustomPropertyListForPostConstruction();
+		if (UBlueprintGeneratedClass* BPGClass = CastChecked<UBlueprintGeneratedClass>(CurrentReinstancer->ClassToReinstance))
+		{
+			BPGClass->UpdateCustomPropertyListForPostConstruction();
 
-				// patch new cdo into linker table:
-				if(ObjLoaded)
+			// patch new cdo into linker table:
+			if(OldCDO && ObjLoaded)
+			{
+				UBlueprint* CurrentBP = CastChecked<UBlueprint>(BPGClass->ClassGeneratedBy);
+				if(FLinkerLoad* CurrentLinker = CurrentBP->GetLinker())
 				{
-					UBlueprint* CurrentBP = CastChecked<UBlueprint>(BPGClass->ClassGeneratedBy);
-					if(FLinkerLoad* CurrentLinker = CurrentBP->GetLinker())
-					{
-						int32 OldCDOIndex = INDEX_NONE;
+					int32 OldCDOIndex = INDEX_NONE;
 
-						for (int32 i = 0; i < CurrentLinker->ExportMap.Num(); i++)
+					for (int32 i = 0; i < CurrentLinker->ExportMap.Num(); i++)
+					{
+						FObjectExport& ThisExport = CurrentLinker->ExportMap[i];
+						if (ThisExport.ObjectFlags & RF_ClassDefaultObject)
 						{
-							FObjectExport& ThisExport = CurrentLinker->ExportMap[i];
-							if (ThisExport.ObjectFlags & RF_ClassDefaultObject)
-							{
-								OldCDOIndex = i;
-								break;
-							}
+							OldCDOIndex = i;
+							break;
 						}
+					}
 
-						if(OldCDOIndex != INDEX_NONE)
-						{
-							FBlueprintEditorUtils::PatchNewCDOIntoLinker(CurrentBP->GeneratedClass->ClassDefaultObject, CurrentLinker, OldCDOIndex, *ObjLoaded);
-							FBlueprintEditorUtils::PatchCDOSubobjectsIntoExport(OldCDO, CurrentBP->GeneratedClass->ClassDefaultObject);
-						}
+					if(OldCDOIndex != INDEX_NONE)
+					{
+						FBlueprintEditorUtils::PatchNewCDOIntoLinker(CurrentBP->GeneratedClass->ClassDefaultObject, CurrentLinker, OldCDOIndex, *ObjLoaded);
+						FBlueprintEditorUtils::PatchCDOSubobjectsIntoExport(OldCDO, CurrentBP->GeneratedClass->ClassDefaultObject);
 					}
 				}
 			}
@@ -1397,7 +1402,7 @@ void FBlueprintCompilationManagerImpl::ReinstanceBatch(TArray<FReinstancingJob>&
 	{
 		const TSharedPtr<FBlueprintCompileReinstancer>& CurrentReinstancer = ReinstancingJob.Reinstancer;
 		UClass* OldClass = CurrentReinstancer->DuplicatedClass;
-		if(ensure(OldClass))
+		if(OldClass)
 		{
 			TArray<UObject*> ArchetypeObjects;
 			GetObjectsOfClass(OldClass, ArchetypeObjects, false);
diff --git a/Engine/Source/Editor/UnrealEd/Private/Kismet2/Kismet2.cpp b/Engine/Source/Editor/UnrealEd/Private/Kismet2/Kismet2.cpp
--- a/Engine/Source/Editor/UnrealEd/Private/Kismet2/Kismet2.cpp
+++ b/Engine/Source/Editor/UnrealEd/Private/Kismet2/Kismet2.cpp
@@ -503,12 +503,21 @@ UBlueprint* FKismetEditorUtilities::CreateBlueprint(UClass* ParentClass, UObject
 
 	// Create initial UClass
 	IKismetCompilerInterface& Compiler = FModuleManager::LoadModuleChecked<IKismetCompilerInterface>(KISMET_COMPILER_MODULENAME);
-	
-	FCompilerResultsLog Results;
-	const bool bReplaceExistingInstances = false;
-	NewBP->Status = BS_Dirty;
-	FKismetCompilerOptions CompileOptions;
-	Compiler.CompileBlueprint(NewBP, CompileOptions, Results);
+
+	if (GBlueprintUseCompilationManager)
+	{
+		FBlueprintCompilationManager::CompileSynchronously(
+			FBPCompileRequest(NewBP, EBlueprintCompileOptions::SkipGarbageCollection, nullptr)
+		);
+	}
+	else
+	{
+		FCompilerResultsLog Results;
+		const bool bReplaceExistingInstances = false;
+		NewBP->Status = BS_Dirty;
+		FKismetCompilerOptions CompileOptions;
+		Compiler.CompileBlueprint(NewBP, CompileOptions, Results);
+	}
 
 	// Mark the BP as being regenerated, so it will not be confused as needing to be loaded and regenerated when a referenced BP loads.
 	NewBP->bHasBeenRegenerated = true;
@@ -1235,9 +1244,11 @@ void FKismetEditorUtilities::ConformBlueprintFlagsAndComponents(UBlueprint* Blue
 	if( UClass* GenClass = BlueprintObj->GeneratedClass )
 	{
 		GenClass->ClassFlags |= (ParentClass->ClassFlags & CLASS_ScriptInherit);
-		UObject* GenCDO = GenClass->GetDefaultObject();
-		ConformComponentsUtils::ConformRemovedNativeComponents(GenCDO);
-		GenCDO->InstanceSubobjectTemplates();
+		if (UObject* GenCDO = GenClass->ClassDefaultObject)
+		{
+			ConformComponentsUtils::ConformRemovedNativeComponents(GenCDO);
+			GenCDO->InstanceSubobjectTemplates();
+		}
 	}
 
 	UInheritableComponentHandler* InheritableComponentHandler = BlueprintObj->GetInheritableComponentHandler(false);
diff --git a/Engine/Source/Editor/UnrealEd/Private/Kismet2/KismetReinstanceUtilities.cpp b/Engine/Source/Editor/UnrealEd/Private/Kismet2/KismetReinstanceUtilities.cpp
--- a/Engine/Source/Editor/UnrealEd/Private/Kismet2/KismetReinstanceUtilities.cpp
+++ b/Engine/Source/Editor/UnrealEd/Private/Kismet2/KismetReinstanceUtilities.cpp
@@ -226,7 +226,7 @@ FBlueprintCompileReinstancer::FBlueprintCompileReinstancer(UClass* InClassToRein
 	, bIsRootReinstancer(false)
 	, bAllowResaveAtTheEndIfRequested(false)
 {
-	if( InClassToReinstance != nullptr )
+	if( InClassToReinstance != nullptr && InClassToReinstance->ClassDefaultObject )
 	{
 		bool bAutoInferSaveOnCompile = !!(Flags & EBlueprintCompileReinstancerFlags::AutoInferSaveOnCompile);
 		bool bIsBytecodeOnly = !!(Flags & EBlueprintCompileReinstancerFlags::BytecodeOnly);
