Prevent FastGenerateSkeletonClass from clearing DefaultValues field in UBlueprint::NewVariables
#rnx

diff --git a/Engine/Source/Editor/Kismet/Private/BlueprintCompilationManager.cpp b/Engine/Source/Editor/Kismet/Private/BlueprintCompilationManager.cpp
--- a/Engine/Source/Editor/Kismet/Private/BlueprintCompilationManager.cpp
+++ b/Engine/Source/Editor/Kismet/Private/BlueprintCompilationManager.cpp
@@ -1950,13 +1950,14 @@ UClass* FBlueprintCompilationManagerImpl::FastGenerateSkeletonClass(UBlueprint*
 		MakeEventFunction(Timeline->GetFinishedFunctionName(), EFunctionFlags::FUNC_None, TArray<UEdGraphPin*>(), TArray< TSharedPtr<FUserPinInfo> >(), nullptr, false);
 	}
 
-	CompilerContext.NewClass = Ret;
-	CompilerContext.bAssignDelegateSignatureFunction = true;
-	CompilerContext.bGenerateSubInstanceVariables = true;
-	CompilerContext.CreateClassVariablesFromBlueprint();
-	CompilerContext.bAssignDelegateSignatureFunction = false;
-	CompilerContext.bGenerateSubInstanceVariables = false;
-	CompilerContext.NewClass = OriginalNewClass;
+	{
+		CompilerContext.NewClass = Ret;
+		TGuardValue<bool> GuardAssignDelegateSignatureFunction( CompilerContext.bAssignDelegateSignatureFunction, true);
+		TGuardValue<bool> GuardGenerateSubInstanceVariables( CompilerContext.bGenerateSubInstanceVariables, true);
+		TGuardValue<EKismetCompileType::Type> GuardCompileType( CompilerContext.CompileOptions.CompileType, EKismetCompileType::SkeletonOnly);
+		CompilerContext.CreateClassVariablesFromBlueprint();
+		CompilerContext.NewClass = OriginalNewClass;
+	}
 	UField* Iter = Ret->Children;
 	while(Iter)
 	{
