Avoid crash when creating a levelblueprint
#rnx

diff --git a/Engine/Source/Editor/Kismet/Private/BlueprintCompilationManager.cpp b/Engine/Source/Editor/Kismet/Private/BlueprintCompilationManager.cpp
--- a/Engine/Source/Editor/Kismet/Private/BlueprintCompilationManager.cpp
+++ b/Engine/Source/Editor/Kismet/Private/BlueprintCompilationManager.cpp
@@ -210,7 +210,11 @@ void FBlueprintCompilationManagerImpl::CompileSynchronouslyImpl(const FBPCompile
 		// been updated as it respawns the actor:
 		if (ULevel* BPLevel = LevelScriptBP->GetLevel())
 		{
-			BPLevel->OnLevelScriptBlueprintChanged(LevelScriptBP);
+			// Newly created levels don't need this notification:
+			if (BPLevel->GetLevelScriptBlueprint(true))
+			{
+				BPLevel->OnLevelScriptBlueprintChanged(LevelScriptBP);
+			}
 		}
 	}
 
