Remove unnecessary ensure. This matches changes that were made in the 4.20 branch and fixes the DDC.

diff --git a/Engine/Source/Editor/Kismet/Private/WatchPointViewer.cpp b/Engine/Source/Editor/Kismet/Private/WatchPointViewer.cpp
--- a/Engine/Source/Editor/Kismet/Private/WatchPointViewer.cpp
+++ b/Engine/Source/Editor/Kismet/Private/WatchPointViewer.cpp
@@ -214,11 +214,6 @@ namespace
 
 	void UpdateWatchListFromBlueprintImpl(TWeakObjectPtr<UBlueprint> BlueprintObj, const bool bShouldWatch)
 	{
-		if (!ensure(BlueprintObj.IsValid()))
-		{
-			return;
-		}
-
 		if (bShouldWatch)
 		{
 			// make sure the blueprint is in our list
