Mirror 4227401 from Release-4.20
Clear node error message on load at proper time when using the compilation manager
#rnx

diff --git a/Engine/Source/Editor/Kismet/Private/BlueprintCompilationManager.cpp b/Engine/Source/Editor/Kismet/Private/BlueprintCompilationManager.cpp
--- a/Engine/Source/Editor/Kismet/Private/BlueprintCompilationManager.cpp
+++ b/Engine/Source/Editor/Kismet/Private/BlueprintCompilationManager.cpp
@@ -317,6 +317,7 @@ struct FCompilerData
 
 	bool IsSkeletonOnly() const { return JobType == ECompilationManagerJobType::SkeletonOnly; }
 	bool ShouldSetTemporaryBlueprintFlags() const { return JobType != ECompilationManagerJobType::RelinkOnly; }
+	bool ShouldResetErrorState() const { return JobType == ECompilationManagerJobType::Normal && InternalOptions.CompileType != EKismetCompileType::BytecodeOnly; }
 	bool ShouldValidateVariableNames() const { return JobType == ECompilationManagerJobType::Normal; }
 	bool ShouldRegenerateSkeleton() const { return JobType != ECompilationManagerJobType::RelinkOnly; }
 	bool ShouldMarkUpToDateAfterSkeletonStage() const { return IsSkeletonOnly(); }
@@ -543,16 +544,31 @@ void FBlueprintCompilationManagerImpl::FlushCompilationQueueImpl(TArray<UObject*
 		// STAGE IV: Set UBlueprint flags (bBeingCompiled, bIsRegeneratingOnLoad)
 		for (FCompilerData& CompilerData : CurrentlyCompilingBPs)
 		{
-			if(CompilerData.ShouldSetTemporaryBlueprintFlags())
+			if (!CompilerData.ShouldSetTemporaryBlueprintFlags())
 			{
-				UBlueprint* BP = CompilerData.BP;
-				BP->bBeingCompiled = true;
-				BP->CurrentMessageLog = CompilerData.ActiveResultsLog;
-				BP->bIsRegeneratingOnLoad = !BP->bHasBeenRegenerated && BP->GetLinker();
-				if(BP->bIsRegeneratingOnLoad)
+				continue;
+			}
+
+			UBlueprint* BP = CompilerData.BP;
+			BP->bBeingCompiled = true;
+			BP->CurrentMessageLog = CompilerData.ActiveResultsLog;
+			BP->bIsRegeneratingOnLoad = !BP->bHasBeenRegenerated && BP->GetLinker();
+			if(BP->bIsRegeneratingOnLoad)
+			{
+				// we may have cached dependencies before being fully loaded:
+				BP->bCachedDependenciesUpToDate = false;
+			}
+
+			if(CompilerData.ShouldResetErrorState())
+			{
+				TArray<UEdGraph*> AllGraphs;
+				BP->GetAllGraphs(AllGraphs);
+				for (UEdGraph* Graph : AllGraphs )
 				{
-					// we may have cached dependencies before being fully loaded:
-					BP->bCachedDependenciesUpToDate = false;
+					for (UEdGraphNode* GraphNode : Graph->Nodes)
+					{
+						GraphNode->ClearCompilerMessage();
+					}
 				}
 			}
 		}
diff --git a/Engine/Source/Editor/KismetCompiler/Private/KismetCompiler.cpp b/Engine/Source/Editor/KismetCompiler/Private/KismetCompiler.cpp
--- a/Engine/Source/Editor/KismetCompiler/Private/KismetCompiler.cpp
+++ b/Engine/Source/Editor/KismetCompiler/Private/KismetCompiler.cpp
@@ -64,6 +64,8 @@ FSimpleMulticastDelegate FKismetCompilerContext::OnPostCompile;
 
 #define LOCTEXT_NAMESPACE "KismetCompiler"
 
+extern COREUOBJECT_API bool GBlueprintUseCompilationManager;
+
 //////////////////////////////////////////////////////////////////////////
 // Stats for this module
 DECLARE_CYCLE_STAT(TEXT("Create Schema"), EKismetCompilerStats_CreateSchema, STATGROUP_KismetCompiler );
@@ -3815,7 +3817,7 @@ void FKismetCompilerContext::CompileClassLayout(EInternalCompilerFlags InternalF
 		}
 	}
 
-	if (CompileOptions.DoesRequireBytecodeGeneration() && !Blueprint->bIsRegeneratingOnLoad)
+	if (CompileOptions.DoesRequireBytecodeGeneration() && !GBlueprintUseCompilationManager)
 	{
 		TArray<UEdGraph*> AllGraphs;
 		Blueprint->GetAllGraphs(AllGraphs);
