Mirror 4241766 from Release-4.20
Fix for crash when loading blueprints with missing nodes and using the compilation manager

diff --git a/Engine/Source/Editor/Kismet/Private/BlueprintCompilationManager.cpp b/Engine/Source/Editor/Kismet/Private/BlueprintCompilationManager.cpp
--- a/Engine/Source/Editor/Kismet/Private/BlueprintCompilationManager.cpp
+++ b/Engine/Source/Editor/Kismet/Private/BlueprintCompilationManager.cpp
@@ -567,7 +567,10 @@ void FBlueprintCompilationManagerImpl::FlushCompilationQueueImpl(TArray<UObject*
 				{
 					for (UEdGraphNode* GraphNode : Graph->Nodes)
 					{
-						GraphNode->ClearCompilerMessage();
+						if (GraphNode)
+						{
+							GraphNode->ClearCompilerMessage();
+						}
 					}
 				}
 			}
