Fix broken FiB searches for BP assets by interface name.

Change summary:
- Modified FImaginaryFiBData::CreateSearchTree() to no longer exclude tag+value categories from the search results.
- Modified FFindInBlueprintSearchManager::ExtractUnloadedFiBData() to fix incorrect parsing of the Blueprint's 'ImplementedInterfaces' property value exported to an asset registry tag string.

diff --git a/Engine/Source/Editor/Kismet/Private/FindInBlueprintManager.cpp b/Engine/Source/Editor/Kismet/Private/FindInBlueprintManager.cpp
--- a/Engine/Source/Editor/Kismet/Private/FindInBlueprintManager.cpp
+++ b/Engine/Source/Editor/Kismet/Private/FindInBlueprintManager.cpp
@@ -1514,20 +1514,47 @@ void FFindInBlueprintSearchManager::ExtractUnloadedFiBData(const FAssetData& InA
 	const FString ImplementedInterfaces = InAssetData.GetTagValueRef<FString>(FBlueprintTags::ImplementedInterfaces);
 	if(!ImplementedInterfaces.IsEmpty())
 	{
-		FString FullInterface;
-		FString RemainingString;
-		FString InterfaceName;
-		FString CurrentString = ImplementedInterfaces;
-		while(CurrentString.Split(TEXT(","), &FullInterface, &RemainingString))
+		// ImplementedInterfaces is an array of structs (FBPInterfaceDescription). When exported to an AR tag value, each entry will be formatted as:
+		//
+		//	Entry := (Interface=Type'Package.Class') OR
+		//  Entry := (Interface=Type'Package.Class',Graphs=(Type'Package.Blueprint:Graph1',Type'Package.Blueprint:Graph2',...Type'Package.Blueprint:GraphN'))
+		//
+		// The full tag value (array of exported struct values) will then be formatted as follows:
+		//
+		//  Value := (Entry1,Entry2,...EntryN)
+		//
+		// Here we parse out the 'Interface' values, and add only the 'Name' part of the full interface path substrings into the new search data as keywords.
+
+		auto FindSubStringPosLambda = [&ImplementedInterfaces](const FString& InSubString, int32 StartPosition) -> int32
+		{
+			return ImplementedInterfaces.Find(InSubString, ESearchCase::CaseSensitive, ESearchDir::FromStart, StartPosition);
+		};
+
+		static const FString InterfaceFieldName = GET_MEMBER_NAME_STRING_CHECKED(FBPInterfaceDescription, Interface);
+
+		int32 CurPos = FindSubStringPosLambda(InterfaceFieldName, 0);
+		while (CurPos != INDEX_NONE)
 		{
-			if(FullInterface.Split(TEXT("."), &CurrentString, &InterfaceName, ESearchCase::CaseSensitive, ESearchDir::FromEnd))
+			CurPos = FindSubStringPosLambda(TEXT("="), CurPos);
+			if (CurPos != INDEX_NONE)
 			{
-				if(!CurrentString.StartsWith(TEXT("Graphs=(")))
+				CurPos = FindSubStringPosLambda(TEXT("."), CurPos);
+				if (CurPos != INDEX_NONE)
 				{
-					NewSearchData.Interfaces.Add(InterfaceName);
+					const int32 StartPos = CurPos + 1;
+					CurPos = FindSubStringPosLambda(TEXT("\'"), StartPos);
+					if (CurPos != INDEX_NONE)
+					{
+						const FString InterfaceName = ImplementedInterfaces.Mid(StartPos, CurPos - StartPos);
+						if (!InterfaceName.IsEmpty())
+						{
+							NewSearchData.Interfaces.Add(InterfaceName.TrimQuotes());
+						}
+
+						CurPos = FindSubStringPosLambda(InterfaceFieldName, CurPos + 1);
+					}
 				}
 			}
-			CurrentString = RemainingString;
 		}
 	}
 
diff --git a/Engine/Source/Editor/Kismet/Private/ImaginaryBlueprintData.cpp b/Engine/Source/Editor/Kismet/Private/ImaginaryBlueprintData.cpp
--- a/Engine/Source/Editor/Kismet/Private/ImaginaryBlueprintData.cpp
+++ b/Engine/Source/Editor/Kismet/Private/ImaginaryBlueprintData.cpp
@@ -108,8 +108,10 @@ FSearchResult FImaginaryFiBData::CreateSearchTree(FSearchResult InParentSearchRe
 			}
 		}
 
-		// If the children did not match the search results but this item does, then we will want to return true
-		if (!bValidSearchResults && !CurrentData->IsCategory() && (InValidSearchResults.Find(CurrentData) != INDEX_NONE || InMatchingSearchComponents.Find(CurrentData)))
+		// If the children did not match the search results but this item does, then we will want to return true.
+		// Include "tag+value" categories in the search tree, as the relevant results need to be added as children.
+		const bool bInvalidSearchResultsCategory = CurrentData->IsCategory() && !CurrentData->IsTagAndValueCategory();
+		if (!bValidSearchResults && !bInvalidSearchResultsCategory && (InValidSearchResults.Find(CurrentData) != INDEX_NONE || InMatchingSearchComponents.Find(CurrentData)))
 		{
 			bValidSearchResults = true;
 		}
