Generalize the logic that determines what entry points appear under graph in the my blueprints panel.
Make InputTouch show up as an entry point in the graph list.

diff --git a/Engine/Source/Editor/BlueprintGraph/Classes/K2Node_Event.h b/Engine/Source/Editor/BlueprintGraph/Classes/K2Node_Event.h
--- a/Engine/Source/Editor/BlueprintGraph/Classes/K2Node_Event.h
+++ b/Engine/Source/Editor/BlueprintGraph/Classes/K2Node_Event.h
@@ -10,6 +10,7 @@
 #include "Engine/MemberReference.h"
 #include "EdGraph/EdGraphNodeUtils.h"
 #include "K2Node_EditablePinBase.h"
+#include "K2Node_EventNodeInterface.h"
 #include "K2Node_Event.generated.h"
 
 class FKismetCompilerContext;
@@ -17,7 +18,7 @@ class FNodeHandlingFunctor;
 class UEdGraph;
 
 UCLASS(MinimalAPI)
-class UK2Node_Event : public UK2Node_EditablePinBase
+class UK2Node_Event : public UK2Node_EditablePinBase, public IK2Node_EventNodeInterface
 {
 	GENERATED_UCLASS_BODY()
 	BLUEPRINTGRAPH_API static const FName DelegateOutputName;
@@ -89,6 +90,10 @@ class UK2Node_Event : public UK2Node_EditablePinBase
 	BLUEPRINTGRAPH_API virtual bool HasExternalDependencies(TArray<class UStruct*>* OptionalOutput) const override;
 	//~ End UK2Node Interface
 
+	//~ Begin IK2Node_EventNodeInterface Interface.
+	virtual TSharedPtr<FEdGraphSchemaAction> GetEventNodeAction(const FText& ActionCategory) override;
+	//~ End IK2Node_EventNodeInterface Interface.
+
 	/** Checks whether the parameters for this event node are compatible with the specified function entry node */
 	BLUEPRINTGRAPH_API virtual bool IsFunctionEntryCompatible(const class UK2Node_FunctionEntry* EntryNode) const;
 	
diff --git a/Engine/Source/Editor/BlueprintGraph/Classes/K2Node_EventNodeInterface.h b/Engine/Source/Editor/BlueprintGraph/Classes/K2Node_EventNodeInterface.h
new file mode 100644
--- /dev/null
+++ b/Engine/Source/Editor/BlueprintGraph/Classes/K2Node_EventNodeInterface.h
@@ -0,0 +1,24 @@
+// Copyright 1998-2018 Epic Games, Inc. All Rights Reserved.
+
+
+#pragma once
+
+#include "CoreMinimal.h"
+#include "UObject/Interface.h"
+#include "K2Node_EventNodeInterface.generated.h"
+
+struct FEdGraphSchemaAction;
+
+UINTERFACE(meta=(CannotImplementInterfaceInBlueprint))
+class BLUEPRINTGRAPH_API UK2Node_EventNodeInterface : public UInterface
+{
+	GENERATED_BODY()
+};
+
+class BLUEPRINTGRAPH_API IK2Node_EventNodeInterface
+{
+	GENERATED_BODY()
+
+public:
+	virtual TSharedPtr<FEdGraphSchemaAction> GetEventNodeAction(const FText& ActionCategory) = 0;
+};
diff --git a/Engine/Source/Editor/BlueprintGraph/Classes/K2Node_InputAction.h b/Engine/Source/Editor/BlueprintGraph/Classes/K2Node_InputAction.h
--- a/Engine/Source/Editor/BlueprintGraph/Classes/K2Node_InputAction.h
+++ b/Engine/Source/Editor/BlueprintGraph/Classes/K2Node_InputAction.h
@@ -10,13 +10,14 @@
 #include "EdGraph/EdGraphNodeUtils.h"
 #include "BlueprintNodeSignature.h"
 #include "K2Node.h"
+#include "K2Node_EventNodeInterface.h"
 #include "K2Node_InputAction.generated.h"
 
 class FBlueprintActionDatabaseRegistrar;
 class FKismetCompilerContext;
 
 UCLASS()
-class BLUEPRINTGRAPH_API UK2Node_InputAction : public UK2Node
+class BLUEPRINTGRAPH_API UK2Node_InputAction : public UK2Node, public IK2Node_EventNodeInterface
 {
 	GENERATED_UCLASS_BODY()
 
@@ -57,6 +58,10 @@ class BLUEPRINTGRAPH_API UK2Node_InputAction : public UK2Node
 	virtual FBlueprintNodeSignature GetSignature() const override;
 	//~ End UK2Node Interface
 
+	//~ Begin IK2Node_EventNodeInterface Interface.
+	virtual TSharedPtr<FEdGraphSchemaAction> GetEventNodeAction(const FText& ActionCategory) override;
+	//~ End IK2Node_EventNodeInterface Interface.
+
 	/** Get the 'pressed' input pin */
 	UEdGraphPin* GetPressedPin() const;
 
diff --git a/Engine/Source/Editor/BlueprintGraph/Classes/K2Node_InputKey.h b/Engine/Source/Editor/BlueprintGraph/Classes/K2Node_InputKey.h
--- a/Engine/Source/Editor/BlueprintGraph/Classes/K2Node_InputKey.h
+++ b/Engine/Source/Editor/BlueprintGraph/Classes/K2Node_InputKey.h
@@ -11,6 +11,7 @@
 #include "EdGraph/EdGraphNodeUtils.h"
 #include "BlueprintNodeSignature.h"
 #include "K2Node.h"
+#include "K2Node_EventNodeInterface.h"
 
 #include "K2Node_InputKey.generated.h"
 
@@ -18,7 +19,7 @@ class FBlueprintActionDatabaseRegistrar;
 class FKismetCompilerContext;
 
 UCLASS()
-class BLUEPRINTGRAPH_API UK2Node_InputKey : public UK2Node
+class BLUEPRINTGRAPH_API UK2Node_InputKey : public UK2Node, public IK2Node_EventNodeInterface
 {
 	GENERATED_UCLASS_BODY()
 
@@ -77,6 +78,11 @@ class BLUEPRINTGRAPH_API UK2Node_InputKey : public UK2Node
 	virtual bool IsCompatibleWithGraph(UEdGraph const* Graph) const override;
 	//~ End UEdGraphNode Interface.
 
+	//~ Begin IK2Node_EventNodeInterface Interface.
+	virtual TSharedPtr<FEdGraphSchemaAction> GetEventNodeAction(const FText& ActionCategory) override;
+	//~ End IK2Node_EventNodeInterface Interface.
+
+
 	FText GetModifierText() const;
 
 	FName GetModifierName() const;
diff --git a/Engine/Source/Editor/BlueprintGraph/Classes/K2Node_InputTouch.h b/Engine/Source/Editor/BlueprintGraph/Classes/K2Node_InputTouch.h
--- a/Engine/Source/Editor/BlueprintGraph/Classes/K2Node_InputTouch.h
+++ b/Engine/Source/Editor/BlueprintGraph/Classes/K2Node_InputTouch.h
@@ -7,6 +7,7 @@
 #include "UObject/ObjectMacros.h"
 #include "Textures/SlateIcon.h"
 #include "K2Node.h"
+#include "K2Node_EventNodeInterface.h"
 #include "K2Node_InputTouch.generated.h"
 
 class FBlueprintActionDatabaseRegistrar;
@@ -15,7 +16,7 @@ class UEdGraph;
 class UEdGraphPin;
 
 UCLASS(MinimalAPI)
-class UK2Node_InputTouch : public UK2Node
+class UK2Node_InputTouch : public UK2Node, public IK2Node_EventNodeInterface
 {
 	GENERATED_UCLASS_BODY()
 
@@ -51,6 +52,10 @@ class UK2Node_InputTouch : public UK2Node
 	virtual bool IsCompatibleWithGraph(const UEdGraph* TargetGraph) const override;
 	//~ End UEdGraphNode Interface.
 
+	//~ Begin IK2Node_EventNodeInterface Interface.
+	virtual TSharedPtr<FEdGraphSchemaAction> GetEventNodeAction(const FText& ActionCategory) override;
+	//~ End IK2Node_EventNodeInterface Interface.
+
 	BLUEPRINTGRAPH_API static UEnum* GetTouchIndexEnum();
 
 	/** Get the 'pressed' input pin */
diff --git a/Engine/Source/Editor/BlueprintGraph/Classes/K2Node_InputTouchEvent.h b/Engine/Source/Editor/BlueprintGraph/Classes/K2Node_InputTouchEvent.h
--- a/Engine/Source/Editor/BlueprintGraph/Classes/K2Node_InputTouchEvent.h
+++ b/Engine/Source/Editor/BlueprintGraph/Classes/K2Node_InputTouchEvent.h
@@ -20,13 +20,13 @@ class UK2Node_InputTouchEvent : public UK2Node_Event
 	TEnumAsByte<EInputEvent> InputKeyEvent;
 
 	UPROPERTY()
-	uint32 bConsumeInput:1;
+	uint8 bConsumeInput:1;
 
 	UPROPERTY()
-	uint32 bExecuteWhenPaused:1;
+	uint8 bExecuteWhenPaused:1;
 
 	UPROPERTY()
-	uint32 bOverrideParentBinding:1;
+	uint8 bOverrideParentBinding:1;
 
 	//~ Begin UK2Node Interface
 	virtual UClass* GetDynamicBindingClass() const override;
diff --git a/Engine/Source/Editor/BlueprintGraph/Private/K2Node_Event.cpp b/Engine/Source/Editor/BlueprintGraph/Private/K2Node_Event.cpp
--- a/Engine/Source/Editor/BlueprintGraph/Private/K2Node_Event.cpp
+++ b/Engine/Source/Editor/BlueprintGraph/Private/K2Node_Event.cpp
@@ -7,6 +7,7 @@
 #include "GameFramework/Actor.h"
 #include "GraphEditorSettings.h"
 #include "EdGraphSchema_K2.h"
+#include "EdGraphSchema_K2_Actions.h"
 #include "K2Node_CallFunction.h"
 #include "K2Node_CreateDelegate.h"
 #include "K2Node_FunctionEntry.h"
@@ -859,4 +860,11 @@ bool UK2Node_Event::HasExternalDependencies(TArray<class UStruct*>* OptionalOutp
 	return bSuperResult || bResult;
 }
 
+TSharedPtr<FEdGraphSchemaAction> UK2Node_Event::GetEventNodeAction(const FText& ActionCategory)
+{
+	TSharedPtr<FEdGraphSchemaAction_K2Event> EventNodeAction = MakeShareable(new FEdGraphSchemaAction_K2Event(ActionCategory, GetNodeTitle(ENodeTitleType::EditableTitle), GetTooltipText(), 0));
+	EventNodeAction->NodeTemplate = this;
+	return EventNodeAction;
+}
+
 #undef LOCTEXT_NAMESPACE
diff --git a/Engine/Source/Editor/BlueprintGraph/Private/K2Node_InputAction.cpp b/Engine/Source/Editor/BlueprintGraph/Private/K2Node_InputAction.cpp
--- a/Engine/Source/Editor/BlueprintGraph/Private/K2Node_InputAction.cpp
+++ b/Engine/Source/Editor/BlueprintGraph/Private/K2Node_InputAction.cpp
@@ -5,6 +5,7 @@
 #include "GameFramework/InputSettings.h"
 #include "GraphEditorSettings.h"
 #include "EdGraphSchema_K2.h"
+#include "EdGraphSchema_K2_Actions.h"
 #include "K2Node_AssignmentStatement.h"
 #include "K2Node_TemporaryVariable.h"
 #include "Kismet2/BlueprintEditorUtils.h"
@@ -280,4 +281,11 @@ FBlueprintNodeSignature UK2Node_InputAction::GetSignature() const
 	return NodeSignature;
 }
 
+TSharedPtr<FEdGraphSchemaAction> UK2Node_InputAction::GetEventNodeAction(const FText& ActionCategory)
+{
+	TSharedPtr<FEdGraphSchemaAction_K2InputAction> EventNodeAction = MakeShareable(new FEdGraphSchemaAction_K2InputAction(ActionCategory, GetNodeTitle(ENodeTitleType::EditableTitle), GetTooltipText(), 0));
+	EventNodeAction->NodeTemplate = this;
+	return EventNodeAction;
+}
+
 #undef LOCTEXT_NAMESPACE
diff --git a/Engine/Source/Editor/BlueprintGraph/Private/K2Node_InputKey.cpp b/Engine/Source/Editor/BlueprintGraph/Private/K2Node_InputKey.cpp
--- a/Engine/Source/Editor/BlueprintGraph/Private/K2Node_InputKey.cpp
+++ b/Engine/Source/Editor/BlueprintGraph/Private/K2Node_InputKey.cpp
@@ -3,6 +3,7 @@
 #include "K2Node_InputKey.h"
 #include "GraphEditorSettings.h"
 #include "EdGraphSchema_K2.h"
+#include "EdGraphSchema_K2_Actions.h"
 #include "K2Node_AssignmentStatement.h"
 #include "K2Node_TemporaryVariable.h"
 #include "Kismet2/BlueprintEditorUtils.h"
@@ -421,4 +422,11 @@ FBlueprintNodeSignature UK2Node_InputKey::GetSignature() const
 	return NodeSignature;
 }
 
+TSharedPtr<FEdGraphSchemaAction> UK2Node_InputKey::GetEventNodeAction(const FText& ActionCategory)
+{
+	TSharedPtr<FEdGraphSchemaAction_K2InputAction> EventNodeAction = MakeShareable(new FEdGraphSchemaAction_K2InputAction(ActionCategory, GetNodeTitle(ENodeTitleType::EditableTitle), GetTooltipText(), 0));
+	EventNodeAction->NodeTemplate = this;
+	return EventNodeAction;
+}
+
 #undef LOCTEXT_NAMESPACE
diff --git a/Engine/Source/Editor/BlueprintGraph/Private/K2Node_InputTouch.cpp b/Engine/Source/Editor/BlueprintGraph/Private/K2Node_InputTouch.cpp
--- a/Engine/Source/Editor/BlueprintGraph/Private/K2Node_InputTouch.cpp
+++ b/Engine/Source/Editor/BlueprintGraph/Private/K2Node_InputTouch.cpp
@@ -3,6 +3,7 @@
 #include "K2Node_InputTouch.h"
 #include "GraphEditorSettings.h"
 #include "EdGraphSchema_K2.h"
+#include "EdGraphSchema_K2_Actions.h"
 #include "K2Node_AssignmentStatement.h"
 #include "K2Node_TemporaryVariable.h"
 #include "Kismet2/BlueprintEditorUtils.h"
@@ -107,6 +108,13 @@ FText UK2Node_InputTouch::GetMenuCategory() const
 	return FEditorCategoryUtils::GetCommonCategory(FCommonEditorCategory::Input);
 }
 
+TSharedPtr<FEdGraphSchemaAction> UK2Node_InputTouch::GetEventNodeAction(const FText& ActionCategory)
+{
+	TSharedPtr<FEdGraphSchemaAction_K2InputAction> EventNodeAction = MakeShareable(new FEdGraphSchemaAction_K2InputAction(ActionCategory, GetNodeTitle(ENodeTitleType::EditableTitle), GetTooltipText(), 0));
+	EventNodeAction->NodeTemplate = this;
+	return EventNodeAction;
+}
+
 bool UK2Node_InputTouch::IsCompatibleWithGraph(const UEdGraph* TargetGraph) const
 {
 	// This node expands into event nodes and must be placed in a Ubergraph
diff --git a/Engine/Source/Editor/Kismet/Private/SMyBlueprint.cpp b/Engine/Source/Editor/Kismet/Private/SMyBlueprint.cpp
--- a/Engine/Source/Editor/Kismet/Private/SMyBlueprint.cpp
+++ b/Engine/Source/Editor/Kismet/Private/SMyBlueprint.cpp
@@ -24,11 +24,9 @@
 #include "K2Node_CreateDelegate.h"
 #include "K2Node_CustomEvent.h"
 #include "K2Node_FunctionEntry.h"
-#include "K2Node_InputAction.h"
-#include "K2Node_InputKey.h"
+#include "K2Node_EventNodeInterface.h"
 #include "ScopedTransaction.h"
 
-
 #include "DetailLayoutBuilder.h"
 
 #include "SKismetInspector.h"
@@ -889,7 +887,7 @@ TSharedRef<SWidget> SMyBlueprint::OnCreateWidgetForAction(FCreateWidgetForAction
 	return BlueprintEditorPtr.IsValid() ? SNew(SBlueprintPaletteItem, InCreateData, BlueprintEditorPtr.Pin()) : SNew(SBlueprintPaletteItem, InCreateData, GetBlueprintObj());
 }
 
-void SMyBlueprint::GetChildGraphs(UEdGraph* InEdGraph, int32 const SectionId, FGraphActionSort& SortList, FText ParentCategory)
+void SMyBlueprint::GetChildGraphs(UEdGraph* InEdGraph, int32 const SectionId, FGraphActionSort& SortList, const FText& ParentCategory)
 {
 	check(InEdGraph);
 
@@ -924,11 +922,11 @@ void SMyBlueprint::GetChildGraphs(UEdGraph* InEdGraph, int32 const SectionId, FG
 			Category = EdGraphDisplayName;
 		}
 
-		const FText ChildTooltip = DisplayText;
-		const FText ChildDesc = DisplayText;
-		const FName DisplayName =  FName(*DisplayText.ToString());
+		const FName DisplayName = FName(*DisplayText.ToString());
+		FText ChildTooltip = DisplayText;
+		FText ChildDesc = MoveTemp(DisplayText);
 
-		TSharedPtr<FEdGraphSchemaAction_K2Graph> NewChildAction = MakeShareable(new FEdGraphSchemaAction_K2Graph(EEdGraphSchemaAction_K2Graph::Subgraph, Category, ChildDesc, ChildTooltip, 1, SectionId));
+		TSharedPtr<FEdGraphSchemaAction_K2Graph> NewChildAction = MakeShareable(new FEdGraphSchemaAction_K2Graph(EEdGraphSchemaAction_K2Graph::Subgraph, Category, MoveTemp(ChildDesc), MoveTemp(ChildTooltip), 1, SectionId));
 		NewChildAction->FuncName = DisplayName;
 		NewChildAction->EdGraph = Graph;
 		SortList.AddAction(NewChildAction);
@@ -938,27 +936,7 @@ void SMyBlueprint::GetChildGraphs(UEdGraph* InEdGraph, int32 const SectionId, FG
 	}
 }
 
-struct FCreateEdGraphSchemaActionHelper
-{
-	template<class ActionType, class NodeType>
-	static void CreateAll(UEdGraph const* EdGraph, int32 SectionId, FGraphActionSort& SortList, const FText& ActionCategory)
-	{
-		TArray<NodeType*> EventNodes;
-		EdGraph->GetNodesOfClass<NodeType>(EventNodes);
-		for (NodeType* const EventNode : EventNodes)
-		{
-			const FText Tooltip = EventNode->GetTooltipText();
-			const FText Description = EventNode->GetNodeTitle(ENodeTitleType::EditableTitle);
-
-			TSharedPtr<ActionType> EventNodeAction = MakeShareable(new ActionType(ActionCategory, Description, Tooltip, 0));
-			EventNodeAction->NodeTemplate = EventNode;
-			EventNodeAction->SectionID = SectionId;
-			SortList.AddAction(EventNodeAction);
-		}
-	}
-};
-
-void SMyBlueprint::GetChildEvents(UEdGraph const* InEdGraph, int32 const SectionId, FGraphActionSort& SortList, FText ParentCategory) const
+void SMyBlueprint::GetChildEvents(UEdGraph const* InEdGraph, int32 const SectionId, FGraphActionSort& SortList, const FText& ParentCategory) const
 {
 	if (!ensure(InEdGraph != NULL))
 	{
@@ -971,7 +949,7 @@ void SMyBlueprint::GetChildEvents(UEdGraph const* InEdGraph, int32 const Section
 	{
 		Schema->GetGraphDisplayInformation(*InEdGraph, EdGraphDisplayInfo);
 	}
-	FText const EdGraphDisplayName = EdGraphDisplayInfo.DisplayName;
+	FText EdGraphDisplayName = EdGraphDisplayInfo.DisplayName;
 	FText ActionCategory;
 	if (!ParentCategory.IsEmpty())
 	{
@@ -979,12 +957,18 @@ void SMyBlueprint::GetChildEvents(UEdGraph const* InEdGraph, int32 const Section
 	}
 	else
 	{
-		ActionCategory = EdGraphDisplayName;
+		ActionCategory = MoveTemp(EdGraphDisplayName);
 	}
 
-	FCreateEdGraphSchemaActionHelper::CreateAll<FEdGraphSchemaAction_K2Event, UK2Node_Event>(InEdGraph, SectionId, SortList, ActionCategory);
-	FCreateEdGraphSchemaActionHelper::CreateAll<FEdGraphSchemaAction_K2InputAction, UK2Node_InputKey>(InEdGraph, SectionId, SortList, ActionCategory);
-	FCreateEdGraphSchemaActionHelper::CreateAll<FEdGraphSchemaAction_K2InputAction, UK2Node_InputAction>(InEdGraph, SectionId, SortList, ActionCategory);
+	for (UEdGraphNode* GraphNode : InEdGraph->Nodes)
+	{
+		if (GraphNode && GraphNode->GetClass()->ImplementsInterface(UK2Node_EventNodeInterface::StaticClass()))
+		{
+			TSharedPtr<FEdGraphSchemaAction> EventNodeAction = CastChecked<IK2Node_EventNodeInterface>(GraphNode)->GetEventNodeAction(ActionCategory);
+			EventNodeAction->SectionID = SectionId;
+			SortList.AddAction(EventNodeAction);
+		}
+	}
 }
 
 void SMyBlueprint::GetLocalVariables(FGraphActionSort& SortList) const
@@ -999,7 +983,6 @@ void SMyBlueprint::GetLocalVariables(FGraphActionSort& SortList) const
 		{
 			Schema->GetGraphDisplayInformation(*TopLevelGraph, EdGraphDisplayInfo);
 		}
-		FText const EdGraphDisplayName = EdGraphDisplayInfo.DisplayName;
 
 		TArray<UK2Node_FunctionEntry*> FunctionEntryNodes;
 		TopLevelGraph->GetNodesOfClass<UK2Node_FunctionEntry>(FunctionEntryNodes);
diff --git a/Engine/Source/Editor/Kismet/Private/SMyBlueprint.h b/Engine/Source/Editor/Kismet/Private/SMyBlueprint.h
--- a/Engine/Source/Editor/Kismet/Private/SMyBlueprint.h
+++ b/Engine/Source/Editor/Kismet/Private/SMyBlueprint.h
@@ -162,8 +162,8 @@ private:
 	/** Callback used to populate all actions list in SGraphActionMenu */
 	void CollectAllActions(FGraphActionListBuilderBase& OutAllActions);
 	void CollectStaticSections(TArray<int32>& StaticSectionIDs);
-	void GetChildGraphs(UEdGraph* EdGraph, int32 const SectionId, FGraphActionSort& SortList, FText ParentCategory = FText::GetEmpty());
-	void GetChildEvents(UEdGraph const* EdGraph, int32 const SectionId, FGraphActionSort& SortList, FText ParentCategory = FText::GetEmpty()) const;
+	void GetChildGraphs(UEdGraph* EdGraph, int32 const SectionId, FGraphActionSort& SortList, const FText& ParentCategory = FText::GetEmpty());
+	void GetChildEvents(UEdGraph const* EdGraph, int32 const SectionId, FGraphActionSort& SortList, const FText& ParentCategory = FText::GetEmpty()) const;
 	void GetLocalVariables(FGraphActionSort& SortList) const;
 	
 	/** Handles the visibility of the local action list */
