Blueprints: Fixed an issue where any locally declared macros with no specified category would result in an extra empty category in the dropdown
PR #4378: Prevent empty category entries in dropdown (Contributed by [REDACTED_PERSON])

diff --git a/Engine/Source/Editor/Kismet/Private/BlueprintDetailsCustomization.cpp b/Engine/Source/Editor/Kismet/Private/BlueprintDetailsCustomization.cpp
--- a/Engine/Source/Editor/Kismet/Private/BlueprintDetailsCustomization.cpp
+++ b/Engine/Source/Editor/Kismet/Private/BlueprintDetailsCustomization.cpp
@@ -1419,14 +1419,17 @@ void FBlueprintVarActionDetails::PopulateCategories(SMyBlueprint* MyBlueprint, T
 		UK2Node_EditablePinBase* EntryNode = FBlueprintEditorUtils::GetEntryNode(MacroGraph);
 		if (UK2Node_Tunnel* TypedEntryNode = ExactCast<UK2Node_Tunnel>(EntryNode))
 		{
-			bool bNewCategory = true;
-			for (int32 j = 0; j < CategorySource.Num() && bNewCategory; ++j)
-			{
-				bNewCategory &= !CategorySource[j].Get()->EqualTo(TypedEntryNode->MetaData.Category);
-			}
-			if (bNewCategory)
+			if (!TypedEntryNode->MetaData.Category.IsEmpty())
 			{
-				CategorySource.Add(MakeShareable(new FText(TypedEntryNode->MetaData.Category)));
+				bool bNewCategory = true;
+				for (int32 j = 0; j < CategorySource.Num() && bNewCategory; ++j)
+				{
+					bNewCategory &= !CategorySource[j].Get()->EqualTo(TypedEntryNode->MetaData.Category);
+				}
+				if (bNewCategory)
+				{
+					CategorySource.Add(MakeShareable(new FText(TypedEntryNode->MetaData.Category)));
+				}
 			}
 		}
 	}
