Initial support for Blueprint Graph editor "Quick Jump" bookmarks.

Change summary:
- Moved the FBlueprintEditor::OpenGraphAndBringToFront() API up into the IBlueprintEditor interface.
- Modified FGraphEditorCommandsImpl to define new "Quick Jump" commands for graph editor contexts. Keyboard shortcuts are currently set up as follows (where 'slot_index' is '0' through '9' on the keyboard):
    - Quick Jump to Bookmark: SHIFT + <slot_index>
    - Set Quick Jump Bookmark: CTRL + <slot_index>
- Modified FBlueprintEditor::CreateDefaultCommands() to map Quick Jump commands to actions within the top-level Blueprint editor toolkit context.
- Added implementations to FBlueprintEditor for handling Quick Jump commands within the relevant editor contexts.
- Added UBlueprintEditorSettings::GraphEditorQuickJumps to save Quick Jump bookmark info to local user settings (.ini).

Known issues:
- Keyboard input focus will be lost whenever the user clicks on a tab well at the top of the Blueprint editor window. As a result, quick jump keyboard shortcuts won't work again until the user clicks somewhere inside the tab to reclaim the input focus. This is a general UI problem that's not tied to this feature specifically.
- Similarly, top-level editor command keyboard shortcuts (including Quick Jump) won't work immediately after opening a Blueprint asset in a new editor window. The keyboard input focus is not set until the user clicks somewhere inside the tab area. All editors have this issue (not just Blueprint editor windows).

diff --git a/Engine/Source/Editor/BlueprintGraph/Public/BlueprintEditorSettings.h b/Engine/Source/Editor/BlueprintGraph/Public/BlueprintEditorSettings.h
--- a/Engine/Source/Editor/BlueprintGraph/Public/BlueprintEditorSettings.h
+++ b/Engine/Source/Editor/BlueprintGraph/Public/BlueprintEditorSettings.h
@@ -106,6 +106,10 @@ public:
 	UPROPERTY(config)
 	bool bShowBookmarksForCurrentDocumentOnlyInTab;
 
+	/** Blueprint graph editor "Quick Jump" command bindings */
+	UPROPERTY(config)
+	TMap<int32, FEditedDocumentInfo> GraphEditorQuickJumps;
+
 // Compiler Settings
 public:
 	/** Determines when to save Blueprints post-compile */
diff --git a/Engine/Source/Editor/GraphEditor/Private/GraphEditorActions.cpp b/Engine/Source/Editor/GraphEditor/Private/GraphEditorActions.cpp
--- a/Engine/Source/Editor/GraphEditor/Private/GraphEditorActions.cpp
+++ b/Engine/Source/Editor/GraphEditor/Private/GraphEditorActions.cpp
@@ -100,9 +100,61 @@ void FGraphEditorCommandsImpl::RegisterCommands()
 	UI_COMMAND( ZoomOut, "Zoom Out", "Zoom out from the graph editor", EUserInterfaceActionType::Button, FInputChord(EKeys::Subtract))
 
 	UI_COMMAND( GoToDocumentation, "View Documentation", "View documentation for this node.", EUserInterfaceActionType::Button, FInputChord());
-}
-
 
+	// Map quick jump index to command key bindings.
+	TArray< FKey, TInlineAllocator<10> > NumberKeys;
+	NumberKeys.Add( EKeys::Zero );
+	NumberKeys.Add( EKeys::One );
+	NumberKeys.Add( EKeys::Two );
+	NumberKeys.Add( EKeys::Three );
+	NumberKeys.Add( EKeys::Four );
+	NumberKeys.Add( EKeys::Five );
+	NumberKeys.Add( EKeys::Six );
+	NumberKeys.Add( EKeys::Seven );
+	NumberKeys.Add( EKeys::Eight );
+	NumberKeys.Add( EKeys::Nine );
+
+	const int32 NumQuickJumpCommands = NumberKeys.Num();
+	QuickJumpCommands.Reserve(NumQuickJumpCommands);
+
+	for (int32 QuickJumpIndex = 0; QuickJumpIndex < NumQuickJumpCommands; ++QuickJumpIndex)
+	{
+		const FText QuickJumpIndexText = FText::AsNumber(QuickJumpIndex);
+
+		FQuickJumpCommandInfo QuickJumpCommandInfo;
+
+		QuickJumpCommandInfo.QuickJump =
+			FUICommandInfoDecl(
+				this->AsShared(), //Command class
+				FName(*FString::Printf(TEXT("QuickJump%i"), QuickJumpIndex)), //CommandName
+				FText::Format(LOCTEXT("QuickJump", "Quick Jump {0}"), QuickJumpIndexText), //Localized label
+				FText::Format(LOCTEXT("QuickJump_ToolTip", "Jump to the location and zoom level bound to {0}"), QuickJumpIndexText))//Localized tooltip
+			.UserInterfaceType(EUserInterfaceActionType::Button) //interface type
+			.DefaultChord(FInputChord(EModifierKey::Shift, NumberKeys[QuickJumpIndex])); //default chord
+
+		QuickJumpCommandInfo.SetQuickJump =
+			FUICommandInfoDecl(
+				this->AsShared(), //Command class
+				FName(*FString::Printf(TEXT("SetQuickJump%i"), QuickJumpIndex)), //CommandName
+				FText::Format(LOCTEXT("SetQuickJump", "Set Quick Jump {0}"), QuickJumpIndexText), //Localized label
+				FText::Format(LOCTEXT("SetQuickJump_ToolTip", "Save the graph's current location and zoom level as quick jump {0}"), QuickJumpIndexText))//Localized tooltip
+			.UserInterfaceType(EUserInterfaceActionType::Button) //interface type
+			.DefaultChord(FInputChord(EModifierKey::Control, NumberKeys[QuickJumpIndex])); //default chord
+
+		QuickJumpCommandInfo.ClearQuickJump =
+			FUICommandInfoDecl(
+				this->AsShared(), //Command class
+				FName(*FString::Printf(TEXT("ClearQuickJump%i"), QuickJumpIndex)), //CommandName
+				FText::Format(LOCTEXT("ClearQuickJump", "Clear Quick Jump {0}"), QuickJumpIndexText), //Localized label
+				FText::Format(LOCTEXT("ClearQuickJump_ToolTip", "Clear the saved location and zoom level at quick jump {0}"), QuickJumpIndexText))//Localized tooltip
+			.UserInterfaceType(EUserInterfaceActionType::Button) //interface type
+			.DefaultChord(FInputChord()); //default chord 
+
+		QuickJumpCommands.Add(QuickJumpCommandInfo);
+	}
+
+	UI_COMMAND( ClearAllQuickJumps, "Clear All Quick Jumps", "Clear all quick jump bindings", EUserInterfaceActionType::Button, FInputChord() );
+}
 
 void FGraphEditorCommands::Register()
 {
diff --git a/Engine/Source/Editor/GraphEditor/Public/GraphEditorActions.h b/Engine/Source/Editor/GraphEditor/Public/GraphEditorActions.h
--- a/Engine/Source/Editor/GraphEditor/Public/GraphEditorActions.h
+++ b/Engine/Source/Editor/GraphEditor/Public/GraphEditorActions.h
@@ -138,6 +138,16 @@ public:
 
 	// Go to node documentation
 	TSharedPtr< FUICommandInfo > GoToDocumentation;
+
+	// Quick jump commands
+	struct FQuickJumpCommandInfo
+	{
+		TSharedPtr< FUICommandInfo > QuickJump;
+		TSharedPtr< FUICommandInfo > SetQuickJump;
+		TSharedPtr< FUICommandInfo > ClearQuickJump;
+	};
+	TArray< FQuickJumpCommandInfo > QuickJumpCommands;
+	TSharedPtr< FUICommandInfo > ClearAllQuickJumps;
 };
 
 class GRAPHEDITOR_API FGraphEditorCommands
@@ -148,6 +158,4 @@ public:
 	static const FGraphEditorCommandsImpl& Get();
 
 	static void Unregister();
-};
-
-
+};
\ No newline at end of file
diff --git a/Engine/Source/Editor/Kismet/Private/BlueprintEditor.cpp b/Engine/Source/Editor/Kismet/Private/BlueprintEditor.cpp
--- a/Engine/Source/Editor/Kismet/Private/BlueprintEditor.cpp
+++ b/Engine/Source/Editor/Kismet/Private/BlueprintEditor.cpp
@@ -2623,6 +2623,29 @@ void FBlueprintEditor::CreateDefaultCommands()
 			}),
 		FCanExecuteAction(),
 		FIsActionChecked::CreateLambda([]()->bool{ return GetDefault<UBlueprintEditorSettings>()->bShowActionMenuItemSignatures; }));
+
+	for (int32 QuickJumpIndex = 0; QuickJumpIndex < FGraphEditorCommands::Get().QuickJumpCommands.Num(); ++QuickJumpIndex)
+	{
+		ToolkitCommands->MapAction(
+			FGraphEditorCommands::Get().QuickJumpCommands[QuickJumpIndex].QuickJump,
+			FExecuteAction::CreateSP(this, &FBlueprintEditor::OnGraphEditorQuickJump, QuickJumpIndex)
+		);
+
+		ToolkitCommands->MapAction(
+			FGraphEditorCommands::Get().QuickJumpCommands[QuickJumpIndex].SetQuickJump,
+			FExecuteAction::CreateSP(this, &FBlueprintEditor::SetGraphEditorQuickJump, QuickJumpIndex)
+		);
+
+		ToolkitCommands->MapAction(
+			FGraphEditorCommands::Get().QuickJumpCommands[QuickJumpIndex].ClearQuickJump,
+			FExecuteAction::CreateSP(this, &FBlueprintEditor::ClearGraphEditorQuickJump, QuickJumpIndex)
+		);
+	}
+
+	ToolkitCommands->MapAction(
+		FGraphEditorCommands::Get().ClearAllQuickJumps,
+		FExecuteAction::CreateSP(this, &FBlueprintEditor::ClearAllGraphEditorQuickJumps)
+	);
 }
 
 void FBlueprintEditor::OpenNativeCodeGenerationTool()
@@ -8481,6 +8504,59 @@ void FBlueprintEditor::RemoveBookmark(const FGuid& BookmarkNodeId, bool bRefresh
 	}
 }
 
+void FBlueprintEditor::SetGraphEditorQuickJump(int32 QuickJumpIndex)
+{
+	TSharedPtr<SGraphEditor> FocusedGraphEd = FocusedGraphEdPtr.Pin();
+	if (FocusedGraphEd.IsValid())
+	{
+		if (UEdGraph* GraphObject = FocusedGraphEd->GetCurrentGraph())
+		{
+			UBlueprintEditorSettings* LocalSettings = GetMutableDefault<UBlueprintEditorSettings>();
+			FEditedDocumentInfo& QuickJumpInfo = LocalSettings->GraphEditorQuickJumps.FindOrAdd(QuickJumpIndex);
+
+			QuickJumpInfo.EditedObjectPath = GraphObject;
+			FocusedGraphEd->GetViewLocation(QuickJumpInfo.SavedViewOffset, QuickJumpInfo.SavedZoomAmount);
+
+			LocalSettings->SaveConfig();
+		}
+	}
+}
+
+void FBlueprintEditor::ClearGraphEditorQuickJump(int32 QuickJumpIndex)
+{
+	UBlueprintEditorSettings* LocalSettings = GetMutableDefault<UBlueprintEditorSettings>();
+	LocalSettings->GraphEditorQuickJumps.Remove(QuickJumpIndex);
+	LocalSettings->SaveConfig();
+}
+
+void FBlueprintEditor::OnGraphEditorQuickJump(int32 QuickJumpIndex)
+{
+	const UBlueprintEditorSettings* LocalSettings = GetDefault<UBlueprintEditorSettings>();
+	if (const FEditedDocumentInfo* QuickJumpInfo = LocalSettings->GraphEditorQuickJumps.Find(QuickJumpIndex))
+	{
+		if (UObject* EditedObject = QuickJumpInfo->EditedObjectPath.TryLoad())
+		{
+			TSharedPtr<IBlueprintEditor> IBlueprintEditorPtr = FKismetEditorUtilities::GetIBlueprintEditorForObject(EditedObject, true);
+			if (IBlueprintEditorPtr.IsValid())
+			{
+				IBlueprintEditorPtr->FocusWindow();
+				TSharedPtr<SGraphEditor> GraphEditorPtr = IBlueprintEditorPtr->OpenGraphAndBringToFront(Cast<UEdGraph>(EditedObject));
+				if (GraphEditorPtr.IsValid())
+				{
+					GraphEditorPtr->SetViewLocation(QuickJumpInfo->SavedViewOffset, QuickJumpInfo->SavedZoomAmount);
+				}
+			}
+		}
+	}
+}
+
+void FBlueprintEditor::ClearAllGraphEditorQuickJumps()
+{
+	UBlueprintEditorSettings* LocalSettings = GetMutableDefault<UBlueprintEditorSettings>();
+	LocalSettings->GraphEditorQuickJumps.Empty();
+	LocalSettings->SaveConfig();
+}
+
 /////////////////////////////////////////////////////
 /////////////////////////////////////////////////////
 /////////////////////////////////////////////////////
diff --git a/Engine/Source/Editor/Kismet/Public/BlueprintEditor.h b/Engine/Source/Editor/Kismet/Public/BlueprintEditor.h
--- a/Engine/Source/Editor/Kismet/Public/BlueprintEditor.h
+++ b/Engine/Source/Editor/Kismet/Public/BlueprintEditor.h
@@ -191,6 +191,7 @@ public:
 	virtual void JumpToPin(const class UEdGraphPin* Pin) override;
 	virtual void SummonSearchUI(bool bSetFindWithinBlueprint, FString NewSearchTerms = FString(), bool bSelectFirstResult = false) override;
 	virtual void SummonFindAndReplaceUI() override;
+	virtual TSharedPtr<SGraphEditor> OpenGraphAndBringToFront(UEdGraph* Graph) override;
 	virtual TArray<TSharedPtr<class FSCSEditorTreeNode> > GetSelectedSCSEditorTreeNodes() const override;
 	virtual TSharedPtr<class FSCSEditorTreeNode> FindAndSelectSCSEditorTreeNode(const UActorComponent* InComponent, bool IsCntrlDown) override;
 	virtual int32 GetNumberOfSelectedNodes() const override;
@@ -978,6 +979,18 @@ protected:
 	/** Attempts to invoke the details tab if it's currently possible to. */
 	void TryInvokingDetailsTab(bool bFlash = true);
 
+	/** Handles activation of a graph editor "quick jump" command */
+	void OnGraphEditorQuickJump(int32 BookmarkIndex);
+
+	/** Binds a graph editor "quick jump" to the command at the given index */
+	void SetGraphEditorQuickJump(int32 BookmarkIndex);
+
+	/** Unbinds a graph editor "quick jump" from the command at the given index */
+	void ClearGraphEditorQuickJump(int32 BookmarkIndex);
+
+	/** Unbinds all graph editor "quick jump" commands */
+	void ClearAllGraphEditorQuickJumps();
+
 private:
 
 	/** Returns true if modules can be recompiled */
@@ -1157,9 +1170,6 @@ protected:
 	mutable TWeakObjectPtr<AActor> PreviewActorPtr;
 
 public:
-	// Tries to open the specified graph and bring it's document to the front (note: this can return NULL)
-	TSharedPtr<SGraphEditor> OpenGraphAndBringToFront(UEdGraph* Graph);
-
 	//@TODO: To be moved/merged
 	TSharedPtr<SDockTab> OpenDocument(const UObject* DocumentID, FDocumentTracker::EOpenDocumentCause Cause);
 
@@ -1169,7 +1179,6 @@ public:
 	// Finds any open tabs containing the specified document and adds them to the specified array; returns true if at least one is found
 	bool FindOpenTabsContainingDocument(const UObject* DocumentID, /*inout*/ TArray< TSharedPtr<SDockTab> >& Results);
 
-
 public:
 	/** Broadcasts a notification whenever the editor needs associated controls to refresh */
 	DECLARE_EVENT ( FBlueprintEditor, FOnRefreshEvent );
diff --git a/Engine/Source/Editor/Kismet/Public/BlueprintEditorModule.h b/Engine/Source/Editor/Kismet/Public/BlueprintEditorModule.h
--- a/Engine/Source/Editor/Kismet/Public/BlueprintEditorModule.h
+++ b/Engine/Source/Editor/Kismet/Public/BlueprintEditorModule.h
@@ -64,6 +64,9 @@ public:
 	/** Invokes the Find and Replace UI */
 	virtual void SummonFindAndReplaceUI() = 0;
 
+	/** Tries to open the specified graph and bring it's document to the front (note: this can return NULL) */
+	virtual TSharedPtr<class SGraphEditor> OpenGraphAndBringToFront(class UEdGraph* Graph) = 0;
+
 	virtual void RefreshEditors(ERefreshBlueprintEditorReason::Type Reason = ERefreshBlueprintEditorReason::UnknownReason) = 0;
 
 	virtual void AddToSelection(UEdGraphNode* InNode) = 0;
