Blueprints: Formalized the concept of a blueprint embedded within an asset for nativization and "Find in Blueprints" purposes

  - Added IBlueprintAssetHandler to define how blueprints should be retrieved from particular assets, registered by class name
  - Replaced the majority of UWorld specific logic from the find in blueprint manager with calls to this new handler
  - All search entries are stored as the full Blueprint path, all unindexed (uncached) entries are stored as the full asset path.
  - Added a function (SupportsNativization) to UBlueprint to enable specific blueprints to explicitly disallow nativization

diff --git a/Engine/Source/Developer/BlueprintNativeCodeGen/Private/BlueprintNativeCodeGenModule.cpp b/Engine/Source/Developer/BlueprintNativeCodeGen/Private/BlueprintNativeCodeGenModule.cpp
--- a/Engine/Source/Developer/BlueprintNativeCodeGen/Private/BlueprintNativeCodeGenModule.cpp
+++ b/Engine/Source/Developer/BlueprintNativeCodeGen/Private/BlueprintNativeCodeGenModule.cpp
@@ -971,20 +971,9 @@ EReplacementResult FBlueprintNativeCodeGenModule::IsTargetedForReplacement(const
 			}
 
 			// Unconvertable Blueprint
+			if (!Blueprint->SupportsNativization())
 			{
-				const EBlueprintType UnconvertableBlueprintTypes[] = {
-					//BPTYPE_Const,		// What is a "const" Blueprint?
-					BPTYPE_MacroLibrary,
-					BPTYPE_LevelScript,
-				};
-				const EBlueprintType BlueprintType = Blueprint->BlueprintType;
-				for (int32 TypeIndex = 0; TypeIndex < ARRAY_COUNT(UnconvertableBlueprintTypes); ++TypeIndex)
-				{
-					if (BlueprintType == UnconvertableBlueprintTypes[TypeIndex])
-					{
-						return true;
-					}
-				}
+				return true;
 			}
 
 			// ExcludedBlueprintTypes
diff --git a/Engine/Source/Editor/Kismet/Private/BlueprintAssetHandler.cpp b/Engine/Source/Editor/Kismet/Private/BlueprintAssetHandler.cpp
new file mode 100644
--- /dev/null
+++ b/Engine/Source/Editor/Kismet/Private/BlueprintAssetHandler.cpp
@@ -0,0 +1,75 @@
+// Copyright 1998-2018 Epic Games, Inc. All Rights Reserved.
+
+#include "BlueprintAssetHandler.h"
+#include "Engine/World.h"
+#include "Engine/LevelScriptBlueprint.h"
+#include "AssetData.h"
+
+class FLevelBlueprintAssetHandler : public IBlueprintAssetHandler
+{
+	virtual UBlueprint* RetrieveBlueprint(UObject* InObject) const override
+	{
+		UWorld* World = CastChecked<UWorld>(InObject);
+
+		const bool bCreateLevelScript = false;
+		return World->PersistentLevel ? World->PersistentLevel->GetLevelScriptBlueprint(bCreateLevelScript) : nullptr;
+	}
+
+	virtual bool AssetContainsBlueprint(const FAssetData& InAssetData) const
+	{
+		static FName FiB("FiB");
+		// Worlds are only considered to contain a blueprint if they have FiB data
+		return InAssetData.TagsAndValues.Find(FiB) || InAssetData.TagsAndValues.Find(FBlueprintTags::FindInBlueprintsData);
+	}
+};
+
+
+class FBlueprintAssetTypeHandler : public IBlueprintAssetHandler
+{
+	virtual UBlueprint* RetrieveBlueprint(UObject* InObject) const override
+	{
+		// The object is the blueprint for UBlueprint (and derived) assets
+		return CastChecked<UBlueprint>(InObject);
+	}
+
+	virtual bool AssetContainsBlueprint(const FAssetData& InAssetData) const
+	{
+		return true;
+	}
+};
+
+FBlueprintAssetHandler::FBlueprintAssetHandler()
+{
+	// Register default handlers
+	RegisterHandler<FLevelBlueprintAssetHandler>(UWorld::StaticClass()->GetFName());
+	RegisterHandler<FBlueprintAssetTypeHandler>(UBlueprint::StaticClass()->GetFName());
+}
+
+FBlueprintAssetHandler& FBlueprintAssetHandler::Get()
+{
+	static FBlueprintAssetHandler Singleton;
+	return Singleton;
+}
+
+void FBlueprintAssetHandler::RegisterHandler(FName EligibleClass, TUniquePtr<IBlueprintAssetHandler>&& InHandler)
+{
+	ClassNames.Add(EligibleClass);
+	Handlers.Add(MoveTemp(InHandler));
+}
+
+const IBlueprintAssetHandler* FBlueprintAssetHandler::FindHandler(const UClass* InClass) const
+{
+	UClass* StopAtClass = UObject::StaticClass();
+	while (InClass && InClass != StopAtClass)
+	{
+		int32 Index = ClassNames.IndexOfByKey(InClass->GetFName());
+		if (Index != INDEX_NONE)
+		{
+			return Handlers[Index].Get();
+		}
+
+		InClass = InClass->GetSuperClass();
+	}
+
+	return nullptr;
+}
\ No newline at end of file
diff --git a/Engine/Source/Editor/Kismet/Private/BlueprintDetailsCustomization.cpp b/Engine/Source/Editor/Kismet/Private/BlueprintDetailsCustomization.cpp
--- a/Engine/Source/Editor/Kismet/Private/BlueprintDetailsCustomization.cpp
+++ b/Engine/Source/Editor/Kismet/Private/BlueprintDetailsCustomization.cpp
@@ -5247,12 +5247,6 @@ FText FBlueprintGlobalOptionsDetails::GetDeprecatedTooltip() const
 	return LOCTEXT("DisabledDeprecateBlueprintTooltip", "This Blueprint is deprecated because of a parent, it is not possible to remove deprecation from it!");
 }
 
-/** Shared tooltip text for both the label and the value field */
-static FText GetNativizeLabelTooltip()
-{
-	return LOCTEXT("NativizeTooltip", "When exclusive nativization is enabled, then this asset will be nativized. NOTE: All super classes must be also nativized.");
-}
-
 void FBlueprintGlobalOptionsDetails::CustomizeDetails(IDetailLayoutBuilder& DetailLayout)
 {
 	const UBlueprint* Blueprint = GetBlueprintObj();
@@ -5354,9 +5348,8 @@ void FBlueprintGlobalOptionsDetails::CustomizeDetails(IDetailLayoutBuilder& Deta
 			.NameContent()
 			[
 				SNew(STextBlock)
-					.Text(LOCTEXT("NativizeLabel", "Nativize"))
-					.ToolTipText_Static(&GetNativizeLabelTooltip)
-					.Font(IDetailLayoutBuilder::GetDetailFont())
+				.Text(LOCTEXT("NativizeLabel", "Nativize"))
+				.Font(IDetailLayoutBuilder::GetDetailFont())
 			]
 			.ValueContent()
 			[
@@ -5374,7 +5367,7 @@ bool FBlueprintGlobalOptionsDetails::IsNativizeEnabled() const
 	bool bIsEnabled = false;
 	if (UBlueprint* Blueprint = GetBlueprintObj())
 	{
-		bIsEnabled = (Blueprint->BlueprintType != BPTYPE_MacroLibrary) && (Blueprint->BlueprintType != BPTYPE_LevelScript) && (!FBlueprintEditorUtils::ShouldNativizeImplicitly(Blueprint));
+		bIsEnabled = Blueprint->SupportsNativization() && !FBlueprintEditorUtils::ShouldNativizeImplicitly(Blueprint);
 	}
 	return bIsEnabled;
 }
@@ -5412,22 +5405,28 @@ ECheckBoxState FBlueprintGlobalOptionsDetails::GetNativizeState() const
 
 FText FBlueprintGlobalOptionsDetails::GetNativizeTooltip() const
 {
-	if (!IsNativizeEnabled())
+	UBlueprint* Blueprint = GetBlueprintObj();
+
+	if (Blueprint)
 	{
-		if (FBlueprintEditorUtils::ShouldNativizeImplicitly(GetBlueprintObj()))
+		if (FBlueprintEditorUtils::ShouldNativizeImplicitly(Blueprint))
 		{
 			return LOCTEXT("NativizeImplicitlyTooltip", "This Blueprint must be nativized because it overrides one or more BlueprintCallable functions inherited from a parent Blueprint class that has also been flagged for nativization.");
 		}
-		else
+
+		FText Reason;
+		if (!Blueprint->SupportsNativization(&Reason))
 		{
-			return LOCTEXT("NativizeDisabledTooltip", "Macro libraries and level Blueprints cannot be nativized.");
+			return Reason.IsEmpty() ? LOCTEXT("NativizeDisabledTooltip", "This blueprint does not support nativization.") : Reason;
+		}
+
+		if (Blueprint->NativizationFlag == EBlueprintNativizationFlag::Dependency)
+		{
+			return LOCTEXT("NativizeAsDependencyTooltip", "This Blueprint has been flagged to nativize as a dependency needed by another Blueprint. This will be applied once that Blueprint is saved.");
 		}
 	}
-	else if (GetNativizeState() == ECheckBoxState::Undetermined)
-	{
-		return LOCTEXT("NativizeAsDependencyTooltip", "This Blueprint has been flagged to nativize as a dependency needed by another Blueprint. This will be applied once that Blueprint is saved.");
-	}
-	return GetNativizeLabelTooltip();
+
+	return LOCTEXT("NativizeTooltip", "When exclusive nativization is enabled, then this asset will be nativized. NOTE: All super classes must be also nativized.");
 }
 
 void FBlueprintGlobalOptionsDetails::OnNativizeToggled(ECheckBoxState NewState) const
diff --git a/Engine/Source/Editor/Kismet/Private/FindInBlueprintManager.cpp b/Engine/Source/Editor/Kismet/Private/FindInBlueprintManager.cpp
--- a/Engine/Source/Editor/Kismet/Private/FindInBlueprintManager.cpp
+++ b/Engine/Source/Editor/Kismet/Private/FindInBlueprintManager.cpp
@@ -45,6 +45,7 @@
 #include "ImaginaryBlueprintData.h"
 #include "FiBSearchInstance.h"
 #include "Misc/HotReloadInterface.h"
+#include "BlueprintAssetHandler.h"
 
 #include "JsonObjectConverter.h"
 #include "Framework/Notifications/NotificationManager.h"
@@ -1079,9 +1080,9 @@ class FCacheAllBlueprintsTickableObject
 
 public:
 
-	FCacheAllBlueprintsTickableObject(TSet<FName> InUncachedBlueprints, bool bInCheckOutAndSave, FSimpleDelegate InOnFinished = FSimpleDelegate())
+	FCacheAllBlueprintsTickableObject(TSet<FName> InUncachedAssets, bool bInCheckOutAndSave, FSimpleDelegate InOnFinished = FSimpleDelegate())
 		: TickCacheIndex(0)
-		, UncachedBlueprints(InUncachedBlueprints.Array())
+		, UncachedAssets(InUncachedAssets.Array())
 		, bIsStarted(false)
 		, bIsCancelled(false)
 		, bRecursionGuard(false)
@@ -1116,9 +1117,9 @@ public:
 	/** Returns the name of the current Blueprint being cached */
 	FName GetCurrentCacheBlueprintName() const
 	{
-		if(UncachedBlueprints.Num() && TickCacheIndex >= 0)
+		if(UncachedAssets.Num() && TickCacheIndex >= 0)
 		{
-			return UncachedBlueprints[TickCacheIndex];
+			return UncachedAssets[TickCacheIndex];
 		}
 		return NAME_None;
 	}
@@ -1126,19 +1127,19 @@ public:
 	/** Returns the progress as a percent */
 	float GetCacheProgress() const
 	{
-		return (float)TickCacheIndex / (float)UncachedBlueprints.Num();
+		return (float)TickCacheIndex / (float)UncachedAssets.Num();
 	}
 
-	/** Returns the number of uncached Blueprints */
-	int32 GetUncachedBlueprintCount()
+	/** Returns the number of uncached assets */
+	int32 GetUncachedAssetCount()
 	{
-		return UncachedBlueprints.Num();
+		return UncachedAssets.Num();
 	}
 
-	/** Returns the entire list of uncached Blueprints that this object will attempt to cache */
-	const TArray<FName>& GetUncachedBlueprintList() const
+	/** Returns the entire list of uncached assets that this object will attempt to cache */
+	const TArray<FName>& GetUncachedAssetList() const
 	{
-		return UncachedBlueprints;
+		return UncachedAssets;
 	}
 
 	/** True if there is a callback when done caching, this will prevent a re-query from occuring */
@@ -1194,9 +1195,9 @@ public:
 		else
 		{
 			FAssetRegistryModule* AssetRegistryModule = &FModuleManager::LoadModuleChecked<FAssetRegistryModule>(TEXT("AssetRegistry"));
-			FAssetData AssetData = AssetRegistryModule->Get().GetAssetByObjectPath(UncachedBlueprints[TickCacheIndex]);
+			FAssetData AssetData = AssetRegistryModule->Get().GetAssetByObjectPath(UncachedAssets[TickCacheIndex]);
 
-			bool bIsWorldAsset = AssetData.GetClass() == UWorld::StaticClass();
+			const bool bIsWorldAsset = AssetData.AssetClass == UWorld::StaticClass()->GetFName();
 
 			// Construct a full package filename with path so we can query the read only status and save to disk
 			FString FinalPackageFilename = FPackageName::LongPackageNameToFilename(AssetData.PackageName.ToString());
@@ -1267,13 +1268,13 @@ public:
 
 			if (bFailedToCache)
 			{
-				FailedToCacheList.Add(UncachedBlueprints[TickCacheIndex]);
+				FailedToCacheList.Add(UncachedAssets[TickCacheIndex]);
 			}
 
 			++TickCacheIndex;
 
 			// Check if done caching Blueprints
-			if(TickCacheIndex == UncachedBlueprints.Num())
+			if(TickCacheIndex == UncachedAssets.Num())
 			{
 				ProgressNotification.Pin()->SetCompletionState(SNotificationItem::CS_Success);
 				ProgressNotification.Pin()->ExpireAndFadeout();
@@ -1302,8 +1303,8 @@ private:
 	/** The current index, increases at a rate of once per tick */
 	int32 TickCacheIndex;
 
-	/** The list of uncached Blueprints that are in the process of being cached */
-	TArray<FName> UncachedBlueprints;
+	/** The list of uncached assets that are in the process of being cached */
+	TArray<FName> UncachedAssets;
 
 	/** Notification that appears and details progress */
 	TWeakPtr<SNotificationItem> ProgressNotification;
@@ -1437,78 +1438,82 @@ void FFindInBlueprintSearchManager::OnAssetAdded(const FAssetData& InAssetData)
 		}
 	}
 
-	bool bIsLevel = AssetClass && AssetClass->IsChildOf(UWorld::StaticClass());
-	bool bIsBlueprint = AssetClass && AssetClass->IsChildOf(UBlueprint::StaticClass());
+	const IBlueprintAssetHandler* Handler = AssetClass ? FBlueprintAssetHandler::Get().FindHandler(AssetClass) : nullptr;
 
-	if(bIsLevel || bIsBlueprint)
+	// No handler means we can't process this asset
+	if (!Handler)
 	{
-		// Confirm that the Blueprint has not been added already, this can occur during duplication of Blueprints.
-		int32* IndexPtr = SearchMap.Find(InAssetData.ObjectPath);
-		if(!IndexPtr)
-		{
-			if(InAssetData.IsAssetLoaded())
-			{
-				if(bIsBlueprint)
-				{
-					if(UBlueprint* BlueprintAsset = Cast<UBlueprint>(InAssetData.GetAsset()))
-					{
-						// Cache the searchable data
-						AddOrUpdateBlueprintSearchMetadata(BlueprintAsset);
-					}
-				}
-				else if(bIsLevel)
-				{
-					UWorld* WorldAsset = Cast<UWorld>(InAssetData.GetAsset());
-					if(WorldAsset->PersistentLevel)
-					{
-						TArray<UBlueprint*> LevelBlueprints;
-						LevelBlueprints = WorldAsset->PersistentLevel->GetLevelBlueprints();
-
-						for(UBlueprint* BlueprintAsset : LevelBlueprints)
-						{
-							// Cache the searchable data
-							AddOrUpdateBlueprintSearchMetadata(BlueprintAsset);
-						}
+		return;
+	}
 
-					}
-				}
-			}
-			else if(const FString* FiBSearchData = InAssetData.TagsAndValues.Find("FiB"))
-			{
-				ExtractUnloadedFiBData(InAssetData, *FiBSearchData, false);
-			}
-			else if(const FString* FiBVersionedSearchData = InAssetData.TagsAndValues.Find(FBlueprintTags::FindInBlueprintsData))
+	if (InAssetData.IsAssetLoaded())
+	{
+		UObject*    AssetObject = InAssetData.GetAsset();
+		UBlueprint* Blueprint   = Handler->RetrieveBlueprint(AssetObject);
+		if (Blueprint)
+		{
+			AddOrUpdateBlueprintSearchMetadata(Blueprint);
+		}
+	}
+	else if (Handler->AssetContainsBlueprint(InAssetData))
+	{
+		// Check first for versioned FiB data (latest codepath)
+		if(const FString* FiBVersionedSearchData = InAssetData.TagsAndValues.Find(FBlueprintTags::FindInBlueprintsData))
+		{
+			if (FiBVersionedSearchData->Len() == 0)
 			{
-				if (FiBVersionedSearchData->Len() == 0)
-				{
-					if (bIsBlueprint)
-					{
-						UncachedBlueprints.Add(InAssetData.ObjectPath);
-					}
-				}
-				else
-				{
-					ExtractUnloadedFiBData(InAssetData, *FiBVersionedSearchData, true);
-				}
+				UncachedAssets.Add(InAssetData.ObjectPath);
 			}
 			else
 			{
-				// The asset is uncached, we will want to inform the user that this is the case
-				// Maps may have no data because they have no blueprints, assume they are empty instead of uncached
-				if (bIsBlueprint)
-				{
-					UncachedBlueprints.Add(InAssetData.ObjectPath);
-				}
+				ExtractUnloadedFiBData(InAssetData, *FiBVersionedSearchData, true);
 			}
 		}
+		// Check for legacy (unversioned) FiB data
+		else if(const FString* FiBSearchData = InAssetData.TagsAndValues.Find("FiB"))
+		{
+			ExtractUnloadedFiBData(InAssetData, *FiBSearchData, false);
+		}
+		// The asset has no FiB data, keep track of it so we can inform the user
+		else
+		{
+			UncachedAssets.Add(InAssetData.ObjectPath);
+		}
 	}
 }
 
 void FFindInBlueprintSearchManager::ExtractUnloadedFiBData(const FAssetData& InAssetData, const FString& InFiBData, bool bIsVersioned)
 {
-	FSearchData NewSearchData;
+	// Add the actual blueprint path to the search map rather than the asset.
+	FName BlueprintPath;
+
+	const FString* BlueprintPathWithinPackage = InAssetData.TagsAndValues.Find(FBlueprintTags::BlueprintPathWithinPackage);
+	if (BlueprintPathWithinPackage)
+	{
+		BlueprintPath = *(InAssetData.PackageName.ToString() + TEXT(".") + (*BlueprintPathWithinPackage));
+	}
+	// ------------------------
+	// Legacy handling for assets saved without FBlueprintTags::BlueprintPathWithinPackage tags.
+	// Prior to this tag, only UBlueprint and UWorld assets were considered.
+	else if (InAssetData.AssetClass == UWorld::StaticClass()->GetFName())
+	{
+		// Assume the path of the level BP to be PersistentLevel.LevelName
+		BlueprintPath = *(InAssetData.ObjectPath.ToString() + TEXT(":PersistentLevel.") + InAssetData.AssetName.ToString());
+	}
+	else
+	{
+		BlueprintPath = InAssetData.ObjectPath;
+	}
+	// End legacy handling
+	// ------------------------
+
+	if (SearchMap.Contains(BlueprintPath))
+	{
+		return;
+	}
 
-	NewSearchData.BlueprintPath = InAssetData.ObjectPath;
+	FSearchData NewSearchData;
+	NewSearchData.BlueprintPath = BlueprintPath;
 	InAssetData.GetTagValue(FBlueprintTags::ParentClassPath, NewSearchData.ParentClass);
 
 	const FString ImplementedInterfaces = InAssetData.GetTagValueRef<FString>(FBlueprintTags::ImplementedInterfaces);
@@ -1586,27 +1591,13 @@ void FFindInBlueprintSearchManager::OnAssetRenamed(const struct FAssetData& InAs
 
 void FFindInBlueprintSearchManager::OnAssetLoaded(UObject* InAsset)
 {
-	UBlueprint* BlueprintAsset = nullptr;
-	FName BlueprintPath = *InAsset->GetPathName();
+	const IBlueprintAssetHandler* Handler = FBlueprintAssetHandler::Get().FindHandler(InAsset->GetClass());
+	UBlueprint* BlueprintAsset = Handler ? Handler->RetrieveBlueprint(InAsset) : nullptr;
 
-	if(UWorld* WorldAsset = Cast<UWorld>(InAsset))
-	{
-		if(WorldAsset->PersistentLevel)
-		{
-			BlueprintAsset = Cast<UBlueprint>((UObject*)WorldAsset->PersistentLevel->GetLevelScriptBlueprint(true));
-			if(BlueprintAsset)
-			{
-				BlueprintPath = *BlueprintAsset->GetPathName();
-			}
-		}
-	}
-	else
+	if (BlueprintAsset)
 	{
-		BlueprintAsset = Cast<UBlueprint>(InAsset);
-	}
+		FName BlueprintPath = *BlueprintAsset->GetPathName();
 
-	if(BlueprintAsset)
-	{
 		// Find and update the item in the search array. Searches may currently be active, this will do no harm to them
 
 		// Confirm that the Blueprint has not been added already, this can occur during duplication of Blueprints.
@@ -1620,7 +1611,8 @@ void FFindInBlueprintSearchManager::OnAssetLoaded(UObject* InAsset)
 			ensureMsgf(!SearchArray[*IndexPtr].Blueprint.IsValid() || SearchArray[*IndexPtr].BlueprintPath == BlueprintPath, TEXT("Blueprint in database has path %s and is being stomped by %s"), *(SearchArray[*IndexPtr].BlueprintPath.ToString()), *BlueprintPath.ToString());
 			SearchArray[*IndexPtr].Blueprint = BlueprintAsset;
 		}
-		UncachedBlueprints.Remove(BlueprintPath);
+
+		UncachedAssets.Remove(FName(*InAsset->GetPathName()));
 	}
 }
 
@@ -1740,18 +1732,7 @@ void FFindInBlueprintSearchManager::AddOrUpdateBlueprintSearchMetadata(UBlueprin
 	// Allow only one thread modify the search data at a time
 	FScopeLock ScopeLock(&SafeModifyCacheCriticalSection);
 
-	FName BlueprintPath;
-	if(FBlueprintEditorUtils::IsLevelScriptBlueprint(InBlueprint))
-	{
-		if(UWorld* World = InBlueprint->GetTypedOuter<UWorld>())
-		{
-			BlueprintPath = *World->GetPathName();
-		}
-	}
-	else
-	{
-		BlueprintPath = *InBlueprint->GetPathName();
-	}
+	FName BlueprintPath = *InBlueprint->GetPathName();
 
 	int32* IndexPtr = SearchMap.Find(BlueprintPath);
 	int32 Index = 0;
@@ -1889,13 +1870,7 @@ FString FFindInBlueprintSearchManager::QuerySingleBlueprint(UBlueprint* InBluepr
 		}
 
 		FName Key = *InBlueprint->GetPathName();
-		if (ULevel* LevelOuter = Cast<ULevel>(InBlueprint->GetOuter()))
-		{
-			if (UWorld* WorldOuter = Cast<UWorld>(LevelOuter->GetOuter()))
-			{
-				Key = *WorldOuter->GetPathName();
-			}
-		}
+
 		int32* ArrayIdx = SearchMap.Find(Key);
 		// This should always be true since we make sure to refresh the search data for this Blueprint when doing the search, unless we do not rebuild the searchable data
 		check((bInRebuildSearchData && ArrayIdx && *ArrayIdx < SearchArray.Num()) || !bInRebuildSearchData);
@@ -2025,8 +2000,11 @@ void FFindInBlueprintSearchManager::BuildCache()
 	TArray< FAssetData > BlueprintAssets;
 	FARFilter ClassFilter;
 	ClassFilter.bRecursiveClasses = true;
-	ClassFilter.ClassNames.Add(UBlueprint::StaticClass()->GetFName());
-	ClassFilter.ClassNames.Add(UWorld::StaticClass()->GetFName());
+
+	for (FName ClassName : FBlueprintAssetHandler::Get().GetRegisteredClassNames())
+	{
+		ClassFilter.ClassNames.Add(ClassName);
+	}
 
 	AssetRegistryModule->Get().GetAssets(ClassFilter, BlueprintAssets);
 	
@@ -2063,7 +2041,7 @@ FString FFindInBlueprintSearchManager::ConvertFTextToHexString(FText InValue)
 	return BytesToHex(SerializedData.GetData(), SerializedData.Num());
 }
 
-void FFindInBlueprintSearchManager::OnCacheAllUncachedBlueprints(bool bInSourceControlActive, bool bInCheckoutAndSave)
+void FFindInBlueprintSearchManager::OnCacheAllUncachedAssets(bool bInSourceControlActive, bool bInCheckoutAndSave)
 {
 	// Multiple threads can be adding to this at the same time
 	FScopeLock ScopeLock(&SafeModifyCacheCriticalSection);
@@ -2073,15 +2051,15 @@ void FFindInBlueprintSearchManager::OnCacheAllUncachedBlueprints(bool bInSourceC
 	{
 		if(bInSourceControlActive && bInCheckoutAndSave)
 		{
-			TArray<FString> UncachedBlueprintStrings;
-			const TArray<FName>& TotalUncachedBlueprints = CachingObject->GetUncachedBlueprintList();
+			TArray<FString> UncachedAssetStrings;
+			const TArray<FName>& TotalUncachedAssets = CachingObject->GetUncachedAssetList();
 		
-			UncachedBlueprintStrings.Reserve(TotalUncachedBlueprints.Num());
-			for (const FName& UncachedBlueprint : TotalUncachedBlueprints)
+			UncachedAssetStrings.Reserve(TotalUncachedAssets.Num());
+			for (const FName& UncachedAsset : TotalUncachedAssets)
 			{
-				UncachedBlueprintStrings.Add(UncachedBlueprint.ToString());
+				UncachedAssetStrings.Add(UncachedAsset.ToString());
 			}
-			FEditorFileUtils::CheckoutPackages(UncachedBlueprintStrings);
+			FEditorFileUtils::CheckoutPackages(UncachedAssetStrings);
 		}
 
 		// Start the cache process.
@@ -2089,7 +2067,7 @@ void FFindInBlueprintSearchManager::OnCacheAllUncachedBlueprints(bool bInSourceC
 	}
 }
 
-void FFindInBlueprintSearchManager::CacheAllUncachedBlueprints(TWeakPtr< SFindInBlueprints > InSourceWidget, FWidgetActiveTimerDelegate& InOutActiveTimerDelegate, FSimpleDelegate InOnFinished/* = FSimpleDelegate()*/, EFiBVersion InMinimiumVersionRequirement/* = EFiBVersion::FIB_VER_LATEST*/)
+void FFindInBlueprintSearchManager::CacheAllUncachedAssets(TWeakPtr< SFindInBlueprints > InSourceWidget, FWidgetActiveTimerDelegate& InOutActiveTimerDelegate, FSimpleDelegate InOnFinished/* = FSimpleDelegate()*/, EFiBVersion InMinimiumVersionRequirement/* = EFiBVersion::FIB_VER_LATEST*/)
 {
 	// Do not start another caching process if one is in progress
 	if(!IsCacheInProgress())
@@ -2106,24 +2084,24 @@ void FFindInBlueprintSearchManager::CacheAllUncachedBlueprints(TWeakPtr< SFindIn
 
 		FText DialogTitle = LOCTEXT("ConfirmIndexAll_Title", "Indexing All");
 		FFormatNamedArguments Args;
-		Args.Add(TEXT("PackageCount"), UncachedBlueprints.Num() + BlueprintsToUpdate.Num());
+		Args.Add(TEXT("PackageCount"), UncachedAssets.Num() + BlueprintsToUpdate.Num());
 
 		FText DialogDisplayText;
 		
-		if (UncachedBlueprints.Num() && BlueprintsToUpdate.Num())
+		if (UncachedAssets.Num() && BlueprintsToUpdate.Num())
 		{
-			Args.Add(TEXT("PackageCount"), UncachedBlueprints.Num() + BlueprintsToUpdate.Num());
-			Args.Add(TEXT("UnindexedCount"), UncachedBlueprints.Num());
+			Args.Add(TEXT("PackageCount"), UncachedAssets.Num() + BlueprintsToUpdate.Num());
+			Args.Add(TEXT("UnindexedCount"), UncachedAssets.Num());
 			Args.Add(TEXT("OutOfDateCount"), BlueprintsToUpdate.Num());
 			DialogDisplayText = FText::Format(LOCTEXT("CacheAllConfirmationMessage_UncachedAndBlueprints", "This process can take a long time and the editor may become unresponsive; there are {PackageCount} ({UnindexedCount} Unindexed/{OutOfDateCount} Out-of-Date) Blueprints to load. \
 																					\n\nWould you like to checkout, load, and save all Blueprints to make this indexing permanent? Otherwise, all Blueprints will still be loaded but you will be required to re-index the next time you start the editor!"), Args);
 		}
-		else if (UncachedBlueprints.Num() && BlueprintsToUpdate.Num() == 0)
+		else if (UncachedAssets.Num() && BlueprintsToUpdate.Num() == 0)
 		{
 			DialogDisplayText = FText::Format(LOCTEXT("CacheAllConfirmationMessage_UncachedOnly", "This process can take a long time and the editor may become unresponsive; there are {PackageCount} unindexed Blueprints to load. \
 																					 \n\nWould you like to checkout, load, and save all Blueprints to make this indexing permanent? Otherwise, all Blueprints will still be loaded but you will be required to re-index the next time you start the editor!"), Args);
 		}
-		else if (UncachedBlueprints.Num() == 0 && BlueprintsToUpdate.Num())
+		else if (UncachedAssets.Num() == 0 && BlueprintsToUpdate.Num())
 		{
 			DialogDisplayText = FText::Format(LOCTEXT("CacheAllConfirmationMessage_BlueprintsOnly", "This process can take a long time and the editor may become unresponsive; there are {PackageCount} out-of-date Blueprints to load. \
 																					 \n\nWould you like to checkout, load, and save all Blueprints to make this indexing permanent? Otherwise, all Blueprints will still be loaded but you will be required to re-index the next time you start the editor!"), Args);
@@ -2136,23 +2114,23 @@ void FFindInBlueprintSearchManager::CacheAllUncachedBlueprints(TWeakPtr< SFindIn
 		{
 			FailedToCachePaths.Empty();
 			
-			TSet<FName> TempUncachedBlueprints;
-			TempUncachedBlueprints.Append(UncachedBlueprints);
-			TempUncachedBlueprints.Append(BlueprintsToUpdate);
+			TSet<FName> TempUncachedAssets;
+			TempUncachedAssets.Append(UncachedAssets);
+			TempUncachedAssets.Append(BlueprintsToUpdate);
 
 			const bool bCheckoutAndSave = ReturnValue == EAppReturnType::Yes;
-			CachingObject = new FCacheAllBlueprintsTickableObject(MoveTemp(TempUncachedBlueprints), bCheckoutAndSave, InOnFinished);
+			CachingObject = new FCacheAllBlueprintsTickableObject(MoveTemp(TempUncachedAssets), bCheckoutAndSave, InOnFinished);
 			InOutActiveTimerDelegate.BindRaw(CachingObject, &FCacheAllBlueprintsTickableObject::Tick);
 
 			const bool bIsSourceControlEnabled = ISourceControlModule::Get().IsEnabled();
 			if(!bIsSourceControlEnabled && bCheckoutAndSave)
 			{
 				// Offer to start up Source Control
-				ISourceControlModule::Get().ShowLoginDialog(FSourceControlLoginClosed::CreateRaw(this, &FFindInBlueprintSearchManager::OnCacheAllUncachedBlueprints, bCheckoutAndSave), ELoginWindowMode::Modeless, EOnLoginWindowStartup::PreserveProvider);
+				ISourceControlModule::Get().ShowLoginDialog(FSourceControlLoginClosed::CreateRaw(this, &FFindInBlueprintSearchManager::OnCacheAllUncachedAssets, bCheckoutAndSave), ELoginWindowMode::Modeless, EOnLoginWindowStartup::PreserveProvider);
 			}
 			else
 			{
-				OnCacheAllUncachedBlueprints(bIsSourceControlEnabled, bCheckoutAndSave);
+				OnCacheAllUncachedAssets(bIsSourceControlEnabled, bCheckoutAndSave);
 			}
 
 			SourceCachingWidget = InSourceWidget;
@@ -2203,12 +2181,12 @@ float FFindInBlueprintSearchManager::GetCacheProgress() const
 	return ReturnCacheValue;
 }
 
-int32 FFindInBlueprintSearchManager::GetNumberUncachedBlueprints() const
+int32 FFindInBlueprintSearchManager::GetNumberUncachedAssets() const
 {
-	int32 ReturnCount = UncachedBlueprints.Num();
+	int32 ReturnCount = UncachedAssets.Num();
 	if (CachingObject)
 	{
-		ReturnCount = CachingObject->GetUncachedBlueprintCount();
+		ReturnCount = CachingObject->GetUncachedAssetCount();
 	}
 	return ReturnCount;
 }
diff --git a/Engine/Source/Editor/Kismet/Private/FindInBlueprints.cpp b/Engine/Source/Editor/Kismet/Private/FindInBlueprints.cpp
--- a/Engine/Source/Editor/Kismet/Private/FindInBlueprints.cpp
+++ b/Engine/Source/Editor/Kismet/Private/FindInBlueprints.cpp
@@ -720,7 +720,7 @@ void SFindInBlueprints::ConditionallyAddCacheBar()
 	FFindInBlueprintSearchManager& FindInBlueprintManager = FFindInBlueprintSearchManager::Get();
 
 	// Do not add a second cache bar and do not add it when there are no uncached Blueprints
-	if(FindInBlueprintManager.GetNumberUncachedBlueprints() > 0 || FindInBlueprintManager.GetFailedToCacheCount() > 0)
+	if(FindInBlueprintManager.GetNumberUncachedAssets() > 0 || FindInBlueprintManager.GetFailedToCacheCount() > 0)
 	{
 		if(MainVerticalBox.IsValid() && !CacheBarSlot.IsValid())
 		{
@@ -788,7 +788,7 @@ void SFindInBlueprints::ConditionallyAddCacheBar()
 							.AutoWidth()
 							[
 								SNew(STextBlock)
-								.Text(this, &SFindInBlueprints::GetUncachedBlueprintWarningText)
+								.Text(this, &SFindInBlueprints::GetUncachedAssetWarningText)
 								.ColorAndOpacity( FCoreStyle::Get().GetColor("ErrorReporting.ForegroundColor") )
 							]
 
@@ -1270,7 +1270,7 @@ FReply SFindInBlueprints::OnCacheAllBlueprints(FSimpleDelegate InOnFinished/* =
 	{
 		// Request from the SearchManager a delegate to use for ticking the cache system.
 		FWidgetActiveTimerDelegate WidgetActiveTimer;
-		FFindInBlueprintSearchManager::Get().CacheAllUncachedBlueprints(SharedThis(this), WidgetActiveTimer, InOnFinished, InMinimiumVersionRequirement);
+		FFindInBlueprintSearchManager::Get().CacheAllUncachedAssets(SharedThis(this), WidgetActiveTimer, InOnFinished, InMinimiumVersionRequirement);
 		RegisterActiveTimer(0.f, WidgetActiveTimer);
 	}
 
@@ -1310,7 +1310,7 @@ EVisibility SFindInBlueprints::GetCacheAllButtonVisibility() const
 EVisibility SFindInBlueprints::GetCachingBarVisibility() const
 {
 	FFindInBlueprintSearchManager& FindInBlueprintManager = FFindInBlueprintSearchManager::Get();
-	return (FindInBlueprintManager.GetNumberUncachedBlueprints() > 0 || FindInBlueprintManager.GetFailedToCacheCount())? EVisibility::Visible : EVisibility::Collapsed;
+	return (FindInBlueprintManager.GetNumberUncachedAssets() > 0 || FindInBlueprintManager.GetFailedToCacheCount())? EVisibility::Visible : EVisibility::Collapsed;
 }
 
 EVisibility SFindInBlueprints::GetCachingBlueprintNameVisiblity() const
@@ -1340,7 +1340,7 @@ FSlateColor SFindInBlueprints::GetCachingBarColor() const
 	return ReturnColor;
 }
 
-FText SFindInBlueprints::GetUncachedBlueprintWarningText() const
+FText SFindInBlueprints::GetUncachedAssetWarningText() const
 {
 	FFindInBlueprintSearchManager& FindInBlueprintManager = FFindInBlueprintSearchManager::Get();
 
@@ -1348,7 +1348,7 @@ FText SFindInBlueprints::GetUncachedBlueprintWarningText() const
 
 	// The number of unindexed Blueprints is the total of those that failed to cache and those that haven't been attempted yet.
 	FFormatNamedArguments Args;
-	Args.Add(TEXT("Count"), FindInBlueprintManager.GetNumberUncachedBlueprints() + OutOfDateWithLastSearchBPCount);
+	Args.Add(TEXT("Count"), FindInBlueprintManager.GetNumberUncachedAssets() + OutOfDateWithLastSearchBPCount);
 
 	FText ReturnDisplayText;
 	if(IsCacheInProgress())
@@ -1359,17 +1359,17 @@ FText SFindInBlueprints::GetUncachedBlueprintWarningText() const
 	}
 	else
 	{
-		Args.Add(TEXT("UnindexedCount"), FindInBlueprintManager.GetNumberUncachedBlueprints());
+		Args.Add(TEXT("UnindexedCount"), FindInBlueprintManager.GetNumberUncachedAssets());
 		Args.Add(TEXT("OutOfDateCount"), OutOfDateWithLastSearchBPCount);
 
-		ReturnDisplayText = FText::Format(LOCTEXT("UncachedBlueprints", "Search incomplete. {Count} ({UnindexedCount} Unindexed/{OutOfDateCount} Out-of-Date) Blueprints need to be indexed!"), Args);
+		ReturnDisplayText = FText::Format(LOCTEXT("UncachedAssets", "Search incomplete. {Count} ({UnindexedCount} Unindexed/{OutOfDateCount} Out-of-Date) Blueprints need to be indexed!"), Args);
 
 		if (FailedToCacheCount > 0)
 		{
 			FFormatNamedArguments ArgsWithCacheFails;
 			Args.Add(TEXT("BaseMessage"), ReturnDisplayText);
 			Args.Add(TEXT("CacheFails"), FailedToCacheCount);
-			ReturnDisplayText = FText::Format(LOCTEXT("UncachedBlueprintsWithCacheFails", "{BaseMessage} {CacheFails} Blueprints failed to cache."), Args);
+			ReturnDisplayText = FText::Format(LOCTEXT("UncachedAssetsWithCacheFails", "{BaseMessage} {CacheFails} Blueprints failed to cache."), Args);
 		}
 	}
 
diff --git a/Engine/Source/Editor/Kismet/Public/BlueprintAssetHandler.h b/Engine/Source/Editor/Kismet/Public/BlueprintAssetHandler.h
new file mode 100644
--- /dev/null
+++ b/Engine/Source/Editor/Kismet/Public/BlueprintAssetHandler.h
@@ -0,0 +1,107 @@
+// Copyright 1998-2018 Epic Games, Inc. All Rights Reserved.
+
+#pragma once
+
+#include "CoreTypes.h"
+#include "Templates/UniquePtr.h"
+#include "Containers/ArrayView.h"
+#include "Containers/Array.h"
+
+struct FAssetData;
+
+class UBlueprint;
+class UObject;
+class UClass;
+
+/**
+ * Interface used to define how to interact with a blueprint within an asset
+ */
+class IBlueprintAssetHandler
+{
+public:
+	virtual ~IBlueprintAssetHandler() {}
+
+	/**
+	 * Retrieve the blueprint from the specified asset object
+	 *
+	 *@param InAsset         The asset object to retrieve the blueprint from
+	 *@return The blueprint contained within the specified asset, or nullptr if non exists
+	 */
+	virtual UBlueprint* RetrieveBlueprint(UObject* InAsset) const = 0;
+
+
+	/**
+	 * Check whether the specified asset registry data contains a blueprint
+	 *
+	 *@param InAssetData     The asset object to retrieve the blueprint from
+	 *@return True if the asset contains a blueprint, false otherwise
+	 */
+	virtual bool AssetContainsBlueprint(const FAssetData& InAssetData) const = 0;
+};
+
+
+/**
+ * Singleton class that marshals different blueprint asset handlers for different asset class types
+ */
+class FBlueprintAssetHandler
+{
+public:
+
+	/**
+	 * Retrieve the singleton instance of this class
+	 */
+	KISMET_API static FBlueprintAssetHandler& Get();
+
+
+	/**
+	 * Get all the currently registered class names
+	 */
+	TArrayView<const FName> GetRegisteredClassNames() const
+	{
+		return ClassNames;
+	}
+
+
+	/**
+	 * Register an asset for the specified class name
+	 * @note: Any assets whose class is a child of the specified class will use this handler (unless there is a more specific handler registered)
+	 *
+	 * @param ClassName      The name of the class this handler relates to.
+	 */
+	template<typename HandlerType>
+	void RegisterHandler(FName ClassName)
+	{
+		RegisterHandler(ClassName, MakeUnique<HandlerType>());
+	}
+
+
+	/**
+	 * Register an asset for the specified class name
+	 * @note: Any assets whose class is a child of the specified class will use this handler (unless there is a more specific handler registered)
+	 *
+	 * @param ClassName      The name of the class this handler relates to.
+	 * @param InHandler      The implementation of the handler to use
+	 */
+	KISMET_API void RegisterHandler(FName ClassName, TUniquePtr<IBlueprintAssetHandler>&& InHandler);
+
+
+	/**
+	 * Find a handler that applies to the specified class
+	 *
+	 * @param InClass        The class to find a handler for
+	 * @return A valid asset handler, or nullptr if none exists for this class
+	 */
+	KISMET_API const IBlueprintAssetHandler* FindHandler(const UClass* InClass) const;
+
+
+private:
+
+	/** Private constructor - use singleton accessor (::Get) */
+	FBlueprintAssetHandler();
+
+	/** Unsorted array of class names, one per-handler below */
+	TArray<FName> ClassNames;
+
+	/** Array of handlers that relate to the class names above */
+	TArray<TUniquePtr<IBlueprintAssetHandler>> Handlers;
+};
\ No newline at end of file
diff --git a/Engine/Source/Editor/Kismet/Public/FindInBlueprintManager.h b/Engine/Source/Editor/Kismet/Public/FindInBlueprintManager.h
--- a/Engine/Source/Editor/Kismet/Public/FindInBlueprintManager.h
+++ b/Engine/Source/Editor/Kismet/Public/FindInBlueprintManager.h
@@ -406,7 +406,7 @@ public:
 	static FString ConvertFTextToHexString(FText InValue);
 
 	/** Returns the number of uncached Blueprints */
-	int32 GetNumberUncachedBlueprints() const;
+	int32 GetNumberUncachedAssets() const;
 
 	/**
 	 * Starts caching all uncached Blueprints at a rate of 1 per tick
@@ -416,7 +416,7 @@ public:
 	 * @param InOnFinished					Callback when caching is finished
 	 * @param InMinimiumVersionRequirement	Minimum version requirement for caching, any Blueprints below this version will be re-indexed
 	 */
-	void CacheAllUncachedBlueprints(TWeakPtr< class SFindInBlueprints > InSourceWidget, FWidgetActiveTimerDelegate& InOutActiveTimerDelegate, FSimpleDelegate InOnFinished = FSimpleDelegate(), EFiBVersion InMinimiumVersionRequirement = EFiBVersion::FIB_VER_LATEST);
+	void CacheAllUncachedAssets(TWeakPtr< class SFindInBlueprints > InSourceWidget, FWidgetActiveTimerDelegate& InOutActiveTimerDelegate, FSimpleDelegate InOnFinished = FSimpleDelegate(), EFiBVersion InMinimiumVersionRequirement = EFiBVersion::FIB_VER_LATEST);
 	
 	/**
 	 * Starts the actual caching process
@@ -424,7 +424,7 @@ public:
 	 * @param bInSourceControlActive		TRUE if source control is active
 	 * @param bInCheckoutAndSave			TRUE if the system should checkout and save all assets that need to be reindexed
 	 */
-	void OnCacheAllUncachedBlueprints(bool bInSourceControlActive, bool bInCheckoutAndSave);
+	void OnCacheAllUncachedAssets(bool bInSourceControlActive, bool bInCheckoutAndSave);
 
 	/** Stops the caching process where it currently is at, the rest can be continued later */
 	void CancelCacheAll(SFindInBlueprints* InFindInBlueprintWidget);
@@ -571,8 +571,8 @@ protected:
 	/** FindInBlueprints widget that started the cache process */
 	TWeakPtr<SFindInBlueprints> SourceCachingWidget;
 
-	/** Blueprint paths that have not been cached for searching due to lack of data, this means that they are either older Blueprints, or the DDC cannot find the data */
-	TSet<FName> UncachedBlueprints;
+	/** Asset paths that have not been cached for searching due to lack of FiB data, this means that they are either older Blueprints, or the DDC cannot find the data */
+	TSet<FName> UncachedAssets;
 
 	/** List of paths for Blueprints that failed to cache */
 	TSet<FName> FailedToCachePaths;
diff --git a/Engine/Source/Editor/Kismet/Public/FindInBlueprints.h b/Engine/Source/Editor/Kismet/Public/FindInBlueprints.h
--- a/Engine/Source/Editor/Kismet/Public/FindInBlueprints.h
+++ b/Engine/Source/Editor/Kismet/Public/FindInBlueprints.h
@@ -285,7 +285,7 @@ private:
 	FReply OnRemoveCacheBar();
 
 	/** Callback to return the cache bar's display text, informing the user of the situation */
-	FText GetUncachedBlueprintWarningText() const;
+	FText GetUncachedAssetWarningText() const;
 
 	/** Callback to return the cache bar's current indexing Blueprint name */
 	FText GetCurrentCacheBlueprintName() const;
diff --git a/Engine/Source/Runtime/CoreUObject/Private/Blueprint/BlueprintSupport.cpp b/Engine/Source/Runtime/CoreUObject/Private/Blueprint/BlueprintSupport.cpp
--- a/Engine/Source/Runtime/CoreUObject/Private/Blueprint/BlueprintSupport.cpp
+++ b/Engine/Source/Runtime/CoreUObject/Private/Blueprint/BlueprintSupport.cpp
@@ -43,6 +43,7 @@ const FName FBlueprintTags::FindInBlueprintsData(TEXT("FiBData"));
 const FName FBlueprintTags::NumReplicatedProperties(TEXT("NumReplicatedProperties"));
 const FName FBlueprintTags::NumNativeComponents(TEXT("NativeComponents"));
 const FName FBlueprintTags::NumBlueprintComponents(TEXT("BlueprintComponents"));
+const FName FBlueprintTags::BlueprintPathWithinPackage(TEXT("BlueprintPath"));
 
 /**
  * Defined in BlueprintSupport.cpp
diff --git a/Engine/Source/Runtime/CoreUObject/Public/Blueprint/BlueprintSupport.h b/Engine/Source/Runtime/CoreUObject/Public/Blueprint/BlueprintSupport.h
--- a/Engine/Source/Runtime/CoreUObject/Public/Blueprint/BlueprintSupport.h
+++ b/Engine/Source/Runtime/CoreUObject/Public/Blueprint/BlueprintSupport.h
@@ -40,6 +40,8 @@ struct COREUOBJECT_API FBlueprintTags
 	static const FName NumNativeComponents;
 	/** Number of blueprint components */
 	static const FName NumBlueprintComponents;
+	/** The subpath of a blueprint contained within the asset. Used to determine whether, and where a blueprint exists in a package. */
+	static const FName BlueprintPathWithinPackage;
 };
 
 struct FBlueprintWarningDeclaration
diff --git a/Engine/Source/Runtime/Engine/Classes/Engine/Blueprint.h b/Engine/Source/Runtime/Engine/Classes/Engine/Blueprint.h
--- a/Engine/Source/Runtime/Engine/Classes/Engine/Blueprint.h
+++ b/Engine/Source/Runtime/Engine/Classes/Engine/Blueprint.h
@@ -761,6 +761,11 @@ public:
 	/** Some Blueprints (and classes) can recompile while we are debugging a live session. This function controls whether this can occur. */
 	virtual bool CanRecompileWhilePlayingInEditor() const { return false; }
 
+	/**
+	 * Check whether this blueprint can be nativized or not
+	 */
+	virtual bool SupportsNativization(FText* OutReason = nullptr) const;
+
 private:
 
 	/** Sets the current object being debugged */
diff --git a/Engine/Source/Runtime/Engine/Private/Blueprint.cpp b/Engine/Source/Runtime/Engine/Private/Blueprint.cpp
--- a/Engine/Source/Runtime/Engine/Private/Blueprint.cpp
+++ b/Engine/Source/Runtime/Engine/Private/Blueprint.cpp
@@ -810,6 +810,29 @@ UClass* UBlueprint::GetBlueprintClass() const
 	return UBlueprintGeneratedClass::StaticClass();
 }
 
+bool UBlueprint::SupportsNativization(FText* OutReason) const
+{
+	// Previously commented out in FBlueprintNativeCodeGenModule::IsTargetedForReplacement - should 'const' blueprints be nativized??
+	//BPTYPE_Const,		// What is a "const" Blueprint?
+	if (BlueprintType == BPTYPE_MacroLibrary)
+	{
+		if (OutReason)
+		{
+			*OutReason = NSLOCTEXT("Blueprint", "MacroLibraryNativizationReason", "Macro Libraries cannot be nativized.");
+		}
+		return false;
+	}
+	else if (BlueprintType == BPTYPE_LevelScript)
+	{
+		if (OutReason)
+		{
+			*OutReason = NSLOCTEXT("Blueprint", "LevelScriptNativizationReason", "Level Blueprints cannot be nativized.");
+		}
+		return false;
+	}
+	return true;
+}
+
 void UBlueprint::SetObjectBeingDebugged(UObject* NewObject)
 {
 	// Unregister the old object
@@ -928,6 +951,8 @@ void UBlueprint::GetAssetRegistryTags(TArray<FAssetRegistryTag>& OutTags) const
 		NativeParentClassName = TEXT("None");
 	}
 
+
+	OutTags.Add(FAssetRegistryTag(FBlueprintTags::BlueprintPathWithinPackage, GetPathName(GetOutermost()), FAssetRegistryTag::TT_Hidden));
 	OutTags.Add(FAssetRegistryTag(FBlueprintTags::GeneratedClassPath, GeneratedClassVal, FAssetRegistryTag::TT_Hidden));
 	OutTags.Add(FAssetRegistryTag(FBlueprintTags::NativeParentClassPath, NativeParentClassName, FAssetRegistryTag::TT_Alphabetical));
 
