Fixed missing include

diff --git a/Engine/Source/Editor/Kismet/Public/BlueprintAssetHandler.h b/Engine/Source/Editor/Kismet/Public/BlueprintAssetHandler.h
--- a/Engine/Source/Editor/Kismet/Public/BlueprintAssetHandler.h
+++ b/Engine/Source/Editor/Kismet/Public/BlueprintAssetHandler.h
@@ -6,6 +6,7 @@
 #include "Templates/UniquePtr.h"
 #include "Containers/ArrayView.h"
 #include "Containers/Array.h"
+#include "UObject/NameTypes.h"
 
 struct FAssetData;
 
