Updated variable details panels to not display unusable metadata options for UMG widget references

diff --git a/Engine/Source/Editor/Kismet/Private/BlueprintDetailsCustomization.cpp b/Engine/Source/Editor/Kismet/Private/BlueprintDetailsCustomization.cpp
--- a/Engine/Source/Editor/Kismet/Private/BlueprintDetailsCustomization.cpp
+++ b/Engine/Source/Editor/Kismet/Private/BlueprintDetailsCustomization.cpp
@@ -1056,13 +1056,23 @@ TSharedRef<ITableRow> FBlueprintVarActionDetails::OnGenerateWidgetForPropertyLis
 		];
 }
 
-bool FBlueprintVarActionDetails::IsASCSVariable(UProperty* VariableProperty) const
+bool FBlueprintVarActionDetails::IsAUserVariable(UProperty* VariableProperty) const
 {
 	UObjectProperty* VariableObjProp = VariableProperty ? Cast<UObjectProperty>(VariableProperty) : NULL;
 
-	if (VariableObjProp != NULL && VariableObjProp->PropertyClass != NULL && VariableObjProp->PropertyClass->IsChildOf(UActorComponent::StaticClass()))
+	if (VariableObjProp != NULL && VariableObjProp->PropertyClass != NULL)
+	{
+		return FBlueprintEditorUtils::IsVariableCreatedByBlueprint(GetBlueprintObj(), VariableObjProp);
+	}
+	return true;
+}
+
+bool FBlueprintVarActionDetails::IsASCSVariable(UProperty* VariableProperty) const
+{
+	UObjectProperty* VariableObjProp = VariableProperty ? Cast<UObjectProperty>(VariableProperty) : NULL;
+	if (VariableObjProp != NULL && VariableObjProp->PropertyClass != NULL)
 	{
-		return !FBlueprintEditorUtils::IsVariableCreatedByBlueprint(GetBlueprintObj(), VariableObjProp);
+		return (!IsAUserVariable(VariableProperty) && VariableObjProp->PropertyClass->IsChildOf(UActorComponent::StaticClass()));
 	}
 	return false;
 }
@@ -1607,7 +1617,7 @@ EVisibility FBlueprintVarActionDetails::ShowEditableCheckboxVisibilty() const
 	UProperty* VariableProperty = CachedVariableProperty.Get();
 	if (VariableProperty && GetPropertyOwnerBlueprint())
 	{
-		if (IsABlueprintVariable(VariableProperty) && !IsASCSVariable(VariableProperty))
+		if (IsABlueprintVariable(VariableProperty) && IsAUserVariable(VariableProperty))
 		{
 			return EVisibility::Visible;
 		}
@@ -1641,7 +1651,7 @@ EVisibility FBlueprintVarActionDetails::ShowReadOnlyCheckboxVisibilty() const
 	UProperty* VariableProperty = CachedVariableProperty.Get();
 	if (VariableProperty && GetPropertyOwnerBlueprint())
 	{
-		if (IsABlueprintVariable(VariableProperty) && !IsASCSVariable(VariableProperty))
+		if (IsABlueprintVariable(VariableProperty) && IsAUserVariable(VariableProperty))
 		{
 			return EVisibility::Visible;
 		}
@@ -1757,7 +1767,7 @@ EVisibility FBlueprintVarActionDetails::ExposeOnSpawnVisibility() const
 		FEdGraphPinType VariablePinType;
 		K2Schema->ConvertPropertyToPinType(VariableProperty, VariablePinType);
 
-		const bool bShowPrivacySetting = IsABlueprintVariable(VariableProperty) && !IsASCSVariable(VariableProperty);
+		const bool bShowPrivacySetting = IsABlueprintVariable(VariableProperty) && IsAUserVariable(VariableProperty);
 		if (bShowPrivacySetting && (K2Schema->FindSetVariableByNameFunction(VariablePinType) != NULL))
 		{
 			return EVisibility::Visible;
@@ -1798,7 +1808,7 @@ EVisibility FBlueprintVarActionDetails::ExposePrivateVisibility() const
 	UProperty* Property = CachedVariableProperty.Get();
 	if (Property && GetPropertyOwnerBlueprint())
 	{
-		if (IsABlueprintVariable(Property) && !IsASCSVariable(Property))
+		if (IsABlueprintVariable(Property) && IsAUserVariable(Property))
 		{
 			return EVisibility::Visible;
 		}
@@ -2004,7 +2014,7 @@ EVisibility FBlueprintVarActionDetails::ExposeConfigVisibility() const
 	UProperty* Property = CachedVariableProperty.Get();
 	if (Property)
 	{
-		if (IsABlueprintVariable(Property) && !IsASCSVariable(Property))
+		if (IsABlueprintVariable(Property) && IsAUserVariable(Property))
 		{
 			return EVisibility::Visible;
 		}
@@ -2336,7 +2346,7 @@ EVisibility FBlueprintVarActionDetails::ReplicationVisibility() const
 	UProperty* VariableProperty = CachedVariableProperty.Get();
 	if(VariableProperty)
 	{
-		if (!IsASCSVariable(VariableProperty) && IsABlueprintVariable(VariableProperty))
+		if (IsAUserVariable(VariableProperty) && IsABlueprintVariable(VariableProperty))
 		{
 			return EVisibility::Visible;
 		}
@@ -2376,7 +2386,7 @@ EVisibility FBlueprintVarActionDetails::GetTransientVisibility() const
 	UProperty* VariableProperty = CachedVariableProperty.Get();
 	if (VariableProperty)
 	{
-		if (IsABlueprintVariable(VariableProperty) && !IsASCSVariable(VariableProperty))
+		if (IsABlueprintVariable(VariableProperty) && IsAUserVariable(VariableProperty))
 		{
 			return EVisibility::Visible;
 		}
@@ -2409,7 +2419,7 @@ EVisibility FBlueprintVarActionDetails::GetSaveGameVisibility() const
 	UProperty* VariableProperty = CachedVariableProperty.Get();
 	if (VariableProperty)
 	{
-		if (IsABlueprintVariable(VariableProperty) && !IsASCSVariable(VariableProperty))
+		if (IsABlueprintVariable(VariableProperty) && IsAUserVariable(VariableProperty))
 		{
 			return EVisibility::Visible;
 		}
@@ -2442,7 +2452,7 @@ EVisibility FBlueprintVarActionDetails::GetAdvancedDisplayVisibility() const
 	UProperty* VariableProperty = CachedVariableProperty.Get();
 	if (VariableProperty)
 	{
-		if (IsABlueprintVariable(VariableProperty))
+		if (IsABlueprintVariable(VariableProperty) && IsAUserVariable(VariableProperty))
 		{
 			return EVisibility::Visible;
 		}
@@ -2509,7 +2519,7 @@ EVisibility FBlueprintVarActionDetails::IsTooltipEditVisible() const
 	UProperty* VariableProperty = CachedVariableProperty.Get();
 	if (VariableProperty)
 	{
-		if ((IsABlueprintVariable(VariableProperty) && !IsASCSVariable(VariableProperty)) || IsALocalVariable(VariableProperty))
+		if ((IsABlueprintVariable(VariableProperty) && IsAUserVariable(VariableProperty)) || IsALocalVariable(VariableProperty))
 		{
 			return EVisibility::Visible;
 		}
diff --git a/Engine/Source/Editor/Kismet/Private/BlueprintDetailsCustomization.h b/Engine/Source/Editor/Kismet/Private/BlueprintDetailsCustomization.h
--- a/Engine/Source/Editor/Kismet/Private/BlueprintDetailsCustomization.h
+++ b/Engine/Source/Editor/Kismet/Private/BlueprintDetailsCustomization.h
@@ -118,6 +118,7 @@ private:
 	FName GetVariableName() const;
 
 	/** Commonly queried attributes about the schema action */
+	bool IsAUserVariable(UProperty* VariableProperty) const;
 	bool IsASCSVariable(UProperty* VariableProperty) const;
 	bool IsABlueprintVariable(UProperty* VariableProperty) const;
 	bool IsALocalVariable(UProperty* VariableProperty) const;
