Prevent user from entering 'None' as parameter name on a blueprint function

diff --git a/Engine/Source/Editor/Kismet/Private/BlueprintDetailsCustomization.cpp b/Engine/Source/Editor/Kismet/Private/BlueprintDetailsCustomization.cpp
--- a/Engine/Source/Editor/Kismet/Private/BlueprintDetailsCustomization.cpp
+++ b/Engine/Source/Editor/Kismet/Private/BlueprintDetailsCustomization.cpp
@@ -4137,6 +4137,12 @@ bool FBaseBlueprintGraphActionDetails::OnVerifyPinRename(UK2Node_EditablePinBase
 		return false;
 	}
 
+	if (InNewName == TEXT("None"))
+	{
+		OutErrorMessage = FText::Format( LOCTEXT("PinNameNone", "'None' is a reserved name"), FText::AsNumber( NAME_SIZE ) );
+		return false;
+	}
+
 	if (InTargetNode)
 	{
 		// Check if the name conflicts with any of the other internal UFunction's property names (local variables and parameters).
