PR #5028: UE-63225: Use correct source BP when retrieving the tooltip text (Contributed by [REDACTED_PERSON])

diff --git a/Engine/Source/Editor/Kismet/Private/SBlueprintPalette.cpp b/Engine/Source/Editor/Kismet/Private/SBlueprintPalette.cpp
--- a/Engine/Source/Editor/Kismet/Private/SBlueprintPalette.cpp
+++ b/Engine/Source/Editor/Kismet/Private/SBlueprintPalette.cpp
@@ -135,7 +135,8 @@ static FString GetVarTooltip(UBlueprint* InBlueprint, UClass* VarClass, FName Va
 			}
 			else
 			{
-				FBlueprintEditorUtils::GetBlueprintVariableMetaData(InBlueprint, VarName, nullptr, TEXT("tooltip"), ResultTooltip);
+				const UBlueprint* SourceBlueprint = SourceClass ? Cast<UBlueprint>(SourceClass->ClassGeneratedBy) : nullptr;
+				FBlueprintEditorUtils::GetBlueprintVariableMetaData(SourceBlueprint ? SourceBlueprint : InBlueprint, VarName, nullptr, TEXT("tooltip"), ResultTooltip);
 			}
 		}
 	}
