Allow the Feature Level Preview to work within the Blueprint Viewport

diff --git a/Engine/Source/Editor/Kismet/Private/SSCSEditorViewport.cpp b/Engine/Source/Editor/Kismet/Private/SSCSEditorViewport.cpp
--- a/Engine/Source/Editor/Kismet/Private/SSCSEditorViewport.cpp
+++ b/Engine/Source/Editor/Kismet/Private/SSCSEditorViewport.cpp
@@ -237,12 +237,42 @@ void SSCSEditorViewport::Construct(const FArguments& InArgs)
 
 	SEditorViewport::Construct( SEditorViewport::FArguments() );
 
+	// Restore last used feature level
+	if (ViewportClient.IsValid())
+	{
+		UWorld* World = ViewportClient->GetPreviewScene()->GetWorld();
+		if (World != nullptr)
+		{
+			World->ChangeFeatureLevel(GWorld->FeatureLevel);
+		}
+	}
+
+	// Use a delegate to inform the attached world of feature level changes.
+	UEditorEngine* Editor = (UEditorEngine*)GEngine;
+	PreviewFeatureLevelChangedHandle = Editor->OnPreviewFeatureLevelChanged().AddLambda([this](ERHIFeatureLevel::Type NewFeatureLevel)
+		{
+			if(ViewportClient.IsValid())
+			{
+				UWorld* World = ViewportClient->GetPreviewScene()->GetWorld();
+				if (World != nullptr)
+				{
+					World->ChangeFeatureLevel(NewFeatureLevel);
+
+					// Refresh the preview scene. Don't change the camera.
+					RequestRefresh(false);
+				}
+			}
+		});
+
 	// Refresh the preview scene
 	RequestRefresh(true);
 }
 
 SSCSEditorViewport::~SSCSEditorViewport()
 {
+	UEditorEngine* Editor = (UEditorEngine*)GEngine;
+	Editor->OnPreviewFeatureLevelChanged().Remove(PreviewFeatureLevelChangedHandle);
+
 	if(ViewportClient.IsValid())
 	{
 		// Reset this to ensure it's no longer in use after destruction
@@ -277,6 +307,19 @@ TSharedPtr<SWidget> SSCSEditorViewport::MakeViewportToolbar()
 		.IsEnabled(FSlateApplication::Get().GetNormalExecutionAttribute());
 }
 
+void SSCSEditorViewport::PopulateViewportOverlays(TSharedRef<class SOverlay> Overlay)
+{
+	SEditorViewport::PopulateViewportOverlays(Overlay);
+
+	// add the feature level display widget
+	Overlay->AddSlot()
+		.VAlign(VAlign_Bottom)
+		.HAlign(HAlign_Right)
+		.Padding(5.0f)
+		[
+			BuildFeatureLevelWidget()
+		];
+}
 
 void SSCSEditorViewport::BindCommands()
 {
diff --git a/Engine/Source/Editor/Kismet/Private/SSCSEditorViewport.h b/Engine/Source/Editor/Kismet/Private/SSCSEditorViewport.h
--- a/Engine/Source/Editor/Kismet/Private/SSCSEditorViewport.h
+++ b/Engine/Source/Editor/Kismet/Private/SSCSEditorViewport.h
@@ -89,6 +89,7 @@ protected:
 	/** SEditorViewport interface */
 	virtual TSharedRef<class FEditorViewportClient> MakeEditorViewportClient() override;
 	virtual TSharedPtr<class SWidget> MakeViewportToolbar() override;
+	virtual void PopulateViewportOverlays(TSharedRef<class SOverlay> Overlay) override;
 	virtual void BindCommands() override;
 
 private:
@@ -107,4 +108,7 @@ private:
 
 	/** The owner dock tab for this viewport. */
 	TWeakPtr<SDockTab> OwnerTab;
+
+	/** Handle to the registered OnPreviewFeatureLevelChanged delegate. */
+	FDelegateHandle PreviewFeatureLevelChangedHandle;
 };
diff --git a/Engine/Source/Editor/UnrealEd/Public/EditorViewportClient.h b/Engine/Source/Editor/UnrealEd/Public/EditorViewportClient.h
--- a/Engine/Source/Editor/UnrealEd/Public/EditorViewportClient.h
+++ b/Engine/Source/Editor/UnrealEd/Public/EditorViewportClient.h
@@ -277,6 +277,16 @@ public:
 	FEditorViewportClient(const FEditorViewportClient&) = delete;
 	FEditorViewportClient& operator=(const FEditorViewportClient&) = delete;
 
+	/**
+	 * Retrieves the FPreviewScene used by this instance of FEditorViewportClient.
+	 *
+	 * @return		The internal FPreviewScene pointer.
+	 */
+	FPreviewScene* GetPreviewScene()
+	{
+		return PreviewScene;
+	}
+
 	/**
 	 * Toggles whether or not the viewport updates in realtime and returns the updated state.
 	 *
