Fix crash when dragging from Blueprint Editor components list when component does not have a corresponding property

diff --git a/Engine/Source/Editor/Kismet/Private/BPVariableDragDropAction.h b/Engine/Source/Editor/Kismet/Private/BPVariableDragDropAction.h
--- a/Engine/Source/Editor/Kismet/Private/BPVariableDragDropAction.h
+++ b/Engine/Source/Editor/Kismet/Private/BPVariableDragDropAction.h
@@ -42,7 +42,7 @@ public:
 	{
 		if (VariableSource.IsValid() && VariableName != NAME_None)
 		{
-			return FindFieldChecked<UProperty>(VariableSource.Get(), VariableName);
+			return FindField<UProperty>(VariableSource.Get(), VariableName);
 		}
 		return nullptr;
 	}
