Ensured that the 'create + rename' transaction was ended if the component name editable text failed to scroll into view during the frame the component was created. The previous code assumed the name would always be scrolled into view with the focus to let the user rename the newly created component, then close the 'create + rename' transaction as a single 'undo' step, but if the widget did not have any space to display the item, the operation would be postponed until the UI gets resized to display it, leaving the transaction unexpectedly open for an undefined amount of frames (messing the undo/redo buffer). The new code detects if the component was scrolled into view during the frame and if not, the transaction is closed. In such case, the user has to rezise the UI to bring the component into view and rename it.

diff --git a/Engine/Source/Editor/Kismet/Private/SSCSEditor.cpp b/Engine/Source/Editor/Kismet/Private/SSCSEditor.cpp
--- a/Engine/Source/Editor/Kismet/Private/SSCSEditor.cpp
+++ b/Engine/Source/Editor/Kismet/Private/SSCSEditor.cpp
@@ -3721,7 +3721,6 @@ void SSCSEditor::Construct( const FArguments& InArgs )
 		];
 	}
 
-
 	this->ChildSlot
 	[
 		Contents.ToSharedRef()
@@ -3735,6 +3734,8 @@ void SSCSEditor::Construct( const FArguments& InArgs )
 		GEngine->OnLevelComponentRequestRename().AddSP(this, &SSCSEditor::OnLevelComponentRequestRename);
 		GEditor->OnObjectsReplaced().AddSP(this, &SSCSEditor::OnObjectsReplaced);
 	}
+
+	FSlateApplication::Get().OnPostTick().AddSP(this, &SSCSEditor::OnPostTick);
 }
 END_SLATE_FUNCTION_BUILD_OPTIMIZATION
 
@@ -5118,7 +5119,7 @@ UActorComponent* SSCSEditor::AddNewComponent( UClass* NewComponentClass, UObject
 	}
 
 	// Begin a transaction. The transaction will end when the component name will be provided/confirmed by the user.
-	check(!OngoingCreateTransaction.IsValid())
+	check(!DeferredOngoingCreateTransaction.IsValid())
 	TUniquePtr<FScopedTransaction> AddTransaction = MakeUnique<FScopedTransaction>( LOCTEXT("AddComponent", "Add Component") );
 
 	UActorComponent* NewComponent = nullptr;
@@ -6175,7 +6176,7 @@ void SSCSEditor::OnItemScrolledIntoView( FSCSEditorTreeNodePtrType InItem, const
 		if(DeferredRenameRequest == ItemName)
 		{
 			DeferredRenameRequest = NAME_None;
-			InItem->OnRequestRename(MoveTemp(OngoingCreateTransaction)); // Transfer responsibility to end the 'create + give initial name' transaction to the tree item if such transaction is ongoing.
+			InItem->OnRequestRename(MoveTemp(DeferredOngoingCreateTransaction)); // Transfer responsibility to end the 'create + give initial name' transaction to the tree item if such transaction is ongoing.
 		}
 	}
 }
@@ -6200,12 +6201,21 @@ void SSCSEditor::OnRenameComponent(TUniquePtr<FScopedTransaction> InComponentCre
 
 	DeferredRenameRequest = SelectedItems[0]->GetNodeID();
 
-	check(!OngoingCreateTransaction.IsValid()); // If this fails, something in the chain of responsibility failed to end the previous transaction.
-	OngoingCreateTransaction = MoveTemp(InComponentCreateTransaction); // If a 'create + give initial name' transaction is ongoing, take responsibility of ending it until it's transfered to the selected item.
+	check(!DeferredOngoingCreateTransaction.IsValid()); // If this fails, something in the chain of responsibility failed to end the previous transaction.
+	DeferredOngoingCreateTransaction = MoveTemp(InComponentCreateTransaction); // If a 'create + give initial name' transaction is ongoing, take responsibility of ending it until the selected item is scrolled into view.
 
 	SCSTreeWidget->RequestScrollIntoView(SelectedItems[0]);
 }
 
+void SSCSEditor::OnPostTick(float)
+{
+	// If a 'create + give initial name' is ongoing and the transaction ownership was not transferred during the frame it was requested, it is most likely because the newly
+	// created item could not be scrolled into view (should say 'teleported', the scrolling is not animated). But the tree view will not put the item in view if the there is
+	// no space left to display the item. (ex a splitter where all the display space is used by the other component). End the transaction before starting a new frame. The user
+	// will not be able to rename on creation, the widget is likely not in view and cannot be edited anyway.
+	DeferredOngoingCreateTransaction.Reset();
+}
+
 bool SSCSEditor::CanRenameComponent() const
 {
 	return IsEditingAllowed() && SCSTreeWidget->GetSelectedItems().Num() == 1 && SCSTreeWidget->GetSelectedItems()[0]->CanRename();
diff --git a/Engine/Source/Editor/Kismet/Public/SSCSEditor.h b/Engine/Source/Editor/Kismet/Public/SSCSEditor.h
--- a/Engine/Source/Editor/Kismet/Public/SSCSEditor.h
+++ b/Engine/Source/Editor/Kismet/Public/SSCSEditor.h
@@ -976,6 +976,9 @@ public:
 	/** Callback for the action trees to get the filter text */
 	FText GetFilterText() const;
 
+	/** Called at the end of each frame. */
+	void OnPostTick(float);
+
 protected:
 	FString GetSelectedClassText() const;
 
@@ -1143,7 +1146,7 @@ public:
 	FName DeferredRenameRequest;
 
 	/** Scope the creation of a component which ends when the initial component 'name' is given/accepted by the user, which can be several frames after the component was actually created. */
-	TUniquePtr<FScopedTransaction> OngoingCreateTransaction;
+	TUniquePtr<FScopedTransaction> DeferredOngoingCreateTransaction;
 
 	/** Attribute that provides access to the Actor context for which we are viewing/editing the SCS. */
 	TAttribute<class AActor*> ActorContext;
