Remove delegate when requesting window close so that shared pointer in invocation list doesn't prevent SBlueprintDiff from being destroyed
#rnx

diff --git a/Engine/Source/Editor/Kismet/Private/SBlueprintDiff.cpp b/Engine/Source/Editor/Kismet/Private/SBlueprintDiff.cpp
--- a/Engine/Source/Editor/Kismet/Private/SBlueprintDiff.cpp
+++ b/Engine/Source/Editor/Kismet/Private/SBlueprintDiff.cpp
@@ -834,9 +834,14 @@ void SBlueprintDiff::OnCloseAssetEditor(UObject* Asset, EAssetEditorCloseReason
 		// Tell our window to close and set our selves to collapsed to try and stop it from ticking
 		SetVisibility(EVisibility::Collapsed);
 
+		if (AssetEditorCloseDelegate.IsValid())
+		{
+			FAssetEditorManager::Get().OnAssetEditorRequestClose().Remove(AssetEditorCloseDelegate);
+		}
+
 		if (WeakParentWindow.IsValid())
 		{
-			WeakParentWindow.Pin()->DestroyWindowImmediately();
+			WeakParentWindow.Pin()->RequestDestroyWindow();
 		}
 	}
 }
