Minor improvement from the previous fix CL[REDACTED_HANDLE], avoiding to keep the 'post tick' ticking if not required.

diff --git a/Engine/Source/Editor/Kismet/Private/SSCSEditor.cpp b/Engine/Source/Editor/Kismet/Private/SSCSEditor.cpp
--- a/Engine/Source/Editor/Kismet/Private/SSCSEditor.cpp
+++ b/Engine/Source/Editor/Kismet/Private/SSCSEditor.cpp
@@ -3734,8 +3734,6 @@ void SSCSEditor::Construct( const FArguments& InArgs )
 		GEngine->OnLevelComponentRequestRename().AddSP(this, &SSCSEditor::OnLevelComponentRequestRename);
 		GEditor->OnObjectsReplaced().AddSP(this, &SSCSEditor::OnObjectsReplaced);
 	}
-
-	FSlateApplication::Get().OnPostTick().AddSP(this, &SSCSEditor::OnPostTick);
 }
 END_SLATE_FUNCTION_BUILD_OPTIMIZATION
 
@@ -6205,15 +6203,24 @@ void SSCSEditor::OnRenameComponent(TUniquePtr<FScopedTransaction> InComponentCre
 	DeferredOngoingCreateTransaction = MoveTemp(InComponentCreateTransaction); // If a 'create + give initial name' transaction is ongoing, take responsibility of ending it until the selected item is scrolled into view.
 
 	SCSTreeWidget->RequestScrollIntoView(SelectedItems[0]);
+
+	if (DeferredOngoingCreateTransaction.IsValid())
+	{
+		// Ensure the item will be scrolled into view during the frame (See explanation in OnPostTick()).
+		PostTickHandle = FSlateApplication::Get().OnPostTick().AddSP(this, &SSCSEditor::OnPostTick);
+	}
 }
 
 void SSCSEditor::OnPostTick(float)
 {
 	// If a 'create + give initial name' is ongoing and the transaction ownership was not transferred during the frame it was requested, it is most likely because the newly
-	// created item could not be scrolled into view (should say 'teleported', the scrolling is not animated). But the tree view will not put the item in view if the there is
+	// created item could not be scrolled into view (should say 'teleported', the scrolling is not animated). The tree view will not put the item in view if the there is
 	// no space left to display the item. (ex a splitter where all the display space is used by the other component). End the transaction before starting a new frame. The user
 	// will not be able to rename on creation, the widget is likely not in view and cannot be edited anyway.
 	DeferredOngoingCreateTransaction.Reset();
+
+	// The post tick event handler is not required anymore.
+	FSlateApplication::Get().OnPostTick().Remove(PostTickHandle);
 }
 
 bool SSCSEditor::CanRenameComponent() const
diff --git a/Engine/Source/Editor/Kismet/Public/SSCSEditor.h b/Engine/Source/Editor/Kismet/Public/SSCSEditor.h
--- a/Engine/Source/Editor/Kismet/Public/SSCSEditor.h
+++ b/Engine/Source/Editor/Kismet/Public/SSCSEditor.h
@@ -1148,6 +1148,9 @@ public:
 	/** Scope the creation of a component which ends when the initial component 'name' is given/accepted by the user, which can be several frames after the component was actually created. */
 	TUniquePtr<FScopedTransaction> DeferredOngoingCreateTransaction;
 
+	/** Used to unregister from the post tick event. */
+	FDelegateHandle PostTickHandle;
+
 	/** Attribute that provides access to the Actor context for which we are viewing/editing the SCS. */
 	TAttribute<class AActor*> ActorContext;
 
