Minor QoL fix for BP debugging "world filter" combo box during PIE - exclude streaming subworlds that appear as 'standalone'.

diff --git a/Engine/Source/Editor/Kismet/Private/BlueprintEditor.cpp b/Engine/Source/Editor/Kismet/Private/BlueprintEditor.cpp
--- a/Engine/Source/Editor/Kismet/Private/BlueprintEditor.cpp
+++ b/Engine/Source/Editor/Kismet/Private/BlueprintEditor.cpp
@@ -8028,6 +8028,12 @@ FText FBlueprintEditor::GetPIEStatus() const
 	}
 	else if (NetMode == NM_Client)
 	{
+		FWorldContext* PIEContext = GEngine->GetWorldContextFromWorld(DebugWorld);
+		if (PIEContext && PIEContext->PIEInstance > 1)
+		{
+			return FText::Format(LOCTEXT("PIEStatusClientSimulatingFormat", "CLIENT {0} - SIMULATING"), FText::AsNumber(PIEContext->PIEInstance - 1));
+		}
+		
 		return LOCTEXT("PIEStatusClientSimulating", "CLIENT - SIMULATING");
 	}
 
diff --git a/Engine/Source/Editor/Kismet/Private/SBlueprintEditorSelectedDebugObjectWidget.cpp b/Engine/Source/Editor/Kismet/Private/SBlueprintEditorSelectedDebugObjectWidget.cpp
--- a/Engine/Source/Editor/Kismet/Private/SBlueprintEditorSelectedDebugObjectWidget.cpp
+++ b/Engine/Source/Editor/Kismet/Private/SBlueprintEditorSelectedDebugObjectWidget.cpp
@@ -243,7 +243,14 @@ void SBlueprintEditorSelectedDebugObjectWidget::GenerateDebugWorldNames(bool bRe
 	for (TObjectIterator<UWorld> It; It; ++It)
 	{
 		UWorld *TestWorld = *It;
-		if (!TestWorld || TestWorld->WorldType != EWorldType::PIE)
+
+		// Include only PIE and worlds that own the persistent level (i.e. non-streaming levels).
+		const bool bIsValidDebugWorld = (TestWorld != nullptr)
+			&& TestWorld->WorldType == EWorldType::PIE
+			&& TestWorld->PersistentLevel != nullptr
+			&& TestWorld->PersistentLevel->OwningWorld == TestWorld;
+
+		if (!bIsValidDebugWorld)
 		{
 			continue;
 		}
