Minor QoL fix for BP debugging - include client/server context label in graph overlay text for streaming level actor instances.

diff --git a/Engine/Source/Editor/Kismet/Private/BlueprintEditor.cpp b/Engine/Source/Editor/Kismet/Private/BlueprintEditor.cpp
--- a/Engine/Source/Editor/Kismet/Private/BlueprintEditor.cpp
+++ b/Engine/Source/Editor/Kismet/Private/BlueprintEditor.cpp
@@ -7997,7 +7997,7 @@ TSharedPtr<ISCSEditorCustomization> FBlueprintEditor::CustomizeSCSEditor(USceneC
 FText FBlueprintEditor::GetPIEStatus() const
 {
 	UBlueprint* CurrentBlueprint = GetBlueprintObj();
-	UWorld *DebugWorld = NULL;
+	UWorld *DebugWorld = nullptr;
 	ENetMode NetMode = NM_Standalone;
 	if (CurrentBlueprint)
 	{
@@ -8009,18 +8009,24 @@ FText FBlueprintEditor::GetPIEStatus() const
 		else
 		{
 			UObject* ObjOuter = CurrentBlueprint->GetObjectBeingDebugged();
-			while(DebugWorld == NULL && ObjOuter != NULL)
+			while(DebugWorld == nullptr && ObjOuter != nullptr)
 			{
 				ObjOuter = ObjOuter->GetOuter();
 				DebugWorld = Cast<UWorld>(ObjOuter);
 			}
+
+			if (DebugWorld)
+			{
+				// Redirect through streaming levels to find the owning world; this ensures that we always use the appropriate NetMode for the context string below.
+				if (DebugWorld->PersistentLevel != nullptr && DebugWorld->PersistentLevel->OwningWorld != nullptr)
+				{
+					DebugWorld = DebugWorld->PersistentLevel->OwningWorld;
+				}
+
+				NetMode = DebugWorld->GetNetMode();
+			}
 		}
 	}
-	
-	if (DebugWorld)
-	{
-		NetMode = DebugWorld->GetNetMode();
-	}
 
 	if (NetMode == NM_ListenServer || NetMode == NM_DedicatedServer)
 	{
