Fix for an assert on attempt to locally find terms or references in a new, unsaved Blueprint. (regression)

#rnx

diff --git a/Engine/Source/Editor/Kismet/Private/FindInBlueprintManager.cpp b/Engine/Source/Editor/Kismet/Private/FindInBlueprintManager.cpp
--- a/Engine/Source/Editor/Kismet/Private/FindInBlueprintManager.cpp
+++ b/Engine/Source/Editor/Kismet/Private/FindInBlueprintManager.cpp
@@ -1789,7 +1789,7 @@ void FFindInBlueprintSearchManager::AddOrUpdateBlueprintSearchMetadata(UBlueprin
 	//	c) The Blueprint was loaded for diffing. It makes search all very strange and allows you to fully open those Blueprints.
 	//	d) The Blueprint was loaded/copied for PIE. These assets are temporarily created for a session and don't need to be re-indexed.
 	if (!bEnableGatheringData
-		|| !InBlueprint->HasAnyFlags(RF_WasLoaded)
+		|| InBlueprint->HasAnyFlags(RF_NeedLoad | RF_NeedPostLoad)
 		|| InBlueprint->GetOutermost()->HasAnyPackageFlags(PKG_ForDiffing | PKG_PlayInEditor))
 	{
 		return;
