Mirror 5695696
Avoid replacing references on out of date classes, now that the CDO is remapped during archetype reinstancing this can result in the out of date UClass's ClassDefaultObject pointer being overwritten

diff --git a/Engine/Source/Editor/Kismet/Private/BlueprintCompilationManager.cpp b/Engine/Source/Editor/Kismet/Private/BlueprintCompilationManager.cpp
--- a/Engine/Source/Editor/Kismet/Private/BlueprintCompilationManager.cpp
+++ b/Engine/Source/Editor/Kismet/Private/BlueprintCompilationManager.cpp
@@ -1744,7 +1744,7 @@ void FBlueprintCompilationManagerImpl::ReinstanceBatch(TArray<FReinstancingJob>&
 					{
 						UBlueprintGeneratedClass* IterAsBPGC = Cast<UBlueprintGeneratedClass>(Iter);
 						if(Iter->HasAnyFlags(RF_ClassDefaultObject)
-							|| IterAsBPGC
+							|| (IterAsBPGC && !IterAsBPGC->HasAnyClassFlags(CLASS_NewerVersionExists))
 							|| Cast<UBlueprint>(Iter) )
 						{
 							ArchetypeReferencers.Add(Iter);
