Fix crash on attempt to find selected Actor references in a level Blueprint during PIE.

diff --git a/Engine/Source/Editor/Kismet/Private/FindInBlueprintManager.cpp b/Engine/Source/Editor/Kismet/Private/FindInBlueprintManager.cpp
--- a/Engine/Source/Editor/Kismet/Private/FindInBlueprintManager.cpp
+++ b/Engine/Source/Editor/Kismet/Private/FindInBlueprintManager.cpp
@@ -1957,21 +1957,25 @@ FString FFindInBlueprintSearchManager::QuerySingleBlueprint(UBlueprint* InBluepr
 		FName Key = *AssetObject->GetPathName();
 
 		int32* ArrayIdx = SearchMap.Find(Key);
-		// This should always be true since we make sure to refresh the search data for this Blueprint when doing the search, unless we do not rebuild the searchable data
-		checkf((bInRebuildSearchData && ArrayIdx && *ArrayIdx < SearchArray.Num()) || !bInRebuildSearchData,
-			TEXT("bInRebuildSearchData: %s, ArrayIdx is %s:%d, SearchArray.Num():%d"),
-			bInRebuildSearchData ? TEXT("true") : TEXT("false"),
-			(ArrayIdx != nullptr) ? TEXT("Valid") : TEXT("Invalid"),
-			(ArrayIdx != nullptr) ? *ArrayIdx : INDEX_NONE,
-			SearchArray.Num());
-
 		if (ArrayIdx)
 		{
+			checkf(*ArrayIdx < SearchArray.Num(),
+				TEXT("ArrayIdx:%d, SearchArray.Num():%d"),
+				*ArrayIdx,
+				SearchArray.Num());
+
 			return SearchArray[*ArrayIdx].Value;
 		}
+		else if(bInRebuildSearchData)
+		{
+			// Warn here, since we make sure to refresh the search data for this Blueprint when doing the search, and we expect that it should have
+			// been indexed. Note that there are some situations in which we never index a Blueprint asset (@see AddOrUpdateBlueprintSearchMetadata).
+			UE_LOG(LogBlueprint, Warning, TEXT("Attempted to query a Blueprint (%s) that was not re-indexed even after rebuilding. No results can be returned."), *InBlueprint->GetPathName());
+		}
 	}
 	else
 	{
+		// Also warn here as we do not index diff-only packages.
 		UE_LOG(LogBlueprint, Warning, TEXT("Attempted to query an old Blueprint package opened for diffing!"));
 	}
 	return FString();
diff --git a/Engine/Source/Editor/Kismet/Private/FindInBlueprints.cpp b/Engine/Source/Editor/Kismet/Private/FindInBlueprints.cpp
--- a/Engine/Source/Editor/Kismet/Private/FindInBlueprints.cpp
+++ b/Engine/Source/Editor/Kismet/Private/FindInBlueprints.cpp
@@ -1074,19 +1074,36 @@ void SFindInBlueprints::MakeSearchQuery(FString InSearchString, bool bInIsFindWi
 			{
 				Interfaces.Add(InterfaceDesc.Interface->GetPathName());
 			}
-			FImaginaryFiBDataSharedPtr ImaginaryBlueprint(new FImaginaryBlueprint(Blueprint->GetName(), Blueprint->GetPathName(), ParentClass, Interfaces, FFindInBlueprintSearchManager::Get().QuerySingleBlueprint(Blueprint)));
-			TSharedPtr< FFiBSearchInstance > SearchInstance(new FFiBSearchInstance);
-			FSearchResult SearchResult = RootSearchResult = SearchInstance->StartSearchQuery(SearchValue, ImaginaryBlueprint);
 
-			if (SearchResult.IsValid())
+			const FString UnparsedStringData = FFindInBlueprintSearchManager::Get().QuerySingleBlueprint(Blueprint);
+			const bool bIsBlueprintIndexed = !UnparsedStringData.IsEmpty();
+
+			if (bIsBlueprintIndexed)
 			{
-				ItemsFound = SearchResult->Children;
+				FImaginaryFiBDataSharedPtr ImaginaryBlueprint(new FImaginaryBlueprint(Blueprint->GetName(), Blueprint->GetPathName(), ParentClass, Interfaces, UnparsedStringData));
+				TSharedPtr< FFiBSearchInstance > SearchInstance(new FFiBSearchInstance);
+				FSearchResult SearchResult = RootSearchResult = SearchInstance->StartSearchQuery(SearchValue, ImaginaryBlueprint);
+
+				if (SearchResult.IsValid())
+				{
+					ItemsFound = SearchResult->Children;
+				}
 			}
 
 			if(ItemsFound.Num() == 0)
 			{
+				FText NoResultsText;
+				if (bIsBlueprintIndexed)
+				{
+					NoResultsText = LOCTEXT("BlueprintSearchNoResults", "No Results found");
+				}
+				else
+				{
+					NoResultsText = LOCTEXT("BlueprintSearchNotIndexed", "This Blueprint is not indexed for searching");
+				}
+
 				// Insert a fake result to inform user if none found
-				ItemsFound.Add(FSearchResult(new FFindInBlueprintsResult(LOCTEXT("BlueprintSearchNoResults", "No Results found"))));
+				ItemsFound.Add(FSearchResult(new FFindInBlueprintsResult(NoResultsText)));
 				HighlightText = FText::GetEmpty();
 			}
 			else
diff --git a/Engine/Source/Editor/UnrealEd/Private/EditorEngine.cpp b/Engine/Source/Editor/UnrealEd/Private/EditorEngine.cpp
--- a/Engine/Source/Editor/UnrealEd/Private/EditorEngine.cpp
+++ b/Engine/Source/Editor/UnrealEd/Private/EditorEngine.cpp
@@ -3184,6 +3184,16 @@ void UEditorEngine::FindSelectedActorsInLevelScript()
 	AActor* Actor = GetSelectedActors()->GetTop<AActor>();
 	if(Actor != NULL)
 	{
+		if (PlayWorld)
+		{
+			// Redirect to editor world counterpart if PIE is active. We don't index cloned PIE LSBPs for searching.
+			AActor* EditorActor = EditorUtilities::GetEditorWorldCounterpartActor(Actor);
+			if (EditorActor)
+			{
+				Actor = EditorActor;
+			}
+		}
+
 		FKismetEditorUtilities::ShowActorReferencesInLevelScript(Actor);
 	}
 }
