An assertion fired up when SSCSEditor::AddNewComponent() was called in a loop (or multiple times in the same frame) with every component scheduled to get the focus. This unexpectedly activated the 'create + rename' transaction mechansim multiple times but only one item (the last in the frame) would end up being focused, scrolled into view and renamable on creation. Instead of asserting, if an ongoing transaction is ongoing, assumes it is because AddNewComponet() is being called in a loop, close the previous transaction and open a new one, expecting the last component added in the frame (with focus request) to close its own transaction once scrolled into view and edited by user.

diff --git a/Engine/Source/Editor/Kismet/Private/SSCSEditor.cpp b/Engine/Source/Editor/Kismet/Private/SSCSEditor.cpp
--- a/Engine/Source/Editor/Kismet/Private/SSCSEditor.cpp
+++ b/Engine/Source/Editor/Kismet/Private/SSCSEditor.cpp
@@ -5116,8 +5116,16 @@ UActorComponent* SSCSEditor::AddNewComponent( UClass* NewComponentClass, UObject
 		return nullptr;
 	}
 
+	// If an 'add' transaction is ongoing, it is most likely because AddNewComponent() is being called in a tight loop inside a larger transaction (e.g. 'duplicate')
+	// and bSetFocusToNewItem was true for each element.
+	if (DeferredOngoingCreateTransaction.IsValid() && bSetFocusToNewItem)
+	{
+		// Close the ongoing 'add' sub-transaction before staring another one. The user will not be able to edit the name of that component because the
+		// new component is going to still focus.
+		DeferredOngoingCreateTransaction.Reset();
+	}
+
 	// Begin a transaction. The transaction will end when the component name will be provided/confirmed by the user.
-	check(!DeferredOngoingCreateTransaction.IsValid())
 	TUniquePtr<FScopedTransaction> AddTransaction = MakeUnique<FScopedTransaction>( LOCTEXT("AddComponent", "Add Component") );
 
 	UActorComponent* NewComponent = nullptr;
@@ -6204,7 +6212,7 @@ void SSCSEditor::OnRenameComponent(TUniquePtr<FScopedTransaction> InComponentCre
 
 	SCSTreeWidget->RequestScrollIntoView(SelectedItems[0]);
 
-	if (DeferredOngoingCreateTransaction.IsValid())
+	if (DeferredOngoingCreateTransaction.IsValid() && !PostTickHandle.IsValid())
 	{
 		// Ensure the item will be scrolled into view during the frame (See explanation in OnPostTick()).
 		PostTickHandle = FSlateApplication::Get().OnPostTick().AddSP(this, &SSCSEditor::OnPostTick);
@@ -6221,6 +6229,7 @@ void SSCSEditor::OnPostTick(float)
 
 	// The post tick event handler is not required anymore.
 	FSlateApplication::Get().OnPostTick().Remove(PostTickHandle);
+	PostTickHandle.Reset();
 }
 
 bool SSCSEditor::CanRenameComponent() const
