Detach from parent not the node the detach action is being instigated by

diff --git a/Engine/Source/Editor/Kismet/Private/SSCSEditor.cpp b/Engine/Source/Editor/Kismet/Private/SSCSEditor.cpp
--- a/Engine/Source/Editor/Kismet/Private/SSCSEditor.cpp
+++ b/Engine/Source/Editor/Kismet/Private/SSCSEditor.cpp
@@ -2630,9 +2630,6 @@ void SSCS_RowWidget::OnAttachToDropAction(const TArray<FSCSEditorTreeNodePtrType
 
 void SSCS_RowWidget::OnDetachFromDropAction(const TArray<FSCSEditorTreeNodePtrType>& DroppedNodePtrs)
 {
-	FSCSEditorTreeNodePtrType NodePtr = GetNode();
-
-	check(NodePtr.IsValid());
 	check(DroppedNodePtrs.Num() > 0);
 
 	TSharedPtr<SSCSEditor> SCSEditorPtr = SCSEditor.Pin();
@@ -2646,7 +2643,7 @@ void SSCS_RowWidget::OnDetachFromDropAction(const TArray<FSCSEditorTreeNodePtrTy
 		AActor* PreviewActor = SCSEditorPtr->PreviewActor.Get();
 		check(PreviewActor);
 
-		for(const auto& DroppedNodePtr : DroppedNodePtrs)
+		for (const FSCSEditorTreeNodePtrType& DroppedNodePtr : DroppedNodePtrs)
 		{
 			FVector OldRelativeLocation, OldRelativeScale3D;
 			FRotator OldRelativeRotation;
@@ -2654,7 +2651,9 @@ void SSCS_RowWidget::OnDetachFromDropAction(const TArray<FSCSEditorTreeNodePtrTy
 			check(DroppedNodePtr.IsValid());
 
 			// Detach the node from its parent
-			NodePtr->RemoveChild(DroppedNodePtr);
+			FSCSEditorTreeNodePtrType ParentNodePtr = DroppedNodePtr->GetParent();
+			check(ParentNodePtr.IsValid());
+			ParentNodePtr->RemoveChild(DroppedNodePtr);
 
 			// If the associated component template is a scene component, maintain its current world position
 			USceneComponent* SceneComponentTemplate = Cast<USceneComponent>(DroppedNodePtr->GetComponentTemplate());
@@ -2741,12 +2740,14 @@ void SSCS_RowWidget::OnDetachFromDropAction(const TArray<FSCSEditorTreeNodePtrTy
 	}
 	else    // EComponentEditorMode::ActorInstance
 	{
-		for(const auto& DroppedNodePtr : DroppedNodePtrs)
+		for (const FSCSEditorTreeNodePtrType& DroppedNodePtr : DroppedNodePtrs)
 		{
 			check(DroppedNodePtr.IsValid());
 
 			// Detach the node from its parent
-			NodePtr->RemoveChild(DroppedNodePtr);
+			FSCSEditorTreeNodePtrType ParentNodePtr = DroppedNodePtr->GetParent();
+			check(ParentNodePtr.IsValid());
+			ParentNodePtr->RemoveChild(DroppedNodePtr);
 
 			// Attach the dropped node to the current scene root node
 			check(SCSEditorPtr->SceneRootNodePtr.IsValid());
