Create the expected name when duplicating a scs component that has a variable name that is entirely different than its underlying component template

diff --git a/Engine/Source/Editor/Kismet/Private/SSCSEditor.cpp b/Engine/Source/Editor/Kismet/Private/SSCSEditor.cpp
--- a/Engine/Source/Editor/Kismet/Private/SSCSEditor.cpp
+++ b/Engine/Source/Editor/Kismet/Private/SSCSEditor.cpp
@@ -4096,7 +4096,9 @@ void SSCSEditor::OnDuplicateComponent()
 		{
 			if (UActorComponent* ComponentTemplate = SelectedNodes[i]->GetComponentTemplate())
 			{
-				UActorComponent* CloneComponent = AddNewComponent(ComponentTemplate->GetClass(), ComponentTemplate);
+				USCS_Node* SCSNode = SelectedNodes[i]->GetSCSNode();
+				check(SCSNode == nullptr || SCSNode->ComponentTemplate == ComponentTemplate);
+				UActorComponent* CloneComponent = AddNewComponent(ComponentTemplate->GetClass(), (SCSNode ? (UObject*)SCSNode : ComponentTemplate));
 				if (USceneComponent* SceneClone = Cast<USceneComponent>(CloneComponent))
 				{
 					DuplicateSceneComponentMap.Add(CastChecked<USceneComponent>(ComponentTemplate), SceneClone);
@@ -5130,9 +5132,17 @@ UActorComponent* SSCSEditor::AddNewComponent( UClass* NewComponentClass, UObject
 	TUniquePtr<FScopedTransaction> AddTransaction = MakeUnique<FScopedTransaction>( LOCTEXT("AddComponent", "Add Component") );
 
 	UActorComponent* NewComponent = nullptr;
-	UActorComponent* ComponentTemplate = Cast<UActorComponent>(Asset);
+	FName TemplateVariableName;
 
-	if (ComponentTemplate)
+	USCS_Node* SCSNode = Cast<USCS_Node>(Asset);
+	UActorComponent* ComponentTemplate = (SCSNode ? SCSNode->ComponentTemplate : Cast<UActorComponent>(Asset));
+
+	if (SCSNode)
+	{
+		TemplateVariableName = SCSNode->GetVariableName();
+		Asset = nullptr;
+	}
+	else if (ComponentTemplate)
 	{
 		Asset = nullptr;
 	}
@@ -5155,10 +5165,17 @@ UActorComponent* SSCSEditor::AddNewComponent( UClass* NewComponentClass, UObject
 		FName NewVariableName;
 		if (ComponentTemplate)
 		{
-			FString TemplateName = ComponentTemplate->GetName();
-			NewVariableName = (TemplateName.EndsWith(USimpleConstructionScript::ComponentTemplateNameSuffix) 
-								? FName(*TemplateName.LeftChop(USimpleConstructionScript::ComponentTemplateNameSuffix.Len()))
-								: ComponentTemplate->GetFName());
+			if (!TemplateVariableName.IsNone())
+			{
+				NewVariableName = TemplateVariableName;
+			}
+			else
+			{
+				FString TemplateName = ComponentTemplate->GetName();
+				NewVariableName = (TemplateName.EndsWith(USimpleConstructionScript::ComponentTemplateNameSuffix) 
+									? FName(*TemplateName.LeftChop(USimpleConstructionScript::ComponentTemplateNameSuffix.Len()))
+									: ComponentTemplate->GetFName());
+			}
 		}
 		else if (Asset)
 		{
