Guard against an invalid array access check possibly caused by an empty array during a user-initiated async FiB indexing operation.

diff --git a/Engine/Source/Editor/Kismet/Private/FindInBlueprintManager.cpp b/Engine/Source/Editor/Kismet/Private/FindInBlueprintManager.cpp
--- a/Engine/Source/Editor/Kismet/Private/FindInBlueprintManager.cpp
+++ b/Engine/Source/Editor/Kismet/Private/FindInBlueprintManager.cpp
@@ -1156,7 +1156,7 @@ public:
 	/** Returns the progress as a percent */
 	float GetCacheProgress() const
 	{
-		return (float)TickCacheIndex / (float)UncachedAssets.Num();
+		return UncachedAssets.Num() > 0 ? (float)TickCacheIndex / (float)UncachedAssets.Num() : 1.0f;
 	}
 
 	/** Returns the number of uncached assets */
@@ -1193,6 +1193,7 @@ public:
 			// Sometimes we can't wait another tick to shutdown, so make the callback immediately.
 			if (bIsImmediate)
 			{
+				// Note: This will effectively delete this instance. It should not be used after this!
 				FFindInBlueprintSearchManager::Get().FinishedCachingBlueprints(TickCacheIndex, FailedToCacheList);
 			}
 			else
@@ -1219,8 +1220,17 @@ public:
 			return;
 		}
 
+		if (UncachedAssets.Num() == 0)
+		{
+			// Immediately finish if we have no assets to index. This will delete this instance!
+			Finish();
+
+			return;
+		}
+
 		if (bIsCancelled || GWarn->ReceivedUserCancel())
 		{
+			// Note: This will effectively delete this instance. It should not be used after this!
 			FFindInBlueprintSearchManager::Get().FinishedCachingBlueprints(TickCacheIndex, FailedToCacheList);
 		}
 		else
@@ -1333,18 +1343,8 @@ public:
 			// Check if done caching Blueprints
 			if(TickCacheIndex == UncachedAssets.Num())
 			{
-				if (ProgressNotification.IsValid())
-				{
-					ProgressNotification.Pin()->SetCompletionState(SNotificationItem::CS_Success);
-					ProgressNotification.Pin()->ExpireAndFadeout();
-
-					ProgressNotification.Pin()->SetText(LOCTEXT("BlueprintIndexComplete", "Finished indexing Blueprints!"));
-				}
-
-				// We have actually finished, use the OnFinished callback.
-				CacheParams.OnFinished.ExecuteIfBound();
-
-				FFindInBlueprintSearchManager::Get().FinishedCachingBlueprints(TickCacheIndex, FailedToCacheList);
+				// Note: This will effectively delete this instance, do not use after this!
+				Finish();
 			}
 			else if(ProgressNotification.IsValid())
 			{
@@ -1355,6 +1355,25 @@ public:
 		}
 	}
 
+protected:
+	/** Completes a successful caching process */
+	void Finish()
+	{
+		if (ProgressNotification.IsValid())
+		{
+			ProgressNotification.Pin()->SetCompletionState(SNotificationItem::CS_Success);
+			ProgressNotification.Pin()->ExpireAndFadeout();
+
+			ProgressNotification.Pin()->SetText(LOCTEXT("BlueprintIndexComplete", "Finished indexing Blueprints!"));
+		}
+
+		// We have actually finished, use the OnFinished callback.
+		CacheParams.OnFinished.ExecuteIfBound();
+
+		// Note: This will effectively delete this instance. It should not be used after this!
+		FFindInBlueprintSearchManager::Get().FinishedCachingBlueprints(TickCacheIndex, FailedToCacheList);
+	}
+
 private:
 
 	/** The current index, increases at a rate of once per tick */
