BP debugging QoL - add server/client context to debug object combo box, more consistent debug object selection.

diff --git a/Engine/Source/Editor/Kismet/Private/BlueprintEditor.cpp b/Engine/Source/Editor/Kismet/Private/BlueprintEditor.cpp
--- a/Engine/Source/Editor/Kismet/Private/BlueprintEditor.cpp
+++ b/Engine/Source/Editor/Kismet/Private/BlueprintEditor.cpp
@@ -7944,17 +7944,6 @@ FReply FBlueprintEditor::OnSpawnGraphNodeByShortcut(FInputChord InChord, const F
 	return FReply::Handled();
 }
 
-void FBlueprintEditor::ToolkitBroughtToFront()
-{
-	UBlueprint* CurrentBlueprint = GetBlueprintObj();
-	if( CurrentBlueprint != NULL )
-	{
-		UObject* DebugInstance = CurrentBlueprint->GetObjectBeingDebugged();
-		CurrentBlueprint->SetObjectBeingDebugged( NULL );
-		CurrentBlueprint->SetObjectBeingDebugged( DebugInstance );
-	}
-}
-
 void FBlueprintEditor::OnNodeSpawnedByKeymap()
 {
 	UpdateNodeCreationStats( ENodeCreateAction::Keymap );
diff --git a/Engine/Source/Editor/Kismet/Private/SBlueprintEditorSelectedDebugObjectWidget.cpp b/Engine/Source/Editor/Kismet/Private/SBlueprintEditorSelectedDebugObjectWidget.cpp
--- a/Engine/Source/Editor/Kismet/Private/SBlueprintEditorSelectedDebugObjectWidget.cpp
+++ b/Engine/Source/Editor/Kismet/Private/SBlueprintEditorSelectedDebugObjectWidget.cpp
@@ -46,7 +46,7 @@ void SBlueprintEditorSelectedDebugObjectWidget::Construct(const FArguments& InAr
 		.OnComboBoxOpening(this, &SBlueprintEditorSelectedDebugObjectWidget::GenerateDebugWorldNames, true)
 		.OnSelectionChanged(this, &SBlueprintEditorSelectedDebugObjectWidget::DebugWorldSelectionChanged);
 
-	DebugObjectsComboBox = SNew(STextComboBox)
+	DebugObjectsComboBox = SNew(SComboBox<TSharedPtr<FString>>)
 		.ButtonStyle(FEditorStyle::Get(), "FlatButton.Light")
 		.ToolTip(IDocumentation::Get()->CreateToolTip(
 		LOCTEXT("BlueprintDebugObjectTooltip", "Select an object to debug"),
@@ -57,7 +57,12 @@ void SBlueprintEditorSelectedDebugObjectWidget::Construct(const FArguments& InAr
 		.InitiallySelectedItem(GetDebugObjectName())
 		.OnComboBoxOpening(this, &SBlueprintEditorSelectedDebugObjectWidget::GenerateDebugObjectNames, true)
 		.OnSelectionChanged(this, &SBlueprintEditorSelectedDebugObjectWidget::DebugObjectSelectionChanged)
-		.AddMetaData<FTagMetaData>(TEXT("SelectDebugObjectCobmo"));
+		.OnGenerateWidget(this, &SBlueprintEditorSelectedDebugObjectWidget::CreateDebugObjectItemWidget)
+		.AddMetaData<FTagMetaData>(TEXT("SelectDebugObjectCobmo"))
+		[
+			SNew(STextBlock)
+			.Text(this, &SBlueprintEditorSelectedDebugObjectWidget::GetSelectedDebugObjectTextLabel)
+		];
 
 	ChildSlot
 	[
@@ -202,36 +207,18 @@ void SBlueprintEditorSelectedDebugObjectWidget::OnRefresh()
 	if (GetBlueprintObj())
 	{
 		GenerateDebugWorldNames(false);
-		if (UObject* Object = GetBlueprintObj()->GetObjectBeingDebugged())
-		{
-			GenerateDebugObjectNames(false);
-			if (AActor* Actor = Cast<AActor>(Object))
-			{
-				DebugObjectsComboBox->SetSelectedItem(MakeShareable(new FString(Actor->GetActorLabel())));
-			}
-			else
-			{
-				DebugObjectsComboBox->SetSelectedItem(MakeShareable(new FString(Object->GetName())));
-			}
-		}
-		else
+		GenerateDebugObjectNames(false);
+
+		if (DebugObjectsComboBox.IsValid())
 		{
-			// If we didn't find and object its probably a good idea to regenerate the names now - this will also ensure the combo box has a valid selection.
-			GenerateDebugObjectNames(false);
+			DebugObjectsComboBox->SetSelectedItem(GetDebugWorldName());
+			DebugObjectsComboBox->SetSelectedItem(GetDebugObjectName());
 		}
 	}
 }
 
 void SBlueprintEditorSelectedDebugObjectWidget::GenerateDebugWorldNames(bool bRestoreSelection)
 {
-	TSharedPtr<FString> OldSelection;
-
-	// Store off the old selection
-	if (bRestoreSelection && DebugWorldsComboBox.IsValid())
-	{
-		OldSelection = DebugWorldsComboBox->GetSelectedItem();
-	}
-
 	DebugWorldNames.Empty();
 	DebugWorlds.Empty();
 
@@ -300,31 +287,19 @@ void SBlueprintEditorSelectedDebugObjectWidget::GenerateDebugWorldNames(bool bRe
 		}
 	}
 
-	// Attempt to restore the old selection
-	if (bRestoreSelection && DebugWorldsComboBox.IsValid())
+	if (DebugWorldsComboBox.IsValid())
 	{
-		bool bMatchFound = false;
-		for (int32 WorldIdx = 0; WorldIdx < DebugWorldNames.Num(); ++WorldIdx)
+		// Attempt to restore the old selection
+		if (bRestoreSelection)
 		{
-			if (*DebugWorldNames[WorldIdx] == *OldSelection)
+			TSharedPtr<FString> CurrentDebugWorld = GetDebugWorldName();
+			if (CurrentDebugWorld.IsValid())
 			{
-				DebugWorldsComboBox->SetSelectedItem(DebugWorldNames[WorldIdx]);
-				bMatchFound = true;
-				break;
+				DebugWorldsComboBox->SetSelectedItem(CurrentDebugWorld);
 			}
 		}
 
-		// No match found, use the default option
-		if (!bMatchFound)
-		{
-			DebugWorldsComboBox->SetSelectedItem(DebugWorldNames[0]);
-		}
-	}
-
-
-	// Finally ensure we have a valid selection
-	if (DebugWorldsComboBox.IsValid())
-	{
+		// Finally ensure we have a valid selection
 		TSharedPtr<FString> CurrentSelection = DebugWorldsComboBox->GetSelectedItem();
 		if (DebugWorldNames.Find(CurrentSelection) == INDEX_NONE)
 		{
@@ -337,19 +312,13 @@ void SBlueprintEditorSelectedDebugObjectWidget::GenerateDebugWorldNames(bool bRe
 				DebugWorldsComboBox->ClearSelection();
 			}
 		}
+
+		DebugWorldsComboBox->RefreshOptions();
 	}
 }
 
 void SBlueprintEditorSelectedDebugObjectWidget::GenerateDebugObjectNames(bool bRestoreSelection)
 {
-	TSharedPtr<FString> OldSelection;
-
-	// Store off the old selection
-	if (bRestoreSelection && DebugObjectsComboBox.IsValid())
-	{
-		OldSelection = DebugObjectsComboBox->GetSelectedItem();
-	}
-
 	// Empty the lists of actors and regenerate them
 	DebugObjects.Empty();
 	DebugObjectNames.Empty();
@@ -464,83 +433,72 @@ void SBlueprintEditorSelectedDebugObjectWidget::GenerateDebugObjectNames(bool bR
 	}
 	else
 	{
-	for (TObjectIterator<UObject> It; It; ++It)
-	{
-		UObject* TestObject = *It;
-
-		// Skip Blueprint preview objects (don't allow them to be selected for debugging)
-		if (PreviewWorld != nullptr && TestObject->IsIn(PreviewWorld))
+		for (TObjectIterator<UObject> It; It; ++It)
 		{
-			continue;
-		}
-
-		const bool bPassesFlags = !TestObject->HasAnyFlags(RF_ClassDefaultObject) && !TestObject->IsPendingKill();
-		const bool bGeneratedByAnyBlueprint = TestObject->GetClass()->ClassGeneratedBy != nullptr;
-		const bool bGeneratedByThisBlueprint = bGeneratedByAnyBlueprint && GetBlueprintObj()->GeneratedClass && TestObject->IsA(GetBlueprintObj()->GeneratedClass);
+			UObject* TestObject = *It;
 
-		if (bPassesFlags && bGeneratedByThisBlueprint)
-		{
-			UObject *ObjOuter = TestObject;
-			UWorld *ObjWorld = nullptr;
-			do		// Run through at least once in case the TestObject is a UGameInstance
+			// Skip Blueprint preview objects (don't allow them to be selected for debugging)
+			if (PreviewWorld != nullptr && TestObject->IsIn(PreviewWorld))
 			{
-				UGameInstance *ObjGameInstance = Cast<UGameInstance>(ObjOuter);
+				continue;
+			}
 
-				ObjOuter = ObjOuter->GetOuter();
-				ObjWorld = ObjGameInstance ? ObjGameInstance->GetWorld() : Cast<UWorld>(ObjOuter);
-			} while (ObjWorld == nullptr && ObjOuter != nullptr);
+			const bool bPassesFlags = !TestObject->HasAnyFlags(RF_ClassDefaultObject) && !TestObject->IsPendingKill();
+			const bool bGeneratedByAnyBlueprint = TestObject->GetClass()->ClassGeneratedBy != nullptr;
+			const bool bGeneratedByThisBlueprint = bGeneratedByAnyBlueprint && GetBlueprintObj()->GeneratedClass && TestObject->IsA(GetBlueprintObj()->GeneratedClass);
 
-			if (ObjWorld)
+			if (bPassesFlags && bGeneratedByThisBlueprint)
 			{
-				// Make check on owning level (not streaming level)
-				if (ObjWorld->PersistentLevel && ObjWorld->PersistentLevel->OwningWorld)
+				UObject *ObjOuter = TestObject;
+				UWorld *ObjWorld = nullptr;
+				do		// Run through at least once in case the TestObject is a UGameInstance
 				{
-					ObjWorld = ObjWorld->PersistentLevel->OwningWorld;
-				}
+					UGameInstance *ObjGameInstance = Cast<UGameInstance>(ObjOuter);
 
-				// We have a specific debug world and the object isn't in it
-				if (DebugWorld && ObjWorld != DebugWorld)
-				{
-					continue;
-				}
+					ObjOuter = ObjOuter->GetOuter();
+					ObjWorld = ObjGameInstance ? ObjGameInstance->GetWorld() : Cast<UWorld>(ObjOuter);
+				} while (ObjWorld == nullptr && ObjOuter != nullptr);
 
-				if ((ObjWorld->WorldType == EWorldType::Editor) && (GUnrealEd->GetPIEViewport() == nullptr))
-				{
-					AddDebugObject(TestObject);
-				}
-				else if (ObjWorld->WorldType == EWorldType::PIE)
+				if (ObjWorld)
 				{
-					AddDebugObject(TestObject);
+					// Make check on owning level (not streaming level)
+					if (ObjWorld->PersistentLevel && ObjWorld->PersistentLevel->OwningWorld)
+					{
+						ObjWorld = ObjWorld->PersistentLevel->OwningWorld;
+					}
+
+					// We have a specific debug world and the object isn't in it
+					if (DebugWorld && ObjWorld != DebugWorld)
+					{
+						continue;
+					}
+
+					if ((ObjWorld->WorldType == EWorldType::Editor) && (GUnrealEd->GetPIEViewport() == nullptr))
+					{
+						AddDebugObject(TestObject);
+					}
+					else if (ObjWorld->WorldType == EWorldType::PIE)
+					{
+						AddDebugObject(TestObject);
+					}
 				}
 			}
 		}
 	}
-	}
 
-	// Attempt to restore the old selection
-	if (bRestoreSelection && DebugObjectsComboBox.IsValid())
+	if (DebugObjectsComboBox.IsValid())
 	{
-		bool bMatchFound = false;
-		for (int32 ObjectIndex = 0; ObjectIndex < DebugObjectNames.Num(); ++ObjectIndex)
+		// Attempt to restore the old selection
+		if (bRestoreSelection)
 		{
-			if (*DebugObjectNames[ObjectIndex] == *OldSelection)
+			TSharedPtr<FString> CurrentDebugObject = GetDebugObjectName();
+			if (CurrentDebugObject.IsValid())
 			{
-				DebugObjectsComboBox->SetSelectedItem(DebugObjectNames[ObjectIndex]);
-				bMatchFound = true;
-				break;
+				DebugObjectsComboBox->SetSelectedItem(CurrentDebugObject);
 			}
 		}
 
-		// No match found, use the default option
-		if (!bMatchFound)
-		{
-			DebugObjectsComboBox->SetSelectedItem(DebugObjectNames[0]);
-		}
-	}
-
-	// Finally ensure we have a valid selection
-	if (DebugObjectsComboBox.IsValid())
-	{
+		// Finally ensure we have a valid selection
 		TSharedPtr<FString> CurrentSelection = DebugObjectsComboBox->GetSelectedItem();
 		if (DebugObjectNames.Find(CurrentSelection) == INDEX_NONE)
 		{
@@ -581,13 +539,14 @@ TSharedPtr<FString> SBlueprintEditorSelectedDebugObjectWidget::GetDebugWorldName
 	check(GetBlueprintObj());
 	if (ensure(DebugWorlds.Num() == DebugWorldNames.Num()))
 	{
-		if (UObject* DebugObj = GetBlueprintObj()->GetObjectBeingDebugged())
+		UWorld* DebugWorld = GetBlueprintObj()->GetWorldBeingDebugged();
+		if (DebugWorld != nullptr)
 		{
-			for (int32 WorldIdx = 0; WorldIdx < DebugWorlds.Num(); ++WorldIdx)
+			for (int32 WorldIndex = 0; WorldIndex < DebugWorlds.Num(); ++WorldIndex)
 			{
-				if (DebugObj->IsIn(DebugWorlds[WorldIdx].Get()))
+				if (DebugWorlds[WorldIndex].IsValid() && (DebugWorlds[WorldIndex].Get() == DebugWorld))
 				{
-					return DebugWorldNames[WorldIdx];
+					return DebugWorldNames[WorldIndex];
 				}
 			}
 		}
@@ -598,33 +557,39 @@ TSharedPtr<FString> SBlueprintEditorSelectedDebugObjectWidget::GetDebugWorldName
 
 void SBlueprintEditorSelectedDebugObjectWidget::DebugWorldSelectionChanged(TSharedPtr<FString> NewSelection, ESelectInfo::Type SelectInfo)
 {
-	if (!NewSelection.IsValid())
+	if (NewSelection != GetDebugWorldName())
 	{
-		return;
-	}
-
-	check(DebugObjects.Num() == DebugObjectNames.Num());
-	for (int32 WorldIdx = 0; WorldIdx < DebugWorldNames.Num(); ++WorldIdx)
-	{
-		if (*DebugWorldNames[WorldIdx] == *NewSelection)
+		check(DebugWorlds.Num() == DebugWorldNames.Num());
+		for (int32 WorldIdx = 0; WorldIdx < DebugWorldNames.Num(); ++WorldIdx)
 		{
-			GetBlueprintObj()->SetWorldBeingDebugged(DebugWorlds[WorldIdx].Get());
-			GenerateDebugObjectNames(false);
+			if (DebugWorldNames[WorldIdx] == NewSelection)
+			{
+				GetBlueprintObj()->SetWorldBeingDebugged(DebugWorlds[WorldIdx].Get());
+
+				GetBlueprintObj()->SetObjectBeingDebugged(nullptr);
+				LastObjectObserved.Reset();
+
+				GenerateDebugObjectNames(false);
+				break;
+			}
 		}
 	}
 }
 
 void SBlueprintEditorSelectedDebugObjectWidget::DebugObjectSelectionChanged(TSharedPtr<FString> NewSelection, ESelectInfo::Type SelectInfo)
 {
-	if (SelectInfo != ESelectInfo::Direct)
+	if (NewSelection != GetDebugObjectName())
 	{
 		check(DebugObjects.Num() == DebugObjectNames.Num());
 		for (int32 ObjectIndex = 0; ObjectIndex < DebugObjectNames.Num(); ++ObjectIndex)
 		{
-			if (*DebugObjectNames[ObjectIndex] == *NewSelection)
+			if (DebugObjectNames[ObjectIndex] == NewSelection)
 			{
-				UObject* DebugObj = DebugObjects[ObjectIndex].IsValid() ? DebugObjects[ObjectIndex].Get() : nullptr;
+				UObject* DebugObj = DebugObjects[ObjectIndex].Get();
 				GetBlueprintObj()->SetObjectBeingDebugged(DebugObj);
+
+				LastObjectObserved = DebugObj;
+				break;
 			}
 		}
 	}
@@ -678,17 +643,63 @@ EVisibility SBlueprintEditorSelectedDebugObjectWidget::IsDebugWorldComboVisible(
 	return EVisibility::Collapsed;
 }
 
-void SBlueprintEditorSelectedDebugObjectWidget::AddDebugObject(UObject* TestObject)
+FString SBlueprintEditorSelectedDebugObjectWidget::MakeDebugObjectLabel(UObject* TestObject, bool bAddContextIfSelectedInEditor) const
 {
-	FString Label;
-	if (AActor* Actor = Cast<AActor>(TestObject))
+	auto GetActorLabelStringLambda = [](AActor* InActor, bool bIncludeNetModeSuffix, bool bIncludeSelectedSuffix)
 	{
-		Label = Actor->GetActorLabel();
+		FString Label = InActor->GetActorLabel();
+
+		FString Context;
 
-		if (Actor->IsSelected())
+		if (bIncludeNetModeSuffix)
 		{
-			Label += TEXT(" (selected)");
+			switch (InActor->GetNetMode())
+			{
+			case ENetMode::NM_Client:
+			{
+				Context = NSLOCTEXT("BlueprintEditor", "DebugWorldClient", "Client").ToString();
+
+				FWorldContext* WorldContext = GEngine->GetWorldContextFromWorld(InActor->GetWorld());
+				if (WorldContext != nullptr && WorldContext->PIEInstance > 1)
+				{
+					Context += TEXT(" ");
+					Context += FText::AsNumber(WorldContext->PIEInstance - 1).ToString();
+				}
+			}
+			break;
+
+			case ENetMode::NM_ListenServer:
+			case ENetMode::NM_DedicatedServer:
+				Context = NSLOCTEXT("BlueprintEditor", "DebugWorldServer", "Server").ToString();
+				break;
+			}
 		}
+
+		if (bIncludeSelectedSuffix && InActor->IsSelected())
+		{
+			if (!Context.IsEmpty())
+			{
+				Context += TEXT(" - ");
+			}
+
+			Context += NSLOCTEXT("BlueprintEditor", "DebugObjectSelected", "selected").ToString();
+		}
+
+		if (!Context.IsEmpty())
+		{
+			Label = FString::Printf(TEXT("%s (%s)"), *Label, *Context);
+		}
+
+		return Label;
+	};
+
+	// Include net mode suffix when "All worlds" is selected.
+	const bool bIncludeNetModeSuffix = *GetDebugWorldName() == GetDebugAllWorldsString();
+
+	FString Label;
+	if (AActor* Actor = Cast<AActor>(TestObject))
+	{
+		Label = GetActorLabelStringLambda(Actor, bIncludeNetModeSuffix, bAddContextIfSelectedInEditor);
 	}
 	else
 	{
@@ -697,7 +708,7 @@ void SBlueprintEditorSelectedDebugObjectWidget::AddDebugObject(UObject* TestObje
 			// This gives the most precision, but is pretty long for the combo box
 			//const FString RelativePath = TestObject->GetPathName(/*StopOuter=*/ ParentActor);
 			const FString RelativePath = TestObject->GetName();
-			Label = FString::Printf(TEXT("%s in %s"), *RelativePath, *ParentActor->GetActorLabel());
+			Label = FString::Printf(TEXT("%s in %s"), *RelativePath, *GetActorLabelStringLambda(ParentActor, bIncludeNetModeSuffix, bAddContextIfSelectedInEditor));
 		}
 		else
 		{
@@ -705,6 +716,12 @@ void SBlueprintEditorSelectedDebugObjectWidget::AddDebugObject(UObject* TestObje
 		}
 	}
 
+	return Label;
+}
+
+void SBlueprintEditorSelectedDebugObjectWidget::AddDebugObject(UObject* TestObject)
+{
+	FString Label = MakeDebugObjectLabel(TestObject, true);
 	AddDebugObjectWithName(TestObject, Label);
 }
 
@@ -714,6 +731,42 @@ void SBlueprintEditorSelectedDebugObjectWidget::AddDebugObjectWithName(UObject*
 	DebugObjectNames.Add(MakeShareable(new FString(TestObjectName)));
 }
 
+TSharedRef<SWidget> SBlueprintEditorSelectedDebugObjectWidget::CreateDebugObjectItemWidget(TSharedPtr<FString> InItem)
+{
+	FString ItemString;
+
+	if (InItem.IsValid())
+	{
+		ItemString = *InItem;
+	}
+
+	return SNew(STextBlock)
+		.Text(FText::FromString(*ItemString));
+}
+
+FText SBlueprintEditorSelectedDebugObjectWidget::GetSelectedDebugObjectTextLabel() const
+{
+	FString Label;
+
+	UBlueprint* Blueprint = GetBlueprintObj();
+	if (Blueprint != nullptr)
+	{
+		UObject* DebugObj = Blueprint->GetObjectBeingDebugged();
+		if (DebugObj != nullptr)
+		{
+			// Exclude the editor selection suffix for the combo button's label.
+			Label = MakeDebugObjectLabel(DebugObj, false);
+		}
+	}
+		
+	if (Label.IsEmpty())
+	{
+		Label = *GetDebugObjectName();
+	}
+
+	return FText::FromString(Label);
+}
+
 //////////////////////////////////////////////////////////////////////////
 
 #undef LOCTEXT_NAMESPACE
diff --git a/Engine/Source/Editor/Kismet/Private/SBlueprintEditorSelectedDebugObjectWidget.h b/Engine/Source/Editor/Kismet/Private/SBlueprintEditorSelectedDebugObjectWidget.h
--- a/Engine/Source/Editor/Kismet/Private/SBlueprintEditorSelectedDebugObjectWidget.h
+++ b/Engine/Source/Editor/Kismet/Private/SBlueprintEditorSelectedDebugObjectWidget.h
@@ -6,6 +6,7 @@
 #include "Layout/Visibility.h"
 #include "Widgets/DeclarativeSyntaxSupport.h"
 #include "Widgets/SCompoundWidget.h"
+#include "Widgets/Input/SComboBox.h"
 #include "BlueprintEditor.h"
 
 class KISMET_API SBlueprintEditorSelectedDebugObjectWidget : public SCompoundWidget
@@ -61,6 +62,16 @@ private:
 	const FString& GetDebugAllWorldsString() const;
 
 	TSharedRef<SWidget> OnGetActiveDetailSlotContent(bool bChangedToHighDetail);
+
+	/** Helper method to construct a debug object label string */
+	FString MakeDebugObjectLabel(UObject* TestObject, bool bAddContextIfSelectedInEditor) const;
+
+	/** Called to create a widget for each debug object item */
+	TSharedRef<SWidget> CreateDebugObjectItemWidget(TSharedPtr<FString> InItem);
+
+	/** Returns the combo button label to use for the currently-selected debug object item */
+	FText GetSelectedDebugObjectTextLabel() const;
+
 private:
 	/** Pointer back to the blueprint editor tool that owns us */
 	TWeakPtr<FBlueprintEditor> BlueprintEditor;
@@ -73,8 +84,8 @@ private:
 	TArray< TWeakObjectPtr<UWorld> > DebugWorlds;
 	TArray< TSharedPtr<FString> > DebugWorldNames;
 
-	/** Widget containing the names of all possible debug actors */
-	TSharedPtr<STextComboBox> DebugObjectsComboBox;
+	/** Widget containing the names of all possible debug actors. This is a "generic" SComboBox rather than an STextComboBox so that we can customize the label on the combo button widget. */
+	TSharedPtr<SComboBox<TSharedPtr<FString>>> DebugObjectsComboBox;
 
 	TSharedPtr<STextComboBox> DebugWorldsComboBox;
 
diff --git a/Engine/Source/Editor/Kismet/Public/BlueprintEditor.h b/Engine/Source/Editor/Kismet/Public/BlueprintEditor.h
--- a/Engine/Source/Editor/Kismet/Public/BlueprintEditor.h
+++ b/Engine/Source/Editor/Kismet/Public/BlueprintEditor.h
@@ -166,7 +166,6 @@ public:
 public:
 	//~ Begin FAssetEditorToolkit Interface
 	virtual bool OnRequestClose() override;
-	virtual void ToolkitBroughtToFront() override;
 	// End of FAssetEditorToolkit 
 
 	//~ Begin IToolkit Interface
diff --git a/Engine/Source/Editor/UnrealEd/Private/Kismet2/KismetReinstanceUtilities.cpp b/Engine/Source/Editor/UnrealEd/Private/Kismet2/KismetReinstanceUtilities.cpp
--- a/Engine/Source/Editor/UnrealEd/Private/Kismet2/KismetReinstanceUtilities.cpp
+++ b/Engine/Source/Editor/UnrealEd/Private/Kismet2/KismetReinstanceUtilities.cpp
@@ -1909,7 +1909,8 @@ void FBlueprintCompileReinstancer::ReplaceInstancesOfClass_Inner(TMap<UClass*, U
 	{
 		if (UBlueprint* OldObjBlueprint = Cast<UBlueprint>(InOldObject->GetClass()->ClassGeneratedBy))
 		{
-			if (OldObjBlueprint->GetObjectBeingDebugged() == InOldObject)
+			const UObject* DebugObj = OldObjBlueprint->GetObjectBeingDebugged(EGetObjectOrWorldBeingDebuggedFlags::IgnorePendingKill);
+			if (DebugObj == InOldObject)
 			{
 				OldObjBlueprint->SetObjectBeingDebugged(InNewObject);
 			}
diff --git a/Engine/Source/Runtime/Engine/Classes/Engine/Blueprint.h b/Engine/Source/Runtime/Engine/Classes/Engine/Blueprint.h
--- a/Engine/Source/Runtime/Engine/Classes/Engine/Blueprint.h
+++ b/Engine/Source/Runtime/Engine/Classes/Engine/Blueprint.h
@@ -402,6 +402,19 @@ enum class EBlueprintNativizationFlag : uint8
 	ExplicitlyEnabled
 };
 
+#if WITH_EDITOR
+/** Control flags for current object/world accessor methods */
+enum class EGetObjectOrWorldBeingDebuggedFlags
+{
+	/** Use normal weak ptr semantics when accessing the referenced object. */
+	None = 0,
+	/** Return a valid ptr even if the PendingKill flag is set on the referenced object. */
+	IgnorePendingKill = 1 << 0
+};
+
+ENUM_CLASS_FLAGS(EGetObjectOrWorldBeingDebuggedFlags);
+#endif
+
 
 /**
  * Blueprints are special assets that provide an intuitive, node-based interface that can be used to create new types of Actors
@@ -777,9 +790,18 @@ private:
 public:
 
 	/** @return the current object being debugged, which can be nullptr */
-	virtual UObject* GetObjectBeingDebugged();
+	UObject* GetObjectBeingDebugged(EGetObjectOrWorldBeingDebuggedFlags InFlags = EGetObjectOrWorldBeingDebuggedFlags::None) const
+	{
+		const bool bEvenIfPendingKill = EnumHasAnyFlags(InFlags, EGetObjectOrWorldBeingDebuggedFlags::IgnorePendingKill);
+		return CurrentObjectBeingDebugged.Get(bEvenIfPendingKill);
+	}
 
-	virtual class UWorld* GetWorldBeingDebugged();
+	/** @return the current world being debugged, which can be nullptr */
+	class UWorld* GetWorldBeingDebugged(EGetObjectOrWorldBeingDebuggedFlags InFlags = EGetObjectOrWorldBeingDebuggedFlags::None) const
+	{
+		const bool bEvenIfPendingKill = EnumHasAnyFlags(InFlags, EGetObjectOrWorldBeingDebuggedFlags::IgnorePendingKill);
+		return CurrentWorldBeingDebugged.Get(bEvenIfPendingKill);
+	}
 
 	/** Renames only the generated classes. Should only be used internally or when testing for rename. */
 	virtual bool RenameGeneratedClasses(const TCHAR* NewName = nullptr, UObject* NewOuter = nullptr, ERenameFlags Flags = REN_None);
diff --git a/Engine/Source/Runtime/Engine/Classes/Engine/LevelScriptActor.h b/Engine/Source/Runtime/Engine/Classes/Engine/LevelScriptActor.h
--- a/Engine/Source/Runtime/Engine/Classes/Engine/LevelScriptActor.h
+++ b/Engine/Source/Runtime/Engine/Classes/Engine/LevelScriptActor.h
@@ -58,11 +58,7 @@ class ENGINE_API ALevelScriptActor : public AActor
 	UFUNCTION(BlueprintImplementableEvent)
 	void WorldOriginLocationChanged(FIntVector OldOriginLocation, FIntVector NewOriginLocation);
 	
-#if WITH_EDITOR
 	//~ Begin UObject Interface
-	virtual void PostDuplicate(bool bDuplicateForPIE) override;
-	virtual void BeginDestroy() override;
-#endif
 	virtual void PreInitializeComponents() override;
 	//~ End UObject Interface
 
diff --git a/Engine/Source/Runtime/Engine/Classes/Engine/LevelScriptBlueprint.h b/Engine/Source/Runtime/Engine/Classes/Engine/LevelScriptBlueprint.h
--- a/Engine/Source/Runtime/Engine/Classes/Engine/LevelScriptBlueprint.h
+++ b/Engine/Source/Runtime/Engine/Classes/Engine/LevelScriptBlueprint.h
@@ -38,8 +38,6 @@ class ULevelScriptBlueprint : public UBlueprint
 	}
 
 	//~ Begin UBlueprint Interface
-	ENGINE_API virtual UObject* GetObjectBeingDebugged() override;
-	ENGINE_API virtual void SetObjectBeingDebugged(UObject* NewObject) override;
 	ENGINE_API virtual FString GetFriendlyName() const override;
 	//~ End UBlueprint Interface
 
diff --git a/Engine/Source/Runtime/Engine/Private/Blueprint.cpp b/Engine/Source/Runtime/Engine/Private/Blueprint.cpp
--- a/Engine/Source/Runtime/Engine/Private/Blueprint.cpp
+++ b/Engine/Source/Runtime/Engine/Private/Blueprint.cpp
@@ -862,8 +862,8 @@ bool UBlueprint::SupportsNativization(FText* OutReason) const
 
 void UBlueprint::SetObjectBeingDebugged(UObject* NewObject)
 {
-	// Unregister the old object
-	if (UObject* OldObject = CurrentObjectBeingDebugged.Get())
+	// Unregister the old object (even if PendingKill)
+	if (UObject* OldObject = CurrentObjectBeingDebugged.Get(true))
 	{
 		if (OldObject == NewObject)
 		{
@@ -908,36 +908,6 @@ void UBlueprint::GetReparentingRules(TSet< const UClass* >& AllowedChildrenOfCla
 
 }
 
-UObject* UBlueprint::GetObjectBeingDebugged()
-{
-	UObject* DebugObj = CurrentObjectBeingDebugged.Get();
-	if(DebugObj)
-	{
-		//Check whether the object has been deleted.
-		if(DebugObj->IsPendingKill())
-		{
-			SetObjectBeingDebugged(NULL);
-			DebugObj = NULL;
-		}
-	}
-	return DebugObj;
-}
-
-UWorld* UBlueprint::GetWorldBeingDebugged()
-{
-	UWorld* DebugWorld = CurrentWorldBeingDebugged.Get();
-	if (DebugWorld)
-	{
-		if(DebugWorld->IsPendingKill())
-		{
-			SetWorldBeingDebugged(NULL);
-			DebugWorld = NULL;
-		}
-	}
-
-	return DebugWorld;
-}
-
 void UBlueprint::GetAssetRegistryTags(TArray<FAssetRegistryTag>& OutTags) const
 {
 	// We use Generated instead of Skeleton because the CDO data is more accurate on Generated
diff --git a/Engine/Source/Runtime/Engine/Private/Level.cpp b/Engine/Source/Runtime/Engine/Private/Level.cpp
--- a/Engine/Source/Runtime/Engine/Private/Level.cpp
+++ b/Engine/Source/Runtime/Engine/Private/Level.cpp
@@ -2014,11 +2014,20 @@ void ULevel::OnLevelScriptBlueprintChanged(ULevelScriptBlueprint* InBlueprint)
 		// Make sure this is OUR level scripting blueprint
 		ensureMsgf(InBlueprint == LevelScriptBlueprint, TEXT("Level ('%s') received OnLevelScriptBlueprintChanged notification for the wrong Blueprint ('%s')."), LevelScriptBlueprint ? *LevelScriptBlueprint->GetPathName() : TEXT("NULL"), *InBlueprint->GetPathName()) )
 	{
+		bool bResetDebugObject = false;
+
 		UClass* SpawnClass = (LevelScriptBlueprint->GeneratedClass) ? LevelScriptBlueprint->GeneratedClass : LevelScriptBlueprint->SkeletonGeneratedClass;
 
 		// Get rid of the old LevelScriptActor
 		if( LevelScriptActor )
 		{
+			// Clear the current debug object and indicate that it needs to be reset (below).
+			if (InBlueprint->GetObjectBeingDebugged() == LevelScriptActor)
+			{
+				bResetDebugObject = true;
+				InBlueprint->SetObjectBeingDebugged(nullptr);
+			}
+
 			LevelScriptActor->MarkPendingKill();
 			LevelScriptActor = nullptr;
 		}
@@ -2031,6 +2040,12 @@ void ULevel::OnLevelScriptBlueprintChanged(ULevelScriptBlueprint* InBlueprint)
 
 		if( LevelScriptActor )
 		{
+			// Reset the current debug object to the new instance if it was previously set to the old instance.
+			if (bResetDebugObject)
+			{
+				InBlueprint->SetObjectBeingDebugged(LevelScriptActor);
+			}
+
 			LevelScriptActor->ClearFlags(RF_Transactional);
 			check(LevelScriptActor->GetOuter() == this);
 			// Finally, fixup all the bound events to point to their new LSA
diff --git a/Engine/Source/Runtime/Engine/Private/LevelScriptActor.cpp b/Engine/Source/Runtime/Engine/Private/LevelScriptActor.cpp
--- a/Engine/Source/Runtime/Engine/Private/LevelScriptActor.cpp
+++ b/Engine/Source/Runtime/Engine/Private/LevelScriptActor.cpp
@@ -30,29 +30,6 @@ ALevelScriptActor::ALevelScriptActor(const FObjectInitializer& ObjectInitializer
 	bReplayRewindable = true;
 }
 
-#if WITH_EDITOR
-void ALevelScriptActor::PostDuplicate(bool bDuplicateForPIE)
-{
-	ULevelScriptBlueprint* MyBlueprint = Cast<ULevelScriptBlueprint>(GetClass()->ClassGeneratedBy);
-	if (MyBlueprint && !GIsDuplicatingClassForReinstancing && !IsPendingKill())
-	{
-		MyBlueprint->SetObjectBeingDebugged(this);
-	}
-
-	Super::PostDuplicate(bDuplicateForPIE);
-}
-
-void ALevelScriptActor::BeginDestroy()
-{
-	if (ULevelScriptBlueprint* MyBlueprint = Cast<ULevelScriptBlueprint>(GetClass()->ClassGeneratedBy))
-	{
-		MyBlueprint->SetObjectBeingDebugged(NULL);
-	}
-
-	Super::BeginDestroy();
-}
-#endif
-
 void ALevelScriptActor::PreInitializeComponents()
 {
 	if (UInputDelegateBinding::SupportsInputDelegate(GetClass()))
diff --git a/Engine/Source/Runtime/Engine/Private/LevelScriptBlueprint.cpp b/Engine/Source/Runtime/Engine/Private/LevelScriptBlueprint.cpp
--- a/Engine/Source/Runtime/Engine/Private/LevelScriptBlueprint.cpp
+++ b/Engine/Source/Runtime/Engine/Private/LevelScriptBlueprint.cpp
@@ -15,33 +15,6 @@ ULevelScriptBlueprint::ULevelScriptBlueprint(const FObjectInitializer& ObjectIni
 
 #if WITH_EDITOR
 
-void ULevelScriptBlueprint::SetObjectBeingDebugged(UObject* NewObject)
-{
-	// Only allowed because UWorld::TransferBlueprintDebugReferences needs to get us looking at our duplicated self
-	UBlueprint::SetObjectBeingDebugged(NewObject);
-}
-
-UObject* ULevelScriptBlueprint::GetObjectBeingDebugged()
-{ 
-	UObject* CachedResult = UBlueprint::GetObjectBeingDebugged();
-
-	if (CachedResult != NULL)
-	{
-		return CachedResult;
-	}
-	else
-	{
-		if (ULevel* Level = GetLevel()) 
-		{ 
-			CachedResult = Level->GetLevelScriptActor(); 
-		}
-
-		SetObjectBeingDebugged(CachedResult);
-	}
-
-	return CachedResult; 
-} 
-
 FString ULevelScriptBlueprint::GetFriendlyName() const
 {
 #if WITH_EDITORONLY_DATA
diff --git a/Engine/Source/Runtime/Engine/Private/World.cpp b/Engine/Source/Runtime/Engine/Private/World.cpp
--- a/Engine/Source/Runtime/Engine/Private/World.cpp
+++ b/Engine/Source/Runtime/Engine/Private/World.cpp
@@ -1996,17 +1996,7 @@ void UWorld::TransferBlueprintDebugReferences(UWorld* NewWorld)
 			}
 		}
 	}
-	// Ensure the level script actor debug references are transferred to the new world.
-	if (NewWorld)
-	{
-		if (ALevelScriptActor* LevelScript = NewWorld->GetLevelScriptActor())
-		{
-			if (UBlueprint* LevelScriptBlueprint = Cast<UBlueprint>(LevelScript->GetClass()->ClassGeneratedBy))
-			{
-				LevelScriptBlueprint->SetObjectBeingDebugged(LevelScript);
-			}
-		}
-	}
+
 	// Empty the map, anything useful got moved over the map in the new world
 	BlueprintObjectsBeingDebugged.Empty();
 #endif	//#if WITH_EDITOR
