Add compiler extensions, providing game developer modules/monolithic editor extensions a chance to do further game specific blueprint validation

diff --git a/Engine/Source/Editor/Kismet/Private/BlueprintCompilationManager.cpp b/Engine/Source/Editor/Kismet/Private/BlueprintCompilationManager.cpp
--- a/Engine/Source/Editor/Kismet/Private/BlueprintCompilationManager.cpp
+++ b/Engine/Source/Editor/Kismet/Private/BlueprintCompilationManager.cpp
@@ -3,6 +3,7 @@
 #include "BlueprintCompilationManager.h"
 
 #include "Kismet2/BlueprintEditorUtils.h"
+#include "BlueprintCompilerExtension.h"
 #include "BlueprintEditorSettings.h"
 #include "Blueprint/BlueprintSupport.h"
 #include "Kismet2/CompilerResultsLog.h"
@@ -80,6 +81,7 @@
 
 struct FReinstancingJob;
 struct FSkeletonFixupData;
+struct FCompilerData;
 
 struct FBlueprintCompilationManagerImpl : public FGCObject
 {
@@ -89,9 +91,12 @@ struct FBlueprintCompilationManagerImpl : public FGCObject
 	// FGCObject:
 	virtual void AddReferencedObjects(FReferenceCollector& Collector);
 
+	void RegisterCompilerExtension(TSubclassOf<UBlueprint> BlueprintType, UBlueprintCompilerExtension* Extension);
+
 	void QueueForCompilation(const FBPCompileRequest& CompileJob);
 	void CompileSynchronouslyImpl(const FBPCompileRequest& Request);
 	void FlushCompilationQueueImpl(bool bSuppressBroadcastCompiled, TArray<UBlueprint*>* BlueprintsCompiled, TArray<UBlueprint*>* BlueprintsCompiledOrSkeletonCompiled, FUObjectSerializeContext* InLoadContext);
+	void ProcessExtensions(const TArray<FCompilerData>& InCurrentlyCompilingBPs);
 	void FlushReinstancingQueueImpl();
 	bool HasBlueprintsToCompile() const;
 	bool IsGeneratedClassLayoutReady() const;
@@ -113,6 +118,11 @@ struct FBlueprintCompilationManagerImpl : public FGCObject
 		virtual FArchive& operator<<( UObject*& Obj ) override;
 	};
 
+	// Extension data, could be organized in many ways, but this provides an easy way
+	// to extend blueprint compilation after the graph has been pruned and functions
+	// have been generated (but before code is generated):
+	TMap<UClass*, TArray<UBlueprintCompilerExtension*> > CompilerExtensions;
+
 	// Queued requests to be processed in the next FlushCompilationQueueImpl call:
 	TArray<FBPCompileRequest> QueuedRequests;
 	
@@ -148,6 +158,12 @@ FBlueprintCompilationManagerImpl::~FBlueprintCompilationManagerImpl()
 
 void FBlueprintCompilationManagerImpl::AddReferencedObjects(FReferenceCollector& Collector)
 {
+	for(TPair<UClass*, TArray<UBlueprintCompilerExtension*>>& Extensions : CompilerExtensions)
+	{
+		Collector.AddReferencedObject(Extensions.Key);
+		Collector.AddReferencedObjects(Extensions.Value);
+	}
+
 	for( FBPCompileRequest& Job : QueuedRequests )
 	{
 		Collector.AddReferencedObject(Job.BPToCompile);
@@ -157,6 +173,11 @@ void FBlueprintCompilationManagerImpl::AddReferencedObjects(FReferenceCollector&
 	Collector.AddReferencedObjects(OldCDOs);
 }
 
+void FBlueprintCompilationManagerImpl::RegisterCompilerExtension(TSubclassOf<UBlueprint> BlueprintType, UBlueprintCompilerExtension* Extension)
+{
+	CompilerExtensions.FindOrAdd(BlueprintType).Emplace(Extension);
+}
+
 void FBlueprintCompilationManagerImpl::QueueForCompilation(const FBPCompileRequest& CompileJob)
 {
 	if(!CompileJob.BPToCompile->bQueuedForCompilation)
@@ -1009,7 +1030,9 @@ void FBlueprintCompilationManagerImpl::FlushCompilationQueueImpl(bool bSuppressB
 			}
 		}
 		bGeneratedClassLayoutReady = true;
-	
+		
+		ProcessExtensions(CurrentlyCompilingBPs);
+
 		// STAGE XIII: Compile functions
 		UBlueprintEditorSettings* Settings = GetMutableDefault<UBlueprintEditorSettings>();
 		
@@ -1286,6 +1309,57 @@ void FBlueprintCompilationManagerImpl::FlushCompilationQueueImpl(bool bSuppressB
 	ensure(QueuedRequests.Num() == 0);
 }
 
+void FBlueprintCompilationManagerImpl::ProcessExtensions(const TArray<FCompilerData>& InCurrentlyCompilingBPs)
+{
+	if(CompilerExtensions.Num() == 0)
+	{
+		return;
+	}
+
+	for (const FCompilerData& CompilerData : InCurrentlyCompilingBPs)
+	{
+		UBlueprint* BP = CompilerData.BP;
+		FBlueprintCompiledData CompiledData; // populate only if we find an extension:
+
+		if(CompilerData.ShouldCompileClassLayout())
+		{
+			// give extension a chance to raise errors, or save off data:
+			UClass* Iter = BP->GetClass();
+			while(Iter != UBlueprint::StaticClass()->GetSuperClass())
+			{
+				TArray<UBlueprintCompilerExtension*>* Extensions = CompilerExtensions.Find(Iter);
+				if(Extensions)
+				{
+					// extension found, store off data from compiler that we want to expose to extensions:
+					if(CompiledData.IntermediateGraphs.Num() == 0)
+					{
+						if(CompilerData.Compiler->ConsolidatedEventGraph)
+						{
+							CompiledData.IntermediateGraphs.Emplace(CompilerData.Compiler->ConsolidatedEventGraph);
+						}
+								
+						for(const FKismetFunctionContext& Fn : CompilerData.Compiler->FunctionList)
+						{
+							if(Fn.SourceGraph == CompilerData.Compiler->ConsolidatedEventGraph)
+							{
+								continue;
+							}
+
+							CompiledData.IntermediateGraphs.Emplace(Fn.SourceGraph);
+						}
+					}
+
+					for(UBlueprintCompilerExtension* Extension : *Extensions)
+					{
+						Extension->BlueprintCompiled(*CompilerData.Compiler, CompiledData);
+					}
+				}
+				Iter = Iter->GetSuperClass();
+			}
+		}
+	}
+}
+
 void FBlueprintCompilationManagerImpl::FlushReinstancingQueueImpl()
 {
 	if(GCompilingBlueprint)
@@ -2625,5 +2699,11 @@ void FBlueprintCompilationManager::ReparentHierarchies(const TMap<UClass*, UClas
 	FBlueprintCompilationManagerImpl::ReparentHierarchies(OldClassToNewClass);
 }
 
+void FBlueprintCompilationManager::RegisterCompilerExtension(TSubclassOf<UBlueprint> BlueprintType, UBlueprintCompilerExtension* Extension)
+{
+	Initialize();
+	BPCMImpl->RegisterCompilerExtension(BlueprintType, Extension);
+}
+
 #undef LOCTEXT_NAMESPACE
 
diff --git a/Engine/Source/Editor/Kismet/Private/BlueprintCompilerExtension.cpp b/Engine/Source/Editor/Kismet/Private/BlueprintCompilerExtension.cpp
new file mode 100644
--- /dev/null
+++ b/Engine/Source/Editor/Kismet/Private/BlueprintCompilerExtension.cpp
@@ -0,0 +1,14 @@
+// Copyright 1998-2019 Epic Games, Inc. All Rights Reserved.
+
+#include "BlueprintCompilerExtension.h"
+
+UBlueprintCompilerExtension::UBlueprintCompilerExtension(const FObjectInitializer& ObjectInitializer)
+	: UObject(ObjectInitializer)
+{
+}
+
+void UBlueprintCompilerExtension::BlueprintCompiled(const FKismetCompilerContext& CompilationContext, const FBlueprintCompiledData& Data)
+{
+	// common entry point in case we need to add logging, profiling, etc
+	ProcessBlueprintCompiled(CompilationContext, Data);
+}
diff --git a/Engine/Source/Editor/Kismet/Public/BlueprintCompilationManager.h b/Engine/Source/Editor/Kismet/Public/BlueprintCompilationManager.h
--- a/Engine/Source/Editor/Kismet/Public/BlueprintCompilationManager.h
+++ b/Engine/Source/Editor/Kismet/Public/BlueprintCompilationManager.h
@@ -7,6 +7,7 @@
 #include "Kismet2/KismetEditorUtilities.h"
 
 class UBlueprint;
+class UBlueprintCompilerExtension;
 class FCompilerResultsLog;
 
 struct FBPCompileRequest
@@ -81,6 +82,14 @@ struct KISMET_API FBlueprintCompilationManager
 	 * reinstancing could be avoided when layouts match.
 	 */
 	static void ReparentHierarchies(const TMap<UClass*, UClass*>& OldClassToNewClass);
+
+	/** 
+	 * Registers a blueprint compiler extension - anytime a blueprint of the provided type is compiled
+	 * the extension will be activated. Note that because editor initialization may require blueprint
+	 * compilation there may be blueprints compiled before the extension is registed unless special
+	 * care has been taken.
+	 */
+	static void RegisterCompilerExtension(TSubclassOf<UBlueprint> BlueprintType, UBlueprintCompilerExtension* Extension);
 private:
 	FBlueprintCompilationManager();
 };
diff --git a/Engine/Source/Editor/Kismet/Public/BlueprintCompilerExtension.h b/Engine/Source/Editor/Kismet/Public/BlueprintCompilerExtension.h
new file mode 100644
--- /dev/null
+++ b/Engine/Source/Editor/Kismet/Public/BlueprintCompilerExtension.h
@@ -0,0 +1,32 @@
+// Copyright 1998-2019 Epic Games, Inc. All Rights Reserved.
+#pragma once
+
+#include "KismetCompiler.h"
+
+#include "BlueprintCompilerExtension.generated.h"
+
+USTRUCT()
+struct FBlueprintCompiledData
+{
+	GENERATED_BODY()
+
+	UPROPERTY()
+	TArray<UEdGraph*> IntermediateGraphs;
+};
+
+UCLASS(Abstract)
+class KISMET_API UBlueprintCompilerExtension : public UObject
+{
+	GENERATED_BODY()
+public:
+	UBlueprintCompilerExtension(const FObjectInitializer& ObjectInitializer);
+
+	void BlueprintCompiled(const FKismetCompilerContext& CompilationContext, const FBlueprintCompiledData& Data);
+
+protected:
+	/** 
+	 * Override this if you're interested in running logic after class layout has been
+	 * generated, but before bytecode and member variables have been 
+	 */
+	virtual void ProcessBlueprintCompiled(const FKismetCompilerContext& CompilationContext, const FBlueprintCompiledData& Data) {}
+};
