Honor ClassWithin contract when reinstancing archetypes

diff --git a/Engine/Source/Editor/Kismet/Private/BlueprintCompilationManager.cpp b/Engine/Source/Editor/Kismet/Private/BlueprintCompilationManager.cpp
--- a/Engine/Source/Editor/Kismet/Private/BlueprintCompilationManager.cpp
+++ b/Engine/Source/Editor/Kismet/Private/BlueprintCompilationManager.cpp
@@ -107,6 +107,7 @@ struct FBlueprintCompilationManagerImpl : public FGCObject
 	static void ReinstanceBatch(TArray<FReinstancingJob>& Reinstancers, TMap< UClass*, UClass* >& InOutOldToNewClassMap, FUObjectSerializeContext* InLoadContext);
 	static UClass* FastGenerateSkeletonClass(UBlueprint* BP, FKismetCompilerContext& CompilerContext, bool bIsSkeletonOnly, TArray<FSkeletonFixupData>& OutSkeletonFixupData);
 	static bool IsQueuedForCompilation(UBlueprint* BP);
+	static UObject* GetOuterForRename(UClass* ForClass);
 	
 	// Declaration of archive to fix up bytecode references of blueprints that are actively compiled:
 	class FFixupBytecodeReferences : public FArchiveUObject
@@ -1844,11 +1845,13 @@ void FBlueprintCompilationManagerImpl::ReinstanceBatch(TArray<FReinstancingJob>&
 				FName OriginalName = Archetype->GetFName();
 				UObject* OriginalOuter = Archetype->GetOuter();
 				EObjectFlags OriginalFlags = Archetype->GetFlags();
+
+				UObject* Destination = GetOuterForRename(Archetype->GetClass());
 				Archetype->Rename(
 					nullptr,
 					// destination - this is the important part of this call. Moving the object 
 					// out of the way so we can reuse its name:
-					GetTransientPackage(), 
+					Destination, 
 					// Rename options:
 					REN_DoNotDirty | REN_DontCreateRedirectors | REN_ForceNoResetLoaders );
 
@@ -2459,6 +2462,19 @@ bool FBlueprintCompilationManagerImpl::IsQueuedForCompilation(UBlueprint* BP)
 	return BP->bQueuedForCompilation;
 }
 
+UObject* FBlueprintCompilationManagerImpl::GetOuterForRename(UClass* ForClass)
+{
+	// if someone has tautologically placed themself within their own hierarchy then we'll
+	// just assume they're ok with eventually being outered to a upackage, similar UPackage
+	// is a UObject, so if someone demands that they be outered to 'a uobject' we'll 
+	// just leave them directly parented to the transient package:
+	if(ForClass->ClassWithin && ForClass->ClassWithin != ForClass && ForClass->ClassWithin != UObject::StaticClass())
+	{
+		return NewObject<UObject>( GetOuterForRename(ForClass->ClassWithin), ForClass->ClassWithin, NAME_None, RF_Transient );
+	}
+	return GetTransientPackage();
+}
+
 // FFixupBytecodeReferences Implementation:
 FBlueprintCompilationManagerImpl::FFixupBytecodeReferences::FFixupBytecodeReferences(UObject* InObject)
 {
