Blueprint: Reordering categories now correctly creates an undo record and marks the Blueprint as modified

diff --git a/Engine/Source/Editor/Kismet/Private/SMyBlueprint.cpp b/Engine/Source/Editor/Kismet/Private/SMyBlueprint.cpp
--- a/Engine/Source/Editor/Kismet/Private/SMyBlueprint.cpp
+++ b/Engine/Source/Editor/Kismet/Private/SMyBlueprint.cpp
@@ -2815,33 +2815,37 @@ void SMyBlueprint::ExpandCategory(const FText& CategoryName)
 	GraphActionMenu->ExpandCategory(CategoryName);
 }
 
-bool SMyBlueprint::MoveCategoryBeforeCategory( const FText& InCategoryToMove, const FText& InTargetCategory )
+bool SMyBlueprint::MoveCategoryBeforeCategory(const FText& InCategoryToMove, const FText& InTargetCategory)
 {
 	bool bResult = false;
-	UBlueprint* BlueprintObj = BlueprintEditorPtr.Pin()->GetBlueprintObj();
 
 	FString CategoryToMoveString = InCategoryToMove.ToString();
 	FString TargetCategoryString = InTargetCategory.ToString();
-	if( BlueprintObj )
+	if (UBlueprint* BlueprintObj = BlueprintEditorPtr.Pin()->GetBlueprintObj())
 	{
+		FScopedTransaction Transaction(LOCTEXT("ReorderCategories", "Reorder Categories"));
+		BlueprintObj->Modify();
+
 		// Find root categories
-		int32 RootCategoryDelim = CategoryToMoveString.Find( TEXT( "|" ), ESearchCase::CaseSensitive );
-		FName CategoryToMove = RootCategoryDelim == INDEX_NONE ? *CategoryToMoveString : *CategoryToMoveString.Left( RootCategoryDelim );
-		RootCategoryDelim = TargetCategoryString.Find( TEXT( "|" ), ESearchCase::CaseSensitive );
-		FName TargetCategory = RootCategoryDelim == INDEX_NONE ? *TargetCategoryString : *TargetCategoryString.Left( RootCategoryDelim );
+		int32 RootCategoryDelim = CategoryToMoveString.Find(TEXT("|"), ESearchCase::CaseSensitive);
+		FName CategoryToMove = RootCategoryDelim == INDEX_NONE ? *CategoryToMoveString : *CategoryToMoveString.Left(RootCategoryDelim);
+		RootCategoryDelim = TargetCategoryString.Find(TEXT("|"), ESearchCase::CaseSensitive);
+		FName TargetCategory = RootCategoryDelim == INDEX_NONE ? *TargetCategoryString : *TargetCategoryString.Left(RootCategoryDelim);
 
 		TArray<FName>& CategorySort = BlueprintObj->CategorySorting;
-		const int32 RemovalIndex = CategorySort.Find( CategoryToMove );
+
 		// Remove existing sort index
-		if( RemovalIndex != INDEX_NONE )
+		const int32 RemovalIndex = CategorySort.Find(CategoryToMove);
+		if (RemovalIndex != INDEX_NONE)
 		{
-			CategorySort.RemoveAt( RemovalIndex );
+			CategorySort.RemoveAt(RemovalIndex);
 		}
+
 		// Update the Category sort order and refresh ( if the target category has an entry )
-		const int32 InsertIndex = CategorySort.Find( TargetCategory );
-		if( InsertIndex != INDEX_NONE )
+		const int32 InsertIndex = CategorySort.Find(TargetCategory);
+		if (InsertIndex != INDEX_NONE)
 		{
-			CategorySort.Insert( CategoryToMove, InsertIndex );
+			CategorySort.Insert(CategoryToMove, InsertIndex);
 			Refresh();
 			bResult = true;
 		}
