Blueprints: Changed the function icon for pure functions to be green in the My Blueprint list
Adapted from PR #5757: Green pure functions in blueprint function list (Contributed by [REDACTED_PERSON])

diff --git a/Engine/Source/Editor/EditorStyle/Private/SlateEditorStyle.cpp b/Engine/Source/Editor/EditorStyle/Private/SlateEditorStyle.cpp
--- a/Engine/Source/Editor/EditorStyle/Private/SlateEditorStyle.cpp
+++ b/Engine/Source/Editor/EditorStyle/Private/SlateEditorStyle.cpp
@@ -4467,8 +4467,10 @@ void FSlateEditorStyle::FStyle::SetupGraphEditorStyles()
 		Set( "GraphEditor.InterfaceFunction_16x", new IMAGE_BRUSH( "Icons/icon_Blueprint_Interfacefunction_16x", Icon16x16 ) );
 		Set( "GraphEditor.Macro_16x", new IMAGE_BRUSH( "Icons/icon_Blueprint_Macro_16x", Icon16x16 ) );
 		Set( "GraphEditor.Function_16x", new IMAGE_BRUSH( "Icons/icon_Blueprint_NewFunction_16x", Icon16x16 ) );
+		Set( "GraphEditor.PureFunction_16x", new IMAGE_BRUSH( "Icons/icon_Blueprint_NewPureFunction_16x", Icon16x16 ) );
 		Set( "GraphEditor.PotentialOverrideFunction_16x", new IMAGE_BRUSH( "Icons/icon_Blueprint_OverrideableFunction_16x", Icon16x16 ) );
 		Set( "GraphEditor.OverrideFunction_16x", new IMAGE_BRUSH( "Icons/icon_Blueprint_OverrideFunction_16x", Icon16x16 ) );
+		Set( "GraphEditor.OverridePureFunction_16x", new IMAGE_BRUSH( "Icons/icon_Blueprint_OverridePureFunction_16x", Icon16x16 ) );
 		Set( "GraphEditor.SubGraph_16x", new IMAGE_BRUSH( "Icons/icon_Blueprint_SubgraphComposite_16x", Icon16x16 ) );
 		Set( "GraphEditor.Animation_16x", new IMAGE_BRUSH( "Icons/icon_Blueprint_Anim_16x", Icon16x16 ) );
 		Set( "GraphEditor.Conduit_16x", new IMAGE_BRUSH( "Icons/icon_Blueprint_Conduit_16x", Icon16x16 ) );
diff --git a/Engine/Source/Editor/Kismet/Private/SBlueprintPalette.cpp b/Engine/Source/Editor/Kismet/Private/SBlueprintPalette.cpp
--- a/Engine/Source/Editor/Kismet/Private/SBlueprintPalette.cpp
+++ b/Engine/Source/Editor/Kismet/Private/SBlueprintPalette.cpp
@@ -53,6 +53,7 @@
 #include "BlueprintEditorSettings.h"
 #include "Widgets/Text/SInlineEditableTextBlock.h"
 #include "SPinTypeSelector.h"
+#include "GraphEditorSettings.h"
 
 #define LOCTEXT_NAMESPACE "BlueprintPalette"
 
@@ -250,25 +251,29 @@ static void GetSubGraphIcon(FEdGraphSchemaAction_K2Graph const* const ActionIn,
 		break;
 	case EEdGraphSchemaAction_K2Graph::Function:
 		{
-			if ( ActionIn->EdGraph == nullptr )
+			if (ActionIn->EdGraph == nullptr)
 			{
 				IconOut = FEditorStyle::GetBrush(TEXT("GraphEditor.PotentialOverrideFunction_16x"));
 				ToolTipOut = LOCTEXT("PotentialOverride_Tooltip", "Potential Override");	
 			}
 			else
 			{
-				if ( ActionIn->EdGraph->IsA(UAnimationGraph::StaticClass()) )
+				if (ActionIn->EdGraph->IsA(UAnimationGraph::StaticClass()))
 				{
 					IconOut = FEditorStyle::GetBrush(TEXT("GraphEditor.Animation_16x"));
 				}
 				else if (UFunction* OverrideFunc = FindField<UFunction>(BlueprintIn->ParentClass, ActionIn->FuncName))
 				{
-					IconOut = FEditorStyle::GetBrush(TEXT("GraphEditor.OverrideFunction_16x"));
+					const bool bIsPureFunction = OverrideFunc && OverrideFunc->HasAnyFunctionFlags(FUNC_BlueprintPure);
+					IconOut = FEditorStyle::GetBrush(bIsPureFunction ? TEXT("GraphEditor.OverridePureFunction_16x") : TEXT("GraphEditor.OverrideFunction_16x"));
 					ToolTipOut = LOCTEXT("Override_Tooltip", "Override");
 				}
 				else
 				{
-					IconOut = FEditorStyle::GetBrush(TEXT("GraphEditor.Function_16x"));
+					UFunction* Function = FindField<UFunction>(BlueprintIn->SkeletonGeneratedClass, ActionIn->FuncName);
+					const bool bIsPureFunction = Function && Function->HasAnyFunctionFlags(FUNC_BlueprintPure);
+
+					IconOut = FEditorStyle::GetBrush(bIsPureFunction ? TEXT("GraphEditor.PureFunction_16x") : TEXT("GraphEditor.Function_16x"));
 					if (ActionIn->EdGraph->IsA(UAnimationGraph::StaticClass()))
 					{
 						ToolTipOut = LOCTEXT("AnimationGraph_Tooltip", "Animation Graph");
