Added a cvar bp.EnableAutomaticLibraryAssetLoading to allow disabling the automatic loading of Blueprint macro/function library assets when a BP editor window is opened
- This dramatically speeds up the first time to open the editor in a large project, but the unloaded nodes are *not* present in the context menu until something else causes a given function library / macro library to be loaded, so it isn't a good all-purpose solution

diff --git a/Engine/Source/Editor/Kismet/Private/BlueprintEditor.cpp b/Engine/Source/Editor/Kismet/Private/BlueprintEditor.cpp
--- a/Engine/Source/Editor/Kismet/Private/BlueprintEditor.cpp
+++ b/Engine/Source/Editor/Kismet/Private/BlueprintEditor.cpp
@@ -161,6 +161,13 @@
 
 #define LOCTEXT_NAMESPACE "BlueprintEditor"
 
+static int32 EnableAutomaticLibraryAssetLoading = 1;
+static FAutoConsoleVariableRef CVarEnableAutomaticLibraryAssetLoading(
+	TEXT("bp.EnableAutomaticLibraryAssetLoading"),
+	EnableAutomaticLibraryAssetLoading,
+	TEXT("Should opening the BP editor load all macro and function library assets or not?\n0: Disable, 1: Enable (defaults to enabled)\nNodes defined in unloaded libraries will not show up in the context menu!"),
+	ECVF_Default);
+
 /////////////////////////////////////////////////////
 // FSelectionDetailsSummoner
 
@@ -1681,6 +1688,11 @@ void FBlueprintEditor::CommonInitialization(const TArray<UBlueprint*>& InitBluep
 
 void FBlueprintEditor::LoadLibrariesFromAssetRegistry()
 {
+	if (EnableAutomaticLibraryAssetLoading == 0)
+	{
+		return;
+	}
+
 	if( GetBlueprintObj() )
 	{
 		FString UserDeveloperPath = FPackageName::FilenameToLongPackageName( FPaths::GameUserDeveloperDir());
