Before duplicating, clear the application keyboard focus to ensure that any component being renamed lost the focus and commit their text first. The duplicate operation triggers a regeneration of the tree view, releasing the row widgets. If one row had the keyboard focus (editing the component name), the 'weak focus path' maintained by Slate gets truncated (they are weak pointers, so this is fine) but when the focus changed, Slate could not pin the row widget that had the focus and therefore, the 'focus lost' event was never handled by the now deleted widget. Since the 'rename' transaction are performed on focus lost, the rename was not taken into account. For 'create + give initial name' transaction, the responsibility to end the create transaction was passed to the row item, but the transaction would not end because on focus lost was never called and the rename never ended. In that case, any further transactions were appended to the 'create' transaction until the node item was released. (Like deleting the component).

diff --git a/Engine/Source/Editor/Kismet/Private/SSCSEditor.cpp b/Engine/Source/Editor/Kismet/Private/SSCSEditor.cpp
--- a/Engine/Source/Editor/Kismet/Private/SSCSEditor.cpp
+++ b/Engine/Source/Editor/Kismet/Private/SSCSEditor.cpp
@@ -4167,6 +4167,11 @@ void SSCSEditor::OnDuplicateComponent()
 	TArray<FSCSEditorTreeNodePtrType> SelectedNodes = SCSTreeWidget->GetSelectedItems();
 	if(SelectedNodes.Num() > 0)
 	{
+		// Force the text box being edited (if any) to commit its text. The duplicate operation may trigger a regeneration of the tree view,
+		// releasing all row widgets. If one row was in edit mode (rename/rename on create), it was released before losing the focus and
+		// this would prevent the completion of the 'rename' or 'create + give initial name' transaction (occurring on focus lost).
+		FSlateApplication::Get().ClearKeyboardFocus();
+
 		const FScopedTransaction Transaction(SelectedNodes.Num() > 1 ? LOCTEXT("DuplicateComponents", "Duplicate Components") : LOCTEXT("DuplicateComponent", "Duplicate Component"));
 
 		TMap<USceneComponent*, USceneComponent*> DuplicateSceneComponentMap;
