Fix invalid LoadObject() call on attempt to change focus to an empty FiB search result.

diff --git a/Engine/Source/Editor/Kismet/Private/FindInBlueprints.cpp b/Engine/Source/Editor/Kismet/Private/FindInBlueprints.cpp
--- a/Engine/Source/Editor/Kismet/Private/FindInBlueprints.cpp
+++ b/Engine/Source/Editor/Kismet/Private/FindInBlueprints.cpp
@@ -956,7 +956,7 @@ EActiveTimerReturnType SFindInBlueprints::UpdateSearchResults( double InCurrentT
 			if ( ItemsFound.Num() == 0 )
 			{
 				// Insert a fake result to inform user if none found
-				ItemsFound.Add( FSearchResult( new FFindInBlueprintsResult( LOCTEXT( "BlueprintSearchNoResults", "No Results found" ) ) ) );
+				ItemsFound.Add( FSearchResult( new FFindInBlueprintsNoResult( LOCTEXT( "BlueprintSearchNoResults", "No Results found" ) ) ) );
 				TreeView->RequestTreeRefresh();
 			}
 
@@ -1103,7 +1103,7 @@ void SFindInBlueprints::MakeSearchQuery(FString InSearchString, bool bInIsFindWi
 				}
 
 				// Insert a fake result to inform user if none found
-				ItemsFound.Add(FSearchResult(new FFindInBlueprintsResult(NoResultsText)));
+				ItemsFound.Add(FSearchResult(new FFindInBlueprintsNoResult(NoResultsText)));
 				HighlightText = FText::GetEmpty();
 			}
 			else
diff --git a/Engine/Source/Editor/Kismet/Public/FindInBlueprints.h b/Engine/Source/Editor/Kismet/Public/FindInBlueprints.h
--- a/Engine/Source/Editor/Kismet/Public/FindInBlueprints.h
+++ b/Engine/Source/Editor/Kismet/Public/FindInBlueprints.h
@@ -85,6 +85,22 @@ namespace FindInBlueprintsHelpers
 	void ExpandAllChildren(FSearchResult InTreeNode, TSharedPtr<STreeView<TSharedPtr<FFindInBlueprintsResult>>> InTreeView);
 }
 
+/** Class used to denote an empty search result */
+class FFindInBlueprintsNoResult : public FFindInBlueprintsResult
+{
+public:
+	FFindInBlueprintsNoResult(const FText& InDisplayText)
+		:FFindInBlueprintsResult(InDisplayText, nullptr)
+	{
+	}
+
+	virtual FReply OnClick() override
+	{
+		// Do nothing on click.
+		return FReply::Handled();
+	}
+};
+
 /** Graph nodes use this class to store their data */
 class FFindInBlueprintsGraphNode : public FFindInBlueprintsResult
 {
