Edigrate 6960235
Fix serious errors when reloading blueprints that have derived types already loaded, reduce asset reloading batch size to 1 to improve correctness when reloading many assets

diff --git a/Engine/Source/Editor/Kismet/Private/BlueprintCompilationManager.cpp b/Engine/Source/Editor/Kismet/Private/BlueprintCompilationManager.cpp
--- a/Engine/Source/Editor/Kismet/Private/BlueprintCompilationManager.cpp
+++ b/Engine/Source/Editor/Kismet/Private/BlueprintCompilationManager.cpp
@@ -109,7 +109,8 @@ struct FBlueprintCompilationManagerImpl : public FGCObject
 	static UClass* FastGenerateSkeletonClass(UBlueprint* BP, FKismetCompilerContext& CompilerContext, bool bIsSkeletonOnly, TArray<FSkeletonFixupData>& OutSkeletonFixupData);
 	static bool IsQueuedForCompilation(UBlueprint* BP);
 	static UObject* GetOuterForRename(UClass* ForClass);
-	
+	static bool ReinstancerOrderingFunction( UClass* A, UClass* B );
+
 	// Declaration of archive to fix up bytecode references of blueprints that are actively compiled:
 	class FFixupBytecodeReferences : public FArchiveUObject
 	{
@@ -392,10 +393,45 @@ struct FCompilerData
 
 struct FReinstancingJob
 {
+	FReinstancingJob(TSharedPtr<FBlueprintCompileReinstancer> InReinstancer);
+	FReinstancingJob(TSharedPtr<FBlueprintCompileReinstancer> InReinstancer, TSharedPtr<FKismetCompilerContext> InCompiler);
+	FReinstancingJob(TPair<UClass*, UClass*> InOldToNew);
+
+	// optional:
 	TSharedPtr<FBlueprintCompileReinstancer> Reinstancer;
 	TSharedPtr<FKismetCompilerContext> Compiler;
+
+	// always set:
+	TPair<UClass*, UClass*> OldToNew;
 };
 
+FReinstancingJob::FReinstancingJob(TSharedPtr<FBlueprintCompileReinstancer> InReinstancer)
+	: Reinstancer(InReinstancer)
+	, Compiler()
+	, OldToNew(nullptr, nullptr)
+{
+	check(InReinstancer.IsValid());
+	OldToNew.Key = InReinstancer->DuplicatedClass;
+	OldToNew.Value = InReinstancer->ClassToReinstance;
+}
+
+FReinstancingJob::FReinstancingJob(TSharedPtr<FBlueprintCompileReinstancer> InReinstancer, TSharedPtr<FKismetCompilerContext> InCompiler)
+	: Reinstancer(InReinstancer)
+	, Compiler(InCompiler)
+	, OldToNew(nullptr, nullptr)
+{
+	check(InReinstancer.IsValid());
+	OldToNew.Key = InReinstancer->DuplicatedClass;
+	OldToNew.Value = InReinstancer->ClassToReinstance;
+}
+
+FReinstancingJob::FReinstancingJob(TPair<UClass*, UClass*> InOldToNew)
+	: Reinstancer()
+	, Compiler()
+	, OldToNew(InOldToNew)
+{
+}
+
 void FBlueprintCompilationManagerImpl::FlushCompilationQueueImpl(bool bSuppressBroadcastCompiled, TArray<UBlueprint*>* BlueprintsCompiled, TArray<UBlueprint*>* BlueprintsCompiledOrSkeletonCompiled, FUObjectSerializeContext* InLoadContext)
 {
 	TGuardValue<bool> GuardTemplateNameFlag(GCompilingBlueprint, true);
@@ -607,27 +643,7 @@ void FBlueprintCompilationManagerImpl::FlushCompilationQueueImpl(bool bSuppressB
 				return false;
 			}
 
-			int32 DepthA = 0;
-			int32 DepthB = 0;
-			UStruct* Iter = *(A.GeneratedClass) ? A.GeneratedClass->GetSuperStruct() : nullptr;
-			while (Iter)
-			{
-				++DepthA;
-				Iter = Iter->GetSuperStruct();
-			}
-
-			Iter = *(B.GeneratedClass) ? B.GeneratedClass->GetSuperStruct() : nullptr;
-			while (Iter)
-			{
-				++DepthB;
-				Iter = Iter->GetSuperStruct();
-			}
-
-			if (DepthA == DepthB)
-			{
-				return A.GetFName().LexicalLess(B.GetFName());
-			}
-			return DepthA < DepthB;
+			return FBlueprintCompilationManagerImpl::ReinstancerOrderingFunction(A.GeneratedClass, B.GeneratedClass);
 		};
 		CurrentlyCompilingBPs.Sort( HierarchyDepthSortFn );
 
@@ -1151,10 +1167,7 @@ void FBlueprintCompilationManagerImpl::FlushCompilationQueueImpl(bool bSuppressB
 				if(CompilerData.Reinstancer.IsValid() && CompilerData.Reinstancer->ClassToReinstance)
 				{
 					Reinstancers.Push(
-						FReinstancingJob {
-							CompilerData.Reinstancer,
-							CompilerData.Compiler
-						}
+						FReinstancingJob( CompilerData.Reinstancer, CompilerData.Compiler )
 					);
 				}
 			}
@@ -1458,65 +1471,89 @@ void FBlueprintCompilationManagerImpl::GetDefaultValue(const UClass* ForClass, c
 	}
 }
 
-void FBlueprintCompilationManagerImpl::ReparentHierarchies(const TMap<UClass*, UClass*>& OldClassToNewClass)
+void FBlueprintCompilationManagerImpl::ReparentHierarchies(const TMap<UClass*, UClass*>& OldToNewClasses)
 {
 	// something has decided to replace instances of a class. We need to update all the children of those types:
-	TSet<TPair<UClass*, UClass*>> ClassesToReinstance;
-	for(const TPair<UClass*, UClass*>& OldToNewClass : OldClassToNewClass)
+	TArray< UClass* > ClassesOrdered;
+	// Map used to distinguish between new classes and classes that need to be reinstanced (reparented) via a new reinstancer:
+	TMap<UClass*, UClass*> NewToOldClasses;
 	{
-		TArray<UClass*> DerivedClasses;
-		// Just like when compiling we have to gather all children, not just immediate children. This is so that we can
-		// update things like the ClassConstructor pointer in case it changed:
-		GetDerivedClasses(OldToNewClass.Key, DerivedClasses);
-
-		for(UClass* DerivedClass : DerivedClasses)
+		TSet<UClass*> Classes;
+		for(const TPair<UClass*, UClass*>& OldToNewClass : OldToNewClasses)
 		{
-			if(DerivedClass->GetSuperClass() == OldToNewClass.Key)
+			// classes with no CDO do not need to be reinstanced:
+			if(OldToNewClass.Key->ClassDefaultObject == nullptr)
 			{
-				// need to reparent, change the old parent class to the new one:
-				ClassesToReinstance.Add(
-					TPair<UClass*, UClass*>(DerivedClass, OldToNewClass.Value)
-				);
+				continue;
 			}
-			else
+
+			Classes.Add(OldToNewClass.Value);
+			NewToOldClasses.Add(OldToNewClass.Value, OldToNewClass.Key);
+
+			TArray<UClass*> DerivedClasses;
+			// Just like when compiling we have to gather all children, not just immediate children. This is so that we can
+			// update things like the ClassConstructor pointer in case it changed:
+			GetDerivedClasses(OldToNewClass.Key, DerivedClasses);
+
+			for(UClass* DerivedClass : DerivedClasses)
 			{
-				// just need to reinstance, parent class can remain the same as
-				// it is generally stable (outside of hotreload/asset reload):
-				ClassesToReinstance.Add(
-					TPair<UClass*, UClass*>(DerivedClass, DerivedClass->GetSuperClass())
-				);
+				if(DerivedClass->ClassDefaultObject == nullptr)
+				{
+					// on CDO->no other instances->no need to reinstance..
+					continue;
+				}
+
+				Classes.Add(DerivedClass);
 			}
 		}
+
+		ClassesOrdered = Classes.Array();
 	}
 
-	// create reinstancers:
+	// Order the classes we're about to reinstance by hierarchy depth. This will improve determinism
+	// and is as logical an order as I can come up with:
+	ClassesOrdered.Sort( [](UClass& A, UClass& B)->bool { return FBlueprintCompilationManagerImpl::ReinstancerOrderingFunction(&A, &B); } );
+
+	// create reinstancing jobs, no need to create a reinstancer when there is a new UClass* available (e.g. asset reload, hot reload):
 	TArray<FReinstancingJob> Reinstancers;
-	for(const TPair<UClass*, UClass*> OldToNewParentClass : ClassesToReinstance)
+	for(UClass* Class : ClassesOrdered)
 	{
-		Reinstancers.Push( {
-			TSharedPtr<FBlueprintCompileReinstancer>(
-				new FBlueprintCompileReinstancer(
-					OldToNewParentClass.Key,
-					EBlueprintCompileReinstancerFlags::AutoInferSaveOnCompile | EBlueprintCompileReinstancerFlags::AvoidCDODuplication
-				)
-			),
-			TSharedPtr<FKismetCompilerContext>()
-		} );
-	}
+		UClass* const* OldClass = NewToOldClasses.Find(Class);
+		if(OldClass)
+		{
+			Reinstancers.Push( FReinstancingJob( TPair<UClass*, UClass*>(*OldClass, Class) ) );
+		}
+		else
+		{
+			Reinstancers.Push( FReinstancingJob( 
+				TSharedPtr<FBlueprintCompileReinstancer>(
+					new FBlueprintCompileReinstancer(
+						Class,
+						EBlueprintCompileReinstancerFlags::AutoInferSaveOnCompile | EBlueprintCompileReinstancerFlags::AvoidCDODuplication
+					)
+				),
+				TSharedPtr<FKismetCompilerContext>()
+			) );
 
-	// reparent blueprints, now that their cdo has been moved aside and the REINST_ type has been created for old instances:
-	TMap<UClass*, UClass*> OldClassToNewClassIncludingChildren = OldClassToNewClass;
+			FReinstancingJob& ReinstancingJob = Reinstancers.Last();
+			ensure(ReinstancingJob.Reinstancer->DuplicatedClass && ReinstancingJob.Reinstancer->ClassToReinstance);
+		}
+	}
+	
+	// Reparent and Link:
+	TMap<UClass*, UClass*> OldClassToNewClassIncludingChildren = OldToNewClasses;
 	for(const FReinstancingJob& ReinstancingJob : Reinstancers)
 	{
-		UClass* ClassToReinstance = ReinstancingJob.Reinstancer->ClassToReinstance;
+		UClass* ClassToReinstance = ReinstancingJob.OldToNew.Value;
+
 		ClassToReinstance->ClassConstructor = nullptr;
 		ClassToReinstance->ClassVTableHelperCtorCaller = nullptr;
 		ClassToReinstance->ClassAddReferencedObjects = nullptr;
-
-		UClass* const* NewParent = OldClassToNewClass.Find(ClassToReinstance->GetSuperClass());
+		
+		// check to see if we're a direct parent of one of the new classes:
+		UClass* const* NewParent = OldToNewClasses.Find(ClassToReinstance->GetSuperClass());
 		if(NewParent)
 		{
-			check(*NewParent);
 			ClassToReinstance->SetSuperStruct(*NewParent);
 		}
 
@@ -1524,17 +1561,17 @@ void FBlueprintCompilationManagerImpl::ReparentHierarchies(const TMap<UClass*, U
 		ClassToReinstance->ClearFunctionMapsCaches();
 		ClassToReinstance->StaticLink(true);
 
-		OldClassToNewClassIncludingChildren.Add(ReinstancingJob.Reinstancer->DuplicatedClass, ClassToReinstance);
+		OldClassToNewClassIncludingChildren.Add(ReinstancingJob.OldToNew.Key, ClassToReinstance);
 	}
-
-	// reparenting done, reinstance the hierarchy and update archetypes:
+	
+	// Reparenting done, reinstance the hierarchy and update archetypes:
 	ReinstanceBatch(Reinstancers, OldClassToNewClassIncludingChildren, nullptr);
 
-	// reinstance instances - but only the classes we created, rest will be handled by caller
+	// Reinstance (non archetype) instances
 	TMap<UClass*, UClass*> OldClassToNewClassDerivedTypes;
 	for(const FReinstancingJob& ReinstancingJob : Reinstancers)
 	{
-		OldClassToNewClassDerivedTypes.Add(ReinstancingJob.Reinstancer->DuplicatedClass, ReinstancingJob.Reinstancer->ClassToReinstance);
+		OldClassToNewClassDerivedTypes.Add(ReinstancingJob.OldToNew);
 	}
 	TGuardValue<bool> ReinstancingGuard(GIsReinstancing, true);
 	FBlueprintCompileReinstancer::BatchReplaceInstancesOfClass( OldClassToNewClassDerivedTypes, /* bArchetypesAreUpToDate */ true );
@@ -1567,7 +1604,8 @@ void FBlueprintCompilationManagerImpl::ReinstanceBatch(TArray<FReinstancingJob>&
 {
 	const auto FilterOutOfDateClasses = [](TArray<UClass*>& ClassList)
 	{
-		ClassList.RemoveAllSwap( [](UClass* Class) { return Class->HasAnyClassFlags(CLASS_NewerVersionExists); } );
+		// Old versions of classes can be abandoned, classes without CDOs have no instances and don't require reinstancing
+		ClassList.RemoveAllSwap( [](UClass* Class) { return Class->HasAnyClassFlags(CLASS_NewerVersionExists) || Class->ClassDefaultObject == nullptr; } );
 	};
 
 	const auto HasChildren = [FilterOutOfDateClasses](UClass* InClass) -> bool
@@ -1585,21 +1623,20 @@ void FBlueprintCompilationManagerImpl::ReinstanceBatch(TArray<FReinstancingJob>&
 	// haven't already been reinstanced:
 	for (const FReinstancingJob& ReinstancingJob : Reinstancers)
 	{
-		const TSharedPtr<FBlueprintCompileReinstancer>& CurrentReinstancer = ReinstancingJob.Reinstancer;
-		UClass* OldClass = CurrentReinstancer->DuplicatedClass;
+		UClass* OldClass = ReinstancingJob.OldToNew.Key;
 		if(!OldClass)
 		{
 			continue;
 		}
 
-		InOutOldToNewClassMap.Add(OldClass, CurrentReinstancer->ClassToReinstance);
+		InOutOldToNewClassMap.Add(OldClass, ReinstancingJob.OldToNew.Value);
 
 		if(!HasChildren(OldClass))
 		{
 			continue;
 		}
 
-		bool bParentLayoutChanged = !FStructUtils::TheSameLayout(OldClass, CurrentReinstancer->ClassToReinstance);
+		bool bParentLayoutChanged = !FStructUtils::TheSameLayout(OldClass, ReinstancingJob.OldToNew.Value);
 		if(bParentLayoutChanged)
 		{
 			// we need *all* derived types:
@@ -1640,7 +1677,7 @@ void FBlueprintCompilationManagerImpl::ReinstanceBatch(TArray<FReinstancingJob>&
 	{
 		UObject* OriginalCDO = Class->ClassDefaultObject;
 		Reinstancers.Emplace(
-			FReinstancingJob {
+			FReinstancingJob ( 
 				TSharedPtr<FBlueprintCompileReinstancer>( 
 					new FBlueprintCompileReinstancer(
 						Class, 
@@ -1648,11 +1685,12 @@ void FBlueprintCompilationManagerImpl::ReinstanceBatch(TArray<FReinstancingJob>&
 					)
 				),
 				TSharedPtr<FKismetCompilerContext>()
-			}
+			)
 		);
 
 		// make sure we have the newest parent now that CDO has been moved to duplicate class:
 		TSharedPtr<FBlueprintCompileReinstancer>& NewestReinstancer = Reinstancers.Last().Reinstancer;
+		ensure(NewestReinstancer->DuplicatedClass && NewestReinstancer->ClassToReinstance);
 
 		UClass* SuperClass = NewestReinstancer->ClassToReinstance->GetSuperClass();
 		if(ensure(SuperClass))
@@ -1671,14 +1709,16 @@ void FBlueprintCompilationManagerImpl::ReinstanceBatch(TArray<FReinstancingJob>&
 	// run UpdateBytecodeReferences:
 	for (const FReinstancingJob& ReinstancingJob : Reinstancers)
 	{
-		const TSharedPtr<FBlueprintCompileReinstancer>& CurrentReinstancer = ReinstancingJob.Reinstancer;
-		if (CurrentReinstancer->DuplicatedClass)
+		if (ReinstancingJob.OldToNew.Key)
 		{
-			InOutOldToNewClassMap.Add(CurrentReinstancer->DuplicatedClass, CurrentReinstancer->ClassToReinstance);
+			InOutOldToNewClassMap.Add(ReinstancingJob.OldToNew.Key, ReinstancingJob.OldToNew.Value);
 		}
 			
-		UBlueprint* CompiledBlueprint = UBlueprint::GetBlueprintFromClass(CurrentReinstancer->ClassToReinstance);
-		CurrentReinstancer->UpdateBytecodeReferences();
+		if(ReinstancingJob.Reinstancer.IsValid())
+		{
+			UBlueprint* CompiledBlueprint = UBlueprint::GetBlueprintFromClass(ReinstancingJob.OldToNew.Value);
+			ReinstancingJob.Reinstancer->UpdateBytecodeReferences();
+		}
 	}
 	
 	// Now we can update templates and archetypes - note that we don't look for direct references to archetypes - doing
@@ -1691,32 +1731,10 @@ void FBlueprintCompilationManagerImpl::ReinstanceBatch(TArray<FReinstancingJob>&
 	Reinstancers.Sort(
 		[](const FReinstancingJob& ReinstancingDataA, const FReinstancingJob& ReinstancingDataB)
 		{
-			const TSharedPtr<FBlueprintCompileReinstancer>& CompilerDataA = ReinstancingDataA.Reinstancer;
-			const TSharedPtr<FBlueprintCompileReinstancer>& CompilerDataB = ReinstancingDataB.Reinstancer;
-
-			UClass* A = CompilerDataA->ClassToReinstance;
-			UClass* B = CompilerDataB->ClassToReinstance;
-			int32 DepthA = 0;
-			int32 DepthB = 0;
-			UStruct* Iter = A ? A->GetSuperStruct() : nullptr;
-			while (Iter)
-			{
-				++DepthA;
-				Iter = Iter->GetSuperStruct();
-			}
-
-			Iter = B ? B->GetSuperStruct() : nullptr;
-			while (Iter)
-			{
-				++DepthB;
-				Iter = Iter->GetSuperStruct();
-			}
-
-			if (DepthA == DepthB && A && B)
-			{
-				return A->GetFName().LexicalLess(B->GetFName());
-			}
-			return DepthA < DepthB;
+			return FBlueprintCompilationManagerImpl::ReinstancerOrderingFunction(
+				ReinstancingDataA.OldToNew.Value, 
+				ReinstancingDataB.OldToNew.Value
+			);
 		}
 	);
 
@@ -1724,14 +1742,13 @@ void FBlueprintCompilationManagerImpl::ReinstanceBatch(TArray<FReinstancingJob>&
 	// so use GetDefaultObject(true):
 	for (const FReinstancingJob& ReinstancingJob : Reinstancers)
 	{
-		const TSharedPtr<FBlueprintCompileReinstancer>& CurrentReinstancer = ReinstancingJob.Reinstancer;
 		UObject* OldCDO = nullptr;
-		if (CurrentReinstancer->DuplicatedClass)
+		if (ReinstancingJob.OldToNew.Key)
 		{
-			OldCDO = CurrentReinstancer->DuplicatedClass->ClassDefaultObject;
-			if (OldCDO)
+			OldCDO = ReinstancingJob.OldToNew.Key->ClassDefaultObject;
+			if (OldCDO && ReinstancingJob.Reinstancer.IsValid())
 			{
-				UObject* NewCDO = CurrentReinstancer->ClassToReinstance->GetDefaultObject(true);
+				UObject* NewCDO = ReinstancingJob.OldToNew.Value->GetDefaultObject(true);
 				FBlueprintCompileReinstancer::CopyPropertiesForUnrelatedObjects(OldCDO, NewCDO, true);
 
 				if (ReinstancingJob.Compiler.IsValid())
@@ -1741,12 +1758,12 @@ void FBlueprintCompilationManagerImpl::ReinstanceBatch(TArray<FReinstancingJob>&
 			}
 		}
 
-		if (UBlueprintGeneratedClass* BPGClass = CastChecked<UBlueprintGeneratedClass>(CurrentReinstancer->ClassToReinstance))
+		if (UBlueprintGeneratedClass* BPGClass = CastChecked<UBlueprintGeneratedClass>(ReinstancingJob.OldToNew.Value))
 		{
 			BPGClass->UpdateCustomPropertyListForPostConstruction();
 
 			// patch new cdo into linker table:
-			if(OldCDO)
+			if(OldCDO && ReinstancingJob.Reinstancer.IsValid())
 			{
 				UBlueprint* CurrentBP = CastChecked<UBlueprint>(BPGClass->ClassGeneratedBy);
 				if(FLinkerLoad* CurrentLinker = CurrentBP->GetLinker())
@@ -1780,11 +1797,10 @@ void FBlueprintCompilationManagerImpl::ReinstanceBatch(TArray<FReinstancingJob>&
 	TSet<UObject*> ArchetypeReferencers;
 	for (const FReinstancingJob& ReinstancingJob : Reinstancers)
 	{
-		const TSharedPtr<FBlueprintCompileReinstancer>& CurrentReinstancer = ReinstancingJob.Reinstancer;
-		UClass* OldClass = CurrentReinstancer->DuplicatedClass;
+		UClass* OldClass = ReinstancingJob.OldToNew.Key;
 		if(OldClass)
 		{
-			UClass* NewClass = CurrentReinstancer->ClassToReinstance;
+			UClass* NewClass = ReinstancingJob.OldToNew.Value;
             if (NewClass && 
 				OldClass->ClassDefaultObject && 
 				NewClass->ClassDefaultObject &&
@@ -1811,7 +1827,7 @@ void FBlueprintCompilationManagerImpl::ReinstanceBatch(TArray<FReinstancingJob>&
 					// that things that are not directly outered to the transient package will be 
 					// 'reinst'd', this is specifically to handle components, which need to be up to date
 					// on the REINST_ actor class:
-					return !bIsArchetype || Obj->GetOuter() == GetTransientPackage(); 
+					return !bIsArchetype || Obj->GetOuter() == GetTransientPackage() || Obj->HasAnyFlags(RF_NewerVersionExists); 
 				}
 			);
 
@@ -1862,9 +1878,9 @@ void FBlueprintCompilationManagerImpl::ReinstanceBatch(TArray<FReinstancingJob>&
 					REN_DoNotDirty | REN_DontCreateRedirectors | REN_ForceNoResetLoaders );
 
 				// reconstruct
-				FMakeClassSpawnableOnScope TemporarilySpawnable(CurrentReinstancer->ClassToReinstance);
+				FMakeClassSpawnableOnScope TemporarilySpawnable(NewClass);
 				const EObjectFlags FlagMask = RF_Public | RF_ArchetypeObject | RF_Transactional | RF_Transient | RF_TextExportTransient | RF_InheritableComponentTemplate | RF_Standalone; //TODO: what about RF_RootSet?
-				UObject* NewArchetype = NewObject<UObject>(OriginalOuter, CurrentReinstancer->ClassToReinstance, OriginalName, OriginalFlags & FlagMask);
+				UObject* NewArchetype = NewObject<UObject>(OriginalOuter, NewClass, OriginalName, OriginalFlags & FlagMask);
 
 				// grab the old archetype's subobjects:
 				{
@@ -1908,8 +1924,7 @@ void FBlueprintCompilationManagerImpl::ReinstanceBatch(TArray<FReinstancingJob>&
 	// This loop finishes the reinstancing of archetypes after the entire Outer hierarchy has been updated with new instances:
 	for (const FReinstancingJob& ReinstancingJob : Reinstancers)
 	{
-		const TSharedPtr<FBlueprintCompileReinstancer>& CurrentReinstancer = ReinstancingJob.Reinstancer;
-		UClass* OldClass = CurrentReinstancer->DuplicatedClass;
+		UClass* OldClass = ReinstancingJob.OldToNew.Key;
 		if(OldClass)
 		{
 			TArray<UObject*> OldInstances;
@@ -1935,10 +1950,15 @@ void FBlueprintCompilationManagerImpl::ReinstanceBatch(TArray<FReinstancingJob>&
 	// all references in our UBlueprint and its generated class:
 	for (const FReinstancingJob& ReinstancingJob : Reinstancers)
 	{
-		const TSharedPtr<FBlueprintCompileReinstancer>& CurrentReinstancer = ReinstancingJob.Reinstancer;
-		ArchetypeReferencers.Add(CurrentReinstancer->ClassToReinstance);
-		ArchetypeReferencers.Add(CurrentReinstancer->ClassToReinstance->ClassGeneratedBy);
-		if(UBlueprint* BP = Cast<UBlueprint>(CurrentReinstancer->ClassToReinstance->ClassGeneratedBy))
+		ArchetypeReferencers.Add(ReinstancingJob.OldToNew.Value);
+		ArchetypeReferencers.Add(ReinstancingJob.OldToNew.Value->ClassGeneratedBy);
+
+		if(!ReinstancingJob.Reinstancer.IsValid())
+		{
+			continue;
+		}
+
+		if(UBlueprint* BP = Cast<UBlueprint>(ReinstancingJob.OldToNew.Value->ClassGeneratedBy))
 		{
 			// The only known way to cause this ensure to trip is to enqueue blueprints for compilation
 			// while blueprints are already compiling:
@@ -2500,6 +2520,31 @@ UObject* FBlueprintCompilationManagerImpl::GetOuterForRename(UClass* ForClass)
 	return GetTransientPackage();
 }
 
+bool FBlueprintCompilationManagerImpl::ReinstancerOrderingFunction( UClass* A, UClass* B )
+{
+	int32 DepthA = 0;
+	int32 DepthB = 0;
+	UStruct* Iter = A ? A->GetSuperStruct() : nullptr;
+	while (Iter)
+	{
+		++DepthA;
+		Iter = Iter->GetSuperStruct();
+	}
+
+	Iter = B ? B->GetSuperStruct() : nullptr;
+	while (Iter)
+	{
+		++DepthB;
+		Iter = Iter->GetSuperStruct();
+	}
+
+	if (DepthA == DepthB && A && B)
+	{
+		return A->GetFName().LexicalLess(B->GetFName());
+	}
+	return DepthA < DepthB;
+}
+
 // FFixupBytecodeReferences Implementation:
 FBlueprintCompilationManagerImpl::FFixupBytecodeReferences::FFixupBytecodeReferences(UObject* InObject)
 {
diff --git a/Engine/Source/Editor/UnrealEd/Public/Kismet2/KismetReinstanceUtilities.h b/Engine/Source/Editor/UnrealEd/Public/Kismet2/KismetReinstanceUtilities.h
--- a/Engine/Source/Editor/UnrealEd/Public/Kismet2/KismetReinstanceUtilities.h
+++ b/Engine/Source/Editor/UnrealEd/Public/Kismet2/KismetReinstanceUtilities.h
@@ -66,6 +66,7 @@ class UNREALED_API FBlueprintCompileReinstancer : public TSharedFromThis<FBluepr
 {
 protected:
 	friend struct FBlueprintCompilationManagerImpl;
+	friend struct FReinstancingJob;
 
 	static TSet<TWeakObjectPtr<UBlueprint>> DependentBlueprintsToRefresh;
 	static TSet<TWeakObjectPtr<UBlueprint>> DependentBlueprintsToRecompile;
diff --git a/Engine/Source/Runtime/CoreUObject/Private/UObject/PackageReload.cpp b/Engine/Source/Runtime/CoreUObject/Private/UObject/PackageReload.cpp
--- a/Engine/Source/Runtime/CoreUObject/Private/UObject/PackageReload.cpp
+++ b/Engine/Source/Runtime/CoreUObject/Private/UObject/PackageReload.cpp
@@ -455,8 +455,15 @@ UPackage* ReloadPackage(UPackage* InPackageToReload, const uint32 InLoadFlags)
 	return ReloadedPackages[0];
 }
 
-void ReloadPackages(const TArrayView<FReloadPackageData>& InPackagesToReload, TArray<UPackage*>& OutReloadedPackages, const int32 InNumPackagesPerBatch)
+void ReloadPackages(const TArrayView<FReloadPackageData>& InPackagesToReload, TArray<UPackage*>& OutReloadedPackages, int32 InNumPackagesPerBatch)
 {
+	// Interdependencies between packages (in particular Blueprints) make it unsafe to run this logic in batches. 
+	// There are a number of edge cases that would have to be addressed if the batching logic were to be re-enabled, but 
+	// most likely the blueprint reparenting step would have to take in a TMap<UObject*, UObject*> that it could update
+	// when it decided that it needed to replace an instance due to hierarchy changes (e.g. class layout changing
+	// due to SuperStruct changes). For now, just process assets in a batch size of 1:
+	InNumPackagesPerBatch = 1;
+
 	FString Msg;
 	Msg.Append(FString::Printf(TEXT("Reloading %d Package(s):"), InPackagesToReload.Num()));
 	const int32 MAX_PACKAGES_TO_LOG = 10;
@@ -575,7 +582,7 @@ void ReloadPackages(const TArrayView<FReloadPackageData>& InPackagesToReload, TA
 				// in class layout change and subsequently crashes because instances will not match this new class layout:
 				if(UClass* AsClass = Cast<UClass>(PotentialReferencer))
 				{
-					if(AsClass->HasAnyClassFlags(CLASS_NewerVersionExists))
+					if(AsClass->HasAnyClassFlags(CLASS_NewerVersionExists) || AsClass->HasAnyFlags(RF_NewerVersionExists))
 					{
 						continue;
 					}
