Fix crash on saving a Blueprint with a Map variable whose default value includes an empty (unset) key.

diff --git a/Engine/Source/Editor/Kismet/Private/FindInBlueprintManager.cpp b/Engine/Source/Editor/Kismet/Private/FindInBlueprintManager.cpp
--- a/Engine/Source/Editor/Kismet/Private/FindInBlueprintManager.cpp
+++ b/Engine/Source/Editor/Kismet/Private/FindInBlueprintManager.cpp
@@ -591,7 +591,7 @@ namespace BlueprintSearchMetaDataHelpers
 		}
 		else if(InJsonValue->Type == EJson::Array)
 		{
-			auto JsonArray = InJsonValue->AsArray();
+			const TArray<TSharedPtr<FJsonValue>>& JsonArray = InJsonValue->AsArray();
 			if(JsonArray.Num() > 0)
 			{
 				// Some types are never interesting and the contents of the array should be ignored. Other types can be interesting, the contents of the array should be stored (even if
@@ -608,11 +608,16 @@ namespace BlueprintSearchMetaDataHelpers
 			bValidPropetyValue = false;
 
 			// Go through all value/key pairs to see if any of them are searchable, remove the ones that are not
-			auto JsonObject = InJsonValue->AsObject();
-			for(auto Iter = JsonObject->Values.CreateIterator(); Iter; ++Iter)
+			const TSharedPtr<FJsonObject>& JsonObject = InJsonValue->AsObject();
+			for(TMap<FString, TSharedPtr<FJsonValue>>::TIterator Iter = JsonObject->Values.CreateIterator(); Iter; ++Iter)
 			{
-				if(!CheckIfJsonValueIsSearchable(Iter->Value))
+				// Empty keys don't convert to JSON, so we also remove the entry in that case. Note: This means the entry is not going to be searchable.
+				// @todo - Potentially use a placeholder string that uniquely identifies this as an empty key?
+				const bool bHasEmptyKey = Iter->Key.IsEmpty();
+
+				if(!CheckIfJsonValueIsSearchable(Iter->Value) || bHasEmptyKey)
 				{
+					// Note: It's safe to keep incrementing after this; the underlying logic maps to TSparseArray/TConstSetBitIterator::RemoveCurrent().
 					Iter.RemoveCurrent();
 				}
 				else
