If there is no root component for an Actor, when adding a component to the instance use the ActorNode as the parent.
Also fixed up a bunch of whitespace that had gotten broken while resolving conflicts.
#rnx

diff --git a/Engine/Source/Editor/Kismet/Private/SSCSEditor.cpp b/Engine/Source/Editor/Kismet/Private/SSCSEditor.cpp
--- a/Engine/Source/Editor/Kismet/Private/SSCSEditor.cpp
+++ b/Engine/Source/Editor/Kismet/Private/SSCSEditor.cpp
@@ -995,7 +995,7 @@ FSCSEditorTreeNodeInstancedInheritedComponent::FSCSEditorTreeNodeInstancedInheri
 
 	InstancedComponentOwnerPtr = Owner;
 
-			SetComponentTemplate(ComponentInstance);
+	SetComponentTemplate(ComponentInstance);
 }
 
 bool FSCSEditorTreeNodeInstancedInheritedComponent::IsNative() const
@@ -4592,7 +4592,7 @@ void SSCSEditor::UpdateTree(bool bRegenerateTreeNodes)
 				
 				for (UActorComponent* Component : Components)
 				{
-						// Add the rest of the native base class SceneComponent hierarchy
+					// Add the rest of the native base class SceneComponent hierarchy
 					AddTreeNodeFromComponent(Component, FindOrCreateParentForExistingComponent(Component, GetActorNode()));
 				}
 			}
@@ -4970,15 +4970,15 @@ FSCSEditorActorNodePtrType SSCSEditor::GetActorNode() const
 }
 
 FSCSEditorTreeNodePtrType SSCSEditor::GetSceneRootNode() const
-	{
+{
 	FSCSEditorActorNodePtrType ActorNode = GetActorNode();
 	if (ActorNode.IsValid())
-		{
+	{
 		return ActorNode->GetSceneRootNode();
-		}
+	}
 
 	return FSCSEditorTreeNodePtrType();
-	}
+}
 
 void SSCSEditor::SetSceneRootNode(FSCSEditorTreeNodePtrType NewSceneRootNode)
 {
@@ -5349,10 +5349,10 @@ FSCSEditorTreeNodePtrType SSCSEditor::FindOrCreateParentForExistingComponent(UAc
 }
 
 FSCSEditorTreeNodePtrType SSCSEditor::FindParentForNewComponent(UActorComponent* NewComponent) const
-	{
+{
 	// Find Parent to attach to (depending on the new Node type).
 	FSCSEditorTreeNodePtrType TargetParentNode;
-		TArray<FSCSEditorTreeNodePtrType> SelectedTreeNodes;
+	TArray<FSCSEditorTreeNodePtrType> SelectedTreeNodes;
 	if (SCSTreeWidget.IsValid() && SCSTreeWidget->GetSelectedItems(SelectedTreeNodes))
 	{
 		TargetParentNode = SelectedTreeNodes[0];
@@ -5367,14 +5367,18 @@ FSCSEditorTreeNodePtrType SSCSEditor::FindParentForNewComponent(UActorComponent*
 				FSCSEditorActorNodePtrType TargetActorNode = StaticCastSharedPtr<FSCSEditorTreeNodeRootActor>(TargetParentNode);
 				if (TargetActorNode.IsValid())
 				{
-					TargetParentNode = TargetActorNode->GetSceneRootNode();
-					USceneComponent* CastTargetToSceneComponent = Cast<USceneComponent>(TargetParentNode->GetComponentTemplate());
-					if (CastTargetToSceneComponent == nullptr || !NewSceneComponent->CanAttachAsChild(CastTargetToSceneComponent, NAME_None))
+					FSCSEditorTreeNodePtrType TargetSceneRootNode = TargetActorNode->GetSceneRootNode();
+					if (TargetSceneRootNode.IsValid())
 					{
-						TargetParentNode = GetSceneRootNode(); // Default to SceneRoot
+						TargetParentNode = TargetSceneRootNode;
+						USceneComponent* CastTargetToSceneComponent = Cast<USceneComponent>(TargetParentNode->GetComponentTemplate());
+						if (CastTargetToSceneComponent == nullptr || !NewSceneComponent->CanAttachAsChild(CastTargetToSceneComponent, NAME_None))
+						{
+							TargetParentNode = GetSceneRootNode(); // Default to SceneRoot
+						}
+					}
 				}
 			}
-		}
 			else if(TargetParentNode->GetNodeType() == FSCSEditorTreeNode::ComponentNode)
 			{
 				USceneComponent* CastTargetToSceneComponent = Cast<USceneComponent>(TargetParentNode->GetComponentTemplate());
@@ -5410,7 +5414,7 @@ FSCSEditorTreeNodePtrType SSCSEditor::FindParentForNewComponent(UActorComponent*
 	}
 
 FSCSEditorTreeNodePtrType SSCSEditor::FindParentForNewNode(USCS_Node* NewNode) const
-	{
+{
 	return FindParentForNewComponent(NewNode->ComponentTemplate);
 }
 
@@ -5429,7 +5433,7 @@ UActorComponent* SSCSEditor::AddNewNode(TUniquePtr<FScopedTransaction> InOngoing
 	UBlueprint* Blueprint = GetBlueprint();
 	check(Blueprint != nullptr && Blueprint->SimpleConstructionScript != nullptr);
 
-		// Add the new node to the editor tree
+	// Add the new node to the editor tree
 	NewNodePtr = AddTreeNode(NewNode, ParentNodePtr, /*bIsInheritedSCS=*/ false);
 
 	// Potentially adjust variable names for any child blueprints
@@ -5998,46 +6002,45 @@ FSCSEditorTreeNodePtrType SSCSEditor::AddTreeNode(USCS_Node* InSCSNode, FSCSEdit
 	// Determine whether or not the given node is inherited from a parent Blueprint
 	USimpleConstructionScript* NodeSCS = InSCSNode->GetSCS();
 
-			// do this first, because we need a FSCSEditorTreeNodePtrType for the new node
+	// do this first, because we need a FSCSEditorTreeNodePtrType for the new node
 	NewNodePtr = InParentNodePtr->AddChild(InSCSNode, bIsInheritedSCS);
-			RefreshFilteredState(NewNodePtr, /*bRecursive =*/false);
-
-
+	RefreshFilteredState(NewNodePtr, /*bRecursive =*/false);
+	
 	if( InSCSNode->ComponentTemplate && 
 		InSCSNode->ComponentTemplate->IsA(USceneComponent::StaticClass()) && 
 		InParentNodePtr->GetNodeType() == FSCSEditorTreeNode::ComponentNode)
 	{
 		bool bParentIsEditorOnly = InParentNodePtr->GetComponentTemplate()->IsEditorOnly();
-			// if you can't nest this new node under the proposed parent (then swap the two)
+		// if you can't nest this new node under the proposed parent (then swap the two)
 		if (bParentIsEditorOnly && !InSCSNode->ComponentTemplate->IsEditorOnly() && InParentNodePtr->CanReparent())
-			{
+		{
 			FSCSEditorTreeNodePtrType OldParentPtr = InParentNodePtr;
 			InParentNodePtr = OldParentPtr->GetParent();
 
-				OldParentPtr->RemoveChild(NewNodePtr);
-				NodeSCS->RemoveNode(OldParentPtr->GetSCSNode());
+			OldParentPtr->RemoveChild(NewNodePtr);
+			NodeSCS->RemoveNode(OldParentPtr->GetSCSNode());
 
-				// if the grandparent node is invalid (assuming this means that the parent node was the scene-root)
+			// if the grandparent node is invalid (assuming this means that the parent node was the scene-root)
 			if (!InParentNodePtr.IsValid())
-				{
+			{
 				check(OldParentPtr == GetSceneRootNode());
 				SetSceneRootNode(NewNodePtr);
 				NodeSCS->AddNode(NewNodePtr->GetSCSNode());
-				}
-				else 
-				{
+			}
+			else 
+			{
 				InParentNodePtr->AddChild(NewNodePtr);
-				}
+			}
 
-				// move the proposed parent in as a child to the new node
-				NewNodePtr->AddChild(OldParentPtr);
-			} // if bParentIsEditorOnly...
-		}
-		else 
-		{
+			// move the proposed parent in as a child to the new node
+			NewNodePtr->AddChild(OldParentPtr);
+		} // if bParentIsEditorOnly...
+	}
+	else 
+	{
 		// If the SCS root node array does not already contain the given node, this will add it (this should only occur after node creation)
 		if(NodeSCS != nullptr)
-			{
+		{
 			NodeSCS->AddNode(InSCSNode);
 		}
 	}
@@ -6061,11 +6064,11 @@ FSCSEditorTreeNodePtrType SSCSEditor::AddTreeNodeFromComponent(UActorComponent*
 
 	FSCSEditorTreeNodePtrType NewNodePtr = InParentTreeNode->FindChild(InActorComponent);
 	if (!NewNodePtr.IsValid())
-			{
+	{
 		NewNodePtr = FSCSEditorTreeNode::FactoryNodeFromComponent(InActorComponent);
 		InParentTreeNode->AddChild(NewNodePtr);
 		RefreshFilteredState(NewNodePtr, false);
-				}
+	}
 
 	SCSTreeWidget->SetItemExpansion(NewNodePtr, true);
 
