Reuse subobjects during archetype reinstancing path

diff --git a/Engine/Source/Editor/Kismet/Private/BlueprintCompilationManager.cpp b/Engine/Source/Editor/Kismet/Private/BlueprintCompilationManager.cpp
--- a/Engine/Source/Editor/Kismet/Private/BlueprintCompilationManager.cpp
+++ b/Engine/Source/Editor/Kismet/Private/BlueprintCompilationManager.cpp
@@ -1500,49 +1500,32 @@ void FBlueprintCompilationManagerImpl::ReinstanceBatch(TArray<FReinstancingJob>&
 				const EObjectFlags FlagMask = RF_Public | RF_ArchetypeObject | RF_Transactional | RF_Transient | RF_TextExportTransient | RF_InheritableComponentTemplate | RF_Standalone; //TODO: what about RF_RootSet?
 				UObject* NewArchetype = NewObject<UObject>(OriginalOuter, CurrentReinstancer->ClassToReinstance, OriginalName, OriginalFlags & FlagMask);
 
-				// copy old data:
-				FBlueprintCompileReinstancer::CopyPropertiesForUnrelatedObjects(Archetype, NewArchetype, false);
-
-				OldArchetypeToNewArchetype.Add(Archetype, NewArchetype);
-				// Map old subobjects to new subobjects. This is needed by UMG right now, which allows owning archetypes to reference subobjects
-				// in subwidgets:
+				// grab the old archetype's subobjects:
 				{
 					TArray<UObject*> OldSubobjects;
-					GetObjectsWithOuter(Archetype, OldSubobjects);
-					TArray<UObject*> NewSubobjects;
-					GetObjectsWithOuter(NewArchetype, NewSubobjects);
-
-					TMap<FName, UObject*> OldNameMap;
-					for(UObject* OldSubobject : OldSubobjects )
-					{
-						OldNameMap.Add( OldSubobject->GetFName(), OldSubobject );
-					}
-
-					TMap<FName, UObject*> NewNameMap;
-					for(UObject* NewSubobject : NewSubobjects )
-					{
-						NewNameMap.Add( NewSubobject->GetFName(), NewSubobject );
-					}
+					GetObjectsWithOuter( Archetype, OldSubobjects, false );
 
-					for(TPair< FName, UObject* > OldSubobject : OldNameMap )
+					for(UObject* Subobject : OldSubobjects )
 					{
-						UObject** NewSubobject = NewNameMap.Find(OldSubobject.Key);
-						OldArchetypeToNewArchetype.Add(OldSubobject.Value, NewSubobject ? *NewSubobject : nullptr );
-						if(NewSubobject)
+						if(Subobject->HasAnyFlags(RF_DefaultSubObject))
 						{
-							FLinkerLoad::PRIVATE_PatchNewObjectIntoExport(OldSubobject.Value, *NewSubobject);
-						}
-						else
-						{
-							// an object was not recreated, this can be because we are running without GEditor (-game or -server) and
-							// the old subobject was an editor only subobject (CreateEditorOnlyDefaultSubobject):
-							OldSubobject.Value->RemoveFromRoot();
-							OldSubobject.Value->MarkPendingKill();
+							// CPFUO handles DSOs:
+							continue;
 						}
+
+						// Non DSO subobject - just reuse the subobject:
+						Subobject->Rename( 
+							nullptr, 
+							// destination:
+							NewArchetype, 
+							// Rename options:
+							REN_DoNotDirty | REN_DontCreateRedirectors | REN_ForceNoResetLoaders 
+						);
 					}
-					
 				}
 
+				OldArchetypeToNewArchetype.Add(Archetype, NewArchetype);
+
 				ArchetypeReferencers.Add(NewArchetype);
 
 				FLinkerLoad::PRIVATE_PatchNewObjectIntoExport(Archetype, NewArchetype);
@@ -1553,6 +1536,30 @@ void FBlueprintCompilationManagerImpl::ReinstanceBatch(TArray<FReinstancingJob>&
 		}
 	}
 
+	// This loop finishes the reinstancing of archetypes after the entire Outer hierarchy has been updated with new instances:
+	for (const FReinstancingJob& ReinstancingJob : Reinstancers)
+	{
+		const TSharedPtr<FBlueprintCompileReinstancer>& CurrentReinstancer = ReinstancingJob.Reinstancer;
+		UClass* OldClass = CurrentReinstancer->DuplicatedClass;
+		if(OldClass)
+		{
+			TArray<UObject*> OldInstances;
+			GetObjectsOfClass( OldClass, OldInstances, false );
+
+			for(UObject* OldInstance : OldInstances)
+			{
+				UObject** NewInstance = OldArchetypeToNewArchetype.Find(OldInstance);
+				if(NewInstance)
+				{
+					// The new object hierarchy has been created, all of the old instances are in the transient package and new
+					// ones have taken their place. Referenc members will mostly be pointing at *old* instances, and will get fixed
+					// up below:
+					FBlueprintCompileReinstancer::CopyPropertiesForUnrelatedObjects(OldInstance, *NewInstance, false);
+				}
+			}
+		}
+	}
+
 	// 4. update known references to archetypes (e.g. component templates, WidgetTree). We don't want to run the normal 
 	// reference finder to update these because searching the entire object graph is time consuming. Instead we just replace
 	// all references in our UBlueprint and its generated class:
