PR #5075: Mouse decorator when dragging a category from My Blueprint panel (Contributed by [REDACTED_PERSON])

diff --git a/Engine/Source/Editor/Kismet/Private/SMyBlueprint.cpp b/Engine/Source/Editor/Kismet/Private/SMyBlueprint.cpp
--- a/Engine/Source/Editor/Kismet/Private/SMyBlueprint.cpp
+++ b/Engine/Source/Editor/Kismet/Private/SMyBlueprint.cpp
@@ -82,11 +82,11 @@ public:
 		const FSlateBrush* StatusSymbol = FEditorStyle::GetBrush(TEXT("NoBrush")); 
 		FText Message = DraggedCategory;
 
+		FFormatNamedArguments Args;
+		Args.Add(TEXT("DraggedCategory"), DraggedCategory);
+
 		if (!HoveredCategoryName.IsEmpty())
 		{
-			FFormatNamedArguments Args;
-			Args.Add(TEXT("DraggedCategory"), DraggedCategory);
-
 			if(HoveredCategoryName.EqualTo(DraggedCategory))
 			{
 				StatusSymbol = FEditorStyle::GetBrush(TEXT("Graph.ConnectorFeedback.Error"));
@@ -106,6 +106,11 @@ public:
 			StatusSymbol = FEditorStyle::GetBrush(TEXT("Graph.ConnectorFeedback.Error"));
 			Message = LOCTEXT("MoveCatOverAction", "Can only insert before another category.");
 		}
+		else
+		{
+			StatusSymbol = FEditorStyle::GetBrush(TEXT("Graph.ConnectorFeedback.Error"));
+			Message = FText::Format(LOCTEXT("MoveCatAction", "Moving category '{DraggedCategory}'"), Args);
+		}
 
 		SetSimpleFeedbackMessage(StatusSymbol, FLinearColor::White, Message);
 	}
