Improve error messaging when a class creates its CDO while the class layout is being created

diff --git a/Engine/Source/Editor/Kismet/Private/BlueprintCompilationManager.cpp b/Engine/Source/Editor/Kismet/Private/BlueprintCompilationManager.cpp
--- a/Engine/Source/Editor/Kismet/Private/BlueprintCompilationManager.cpp
+++ b/Engine/Source/Editor/Kismet/Private/BlueprintCompilationManager.cpp
@@ -348,7 +348,7 @@ struct FReinstancingJob
 	TSharedPtr<FBlueprintCompileReinstancer> Reinstancer;
 	TSharedPtr<FKismetCompilerContext> Compiler;
 };
-	
+
 void FBlueprintCompilationManagerImpl::FlushCompilationQueueImpl(TArray<UObject*>* ObjLoaded, bool bSuppressBroadcastCompiled, TArray<UBlueprint*>* BlueprintsCompiled, TArray<UBlueprint*>* BlueprintsCompiledOrSkeletonCompiled)
 {
 	TGuardValue<bool> GuardTemplateNameFlag(GCompilingBlueprint, true);
@@ -858,6 +858,11 @@ void FBlueprintCompilationManagerImpl::FlushCompilationQueueImpl(TArray<UObject*
 						EBlueprintCompileReinstancerFlags::AutoInferSaveOnCompile | EBlueprintCompileReinstancerFlags::AvoidCDODuplication
 					)
 				);
+
+				if(BP->GeneratedClass)
+				{
+					BP->GeneratedClass->ClassFlags |= CLASS_LayoutChanging;
+				}
 			}
 		}
 
@@ -926,14 +931,16 @@ void FBlueprintCompilationManagerImpl::FlushCompilationQueueImpl(TArray<UObject*
 		for (FCompilerData& CompilerData : CurrentlyCompilingBPs)
 		{
 			UBlueprint* BP = CompilerData.BP;
+			UClass* BPGC = BP->GeneratedClass;
+
 			if(!CompilerData.ShouldCompileClassFunctions())
 			{
-				if( BP->GeneratedClass &&
-					(	BP->GeneratedClass->ClassDefaultObject == nullptr || 
-						BP->GeneratedClass->ClassDefaultObject->GetClass() != BP->GeneratedClass) )
+				if( BPGC &&
+					(	BPGC->ClassDefaultObject == nullptr || 
+						BPGC->ClassDefaultObject->GetClass() != BPGC) )
 				{
 					// relink, generate CDO:
-					UClass* BPGC = BP->GeneratedClass;
+					BPGC->ClassFlags &= ~CLASS_LayoutChanging;
 					BPGC->Bind();
 					BPGC->StaticLink(true);
 					BPGC->ClassDefaultObject = nullptr;
@@ -942,16 +949,33 @@ void FBlueprintCompilationManagerImpl::FlushCompilationQueueImpl(TArray<UObject*
 			}
 			else
 			{
-				ensure( BP->GeneratedClass == nullptr ||
-						BP->GeneratedClass->ClassDefaultObject == nullptr || 
-						BP->GeneratedClass->ClassDefaultObject->GetClass() != BP->GeneratedClass);
-				
 				// default value propagation occurrs below:
-				if(BP->GeneratedClass)
+				if(BPGC)
 				{
-					BP->GeneratedClass->ClassDefaultObject = nullptr;
-				
+					if( BPGC->ClassDefaultObject && 
+						BPGC->ClassDefaultObject->GetClass() == BPGC)
+					{
+						// the CDO has been created early, it is possible that the reflection data was still
+						// being mutated by CompileClassLayout. Warn the user and and move the CDO aside:
+						ensureAlwaysMsgf(false, 
+							TEXT("ClassDefaultObject for %s created at the wrong time - it may be corrupt. It is recommended that you save all data and restart the editor session"), 
+							*BP->GetName()
+						);
+
+						BPGC->ClassDefaultObject->Rename(
+							nullptr,
+							// destination - this is the important part of this call. Moving the object 
+							// out of the way so we can reuse its name:
+							GetTransientPackage(), 
+							// Rename options:
+							REN_DoNotDirty | REN_DontCreateRedirectors | REN_ForceNoResetLoaders
+						);
+					}
+					BPGC->ClassDefaultObject = nullptr;
 		
+					// class layout is ready, we can clear CLASS_LayoutChanging and CompileFunctions can create the CDO:
+					BPGC->ClassFlags &= ~CLASS_LayoutChanging;
+
 					FKismetCompilerContext& CompilerContext = *(CompilerData.Compiler);
 					CompilerContext.CompileFunctions(
 						EInternalCompilerFlags::PostponeLocalsGenerationUntilPhaseTwo
@@ -988,13 +1012,13 @@ void FBlueprintCompilationManagerImpl::FlushCompilationQueueImpl(TArray<UObject*
 				}
 			}
 
-			if(BP->GeneratedClass)
+			if(BPGC)
 			{
-				BP->GeneratedClass->ClassFlags &= ~CLASS_ReplicationDataIsSetUp;
-				BP->GeneratedClass->SetUpRuntimeReplicationData();
+				BPGC->ClassFlags &= ~CLASS_ReplicationDataIsSetUp;
+				BPGC->SetUpRuntimeReplicationData();
 			}
 
-			ensure(BP->GeneratedClass == nullptr || BP->GeneratedClass->ClassDefaultObject->GetClass() == *(BP->GeneratedClass));
+			ensure(BPGC == nullptr || BPGC->ClassDefaultObject->GetClass() == BPGC);
 		}
 	} // end GTimeCompiling scope
 
diff --git a/Engine/Source/Editor/KismetCompiler/Private/KismetCompiler.cpp b/Engine/Source/Editor/KismetCompiler/Private/KismetCompiler.cpp
--- a/Engine/Source/Editor/KismetCompiler/Private/KismetCompiler.cpp
+++ b/Engine/Source/Editor/KismetCompiler/Private/KismetCompiler.cpp
@@ -289,6 +289,7 @@ void FKismetCompilerContext::CleanAndSanitizeClass(UBlueprintGeneratedClass* Cla
 	}
 
 	// Purge the class to get it back to a "base" state
+	bool bLayoutChanging = ClassToClean->HasAnyClassFlags(CLASS_LayoutChanging);
 	ClassToClean->PurgeClass(bRecompilingOnLoad);
 
 	// Set properties we need to regenerate the class with
@@ -297,6 +298,11 @@ void FKismetCompilerContext::CleanAndSanitizeClass(UBlueprintGeneratedClass* Cla
 	ClassToClean->ClassWithin = ParentClass->ClassWithin ? ParentClass->ClassWithin : UObject::StaticClass();
 	ClassToClean->ClassConfigName = ClassToClean->IsNative() ? FName(ClassToClean->StaticConfigName()) : ParentClass->ClassConfigName;
 	ClassToClean->DebugData = FBlueprintDebugData();
+
+	if(bLayoutChanging)
+	{
+		ClassToClean->ClassFlags |= CLASS_LayoutChanging;
+	}
 }
 
 void FKismetCompilerContext::SaveSubObjectsFromCleanAndSanitizeClass(FSubobjectCollection& SubObjectsToSave, UBlueprintGeneratedClass* ClassToClean)
diff --git a/Engine/Source/Runtime/CoreUObject/Private/UObject/Class.cpp b/Engine/Source/Runtime/CoreUObject/Private/UObject/Class.cpp
--- a/Engine/Source/Runtime/CoreUObject/Private/UObject/Class.cpp
+++ b/Engine/Source/Runtime/CoreUObject/Private/UObject/Class.cpp
@@ -2952,6 +2952,8 @@ UObject* UClass::CreateDefaultObject()
 {
 	if ( ClassDefaultObject == NULL )
 	{
+		ensureMsgf(!HasAnyClassFlags(CLASS_LayoutChanging), TEXT("Class named %s creating its CDO while changing its layout"), *GetName());
+
 		UClass* ParentClass = GetSuperClass();
 		UObject* ParentDefaultObject = NULL;
 		if ( ParentClass != NULL )
diff --git a/Engine/Source/Runtime/CoreUObject/Public/UObject/ObjectMacros.h b/Engine/Source/Runtime/CoreUObject/Public/UObject/ObjectMacros.h
--- a/Engine/Source/Runtime/CoreUObject/Public/UObject/ObjectMacros.h
+++ b/Engine/Source/Runtime/CoreUObject/Public/UObject/ObjectMacros.h
@@ -213,8 +213,8 @@ enum EClassFlags
 	/** all properties and functions in this class are const and should be exported as const */
 	CLASS_Const			      = 0x00010000u,
 
-	/** */
-	//CLASS_ = 0x00020000u,
+	/** Class flag indicating the class is having its layout changed, and therefore is not ready for a CDO to be created */
+	CLASS_LayoutChanging	  = 0x00020000u,
 	
 	/** Indicates that the class was created from blueprint source material */
 	CLASS_CompiledFromBlueprint  = 0x00040000u,
