When recompiling a blueprint macro library we now do a full compile of all dependent blueprints

diff --git a/Engine/Source/Editor/Kismet/Private/BlueprintCompilationManager.cpp b/Engine/Source/Editor/Kismet/Private/BlueprintCompilationManager.cpp
--- a/Engine/Source/Editor/Kismet/Private/BlueprintCompilationManager.cpp
+++ b/Engine/Source/Editor/Kismet/Private/BlueprintCompilationManager.cpp
@@ -365,6 +365,40 @@ void FBlueprintCompilationManagerImpl::FlushCompilationQueueImpl(TArray<UObject*
 
 		// STAGE I: Add any related blueprints that were not compiled, then add any children so that they will be relinked:
 		TArray<UBlueprint*> BlueprintsToRecompile;
+
+		// First add any dependents of macro libraries that are being compiled:
+		for(const FBPCompileRequest& CompileJob : QueuedRequests)
+		{
+			if ((CompileJob.CompileOptions & 
+				(	EBlueprintCompileOptions::RegenerateSkeletonOnly|
+					EBlueprintCompileOptions::IsRegeneratingOnLoad)
+				) != EBlueprintCompileOptions::None)
+			{
+				continue;
+			}
+			
+			if(CompileJob.BPToCompile->BlueprintType == BPTYPE_MacroLibrary)
+			{
+				TArray<UBlueprint*> DependentBlueprints;
+				FBlueprintEditorUtils::GetDependentBlueprints(CompileJob.BPToCompile, DependentBlueprints);
+				for(UBlueprint* DependentBlueprint : DependentBlueprints)
+				{
+					DependentBlueprint->bQueuedForCompilation = true;
+					CurrentlyCompilingBPs.Add(
+						FCompilerData(
+							DependentBlueprint, 
+							ECompilationManagerJobType::Normal, 
+							nullptr, 
+							EBlueprintCompileOptions::None, 
+							false // full compile
+						)
+					);
+					BlueprintsToRecompile.Add(DependentBlueprint);
+				}
+			}
+		}
+
+		// then make sure any normal blueprints have their bytecode dependents recompiled, this is in case a function signature changes:
 		for(const FBPCompileRequest& CompileJob : QueuedRequests)
 		{
 			if ((CompileJob.CompileOptions & EBlueprintCompileOptions::RegenerateSkeletonOnly) != EBlueprintCompileOptions::None)
