Copy of CL 4416373
Extend GCompilingBlueprint lifetime to include BroadcastBlueprintReinstanced when compiling synchronously

diff --git a/Engine/Source/Editor/Kismet/Private/BlueprintCompilationManager.cpp b/Engine/Source/Editor/Kismet/Private/BlueprintCompilationManager.cpp
--- a/Engine/Source/Editor/Kismet/Private/BlueprintCompilationManager.cpp
+++ b/Engine/Source/Editor/Kismet/Private/BlueprintCompilationManager.cpp
@@ -220,6 +220,9 @@ void FBlueprintCompilationManagerImpl::CompileSynchronouslyImpl(const FBPCompile
 
 	if ( GEditor && !bRegenerateSkeletonOnly)
 	{
+		// Make sure clients know they're being reinstanced as part of blueprint compilation. After this point
+		// compilation is completely done:
+		TGuardValue<bool> GuardTemplateNameFlag(GCompilingBlueprint, true);
 		GEditor->BroadcastBlueprintReinstanced();
 	}
 	
