Fixed crash drag-dropping a variable onto a state machine

diff --git a/Engine/Source/Editor/Kismet/Private/BPVariableDragDropAction.cpp b/Engine/Source/Editor/Kismet/Private/BPVariableDragDropAction.cpp
--- a/Engine/Source/Editor/Kismet/Private/BPVariableDragDropAction.cpp
+++ b/Engine/Source/Editor/Kismet/Private/BPVariableDragDropAction.cpp
@@ -229,7 +229,7 @@ FReply FKismetVariableDragDropAction::DroppedOnPin(FVector2D ScreenPosition, FVe
 	{
 		if (!TargetPin->bOrphanedPin)
 		{
-			if (const UEdGraphSchema_K2* Schema = CastChecked<const UEdGraphSchema_K2>(TargetPin->GetSchema()))
+			if (const UEdGraphSchema_K2* Schema = Cast<const UEdGraphSchema_K2>(TargetPin->GetSchema()))
 			{
 				const bool bIsExecPin = Schema->IsExecPin(*TargetPin);
 
