Guard against double-queuing in Macro BP gather phase

diff --git a/Engine/Source/Editor/Kismet/Private/BlueprintCompilationManager.cpp b/Engine/Source/Editor/Kismet/Private/BlueprintCompilationManager.cpp
--- a/Engine/Source/Editor/Kismet/Private/BlueprintCompilationManager.cpp
+++ b/Engine/Source/Editor/Kismet/Private/BlueprintCompilationManager.cpp
@@ -386,17 +386,20 @@ void FBlueprintCompilationManagerImpl::FlushCompilationQueueImpl(TArray<UObject*
 				FBlueprintEditorUtils::GetDependentBlueprints(CompileJob.BPToCompile, DependentBlueprints);
 				for(UBlueprint* DependentBlueprint : DependentBlueprints)
 				{
-					DependentBlueprint->bQueuedForCompilation = true;
-					CurrentlyCompilingBPs.Add(
-						FCompilerData(
-							DependentBlueprint, 
-							ECompilationManagerJobType::Normal, 
-							nullptr, 
-							EBlueprintCompileOptions::None, 
-							false // full compile
-						)
-					);
-					BlueprintsToRecompile.Add(DependentBlueprint);
+					if(!IsQueuedForCompilation(DependentBlueprint))
+					{
+						DependentBlueprint->bQueuedForCompilation = true;
+						CurrentlyCompilingBPs.Add(
+							FCompilerData(
+								DependentBlueprint, 
+								ECompilationManagerJobType::Normal, 
+								nullptr, 
+								EBlueprintCompileOptions::None, 
+								false // full compile
+							)
+						);
+						BlueprintsToRecompile.Add(DependentBlueprint);
+					}
 				}
 			}
 		}
