PR #5216: Fix typo in the tooltip of "Pass-by-Reference" (Contributed by [REDACTED_PERSON])

diff --git a/Engine/Source/Editor/Kismet/Private/BlueprintDetailsCustomization.cpp b/Engine/Source/Editor/Kismet/Private/BlueprintDetailsCustomization.cpp
--- a/Engine/Source/Editor/Kismet/Private/BlueprintDetailsCustomization.cpp
+++ b/Engine/Source/Editor/Kismet/Private/BlueprintDetailsCustomization.cpp
@@ -2792,7 +2792,7 @@ void FBlueprintGraphArgumentLayout::GenerateChildContent( IDetailChildrenBuilder
 				[
 					SNew(STextBlock)
 					.Text(LOCTEXT("FunctionArgDetailsPassByReference", "Pass-by-Reference"))
-				.ToolTipText(LOCTEXT("FunctionArgDetailsPassByReferenceTooltip", "Pass this paremeter by reference?"))
+				.ToolTipText(LOCTEXT("FunctionArgDetailsPassByReferenceTooltip", "Pass this parameter by reference?"))
 				.Font(IDetailLayoutBuilder::GetDetailFont())
 				]
 			.ValueContent()
