Fixed crash when ending a Transform details panel edit operation by making no changes
Reduce check to a simple guard statement

diff --git a/Engine/Source/Editor/Kismet/Private/BlueprintDetailsCustomization.cpp b/Engine/Source/Editor/Kismet/Private/BlueprintDetailsCustomization.cpp
--- a/Engine/Source/Editor/Kismet/Private/BlueprintDetailsCustomization.cpp
+++ b/Engine/Source/Editor/Kismet/Private/BlueprintDetailsCustomization.cpp
@@ -2532,9 +2532,12 @@ EVisibility FBlueprintVarActionDetails::IsTooltipEditVisible() const
 
 void FBlueprintVarActionDetails::OnFinishedChangingProperties(const FPropertyChangedEvent& InPropertyChangedEvent, TSharedPtr<FStructOnScope> InStructData, TWeakObjectPtr<UK2Node_EditablePinBase> InEntryNode)
 {
-	check(InPropertyChangedEvent.MemberProperty
-		&& InPropertyChangedEvent.MemberProperty->GetOwnerStruct()
-		&& InPropertyChangedEvent.MemberProperty->GetOwnerStruct()->IsA<UFunction>());
+	if( !InPropertyChangedEvent.MemberProperty ||
+		!InPropertyChangedEvent.MemberProperty->GetOwnerStruct() ||
+		!InPropertyChangedEvent.MemberProperty->GetOwnerStruct()->IsA<UFunction>())
+	{
+		return;
+	}
 
 	// Find the top level property that was modified within the UFunction
 	const UProperty* DirectProperty = InPropertyChangedEvent.MemberProperty;
