Don't add instanced components to the component tree if they are not editable

diff --git a/Engine/Source/Editor/Kismet/Private/SSCSEditor.cpp b/Engine/Source/Editor/Kismet/Private/SSCSEditor.cpp
--- a/Engine/Source/Editor/Kismet/Private/SSCSEditor.cpp
+++ b/Engine/Source/Editor/Kismet/Private/SSCSEditor.cpp
@@ -4705,7 +4705,8 @@ bool SSCSEditor::ShouldAddInstancedActorComponent(UActorComponent* ActorComp, US
 	return (ActorComp != nullptr
 		&& (!ActorComp->IsVisualizationComponent())
 		&& (ActorComp->CreationMethod != EComponentCreationMethod::UserConstructionScript || !GetDefault<UBlueprintEditorSettings>()->bHideConstructionScriptComponentsInDetailsView)
-		&& (ParentSceneComp == nullptr || !ParentSceneComp->IsCreatedByConstructionScript() || !ActorComp->HasAnyFlags(RF_DefaultSubObject)));
+		&& (ParentSceneComp == nullptr || !ParentSceneComp->IsCreatedByConstructionScript() || !ActorComp->HasAnyFlags(RF_DefaultSubObject)))
+		&& (ActorComp->CreationMethod != EComponentCreationMethod::Native || FComponentEditorUtils::CanEditNativeComponent(ActorComp) );
 }
 
 void SSCSEditor::DumpTree()
