Fix for a regression that caused a crash after creating a new Blueprint class with "Show Friendly Variable Names" turned off.

Notes:
- Regression was caused by the introduction of the BPCM into the FKismetEditorUtilities::CreateBlueprint() code path at 4178102. The BPCM calls FFindInBlueprintSearchManager::Get().AddOrUpdateBlueprintSearchMetadata() as part of Stage IX.
- Previously, the BP creation code path called FKismet2CompilerModule::CompileBlueprint(), a path which does not contain a call to FFindInBlueprintSearchManager::Get().AddOrUpdateBlueprintSearchMetadata() because this is already being handled by FFindInBlueprintSearchManager::OnAssetAdded().
- The crash occurs because when friendly variable names are disabled, AddOrUpdateBlueprintSearchMetadata() calls FTemporarilyUseFriendlyNodeTitles::ForceVisualizationCacheClear(), which then attempted to create the missing CDO before compilation.

Change summary:
- Added EBlueprintCompileOptions::SkipFiBSearchMetaUpdate and modified FCompilerData/FBlueprintCompilationManagerImpl::FlushCompilationQueueImpl() to bypass the FiB update step when it's not needed.
- Modified FKismetEditorUtilities::CreateBlueprint() and FBlueprintEditorUtils::PostDuplicateBlueprint() to bypass the FiB update step, since it's already handled by FFindInBlueprintSearchManager::OnAssetAdded().
- Modified FTemporarilyUseFriendlyNodeTitles::ForceVisualizationCacheClear() to directly iterate over UEdGraphSchema objects rather than indirectly iterating UClass objects and casting each CDO to a UEdGraphSchema.

#rnx

diff --git a/Engine/Source/Editor/Kismet/Private/BlueprintCompilationManager.cpp b/Engine/Source/Editor/Kismet/Private/BlueprintCompilationManager.cpp
--- a/Engine/Source/Editor/Kismet/Private/BlueprintCompilationManager.cpp
+++ b/Engine/Source/Editor/Kismet/Private/BlueprintCompilationManager.cpp
@@ -314,6 +314,7 @@ struct FCompilerData
 		InternalOptions.bReinstanceAndStubOnFailure = false;
 		InternalOptions.bSaveIntermediateProducts = (UserOptions & EBlueprintCompileOptions::SaveIntermediateProducts) != EBlueprintCompileOptions::None;
 		InternalOptions.bSkipDefaultObjectValidation = (UserOptions & EBlueprintCompileOptions::SkipDefaultObjectValidation) != EBlueprintCompileOptions::None;
+		InternalOptions.bSkipFiBSearchMetaUpdate = (UserOptions & EBlueprintCompileOptions::SkipFiBSearchMetaUpdate) != EBlueprintCompileOptions::None;
 		InternalOptions.CompileType = bBytecodeOnly ? EKismetCompileType::BytecodeOnly : EKismetCompileType::Full;
 
 		Compiler = FKismetCompilerContext::GetCompilerForBP(BP, *ActiveResultsLog, InternalOptions);
@@ -333,6 +334,7 @@ struct FCompilerData
 	bool ShouldRegisterCompilerResults() const { return JobType == ECompilationManagerJobType::Normal; }
 	bool ShouldSkipIfDependenciesAreUnchanged() const { return InternalOptions.CompileType == EKismetCompileType::BytecodeOnly || JobType == ECompilationManagerJobType::RelinkOnly; }
 	bool ShouldValidateClassDefaultObject() const { return JobType == ECompilationManagerJobType::Normal && !InternalOptions.bSkipDefaultObjectValidation; }
+	bool ShouldUpdateBlueprintSearchMetadata() const { return JobType == ECompilationManagerJobType::Normal && !InternalOptions.bSkipFiBSearchMetaUpdate; }
 
 	UBlueprint* BP;
 	FCompilerResultsLog* ActiveResultsLog;
@@ -828,13 +830,16 @@ void FBlueprintCompilationManagerImpl::FlushCompilationQueueImpl(TArray<UObject*
 				}
 			}
 
-			// Do not want to run this code without the editor present nor when running commandlets.
-			if (GEditor && GIsEditor)
+			if (CompilerData.ShouldUpdateBlueprintSearchMetadata())
 			{
-				// We do not want to regenerate a search Guid during loads, nothing has changed in the Blueprint and it is cached elsewhere
-				if (!BP->bIsRegeneratingOnLoad)
+				// Do not want to run this code without the editor present nor when running commandlets.
+				if (GEditor && GIsEditor)
 				{
-					FFindInBlueprintSearchManager::Get().AddOrUpdateBlueprintSearchMetadata(BP);
+					// We do not want to regenerate a search Guid during loads, nothing has changed in the Blueprint and it is cached elsewhere
+					if (!BP->bIsRegeneratingOnLoad)
+					{
+						FFindInBlueprintSearchManager::Get().AddOrUpdateBlueprintSearchMetadata(BP);
+					}
 				}
 			}
 
diff --git a/Engine/Source/Editor/Kismet/Private/FindInBlueprintManager.cpp b/Engine/Source/Editor/Kismet/Private/FindInBlueprintManager.cpp
--- a/Engine/Source/Editor/Kismet/Private/FindInBlueprintManager.cpp
+++ b/Engine/Source/Editor/Kismet/Private/FindInBlueprintManager.cpp
@@ -255,14 +255,9 @@ struct FTemporarilyUseFriendlyNodeTitles
 		if (!bCacheShowFriendlyNames)
 		{
 			// Find all Schemas and force a visualization cache clear
-			for ( TObjectIterator<UClass> ClassIt; ClassIt; ++ClassIt )
+			for ( TObjectIterator<UEdGraphSchema> SchemaIt(RF_NoFlags); SchemaIt; ++SchemaIt)
 			{
-				UClass* CurrentClass = *ClassIt;
-
-				if (UEdGraphSchema* Schema = Cast<UEdGraphSchema>(CurrentClass->GetDefaultObject()))
-				{
-					Schema->ForceVisualizationCacheClear();
-				}
+				SchemaIt->ForceVisualizationCacheClear();
 			}
 		}
 	}
diff --git a/Engine/Source/Editor/UnrealEd/Private/Kismet2/BlueprintEditorUtils.cpp b/Engine/Source/Editor/UnrealEd/Private/Kismet2/BlueprintEditorUtils.cpp
--- a/Engine/Source/Editor/UnrealEd/Private/Kismet2/BlueprintEditorUtils.cpp
+++ b/Engine/Source/Editor/UnrealEd/Private/Kismet2/BlueprintEditorUtils.cpp
@@ -1968,9 +1968,14 @@ void FBlueprintEditorUtils::PostDuplicateBlueprint(UBlueprint* Blueprint, bool b
 
 			if (GBlueprintUseCompilationManager)
 			{
-				// Skip CDO validation in this case as we will not have yet propagated values to the new CDO.
+				// Skip CDO validation in this case as we will not have yet propagated values to the new CDO. Also skip
+				// Blueprint search data updates, as that will be handled by an OnAssetAdded() delegate in the FiB manager.
+				const EBlueprintCompileOptions CompileOptions =
+					EBlueprintCompileOptions::SkipDefaultObjectValidation |
+					EBlueprintCompileOptions::SkipFiBSearchMetaUpdate;
+
 				FBlueprintCompilationManager::CompileSynchronously(
-					FBPCompileRequest(Blueprint, EBlueprintCompileOptions::SkipDefaultObjectValidation, nullptr)
+					FBPCompileRequest(Blueprint, CompileOptions, nullptr)
 				);
 			}
 			else
diff --git a/Engine/Source/Editor/UnrealEd/Private/Kismet2/Kismet2.cpp b/Engine/Source/Editor/UnrealEd/Private/Kismet2/Kismet2.cpp
--- a/Engine/Source/Editor/UnrealEd/Private/Kismet2/Kismet2.cpp
+++ b/Engine/Source/Editor/UnrealEd/Private/Kismet2/Kismet2.cpp
@@ -509,8 +509,14 @@ UBlueprint* FKismetEditorUtilities::CreateBlueprint(UClass* ParentClass, UObject
 		// Skip validation of the class default object here, because (a) the new CDO may fail validation since this
 		// is a new Blueprint that the user has not had a chance to modify any defaults for yet, and (b) in some cases,
 		// default value propagation to the new Blueprint's CDO may be deferred until after compilation (e.g. reparenting).
+		// Also skip the Blueprint search data update, as that will be handled by an OnAssetAdded() delegate in the FiB manager.
+		const EBlueprintCompileOptions CompileOptions =
+			EBlueprintCompileOptions::SkipGarbageCollection |
+			EBlueprintCompileOptions::SkipDefaultObjectValidation |
+			EBlueprintCompileOptions::SkipFiBSearchMetaUpdate;
+
 		FBlueprintCompilationManager::CompileSynchronously(
-			FBPCompileRequest(NewBP, EBlueprintCompileOptions::SkipGarbageCollection|EBlueprintCompileOptions::SkipDefaultObjectValidation, nullptr)
+			FBPCompileRequest(NewBP, CompileOptions, nullptr)
 		);
 	}
 	else
diff --git a/Engine/Source/Editor/UnrealEd/Public/Kismet2/KismetEditorUtilities.h b/Engine/Source/Editor/UnrealEd/Public/Kismet2/KismetEditorUtilities.h
--- a/Engine/Source/Editor/UnrealEd/Public/Kismet2/KismetEditorUtilities.h
+++ b/Engine/Source/Editor/UnrealEd/Public/Kismet2/KismetEditorUtilities.h
@@ -50,6 +50,8 @@ enum class EBlueprintCompileOptions
 	RegenerateSkeletonOnly = 0x80,
 	/** Skips class-specific validation of the default object - in some cases we may not have a fully-initialized CDO after reinstancing */
 	SkipDefaultObjectValidation = 0x100,
+	/** Skips Find-in-Blueprint search data update - in some cases (e.g. creating new assets) this is being deferred until after compilation */
+	SkipFiBSearchMetaUpdate = 0x200,
 };
 
 ENUM_CLASS_FLAGS(EBlueprintCompileOptions)
diff --git a/Engine/Source/Runtime/Engine/Classes/Engine/Blueprint.h b/Engine/Source/Runtime/Engine/Classes/Engine/Blueprint.h
--- a/Engine/Source/Runtime/Engine/Classes/Engine/Blueprint.h
+++ b/Engine/Source/Runtime/Engine/Classes/Engine/Blueprint.h
@@ -158,6 +158,9 @@ public:
 	/** Whether or not to skip class default object validation */
 	bool bSkipDefaultObjectValidation;
 
+	/** Whether or not to update Find-in-Blueprint search metadata */
+	bool bSkipFiBSearchMetaUpdate;
+
 	TSharedPtr<FString> OutHeaderSourceCode;
 	TSharedPtr<FString> OutCppSourceCode;
 	FCompilerNativizationOptions NativizationOptions;
@@ -187,6 +190,7 @@ public:
 		, bIsDuplicationInstigated(false)
 		, bReinstanceAndStubOnFailure(true)
 		, bSkipDefaultObjectValidation(false)
+		, bSkipFiBSearchMetaUpdate(false)
 	{
 	};
 };
