Protect against bad rename operation when a component has been moved aside by a new component in the component's owner's archetype

diff --git a/Engine/Source/Editor/Kismet/Private/SSCSEditor.cpp b/Engine/Source/Editor/Kismet/Private/SSCSEditor.cpp
--- a/Engine/Source/Editor/Kismet/Private/SSCSEditor.cpp
+++ b/Engine/Source/Editor/Kismet/Private/SSCSEditor.cpp
@@ -1148,9 +1148,23 @@ void FSCSEditorTreeNodeInstanceAddedComponent::OnCompleteRename(const FText& InN
 	{
 		RenameFlags |= REN_NonTransactional;
 	}
-
-	ComponentInstance->Rename(*InNewName.ToString(), nullptr, RenameFlags);
-	InstancedComponentName = *InNewName.ToString();
+	
+	// name collision could occur due to e.g. our archetype being updated and causing a conflict with our ComponentInstance:
+	FString NewNameAsString = InNewName.ToString();
+	if(StaticFindObject(UObject::StaticClass(), ComponentInstance->GetOuter(), *NewNameAsString) == nullptr)
+	{
+		ComponentInstance->Rename(*NewNameAsString, nullptr, RenameFlags);
+		InstancedComponentName = *NewNameAsString;
+	}
+	else
+	{
+		UObject* Collision = StaticFindObject(UObject::StaticClass(), ComponentInstance->GetOuter(), *NewNameAsString);
+		if(Collision != ComponentInstance)
+		{
+			// use whatever name the ComponentInstance currently has:
+			InstancedComponentName = ComponentInstance->GetFName();
+		}
+	}
 
 	if (TransactionContext)
 	{
