Harden SCS editor against ComponentTemplatePtr being deleted from underneath it

diff --git a/Engine/Source/Editor/Kismet/Private/SSCSEditor.cpp b/Engine/Source/Editor/Kismet/Private/SSCSEditor.cpp
--- a/Engine/Source/Editor/Kismet/Private/SSCSEditor.cpp
+++ b/Engine/Source/Editor/Kismet/Private/SSCSEditor.cpp
@@ -1141,7 +1141,10 @@ void FSCSEditorTreeNodeInstanceAddedComponent::OnCompleteRename(const FText& InN
 	}
 
 	UActorComponent* ComponentInstance = GetComponentTemplate();
-	check(ComponentInstance != nullptr);
+	if(ComponentInstance == nullptr)
+	{
+		return;
+	}
 
 	ERenameFlags RenameFlags = REN_DontCreateRedirectors;
 	if (!TransactionContext)
@@ -5798,9 +5801,8 @@ FSCSEditorTreeNodePtrType SSCSEditor::AddTreeNode(USCS_Node* InSCSNode, FSCSEdit
 	check(InSCSNode != NULL);
 
 	// During diffs, ComponentTemplates can easily be null, so prevent these checks.
-	if (!bIsDiffing)
+	if (!bIsDiffing && InSCSNode->ComponentTemplate)
 	{
-		check(InSCSNode->ComponentTemplate != NULL);
 		checkf(InSCSNode->ParentComponentOrVariableName == NAME_None
 			|| (!InSCSNode->bIsParentComponentNative && InParentNodePtr->GetSCSNode() != NULL && InParentNodePtr->GetSCSNode()->GetVariableName() == InSCSNode->ParentComponentOrVariableName)
 			|| (InSCSNode->bIsParentComponentNative && InParentNodePtr->GetComponentTemplate() != NULL && InParentNodePtr->GetComponentTemplate()->GetFName() == InSCSNode->ParentComponentOrVariableName),
