Fix static analysis warning about possible leak on component rename

diff --git a/Engine/Source/Editor/Kismet/Private/SSCSEditor.cpp b/Engine/Source/Editor/Kismet/Private/SSCSEditor.cpp
--- a/Engine/Source/Editor/Kismet/Private/SSCSEditor.cpp
+++ b/Engine/Source/Editor/Kismet/Private/SSCSEditor.cpp
@@ -1134,11 +1134,8 @@ void FSCSEditorTreeNodeInstanceAddedComponent::RemoveMeAsChild()
 
 void FSCSEditorTreeNodeInstanceAddedComponent::OnCompleteRename(const FText& InNewName)
 {
-	FScopedTransaction* TransactionContext = NULL;
-	if (!GetAndClearNonTransactionalRenameFlag())
-	{
-		TransactionContext = new FScopedTransaction(LOCTEXT("RenameComponentVariable", "Rename Component Variable"));
-	}
+	bool bIsNonTransactionalRename = GetAndClearNonTransactionalRenameFlag();
+	FScopedTransaction TransactionContext(LOCTEXT("RenameComponentVariable", "Rename Component Variable"), !bIsNonTransactionalRename);
 
 	UActorComponent* ComponentInstance = GetComponentTemplate();
 	if(ComponentInstance == nullptr)
@@ -1147,7 +1144,7 @@ void FSCSEditorTreeNodeInstanceAddedComponent::OnCompleteRename(const FText& InN
 	}
 
 	ERenameFlags RenameFlags = REN_DontCreateRedirectors;
-	if (!TransactionContext)
+	if (bIsNonTransactionalRename)
 	{
 		RenameFlags |= REN_NonTransactional;
 	}
@@ -1168,11 +1165,6 @@ void FSCSEditorTreeNodeInstanceAddedComponent::OnCompleteRename(const FText& InN
 			InstancedComponentName = ComponentInstance->GetFName();
 		}
 	}
-
-	if (TransactionContext)
-	{
-		delete TransactionContext;
-	}
 }
 
 
