Editor: Fixed find in blueprints using fully qualified blueprint paths for asset discovery instead of asset paths

FSearchData::BlueprintPath was misleadingly named as it relates to the asset object that contains the blueprint (ie the UWorld asset or UBlueprint asset), not necessarily the path to the blueprint itself. Removed legacy handling introduced in CL 4209438 that ensure the fully qualified Blueprint path was used, and revert to using FAssetData::ObjectPath directly.

#tests Ran through a Find in Blueprints in QA and internal game projects, along with a full index operation.

diff --git a/Engine/Source/Editor/Kismet/Private/FindInBlueprintManager.cpp b/Engine/Source/Editor/Kismet/Private/FindInBlueprintManager.cpp
--- a/Engine/Source/Editor/Kismet/Private/FindInBlueprintManager.cpp
+++ b/Engine/Source/Editor/Kismet/Private/FindInBlueprintManager.cpp
@@ -98,6 +98,17 @@ const FString FFiBMD::FiBSearchableShallowMD = TEXT("BlueprintSearchableShallow"
 const FString FFiBMD::FiBSearchableExplicitMD = TEXT("BlueprintSearchableExplicit");
 const FString FFiBMD::FiBSearchableHiddenExplicitMD = TEXT("BlueprintSearchableHiddenExplicit");
 
+/* Return the outer of the specified object that is a direct child of a package */
+inline UObject* GetAssetObject(UObject* InObject)
+{
+	UObject* AssetObject = InObject;
+	while (AssetObject && !AssetObject->GetOuter()->IsA<UPackage>())
+	{
+		AssetObject = AssetObject->GetOuter();
+	}
+	return AssetObject;
+}
+
 ////////////////////////////////////
 // FStreamSearch
 FStreamSearch::FStreamSearch(const FString& InSearchValue)
@@ -1479,36 +1490,13 @@ void FFindInBlueprintSearchManager::OnAssetAdded(const FAssetData& InAssetData)
 
 void FFindInBlueprintSearchManager::ExtractUnloadedFiBData(const FAssetData& InAssetData, const FString& InFiBData, bool bIsVersioned)
 {
-	// Add the actual blueprint path to the search map rather than the asset.
-	FName BlueprintPath;
-
-	const FString* BlueprintPathWithinPackage = InAssetData.TagsAndValues.Find(FBlueprintTags::BlueprintPathWithinPackage);
-	if (BlueprintPathWithinPackage)
-	{
-		BlueprintPath = *(InAssetData.PackageName.ToString() + TEXT(".") + (*BlueprintPathWithinPackage));
-	}
-	// ------------------------
-	// Legacy handling for assets saved without FBlueprintTags::BlueprintPathWithinPackage tags.
-	// Prior to this tag, only UBlueprint and UWorld assets were considered.
-	else if (InAssetData.AssetClass == UWorld::StaticClass()->GetFName())
-	{
-		// Assume the path of the level BP to be PersistentLevel.LevelName
-		BlueprintPath = *(InAssetData.ObjectPath.ToString() + TEXT(":PersistentLevel.") + InAssetData.AssetName.ToString());
-	}
-	else
-	{
-		BlueprintPath = InAssetData.ObjectPath;
-	}
-	// End legacy handling
-	// ------------------------
-
-	if (SearchMap.Contains(BlueprintPath))
+	if (SearchMap.Contains(InAssetData.ObjectPath))
 	{
 		return;
 	}
 
 	FSearchData NewSearchData;
-	NewSearchData.BlueprintPath = BlueprintPath;
+	NewSearchData.AssetPath = InAssetData.ObjectPath;
 	InAssetData.GetTagValue(FBlueprintTags::ParentClassPath, NewSearchData.ParentClass);
 
 	const FString ImplementedInterfaces = InAssetData.GetTagValueRef<FString>(FBlueprintTags::ImplementedInterfaces);
@@ -1575,12 +1563,12 @@ void FFindInBlueprintSearchManager::ExtractUnloadedFiBData(const FAssetData& InA
 
 int32 FFindInBlueprintSearchManager::AddSearchDataToDatabase(FSearchData InSearchData)
 {
-	FName BlueprintPath = InSearchData.BlueprintPath; // Copy before we move the data into the array
+	FName AssetPath = InSearchData.AssetPath; // Copy before we move the data into the array
 
 	int32 ArrayIndex = SearchArray.Add(MoveTemp(InSearchData));
 
 	// Add the asset file path to the map along with the index into the array
-	SearchMap.Add(BlueprintPath, ArrayIndex);
+	SearchMap.Add(AssetPath, ArrayIndex);
 
 	return ArrayIndex;
 }
@@ -1614,27 +1602,27 @@ void FFindInBlueprintSearchManager::OnAssetRenamed(const struct FAssetData& InAs
 void FFindInBlueprintSearchManager::OnAssetLoaded(UObject* InAsset)
 {
 	const IBlueprintAssetHandler* Handler = FBlueprintAssetHandler::Get().FindHandler(InAsset->GetClass());
-	UBlueprint* BlueprintAsset = Handler ? Handler->RetrieveBlueprint(InAsset) : nullptr;
+	UBlueprint* BlueprintObject = Handler ? Handler->RetrieveBlueprint(InAsset) : nullptr;
 
-	if (BlueprintAsset)
+	if (BlueprintObject)
 	{
-		FName BlueprintPath = *BlueprintAsset->GetPathName();
+		FName AssetPath = *InAsset->GetPathName();
 
 		// Find and update the item in the search array. Searches may currently be active, this will do no harm to them
 
 		// Confirm that the Blueprint has not been added already, this can occur during duplication of Blueprints.
-		int32* IndexPtr = SearchMap.Find(BlueprintPath);
+		int32* IndexPtr = SearchMap.Find(AssetPath);
 
 		// The asset registry might not have informed us of this asset yet.
 		if(IndexPtr)
 		{
 			// That index should never have a Blueprint already, but if it does, it should be the same Blueprint!
-			ensureMsgf(!SearchArray[*IndexPtr].Blueprint.IsValid() || SearchArray[*IndexPtr].Blueprint == BlueprintAsset, TEXT("Blueprint in database has path %s and is being stomped by %s"), *(SearchArray[*IndexPtr].BlueprintPath.ToString()), *BlueprintPath.ToString());
-			ensureMsgf(!SearchArray[*IndexPtr].Blueprint.IsValid() || SearchArray[*IndexPtr].BlueprintPath == BlueprintPath, TEXT("Blueprint in database has path %s and is being stomped by %s"), *(SearchArray[*IndexPtr].BlueprintPath.ToString()), *BlueprintPath.ToString());
-			SearchArray[*IndexPtr].Blueprint = BlueprintAsset;
+			ensureMsgf(!SearchArray[*IndexPtr].Blueprint.IsValid() || SearchArray[*IndexPtr].Blueprint == BlueprintObject, TEXT("Blueprint in database has path %s and is being stomped by %s"), *(SearchArray[*IndexPtr].AssetPath.ToString()), *AssetPath.ToString());
+			ensureMsgf(!SearchArray[*IndexPtr].Blueprint.IsValid() || SearchArray[*IndexPtr].AssetPath == AssetPath, TEXT("Blueprint in database has path %s and is being stomped by %s"), *(SearchArray[*IndexPtr].AssetPath.ToString()), *AssetPath.ToString());
+			SearchArray[*IndexPtr].Blueprint = BlueprintObject;
 		}
 
-		UncachedAssets.Remove(FName(*InAsset->GetPathName()));
+		UncachedAssets.Remove(AssetPath);
 	}
 }
 
@@ -1749,20 +1737,22 @@ void FFindInBlueprintSearchManager::AddOrUpdateBlueprintSearchMetadata(UBlueprin
 		return;
 	}
 
-	check(InBlueprint);
+	UObject* AssetObject = GetAssetObject(InBlueprint);
+
+	check(InBlueprint && AssetObject);
 
 	// Allow only one thread modify the search data at a time
 	FScopeLock ScopeLock(&SafeModifyCacheCriticalSection);
 
-	FName BlueprintPath = *InBlueprint->GetPathName();
+	FName AssetPath = *AssetObject->GetPathName();
 
-	int32* IndexPtr = SearchMap.Find(BlueprintPath);
+	int32* IndexPtr = SearchMap.Find(AssetPath);
 	int32 Index = 0;
 	if(!IndexPtr)
 	{
 		FSearchData SearchData;
 		SearchData.Blueprint = InBlueprint;
-		SearchData.BlueprintPath = BlueprintPath;
+		SearchData.AssetPath = AssetPath;
 		Index = AddSearchDataToDatabase(MoveTemp(SearchData));
 	}
 	else
@@ -1830,7 +1820,7 @@ bool FFindInBlueprintSearchManager::ContinueSearchQuery(const FStreamSearch* InS
 				// If there is FiB data, parse it into an ImaginaryBlueprint
 				if (SearchArray[SearchIdx].Value.Len() > 0)
 				{
-					SearchArray[SearchIdx].ImaginaryBlueprint = MakeShareable(new FImaginaryBlueprint(FPaths::GetBaseFilename(SearchArray[SearchIdx].BlueprintPath.ToString()), SearchArray[SearchIdx].BlueprintPath.ToString(), SearchArray[SearchIdx].ParentClass, SearchArray[SearchIdx].Interfaces, SearchArray[SearchIdx].Value, SearchArray[SearchIdx].Version != 0));
+					SearchArray[SearchIdx].ImaginaryBlueprint = MakeShareable(new FImaginaryBlueprint(FPaths::GetBaseFilename(SearchArray[SearchIdx].AssetPath.ToString()), SearchArray[SearchIdx].AssetPath.ToString(), SearchArray[SearchIdx].ParentClass, SearchArray[SearchIdx].Interfaces, SearchArray[SearchIdx].Value, SearchArray[SearchIdx].Version != 0));
 					SearchArray[SearchIdx].Value.Empty();
 				}
 
@@ -1891,7 +1881,9 @@ FString FFindInBlueprintSearchManager::QuerySingleBlueprint(UBlueprint* InBluepr
 			AddOrUpdateBlueprintSearchMetadata(InBlueprint, true);
 		}
 
-		FName Key = *InBlueprint->GetPathName();
+		UObject* AssetObject = GetAssetObject(InBlueprint);
+		check(AssetObject);
+		FName Key = *AssetObject->GetPathName();
 
 		int32* ArrayIdx = SearchMap.Find(Key);
 		// This should always be true since we make sure to refresh the search data for this Blueprint when doing the search, unless we do not rebuild the searchable data
@@ -1948,7 +1940,7 @@ void FFindInBlueprintSearchManager::CleanCache()
 			FSearchData SearchData;
 	 		ContinueSearchQuery(ActiveSearch, SearchData);
 
-			FName CachePath = SearchData.BlueprintPath;
+			FName CachePath = SearchData.AssetPath;
 	 		CacheQueries.Add(ActiveSearch, CachePath);
 	 	}
 	}
@@ -1976,7 +1968,7 @@ void FFindInBlueprintSearchManager::CleanCache()
 				AssetRegistryModule = &FModuleManager::LoadModuleChecked<FAssetRegistryModule>(TEXT("AssetRegistry"));
 
 				// The asset was not user deleted, so this should usually find the asset. New levels can be deleted if they were not saved
-				FAssetData AssetData = AssetRegistryModule->Get().GetAssetByObjectPath(SearchArray[SearchValuePair.Value].BlueprintPath);
+				FAssetData AssetData = AssetRegistryModule->Get().GetAssetByObjectPath(SearchArray[SearchValuePair.Value].AssetPath);
 				if(AssetData.IsValid())
 				{
 					if(const FString* FiBSearchData = AssetData.TagsAndValues.Find("FiB"))
@@ -2100,7 +2092,7 @@ void FFindInBlueprintSearchManager::CacheAllUncachedAssets(TWeakPtr< SFindInBlue
 		{
 			if ((SearchData.Value.Len() != 0 || SearchData.ImaginaryBlueprint.IsValid()) && SearchData.Version < InMinimiumVersionRequirement)
 			{
-				BlueprintsToUpdate.Add(SearchData.BlueprintPath);
+				BlueprintsToUpdate.Add(SearchData.AssetPath);
 			}
 		}
 
diff --git a/Engine/Source/Editor/Kismet/Public/FindInBlueprintManager.h b/Engine/Source/Editor/Kismet/Public/FindInBlueprintManager.h
--- a/Engine/Source/Editor/Kismet/Public/FindInBlueprintManager.h
+++ b/Engine/Source/Editor/Kismet/Public/FindInBlueprintManager.h
@@ -107,8 +107,8 @@ struct FSearchData
 	/** The Blueprint this search data points to, if available */
 	TWeakObjectPtr<UBlueprint> Blueprint;
 
-	/** The full Blueprint path this search data is associated with */
-	FName BlueprintPath;
+	/** The full asset path this search data is associated with of the form /Game/Path/To/Package.Package */
+	FName AssetPath;
 
 	/** Search data block for the Blueprint */
 	FString Value;
