Fix for other subeditors writing out their layouts on close

diff --git a/Engine/Source/Editor/Kismet/Private/ApplicationMode.cpp b/Engine/Source/Editor/Kismet/Private/ApplicationMode.cpp
--- a/Engine/Source/Editor/Kismet/Private/ApplicationMode.cpp
+++ b/Engine/Source/Editor/Kismet/Private/ApplicationMode.cpp
@@ -3,6 +3,7 @@
 #include "WorkflowOrientedApp/ApplicationMode.h"
 #include "Framework/MultiBox/MultiBoxExtender.h"
 #include "Framework/Docking/LayoutService.h"
+#include "UnrealEdMisc.h"
 
 #define LOCTEXT_NAMESPACE "ApplicationMode"
 
@@ -27,8 +28,10 @@ void FApplicationMode::DeactivateMode(TSharedPtr<FTabManager> InTabManager)
 {
 	// Save the layout to INI
 	check(InTabManager.IsValid());
-	FLayoutSaveRestore::SaveToConfig(GEditorLayoutIni, InTabManager->PersistLayout());
-		
+	if (FUnrealEdMisc::Get().IsSavingLayoutOnClosedAllowed())
+	{
+		FLayoutSaveRestore::SaveToConfig(GEditorLayoutIni, InTabManager->PersistLayout());
+	}
 	// Unregister the tabs
 	/*
 	for (int32 Index = 0; Index < AllowableTabs.Num(); ++Index)
diff --git a/Engine/Source/Editor/LevelEditor/Public/LevelEditor.h b/Engine/Source/Editor/LevelEditor/Public/LevelEditor.h
--- a/Engine/Source/Editor/LevelEditor/Public/LevelEditor.h
+++ b/Engine/Source/Editor/LevelEditor/Public/LevelEditor.h
@@ -323,7 +323,10 @@ private:
 	// Callback for persisting the Level Editor's layout.
 	void HandleTabManagerPersistLayout( const TSharedRef<FTabManager::FLayout>& LayoutToSave )
 	{
-		FLayoutSaveRestore::SaveToConfig(GEditorLayoutIni, LayoutToSave);
+		if (FUnrealEdMisc::Get().IsSavingLayoutOnClosedAllowed())
+		{
+			FLayoutSaveRestore::SaveToConfig(GEditorLayoutIni, LayoutToSave);
+		}
 	}
 
 	/** Spawn the main level editor tab */
diff --git a/Engine/Source/Editor/MainFrame/Private/Frame/MainFrameHandler.h b/Engine/Source/Editor/MainFrame/Private/Frame/MainFrameHandler.h
--- a/Engine/Source/Editor/MainFrame/Private/Frame/MainFrameHandler.h
+++ b/Engine/Source/Editor/MainFrame/Private/Frame/MainFrameHandler.h
@@ -282,7 +282,10 @@ private:
 	// Callback for persisting the Level Editor's layout.
 	void HandleTabManagerPersistLayout( const TSharedRef<FTabManager::FLayout>& LayoutToSave )
 	{
-		FLayoutSaveRestore::SaveToConfig(GEditorLayoutIni, LayoutToSave);
+		if (FUnrealEdMisc::Get().IsSavingLayoutOnClosedAllowed())
+		{
+			FLayoutSaveRestore::SaveToConfig(GEditorLayoutIni, LayoutToSave);
+		}
 	}
 
 private:
