Fix for bad AttachParent data after recompiling a blueprint that is using a ChildActor with components that are attached to eachother (Mirror 4525227)

diff --git a/Engine/Source/Editor/Kismet/Private/BlueprintCompilationManager.cpp b/Engine/Source/Editor/Kismet/Private/BlueprintCompilationManager.cpp
--- a/Engine/Source/Editor/Kismet/Private/BlueprintCompilationManager.cpp
+++ b/Engine/Source/Editor/Kismet/Private/BlueprintCompilationManager.cpp
@@ -94,6 +94,7 @@ struct FBlueprintCompilationManagerImpl : public FGCObject
 	bool HasBlueprintsToCompile() const;
 	bool IsGeneratedClassLayoutReady() const;
 	void GetDefaultValue(const UClass* ForClass, const UProperty* Property, FString& OutDefaultValueAsString) const;
+	static void BuildDSOMap(UObject* OldObject, UObject* NewObject, TMap<UObject*, UObject*>& OutOldToNewDSO);
 	static void ReinstanceBatch(TArray<FReinstancingJob>& Reinstancers, TMap< UClass*, UClass* >& InOutOldToNewClassMap, TArray<UObject*>* ObjLoaded);
 	static UClass* FastGenerateSkeletonClass(UBlueprint* BP, FKismetCompilerContext& CompilerContext, bool bIsSkeletonOnly);
 	static bool IsQueuedForCompilation(UBlueprint* BP);
@@ -1325,6 +1326,28 @@ void FBlueprintCompilationManagerImpl::GetDefaultValue(const UClass* ForClass, c
 	}
 }
 
+void FBlueprintCompilationManagerImpl::BuildDSOMap(UObject* OldObject, UObject* NewObject, TMap<UObject*, UObject*>& OutOldToNewDSO)
+{
+	// IsDefaultObject() unfortunately cannot be relied upon for archetypes, so we explicitly search for the flag:
+	TArray<UObject*> OldSubobjects;
+	ForEachObjectWithOuter(OldObject, [&OldSubobjects](UObject* Object)
+	{
+		if (Object->HasAnyFlags(RF_DefaultSubObject))
+		{
+			OldSubobjects.Add(Object);
+		}
+	}, false);
+
+	for(UObject* OldSubobject : OldSubobjects)
+	{
+		UObject* NewSubobject = NewObject ? StaticFindObjectFast(UObject::StaticClass(), NewObject, OldSubobject->GetFName()) : nullptr;
+		// It may seem aggressive to ensure here, but we should always have a new version of the DSO. If that's
+		// not the case then some testing will need to be done on client of OutOldToNewDSO
+		OutOldToNewDSO.Add(OldSubobject, NewSubobject);
+		BuildDSOMap(OldSubobject, NewSubobject, OutOldToNewDSO);
+	}
+}
+
 void FBlueprintCompilationManagerImpl::ReinstanceBatch(TArray<FReinstancingJob>& Reinstancers, TMap< UClass*, UClass* >& InOutOldToNewClassMap, TArray<UObject*>* ObjLoaded)
 {
 	const auto FilterOutOfDateClasses = [](TArray<UClass*>& ClassList)
@@ -1643,6 +1666,9 @@ void FBlueprintCompilationManagerImpl::ReinstanceBatch(TArray<FReinstancingJob>&
 
 				OldArchetypeToNewArchetype.Add(Archetype, NewArchetype);
 
+				// also map old *default* subobjects to new default subobjects:
+				BuildDSOMap(Archetype, NewArchetype, OldArchetypeToNewArchetype);
+
 				ArchetypeReferencers.Add(NewArchetype);
 
 				FLinkerLoad::PRIVATE_PatchNewObjectIntoExport(Archetype, NewArchetype);
@@ -1666,7 +1692,8 @@ void FBlueprintCompilationManagerImpl::ReinstanceBatch(TArray<FReinstancingJob>&
 			for(UObject* OldInstance : OldInstances)
 			{
 				UObject** NewInstance = OldArchetypeToNewArchetype.Find(OldInstance);
-				if(NewInstance)
+				// NewInstance may be null in the case of deleted or EditorOnly (in -game) DSOs:
+				if(NewInstance && *NewInstance)
 				{
 					// The new object hierarchy has been created, all of the old instances are in the transient package and new
 					// ones have taken their place. Referenc members will mostly be pointing at *old* instances, and will get fixed
