Guard DiffUtils::Identical against nullptr UProperties

diff --git a/Engine/Source/Editor/Kismet/Private/DiffUtils.cpp b/Engine/Source/Editor/Kismet/Private/DiffUtils.cpp
--- a/Engine/Source/Editor/Kismet/Private/DiffUtils.cpp
+++ b/Engine/Source/Editor/Kismet/Private/DiffUtils.cpp
@@ -605,6 +605,15 @@ static void IdenticalHelper(const UProperty* AProperty, const UProperty* BProper
 
 bool DiffUtils::Identical(const FResolvedProperty& AProp, const FResolvedProperty& BProp, const FPropertySoftPath& RootPath, TArray<FPropertySoftPath>& DifferingProperties)
 {
+	if( AProp.Property == nullptr && BProp.Property == nullptr )
+	{
+		return true;
+	}
+	else if( AProp.Property == nullptr || BProp.Property == nullptr )
+	{
+		return false;
+	}
+
 	const void* AValue = AProp.Property->ContainerPtrToValuePtr<void>(AProp.Object);
 	const void* BValue = BProp.Property->ContainerPtrToValuePtr<void>(BProp.Object);
 
