Copy of CL #4581002

diff --git a/Engine/Source/Editor/Kismet/Private/UserDefinedStructureEditor.cpp b/Engine/Source/Editor/Kismet/Private/UserDefinedStructureEditor.cpp
--- a/Engine/Source/Editor/Kismet/Private/UserDefinedStructureEditor.cpp
+++ b/Engine/Source/Editor/Kismet/Private/UserDefinedStructureEditor.cpp
@@ -152,7 +152,9 @@ public:
 		if (Info != FStructureEditorUtils::DefaultValueChanged)
 		{
 			StructData->Initialize(UserDefinedStruct.Get());
-			DetailsView->SetObject(UserDefinedStruct.Get());
+
+			// Force the set object call because we may be called multiple times in a row if more than one struct was changed at the same time
+			DetailsView->SetObject(UserDefinedStruct.Get(), true);
 		}
 
 		UserDefinedStruct.Get()->InitializeDefaultValue(StructData->GetStructMemory());
