Validate Blueprint component tree node data on an input focus loss event that also triggers a component name validity check. Previously, invalid node data caused a hard editor crash.

- Mirrored from //UE4/Release-[REDACTED_PHONE]).

#rnx

diff --git a/Engine/Source/Editor/Kismet/Private/SSCSEditor.cpp b/Engine/Source/Editor/Kismet/Private/SSCSEditor.cpp
--- a/Engine/Source/Editor/Kismet/Private/SSCSEditor.cpp
+++ b/Engine/Source/Editor/Kismet/Private/SSCSEditor.cpp
@@ -3073,27 +3073,36 @@ bool SSCS_RowWidget::OnNameTextVerifyChanged(const FText& InNewText, FText& OutE
 
 	if (!InNewText.IsEmpty())
 	{
-		AActor* ExistingNameSearchScope = NodePtr->GetComponentTemplate()->GetOwner();
-		if ((ExistingNameSearchScope == nullptr) && (Blueprint != nullptr))
+		const UActorComponent* ComponentInstance = NodePtr->GetComponentTemplate();
+		if (ensure(ComponentInstance))
 		{
-			ExistingNameSearchScope = Cast<AActor>(Blueprint->GeneratedClass->GetDefaultObject());
-		}
+			AActor* ExistingNameSearchScope = ComponentInstance->GetOwner();
+			if ((ExistingNameSearchScope == nullptr) && (Blueprint != nullptr))
+			{
+				ExistingNameSearchScope = Cast<AActor>(Blueprint->GeneratedClass->GetDefaultObject());
+			}
 
-		if (!FComponentEditorUtils::IsValidVariableNameString(NodePtr->GetComponentTemplate(), InNewText.ToString()))
-		{
-			OutErrorMessage = LOCTEXT("RenameFailed_EngineReservedName", "This name is reserved for engine use.");
-			return false;
-		}
-		else if (InNewText.ToString().Len() > NAME_SIZE)
-		{
-			FFormatNamedArguments Arguments;
-			Arguments.Add(TEXT("CharCount"), NAME_SIZE);
-			OutErrorMessage = FText::Format(LOCTEXT("ComponentRenameFailed_TooLong", "Component name must be less than {CharCount} characters long."), Arguments);
-			return false;
+			if (!FComponentEditorUtils::IsValidVariableNameString(ComponentInstance, InNewText.ToString()))
+			{
+				OutErrorMessage = LOCTEXT("RenameFailed_EngineReservedName", "This name is reserved for engine use.");
+				return false;
+			}
+			else if (InNewText.ToString().Len() > NAME_SIZE)
+			{
+				FFormatNamedArguments Arguments;
+				Arguments.Add(TEXT("CharCount"), NAME_SIZE);
+				OutErrorMessage = FText::Format(LOCTEXT("ComponentRenameFailed_TooLong", "Component name must be less than {CharCount} characters long."), Arguments);
+				return false;
+			}
+			else if (!FComponentEditorUtils::IsComponentNameAvailable(InNewText.ToString(), ExistingNameSearchScope, ComponentInstance))
+			{
+				OutErrorMessage = LOCTEXT("RenameFailed_ExistingName", "Another component already has the same name.");
+				return false;
+			}
 		}
-		else if (!FComponentEditorUtils::IsComponentNameAvailable(InNewText.ToString(), ExistingNameSearchScope, NodePtr->GetComponentTemplate()))
+		else
 		{
-			OutErrorMessage = LOCTEXT("RenameFailed_ExistingName", "Another component already has the same name.");
+			OutErrorMessage = LOCTEXT("RenameFailed_InvalidComponentInstance", "This node is referencing an invalid component instance and cannot be renamed. Perhaps it was destroyed?");
 			return false;
 		}
 	}
