Remove unused parater to FText::Format
#rnx

diff --git a/Engine/Source/Editor/Kismet/Private/BlueprintDetailsCustomization.cpp b/Engine/Source/Editor/Kismet/Private/BlueprintDetailsCustomization.cpp
--- a/Engine/Source/Editor/Kismet/Private/BlueprintDetailsCustomization.cpp
+++ b/Engine/Source/Editor/Kismet/Private/BlueprintDetailsCustomization.cpp
@@ -4144,7 +4144,7 @@ bool FBaseBlueprintGraphActionDetails::OnVerifyPinRename(UK2Node_EditablePinBase
 
 	if (InNewName == TEXT("None"))
 	{
-		OutErrorMessage = FText::Format( LOCTEXT("PinNameNone", "'None' is a reserved name"), FText::AsNumber( NAME_SIZE ) );
+		OutErrorMessage = FText::Format( LOCTEXT("PinNameNone", "'None' is a reserved name"));
 		return false;
 	}
 
