Copy/Paste of components in bp scs now correctly maintain relative hierarchy
Duplicate of components in bp scs and iwce now correctly maintain relative hierarchy
Copy/Paste and Duplicate of components in bp scs and iwce now maintain the name of the source component where possible and a modified version of the name when already taken instead of using the component class as name

diff --git a/Engine/Source/Editor/Kismet/Private/SSCSEditor.cpp b/Engine/Source/Editor/Kismet/Private/SSCSEditor.cpp
--- a/Engine/Source/Editor/Kismet/Private/SSCSEditor.cpp
+++ b/Engine/Source/Editor/Kismet/Private/SSCSEditor.cpp
@@ -4092,49 +4092,63 @@ void SSCSEditor::OnDuplicateComponent()
 	{
 		const FScopedTransaction Transaction(SelectedNodes.Num() > 1 ? LOCTEXT("DuplicateComponents", "Duplicate Components") : LOCTEXT("DuplicateComponent", "Duplicate Component"));
 
+		TMap<USceneComponent*, USceneComponent*> DuplicateSceneComponentMap;
 		for (int32 i = 0; i < SelectedNodes.Num(); ++i)
 		{
 			if (UActorComponent* ComponentTemplate = SelectedNodes[i]->GetComponentTemplate())
 			{
 				UActorComponent* CloneComponent = AddNewComponent(ComponentTemplate->GetClass(), ComponentTemplate);
-				UActorComponent* OriginalComponent = ComponentTemplate;
-
-				// If we've duplicated a scene component, attempt to reposition the duplicate in the hierarchy if the original
-				// was attached to another scene component as a child. By default, the duplicate is attached to the scene root node.
-				USceneComponent* NewSceneComponent = Cast<USceneComponent>(CloneComponent);
-				if(NewSceneComponent != NULL)
+				if (USceneComponent* SceneClone = Cast<USceneComponent>(CloneComponent))
 				{
-					if (EditorMode == EComponentEditorMode::BlueprintSCS)
-					{
-						// Ensure that any native attachment relationship inherited from the original copy is removed (to prevent a GLEO assertion)
-						NewSceneComponent->DetachFromComponent(FDetachmentTransformRules::KeepWorldTransform);
-					}
+					DuplicateSceneComponentMap.Add(CastChecked<USceneComponent>(ComponentTemplate), SceneClone);
+				}
+			}
+		}
+
+		for (const TPair<USceneComponent*,USceneComponent*>& DuplicatedPair : DuplicateSceneComponentMap)
+		{
+			USceneComponent* OriginalComponent = DuplicatedPair.Key;
+			USceneComponent* NewSceneComponent = DuplicatedPair.Value;
+
+			if (EditorMode == EComponentEditorMode::BlueprintSCS)
+			{
+				// Ensure that any native attachment relationship inherited from the original copy is removed (to prevent a GLEO assertion)
+				NewSceneComponent->DetachFromComponent(FDetachmentTransformRules::KeepWorldTransform);
+			}
 					
-					// Attempt to locate the original node in the SCS tree
-					FSCSEditorTreeNodePtrType OriginalNodePtr = FindTreeNode(OriginalComponent);
-					if(OriginalNodePtr.IsValid())
+			// Attempt to locate the original node in the SCS tree
+			FSCSEditorTreeNodePtrType OriginalNodePtr = FindTreeNode(OriginalComponent);
+			if(OriginalNodePtr.IsValid())
+			{
+				// If we're duplicating the root then we're already a child of it so need to reparent, but we do need to reset the scale
+				// otherwise we'll end up with the square of the root's scale instead of being the same size.
+				if (OriginalNodePtr == SceneRootNodePtr)
+				{
+					NewSceneComponent->RelativeScale3D = FVector(1.f);
+				}
+				else
+				{
+					// If the original node was parented, attempt to add the duplicate as a child of the same parent node if the parent is not
+					// part of the duplicate set, otherwise parent to the parent's duplicate
+					FSCSEditorTreeNodePtrType ParentNodePtr = OriginalNodePtr->GetParent();
+					if (ParentNodePtr.IsValid())
 					{
-						// If we're duplicating the root then we're already a child of it so need to reparent, but we do need to reset the scale
-						// otherwise we'll end up with the square of the root's scale instead of being the same size.
-						if (OriginalNodePtr == SceneRootNodePtr)
-						{
-							NewSceneComponent->RelativeScale3D = FVector(1.f);
-						}
-						else
+						if (USceneComponent** ParentDuplicateComponent = DuplicateSceneComponentMap.Find(Cast<USceneComponent>(ParentNodePtr->GetComponentTemplate())))
 						{
-							// If the original node was parented, attempt to add the duplicate as a child of the same parent node
-							FSCSEditorTreeNodePtrType ParentNodePtr = OriginalNodePtr->GetParent();
-							if (ParentNodePtr.IsValid())
+							FSCSEditorTreeNodePtrType DuplicateParentNodePtr = FindTreeNode(*ParentDuplicateComponent);
+							if (DuplicateParentNodePtr.IsValid())
 							{
-								// Locate the duplicate node (as a child of the current scene root node), and switch it to be a child of the original node's parent
-								FSCSEditorTreeNodePtrType NewChildNodePtr = SceneRootNodePtr->FindChild(NewSceneComponent, true);
-								if (NewChildNodePtr.IsValid())
-								{
-									// Note: This method will handle removal from the scene root node as well
-									ParentNodePtr->AddChild(NewChildNodePtr);
-								}
+								ParentNodePtr = DuplicateParentNodePtr;
 							}
 						}
+
+						// Locate the duplicate node (as a child of the current scene root node), and switch it to be a child of the original node's parent
+						FSCSEditorTreeNodePtrType NewChildNodePtr = SceneRootNodePtr->FindChild(NewSceneComponent, true);
+						if (NewChildNodePtr.IsValid())
+						{
+							// Note: This method will handle removal from the scene root node as well
+							ParentNodePtr->AddChild(NewChildNodePtr);
+						}
 					}
 				}
 			}
@@ -5128,7 +5142,18 @@ UActorComponent* SSCSEditor::AddNewComponent( UClass* NewComponentClass, UObject
 			bAllowTreeUpdates = false;
 		}
 		
-		const FName NewVariableName = (Asset ? FName(*FComponentEditorUtils::GenerateValidVariableNameFromAsset(Asset, nullptr)) : NAME_None);
+		FName NewVariableName;
+		if (ComponentTemplate)
+		{
+			FString TemplateName = ComponentTemplate->GetName();
+			NewVariableName = (TemplateName.EndsWith(USimpleConstructionScript::ComponentTemplateNameSuffix) 
+								? FName(*TemplateName.LeftChop(USimpleConstructionScript::ComponentTemplateNameSuffix.Len()))
+								: ComponentTemplate->GetFName());
+		}
+		else if (Asset)
+		{
+			NewVariableName = *FComponentEditorUtils::GenerateValidVariableNameFromAsset(Asset, nullptr);
+		}
 		NewComponent = AddNewNode(Blueprint->SimpleConstructionScript->CreateNode(NewComponentClass, NewVariableName), Asset, bMarkBlueprintModified, bSetFocusToNewItem);
 
 		if (ComponentTemplate)
@@ -5425,11 +5450,41 @@ void SSCSEditor::CopySelectedNodes()
 		if (ComponentTemplate)
 		{
 			ComponentsToCopy.Add(ComponentTemplate);
+
+			if (EditorMode == EComponentEditorMode::BlueprintSCS)
+			{
+				// CopyComponents uses component attachment to maintain hierarchy, but the SCS templates are not
+				// setup with a relationship to each other. Briefly setup the attachment between the templates being
+				// copied so that the hierarchy is retained upon pasting
+				if (USceneComponent* SceneTemplate = Cast<USceneComponent>(ComponentTemplate))
+				{
+					FSCSEditorTreeNodePtrType SelectedParentNodePtr = SelectedNodePtr->GetParent();
+					if (SelectedParentNodePtr.IsValid())
+					{
+						if (USceneComponent* ParentSceneTemplate = Cast<USceneComponent>(SelectedParentNodePtr->GetComponentTemplate()))
+						{
+							SceneTemplate->SetupAttachment(ParentSceneTemplate);
+						}
+					}
+				}
+			}
 		}
 	}
 
 	// Copy the components to the clipboard
 	FComponentEditorUtils::CopyComponents(ComponentsToCopy);
+
+	if (EditorMode == EComponentEditorMode::BlueprintSCS)
+	{
+		for (UActorComponent* ComponentTemplate : ComponentsToCopy)
+		{
+			if (USceneComponent* SceneTemplate = Cast<USceneComponent>(ComponentTemplate))
+			{
+				// clear back out any temporary attachments we set up for the copy
+				SceneTemplate->SetupAttachment(nullptr);
+			}
+		}
+	}
 }
 
 bool SSCSEditor::CanPasteNodes() const
diff --git a/Engine/Source/Editor/UnrealEd/Private/Kismet2/ComponentEditorUtils.cpp b/Engine/Source/Editor/UnrealEd/Private/Kismet2/ComponentEditorUtils.cpp
--- a/Engine/Source/Editor/UnrealEd/Private/Kismet2/ComponentEditorUtils.cpp
+++ b/Engine/Source/Editor/UnrealEd/Private/Kismet2/ComponentEditorUtils.cpp
@@ -428,7 +428,7 @@ void FComponentEditorUtils::PasteComponents(TArray<UActorComponent*>& OutPastedC
 		check(NewActorComponent);
 
 		// Relocate the instance from the transient package to the Actor and assign it a unique object name
-		FString NewComponentName = FComponentEditorUtils::GenerateValidVariableName(NewActorComponent->GetClass(), TargetActor);
+		FString NewComponentName = FComponentEditorUtils::GenerateValidVariableNameFromAsset(NewActorComponent, TargetActor);
 		NewActorComponent->Rename(*NewComponentName, TargetActor, REN_DontCreateRedirectors | REN_DoNotDirty);
 
 		if (USceneComponent* NewSceneComponent = Cast<USceneComponent>(NewActorComponent))
@@ -631,8 +631,7 @@ UActorComponent* FComponentEditorUtils::DuplicateComponent(UActorComponent* Temp
 	if (!TemplateComponent->IsVisualizationComponent() && Actor)
 	{
 		Actor->Modify();
-		UClass* ComponentClass = TemplateComponent->GetClass();
-		FName NewComponentName = *FComponentEditorUtils::GenerateValidVariableName(ComponentClass, Actor);
+		FName NewComponentName = *FComponentEditorUtils::GenerateValidVariableNameFromAsset(TemplateComponent, Actor);
 
 		bool bKeepWorldLocationOnAttach = false;
 
diff --git a/Engine/Source/Runtime/Engine/Private/SimpleConstructionScript.cpp b/Engine/Source/Runtime/Engine/Private/SimpleConstructionScript.cpp
--- a/Engine/Source/Runtime/Engine/Private/SimpleConstructionScript.cpp
+++ b/Engine/Source/Runtime/Engine/Private/SimpleConstructionScript.cpp
@@ -1253,7 +1253,6 @@ USCS_Node* USimpleConstructionScript::CreateNode(UClass* NewComponentClass, FNam
 	check(NewComponentClass->IsChildOf(UActorComponent::StaticClass()));
 	ensure(Cast<UBlueprintGeneratedClass>(Blueprint->GeneratedClass));
 
-	// note that naming logic is duplicated in CreateNodeAndRenameComponent:
 	NewComponentVariableName = GenerateNewComponentName(NewComponentClass, NewComponentVariableName);
 
 	// At this point we should have a unique, explicit name to use for the template object.
@@ -1277,8 +1276,16 @@ USCS_Node* USimpleConstructionScript::CreateNodeAndRenameComponent(UActorCompone
 {
 	check(NewComponentTemplate);
 
-	// note that naming logic is duplicated in CreateNode:
-	FName NewComponentVariableName = GenerateNewComponentName(NewComponentTemplate->GetClass());
+	// When copying and pasting we'd prefer to keep the component name
+	// However, the incoming template will have the template name suffix on it so
+	// acquire the desired name by stripping the suffix
+	FName DesiredName;
+	FString TemplateName = NewComponentTemplate->GetName();
+	if (TemplateName.EndsWith(ComponentTemplateNameSuffix))
+	{
+		DesiredName = *TemplateName.LeftChop(ComponentTemplateNameSuffix.Len());
+	}
+	FName NewComponentVariableName = GenerateNewComponentName(NewComponentTemplate->GetClass(), DesiredName);
 
 	// At this point we should have a unique, explicit name to use for the template object.
 	check(NewComponentVariableName != NAME_None);
