Moving node alignment tools to a more generic implementation and adding menu options for them to the material editor

diff --git a/Engine/Source/Editor/GraphEditor/Private/SGraphEditorImpl.cpp b/Engine/Source/Editor/GraphEditor/Private/SGraphEditorImpl.cpp
--- a/Engine/Source/Editor/GraphEditor/Private/SGraphEditorImpl.cpp
+++ b/Engine/Source/Editor/GraphEditor/Private/SGraphEditorImpl.cpp
@@ -18,9 +18,39 @@
 #include "ScopedTransaction.h"
 #include "SGraphEditorActionMenu.h"
 #include "Widgets/Notifications/SNotificationList.h"
+#include "EdGraphSchema_K2.h"
 
 #define LOCTEXT_NAMESPACE "GraphEditorModule"
 
+FVector2D GetNodeSize(const SGraphEditor& GraphEditor, const UEdGraphNode* Node)
+{
+	FSlateRect Rect;
+	if (GraphEditor.GetBoundsForNode(Node, Rect, 0.f))
+	{
+		return FVector2D(Rect.Right - Rect.Left, Rect.Bottom - Rect.Top);
+	}
+
+	return FVector2D(Node->NodeWidth, Node->NodeHeight);
+}
+
+/////////////////////////////////////////////////////
+// FAlignmentHelper
+
+FAlignmentData FAlignmentHelper::GetAlignmentDataForNode(UEdGraphNode* Node)
+{
+	float PropertyOffset = 0.f;
+
+	const float NodeSize = Orientation == Orient_Horizontal ? GetNodeSize(*GraphEditor, Node).X : GetNodeSize(*GraphEditor, Node).Y;
+	switch (AlignType)
+	{
+	case EAlignType::Minimum:	PropertyOffset = 0.f;			break;
+	case EAlignType::Middle:	PropertyOffset = NodeSize * .5f;	break;
+	case EAlignType::Maximum:	PropertyOffset = NodeSize;		break;
+	}
+	int32* Property = Orientation == Orient_Horizontal ? &Node->NodePosX : &Node->NodePosY;
+	return FAlignmentData(Node, *Property, PropertyOffset);
+}
+
 /////////////////////////////////////////////////////
 // SGraphEditorImpl
 
@@ -808,6 +838,150 @@ void SGraphEditorImpl::SetNodeFactory(const TSharedRef<class FGraphNodeFactory>&
 	GraphPanel->SetNodeFactory(NewNodeFactory);
 }
 
+
+void SGraphEditorImpl::OnAlignTop()
+{
+	const FScopedTransaction Transaction(FGraphEditorCommands::Get().AlignNodesTop->GetLabel());
+
+	FAlignmentHelper Helper(SharedThis(this), Orient_Vertical, EAlignType::Minimum);
+	Helper.Align();
+}
+
+void SGraphEditorImpl::OnAlignMiddle()
+{
+	const FScopedTransaction Transaction(FGraphEditorCommands::Get().AlignNodesMiddle->GetLabel());
+
+	FAlignmentHelper Helper(SharedThis(this), Orient_Vertical, EAlignType::Middle);
+	Helper.Align();
+}
+
+void SGraphEditorImpl::OnAlignBottom()
+{
+	const FScopedTransaction Transaction(FGraphEditorCommands::Get().AlignNodesBottom->GetLabel());
+
+	FAlignmentHelper Helper(SharedThis(this), Orient_Vertical, EAlignType::Maximum);
+	Helper.Align();
+}
+
+void SGraphEditorImpl::OnAlignLeft()
+{
+	const FScopedTransaction Transaction(FGraphEditorCommands::Get().AlignNodesLeft->GetLabel());
+
+	FAlignmentHelper Helper(SharedThis(this), Orient_Horizontal, EAlignType::Minimum);
+	Helper.Align();
+}
+
+void SGraphEditorImpl::OnAlignCenter()
+{
+	const FScopedTransaction Transaction(FGraphEditorCommands::Get().AlignNodesCenter->GetLabel());
+
+	FAlignmentHelper Helper(SharedThis(this), Orient_Horizontal, EAlignType::Middle);
+	Helper.Align();
+}
+
+void SGraphEditorImpl::OnAlignRight()
+{
+	const FScopedTransaction Transaction(FGraphEditorCommands::Get().AlignNodesRight->GetLabel());
+
+	FAlignmentHelper Helper(SharedThis(this), Orient_Horizontal, EAlignType::Maximum);
+	Helper.Align();
+}
+
+void SGraphEditorImpl::OnStraightenConnections()
+{
+	const FScopedTransaction Transaction(FGraphEditorCommands::Get().StraightenConnections->GetLabel());
+
+	if (UEdGraphPin* Pin = GetGraphPinForMenu())
+	{
+		StraightenConnections(Pin, UEdGraphSchema_K2::GetAndResetStraightenDestinationPin());
+	}
+	else
+	{
+		StraightenConnections();
+	}
+}
+
+/** Distribute the specified array of node data evenly */
+void DistributeNodes(TArray<FAlignmentData>& InData)
+{
+	// Sort the data
+	InData.Sort([](const FAlignmentData& A, const FAlignmentData& B) {
+		return A.TargetProperty + A.TargetOffset / 2 < B.TargetProperty + B.TargetOffset / 2;
+	});
+
+	// Measure the available space
+	float TotalWidthOfNodes = 0.f;
+	for (int32 Index = 1; Index < InData.Num() - 1; ++Index)
+	{
+		TotalWidthOfNodes += InData[Index].TargetOffset;
+	}
+
+	const float SpaceToDistributeIn = InData.Last().TargetProperty - InData[0].GetTarget();
+	const float PaddingAmount = ((SpaceToDistributeIn - TotalWidthOfNodes) / (InData.Num() - 1));
+
+	float TargetPosition = InData[0].GetTarget() + PaddingAmount;
+
+	// Now set all the properties on the target
+	for (int32 Index = 1; Index < InData.Num() - 1; ++Index)
+	{
+		FAlignmentData& Entry = InData[Index];
+
+		Entry.Node->Modify();
+		Entry.TargetProperty = TargetPosition;
+
+		TargetPosition = Entry.GetTarget() + PaddingAmount;
+	}
+}
+
+void SGraphEditorImpl::OnDistributeNodesH()
+{
+	TArray<FAlignmentData> AlignData;
+	for (UObject* It : GetSelectedNodes())
+	{
+		if (UEdGraphNode* Node = Cast<UEdGraphNode>(It))
+		{
+			AlignData.Add(FAlignmentData(Node, Node->NodePosX, GetNodeSize(*this, Node).X));
+		}
+	}
+
+	if (AlignData.Num() > 2)
+	{
+		const FScopedTransaction Transaction(FGraphEditorCommands::Get().DistributeNodesHorizontally->GetLabel());
+		DistributeNodes(AlignData);
+	}
+}
+
+void SGraphEditorImpl::OnDistributeNodesV()
+{
+	TArray<FAlignmentData> AlignData;
+	for (UObject* It : GetSelectedNodes())
+	{
+		if (UEdGraphNode* Node = Cast<UEdGraphNode>(It))
+		{
+			AlignData.Add(FAlignmentData(Node, Node->NodePosY, GetNodeSize(*this, Node).Y));
+		}
+	}
+
+	if (AlignData.Num() > 2)
+	{
+		const FScopedTransaction Transaction(FGraphEditorCommands::Get().DistributeNodesVertically->GetLabel());
+		DistributeNodes(AlignData);
+	}
+}
+
+int32 SGraphEditorImpl::GetNumberOfSelectedNodes() const
+{
+	return GetSelectedNodes().Num();
+}
+
+UEdGraphNode* SGraphEditorImpl::GetSingleSelectedNode() const
+{
+	const FGraphPanelSelectionSet SelectedNodes = GetSelectedNodes();
+	return (SelectedNodes.Num() == 1) ? Cast<UEdGraphNode>(*SelectedNodes.CreateConstIterator()) : nullptr;
+}
+
 /////////////////////////////////////////////////////
 
-#undef LOCTEXT_NAMESPACE 
+
+
+#undef LOCTEXT_NAMESPACE 
\ No newline at end of file
diff --git a/Engine/Source/Editor/GraphEditor/Private/SGraphEditorImpl.h b/Engine/Source/Editor/GraphEditor/Private/SGraphEditorImpl.h
--- a/Engine/Source/Editor/GraphEditor/Private/SGraphEditorImpl.h
+++ b/Engine/Source/Editor/GraphEditor/Private/SGraphEditorImpl.h
@@ -20,6 +20,133 @@ struct FGraphContextMenuArguments;
 struct FNotificationInfo;
 struct Rect;
 
+
+/** Struct used for generically aligning nodes */
+struct FAlignmentData
+{
+	FAlignmentData(UEdGraphNode* InNode, int32& InTargetProperty, float InTargetOffset)
+		: Node(InNode), TargetProperty(InTargetProperty), TargetOffset(InTargetOffset)
+	{}
+
+	/** The node to position */
+	UEdGraphNode* Node;
+	/** The property within the node to read/write */
+	int32& TargetProperty;
+	/** The offset from the property to consider for alignment */
+	float TargetOffset;
+
+	/** Get the destination target from this alignment data (property + offset) */
+	float GetTarget() const
+	{
+		return float(TargetProperty) + TargetOffset;
+	}
+};
+
+enum class EAlignType : uint8
+{
+	Minimum, Middle, Maximum
+};
+
+/** Helper class for aligning nodes */
+struct FAlignmentHelper
+{
+	/** Construct from a graph editor, an orientation, and an alignment type */
+	FAlignmentHelper(TSharedRef<SGraphEditor> InGraphEditor, EOrientation InOrientation, EAlignType InAlignType)
+		: GraphEditor(MoveTemp(InGraphEditor))
+	{
+		// We align to the node that was clicked on, if available (not when invoked from a key shortcut)
+		CardinalNode = GraphEditor->GetGraphNodeForMenu();
+
+		Orientation = InOrientation;
+		AlignType = InAlignType;
+
+		// Collect all the alignment data for all the selected nodes
+		for (UObject* It : GraphEditor->GetSelectedNodes())
+		{
+			if (UEdGraphNode* Node = Cast<UEdGraphNode>(It))
+			{
+				AlignmentData.Add(GetAlignmentDataForNode(Node));
+			}
+		}
+
+		// Sort the data based on target - important for future algorithms
+		AlignmentData.Sort([](const FAlignmentData& A, const FAlignmentData& B) {
+			return A.GetTarget() < B.GetTarget();
+		});
+	}
+
+	/** Align all the nodes */
+	void Align()
+	{
+		if (AlignmentData.Num() > 1)
+		{
+			float Target = DetermineAlignmentTarget();
+
+			for (FAlignmentData& Entry : AlignmentData)
+			{
+				Entry.Node->Modify();
+				Entry.TargetProperty = Target - Entry.TargetOffset;
+			}
+		}
+	}
+
+private:
+
+	/** Collect alignment data for a given node, based on our settings */
+	FAlignmentData GetAlignmentDataForNode(UEdGraphNode* Node);
+
+	/** Determine the horizontal/vertical position that all nodes should align to */
+	float DetermineAlignmentTarget()
+	{
+		if (CardinalNode)
+		{
+			return GetAlignmentDataForNode(CardinalNode).GetTarget();
+		}
+
+		if (AlignType == EAlignType::Minimum)
+		{
+			float Target = TNumericLimits<float>::Max();
+			for (const FAlignmentData& Entry : AlignmentData)
+			{
+				Target = FMath::Min(Target, Entry.GetTarget());
+			}
+			return Target;
+		}
+		else if (AlignType == EAlignType::Maximum)
+		{
+			float Target = TNumericLimits<float>::Lowest();
+			for (const FAlignmentData& Entry : AlignmentData)
+			{
+				Target = FMath::Max(Target, Entry.GetTarget());
+			}
+			return Target;
+		}
+		else
+		{
+			// Use the mean
+			float SumTotal = 0.f;
+			for (const FAlignmentData& Entry : AlignmentData)
+			{
+				SumTotal += Entry.GetTarget();
+			}
+			return SumTotal / AlignmentData.Num();
+		}
+	}
+
+	/** The graph editor */
+	TSharedRef<SGraphEditor> GraphEditor;
+	/** Whether we are aligning horizontally/vertically */
+	EOrientation Orientation;
+	/** Whether we are aligning to the minimum/middle/maximum bounds */
+	EAlignType AlignType;
+	/** The cardinal node that all other nodes should align to (possibly null) */
+	UEdGraphNode* CardinalNode;
+	/** Generated alignment data */
+	TArray<FAlignmentData> AlignmentData;
+};
+
+
+
 /////////////////////////////////////////////////////
 // SGraphEditorImpl
 
@@ -159,6 +286,21 @@ public:
 	virtual void RefreshNode(UEdGraphNode& Node) override;
 	virtual void CaptureKeyboard() override;
 	virtual void SetNodeFactory(const TSharedRef<class FGraphNodeFactory>& NewNodeFactory) override;
+	virtual void OnAlignTop() override;
+	virtual void OnAlignMiddle() override;
+	virtual void OnAlignBottom() override;
+	virtual void OnAlignLeft() override;
+	virtual void OnAlignCenter() override;
+	virtual void OnAlignRight() override;
+
+	virtual void OnStraightenConnections() override;
+
+	virtual void OnDistributeNodesH() override;
+	virtual void OnDistributeNodesV() override;
+
+	virtual int32 GetNumberOfSelectedNodes() const override;
+
+	virtual UEdGraphNode* GetSingleSelectedNode() const override;
 	// End of SGraphEditor interface
 protected:
 	//
diff --git a/Engine/Source/Editor/Kismet/Private/BlueprintEditor.cpp b/Engine/Source/Editor/Kismet/Private/BlueprintEditor.cpp
--- a/Engine/Source/Editor/Kismet/Private/BlueprintEditor.cpp
+++ b/Engine/Source/Editor/Kismet/Private/BlueprintEditor.cpp
@@ -5185,163 +5185,12 @@ bool FBlueprintEditor::CanExpandNodes() const
 	return false;
 }
 
-/** Struct used for generically aligning nodes */
-struct FAlignmentData
-{
-	FAlignmentData(UEdGraphNode* InNode, int32& InTargetProperty, float InTargetOffset)
-		: Node(InNode), TargetProperty(InTargetProperty), TargetOffset(InTargetOffset)
-	{}
-
-	/** The node to position */
-	UEdGraphNode* Node;
-	/** The property within the node to read/write */
-	int32& TargetProperty;
-	/** The offset from the property to consider for alignment */
-	float TargetOffset;
-
-	/** Get the destination target from this alignment data (property + offset) */
-	float GetTarget() const
-	{
-		return float(TargetProperty) + TargetOffset;
-	}
-};
-
-enum class EAlignType : uint8
-{
-	Minimum, Middle, Maximum
-};
-
-FVector2D GetNodeSize(const SGraphEditor& GraphEditor, const UEdGraphNode* Node)
-{
-	FSlateRect Rect;
-	if (GraphEditor.GetBoundsForNode(Node, Rect, 0.f))
-	{
-		return FVector2D(Rect.Right - Rect.Left, Rect.Bottom - Rect.Top);
-	}
-
-	return FVector2D(Node->NodeWidth, Node->NodeHeight);
-}
-
-/** Helper class for aligning nodes */
-struct FAlignmentHelper
-{
-	/** Construct from a graph editor, an orientation, and an alignment type */
-	FAlignmentHelper(TSharedRef<SGraphEditor> InGraphEditor, EOrientation InOrientation, EAlignType InAlignType)
-		: GraphEditor(MoveTemp(InGraphEditor))
-	{
-		// We align to the node that was clicked on, if available (not when invoked from a key shortcut)
-		CardinalNode = GraphEditor->GetGraphNodeForMenu();
-
-		Orientation = InOrientation;
-		AlignType = InAlignType;
-
-		// Collect all the alignment data for all the selected nodes
-		for (UObject* It : GraphEditor->GetSelectedNodes())
-		{
-			if (UEdGraphNode* Node = Cast<UEdGraphNode>(It))
-			{
-				AlignmentData.Add(GetAlignmentDataForNode(Node));
-			}
-		}
-
-		// Sort the data based on target - important for future algorithsm
-		AlignmentData.Sort([](const FAlignmentData& A, const FAlignmentData& B){
-			return A.GetTarget() < B.GetTarget();
-		});
-	}
-
-	/** Align all the nodes */
-	void Align()
-	{
-		if (AlignmentData.Num() > 1)
-		{
-			float Target = DetermineAlignmentTarget();
-
-			for (FAlignmentData& Entry : AlignmentData)
-			{
-				Entry.Node->Modify();
-				Entry.TargetProperty = Target - Entry.TargetOffset;
-			}
-		}
-	}
-
-private:
-
-	/** Collect alignment data for a given node, based on our settings */
-	FAlignmentData GetAlignmentDataForNode(UEdGraphNode* Node)
-	{
-		float PropertyOffset = 0.f;
-
-		const float NodeSize = Orientation == Orient_Horizontal ? GetNodeSize(*GraphEditor, Node).X : GetNodeSize(*GraphEditor, Node).Y;
-		switch(AlignType)
-		{
-			case EAlignType::Minimum:	PropertyOffset = 0.f;			break;
-			case EAlignType::Middle:	PropertyOffset = NodeSize*.5f;	break;
-			case EAlignType::Maximum:	PropertyOffset = NodeSize;		break;
-		}
-		int32* Property = Orientation == Orient_Horizontal ? &Node->NodePosX : &Node->NodePosY;
-		return FAlignmentData(Node, *Property, PropertyOffset);
-	}
-
-	/** Determine the horizontal/vertical position that all nodes should align to */
-	float DetermineAlignmentTarget()
-	{
-		if (CardinalNode)
-		{
-			return GetAlignmentDataForNode(CardinalNode).GetTarget();
-		}
-
-		if (AlignType == EAlignType::Minimum)
-		{
-			float Target = TNumericLimits<float>::Max();
-			for (const FAlignmentData& Entry : AlignmentData)
-			{
-				Target = FMath::Min(Target, Entry.GetTarget());
-			}
-			return Target;
-		}
-		else if (AlignType == EAlignType::Maximum)
-		{
-			float Target = TNumericLimits<float>::Lowest();
-			for (const FAlignmentData& Entry : AlignmentData)
-			{
-				Target = FMath::Max(Target, Entry.GetTarget());
-			}
-			return Target;
-		}
-		else
-		{
-			// Use the mean
-			float SumTotal = 0.f;
-			for (const FAlignmentData& Entry : AlignmentData)
-			{
-				SumTotal += Entry.GetTarget();
-			}
-			return SumTotal / AlignmentData.Num();
-		}
-	}
-
-	/** The graph editor */
-	TSharedRef<SGraphEditor> GraphEditor;
-	/** Whether we are aligning horizontally/vertically */
-	EOrientation Orientation;
-	/** Whether we are aligning to the minimum/middle/maximum bounds */
-	EAlignType AlignType;
-	/** The cardinal node that all other nodes should align to (possibly null) */
-	UEdGraphNode* CardinalNode;
-	/** Generated alignment data */
-	TArray<FAlignmentData> AlignmentData;
-};
-
 void FBlueprintEditor::OnAlignTop()
 {
 	TSharedPtr<SGraphEditor> FocusedGraphEd = FocusedGraphEdPtr.Pin();
 	if (FocusedGraphEd.IsValid())
 	{
-		const FScopedTransaction Transaction(FGraphEditorCommands::Get().AlignNodesTop->GetLabel());
-
-		FAlignmentHelper Helper(FocusedGraphEd.ToSharedRef(), Orient_Vertical, EAlignType::Minimum);
-		Helper.Align();
+		FocusedGraphEd->OnAlignTop();
 	}
 }
 
@@ -5350,10 +5199,7 @@ void FBlueprintEditor::OnAlignMiddle()
 	TSharedPtr<SGraphEditor> FocusedGraphEd = FocusedGraphEdPtr.Pin();
 	if (FocusedGraphEd.IsValid())
 	{
-		const FScopedTransaction Transaction(FGraphEditorCommands::Get().AlignNodesMiddle->GetLabel());
-
-		FAlignmentHelper Helper(FocusedGraphEd.ToSharedRef(), Orient_Vertical, EAlignType::Middle);
-		Helper.Align();
+		FocusedGraphEd->OnAlignMiddle();
 	}
 }
 
@@ -5362,10 +5208,7 @@ void FBlueprintEditor::OnAlignBottom()
 	TSharedPtr<SGraphEditor> FocusedGraphEd = FocusedGraphEdPtr.Pin();
 	if (FocusedGraphEd.IsValid())
 	{
-		const FScopedTransaction Transaction(FGraphEditorCommands::Get().AlignNodesBottom->GetLabel());
-
-		FAlignmentHelper Helper(FocusedGraphEd.ToSharedRef(), Orient_Vertical, EAlignType::Maximum);
-		Helper.Align();
+		FocusedGraphEd->OnAlignMiddle();
 	}
 }
 
@@ -5374,10 +5217,7 @@ void FBlueprintEditor::OnAlignLeft()
 	TSharedPtr<SGraphEditor> FocusedGraphEd = FocusedGraphEdPtr.Pin();
 	if (FocusedGraphEd.IsValid())
 	{
-		const FScopedTransaction Transaction(FGraphEditorCommands::Get().AlignNodesLeft->GetLabel());
-
-		FAlignmentHelper Helper(FocusedGraphEd.ToSharedRef(), Orient_Horizontal, EAlignType::Minimum);
-		Helper.Align();
+		FocusedGraphEd->OnAlignLeft();
 	}
 }
 
@@ -5386,10 +5226,7 @@ void FBlueprintEditor::OnAlignCenter()
 	TSharedPtr<SGraphEditor> FocusedGraphEd = FocusedGraphEdPtr.Pin();
 	if (FocusedGraphEd.IsValid())
 	{
-		const FScopedTransaction Transaction(FGraphEditorCommands::Get().AlignNodesCenter->GetLabel());
-
-		FAlignmentHelper Helper(FocusedGraphEd.ToSharedRef(), Orient_Horizontal, EAlignType::Middle);
-		Helper.Align();
+		FocusedGraphEd->OnAlignCenter();
 	}
 }
 
@@ -5398,10 +5235,7 @@ void FBlueprintEditor::OnAlignRight()
 	TSharedPtr<SGraphEditor> FocusedGraphEd = FocusedGraphEdPtr.Pin();
 	if (FocusedGraphEd.IsValid())
 	{
-		const FScopedTransaction Transaction(FGraphEditorCommands::Get().AlignNodesRight->GetLabel());
-
-		FAlignmentHelper Helper(FocusedGraphEd.ToSharedRef(), Orient_Horizontal, EAlignType::Maximum);
-		Helper.Align();
+		FocusedGraphEd->OnAlignRight();
 	}
 }
 
@@ -5410,96 +5244,25 @@ void FBlueprintEditor::OnStraightenConnections()
 	TSharedPtr<SGraphEditor> FocusedGraphEd = FocusedGraphEdPtr.Pin();
 	if (FocusedGraphEd.IsValid())
 	{
-		const FScopedTransaction Transaction(FGraphEditorCommands::Get().StraightenConnections->GetLabel());
-
-		if (UEdGraphPin* Pin = FocusedGraphEd->GetGraphPinForMenu())
-		{
-			FocusedGraphEd->StraightenConnections(Pin, UEdGraphSchema_K2::GetAndResetStraightenDestinationPin());
-		}
-		else
-		{
-			FocusedGraphEd->StraightenConnections();
-		}
-	}
-}
-
-/** Distribute the specified array of node data evenly */
-void DistributeNodes(TArray<FAlignmentData>& InData)
-{
-	// Sort the data
-	InData.Sort([](const FAlignmentData& A, const FAlignmentData& B){
-		return A.TargetProperty + A.TargetOffset/2 < B.TargetProperty + B.TargetOffset/2;
-	});
-
-	// Measure the available space
-	float TotalWidthOfNodes = 0.f;
-	for (int32 Index = 1; Index < InData.Num() - 1; ++Index)
-	{
-		TotalWidthOfNodes += InData[Index].TargetOffset;
-	}
-
-	const float SpaceToDistributeIn = InData.Last().TargetProperty - InData[0].GetTarget();
-	const float PaddingAmount = ((SpaceToDistributeIn - TotalWidthOfNodes) / (InData.Num() - 1));
-
-	float TargetPosition = InData[0].GetTarget() + PaddingAmount;
-
-	// Now set all the properties on the target
-	for (int32 Index = 1; Index < InData.Num() - 1; ++Index)
-	{
-		FAlignmentData& Entry = InData[Index];
-
-		Entry.Node->Modify();
-		Entry.TargetProperty = TargetPosition;
-
-		TargetPosition = Entry.GetTarget() + PaddingAmount;
+		FocusedGraphEd->OnStraightenConnections();
 	}
 }
 
 void FBlueprintEditor::OnDistributeNodesH()
 {
 	TSharedPtr<SGraphEditor> FocusedGraphEd = FocusedGraphEdPtr.Pin();
-	if (!FocusedGraphEd.IsValid())
-	{
-		return;
-	}
-	
-	TArray<FAlignmentData> AlignData;
-	for (UObject* It : GetSelectedNodes())
-	{
-		if (UEdGraphNode* Node = Cast<UEdGraphNode>(It))
-		{
-			AlignData.Add(FAlignmentData(Node, Node->NodePosX, GetNodeSize(*FocusedGraphEd, Node).X));
-		}
-	}
-
-	if (AlignData.Num() > 2)
+	if (FocusedGraphEd.IsValid())
 	{
-		const FScopedTransaction Transaction(FGraphEditorCommands::Get().DistributeNodesHorizontally->GetLabel());
-		DistributeNodes(AlignData);
+		FocusedGraphEd->OnDistributeNodesH();
 	}
 }
 
 void FBlueprintEditor::OnDistributeNodesV()
 {
 	TSharedPtr<SGraphEditor> FocusedGraphEd = FocusedGraphEdPtr.Pin();
-	if (!FocusedGraphEd.IsValid())
-	{
-		return;
-	}
-	
-	TArray<FAlignmentData> AlignData;
-	for (UObject* It : GetSelectedNodes())
-	{
-		if (UEdGraphNode* Node = Cast<UEdGraphNode>(It))
-		{
-			AlignData.Add(FAlignmentData(Node, Node->NodePosY, GetNodeSize(*FocusedGraphEd, Node).Y));
-		}
-	}
-
-	if (AlignData.Num() > 2)
+	if (FocusedGraphEd.IsValid())
 	{
-		const FScopedTransaction Transaction(FGraphEditorCommands::Get().DistributeNodesVertically->GetLabel());
-		DistributeNodes(AlignData);
+		FocusedGraphEd->OnDistributeNodesV();
 	}
 }
 
diff --git a/Engine/Source/Editor/MaterialEditor/Private/MaterialEditor.cpp b/Engine/Source/Editor/MaterialEditor/Private/MaterialEditor.cpp
--- a/Engine/Source/Editor/MaterialEditor/Private/MaterialEditor.cpp
+++ b/Engine/Source/Editor/MaterialEditor/Private/MaterialEditor.cpp
@@ -3964,6 +3964,79 @@ void FMaterialEditor::RedoGraphAction()
 
 }
 
+void FMaterialEditor::OnAlignTop()
+{
+	if (GraphEditor.IsValid())
+	{
+		GraphEditor->OnAlignTop();
+	}
+}
+
+void FMaterialEditor::OnAlignMiddle()
+{
+	if (GraphEditor.IsValid())
+	{
+		GraphEditor->OnAlignMiddle();
+	}
+}
+
+void FMaterialEditor::OnAlignBottom()
+{
+	if (GraphEditor.IsValid())
+	{
+		GraphEditor->OnAlignMiddle();
+	}
+}
+
+void FMaterialEditor::OnAlignLeft()
+{
+	if (GraphEditor.IsValid())
+	{
+		GraphEditor->OnAlignLeft();
+	}
+}
+
+void FMaterialEditor::OnAlignCenter()
+{
+	if (GraphEditor.IsValid())
+	{
+		GraphEditor->OnAlignCenter();
+	}
+}
+
+void FMaterialEditor::OnAlignRight()
+{
+	if (GraphEditor.IsValid())
+	{
+		GraphEditor->OnAlignRight();
+	}
+}
+
+void FMaterialEditor::OnStraightenConnections()
+{
+	if (GraphEditor.IsValid())
+	{
+		GraphEditor->OnStraightenConnections();
+	}
+}
+
+void FMaterialEditor::OnDistributeNodesH()
+{
+	if (GraphEditor.IsValid())
+	{
+		GraphEditor->OnDistributeNodesH();
+	}
+}
+
+void FMaterialEditor::OnDistributeNodesV()
+{
+	if (GraphEditor.IsValid())
+	{
+		GraphEditor->OnDistributeNodesV();
+	}
+}
+
+
 void FMaterialEditor::PostUndo(bool bSuccess)
 {
 	if (bSuccess)
@@ -4364,6 +4437,44 @@ TSharedRef<SGraphEditor> FMaterialEditor::CreateGraphEditorWidget()
 			FCanExecuteAction::CreateSP(this, &FMaterialEditor::OnCanPromoteToParameter)
 			);
 
+		// Alignment Commands
+		GraphEditorCommands->MapAction(FGraphEditorCommands::Get().AlignNodesTop,
+			FExecuteAction::CreateSP(this, &FMaterialEditor::OnAlignTop)
+		);
+
+		GraphEditorCommands->MapAction(FGraphEditorCommands::Get().AlignNodesMiddle,
+			FExecuteAction::CreateSP(this, &FMaterialEditor::OnAlignMiddle)
+		);
+
+		GraphEditorCommands->MapAction(FGraphEditorCommands::Get().AlignNodesBottom,
+			FExecuteAction::CreateSP(this, &FMaterialEditor::OnAlignBottom)
+		);
+
+		GraphEditorCommands->MapAction(FGraphEditorCommands::Get().AlignNodesLeft,
+			FExecuteAction::CreateSP(this, &FMaterialEditor::OnAlignLeft)
+		);
+
+		GraphEditorCommands->MapAction(FGraphEditorCommands::Get().AlignNodesCenter,
+			FExecuteAction::CreateSP(this, &FMaterialEditor::OnAlignCenter)
+		);
+
+		GraphEditorCommands->MapAction(FGraphEditorCommands::Get().AlignNodesRight,
+			FExecuteAction::CreateSP(this, &FMaterialEditor::OnAlignRight)
+		);
+
+		GraphEditorCommands->MapAction(FGraphEditorCommands::Get().StraightenConnections,
+			FExecuteAction::CreateSP(this, &FMaterialEditor::OnStraightenConnections)
+		);
+
+		// Distribution Commands
+		GraphEditorCommands->MapAction(FGraphEditorCommands::Get().DistributeNodesHorizontally,
+			FExecuteAction::CreateSP(this, &FMaterialEditor::OnDistributeNodesH)
+		);
+
+		GraphEditorCommands->MapAction(FGraphEditorCommands::Get().DistributeNodesVertically,
+			FExecuteAction::CreateSP(this, &FMaterialEditor::OnDistributeNodesV)
+		);
+
 	}
 
 	SGraphEditor::FGraphEditorEvents InEvents;
diff --git a/Engine/Source/Editor/MaterialEditor/Private/MaterialEditor.h b/Engine/Source/Editor/MaterialEditor/Private/MaterialEditor.h
--- a/Engine/Source/Editor/MaterialEditor/Private/MaterialEditor.h
+++ b/Engine/Source/Editor/MaterialEditor/Private/MaterialEditor.h
@@ -500,6 +500,18 @@ protected:
 	/** Called to redo the last undone action */
 	void RedoGraphAction();
 
+	void OnAlignTop();
+	void OnAlignMiddle();
+	void OnAlignBottom();
+	void OnAlignLeft();
+	void OnAlignCenter();
+	void OnAlignRight();
+
+	void OnStraightenConnections();
+
+	void OnDistributeNodesH();
+	void OnDistributeNodesV();
+
 private:
 	/** Builds the toolbar widget for the material editor */
 	void ExtendToolbar();
diff --git a/Engine/Source/Editor/UnrealEd/Public/GraphEditor.h b/Engine/Source/Editor/UnrealEd/Public/GraphEditor.h
--- a/Engine/Source/Editor/UnrealEd/Public/GraphEditor.h
+++ b/Engine/Source/Editor/UnrealEd/Public/GraphEditor.h
@@ -458,6 +458,102 @@ public:
 		}
 	}
 
+	virtual void OnAlignTop()
+	{
+		if (Implementation.IsValid())
+		{
+			Implementation->OnAlignTop();
+		}
+	}
+
+	virtual void OnAlignMiddle()
+	{
+		if (Implementation.IsValid())
+		{
+			Implementation->OnAlignMiddle();
+		}
+	}
+
+	virtual void OnAlignBottom()
+	{
+		if (Implementation.IsValid())
+		{
+			Implementation->OnAlignBottom();
+		}
+	}
+
+	virtual void OnAlignLeft()
+	{
+		if (Implementation.IsValid())
+		{
+			Implementation->OnAlignLeft();
+		}
+	}
+
+	virtual void OnAlignCenter()
+	{
+		if (Implementation.IsValid())
+		{
+			Implementation->OnAlignCenter();
+		}
+	}
+
+	virtual void OnAlignRight()
+	{
+		if (Implementation.IsValid())
+		{
+			Implementation->OnAlignRight();
+		}
+	}
+
+
+	virtual void OnStraightenConnections()
+	{
+		if (Implementation.IsValid())
+		{
+			Implementation->OnStraightenConnections();
+		}
+	}
+
+
+	virtual void OnDistributeNodesH()
+	{
+		if (Implementation.IsValid())
+		{
+			Implementation->OnDistributeNodesH();
+		}
+	}
+
+	virtual void OnDistributeNodesV()
+	{
+		if (Implementation.IsValid())
+		{
+			Implementation->OnDistributeNodesV();
+		}
+	}
+
+
+	virtual int32 GetNumberOfSelectedNodes() const
+	{
+		if (Implementation.IsValid())
+		{
+			Implementation->GetNumberOfSelectedNodes();
+		}
+		return 0;
+	}
+
+
+	/** Returns the currently selected node if there is a single node selected (if there are multiple nodes selected or none selected, it will return nullptr) */
+	virtual UEdGraphNode* GetSingleSelectedNode() const
+	{
+		if (Implementation.IsValid())
+		{
+			Implementation->GetSingleSelectedNode();
+		}
+		return nullptr;
+	}
+
+
 protected:
 	/** Invoked when the underlying Graph is being changed. */
 	virtual void OnGraphChanged(const struct FEdGraphEditAction& InAction)
