Remove last AutoHeight from Debug window to re-enable scrolling

diff --git a/Engine/Source/Editor/Kismet/Private/Debugging/SKismetDebuggingView.cpp b/Engine/Source/Editor/Kismet/Private/Debugging/SKismetDebuggingView.cpp
--- a/Engine/Source/Editor/Kismet/Private/Debugging/SKismetDebuggingView.cpp
+++ b/Engine/Source/Editor/Kismet/Private/Debugging/SKismetDebuggingView.cpp
@@ -1085,7 +1085,6 @@ void SKismetDebuggingView::Construct(const FArguments& InArgs)
 			.Text( this, &SKismetDebuggingView::GetTopText )
 		]
 		+SVerticalBox::Slot()
-		.AutoHeight()
 		[
 			SAssignNew( DebugTreeView, STreeView< FDebugTreeItemPtr > )
 			.TreeItemsSource( &RootTreeItems )
