Use asset picker when selecting external timeline curves
Fixed a bug where a new track may not be created due to duplicate name

diff --git a/Engine/Source/Editor/Kismet/Private/STimelineEditor.cpp b/Engine/Source/Editor/Kismet/Private/STimelineEditor.cpp
--- a/Engine/Source/Editor/Kismet/Private/STimelineEditor.cpp
+++ b/Engine/Source/Editor/Kismet/Private/STimelineEditor.cpp
@@ -21,7 +21,7 @@
 #include "K2Node_Timeline.h"
 #include "ScopedTransaction.h"
 #include "Kismet2/BlueprintEditorUtils.h"
-
+#include "Editor/PropertyEditor/Public/PropertyCustomizationHelpers.h"
 
 #include "BlueprintEditor.h"
 #include "AssetRegistryModule.h"
@@ -33,6 +33,7 @@
 #include "Framework/Notifications/NotificationManager.h"
 #include "Widgets/Notifications/SNotificationList.h"
 #include "Engine/Selection.h"
+#include "AssetData.h"
 
 #define LOCTEXT_NAMESPACE "STimelineEditor"
 
@@ -58,6 +59,22 @@ namespace TimelineEditorHelpers
 		}
 		return NAME_None;
 	}
+
+	TSubclassOf<UCurveBase> TrackTypeToAllowedClass(FTimelineEdTrack::ETrackType TrackType)
+	{
+		switch (TrackType)
+		{
+		case FTimelineEdTrack::TT_Event:
+		case FTimelineEdTrack::TT_FloatInterp:
+			return UCurveFloat::StaticClass();
+		case FTimelineEdTrack::TT_VectorInterp:
+			return UCurveVector::StaticClass();
+		case FTimelineEdTrack::TT_LinearColorInterp:
+			return UCurveLinearColor::StaticClass();
+		default:
+			return UCurveBase::StaticClass();
+		}
+	}
 }
 
 //////////////////////////////////////////////////////////////////////////
@@ -177,7 +194,6 @@ void STimelineEdTrack::Construct(const FArguments& InArgs, TSharedPtr<FTimelineE
 				.AutoWidth()
 				[
 					SNew(SVerticalBox)
-
 					// External Curve Label
 					+SVerticalBox::Slot()
 					.AutoHeight()
@@ -190,71 +206,35 @@ void STimelineEdTrack::Construct(const FArguments& InArgs, TSharedPtr<FTimelineE
 					// External Curve Controls
 					+SVerticalBox::Slot()
 					.AutoHeight()
-					.Padding(2, 0, 2, 4)
+					.Padding(2, 0, 0, 4)
 					[
-						SNew(SHorizontalBox)
-						+SHorizontalBox::Slot()
-						.Padding(0, 0, 1, 0)
-						.FillWidth(1)
-						[
-							//External curve name display box
-							SNew(SEditableTextBox)
-							.Text( this, &STimelineEdTrack::GetExternalCurveName )
-							.ForegroundColor( FLinearColor::Black )
-							.IsReadOnly(true)
-							.ToolTipText(this, &STimelineEdTrack::GetExternalCurvePath )
-							.MinDesiredWidth( 80 )
-							.BackgroundColor( FLinearColor::White )
-						]
-
-						// Use external curve button
-						+SHorizontalBox::Slot()
-						.AutoWidth()
-						.Padding(1,0)
-						.VAlign(VAlign_Center)
+						SNew(SBox)
+						.WidthOverride(170)
 						[
-							SNew(SButton)
-							.ButtonStyle( FEditorStyle::Get(), "NoBorder" )
-							.OnClicked(this, &STimelineEdTrack::OnClickUse)
-							.ContentPadding(1.f)
-							.ToolTipText(NSLOCTEXT("TimelineEdTrack", "TimelineEdTrack_Use", "Use External Curve"))
+							SNew(SHorizontalBox)
+							+SHorizontalBox::Slot()
+							.FillWidth(1)
 							[
-								SNew(SImage)
-								.Image( FEditorStyle::GetBrush(TEXT("PropertyWindow.Button_Use")) )
+								SNew(SObjectPropertyEntryBox)
+								.AllowedClass(TimelineEditorHelpers::TrackTypeToAllowedClass(Track->TrackType))
+								.ObjectPath(this, &STimelineEdTrack::GetExternalCurvePath)
+								.OnObjectChanged(FOnSetObject::CreateSP(this, &STimelineEdTrack::OnChooseCurve))
 							]
-						]
 
-						// Browse external curve button
-						+SHorizontalBox::Slot()
-						.AutoWidth()
-						.Padding(1,0)
-						.VAlign(VAlign_Center)
-						[
-							SNew(SButton)
-							.ButtonStyle( FEditorStyle::Get(), "NoBorder" )
-							.OnClicked(this, &STimelineEdTrack::OnClickBrowse)
-							.ContentPadding(0)
-							.ToolTipText(NSLOCTEXT("TimelineEdTrack", "TimelineEdTrack_Browse", "Browse External Curve"))
+							// Convert to internal curve button
+							+SHorizontalBox::Slot()
+							.AutoWidth()
+							.VAlign(VAlign_Center)
 							[
-								SNew(SImage)
-								.Image( FEditorStyle::GetBrush(TEXT("PropertyWindow.Button_Browse")) )
-							]
-						]
-
-						// Convert to internal curve button
-						+SHorizontalBox::Slot()
-						.AutoWidth()
-						.Padding(1,0)
-						.VAlign(VAlign_Center)
-						[
-							SNew(SButton)
-							.ButtonStyle( FEditorStyle::Get(), "NoBorder" )
-							.OnClicked(this, &STimelineEdTrack::OnClickClear)
-							.ContentPadding(1.f)
-							.ToolTipText(NSLOCTEXT("TimelineEdTrack", "TimelineEdTrack_Clear", "Convert to Internal Curve"))
-							[
-								SNew(SImage)
-								.Image( FEditorStyle::GetBrush(TEXT("PropertyWindow.Button_Clear")) )
+								SNew(SButton)
+								.ButtonStyle( FEditorStyle::Get(), "NoBorder" )
+								.OnClicked(this, &STimelineEdTrack::OnClickClear)
+								.ContentPadding(1.f)
+								.ToolTipText(NSLOCTEXT("TimelineEdTrack", "TimelineEdTrack_Clear", "Convert to Internal Curve"))
+								[
+									SNew(SImage)
+									.Image( FEditorStyle::GetBrush(TEXT("PropertyWindow.Button_Clear")) )
+								]
 							]
 						]
 					]
@@ -423,27 +403,13 @@ void STimelineEdTrack::SwitchToExternalCurve(UCurveBase* AssetCurvePtr)
 
 void STimelineEdTrack::UseExternalCurve( UObject* AssetObj )
 {
-	ResetExternalCurveInfo();
-
-	if( AssetObj )
+	if (AssetObj)
 	{
-		ExternalCurveName = AssetObj->GetName();
-		ExternalCurvePath =  AssetObj->GetFullName();
-
-		int32 StringLen = ExternalCurveName.Len();
-
-		//If string is too long, then truncate (eg. "abcdefgijklmnopq" is converted as "abcd...nopq")
-		const int32 MaxAllowedLength = 12;
-		if( StringLen > MaxAllowedLength )
-		{
-			//Take first 4 characters
-			FString TruncatedStr(ExternalCurveName.Left(4));
-			TruncatedStr += FString( TEXT("..."));
-
-			//Take last 4 characters
-			TruncatedStr += ExternalCurveName.Right(4);
-			ExternalCurveName = TruncatedStr;
-		}
+		ExternalCurvePath = AssetObj->GetPathName();
+	}
+	else
+	{
+		ResetExternalCurveInfo();
 	}
 }
 
@@ -555,37 +521,22 @@ FReply STimelineEdTrack::OnClickClear()
 	return FReply::Handled();
 }
 
-FReply STimelineEdTrack::OnClickUse()
+void STimelineEdTrack::OnChooseCurve(const FAssetData& InObject)
 {
-	FEditorDelegates::LoadSelectedAssetsIfNeeded.Broadcast();
-
-	UCurveBase* SelectedObj = GEditor->GetSelectedObjects()->GetTop<UCurveBase>();
-	if( SelectedObj )
+	UCurveBase* SelectedObj = Cast<UCurveBase>(InObject.GetAsset());
+	if (SelectedObj)
 	{
 		SwitchToExternalCurve(SelectedObj);
 	}
-	return FReply::Handled();
-}
-
-FReply STimelineEdTrack::OnClickBrowse()
-{
-	if(CurveBasePtr)
+	else
 	{
-		TArray<UObject*> Objects;
-		Objects.Add(CurveBasePtr);
-		GEditor->SyncBrowserToObjects(Objects);
+		UseInternalCurve();
 	}
-	return FReply::Handled();
 }
 
-FText STimelineEdTrack::GetExternalCurveName( ) const
+FString STimelineEdTrack::GetExternalCurvePath( ) const
 {
-	return FText::FromString(ExternalCurveName);
-}
-
-FText STimelineEdTrack::GetExternalCurvePath( ) const
-{
-	return FText::FromString(ExternalCurvePath);
+	return ExternalCurvePath;
 }
 
 UCurveBase* STimelineEdTrack::CreateCurveAsset()
@@ -792,7 +743,6 @@ void STimelineEdTrack::OnSetOutputViewRange(float Min, float Max)
 
 void STimelineEdTrack::ResetExternalCurveInfo( )
 {
-	ExternalCurveName = FString( TEXT( "None" ) );
 	ExternalCurvePath = FString( TEXT( "None" ) );
 }
 
@@ -1111,7 +1061,13 @@ TSharedRef<ITableRow> STimelineEditor::MakeTrackWidget( TSharedPtr<FTimelineEdTr
 
 FReply STimelineEditor::CreateNewTrack(FTimelineEdTrack::ETrackType Type)
 {
-	FName TrackName = MakeUniqueObjectName(TimelineObj, UTimelineTemplate::StaticClass(), FName(*(LOCTEXT("NewTrack_DefaultName", "NewTrack").ToString())));
+	FName TrackName;
+	do
+	{
+		// MakeUniqueObjectName is misleading here since tracks aren't UObjects, although the function
+		// will still keep a counter for tracks. This may take a couple tries to find a valid name.
+		TrackName = MakeUniqueObjectName(TimelineObj, UTimelineTemplate::StaticClass(), FName(*(LOCTEXT("NewTrack_DefaultName", "NewTrack").ToString())));
+	} while (!TimelineObj->IsNewTrackNameValid(TrackName));
 		
 	TSharedPtr<FBlueprintEditor> Kismet2 = Kismet2Ptr.Pin();
 	UBlueprint* Blueprint = Kismet2->GetBlueprintObj();
@@ -1121,67 +1077,57 @@ FReply STimelineEditor::CreateNewTrack(FTimelineEdTrack::ETrackType Type)
 
 	FText ErrorMessage;
 
-	if(TimelineObj->IsNewTrackNameValid(TrackName))
+	if (TimelineNode)
 	{
-		if (TimelineNode)
-		{
-			const FScopedTransaction Transaction( LOCTEXT( "TimelineEditor_AddNewTrack", "Add new track" ) );
+		const FScopedTransaction Transaction( LOCTEXT( "TimelineEditor_AddNewTrack", "Add new track" ) );
 
-			TimelineNode->Modify();
-			TimelineObj->Modify();
+		TimelineNode->Modify();
+		TimelineObj->Modify();
 
-			NewTrackPendingRename = TrackName;
-			if(Type == FTimelineEdTrack::TT_Event)
-			{
-				FTTEventTrack NewTrack;
-				NewTrack.SetTrackName(TrackName, TimelineObj);
-				NewTrack.CurveKeys = NewObject<UCurveFloat>(OwnerClass, NAME_None, RF_Public); // Needs to be marked public so that it can be referenced from timeline instances in the level
-				NewTrack.CurveKeys->bIsEventCurve = true;
-				TimelineObj->EventTracks.Add(NewTrack);
-			}
-			else if(Type == FTimelineEdTrack::TT_FloatInterp)
-			{
-				FTTFloatTrack NewTrack;
-				NewTrack.SetTrackName(TrackName, TimelineObj);
-				// @hack for using existing curve assets.  need something better!
-				NewTrack.CurveFloat = FindObject<UCurveFloat>(ANY_PACKAGE, *TrackName.ToString() );
-				if (NewTrack.CurveFloat == nullptr)
-				{
-					NewTrack.CurveFloat = NewObject<UCurveFloat>(OwnerClass, NAME_None, RF_Public);
-				}
-				TimelineObj->FloatTracks.Add(NewTrack);
-			}
-			else if(Type == FTimelineEdTrack::TT_VectorInterp)
-			{
-				FTTVectorTrack NewTrack;
-				NewTrack.SetTrackName(TrackName, TimelineObj);
-				NewTrack.CurveVector = NewObject<UCurveVector>(OwnerClass, NAME_None, RF_Public);
-				TimelineObj->VectorTracks.Add(NewTrack);
-			}
-			else if(Type == FTimelineEdTrack::TT_LinearColorInterp)
+		NewTrackPendingRename = TrackName;
+		if(Type == FTimelineEdTrack::TT_Event)
+		{
+			FTTEventTrack NewTrack;
+			NewTrack.SetTrackName(TrackName, TimelineObj);
+			NewTrack.CurveKeys = NewObject<UCurveFloat>(OwnerClass, NAME_None, RF_Public); // Needs to be marked public so that it can be referenced from timeline instances in the level
+			NewTrack.CurveKeys->bIsEventCurve = true;
+			TimelineObj->EventTracks.Add(NewTrack);
+		}
+		else if(Type == FTimelineEdTrack::TT_FloatInterp)
+		{
+			FTTFloatTrack NewTrack;
+			NewTrack.SetTrackName(TrackName, TimelineObj);
+			// @hack for using existing curve assets.  need something better!
+			NewTrack.CurveFloat = FindObject<UCurveFloat>(ANY_PACKAGE, *TrackName.ToString() );
+			if (NewTrack.CurveFloat == nullptr)
 			{
-				FTTLinearColorTrack NewTrack;
-				NewTrack.SetTrackName(TrackName, TimelineObj);
-				NewTrack.CurveLinearColor = NewObject<UCurveLinearColor>(OwnerClass, NAME_None, RF_Public);
-				TimelineObj->LinearColorTracks.Add(NewTrack);
+				NewTrack.CurveFloat = NewObject<UCurveFloat>(OwnerClass, NAME_None, RF_Public);
 			}
-
-			// Refresh the node that owns this timeline template to get new pin
-			TimelineNode->ReconstructNode();
-			Kismet2->RefreshEditors();
+			TimelineObj->FloatTracks.Add(NewTrack);
+		}
+		else if(Type == FTimelineEdTrack::TT_VectorInterp)
+		{
+			FTTVectorTrack NewTrack;
+			NewTrack.SetTrackName(TrackName, TimelineObj);
+			NewTrack.CurveVector = NewObject<UCurveVector>(OwnerClass, NAME_None, RF_Public);
+			TimelineObj->VectorTracks.Add(NewTrack);
 		}
-		else
+		else if(Type == FTimelineEdTrack::TT_LinearColorInterp)
 		{
-			// invalid node for timeline
-			ErrorMessage = LOCTEXT( "InvalidTimelineNodeCreate","Failed to create track. Timeline node is invalid. Please remove timeline node." );
+			FTTLinearColorTrack NewTrack;
+			NewTrack.SetTrackName(TrackName, TimelineObj);
+			NewTrack.CurveLinearColor = NewObject<UCurveLinearColor>(OwnerClass, NAME_None, RF_Public);
+			TimelineObj->LinearColorTracks.Add(NewTrack);
 		}
+
+		// Refresh the node that owns this timeline template to get new pin
+		TimelineNode->ReconstructNode();
+		Kismet2->RefreshEditors();
 	}
 	else
 	{
-		//name is in use 
-		FFormatNamedArguments Args;
-		Args.Add( TEXT("TrackName"), FText::FromName( TrackName ) );
-		ErrorMessage = FText::Format( LOCTEXT( "DupTrackName","Failed to create track. Duplicate Track name entered. \n\"{TrackName}\" is already in use" ), Args );
+		// invalid node for timeline
+		ErrorMessage = LOCTEXT( "InvalidTimelineNodeCreate","Failed to create track. Timeline node is invalid. Please remove timeline node." );
 	}
 
 	if (!ErrorMessage.IsEmpty())
diff --git a/Engine/Source/Editor/Kismet/Private/STimelineEditor.h b/Engine/Source/Editor/Kismet/Private/STimelineEditor.h
--- a/Engine/Source/Editor/Kismet/Private/STimelineEditor.h
+++ b/Engine/Source/Editor/Kismet/Private/STimelineEditor.h
@@ -21,6 +21,7 @@ class STimelineEditor;
 class SWindow;
 class UTimelineTemplate;
 struct FTTTrackBase;
+struct FAssetData;
 
 //////////////////////////////////////////////////////////////////////////
 // FTimelineEdTrack
@@ -97,9 +98,6 @@ private:
 	/**	Pointer to the curve */
 	UCurveBase* CurveBasePtr;
 
-	/** String to display external curve name in the text box*/
-	FString ExternalCurveName;
-
 	/** String to display external curve path as tooltip*/
 	FString ExternalCurvePath;
 
@@ -127,11 +125,8 @@ private:
 	/** Creates default asset path*/
 	FString CreateUniqueCurveAssetPathName();
 
-	/** Get the current external curve name*/
-	FText GetExternalCurveName( ) const;
-
 	/** Get the current external curve path*/
-	FText GetExternalCurvePath( ) const;
+	FString GetExternalCurvePath( ) const;
 
 	/** Function to replace internal curve with an external curve*/
 	void SwitchToExternalCurve(UCurveBase* AssetCurvePtr);
@@ -142,15 +137,9 @@ private:
 	/** Function to convert external curve to an internal curve*/
 	void UseInternalCurve( );
 
-	/** Callback function to replace internal curve with an external curve*/
-	FReply OnClickUse();
-
 	/** Callback function to replace external curve with an internal curve*/
 	FReply OnClickClear();
 
-	/** Callback function to locate external curve asset in content browser */
-	FReply OnClickBrowse();
-
 	/**Function to reset external curve info*/
 	void ResetExternalCurveInfo( );
 
@@ -184,6 +173,8 @@ private:
 	void OnSetInputViewRange(float Min, float Max);
 	/** Callback to set the output view range for the curve editor. */
 	void OnSetOutputViewRange(float Min, float Max);
+	/** Callback when the user picks a curve from the asset picker for a track */
+	void OnChooseCurve(const FAssetData& InObject);
 
 public:
 	/** Inline block for changing name of track */
