Slight improvement to editor startup/load times by deferring Blueprint search data re-indexing work on load.

Notes:
- Blueprint asset re-indexing is now queued and occurs only when a Blueprint editor window OR any global Find-in-Blueprint tab is open.
- Includes an editor INI option (bDisableDeferredIndexing) to allow the new UX to be disabled if necessary (i.e. revert to previous behavior).

diff --git a/Engine/Config/BaseEditor.ini b/Engine/Config/BaseEditor.ini
--- a/Engine/Config/BaseEditor.ini
+++ b/Engine/Config/BaseEditor.ini
@@ -438,6 +438,9 @@ SpeedUnits=EUnit::MetersPerSecond
 TemperatureUnits=EUnit::Celsius
 ForceUnits=EUnit::Newtons
 
+[BlueprintSearchSettings]
+bDisableDeferredIndexing=false
+
 [BlueprintNativizationSettings]
 bUseStaticVariablesInWrappers=true
 bUseStaticVariablesInClasses=true
diff --git a/Engine/Source/Editor/Kismet/Private/FindInBlueprintManager.cpp b/Engine/Source/Editor/Kismet/Private/FindInBlueprintManager.cpp
--- a/Engine/Source/Editor/Kismet/Private/FindInBlueprintManager.cpp
+++ b/Engine/Source/Editor/Kismet/Private/FindInBlueprintManager.cpp
@@ -31,6 +31,7 @@
 #include "EdGraphSchema_K2.h"
 #include "K2Node_FunctionEntry.h"
 #include "EditorStyleSet.h"
+#include "BlueprintEditorModule.h"
 #include "BlueprintEditorSettings.h"
 #include "Framework/Docking/TabManager.h"
 #include "Widgets/Docking/SDockTab.h"
@@ -44,6 +45,7 @@
 #include "AssetRegistryModule.h"
 #include "ImaginaryBlueprintData.h"
 #include "FiBSearchInstance.h"
+#include "Misc/ConfigCacheIni.h"
 #include "Misc/HotReloadInterface.h"
 #include "BlueprintAssetHandler.h"
 
@@ -1083,29 +1085,51 @@ namespace BlueprintSearchMetaDataHelpers
 
 class FCacheAllBlueprintsTickableObject
 {
-
 public:
+	DECLARE_DELEGATE_OneParam(FOnAssetCached, FName);
+
+	struct FCacheParams
+	{
+		/** If TRUE, allow the operation to be cancelled by the user */
+		bool bAllowCancel = false;
+
+		/** If TRUE, show a progress notification (toast) while processing */
+		bool bShowProgress = true;
+
+		/** If TRUE, Blueprints will be checked out and resaved after being loaded */
+		bool bCheckOutAndSave = false;
 
-	FCacheAllBlueprintsTickableObject(TSet<FName> InUncachedAssets, bool bInCheckOutAndSave, FSimpleDelegate InOnFinished = FSimpleDelegate())
+		/** Callback for when assets are cached */
+		FOnAssetCached OnCached;
+
+		/** Callback for when caching is finished */
+		FSimpleDelegate OnFinished;
+	};
+
+	FCacheAllBlueprintsTickableObject(TSet<FName> InAssets, const FCacheParams& InParams)
 		: TickCacheIndex(0)
-		, UncachedAssets(InUncachedAssets.Array())
+		, UncachedAssets(InAssets.Array())
 		, bIsStarted(false)
 		, bIsCancelled(false)
-		, bRecursionGuard(false)
-		, bCheckOutAndSave(bInCheckOutAndSave)
-		, OnFinished(InOnFinished)
+		, CacheParams(InParams)
 	{
-		// Start the Blueprint indexing 'progress' notification
-		FNotificationInfo Info( LOCTEXT("BlueprintIndexMessage", "Indexing Blueprints...") );
-		Info.bFireAndForget = false;
-		Info.ButtonDetails.Add(FNotificationButtonInfo(
-			LOCTEXT("BlueprintIndexCancel","Cancel"),
-			LOCTEXT("BlueprintIndexCancelToolTip","Cancels indexing Blueprints."), FSimpleDelegate::CreateRaw(this, &FCacheAllBlueprintsTickableObject::OnCancelCaching, false)));
-
-		ProgressNotification = FSlateNotificationManager::Get().AddNotification(Info);
-		if (ProgressNotification.IsValid())
+		if (CacheParams.bShowProgress)
 		{
-			ProgressNotification.Pin()->SetCompletionState(SNotificationItem::CS_Pending);
+			// Start the Blueprint indexing 'progress' notification
+			FNotificationInfo Info(LOCTEXT("BlueprintIndexMessage", "Indexing Blueprints..."));
+			Info.bFireAndForget = false;
+			if (CacheParams.bAllowCancel)
+			{
+				Info.ButtonDetails.Add(FNotificationButtonInfo(
+					LOCTEXT("BlueprintIndexCancel", "Cancel"),
+					LOCTEXT("BlueprintIndexCancelToolTip", "Cancels indexing Blueprints."), FSimpleDelegate::CreateRaw(this, &FCacheAllBlueprintsTickableObject::OnCancelCaching, false)));
+			}
+
+			ProgressNotification = FSlateNotificationManager::Get().AddNotification(Info);
+			if (ProgressNotification.IsValid())
+			{
+				ProgressNotification.Pin()->SetCompletionState(SNotificationItem::CS_Pending);
+			}
 		}
 	}
 
@@ -1151,7 +1175,7 @@ public:
 	/** True if there is a callback when done caching, this will prevent a re-query from occuring */
 	bool HasPostCacheWork() const
 	{
-		return OnFinished.IsBound();
+		return CacheParams.OnFinished.IsBound();
 	}
 
 	/** Cancels caching and destroys this object */
@@ -1159,10 +1183,13 @@ public:
 	{
 		if (!bIsCancelled)
 		{
-			ProgressNotification.Pin()->SetText(LOCTEXT("BlueprintIndexCancelled", "Cancelled Indexing Blueprints!"));
+			if (ProgressNotification.IsValid())
+			{
+				ProgressNotification.Pin()->SetText(LOCTEXT("BlueprintIndexCancelled", "Cancelled Indexing Blueprints!"));
 
-			ProgressNotification.Pin()->SetCompletionState(SNotificationItem::CS_Fail);
-			ProgressNotification.Pin()->ExpireAndFadeout();
+				ProgressNotification.Pin()->SetCompletionState(SNotificationItem::CS_Fail);
+				ProgressNotification.Pin()->ExpireAndFadeout();
+			}
 
 			// Sometimes we can't wait another tick to shutdown, so make the callback immediately.
 			if (bIsImmediate)
@@ -1177,26 +1204,25 @@ public:
 	}
 
 	/** Enables the caching process */
-	void Start() { bIsStarted = true; }
-
-	EActiveTimerReturnType Tick(double InCurrentTime, float InDeltaTime)
+	void Start()
 	{
-		// Protect against Slate recursion if a modal dialog appears from loading/resaving an asset
-		if (bRecursionGuard == true)
+		if (!bIsStarted)
 		{
-			return EActiveTimerReturnType::Continue;
+			bIsStarted = true;
+			FFindInBlueprintSearchManager::Get().StartedCachingBlueprints();
 		}
-		TGuardValue<bool> Guard(bRecursionGuard, true);
+	}
 
-		if(!bIsStarted)
+	void Tick(float InDeltaTime)
+	{
+		if (!bIsStarted)
 		{
-			return EActiveTimerReturnType::Continue;
+			return;
 		}
 
 		if (bIsCancelled || GWarn->ReceivedUserCancel())
 		{
 			FFindInBlueprintSearchManager::Get().FinishedCachingBlueprints(TickCacheIndex, FailedToCacheList);
-			return EActiveTimerReturnType::Stop;
 		}
 		else
 		{
@@ -1219,53 +1245,67 @@ public:
 			}
 
 			bool bIsAssetReadOnlyOnDisk = IFileManager::Get().IsReadOnly( *FinalPackageFilename );
-			bool bFailedToCache = bCheckOutAndSave;
+			bool bFailedToCache = CacheParams.bCheckOutAndSave;
 
-			if (!bIsAssetReadOnlyOnDisk || !bCheckOutAndSave)
+			if (!bIsAssetReadOnlyOnDisk || !CacheParams.bCheckOutAndSave)
 			{
-				UObject* Asset = AssetData.GetAsset();
-				if(Asset && bCheckOutAndSave)
+				if (!FFindInBlueprintSearchManager::Get().IsUnindexedCacheInProgress())
 				{
-					if (UBlueprint* BlueprintAsset = Cast<UBlueprint>(Asset))
+					// Re-index and update the cached value for loaded Blueprint assets only
+					if (AssetData.IsAssetLoaded())
 					{
-						if (BlueprintAsset->SkeletonGeneratedClass == nullptr)
+						if (UBlueprint* LoadedBlueprintAsset = Cast<UBlueprint>(AssetData.GetAsset()))
 						{
-							// There is no skeleton class, something was wrong with the Blueprint during compile on load. This asset will be marked as failing to cache.
-							bFailedToCache = false;
+							FFindInBlueprintSearchManager::Get().AddOrUpdateBlueprintSearchMetadata(LoadedBlueprintAsset, true);
 						}
 					}
-
-					// Still good to attempt to save
-					if (bFailedToCache)
+				}
+				else
+				{
+					UObject* Asset = AssetData.GetAsset();
+					if (Asset && CacheParams.bCheckOutAndSave)
 					{
-						// Assume the package was correctly checked out from SCC
-						bool bOutPackageLocallyWritable = true;
-
-						UPackage* Package = AssetData.GetPackage();
-
-						ISourceControlProvider& SourceControlProvider = ISourceControlModule::Get().GetProvider();
-						// Trusting the SCC status in the package file cache to minimize network activity during save.
-						const FSourceControlStatePtr SourceControlState = SourceControlProvider.GetState(Package, EStateCacheUsage::Use);
-						// If the package is in the depot, and not recognized as editable by source control, and not read-only, then we know the user has made the package locally writable!
-						const bool bSCCCanEdit = !SourceControlState.IsValid() || SourceControlState->CanCheckIn() || SourceControlState->IsIgnored() || SourceControlState->IsUnknown();
-						const bool bSCCIsCheckedOut = SourceControlState.IsValid() && SourceControlState->IsCheckedOut();
-						const bool bInDepot = SourceControlState.IsValid() && SourceControlState->IsSourceControlled();
-						if (!bSCCCanEdit && bInDepot && !bIsAssetReadOnlyOnDisk && SourceControlProvider.UsesLocalReadOnlyState() && !bSCCIsCheckedOut)
+						if (UBlueprint* BlueprintAsset = Cast<UBlueprint>(Asset))
 						{
-							bOutPackageLocallyWritable = false;
+							if (BlueprintAsset->SkeletonGeneratedClass == nullptr)
+							{
+								// There is no skeleton class, something was wrong with the Blueprint during compile on load. This asset will be marked as failing to cache.
+								bFailedToCache = false;
+							}
 						}
 
-						// Save the package if the file is writable
-						if (bOutPackageLocallyWritable)
+						// Still good to attempt to save
+						if (bFailedToCache)
 						{
-							UWorld* WorldAsset = Cast<UWorld>(Asset);
-
-							// Save the package
-							EObjectFlags ObjectFlags = (WorldAsset == nullptr) ? RF_Standalone : RF_NoFlags;
+							// Assume the package was correctly checked out from SCC
+							bool bOutPackageLocallyWritable = true;
+
+							UPackage* Package = AssetData.GetPackage();
+
+							ISourceControlProvider& SourceControlProvider = ISourceControlModule::Get().GetProvider();
+							// Trusting the SCC status in the package file cache to minimize network activity during save.
+							const FSourceControlStatePtr SourceControlState = SourceControlProvider.GetState(Package, EStateCacheUsage::Use);
+							// If the package is in the depot, and not recognized as editable by source control, and not read-only, then we know the user has made the package locally writable!
+							const bool bSCCCanEdit = !SourceControlState.IsValid() || SourceControlState->CanCheckIn() || SourceControlState->IsIgnored() || SourceControlState->IsUnknown();
+							const bool bSCCIsCheckedOut = SourceControlState.IsValid() && SourceControlState->IsCheckedOut();
+							const bool bInDepot = SourceControlState.IsValid() && SourceControlState->IsSourceControlled();
+							if (!bSCCCanEdit && bInDepot && !bIsAssetReadOnlyOnDisk && SourceControlProvider.UsesLocalReadOnlyState() && !bSCCIsCheckedOut)
+							{
+								bOutPackageLocallyWritable = false;
+							}
 
-							if (GEditor->SavePackage(Package, WorldAsset, ObjectFlags, *FinalPackageFilename, GError, nullptr, false, true, SAVE_NoError))
+							// Save the package if the file is writable
+							if (bOutPackageLocallyWritable)
 							{
-								bFailedToCache = false;
+								UWorld* WorldAsset = Cast<UWorld>(Asset);
+
+								// Save the package
+								EObjectFlags ObjectFlags = (WorldAsset == nullptr) ? RF_Standalone : RF_NoFlags;
+
+								if (GEditor->SavePackage(Package, WorldAsset, ObjectFlags, *FinalPackageFilename, GError, nullptr, false, true, SAVE_NoError))
+								{
+									bFailedToCache = false;
+								}
 							}
 						}
 					}
@@ -1276,32 +1316,36 @@ public:
 			{
 				FailedToCacheList.Add(UncachedAssets[TickCacheIndex]);
 			}
+			else
+			{
+				CacheParams.OnCached.ExecuteIfBound(UncachedAssets[TickCacheIndex]);
+			}
 
 			++TickCacheIndex;
 
 			// Check if done caching Blueprints
 			if(TickCacheIndex == UncachedAssets.Num())
 			{
-				ProgressNotification.Pin()->SetCompletionState(SNotificationItem::CS_Success);
-				ProgressNotification.Pin()->ExpireAndFadeout();
+				if (ProgressNotification.IsValid())
+				{
+					ProgressNotification.Pin()->SetCompletionState(SNotificationItem::CS_Success);
+					ProgressNotification.Pin()->ExpireAndFadeout();
 
-				ProgressNotification.Pin()->SetText(LOCTEXT("BlueprintIndexComplete", "Finished indexing Blueprints!"));
+					ProgressNotification.Pin()->SetText(LOCTEXT("BlueprintIndexComplete", "Finished indexing Blueprints!"));
+				}
 
 				// We have actually finished, use the OnFinished callback.
-				OnFinished.ExecuteIfBound();
+				CacheParams.OnFinished.ExecuteIfBound();
 
 				FFindInBlueprintSearchManager::Get().FinishedCachingBlueprints(TickCacheIndex, FailedToCacheList);
-
-				return EActiveTimerReturnType::Stop;
 			}
-			else
+			else if(ProgressNotification.IsValid())
 			{
 				FFormatNamedArguments Args;
 				Args.Add(TEXT("Percent"), FText::AsPercent(GetCacheProgress()));
 				ProgressNotification.Pin()->SetText(FText::Format(LOCTEXT("BlueprintIndexProgress", "Indexing Blueprints... ({Percent})"), Args));
 			}
 		}
-		return EActiveTimerReturnType::Continue;
 	}
 
 private:
@@ -1309,7 +1353,7 @@ private:
 	/** The current index, increases at a rate of once per tick */
 	int32 TickCacheIndex;
 
-	/** The list of uncached assets that are in the process of being cached */
+	/** The list of assets that are in the process of being cached */
 	TArray<FName> UncachedAssets;
 
 	/** Notification that appears and details progress */
@@ -1324,14 +1368,8 @@ private:
 	/** TRUE if the user has requested to cancel the caching process */
 	bool bIsCancelled;
 
-	/** Guard to prevent TickRecursion */
-	bool bRecursionGuard = false;
-
-	/** If TRUE, Blueprints will be checked out and resaved after being loaded */
-	bool bCheckOutAndSave;
-
-	/** Callback for when caching is finished */
-	FSimpleDelegate OnFinished;
+	/** Parameters for task configuration */
+	FCacheParams CacheParams;
 };
 
 FFindInBlueprintSearchManager& FFindInBlueprintSearchManager::Get()
@@ -1347,6 +1385,7 @@ FFindInBlueprintSearchManager& FFindInBlueprintSearchManager::Get()
 
 FFindInBlueprintSearchManager::FFindInBlueprintSearchManager()
 	: bEnableGatheringData(true)
+	, bDisableDeferredIndexing(false)
 	, bIsPausing(false)
 	, AssetRegistryModule(nullptr)
 	, CachingObject(nullptr)
@@ -1377,12 +1416,21 @@ FFindInBlueprintSearchManager::~FFindInBlueprintSearchManager()
 		HotReloadSupport.OnHotReload().RemoveAll(this);
 	}
 
+	if (FModuleManager::Get().IsModuleLoaded("Kismet"))
+	{
+		FBlueprintEditorModule& BlueprintEditorModule = FModuleManager::GetModuleChecked<FBlueprintEditorModule>("Kismet");
+		BlueprintEditorModule.OnBlueprintEditorOpened().RemoveAll(this);
+	}
+
 	// Shut down the global find results tab feature.
 	EnableGlobalFindResults(false);
 }
 
 void FFindInBlueprintSearchManager::Initialize()
 {
+	// Init configuration
+	GConfig->GetBool(TEXT("BlueprintSearchSettings"), TEXT("bDisableDeferredIndexing"), bDisableDeferredIndexing, GEditorIni);
+
 	// Must ensure we do not attempt to load the AssetRegistry Module while saving a package, however, if it is loaded already we can safely obtain it
 	if (!GIsSavingPackage || (GIsSavingPackage && FModuleManager::Get().IsModuleLoaded(TEXT("AssetRegistry"))))
 	{
@@ -1412,6 +1460,10 @@ void FFindInBlueprintSearchManager::Initialize()
 	IHotReloadInterface& HotReloadSupport = FModuleManager::LoadModuleChecked<IHotReloadInterface>("HotReload");
 	HotReloadSupport.OnHotReload().AddRaw(this, &FFindInBlueprintSearchManager::OnHotReload);
 
+	// Register to be notified of a Blueprint editor open event
+	FBlueprintEditorModule& BlueprintEditorModule = FModuleManager::LoadModuleChecked<FBlueprintEditorModule>("Kismet");
+	BlueprintEditorModule.OnBlueprintEditorOpened().AddRaw(this, &FFindInBlueprintSearchManager::HandleBlueprintEditorOpened);
+
 	if(!GIsSavingPackage && AssetRegistryModule)
 	{
 		// Do an immediate load of the cache to catch any Blueprints that were discovered by the asset registry before we initialized.
@@ -1470,7 +1522,7 @@ void FFindInBlueprintSearchManager::OnAssetAdded(const FAssetData& InAssetData)
 			const FString& FiBVersionedSearchData = Result.GetValue();
 			if (FiBVersionedSearchData.Len() == 0)
 			{
-				UncachedAssets.Add(InAssetData.ObjectPath);
+				UnindexedAssets.Add(InAssetData.ObjectPath);
 			}
 			else
 			{
@@ -1488,7 +1540,7 @@ void FFindInBlueprintSearchManager::OnAssetAdded(const FAssetData& InAssetData)
 			// The asset has no FiB data, keep track of it so we can inform the user
 			else
 			{
-				UncachedAssets.Add(InAssetData.ObjectPath);
+				UnindexedAssets.Add(InAssetData.ObjectPath);
 			}
 
 		}
@@ -1629,7 +1681,7 @@ void FFindInBlueprintSearchManager::OnAssetLoaded(UObject* InAsset)
 			SearchArray[*IndexPtr].Blueprint = BlueprintObject;
 		}
 
-		UncachedAssets.Remove(AssetPath);
+		UnindexedAssets.Remove(AssetPath);
 	}
 }
 
@@ -1643,6 +1695,17 @@ void FFindInBlueprintSearchManager::OnHotReload(bool bWasTriggeredAutomatically)
 	CachedAssetClasses.Reset();
 }
 
+void FFindInBlueprintSearchManager::HandleBlueprintEditorOpened(EBlueprintType InBlueprintType)
+{
+	// Kick off an indexing operation if we have pending assets
+	if (!IsCacheInProgress() && PendingAssets.Num() > 0)
+	{
+		FFindInBlueprintCachingOptions CachingOptions;
+		CachingOptions.OpType = EFiBCacheOpType::CachePendingAssets;
+		CacheAllAssets(nullptr, CachingOptions);
+	}
+}
+
 FString FFindInBlueprintSearchManager::GatherBlueprintSearchMetadata(const UBlueprint* Blueprint)
 {	
 	FTemporarilyUseFriendlyNodeTitles TemporarilyUseFriendlyNodeTitles;
@@ -1728,6 +1791,7 @@ FString FFindInBlueprintSearchManager::GatherBlueprintSearchMetadata(const UBlue
 
 	int32 Version = EFiBVersion::FIB_VER_LATEST;
 	SearchMetaData = FiBSerializationHelpers::Serialize(Version, false) + Writer->GetSerializedLookupTable() + SearchMetaData;
+
 	return SearchMetaData;
 }
 
@@ -1766,6 +1830,7 @@ void FFindInBlueprintSearchManager::AddOrUpdateBlueprintSearchMetadata(UBlueprin
 	{
 		Index = *IndexPtr;
 		SearchArray[Index].Blueprint = InBlueprint; // Blueprint instance may change due to reloading
+		SearchArray[Index].bMarkedForDeletion = false;
 	}
 
 	// Build the search data
@@ -1773,13 +1838,25 @@ void FFindInBlueprintSearchManager::AddOrUpdateBlueprintSearchMetadata(UBlueprin
 	{
 		ParentClassProp->ExportTextItem(SearchArray[Index].ParentClass, ParentClassProp->ContainerPtrToValuePtr<uint8>(InBlueprint), nullptr, InBlueprint, 0);
 	}
-	// Cannot successfully gather most searchable data if there is no SkeletonGeneratedClass, so don't try, leave it as whatever it was last set to
-	if(InBlueprint->SkeletonGeneratedClass != nullptr)
+
+	// During unindexed/out-of-date caching we will arrive here as a result of loading the asset, so don't remove the IsCacheInProgress() check!
+	if (bInForceReCache || IsCacheInProgress() || bDisableDeferredIndexing)
 	{
-		SearchArray[Index].Value = GatherBlueprintSearchMetadata(InBlueprint);
-		SearchArray[Index].Version = EFiBVersion::FIB_VER_LATEST;
+		// Cannot successfully gather most searchable data if there is no SkeletonGeneratedClass, so don't try, leave it as whatever it was last set to
+		if (InBlueprint->SkeletonGeneratedClass != nullptr)
+		{
+			SearchArray[Index].Value = GatherBlueprintSearchMetadata(InBlueprint);
+			SearchArray[Index].Version = EFiBVersion::FIB_VER_LATEST;
+		}
+
+		// Remove it from the list of pending assets (if it exists)
+		PendingAssets.Remove(AssetPath);
+	}
+	else
+	{
+		// Add it to the list of assets to be indexed (deferred until the next global search)
+		PendingAssets.Add(AssetPath);
 	}
-	SearchArray[Index].bMarkedForDeletion = false;
 }
 
 void FFindInBlueprintSearchManager::BeginSearchQuery(const FStreamSearch* InSearchOriginator)
@@ -1967,6 +2044,9 @@ void FFindInBlueprintSearchManager::CleanCache()
 		}
 		else
 		{
+			// Remove it from the list of loaded assets that require indexing
+			PendingAssets.Remove(SearchArray[SearchValuePair.Value].AssetPath);
+
 			// Level Blueprints are destroyed when you open a new level, we need to re-add it as an unloaded asset so long as they were not marked for deletion
 			if(!SearchArray[SearchValuePair.Value].bMarkedForDeletion && FModuleManager::Get().IsModuleLoaded(TEXT("AssetRegistry")))
 			{
@@ -2063,7 +2143,7 @@ FString FFindInBlueprintSearchManager::ConvertFTextToHexString(FText InValue)
 	return BytesToHex(SerializedData.GetData(), SerializedData.Num());
 }
 
-void FFindInBlueprintSearchManager::OnCacheAllUncachedAssets(bool bInSourceControlActive, bool bInCheckoutAndSave)
+void FFindInBlueprintSearchManager::OnCacheAllUnindexedAssets(bool bInSourceControlActive, bool bInCheckoutAndSave)
 {
 	// Multiple threads can be adding to this at the same time
 	FScopeLock ScopeLock(&SafeModifyCacheCriticalSection);
@@ -2089,73 +2169,102 @@ void FFindInBlueprintSearchManager::OnCacheAllUncachedAssets(bool bInSourceContr
 	}
 }
 
-void FFindInBlueprintSearchManager::CacheAllUncachedAssets(TWeakPtr< SFindInBlueprints > InSourceWidget, FWidgetActiveTimerDelegate& InOutActiveTimerDelegate, FSimpleDelegate InOnFinished/* = FSimpleDelegate()*/, EFiBVersion InMinimiumVersionRequirement/* = EFiBVersion::FIB_VER_LATEST*/)
+void FFindInBlueprintSearchManager::CacheAllAssets(TWeakPtr< SFindInBlueprints > InSourceWidget, const FFindInBlueprintCachingOptions& InOptions)
 {
 	// Do not start another caching process if one is in progress
 	if(!IsCacheInProgress())
 	{
-		TArray<FName> BlueprintsToUpdate;
-		// Add any out-of-date Blueprints to the list
-		for (FSearchData SearchData : SearchArray)
+		CurrentCacheOpType = InOptions.OpType;
+
+		if (CurrentCacheOpType == EFiBCacheOpType::CachePendingAssets)
 		{
-			if ((SearchData.Value.Len() != 0 || SearchData.ImaginaryBlueprint.IsValid()) && SearchData.Version < InMinimiumVersionRequirement)
+			FCacheAllBlueprintsTickableObject::FCacheParams CacheParams;
+			CacheParams.OnFinished = InOptions.OnFinished;
+			CacheParams.OnCached.BindLambda([this](FName InAssetName)
 			{
-				BlueprintsToUpdate.Add(SearchData.AssetPath);
-			}
-		}
+				// Ensure removal in case it did not get re-indexed for some reason (e.g. not loaded, not a Blueprint asset, etc.)
+				PendingAssets.Remove(InAssetName);
+			});
 
-		FText DialogTitle = LOCTEXT("ConfirmIndexAll_Title", "Indexing All");
-		FFormatNamedArguments Args;
-		Args.Add(TEXT("PackageCount"), UncachedAssets.Num() + BlueprintsToUpdate.Num());
+			// Show the progress notification if we have more than one asset
+			CacheParams.bShowProgress = (PendingAssets.Num() > 1);
 
-		FText DialogDisplayText;
-		
-		if (UncachedAssets.Num() && BlueprintsToUpdate.Num())
-		{
-			Args.Add(TEXT("PackageCount"), UncachedAssets.Num() + BlueprintsToUpdate.Num());
-			Args.Add(TEXT("UnindexedCount"), UncachedAssets.Num());
-			Args.Add(TEXT("OutOfDateCount"), BlueprintsToUpdate.Num());
-			DialogDisplayText = FText::Format(LOCTEXT("CacheAllConfirmationMessage_UncachedAndBlueprints", "This process can take a long time and the editor may become unresponsive; there are {PackageCount} ({UnindexedCount} Unindexed/{OutOfDateCount} Out-of-Date) Blueprints to load. \
-																					\n\nWould you like to checkout, load, and save all Blueprints to make this indexing permanent? Otherwise, all Blueprints will still be loaded but you will be required to re-index the next time you start the editor!"), Args);
-		}
-		else if (UncachedAssets.Num() && BlueprintsToUpdate.Num() == 0)
-		{
-			DialogDisplayText = FText::Format(LOCTEXT("CacheAllConfirmationMessage_UncachedOnly", "This process can take a long time and the editor may become unresponsive; there are {PackageCount} unindexed Blueprints to load. \
-																					 \n\nWould you like to checkout, load, and save all Blueprints to make this indexing permanent? Otherwise, all Blueprints will still be loaded but you will be required to re-index the next time you start the editor!"), Args);
-		}
-		else if (UncachedAssets.Num() == 0 && BlueprintsToUpdate.Num())
-		{
-			DialogDisplayText = FText::Format(LOCTEXT("CacheAllConfirmationMessage_BlueprintsOnly", "This process can take a long time and the editor may become unresponsive; there are {PackageCount} out-of-date Blueprints to load. \
-																					 \n\nWould you like to checkout, load, and save all Blueprints to make this indexing permanent? Otherwise, all Blueprints will still be loaded but you will be required to re-index the next time you start the editor!"), Args);
-		}
+			// Keep track of which global FiB context started the operation (if any)
+			SourceCachingWidget = InSourceWidget;
 
-		const EAppReturnType::Type ReturnValue = FMessageDialog::Open(EAppMsgType::YesNoCancel, DialogDisplayText, &DialogTitle);
+			CachingObject = new FCacheAllBlueprintsTickableObject(PendingAssets, CacheParams);
 
-		// If Yes is chosen, checkout and save all Blueprints, if No is chosen, only load all Blueprints
-		if (ReturnValue != EAppReturnType::Cancel)
+			// Immediately start the operation (non-interactive)
+			check(CachingObject);
+			CachingObject->Start();
+		}
+		else
 		{
-			FailedToCachePaths.Empty();
-			
-			TSet<FName> TempUncachedAssets;
-			TempUncachedAssets.Append(UncachedAssets);
-			TempUncachedAssets.Append(BlueprintsToUpdate);
+			TArray<FName> BlueprintsToUpdate;
+			// Add any out-of-date Blueprints to the list
+			for (FSearchData SearchData : SearchArray)
+			{
+				if ((SearchData.Value.Len() != 0 || SearchData.ImaginaryBlueprint.IsValid()) && SearchData.Version < InOptions.MinimiumVersionRequirement)
+				{
+					BlueprintsToUpdate.Add(SearchData.AssetPath);
+				}
+			}
 
-			const bool bCheckoutAndSave = ReturnValue == EAppReturnType::Yes;
-			CachingObject = new FCacheAllBlueprintsTickableObject(MoveTemp(TempUncachedAssets), bCheckoutAndSave, InOnFinished);
-			InOutActiveTimerDelegate.BindRaw(CachingObject, &FCacheAllBlueprintsTickableObject::Tick);
+			FText DialogTitle = LOCTEXT("ConfirmIndexAll_Title", "Indexing All");
+			FFormatNamedArguments Args;
+			Args.Add(TEXT("PackageCount"), UnindexedAssets.Num() + BlueprintsToUpdate.Num());
 
-			const bool bIsSourceControlEnabled = ISourceControlModule::Get().IsEnabled();
-			if(!bIsSourceControlEnabled && bCheckoutAndSave)
+			FText DialogDisplayText;
+
+			if (UnindexedAssets.Num() && BlueprintsToUpdate.Num())
 			{
-				// Offer to start up Source Control
-				ISourceControlModule::Get().ShowLoginDialog(FSourceControlLoginClosed::CreateRaw(this, &FFindInBlueprintSearchManager::OnCacheAllUncachedAssets, bCheckoutAndSave), ELoginWindowMode::Modeless, EOnLoginWindowStartup::PreserveProvider);
+				Args.Add(TEXT("PackageCount"), UnindexedAssets.Num() + BlueprintsToUpdate.Num());
+				Args.Add(TEXT("UnindexedCount"), UnindexedAssets.Num());
+				Args.Add(TEXT("OutOfDateCount"), BlueprintsToUpdate.Num());
+				DialogDisplayText = FText::Format(LOCTEXT("CacheAllConfirmationMessage_UnindexedAndOutOfDate", "This process can take a long time and the editor may become unresponsive; there are {PackageCount} ({UnindexedCount} Unindexed/{OutOfDateCount} Out-of-Date) Blueprints to load. \
+																					\n\nWould you like to checkout, load, and save all Blueprints to make this indexing permanent? Otherwise, all Blueprints will still be loaded but you will be required to re-index the next time you start the editor!"), Args);
 			}
-			else
+			else if (UnindexedAssets.Num() && BlueprintsToUpdate.Num() == 0)
 			{
-				OnCacheAllUncachedAssets(bIsSourceControlEnabled, bCheckoutAndSave);
+				DialogDisplayText = FText::Format(LOCTEXT("CacheAllConfirmationMessage_UnindexedOnly", "This process can take a long time and the editor may become unresponsive; there are {PackageCount} unindexed Blueprints to load. \
+																					 \n\nWould you like to checkout, load, and save all Blueprints to make this indexing permanent? Otherwise, all Blueprints will still be loaded but you will be required to re-index the next time you start the editor!"), Args);
+			}
+			else if (UnindexedAssets.Num() == 0 && BlueprintsToUpdate.Num())
+			{
+				DialogDisplayText = FText::Format(LOCTEXT("CacheAllConfirmationMessage_OutOfDateOnly", "This process can take a long time and the editor may become unresponsive; there are {PackageCount} out-of-date Blueprints to load. \
+																					 \n\nWould you like to checkout, load, and save all Blueprints to make this indexing permanent? Otherwise, all Blueprints will still be loaded but you will be required to re-index the next time you start the editor!"), Args);
 			}
 
-			SourceCachingWidget = InSourceWidget;
+			const EAppReturnType::Type ReturnValue = FMessageDialog::Open(EAppMsgType::YesNoCancel, DialogDisplayText, &DialogTitle);
+
+			// If Yes is chosen, checkout and save all Blueprints, if No is chosen, only load all Blueprints
+			if (ReturnValue != EAppReturnType::Cancel)
+			{
+				FailedToCachePaths.Empty();
+
+				TSet<FName> TempUncachedAssets;
+				TempUncachedAssets.Append(UnindexedAssets);
+				TempUncachedAssets.Append(BlueprintsToUpdate);
+
+				FCacheAllBlueprintsTickableObject::FCacheParams CacheParams;
+				CacheParams.bAllowCancel = true;
+				CacheParams.bCheckOutAndSave = ReturnValue == EAppReturnType::Yes;
+				CacheParams.OnFinished = InOptions.OnFinished;
+				CachingObject = new FCacheAllBlueprintsTickableObject(MoveTemp(TempUncachedAssets), CacheParams);
+
+				const bool bIsSourceControlEnabled = ISourceControlModule::Get().IsEnabled();
+				if (!bIsSourceControlEnabled && CacheParams.bCheckOutAndSave)
+				{
+					// Offer to start up Source Control
+					ISourceControlModule::Get().ShowLoginDialog(FSourceControlLoginClosed::CreateRaw(this, &FFindInBlueprintSearchManager::OnCacheAllUnindexedAssets, CacheParams.bCheckOutAndSave), ELoginWindowMode::Modeless, EOnLoginWindowStartup::PreserveProvider);
+				}
+				else
+				{
+					OnCacheAllUnindexedAssets(bIsSourceControlEnabled, CacheParams.bCheckOutAndSave);
+				}
+
+				SourceCachingWidget = InSourceWidget;
+			}
 		}
 	}
 }
@@ -2203,14 +2312,36 @@ float FFindInBlueprintSearchManager::GetCacheProgress() const
 	return ReturnCacheValue;
 }
 
+int32 FFindInBlueprintSearchManager::GetNumberPendingAssets() const
+{
+	return PendingAssets.Num();
+}
+
+int32 FFindInBlueprintSearchManager::GetNumberUnindexedAssets() const
+{
+	return UnindexedAssets.Num();
+}
+
 int32 FFindInBlueprintSearchManager::GetNumberUncachedAssets() const
 {
-	int32 ReturnCount = UncachedAssets.Num();
 	if (CachingObject)
 	{
-		ReturnCount = CachingObject->GetUncachedAssetCount();
+		return CachingObject->GetUncachedAssetCount();
+	}
+	
+	return 0;
+}
+
+void FFindInBlueprintSearchManager::StartedCachingBlueprints()
+{
+	// Invoke the callback on any open global widgets
+	for (TWeakPtr<SFindInBlueprints> FindResultsPtr : GlobalFindResults)
+	{
+		if (FindResultsPtr.IsValid())
+		{
+			FindResultsPtr.Pin()->OnCacheStarted(CurrentCacheOpType);
+		}
 	}
-	return ReturnCount;
 }
 
 void FFindInBlueprintSearchManager::FinishedCachingBlueprints(int32 InNumberCached, TSet<FName>& InFailedToCacheList)
@@ -2221,11 +2352,28 @@ void FFindInBlueprintSearchManager::FinishedCachingBlueprints(int32 InNumberCach
 	// Update the list of cache failures
 	FailedToCachePaths = InFailedToCacheList;
 
-	// Signal the callback, so the source FindInBlueprint can resubmit their search queries
+	// Remove any failed attempts from the pending list (should only apply to unindexed caching, but just in case)
+	if (InFailedToCacheList.Num() > 0)
+	{
+		PendingAssets = PendingAssets.Difference(InFailedToCacheList);
+	}
+
+	// Invoke the completion callback on any open widgets that are not the initiating one (e.g. to hide progress bar)
+	for (TWeakPtr<SFindInBlueprints> FindResultsPtr : GlobalFindResults)
+	{
+		if (FindResultsPtr.IsValid() && FindResultsPtr != SourceCachingWidget)
+		{
+			FindResultsPtr.Pin()->OnCacheComplete(CurrentCacheOpType);
+		}
+	}
+
+	// Invoke the completion callback on the initiating widget only if necessary (i.e. only if it doesn't already delegate it)
 	if(SourceCachingWidget.IsValid() && !CachingObject->HasPostCacheWork())
 	{
-		SourceCachingWidget.Pin()->OnCacheComplete();
+		SourceCachingWidget.Pin()->OnCacheComplete(CurrentCacheOpType);
 	}
+
+	// Reset the reference to the widget that initiated the caching operation
 	SourceCachingWidget.Reset();
 
 	// Delete the object and NULL it out so we can do it again in the future if needed (if it was canceled)
@@ -2238,6 +2386,11 @@ bool FFindInBlueprintSearchManager::IsCacheInProgress() const
 	return CachingObject != nullptr;
 }
 
+bool FFindInBlueprintSearchManager::IsUnindexedCacheInProgress() const
+{
+	return IsCacheInProgress() && CurrentCacheOpType == EFiBCacheOpType::CacheUnindexedAssets;
+}
+
 TSharedPtr< FJsonObject > FFindInBlueprintSearchManager::ConvertJsonStringToObject(bool bInIsVersioned, FString InJsonString, TMap<int32, FText>& OutFTextLookupTable)
 {
 	/** The searchable data is more complicated than a Json string, the Json being the main searchable body that is parsed. Below is a diagram of the full data:
@@ -2283,6 +2436,11 @@ void FFindInBlueprintSearchManager::GlobalFindResultsClosed(const TSharedRef<SFi
 		if (FindResultsPtr.Pin() == FindResults)
 		{
 			GlobalFindResults.Remove(FindResultsPtr);
+
+			if (CachingObject && SourceCachingWidget == FindResultsPtr)
+			{
+				SourceCachingWidget.Reset();
+			}
 			break;
 		}
 	}
@@ -2470,4 +2628,31 @@ void FFindInBlueprintSearchManager::CloseOrphanedGlobalFindResultsTabs(TSharedPt
 	}
 }
 
+void FFindInBlueprintSearchManager::Tick(float DeltaTime)
+{
+	if(IsCacheInProgress())
+	{
+		check(CachingObject);
+		CachingObject->Tick(DeltaTime);
+	}
+	else if (PendingAssets.Num() > 0 && GlobalFindResults.Num() > 0)
+	{
+		// Kick off an indexing operation if we have pending assets and at least one open global FiB context
+		FFindInBlueprintCachingOptions CachingOptions;
+		CachingOptions.OpType = EFiBCacheOpType::CachePendingAssets;
+		CacheAllAssets(nullptr, CachingOptions);
+	}
+}
+
+bool FFindInBlueprintSearchManager::IsTickable() const
+{
+	// Tick only if we have an active caching operation or if we have pending assets and an open FiB context
+	return IsCacheInProgress() || (PendingAssets.Num() > 0 && GlobalFindResults.Num() > 0);
+}
+
+TStatId FFindInBlueprintSearchManager::GetStatId() const
+{
+	RETURN_QUICK_DECLARE_CYCLE_STAT(FFindInBlueprintSearchManager, STATGROUP_Tickables);
+}
+
 #undef LOCTEXT_NAMESPACE
diff --git a/Engine/Source/Editor/Kismet/Private/FindInBlueprints.cpp b/Engine/Source/Editor/Kismet/Private/FindInBlueprints.cpp
--- a/Engine/Source/Editor/Kismet/Private/FindInBlueprints.cpp
+++ b/Engine/Source/Editor/Kismet/Private/FindInBlueprints.cpp
@@ -594,6 +594,7 @@ void SFindInBlueprints::Construct( const FArguments& InArgs, TSharedPtr<FBluepri
 
 	HostTab = InArgs._ContainingTab;
 	bIsLocked = false;
+	bKeepCacheBarOpenOnCacheCompletion = false;
 
 	if (HostTab.IsValid())
 	{
@@ -606,121 +607,126 @@ void SFindInBlueprints::Construct( const FArguments& InArgs, TSharedPtr<FBluepri
 	}
 
 	bIsInFindWithinBlueprintMode = BlueprintEditorPtr.IsValid();
-	bHasGlobalSearchResults = !bIsInFindWithinBlueprintMode;
 
 	const bool bHostFindInBlueprintsInGlobalTab = GetDefault<UBlueprintEditorSettings>()->bHostFindInBlueprintsInGlobalTab;
 	
 	this->ChildSlot
+	[
+		SAssignNew(MainVerticalBox, SVerticalBox)
+		+SVerticalBox::Slot()
+		.AutoHeight()
 		[
-			SAssignNew(MainVerticalBox, SVerticalBox)
-			+SVerticalBox::Slot()
-			.AutoHeight()
+			SNew(SHorizontalBox)
+			+SHorizontalBox::Slot()
+			.FillWidth(1)
 			[
-				SNew(SHorizontalBox)
-				+SHorizontalBox::Slot()
-				.FillWidth(1)
-				[
-					SAssignNew(SearchTextField, SSearchBox)
-					.HintText(LOCTEXT("BlueprintSearchHint", "Enter function or event name to find references..."))
-					.OnTextChanged(this, &SFindInBlueprints::OnSearchTextChanged)
-					.OnTextCommitted(this, &SFindInBlueprints::OnSearchTextCommitted)
-					.Visibility(InArgs._bHideSearchBar? EVisibility::Collapsed : EVisibility::Visible)
-				]
-				+SHorizontalBox::Slot()
-				.Padding(4.f, 0.f, 2.f, 0.f)
-				.AutoWidth()
-				[
-					SNew(SButton)
-					.OnClicked(this, &SFindInBlueprints::OnOpenGlobalFindResults)
-					.Visibility(BlueprintEditorPtr.IsValid() && bHostFindInBlueprintsInGlobalTab ? EVisibility::Visible : EVisibility::Collapsed)
-					.ToolTipText(LOCTEXT("OpenInGlobalFindResultsButtonTooltip", "Find in all Blueprints"))
-					[
-						SNew(STextBlock)
-						.TextStyle(FEditorStyle::Get(), "FindResults.FindInBlueprints")
-						.Text(FText::FromString(FString(TEXT("\xf1e5"))) /*fa-binoculars*/)
-					]
-				]
-				+SHorizontalBox::Slot()
-				.Padding(2.f, 0.f)
-				.AutoWidth()
+				SAssignNew(SearchTextField, SSearchBox)
+				.HintText(LOCTEXT("BlueprintSearchHint", "Enter function or event name to find references..."))
+				.OnTextChanged(this, &SFindInBlueprints::OnSearchTextChanged)
+				.OnTextCommitted(this, &SFindInBlueprints::OnSearchTextCommitted)
+				.Visibility(InArgs._bHideSearchBar? EVisibility::Collapsed : EVisibility::Visible)
+			]
+			+SHorizontalBox::Slot()
+			.Padding(4.f, 0.f, 2.f, 0.f)
+			.AutoWidth()
+			[
+				SNew(SButton)
+				.OnClicked(this, &SFindInBlueprints::OnOpenGlobalFindResults)
+				.Visibility(BlueprintEditorPtr.IsValid() && bHostFindInBlueprintsInGlobalTab ? EVisibility::Visible : EVisibility::Collapsed)
+				.ToolTipText(LOCTEXT("OpenInGlobalFindResultsButtonTooltip", "Find in all Blueprints"))
 				[
-					SNew(SCheckBox)
-					.OnCheckStateChanged(this, &SFindInBlueprints::OnFindModeChanged)
-					.IsChecked(this, &SFindInBlueprints::OnGetFindModeChecked)
-					.Visibility(InArgs._bHideSearchBar || bHostFindInBlueprintsInGlobalTab ? EVisibility::Collapsed : EVisibility::Visible)
-					[
-						SNew(STextBlock)
-						.Text(LOCTEXT("BlueprintSearchModeChange", "Find In Current Blueprint Only"))
-					]
+					SNew(STextBlock)
+					.TextStyle(FEditorStyle::Get(), "FindResults.FindInBlueprints")
+					.Text(FText::FromString(FString(TEXT("\xf1e5"))) /*fa-binoculars*/)
 				]
-				+SHorizontalBox::Slot()
-				.AutoWidth()
-				.VAlign(VAlign_Center)
+			]
+			+SHorizontalBox::Slot()
+			.Padding(2.f, 0.f)
+			.AutoWidth()
+			[
+				SNew(SCheckBox)
+				.OnCheckStateChanged(this, &SFindInBlueprints::OnFindModeChanged)
+				.IsChecked(this, &SFindInBlueprints::OnGetFindModeChecked)
+				.Visibility(InArgs._bHideSearchBar || bHostFindInBlueprintsInGlobalTab ? EVisibility::Collapsed : EVisibility::Visible)
 				[
-					SNew(SButton)
-					.VAlign(EVerticalAlignment::VAlign_Center)
-					.ButtonStyle(FEditorStyle::Get(), "FlatButton")
-					.ContentPadding(FMargin(1, 0))
-					.OnClicked(this, &SFindInBlueprints::OnLockButtonClicked)
-					.Visibility(!InArgs._bHideSearchBar && !BlueprintEditorPtr.IsValid() ? EVisibility::Visible : EVisibility::Collapsed)
-					[
-						SNew(SImage)
-						.Image(this, &SFindInBlueprints::OnGetLockButtonImage)
-					]
+					SNew(STextBlock)
+					.Text(LOCTEXT("BlueprintSearchModeChange", "Find In Current Blueprint Only"))
 				]
 			]
-			+SVerticalBox::Slot()
-			.FillHeight(1.0f)
-			.Padding(0.f, 4.f, 0.f, 0.f)
+			+SHorizontalBox::Slot()
+			.AutoWidth()
+			.VAlign(VAlign_Center)
 			[
-				SNew(SBorder)
-				.BorderImage(FEditorStyle::GetBrush("Menu.Background"))
+				SNew(SButton)
+				.VAlign(EVerticalAlignment::VAlign_Center)
+				.ButtonStyle(FEditorStyle::Get(), "FlatButton")
+				.ContentPadding(FMargin(1, 0))
+				.OnClicked(this, &SFindInBlueprints::OnLockButtonClicked)
+				.Visibility(!InArgs._bHideSearchBar && !BlueprintEditorPtr.IsValid() ? EVisibility::Visible : EVisibility::Collapsed)
 				[
-					SAssignNew(TreeView, STreeViewType)
-					.ItemHeight(24)
-					.TreeItemsSource( &ItemsFound )
-					.OnGenerateRow( this, &SFindInBlueprints::OnGenerateRow )
-					.OnGetChildren( this, &SFindInBlueprints::OnGetChildren )
-					.OnMouseButtonDoubleClick(this,&SFindInBlueprints::OnTreeSelectionDoubleClicked)
-					.SelectionMode( ESelectionMode::Multi )
-					.OnContextMenuOpening(this, &SFindInBlueprints::OnContextMenuOpening)
+					SNew(SImage)
+					.Image(this, &SFindInBlueprints::OnGetLockButtonImage)
 				]
 			]
-
-			+SVerticalBox::Slot()
-				.AutoHeight()
+		]
+		+SVerticalBox::Slot()
+		.FillHeight(1.0f)
+		.Padding(0.f, 4.f, 0.f, 0.f)
+		[
+			SNew(SBorder)
+			.BorderImage(FEditorStyle::GetBrush("Menu.Background"))
 			[
-				SNew(SHorizontalBox)
+				SAssignNew(TreeView, STreeViewType)
+				.ItemHeight(24)
+				.TreeItemsSource( &ItemsFound )
+				.OnGenerateRow( this, &SFindInBlueprints::OnGenerateRow )
+				.OnGetChildren( this, &SFindInBlueprints::OnGetChildren )
+				.OnMouseButtonDoubleClick(this,&SFindInBlueprints::OnTreeSelectionDoubleClicked)
+				.SelectionMode( ESelectionMode::Multi )
+				.OnContextMenuOpening(this, &SFindInBlueprints::OnContextMenuOpening)
+			]
+		]
 
-				// Text
-				+SHorizontalBox::Slot()
-				.AutoWidth()
-				.Padding(0, 2)
-				[
-					SNew(STextBlock)
-					.Font( FEditorStyle::GetFontStyle("AssetDiscoveryIndicator.MainStatusFont") )
-					.Text( LOCTEXT("SearchResults", "Searching...") )
-					.Visibility(this, &SFindInBlueprints::GetSearchbarVisiblity)
-				]
+		+SVerticalBox::Slot()
+			.AutoHeight()
+		[
+			SNew(SHorizontalBox)
 
-				// Progress bar
-				+SHorizontalBox::Slot()
-					.FillWidth(1.0f)
-					.Padding(2.0f, 3.0f, 0.0f, 0.0f)
-				[
-					SNew(SProgressBar)
-						.Visibility(this, &SFindInBlueprints::GetSearchbarVisiblity).Percent( this, &SFindInBlueprints::GetPercentCompleteSearch )
-				]
+			// Text
+			+SHorizontalBox::Slot()
+			.AutoWidth()
+			.Padding(0, 2)
+			[
+				SNew(STextBlock)
+				.Font( FEditorStyle::GetFontStyle("AssetDiscoveryIndicator.MainStatusFont") )
+				.Text( LOCTEXT("SearchResults", "Searching...") )
+				.Visibility(this, &SFindInBlueprints::GetSearchbarVisiblity)
 			]
-		];
+
+			// Progress bar
+			+SHorizontalBox::Slot()
+				.FillWidth(1.0f)
+				.Padding(2.0f, 3.0f, 0.0f, 0.0f)
+			[
+				SNew(SProgressBar)
+					.Visibility(this, &SFindInBlueprints::GetSearchbarVisiblity).Percent( this, &SFindInBlueprints::GetPercentCompleteSearch )
+			]
+		]
+	];
+
+	// Include the cache progress bar if we're in global find mode and an active cache operation is still in progress.
+	if (!bIsInFindWithinBlueprintMode && IsCacheInProgress())
+	{
+		ConditionallyAddCacheBar();
+	}
 }
 
 void SFindInBlueprints::ConditionallyAddCacheBar()
 {
 	FFindInBlueprintSearchManager& FindInBlueprintManager = FFindInBlueprintSearchManager::Get();
 
-	// Do not add a second cache bar and do not add it when there are no uncached Blueprints
-	if(FindInBlueprintManager.GetNumberUncachedAssets() > 0 || FindInBlueprintManager.GetFailedToCacheCount() > 0)
+	// Do not add a second cache bar and do not add it when there are no unindexed Blueprints
+	if(FindInBlueprintManager.GetNumberUncachedAssets() > 0 || FindInBlueprintManager.GetNumberUnindexedAssets() > 0 || FindInBlueprintManager.GetFailedToCacheCount() > 0)
 	{
 		if(MainVerticalBox.IsValid() && !CacheBarSlot.IsValid())
 		{
@@ -788,7 +794,7 @@ void SFindInBlueprints::ConditionallyAddCacheBar()
 							.AutoWidth()
 							[
 								SNew(STextBlock)
-								.Text(this, &SFindInBlueprints::GetUncachedAssetWarningText)
+								.Text(this, &SFindInBlueprints::GetUnindexedAssetWarningText)
 								.ColorAndOpacity( FCoreStyle::Get().GetColor("ErrorReporting.ForegroundColor") )
 							]
 
@@ -800,8 +806,8 @@ void SFindInBlueprints::ConditionallyAddCacheBar()
 								[
 									SNew(SButton)
 									.Text(LOCTEXT("IndexAllBlueprints", "Index All"))
-									.OnClicked( this, &SFindInBlueprints::OnCacheAllBlueprints )
-									.Visibility( this, &SFindInBlueprints::GetCacheAllButtonVisibility )
+									.OnClicked( this, &SFindInBlueprints::OnCacheAllUnindexedBlueprints )
+									.Visibility( this, &SFindInBlueprints::GetCacheAllUnindexedButtonVisibility )
 									.ToolTip(IDocumentation::Get()->CreateToolTip(
 									LOCTEXT("IndexAlLBlueprints_Tooltip", "Loads all non-indexed Blueprints and saves them with their search data. This can be a very slow process and the editor may become unresponsive."),
 									NULL,
@@ -844,7 +850,7 @@ void SFindInBlueprints::ConditionallyAddCacheBar()
 									SNew(SButton)
 									.Text(LOCTEXT("CancelCacheAll", "Cancel"))
 									.OnClicked( this, &SFindInBlueprints::OnCancelCacheAll )
-									.Visibility( this, &SFindInBlueprints::GetCachingProgressBarVisiblity )
+									.Visibility( this, &SFindInBlueprints::GetCacheAllCancelButtonVisibility )
 									.ToolTipText( LOCTEXT("CancelCacheAll_Tooltip", "Stops the caching process from where ever it is, can be started back up where it left off when needed.") )
 								]
 
@@ -885,8 +891,9 @@ void SFindInBlueprints::ConditionallyAddCacheBar()
 								.AutoHeight()
 								[
 									SNew(STextBlock)
-									.Text(LOCTEXT("FiBUnresponsiveEditorWarning", "NOTE: the editor may become unresponsive for some time!"))
+									.Text(LOCTEXT("FiBUnresponsiveEditorWarning", "NOTE: The editor may become unresponsive while these assets are loaded for indexing. This may take some time!"))
 									.TextStyle(&FCoreStyle::Get().GetWidgetStyle<FTextBlockStyle>( "SmallText" ))
+									.Visibility(this, &SFindInBlueprints::GetUnresponsiveEditorWarningVisibility)
 								]
 							]
 					]
@@ -918,7 +925,11 @@ SFindInBlueprints::~SFindInBlueprints()
 		StreamSearch->EnsureCompletion();
 	}
 
-	FFindInBlueprintSearchManager::Get().CancelCacheAll(this);
+	// Only cancel unindexed (slow) caching operations upon destruction
+	if (FFindInBlueprintSearchManager::Get().IsUnindexedCacheInProgress())
+	{
+		FFindInBlueprintSearchManager::Get().CancelCacheAll(this);
+	}
 }
 
 EActiveTimerReturnType SFindInBlueprints::UpdateSearchResults( double InCurrentTime, float InDeltaTime )
@@ -1033,7 +1044,11 @@ void SFindInBlueprints::MakeSearchQuery(FString InSearchString, bool bInIsFindWi
 
 	if (InSearchString.Len() > 0)
 	{
-		OnRemoveCacheBar();
+		// Remove the cache bar unless an active cache is in progress (so that we still show the status). It's ok to proceed with the new search while this is ongoing.
+		if (!IsCacheInProgress())
+		{
+			OnRemoveCacheBar();
+		}
 
 		TreeView->RequestTreeRefresh();
 		HighlightText = FText::FromString( InSearchString );
@@ -1090,8 +1105,6 @@ void SFindInBlueprints::MakeSearchQuery(FString InSearchString, bool bInIsFindWi
 		{
 			LaunchStreamThread(InSearchString, InSearchFilterForImaginaryDataReturn, InMinimiumVersionRequirement, InOnSearchComplete);
 		}
-
-		bHasGlobalSearchResults = !bInIsFindWithinBlueprint;
 	}
 }
 
@@ -1254,24 +1267,23 @@ EVisibility SFindInBlueprints::GetSearchbarVisiblity() const
 	return StreamSearch.IsValid()? EVisibility::Visible : EVisibility::Collapsed;
 }
 
-void SFindInBlueprints::CacheAllBlueprints(FSimpleDelegate InOnFinished/* = FSimpleDelegate()*/, EFiBVersion InMinimiumVersionRequirement/* = EFiBVersion::FIB_VER_LATEST*/)
+void SFindInBlueprints::CacheAllBlueprints(const FFindInBlueprintCachingOptions& InOptions)
 {
-	OnCacheAllBlueprints(InOnFinished, InMinimiumVersionRequirement);
+	OnCacheAllBlueprints(InOptions);
 }
 
-FReply SFindInBlueprints::OnCacheAllBlueprints()
+FReply SFindInBlueprints::OnCacheAllUnindexedBlueprints()
 {
-	return OnCacheAllBlueprints(FSimpleDelegate(), LastSearchedFiBVersion);
+	FFindInBlueprintCachingOptions CachingOptions;
+	CachingOptions.OpType = EFiBCacheOpType::CacheUnindexedAssets;
+	return OnCacheAllBlueprints(CachingOptions);
 }
 
-FReply SFindInBlueprints::OnCacheAllBlueprints(FSimpleDelegate InOnFinished/* = FSimpleDelegate()*/, EFiBVersion InMinimiumVersionRequirement/* = EFiBVersion::FIB_VER_LATEST*/)
+FReply SFindInBlueprints::OnCacheAllBlueprints(const FFindInBlueprintCachingOptions& InOptions)
 {
 	if(!FFindInBlueprintSearchManager::Get().IsCacheInProgress())
 	{
-		// Request from the SearchManager a delegate to use for ticking the cache system.
-		FWidgetActiveTimerDelegate WidgetActiveTimer;
-		FFindInBlueprintSearchManager::Get().CacheAllUncachedAssets(SharedThis(this), WidgetActiveTimer, InOnFinished, InMinimiumVersionRequirement);
-		RegisterActiveTimer(0.f, WidgetActiveTimer);
+		FFindInBlueprintSearchManager::Get().CacheAllAssets(SharedThis(this), InOptions);
 	}
 
 	return FReply::Handled();
@@ -1302,15 +1314,23 @@ EVisibility SFindInBlueprints::GetCachingProgressBarVisiblity() const
 	return IsCacheInProgress()? EVisibility::Visible : EVisibility::Hidden;
 }
 
-EVisibility SFindInBlueprints::GetCacheAllButtonVisibility() const
+EVisibility SFindInBlueprints::GetCacheAllUnindexedButtonVisibility() const
 {
 	return IsCacheInProgress()? EVisibility::Collapsed : EVisibility::Visible;
 }
 
+EVisibility SFindInBlueprints::GetCacheAllCancelButtonVisibility() const
+{
+	// The cancel button is visible only if we're doing an unindexed caching operation (slow) and this is the initiating widget
+	const bool bIsUnindexedCacheInProgress = FFindInBlueprintSearchManager::Get().IsUnindexedCacheInProgress();
+	TWeakPtr<SFindInBlueprints> SourceWidgetPtr = FFindInBlueprintSearchManager::Get().GetSourceCachingWidget();
+	return bIsUnindexedCacheInProgress && SourceWidgetPtr.Pin() == SharedThis(this) ? EVisibility::Visible : EVisibility::Collapsed;
+}
+
 EVisibility SFindInBlueprints::GetCachingBarVisibility() const
 {
 	FFindInBlueprintSearchManager& FindInBlueprintManager = FFindInBlueprintSearchManager::Get();
-	return (FindInBlueprintManager.GetNumberUncachedAssets() > 0 || FindInBlueprintManager.GetFailedToCacheCount())? EVisibility::Visible : EVisibility::Collapsed;
+	return (FindInBlueprintManager.GetNumberUncachedAssets() > 0 || FindInBlueprintManager.GetNumberUnindexedAssets() > 0 || FindInBlueprintManager.GetFailedToCacheCount())? EVisibility::Visible : EVisibility::Collapsed;
 }
 
 EVisibility SFindInBlueprints::GetCachingBlueprintNameVisiblity() const
@@ -1323,6 +1343,12 @@ EVisibility SFindInBlueprints::GetFailedToCacheListVisibility() const
 	return FFindInBlueprintSearchManager::Get().GetFailedToCacheCount()? EVisibility::Visible : EVisibility::Collapsed;
 }
 
+EVisibility SFindInBlueprints::GetUnresponsiveEditorWarningVisibility() const
+{
+	// Show the "editor may be unresponsive" note if we're doing an unindexed caching operation, which can block the UI due to file I/O
+	return !IsCacheInProgress() || FFindInBlueprintSearchManager::Get().IsUnindexedCacheInProgress() ? EVisibility::Visible : EVisibility::Collapsed;
+}
+
 bool SFindInBlueprints::IsCacheInProgress() const
 {
 	return FFindInBlueprintSearchManager::Get().IsCacheInProgress();
@@ -1340,36 +1366,35 @@ FSlateColor SFindInBlueprints::GetCachingBarColor() const
 	return ReturnColor;
 }
 
-FText SFindInBlueprints::GetUncachedAssetWarningText() const
+FText SFindInBlueprints::GetUnindexedAssetWarningText() const
 {
 	FFindInBlueprintSearchManager& FindInBlueprintManager = FFindInBlueprintSearchManager::Get();
 
-	int32 FailedToCacheCount = FindInBlueprintManager.GetFailedToCacheCount();
-
-	// The number of unindexed Blueprints is the total of those that failed to cache and those that haven't been attempted yet.
 	FFormatNamedArguments Args;
-	Args.Add(TEXT("Count"), FindInBlueprintManager.GetNumberUncachedAssets() + OutOfDateWithLastSearchBPCount);
-
 	FText ReturnDisplayText;
 	if(IsCacheInProgress())
 	{
 		Args.Add(TEXT("CurrentIndex"), FindInBlueprintManager.GetCurrentCacheIndex());
+		Args.Add(TEXT("Count"), FindInBlueprintManager.GetNumberUncachedAssets());
 
 		ReturnDisplayText = FText::Format(LOCTEXT("CachingBlueprints", "Indexing Blueprints... {CurrentIndex}/{Count}"), Args);
 	}
 	else
 	{
-		Args.Add(TEXT("UnindexedCount"), FindInBlueprintManager.GetNumberUncachedAssets());
+		const int32 UnindexedCount = FindInBlueprintManager.GetNumberUnindexedAssets();
+		Args.Add(TEXT("UnindexedCount"), UnindexedCount);
 		Args.Add(TEXT("OutOfDateCount"), OutOfDateWithLastSearchBPCount);
+		Args.Add(TEXT("Count"), UnindexedCount + OutOfDateWithLastSearchBPCount);
 
-		ReturnDisplayText = FText::Format(LOCTEXT("UncachedAssets", "Search incomplete. {Count} ({UnindexedCount} Unindexed/{OutOfDateCount} Out-of-Date) Blueprints need to be indexed!"), Args);
+		ReturnDisplayText = FText::Format(LOCTEXT("UncachedAssets", "Search incomplete. {Count} ({UnindexedCount} non-indexed/{OutOfDateCount} out-of-date) Blueprints need to be loaded and indexed!"), Args);
 
+		const int32 FailedToCacheCount = FindInBlueprintManager.GetFailedToCacheCount();
 		if (FailedToCacheCount > 0)
 		{
 			FFormatNamedArguments ArgsWithCacheFails;
-			Args.Add(TEXT("BaseMessage"), ReturnDisplayText);
-			Args.Add(TEXT("CacheFails"), FailedToCacheCount);
-			ReturnDisplayText = FText::Format(LOCTEXT("UncachedAssetsWithCacheFails", "{BaseMessage} {CacheFails} Blueprints failed to cache."), Args);
+			ArgsWithCacheFails.Add(TEXT("BaseMessage"), ReturnDisplayText);
+			ArgsWithCacheFails.Add(TEXT("CacheFails"), FailedToCacheCount);
+			ReturnDisplayText = FText::Format(LOCTEXT("UncachedAssetsWithCacheFails", "{BaseMessage} {CacheFails} Blueprints failed to cache."), ArgsWithCacheFails);
 		}
 	}
 
@@ -1381,10 +1406,48 @@ FText SFindInBlueprints::GetCurrentCacheBlueprintName() const
 	return FText::FromName(FFindInBlueprintSearchManager::Get().GetCurrentCacheBlueprintName());
 }
 
-void SFindInBlueprints::OnCacheComplete()
+void SFindInBlueprints::OnCacheStarted(EFiBCacheOpType InOpType)
 {
-	// Resubmit the last search, which will also remove the bar if needed
-	OnSearchTextCommitted(SearchTextField->GetText(), ETextCommit::OnEnter);
+	const bool bIsCacheBarAdded = CacheBarSlot.IsValid();
+
+	// Keep the cache bar open after a re-indexing operation if we're currently showing it (post-search unindexed/out-of-date UI context)
+	bKeepCacheBarOpenOnCacheCompletion = (InOpType == EFiBCacheOpType::CachePendingAssets) && bIsCacheBarAdded;
+
+	// Show progress during a re-indexing operation only if we have multiple assets to process (e.g. avoid showing after compiling a single asset)
+	const bool bShouldShowProgress = InOpType != EFiBCacheOpType::CachePendingAssets || FFindInBlueprintSearchManager::Get().GetNumberPendingAssets() > 1;
+
+	// Ensure that the cache bar is visible to show progress
+	if (!bIsCacheBarAdded && bShouldShowProgress)
+	{
+		ConditionallyAddCacheBar();
+	}
+}
+
+void SFindInBlueprints::OnCacheComplete(EFiBCacheOpType InOpType)
+{
+	if (InOpType == EFiBCacheOpType::CacheUnindexedAssets)
+	{
+		// Resubmit the last search, which will also remove the bar if needed
+		OnSearchTextCommitted(SearchTextField->GetText(), ETextCommit::OnEnter);
+	}
+	else if(!bKeepCacheBarOpenOnCacheCompletion)
+	{
+		const bool bIsCacheBarAdded = CacheBarSlot.IsValid();
+
+		if (bIsCacheBarAdded)
+		{
+			if (!bKeepCacheBarOpenOnCacheCompletion)
+			{
+				// Remove the cache bar
+				OnRemoveCacheBar();
+			}
+		}
+		else if(bKeepCacheBarOpenOnCacheCompletion)
+		{
+			// Restore the cache bar to its previous state (e.g. in case the user closed the progress bar while caching)
+			ConditionallyAddCacheBar();
+		}
+	}
 }
 
 TSharedPtr<SWidget> SFindInBlueprints::OnContextMenuOpening()
diff --git a/Engine/Source/Editor/Kismet/Private/SReplaceNodeReferences.cpp b/Engine/Source/Editor/Kismet/Private/SReplaceNodeReferences.cpp
--- a/Engine/Source/Editor/Kismet/Private/SReplaceNodeReferences.cpp
+++ b/Engine/Source/Editor/Kismet/Private/SReplaceNodeReferences.cpp
@@ -423,7 +423,10 @@ FReply SReplaceNodeReferences::OnFindAndReplaceAll()
 {
 	if (SelectedTargetReferenceItem.IsValid())
 	{
-		FindInBlueprints->CacheAllBlueprints(FSimpleDelegate::CreateSP(this, &SReplaceNodeReferences::OnSubmitSearchQuery, true), EFiBVersion::FIB_VER_VARIABLE_REFERENCE);
+		FFindInBlueprintCachingOptions CachingOptions;
+		CachingOptions.MinimiumVersionRequirement = EFiBVersion::FIB_VER_VARIABLE_REFERENCE;
+		CachingOptions.OnFinished = FSimpleDelegate::CreateSP(this, &SReplaceNodeReferences::OnSubmitSearchQuery, true);
+		FindInBlueprints->CacheAllBlueprints(CachingOptions);
 	}
 	return FReply::Handled();
 }
diff --git a/Engine/Source/Editor/Kismet/Public/FindInBlueprintManager.h b/Engine/Source/Editor/Kismet/Public/FindInBlueprintManager.h
--- a/Engine/Source/Editor/Kismet/Public/FindInBlueprintManager.h
+++ b/Engine/Source/Editor/Kismet/Public/FindInBlueprintManager.h
@@ -10,6 +10,7 @@
 #include "HAL/Runnable.h"
 #include "SlateFwd.h"
 #include "Input/Reply.h"
+#include "TickableEditorObject.h"
 
 struct FAssetData;
 class FFindInBlueprintsResult;
@@ -183,6 +184,26 @@ struct KISMET_API FFiBMD
 	static const FString FiBSearchableHiddenExplicitMD;
 };
 
+/** Which assets to index for caching */
+enum EFiBCacheOpType
+{
+	CachePendingAssets,
+	CacheUnindexedAssets
+};
+
+/** Options to configure the bulk caching task */
+struct FFindInBlueprintCachingOptions
+{
+	/** Type of caching operation */
+	EFiBCacheOpType OpType = EFiBCacheOpType::CachePendingAssets;
+
+	/** Callback for when caching is finished */
+	FSimpleDelegate OnFinished;
+
+	/** Minimum version requirement for caching, any Blueprints below this version will be re-indexed */
+	EFiBVersion MinimiumVersionRequirement = EFiBVersion::FIB_VER_LATEST;
+};
+
 ////////////////////////////////////
 // FFindInBlueprintsResult
 
@@ -341,7 +362,7 @@ public:
 // FFindInBlueprintSearchManager
 
 /** Singleton manager for handling all Blueprint searches, helps to manage the going progress of Blueprints, and is thread-safe. */
-class KISMET_API FFindInBlueprintSearchManager
+class KISMET_API FFindInBlueprintSearchManager : public FTickableEditorObject
 {
 public:
 	static FFindInBlueprintSearchManager* Instance;
@@ -350,6 +371,12 @@ public:
 	FFindInBlueprintSearchManager();
 	~FFindInBlueprintSearchManager();
 
+	//~ Begin FTickableObject Interface
+	virtual void Tick(float DeltaTime) override;
+	virtual bool IsTickable() const override;
+	virtual TStatId GetStatId() const override;
+	//~ End FTickableObject Interface
+
 	/**
 	 * Gathers the Blueprint's search metadata and adds or updates it in the cache
 	 *
@@ -405,18 +432,22 @@ public:
 	/** Serializes an FText to memory and converts the memory into a string of hex characters */
 	static FString ConvertFTextToHexString(FText InValue);
 
-	/** Returns the number of uncached Blueprints */
+	/** Returns the number of assets that are waiting to be re-indexed */
+	int32 GetNumberPendingAssets() const;
+
+	/** Returns the number of unindexed Blueprints, either due to not having been indexed before, or AR data being out-of-date */
+	int32 GetNumberUnindexedAssets() const;
+
+	/** Returns the number of uncached assets during an active indexing operation */
 	int32 GetNumberUncachedAssets() const;
 
 	/**
-	 * Starts caching all uncached Blueprints at a rate of 1 per tick
+	 * Starts a task to cache Blueprints at a rate of 1 per tick
 	 *
 	 * @param InSourceWidget				The source FindInBlueprints widget, this widget will be informed when caching is complete
-	 * @param InOutActiveTimerDelegate		Binds an object that ticks every time the Find-in-Blueprints widget ticks to cache all Blueprints
-	 * @param InOnFinished					Callback when caching is finished
-	 * @param InMinimiumVersionRequirement	Minimum version requirement for caching, any Blueprints below this version will be re-indexed
+	 * @param InCachingOptions				Options to configure the caching task
 	 */
-	void CacheAllUncachedAssets(TWeakPtr< class SFindInBlueprints > InSourceWidget, FWidgetActiveTimerDelegate& InOutActiveTimerDelegate, FSimpleDelegate InOnFinished = FSimpleDelegate(), EFiBVersion InMinimiumVersionRequirement = EFiBVersion::FIB_VER_LATEST);
+	void CacheAllAssets(TWeakPtr< class SFindInBlueprints > InSourceWidget, const FFindInBlueprintCachingOptions& InCachingOptions);
 	
 	/**
 	 * Starts the actual caching process
@@ -424,7 +455,7 @@ public:
 	 * @param bInSourceControlActive		TRUE if source control is active
 	 * @param bInCheckoutAndSave			TRUE if the system should checkout and save all assets that need to be reindexed
 	 */
-	void OnCacheAllUncachedAssets(bool bInSourceControlActive, bool bInCheckoutAndSave);
+	void OnCacheAllUnindexedAssets(bool bInSourceControlActive, bool bInCheckoutAndSave);
 
 	/** Stops the caching process where it currently is at, the rest can be continued later */
 	void CancelCacheAll(SFindInBlueprints* InFindInBlueprintWidget);
@@ -446,6 +477,10 @@ public:
 
 	/** Returns TRUE if caching failed */
 	bool HasCachingFailed() const { return FailedToCachePaths.Num() > 0; };
+
+	/** Callback to note that Blueprint caching is started */
+	void StartedCachingBlueprints();
+
 	/**
 	 * Callback to note that Blueprint caching is complete
 	 *
@@ -456,6 +491,12 @@ public:
 	/** Returns TRUE if Blueprints are being cached. */
 	bool IsCacheInProgress() const;
 
+	/** Returns TRUE if unindexed Blueprints are being cached (since this can block the UI) */
+	bool IsUnindexedCacheInProgress() const;
+
+	/** Returns a weak reference to the widget that initiated the current caching operation */
+	TWeakPtr<SFindInBlueprints> GetSourceCachingWidget() const { return SourceCachingWidget; }
+
 	/** Serializes an FString to memory and converts the memory into a string of hex characters */
 	static FString ConvertFStringToHexString(FString InValue);
 
@@ -505,6 +546,9 @@ private:
 	/** Callback hook from the Hot Reload manager that indicates that a module has been hot-reloaded */
 	void OnHotReload(bool bWasTriggeredAutomatically);
 
+	/** Callback from Kismet when a Blueprint asset is opened in the editor */
+	void HandleBlueprintEditorOpened(EBlueprintType InBlueprintType);
+
 	/** Helper to gathers the Blueprint's search metadata */
 	FString GatherBlueprintSearchMetadata(const UBlueprint* Blueprint);
 
@@ -541,6 +585,9 @@ protected:
 	/** Tells if gathering data is currently allowed */
 	bool bEnableGatheringData;
 
+	/** Tells if deferred indexing is configured as disabled */
+	bool bDisableDeferredIndexing;
+
 	/** Maps the Blueprint paths to their index in the SearchArray */
 	TMap<FName, int32> SearchMap;
 
@@ -571,8 +618,11 @@ protected:
 	/** FindInBlueprints widget that started the cache process */
 	TWeakPtr<SFindInBlueprints> SourceCachingWidget;
 
+	/** Asset paths that are either loaded or modified and require re-indexing before caching */
+	TSet<FName> PendingAssets;
+
 	/** Asset paths that have not been cached for searching due to lack of FiB data, this means that they are either older Blueprints, or the DDC cannot find the data */
-	TSet<FName> UncachedAssets;
+	TSet<FName> UnindexedAssets;
 
 	/** List of paths for Blueprints that failed to cache */
 	TSet<FName> FailedToCachePaths;
@@ -580,6 +630,9 @@ protected:
 	/** Tickable object that does the caching of uncached Blueprints at a rate of once per tick */
 	class FCacheAllBlueprintsTickableObject* CachingObject;
 
+	/** Stores the type of caching operation that's currently in progress */
+	EFiBCacheOpType CurrentCacheOpType;
+
 	/** Mapping between a class name and its UClass instance - used for faster look up in FFindInBlueprintSearchManager::OnAssetAdded */
 	TMap<FName, const UClass*> CachedAssetClasses;
 
diff --git a/Engine/Source/Editor/Kismet/Public/FindInBlueprints.h b/Engine/Source/Editor/Kismet/Public/FindInBlueprints.h
--- a/Engine/Source/Editor/Kismet/Public/FindInBlueprints.h
+++ b/Engine/Source/Editor/Kismet/Public/FindInBlueprints.h
@@ -217,16 +217,18 @@ public:
 	 */
 	void MakeSearchQuery(FString InSearchString, bool bInIsFindWithinBlueprint, enum ESearchQueryFilter InSearchFilterForImaginaryDataReturn = ESearchQueryFilter::AllFilter, EFiBVersion InMinimiumVersionRequirement = EFiBVersion::FIB_VER_LATEST, FOnSearchComplete InOnSearchComplete = FOnSearchComplete());
 
-	/** Called when caching Blueprints is complete, if this widget initiated the indexing */
-	void OnCacheComplete();
+	/** Called when caching Blueprints is started */
+	void OnCacheStarted(EFiBCacheOpType InOpType);
+	
+	/** Called when caching Blueprints is complete */
+	void OnCacheComplete(EFiBCacheOpType InOpType);
 
 	/**
 	 * Asynchronously caches all Blueprints below a specified version.
 	 *
-	 * @param InOnFinished						Callback when the cache process is complete.
-	 * @param InMinimiumVersionRequirement		Version that Blueprints must be below to be loaded and indexed, by default all out-of-date Blueprints
+	 * @param InOptions								Options to configure the caching task
 	 */
-	void CacheAllBlueprints(FSimpleDelegate InOnFinished = FSimpleDelegate(), EFiBVersion InMinimiumVersionRequirement = EFiBVersion::FIB_VER_LATEST);
+	void CacheAllBlueprints(const FFindInBlueprintCachingOptions& InOptions);
 
 	/** If this is a global find results widget, returns the host tab's unique ID. Otherwise, returns NAME_None. */
 	FName GetHostTabId() const;
@@ -285,14 +287,16 @@ private:
 	FReply OnRemoveCacheBar();
 
 	/** Callback to return the cache bar's display text, informing the user of the situation */
-	FText GetUncachedAssetWarningText() const;
+	FText GetUnindexedAssetWarningText() const;
 
 	/** Callback to return the cache bar's current indexing Blueprint name */
 	FText GetCurrentCacheBlueprintName() const;
 
-	/** Callback to cache all uncached Blueprints */
-	FReply OnCacheAllBlueprints();
-	FReply OnCacheAllBlueprints(FSimpleDelegate InOnFinished, EFiBVersion InMinimiumVersionRequirement = EFiBVersion::FIB_VER_LATEST);
+	/** Callback to cache all unindexed Blueprints */
+	FReply OnCacheAllUnindexedBlueprints();
+
+	/** Callback to cache all Blueprints according to the given options */
+	FReply OnCacheAllBlueprints(const FFindInBlueprintCachingOptions& InOptions);
 
 	/** Callback to cancel the caching process */
 	FReply OnCancelCacheAll();
@@ -307,7 +311,10 @@ private:
 	EVisibility GetCachingProgressBarVisiblity() const;
 
 	/** Returns the visibility of the "Cache All" button, visible when not caching, collapsed when caching is in progress */
-	EVisibility GetCacheAllButtonVisibility() const;
+	EVisibility GetCacheAllUnindexedButtonVisibility() const;
+
+	/** Returns the visibility of the "Cancel" button, visible when caching is in progress, collapsed when not caching */
+	EVisibility GetCacheAllCancelButtonVisibility() const;
 
 	/** Returns the caching bar's visibility, it goes invisible when there is nothing to be cached. The next search will remove this bar or make it visible again */
 	EVisibility GetCachingBarVisibility() const;
@@ -318,6 +325,9 @@ private:
 	/** Returns the visibility of the popup button that displays the list of Blueprints that failed to cache */
 	EVisibility GetFailedToCacheListVisibility() const;
 
+	/** Returns the visibility of the unresponsive editor warning note text in the caching progress bar */
+	EVisibility GetUnresponsiveEditorWarningVisibility() const;
+
 	/** Returns TRUE if Blueprint caching is in progress */
 	bool IsCacheInProgress() const;
 
@@ -400,6 +410,6 @@ private:
 	/** True if current search should not be changed by an external source */
 	bool bIsLocked;
 
-	/** True if the most recent search was a global search */
-	bool bHasGlobalSearchResults;
+	/** True if the cache bar should remain active after a caching operation */
+	bool bKeepCacheBarOpenOnCacheCompletion;
 };
diff --git a/Engine/Source/Runtime/Engine/Private/Blueprint.cpp b/Engine/Source/Runtime/Engine/Private/Blueprint.cpp
--- a/Engine/Source/Runtime/Engine/Private/Blueprint.cpp
+++ b/Engine/Source/Runtime/Engine/Private/Blueprint.cpp
@@ -379,8 +379,8 @@ void UBlueprint::PreSave(const class ITargetPlatform* TargetPlatform)
 
 	if (!TargetPlatform || TargetPlatform->HasEditorOnlyData())
 	{
-		// Cache the BP for use
-		FFindInBlueprintSearchManager::Get().AddOrUpdateBlueprintSearchMetadata(this);
+		// Cache the BP for use (immediate, since we're about to save)
+		FFindInBlueprintSearchManager::Get().AddOrUpdateBlueprintSearchMetadata(this, true);
 	}
 }
 #endif // WITH_EDITORONLY_DATA
