Fix assertion on attempt to index an unsaved level Blueprint asset after loading a previously-saved map.

#rnx

diff --git a/Engine/Source/Editor/Kismet/Private/FindInBlueprintManager.cpp b/Engine/Source/Editor/Kismet/Private/FindInBlueprintManager.cpp
--- a/Engine/Source/Editor/Kismet/Private/FindInBlueprintManager.cpp
+++ b/Engine/Source/Editor/Kismet/Private/FindInBlueprintManager.cpp
@@ -1228,97 +1228,99 @@ public:
 		{
 			FAssetRegistryModule* AssetRegistryModule = &FModuleManager::LoadModuleChecked<FAssetRegistryModule>(TEXT("AssetRegistry"));
 			FAssetData AssetData = AssetRegistryModule->Get().GetAssetByObjectPath(UncachedAssets[TickCacheIndex]);
-
-			const bool bIsWorldAsset = AssetData.AssetClass == UWorld::StaticClass()->GetFName();
-
-			// Construct a full package filename with path so we can query the read only status and save to disk
-			FString FinalPackageFilename = FPackageName::LongPackageNameToFilename(AssetData.PackageName.ToString());
-			if( FinalPackageFilename.Len() > 0 && FPaths::GetExtension(FinalPackageFilename).Len() == 0 )
+			if (AssetData.IsValid())
 			{
-				FinalPackageFilename += bIsWorldAsset ? FPackageName::GetMapPackageExtension() : FPackageName::GetAssetPackageExtension();
-			}
-			FText ErrorMessage;
-			bool bValidFilename = FFileHelper::IsFilenameValidForSaving( FinalPackageFilename, ErrorMessage );
-			if ( bValidFilename )
-			{
-				bValidFilename = bIsWorldAsset ? FEditorFileUtils::IsValidMapFilename( FinalPackageFilename, ErrorMessage ) : FPackageName::IsValidLongPackageName( FinalPackageFilename, false, &ErrorMessage );
-			}
+				const bool bIsWorldAsset = AssetData.AssetClass == UWorld::StaticClass()->GetFName();
 
-			bool bIsAssetReadOnlyOnDisk = IFileManager::Get().IsReadOnly( *FinalPackageFilename );
-			bool bFailedToCache = CacheParams.bCheckOutAndSave;
-
-			if (!bIsAssetReadOnlyOnDisk || !CacheParams.bCheckOutAndSave)
-			{
-				if (!FFindInBlueprintSearchManager::Get().IsUnindexedCacheInProgress())
+				// Construct a full package filename with path so we can query the read only status and save to disk
+				FString FinalPackageFilename = FPackageName::LongPackageNameToFilename(AssetData.PackageName.ToString());
+				if (FinalPackageFilename.Len() > 0 && FPaths::GetExtension(FinalPackageFilename).Len() == 0)
 				{
-					// Re-index and update the cached value for loaded Blueprint assets only
-					if (AssetData.IsAssetLoaded())
-					{
-						if (UBlueprint* LoadedBlueprintAsset = Cast<UBlueprint>(AssetData.GetAsset()))
-						{
-							FFindInBlueprintSearchManager::Get().AddOrUpdateBlueprintSearchMetadata(LoadedBlueprintAsset, true);
-						}
-					}
+					FinalPackageFilename += bIsWorldAsset ? FPackageName::GetMapPackageExtension() : FPackageName::GetAssetPackageExtension();
 				}
-				else
+				FText ErrorMessage;
+				bool bValidFilename = FFileHelper::IsFilenameValidForSaving(FinalPackageFilename, ErrorMessage);
+				if (bValidFilename)
+				{
+					bValidFilename = bIsWorldAsset ? FEditorFileUtils::IsValidMapFilename(FinalPackageFilename, ErrorMessage) : FPackageName::IsValidLongPackageName(FinalPackageFilename, false, &ErrorMessage);
+				}
+
+				bool bIsAssetReadOnlyOnDisk = IFileManager::Get().IsReadOnly(*FinalPackageFilename);
+				bool bFailedToCache = CacheParams.bCheckOutAndSave;
+
+				if (!bIsAssetReadOnlyOnDisk || !CacheParams.bCheckOutAndSave)
 				{
-					UObject* Asset = AssetData.GetAsset();
-					if (Asset && CacheParams.bCheckOutAndSave)
+					if (!FFindInBlueprintSearchManager::Get().IsUnindexedCacheInProgress())
 					{
-						if (UBlueprint* BlueprintAsset = Cast<UBlueprint>(Asset))
+						// Re-index and update the cached value for loaded Blueprint assets only
+						if (AssetData.IsAssetLoaded())
 						{
-							if (BlueprintAsset->SkeletonGeneratedClass == nullptr)
+							if (UBlueprint* LoadedBlueprintAsset = Cast<UBlueprint>(AssetData.GetAsset()))
 							{
-								// There is no skeleton class, something was wrong with the Blueprint during compile on load. This asset will be marked as failing to cache.
-								bFailedToCache = false;
+								FFindInBlueprintSearchManager::Get().AddOrUpdateBlueprintSearchMetadata(LoadedBlueprintAsset, true);
 							}
 						}
-
-						// Still good to attempt to save
-						if (bFailedToCache)
+					}
+					else
+					{
+						UObject* Asset = AssetData.GetAsset();
+						if (Asset && CacheParams.bCheckOutAndSave)
 						{
-							// Assume the package was correctly checked out from SCC
-							bool bOutPackageLocallyWritable = true;
-
-							UPackage* Package = AssetData.GetPackage();
-
-							ISourceControlProvider& SourceControlProvider = ISourceControlModule::Get().GetProvider();
-							// Trusting the SCC status in the package file cache to minimize network activity during save.
-							const FSourceControlStatePtr SourceControlState = SourceControlProvider.GetState(Package, EStateCacheUsage::Use);
-							// If the package is in the depot, and not recognized as editable by source control, and not read-only, then we know the user has made the package locally writable!
-							const bool bSCCCanEdit = !SourceControlState.IsValid() || SourceControlState->CanCheckIn() || SourceControlState->IsIgnored() || SourceControlState->IsUnknown();
-							const bool bSCCIsCheckedOut = SourceControlState.IsValid() && SourceControlState->IsCheckedOut();
-							const bool bInDepot = SourceControlState.IsValid() && SourceControlState->IsSourceControlled();
-							if (!bSCCCanEdit && bInDepot && !bIsAssetReadOnlyOnDisk && SourceControlProvider.UsesLocalReadOnlyState() && !bSCCIsCheckedOut)
+							if (UBlueprint* BlueprintAsset = Cast<UBlueprint>(Asset))
 							{
-								bOutPackageLocallyWritable = false;
+								if (BlueprintAsset->SkeletonGeneratedClass == nullptr)
+								{
+									// There is no skeleton class, something was wrong with the Blueprint during compile on load. This asset will be marked as failing to cache.
+									bFailedToCache = false;
+								}
 							}
 
-							// Save the package if the file is writable
-							if (bOutPackageLocallyWritable)
+							// Still good to attempt to save
+							if (bFailedToCache)
 							{
-								UWorld* WorldAsset = Cast<UWorld>(Asset);
-
-								// Save the package
-								EObjectFlags ObjectFlags = (WorldAsset == nullptr) ? RF_Standalone : RF_NoFlags;
+								// Assume the package was correctly checked out from SCC
+								bool bOutPackageLocallyWritable = true;
+
+								UPackage* Package = AssetData.GetPackage();
+
+								ISourceControlProvider& SourceControlProvider = ISourceControlModule::Get().GetProvider();
+								// Trusting the SCC status in the package file cache to minimize network activity during save.
+								const FSourceControlStatePtr SourceControlState = SourceControlProvider.GetState(Package, EStateCacheUsage::Use);
+								// If the package is in the depot, and not recognized as editable by source control, and not read-only, then we know the user has made the package locally writable!
+								const bool bSCCCanEdit = !SourceControlState.IsValid() || SourceControlState->CanCheckIn() || SourceControlState->IsIgnored() || SourceControlState->IsUnknown();
+								const bool bSCCIsCheckedOut = SourceControlState.IsValid() && SourceControlState->IsCheckedOut();
+								const bool bInDepot = SourceControlState.IsValid() && SourceControlState->IsSourceControlled();
+								if (!bSCCCanEdit && bInDepot && !bIsAssetReadOnlyOnDisk && SourceControlProvider.UsesLocalReadOnlyState() && !bSCCIsCheckedOut)
+								{
+									bOutPackageLocallyWritable = false;
+								}
 
-								if (GEditor->SavePackage(Package, WorldAsset, ObjectFlags, *FinalPackageFilename, GError, nullptr, false, true, SAVE_NoError))
+								// Save the package if the file is writable
+								if (bOutPackageLocallyWritable)
 								{
-									bFailedToCache = false;
+									UWorld* WorldAsset = Cast<UWorld>(Asset);
+
+									// Save the package
+									EObjectFlags ObjectFlags = (WorldAsset == nullptr) ? RF_Standalone : RF_NoFlags;
+
+									if (GEditor->SavePackage(Package, WorldAsset, ObjectFlags, *FinalPackageFilename, GError, nullptr, false, true, SAVE_NoError))
+									{
+										bFailedToCache = false;
+									}
 								}
 							}
 						}
 					}
 				}
-			}
 
-			if (bFailedToCache)
-			{
-				FailedToCacheList.Add(UncachedAssets[TickCacheIndex]);
-			}
-			else
-			{
-				CacheParams.OnCached.ExecuteIfBound(UncachedAssets[TickCacheIndex]);
+				if (bFailedToCache)
+				{
+					FailedToCacheList.Add(UncachedAssets[TickCacheIndex]);
+				}
+				else
+				{
+					CacheParams.OnCached.ExecuteIfBound(UncachedAssets[TickCacheIndex]);
+				}
 			}
 
 			++TickCacheIndex;
