Merging using Dev-Anim<->Release4.21
#DUPE MERGE of CL 4596261

Fix on deleting state from AnimBP

diff --git a/Engine/Source/Editor/Kismet/Private/BlueprintEditor.cpp b/Engine/Source/Editor/Kismet/Private/BlueprintEditor.cpp
--- a/Engine/Source/Editor/Kismet/Private/BlueprintEditor.cpp
+++ b/Engine/Source/Editor/Kismet/Private/BlueprintEditor.cpp
@@ -5540,6 +5540,18 @@ void FBlueprintEditor::DeleteSelectedNodes()
 		{
 			if (Node->CanUserDeleteNode())
 			{
+				// if it's state based anim node, make sure we close the sub graph first
+				// otherwise, we leave the orphan graph around
+				UAnimStateNodeBase* StateNode = Cast<UAnimStateNodeBase>(Node);
+				if (StateNode)
+				{
+					UEdGraph* NodeGraph = StateNode->GetBoundGraph();
+					if (NodeGraph)
+					{
+						CloseDocumentTab(NodeGraph);
+					}
+				}
+
 				UK2Node* K2Node = Cast<UK2Node>(Node);
 				if (K2Node != NULL && K2Node->NodeCausesStructuralBlueprintChange())
 				{
