Merging //UE4/Dev-Main to Dev-Anim (//UE4/Dev-Anim)

diff --git a/Engine/Plugins/Runtime/Database/SQLiteCore/Source/SQLiteCore/SQLiteCore.Build.cs b/Engine/Plugins/Runtime/Database/SQLiteCore/Source/SQLiteCore/SQLiteCore.Build.cs
--- a/Engine/Plugins/Runtime/Database/SQLiteCore/Source/SQLiteCore/SQLiteCore.Build.cs
+++ b/Engine/Plugins/Runtime/Database/SQLiteCore/Source/SQLiteCore/SQLiteCore.Build.cs
@@ -20,6 +20,9 @@ namespace UnrealBuildTool.Rules
 				PrivateDefinitions.Add("SQLITE_THREADSAFE=0"); // No threading on HTML5
 			}
 
+			// Use the math.h version of isnan rather than the SQLite version to avoid a -ffast-math error
+			PrivateDefinitions.Add("SQLITE_HAVE_ISNAN=1");
+
 			// Enable SQLite debug checks?
 			//PrivateDefinitions.Add("SQLITE_DEBUG");
 
diff --git a/Engine/Source/Editor/Kismet/Private/FindInBlueprintManager.cpp b/Engine/Source/Editor/Kismet/Private/FindInBlueprintManager.cpp
--- a/Engine/Source/Editor/Kismet/Private/FindInBlueprintManager.cpp
+++ b/Engine/Source/Editor/Kismet/Private/FindInBlueprintManager.cpp
@@ -1228,97 +1228,99 @@ public:
 		{
 			FAssetRegistryModule* AssetRegistryModule = &FModuleManager::LoadModuleChecked<FAssetRegistryModule>(TEXT("AssetRegistry"));
 			FAssetData AssetData = AssetRegistryModule->Get().GetAssetByObjectPath(UncachedAssets[TickCacheIndex]);
-
-			const bool bIsWorldAsset = AssetData.AssetClass == UWorld::StaticClass()->GetFName();
-
-			// Construct a full package filename with path so we can query the read only status and save to disk
-			FString FinalPackageFilename = FPackageName::LongPackageNameToFilename(AssetData.PackageName.ToString());
-			if( FinalPackageFilename.Len() > 0 && FPaths::GetExtension(FinalPackageFilename).Len() == 0 )
+			if (AssetData.IsValid())
 			{
-				FinalPackageFilename += bIsWorldAsset ? FPackageName::GetMapPackageExtension() : FPackageName::GetAssetPackageExtension();
-			}
-			FText ErrorMessage;
-			bool bValidFilename = FFileHelper::IsFilenameValidForSaving( FinalPackageFilename, ErrorMessage );
-			if ( bValidFilename )
-			{
-				bValidFilename = bIsWorldAsset ? FEditorFileUtils::IsValidMapFilename( FinalPackageFilename, ErrorMessage ) : FPackageName::IsValidLongPackageName( FinalPackageFilename, false, &ErrorMessage );
-			}
+				const bool bIsWorldAsset = AssetData.AssetClass == UWorld::StaticClass()->GetFName();
 
-			bool bIsAssetReadOnlyOnDisk = IFileManager::Get().IsReadOnly( *FinalPackageFilename );
-			bool bFailedToCache = CacheParams.bCheckOutAndSave;
-
-			if (!bIsAssetReadOnlyOnDisk || !CacheParams.bCheckOutAndSave)
-			{
-				if (!FFindInBlueprintSearchManager::Get().IsUnindexedCacheInProgress())
+				// Construct a full package filename with path so we can query the read only status and save to disk
+				FString FinalPackageFilename = FPackageName::LongPackageNameToFilename(AssetData.PackageName.ToString());
+				if (FinalPackageFilename.Len() > 0 && FPaths::GetExtension(FinalPackageFilename).Len() == 0)
 				{
-					// Re-index and update the cached value for loaded Blueprint assets only
-					if (AssetData.IsAssetLoaded())
-					{
-						if (UBlueprint* LoadedBlueprintAsset = Cast<UBlueprint>(AssetData.GetAsset()))
-						{
-							FFindInBlueprintSearchManager::Get().AddOrUpdateBlueprintSearchMetadata(LoadedBlueprintAsset, true);
-						}
-					}
+					FinalPackageFilename += bIsWorldAsset ? FPackageName::GetMapPackageExtension() : FPackageName::GetAssetPackageExtension();
 				}
-				else
+				FText ErrorMessage;
+				bool bValidFilename = FFileHelper::IsFilenameValidForSaving(FinalPackageFilename, ErrorMessage);
+				if (bValidFilename)
+				{
+					bValidFilename = bIsWorldAsset ? FEditorFileUtils::IsValidMapFilename(FinalPackageFilename, ErrorMessage) : FPackageName::IsValidLongPackageName(FinalPackageFilename, false, &ErrorMessage);
+				}
+
+				bool bIsAssetReadOnlyOnDisk = IFileManager::Get().IsReadOnly(*FinalPackageFilename);
+				bool bFailedToCache = CacheParams.bCheckOutAndSave;
+
+				if (!bIsAssetReadOnlyOnDisk || !CacheParams.bCheckOutAndSave)
 				{
-					UObject* Asset = AssetData.GetAsset();
-					if (Asset && CacheParams.bCheckOutAndSave)
+					if (!FFindInBlueprintSearchManager::Get().IsUnindexedCacheInProgress())
 					{
-						if (UBlueprint* BlueprintAsset = Cast<UBlueprint>(Asset))
+						// Re-index and update the cached value for loaded Blueprint assets only
+						if (AssetData.IsAssetLoaded())
 						{
-							if (BlueprintAsset->SkeletonGeneratedClass == nullptr)
+							if (UBlueprint* LoadedBlueprintAsset = Cast<UBlueprint>(AssetData.GetAsset()))
 							{
-								// There is no skeleton class, something was wrong with the Blueprint during compile on load. This asset will be marked as failing to cache.
-								bFailedToCache = false;
+								FFindInBlueprintSearchManager::Get().AddOrUpdateBlueprintSearchMetadata(LoadedBlueprintAsset, true);
 							}
 						}
-
-						// Still good to attempt to save
-						if (bFailedToCache)
+					}
+					else
+					{
+						UObject* Asset = AssetData.GetAsset();
+						if (Asset && CacheParams.bCheckOutAndSave)
 						{
-							// Assume the package was correctly checked out from SCC
-							bool bOutPackageLocallyWritable = true;
-
-							UPackage* Package = AssetData.GetPackage();
-
-							ISourceControlProvider& SourceControlProvider = ISourceControlModule::Get().GetProvider();
-							// Trusting the SCC status in the package file cache to minimize network activity during save.
-							const FSourceControlStatePtr SourceControlState = SourceControlProvider.GetState(Package, EStateCacheUsage::Use);
-							// If the package is in the depot, and not recognized as editable by source control, and not read-only, then we know the user has made the package locally writable!
-							const bool bSCCCanEdit = !SourceControlState.IsValid() || SourceControlState->CanCheckIn() || SourceControlState->IsIgnored() || SourceControlState->IsUnknown();
-							const bool bSCCIsCheckedOut = SourceControlState.IsValid() && SourceControlState->IsCheckedOut();
-							const bool bInDepot = SourceControlState.IsValid() && SourceControlState->IsSourceControlled();
-							if (!bSCCCanEdit && bInDepot && !bIsAssetReadOnlyOnDisk && SourceControlProvider.UsesLocalReadOnlyState() && !bSCCIsCheckedOut)
+							if (UBlueprint* BlueprintAsset = Cast<UBlueprint>(Asset))
 							{
-								bOutPackageLocallyWritable = false;
+								if (BlueprintAsset->SkeletonGeneratedClass == nullptr)
+								{
+									// There is no skeleton class, something was wrong with the Blueprint during compile on load. This asset will be marked as failing to cache.
+									bFailedToCache = false;
+								}
 							}
 
-							// Save the package if the file is writable
-							if (bOutPackageLocallyWritable)
+							// Still good to attempt to save
+							if (bFailedToCache)
 							{
-								UWorld* WorldAsset = Cast<UWorld>(Asset);
-
-								// Save the package
-								EObjectFlags ObjectFlags = (WorldAsset == nullptr) ? RF_Standalone : RF_NoFlags;
+								// Assume the package was correctly checked out from SCC
+								bool bOutPackageLocallyWritable = true;
+
+								UPackage* Package = AssetData.GetPackage();
+
+								ISourceControlProvider& SourceControlProvider = ISourceControlModule::Get().GetProvider();
+								// Trusting the SCC status in the package file cache to minimize network activity during save.
+								const FSourceControlStatePtr SourceControlState = SourceControlProvider.GetState(Package, EStateCacheUsage::Use);
+								// If the package is in the depot, and not recognized as editable by source control, and not read-only, then we know the user has made the package locally writable!
+								const bool bSCCCanEdit = !SourceControlState.IsValid() || SourceControlState->CanCheckIn() || SourceControlState->IsIgnored() || SourceControlState->IsUnknown();
+								const bool bSCCIsCheckedOut = SourceControlState.IsValid() && SourceControlState->IsCheckedOut();
+								const bool bInDepot = SourceControlState.IsValid() && SourceControlState->IsSourceControlled();
+								if (!bSCCCanEdit && bInDepot && !bIsAssetReadOnlyOnDisk && SourceControlProvider.UsesLocalReadOnlyState() && !bSCCIsCheckedOut)
+								{
+									bOutPackageLocallyWritable = false;
+								}
 
-								if (GEditor->SavePackage(Package, WorldAsset, ObjectFlags, *FinalPackageFilename, GError, nullptr, false, true, SAVE_NoError))
+								// Save the package if the file is writable
+								if (bOutPackageLocallyWritable)
 								{
-									bFailedToCache = false;
+									UWorld* WorldAsset = Cast<UWorld>(Asset);
+
+									// Save the package
+									EObjectFlags ObjectFlags = (WorldAsset == nullptr) ? RF_Standalone : RF_NoFlags;
+
+									if (GEditor->SavePackage(Package, WorldAsset, ObjectFlags, *FinalPackageFilename, GError, nullptr, false, true, SAVE_NoError))
+									{
+										bFailedToCache = false;
+									}
 								}
 							}
 						}
 					}
 				}
-			}
 
-			if (bFailedToCache)
-			{
-				FailedToCacheList.Add(UncachedAssets[TickCacheIndex]);
-			}
-			else
-			{
-				CacheParams.OnCached.ExecuteIfBound(UncachedAssets[TickCacheIndex]);
+				if (bFailedToCache)
+				{
+					FailedToCacheList.Add(UncachedAssets[TickCacheIndex]);
+				}
+				else
+				{
+					CacheParams.OnCached.ExecuteIfBound(UncachedAssets[TickCacheIndex]);
+				}
 			}
 
 			++TickCacheIndex;
diff --git a/Engine/Source/Programs/UnrealBuildTool/Platform/IOS/IOSToolChain.cs b/Engine/Source/Programs/UnrealBuildTool/Platform/IOS/IOSToolChain.cs
--- a/Engine/Source/Programs/UnrealBuildTool/Platform/IOS/IOSToolChain.cs
+++ b/Engine/Source/Programs/UnrealBuildTool/Platform/IOS/IOSToolChain.cs
@@ -233,15 +233,6 @@ namespace UnrealBuildTool
 				Result += " -fno-exceptions";
 			}
 
-			if (CompileEnvironment.bEnableObjCExceptions)
-			{
-				Result += " -fobjc-exceptions";
-			}
-			else
-			{
-				Result += " -fno-objc-exceptions";
-			}
-
 			string SanitizerMode = Environment.GetEnvironmentVariable("ENABLE_ADDRESS_SANITIZER");
 			if(SanitizerMode != null && SanitizerMode == "YES")
 			{
@@ -417,6 +408,23 @@ namespace UnrealBuildTool
 			return Result;
 		}
 
+		// Conditionally enable (default disabled) Objective-C exceptions
+		static string GetObjCExceptionsFlag(CppCompileEnvironment CompileEnvironment)
+		{
+			string Result = "";
+
+			if (CompileEnvironment.bEnableObjCExceptions)
+			{
+				Result += " -fobjc-exceptions";
+			}
+			else
+			{
+				Result += " -fno-objc-exceptions";
+			}
+
+			return Result;
+		}
+
 		static FileReference GetFrameworkZip(UEBuildFramework Framework)
 		{
 			if (Framework.OwningModule == null)
@@ -623,6 +631,7 @@ namespace UnrealBuildTool
 					// Compile the file as a C++ PCH.
 					FileArguments += GetCompileArguments_PCH();
 					FileArguments += GetRTTIFlag(CompileEnvironment);
+					FileArguments += GetObjCExceptionsFlag(CompileEnvironment);
 				}
 				else if (Extension == ".C")
 				{
@@ -634,17 +643,20 @@ namespace UnrealBuildTool
 					// Compile the file as Objective-C++ code.
 					FileArguments += GetCompileArguments_MM();
 					FileArguments += GetRTTIFlag(CompileEnvironment);
+					FileArguments += GetObjCExceptionsFlag(CompileEnvironment);
 				}
 				else if (Extension == ".M")
 				{
-					// Compile the file as Objective-C++ code.
+					// Compile the file as Objective-C code.
 					FileArguments += GetCompileArguments_M();
+					FileArguments += GetObjCExceptionsFlag(CompileEnvironment);
 				}
 				else
 				{
 					// Compile the file as C++ code.
 					FileArguments += GetCompileArguments_CPP();
 					FileArguments += GetRTTIFlag(CompileEnvironment);
+					FileArguments += GetObjCExceptionsFlag(CompileEnvironment);
 
 					// only use PCH for .cpp files
 					FilePCHArguments = PCHArguments;
