Fix a regression that could lead to repeated attempts to index unregistered temp assets after a PIE session.

#rnx

diff --git a/Engine/Source/Editor/Kismet/Private/FindInBlueprintManager.cpp b/Engine/Source/Editor/Kismet/Private/FindInBlueprintManager.cpp
--- a/Engine/Source/Editor/Kismet/Private/FindInBlueprintManager.cpp
+++ b/Engine/Source/Editor/Kismet/Private/FindInBlueprintManager.cpp
@@ -1322,6 +1322,10 @@ public:
 					CacheParams.OnCached.ExecuteIfBound(UncachedAssets[TickCacheIndex]);
 				}
 			}
+			else
+			{
+				FailedToCacheList.Add(UncachedAssets[TickCacheIndex]);
+			}
 
 			++TickCacheIndex;
 
@@ -2354,7 +2358,7 @@ void FFindInBlueprintSearchManager::FinishedCachingBlueprints(int32 InNumberCach
 	// Update the list of cache failures
 	FailedToCachePaths = InFailedToCacheList;
 
-	// Remove any failed attempts from the pending list (should only apply to unindexed caching, but just in case)
+	// Remove any failed attempts from the pending list
 	if (InFailedToCacheList.Num() > 0)
 	{
 		PendingAssets = PendingAssets.Difference(InFailedToCacheList);
