Fixing cook errors caused by de-globalization of UObject load state variables

diff --git a/Engine/Plugins/Runtime/GameplayAbilities/Source/GameplayAbilities/Private/GameplayCueManager.cpp b/Engine/Plugins/Runtime/GameplayAbilities/Source/GameplayAbilities/Private/GameplayCueManager.cpp
--- a/Engine/Plugins/Runtime/GameplayAbilities/Source/GameplayAbilities/Private/GameplayCueManager.cpp
+++ b/Engine/Plugins/Runtime/GameplayAbilities/Source/GameplayAbilities/Private/GameplayCueManager.cpp
@@ -690,7 +690,7 @@ static void SearchDynamicClassCues(const FName PropertyName, const TArray<FStrin
 				AssetsToLoad.Add(StringRef);
 
 				// Make sure core knows about this ref so it can be properly detected during cook.
-				StringRef.PostLoadPath();
+				StringRef.PostLoadPath(nullptr);
 			}
 			else
 			{
@@ -873,7 +873,7 @@ void UGameplayCueManager::BuildCuesToAddToGlobalSet(const TArray<FAssetData>& As
 				OutAssetsToLoad.Add(StringRef);
 
 				// Make sure core knows about this ref so it can be properly detected during cook.
-				StringRef.PostLoadPath();
+				StringRef.PostLoadPath(GetLinker());
 			}
 			else
 			{
@@ -1009,7 +1009,7 @@ void UGameplayCueManager::HandleAssetAdded(UObject *Object)
 				}
 
 				// Make sure core knows about this ref so it can be properly detected during cook.
-				StringRef.PostLoadPath();
+				StringRef.PostLoadPath(Object->GetLinker());
 
 				for (UGameplayCueSet* Set : GetGlobalCueSets())
 				{
diff --git a/Engine/Source/Editor/Kismet/Private/BlueprintCompilationManager.cpp b/Engine/Source/Editor/Kismet/Private/BlueprintCompilationManager.cpp
--- a/Engine/Source/Editor/Kismet/Private/BlueprintCompilationManager.cpp
+++ b/Engine/Source/Editor/Kismet/Private/BlueprintCompilationManager.cpp
@@ -89,12 +89,12 @@ struct FBlueprintCompilationManagerImpl : public FGCObject
 
 	void QueueForCompilation(const FBPCompileRequest& CompileJob);
 	void CompileSynchronouslyImpl(const FBPCompileRequest& Request);
-	void FlushCompilationQueueImpl(bool bSuppressBroadcastCompiled, TArray<UBlueprint*>* BlueprintsCompiled, TArray<UBlueprint*>* BlueprintsCompiledOrSkeletonCompiled);
+	void FlushCompilationQueueImpl(bool bSuppressBroadcastCompiled, TArray<UBlueprint*>* BlueprintsCompiled, TArray<UBlueprint*>* BlueprintsCompiledOrSkeletonCompiled, FUObjectSerializeContext* InLoadContext);
 	void FlushReinstancingQueueImpl();
 	bool HasBlueprintsToCompile() const;
 	bool IsGeneratedClassLayoutReady() const;
 	void GetDefaultValue(const UClass* ForClass, const UProperty* Property, FString& OutDefaultValueAsString) const;
-	static void ReinstanceBatch(TArray<FReinstancingJob>& Reinstancers, TMap< UClass*, UClass* >& InOutOldToNewClassMap);
+	static void ReinstanceBatch(TArray<FReinstancingJob>& Reinstancers, TMap< UClass*, UClass* >& InOutOldToNewClassMap, FUObjectSerializeContext* InLoadContext);
 	static UClass* FastGenerateSkeletonClass(UBlueprint* BP, FKismetCompilerContext& CompilerContext, bool bIsSkeletonOnly);
 	static bool IsQueuedForCompilation(UBlueprint* BP);
 	
@@ -198,7 +198,7 @@ void FBlueprintCompilationManagerImpl::CompileSynchronouslyImpl(const FBPCompile
 	const bool bSuppressBroadcastCompiled = true;
 	TArray<UBlueprint*> CompiledBlueprints;
 	TArray<UBlueprint*> SkeletonCompiledBlueprints;
-	FlushCompilationQueueImpl(bSuppressBroadcastCompiled, &CompiledBlueprints, &SkeletonCompiledBlueprints);
+	FlushCompilationQueueImpl(bSuppressBroadcastCompiled, &CompiledBlueprints, &SkeletonCompiledBlueprints, nullptr);
 	FlushReinstancingQueueImpl();
 	
 	if (FBlueprintEditorUtils::IsLevelScriptBlueprint(Request.BPToCompile) && !bRegenerateSkeletonOnly)
@@ -354,7 +354,7 @@ struct FReinstancingJob
 	TSharedPtr<FKismetCompilerContext> Compiler;
 };
 
-void FBlueprintCompilationManagerImpl::FlushCompilationQueueImpl(bool bSuppressBroadcastCompiled, TArray<UBlueprint*>* BlueprintsCompiled, TArray<UBlueprint*>* BlueprintsCompiledOrSkeletonCompiled)
+void FBlueprintCompilationManagerImpl::FlushCompilationQueueImpl(bool bSuppressBroadcastCompiled, TArray<UBlueprint*>* BlueprintsCompiled, TArray<UBlueprint*>* BlueprintsCompiledOrSkeletonCompiled, FUObjectSerializeContext* InLoadContext)
 {
 	TGuardValue<bool> GuardTemplateNameFlag(GCompilingBlueprint, true);
 	ensure(bGeneratedClassLayoutReady);
@@ -1087,7 +1087,7 @@ void FBlueprintCompilationManagerImpl::FlushCompilationQueueImpl(bool bSuppressB
 			}
 
 			FScopedDurationTimer ReinstTimer(GTimeReinstancing);
-			ReinstanceBatch(Reinstancers, ClassesToReinstance);
+			ReinstanceBatch(Reinstancers, ClassesToReinstance, InLoadContext);
 
 			// We purposefully do not remove the OldCDOs yet, need to keep them in memory past first GC
 		}
@@ -1325,7 +1325,7 @@ void FBlueprintCompilationManagerImpl::GetDefaultValue(const UClass* ForClass, c
 	}
 }
 
-void FBlueprintCompilationManagerImpl::ReinstanceBatch(TArray<FReinstancingJob>& Reinstancers, TMap< UClass*, UClass* >& InOutOldToNewClassMap)
+void FBlueprintCompilationManagerImpl::ReinstanceBatch(TArray<FReinstancingJob>& Reinstancers, TMap< UClass*, UClass* >& InOutOldToNewClassMap, FUObjectSerializeContext* InLoadContext)
 {
 	const auto FilterOutOfDateClasses = [](TArray<UClass*>& ClassList)
 	{
@@ -1527,7 +1527,7 @@ void FBlueprintCompilationManagerImpl::ReinstanceBatch(TArray<FReinstancingJob>&
 
 					if(OldCDOIndex != INDEX_NONE)
 					{
-						FBlueprintEditorUtils::PatchNewCDOIntoLinker(CurrentBP->GeneratedClass->ClassDefaultObject, CurrentLinker, OldCDOIndex);
+						FBlueprintEditorUtils::PatchNewCDOIntoLinker(CurrentBP->GeneratedClass->ClassDefaultObject, CurrentLinker, OldCDOIndex, InLoadContext);
 						FBlueprintEditorUtils::PatchCDOSubobjectsIntoExport(OldCDO, CurrentBP->GeneratedClass->ClassDefaultObject);
 					}
 				}
@@ -2348,11 +2348,11 @@ void FBlueprintCompilationManager::Shutdown()
 }
 
 // Forward to impl:
-void FBlueprintCompilationManager::FlushCompilationQueue()
+void FBlueprintCompilationManager::FlushCompilationQueue(FUObjectSerializeContext* InLoadContext)
 {
 	if(BPCMImpl)
 	{
-		BPCMImpl->FlushCompilationQueueImpl(false, nullptr, nullptr);
+		BPCMImpl->FlushCompilationQueueImpl(false, nullptr, nullptr, InLoadContext);
 
 		// We can't support save on compile or keeping old CDOs from GCing when reinstancing is deferred:
 		BPCMImpl->CompiledBlueprintsToSave.Empty();
@@ -2364,7 +2364,7 @@ void FBlueprintCompilationManager::FlushCompilationQueueAndReinstance()
 {
 	if(BPCMImpl)
 	{
-		BPCMImpl->FlushCompilationQueueImpl(false, nullptr, nullptr);
+		BPCMImpl->FlushCompilationQueueImpl(false, nullptr, nullptr, nullptr);
 		BPCMImpl->FlushReinstancingQueueImpl();
 
 		BPCMImpl->OldCDOs.Empty();
diff --git a/Engine/Source/Editor/Kismet/Public/BlueprintCompilationManager.h b/Engine/Source/Editor/Kismet/Public/BlueprintCompilationManager.h
--- a/Engine/Source/Editor/Kismet/Public/BlueprintCompilationManager.h
+++ b/Engine/Source/Editor/Kismet/Public/BlueprintCompilationManager.h
@@ -41,7 +41,7 @@ struct KISMET_API FBlueprintCompilationManager
 	 * of objects the compilation manager has loaded. The compilation manager
 	 * will not load data while processing the compilation queue)
 	 */
-	static void FlushCompilationQueue();
+	static void FlushCompilationQueue(FUObjectSerializeContext* InLoadContext);
 	
 	/**
 	 * Flushes the compilation queue and finishes reinstancing
diff --git a/Engine/Source/Editor/KismetCompiler/Private/KismetCompiler.cpp b/Engine/Source/Editor/KismetCompiler/Private/KismetCompiler.cpp
--- a/Engine/Source/Editor/KismetCompiler/Private/KismetCompiler.cpp
+++ b/Engine/Source/Editor/KismetCompiler/Private/KismetCompiler.cpp
@@ -4169,7 +4169,7 @@ void FKismetCompilerContext::CompileFunctions(EInternalCompilerFlags InternalFla
 					if (OldLinker && OldGenLinkerIdx != INDEX_NONE)
 					{
 						// If we have a list of objects that are loading, patch our export table. This also fixes up load flags
-						FBlueprintEditorUtils::PatchNewCDOIntoLinker(Blueprint->GeneratedClass->GetDefaultObject(), OldLinker, OldGenLinkerIdx);
+						FBlueprintEditorUtils::PatchNewCDOIntoLinker(Blueprint->GeneratedClass->GetDefaultObject(), OldLinker, OldGenLinkerIdx, nullptr);
 					}
 
 					UEditorEngine::FCopyPropertiesForUnrelatedObjectsParams CopyDetails;
diff --git a/Engine/Source/Editor/UnrealEd/Private/Kismet2/BlueprintEditorUtils.cpp b/Engine/Source/Editor/UnrealEd/Private/Kismet2/BlueprintEditorUtils.cpp
--- a/Engine/Source/Editor/UnrealEd/Private/Kismet2/BlueprintEditorUtils.cpp
+++ b/Engine/Source/Editor/UnrealEd/Private/Kismet2/BlueprintEditorUtils.cpp
@@ -716,7 +716,7 @@ void FBlueprintEditorUtils::PreloadConstructionScript(UBlueprint* Blueprint)
 	}
 }
 
-void FBlueprintEditorUtils::PatchNewCDOIntoLinker(UObject* CDO, FLinkerLoad* Linker, int32 ExportIndex)
+void FBlueprintEditorUtils::PatchNewCDOIntoLinker(UObject* CDO, FLinkerLoad* Linker, int32 ExportIndex, FUObjectSerializeContext* InLoadContext)
 {
 	if( (CDO != nullptr) && (Linker != nullptr) && (ExportIndex != INDEX_NONE) )
 	{
@@ -731,18 +731,15 @@ void FBlueprintEditorUtils::PatchNewCDOIntoLinker(UObject* CDO, FLinkerLoad* Lin
 			// Copy flags from the old CDO.
 			CDO->SetFlags(OldObjectFlags);
 
-			FUObjectSerializeContext* LoadContext = Linker->GetSerializeContext();
-			TArray<UObject*>& ObjLoaded = LoadContext->GetObjectsLoaded();
+			FUObjectSerializeContext* LoadContext = InLoadContext ? InLoadContext : Linker->GetSerializeContext();
 
 			// Make sure the new CDO gets PostLoad called on it, so either add it to ObjLoaded list, or replace it if already present.
-			int32 ObjLoadedIdx = ObjLoaded.Find(OldCDO);
-			if (ObjLoadedIdx != INDEX_NONE)
+			if (!LoadContext->PRIVATE_PatchNewObjectIntoExport(OldCDO, CDO))
 			{
-				ObjLoaded[ObjLoadedIdx] = CDO;
-			}
-			else if (OldObjectFlags & RF_NeedPostLoad)
-			{
-				ObjLoaded.Add(CDO);
+				if (OldObjectFlags & RF_NeedPostLoad)
+				{
+					LoadContext->AddLoadedObject(CDO);
+				}
 			}
 		}
 
@@ -1457,11 +1454,11 @@ UClass* FBlueprintEditorUtils::RegenerateBlueprintClass(UBlueprint* Blueprint, U
 		// Patch the new CDOs to the old indices in the linker
 		if( Blueprint->SkeletonGeneratedClass )
 		{
-			PatchNewCDOIntoLinker(Blueprint->SkeletonGeneratedClass->GetDefaultObject(), OldLinker, OldSkelLinkerIdx);
+			PatchNewCDOIntoLinker(Blueprint->SkeletonGeneratedClass->GetDefaultObject(), OldLinker, OldSkelLinkerIdx, nullptr);
 		}
 		if( Blueprint->GeneratedClass )
 		{
-			PatchNewCDOIntoLinker(Blueprint->GeneratedClass->GetDefaultObject(), OldLinker, OldGenLinkerIdx);
+			PatchNewCDOIntoLinker(Blueprint->GeneratedClass->GetDefaultObject(), OldLinker, OldGenLinkerIdx, nullptr);
 		}
 
 		// Success or failure, there's no point in trying to recompile this class again when other objects reference it
diff --git a/Engine/Source/Editor/UnrealEd/Public/Kismet2/BlueprintEditorUtils.h b/Engine/Source/Editor/UnrealEd/Public/Kismet2/BlueprintEditorUtils.h
--- a/Engine/Source/Editor/UnrealEd/Public/Kismet2/BlueprintEditorUtils.h
+++ b/Engine/Source/Editor/UnrealEd/Public/Kismet2/BlueprintEditorUtils.h
@@ -183,7 +183,7 @@ public:
 	/** 
 	 * Helper function to patch the new CDO into the linker where the old one existed 
 	 */
-	static void PatchNewCDOIntoLinker(UObject* CDO, FLinkerLoad* Linker, int32 ExportIndex);
+	static void PatchNewCDOIntoLinker(UObject* CDO, FLinkerLoad* Linker, int32 ExportIndex, FUObjectSerializeContext* InLoadContext);
 
 	/** 
 	 * Procedure used to remove old function implementations and child properties from data only blueprints.
diff --git a/Engine/Source/Runtime/CoreUObject/Private/Blueprint/BlueprintSupport.cpp b/Engine/Source/Runtime/CoreUObject/Private/Blueprint/BlueprintSupport.cpp
--- a/Engine/Source/Runtime/CoreUObject/Private/Blueprint/BlueprintSupport.cpp
+++ b/Engine/Source/Runtime/CoreUObject/Private/Blueprint/BlueprintSupport.cpp
@@ -400,8 +400,9 @@ void FBlueprintSupport::ValidateNoExternalRefsToSkeletons()
 UClass* FScopedClassDependencyGather::BatchMasterClass = NULL;
 TArray<UClass*> FScopedClassDependencyGather::BatchClassDependencies;
 
-FScopedClassDependencyGather::FScopedClassDependencyGather(UClass* ClassToGather)
+FScopedClassDependencyGather::FScopedClassDependencyGather(UClass* ClassToGather, FUObjectSerializeContext* InLoadContext)
 	: bMasterClass(false)
+	, LoadContext(InLoadContext)
 {
 	// Do NOT track duplication dependencies, as these are intermediate products that we don't care about
 	if( !GIsDuplicatingClassForReinstancing )
@@ -430,9 +431,9 @@ FScopedClassDependencyGather::~FScopedClassDependencyGather()
 		auto DependencyIter = BatchClassDependencies.CreateIterator();
 		// implemented as a lambda, to prevent duplicated code between 
 		// BatchMasterClass and BatchClassDependencies entries
-		auto RecompileClassLambda = [&DependencyIter](UClass* Class)
+		auto RecompileClassLambda = [&DependencyIter](UClass* Class, FUObjectSerializeContext* InLoadContext)
 		{
-			Class->ConditionalRecompileClass();
+			Class->ConditionalRecompileClass(InLoadContext);
 
 			// because of the above call to ConditionalRecompileClass(), the 
 			// specified Class gets "cleaned and sanitized" (meaning its old 
@@ -466,16 +467,16 @@ FScopedClassDependencyGather::~FScopedClassDependencyGather()
 				UClass* Dependency = *DependencyIter;
 				if( Dependency->ClassGeneratedBy != BatchMasterClass->ClassGeneratedBy )
 				{
-					RecompileClassLambda(Dependency);
+					RecompileClassLambda(Dependency, LoadContext);
 				}
 			}
 
 			// Finally, recompile the master class to make sure it gets updated too
-			RecompileClassLambda(BatchMasterClass);
+			RecompileClassLambda(BatchMasterClass, LoadContext);
 		}
 		else
 		{
-			BatchMasterClass->ConditionalRecompileClass();
+			BatchMasterClass->ConditionalRecompileClass(LoadContext);
 		}
 
 		BatchMasterClass = NULL;
diff --git a/Engine/Source/Runtime/CoreUObject/Private/Misc/RedirectCollector.cpp b/Engine/Source/Runtime/CoreUObject/Private/Misc/RedirectCollector.cpp
--- a/Engine/Source/Runtime/CoreUObject/Private/Misc/RedirectCollector.cpp
+++ b/Engine/Source/Runtime/CoreUObject/Private/Misc/RedirectCollector.cpp
@@ -17,7 +17,7 @@
 
 DEFINE_LOG_CATEGORY_STATIC(LogRedirectors, Log, All);
 
-void FRedirectCollector::OnSoftObjectPathLoaded(const FSoftObjectPath& InPath)
+void FRedirectCollector::OnSoftObjectPathLoaded(const FSoftObjectPath& InPath, FArchive* InArchive)
 {
 	if (InPath.IsNull() || !GIsEditor)
 	{
@@ -32,7 +32,7 @@ void FRedirectCollector::OnSoftObjectPathLoaded(const FSoftObjectPath& InPath)
 	ESoftObjectPathCollectType CollectType = ESoftObjectPathCollectType::AlwaysCollect;
 	ESoftObjectPathSerializeType SerializeType = ESoftObjectPathSerializeType::AlwaysSerialize;
 
-	ThreadContext.GetSerializationOptions(PackageName, PropertyName, CollectType, SerializeType);
+	ThreadContext.GetSerializationOptions(PackageName, PropertyName, CollectType, SerializeType, InArchive);
 
 	if (CollectType == ESoftObjectPathCollectType::NeverCollect)
 	{
@@ -59,7 +59,7 @@ void FRedirectCollector::OnSoftObjectPathLoaded(const FSoftObjectPath& InPath)
 void FRedirectCollector::OnStringAssetReferenceLoaded(const FString& InString)
 {
 	FSoftObjectPath Path(InString);
-	OnSoftObjectPathLoaded(Path);
+	OnSoftObjectPathLoaded(Path, nullptr);
 }
 
 FString FRedirectCollector::OnStringAssetReferenceSaved(const FString& InString)
diff --git a/Engine/Source/Runtime/CoreUObject/Private/Serialization/AsyncLoading.cpp b/Engine/Source/Runtime/CoreUObject/Private/Serialization/AsyncLoading.cpp
--- a/Engine/Source/Runtime/CoreUObject/Private/Serialization/AsyncLoading.cpp
+++ b/Engine/Source/Runtime/CoreUObject/Private/Serialization/AsyncLoading.cpp
@@ -4096,8 +4096,7 @@ void FAsyncPackage::Event_StartPostload()
 	AsyncPackageLoadingState = EAsyncPackageLoadingState::PostLoad_Etc;
 	EventDrivenLoadingComplete();
 	{
-		TArray<UObject*>& ObjLoaded = LoadContext->GetObjectsLoaded();
-		ObjLoaded.Reserve(ObjLoaded.Num() + Linker->ExportMap.Num());
+		LoadContext->ReserveObjectsLoaded(LoadContext->GetNumObjectsLoaded() + Linker->ExportMap.Num());
 		for (int32 LocalExportIndex = 0; LocalExportIndex < Linker->ExportMap.Num(); LocalExportIndex++)
 		{
 			FObjectExport& Export = Linker->ExportMap[LocalExportIndex];
@@ -4106,7 +4105,7 @@ void FAsyncPackage::Event_StartPostload()
 			if (Object && (Object->HasAnyFlags(RF_NeedPostLoad) || Linker->bDynamicClassLinker || Object->HasAnyInternalFlags(EInternalObjectFlags::AsyncLoading)))
 			{
 				check(Object->IsValidLowLevelFast());
-				ObjLoaded.Add(Object);
+				LoadContext->AddLoadedObject(Object);
 			}
 		}
 	}
@@ -6383,7 +6382,7 @@ EAsyncPackageState::Type FAsyncPackage::PreLoadObjects()
 	// GC can't run in here
 	FGCScopeGuard GCGuard;
 
-	TArray<UObject*>& ThreadObjLoaded = LoadContext->GetObjectsLoaded();
+	TArray<UObject*>& ThreadObjLoaded = LoadContext->PRIVATE_GetObjectsLoadedForFAsyncPackage();
 	PackageObjLoaded.Append(ThreadObjLoaded);
 	ThreadObjLoaded.Reset();
 
@@ -6482,7 +6481,7 @@ EAsyncPackageState::Type FAsyncPackage::PostLoadObjects()
 	FUObjectThreadContext& ThreadContext = FUObjectThreadContext::Get();
 	TGuardValue<bool> GuardIsRoutingPostLoad(ThreadContext.IsRoutingPostLoad, true);
 
-	TArray<UObject*>& ThreadObjLoaded = LoadContext->GetObjectsLoaded();
+	TArray<UObject*>& ThreadObjLoaded = LoadContext->PRIVATE_GetObjectsLoadedForFAsyncPackage();
 	if (ThreadObjLoaded.Num())
 	{
 		// New objects have been loaded. They need to go through PreLoad first so exit now and come back after they've been preloaded.
@@ -6562,7 +6561,7 @@ EAsyncPackageState::Type FAsyncPackage::PostLoadDeferredObjects(double InTickSta
 	TGuardValue<bool> GuardIsRoutingPostLoad(PackageScope.ThreadContext.IsRoutingPostLoad, true);
 	FAsyncLoadingTickScope InAsyncLoadingTick;
 
-	TArray<UObject*>& ObjLoadedInPostLoad = LoadContext->GetObjectsLoaded();
+	TArray<UObject*>& ObjLoadedInPostLoad = LoadContext->PRIVATE_GetObjectsLoadedForFAsyncPackage();
 	TArray<UObject*> ObjLoadedInPostLoadLocal;
 
 	STAT(double PostLoadStartTime = FPlatformTime::Seconds());
@@ -6752,7 +6751,7 @@ EAsyncPackageState::Type FAsyncPackage::FinishObjects()
 	LastTypeOfWorkPerformed			= TEXT("finishing all objects");
 
 	check(!Linker || LoadContext == Linker->GetSerializeContext());		
-	TArray<UObject*>& ThreadObjLoaded = LoadContext->GetObjectsLoaded();
+	TArray<UObject*>& ThreadObjLoaded = LoadContext->PRIVATE_GetObjectsLoadedForFAsyncPackage();
 
 	EAsyncLoadingResult::Type LoadingResult;
 	if (!bLoadHasFailed)
diff --git a/Engine/Source/Runtime/CoreUObject/Private/UObject/LinkerLoad.cpp b/Engine/Source/Runtime/CoreUObject/Private/UObject/LinkerLoad.cpp
--- a/Engine/Source/Runtime/CoreUObject/Private/UObject/LinkerLoad.cpp
+++ b/Engine/Source/Runtime/CoreUObject/Private/UObject/LinkerLoad.cpp
@@ -4279,7 +4279,7 @@ UObject* FLinkerLoad::CreateExport( int32 Index )
 					// fully regenerated before another (there is no changing 
 					// that); so dependencies need to be recompiled later (with
 					// all the regenerated classes in place)
-					FScopedClassDependencyGather DependencyHelper(ClassObject);
+					FScopedClassDependencyGather DependencyHelper(ClassObject, GetSerializeContext());
 #endif //WITH_EDITOR
 
 					ClassObject->Bind();
diff --git a/Engine/Source/Runtime/CoreUObject/Private/UObject/SoftObjectPath.cpp b/Engine/Source/Runtime/CoreUObject/Private/UObject/SoftObjectPath.cpp
--- a/Engine/Source/Runtime/CoreUObject/Private/UObject/SoftObjectPath.cpp
+++ b/Engine/Source/Runtime/CoreUObject/Private/UObject/SoftObjectPath.cpp
@@ -101,10 +101,10 @@ bool FSoftObjectPath::PreSavePath(bool* bReportSoftObjectPathRedirects)
 	return false;
 }
 
-void FSoftObjectPath::PostLoadPath() const
+void FSoftObjectPath::PostLoadPath(FArchive* InArchive) const
 {
 #if WITH_EDITOR
-	GRedirectCollector.OnSoftObjectPathLoaded(*this);
+	GRedirectCollector.OnSoftObjectPathLoaded(*this, InArchive);
 #endif // WITH_EDITOR
 }
 
@@ -177,7 +177,7 @@ void FSoftObjectPath::SerializePath(FArchive& Ar)
 	{
 		if (Ar.IsPersistent())
 		{
-			PostLoadPath();
+			PostLoadPath(&Ar);
 
 			// If we think it's going to work, we try to do the pre-save fixup now. This is important because it helps with blueprint CDO save determinism with redirectors
 			// It's important that the entire CDO hierarchy gets fixed up before an instance in a map gets saved otherwise the delta serialization will save too much
@@ -274,13 +274,13 @@ bool FSoftObjectPath::ImportTextItem(const TCHAR*& Buffer, int32 PortFlags, UObj
 		// We're probably reading config for an editor only object, we need to mark this reference as editor only
 		FSoftObjectPathSerializationScope SerializationScope(NAME_None, NAME_None, ESoftObjectPathCollectType::EditorOnlyCollect, ESoftObjectPathSerializeType::AlwaysSerialize);
 
-		PostLoadPath();
+		PostLoadPath(nullptr);
 	}
 	else
 #endif
 	{
 		// Consider this a load, so Config string references get cooked
-		PostLoadPath();
+		PostLoadPath(nullptr);
 	}
 
 	return true;
@@ -338,7 +338,7 @@ bool FSoftObjectPath::SerializeFromMismatchedTag(struct FPropertyTag const& Tag,
 	if (Slot.GetUnderlyingArchive().IsLoading())
 	{
 		SetPath(MoveTemp(Path));
-		PostLoadPath();
+		PostLoadPath(&Slot.GetUnderlyingArchive());
 	}
 
 	return bReturn;
@@ -511,7 +511,7 @@ bool FSoftClassPath::SerializeFromMismatchedTag(struct FPropertyTag const& Tag,
 	if (Slot.GetUnderlyingArchive().IsLoading())
 	{
 		SetPath(MoveTemp(Path));
-		PostLoadPath();
+		PostLoadPath(&Slot.GetUnderlyingArchive());
 	}
 
 	return bReturn;
diff --git a/Engine/Source/Runtime/CoreUObject/Private/UObject/UObjectGlobals.cpp b/Engine/Source/Runtime/CoreUObject/Private/UObject/UObjectGlobals.cpp
--- a/Engine/Source/Runtime/CoreUObject/Private/UObject/UObjectGlobals.cpp
+++ b/Engine/Source/Runtime/CoreUObject/Private/UObject/UObjectGlobals.cpp
@@ -1567,11 +1567,6 @@ void EndLoad(FUObjectSerializeContext* LoadContext)
 		LoadContext->DecrementBeginLoadCount();
 		return;
 	}
-	else if (IsInGameThread())
-	{
-		GGameThreadLoadCounter--;
-		check(GGameThreadLoadCounter >= 0);
-	}
 
 #if WITH_EDITOR
 	FScopedSlowTask SlowTask(0, NSLOCTEXT("Core", "PerformingPostLoad", "Performing post-load..."), ShouldReportProgress());
@@ -1766,7 +1761,17 @@ void EndLoad(FUObjectSerializeContext* LoadContext)
 			}
 		}
 
-		FBlueprintSupport::FlushReinstancingQueue();
+		// If this is the first LoadPackage call, flush the BP queue
+		if (GGameThreadLoadCounter < 2)
+		{
+			FBlueprintSupport::FlushReinstancingQueue();
+		}
+	}
+
+	if (IsInGameThread())
+	{
+		GGameThreadLoadCounter--;
+		check(GGameThreadLoadCounter >= 0);
 	}
 
 	// Loaded new objects, so allow reaccessing asset ptrs
diff --git a/Engine/Source/Runtime/CoreUObject/Public/Blueprint/BlueprintSupport.h b/Engine/Source/Runtime/CoreUObject/Public/Blueprint/BlueprintSupport.h
--- a/Engine/Source/Runtime/CoreUObject/Public/Blueprint/BlueprintSupport.h
+++ b/Engine/Source/Runtime/CoreUObject/Public/Blueprint/BlueprintSupport.h
@@ -142,7 +142,7 @@ private:
 struct COREUOBJECT_API FScopedClassDependencyGather
 {
 public:
-	FScopedClassDependencyGather(UClass* ClassToGather);
+	FScopedClassDependencyGather(UClass* ClassToGather, FUObjectSerializeContext* InLoadContext);
 	~FScopedClassDependencyGather();
 
 	/**
@@ -164,6 +164,9 @@ private:
 	/** List of dependencies (i.e. UClasses that have been newly instantiated) in the scope of this dependency gather */
 	static TArray<UClass*> BatchClassDependencies;
 
+	/** Current load context */
+	FUObjectSerializeContext* LoadContext;
+
 	FScopedClassDependencyGather();
 };
 
diff --git a/Engine/Source/Runtime/CoreUObject/Public/Misc/RedirectCollector.h b/Engine/Source/Runtime/CoreUObject/Public/Misc/RedirectCollector.h
--- a/Engine/Source/Runtime/CoreUObject/Public/Misc/RedirectCollector.h
+++ b/Engine/Source/Runtime/CoreUObject/Public/Misc/RedirectCollector.h
@@ -77,8 +77,9 @@ public:
 	/**
 	 * Called from FSoftObjectPath::PostLoadPath, registers this for later querying
 	 * @param InPath The soft object path that was loaded
+	 * @Param InArchive The archive that loaded this path
 	 */
-	void OnSoftObjectPathLoaded(const struct FSoftObjectPath& InPath);
+	void OnSoftObjectPathLoaded(const struct FSoftObjectPath& InPath, FArchive* InArchive);
 
 	/**
 	 * Load all soft object paths to resolve them, add that to the remap table, and empty the array
diff --git a/Engine/Source/Runtime/CoreUObject/Public/UObject/Class.h b/Engine/Source/Runtime/CoreUObject/Public/UObject/Class.h
--- a/Engine/Source/Runtime/CoreUObject/Public/UObject/Class.h
+++ b/Engine/Source/Runtime/CoreUObject/Public/UObject/Class.h
@@ -2166,7 +2166,7 @@ public:
 	 * Conditionally recompiles the class after loading, in case any dependencies were also newly loaded
 	 * @param ObjLoaded	If set this is the list of objects that are currently loading, usualy GObjLoaded
 	 */
-	virtual void ConditionalRecompileClass() {}
+	virtual void ConditionalRecompileClass(FUObjectSerializeContext* InLoadContext) {}
 	virtual void FlushCompilationQueueForLevel() {}
 #endif //WITH_EDITOR
 
diff --git a/Engine/Source/Runtime/CoreUObject/Public/UObject/SoftObjectPath.h b/Engine/Source/Runtime/CoreUObject/Public/UObject/SoftObjectPath.h
--- a/Engine/Source/Runtime/CoreUObject/Public/UObject/SoftObjectPath.h
+++ b/Engine/Source/Runtime/CoreUObject/Public/UObject/SoftObjectPath.h
@@ -157,8 +157,11 @@ struct COREUOBJECT_API FSoftObjectPath
 	/** Fixes up string asset reference for saving, call if saving with a method that skips SerializePath. This can modify the path, it will return true if it was modified */
 	bool PreSavePath(bool* bReportSoftObjectPathRedirects = nullptr);
 
-	/** Handles when a string asset reference has been loaded, call if loading with a method that skips SerializePath. This does not modify path but might call callbacks */
-	void PostLoadPath() const;
+	/** 
+	 * Handles when a string asset reference has been loaded, call if loading with a method that skips SerializePath. This does not modify path but might call callbacks
+	 * @param InArchive The archive that loaded this path
+	 */
+	void PostLoadPath(FArchive* InArchive) const;
 
 	/** Fixes up this SoftObjectPath to add the PIE prefix depending on what is currently active, returns true if it was modified. The overload that takes an explicit PIE instance is preferred, if it's available. */
 	bool FixupForPIE();
diff --git a/Engine/Source/Runtime/CoreUObject/Public/UObject/UObjectThreadContext.h b/Engine/Source/Runtime/CoreUObject/Public/UObject/UObjectThreadContext.h
--- a/Engine/Source/Runtime/CoreUObject/Public/UObject/UObjectThreadContext.h
+++ b/Engine/Source/Runtime/CoreUObject/Public/UObject/UObjectThreadContext.h
@@ -171,13 +171,24 @@ public:
 		return !!ObjectsLoaded.Num();
 	}
 
-	void PRIVATE_PatchNewObjectIntoExport(UObject* OldObject, UObject* NewObject)
+	bool PRIVATE_PatchNewObjectIntoExport(UObject* OldObject, UObject* NewObject)
 	{
 		const int32 ObjLoadedIdx = ObjectsLoaded.Find(OldObject);
 		if (ObjLoadedIdx != INDEX_NONE)
 		{
 			ObjectsLoaded[ObjLoadedIdx] = NewObject;
+			return true;
 		}
+		else
+		{
+			return false;
+		}
+	}
+
+	/** This is only meant to be used by FAsyncPackage for performance reasons. The ObjectsLoaded array should not be manipulated directly! */
+	TArray<UObject*>& PRIVATE_GetObjectsLoadedForFAsyncPackage()
+	{
+		return ObjectsLoaded;
 	}
 
 	void AppendLoadedObjectsAndEmpty(TArray<UObject*>& InLoadedObject)
@@ -186,9 +197,14 @@ public:
 		ObjectsLoaded.Reset();
 	}
 
-	TArray<UObject*>& GetObjectsLoaded()
+	void ReserveObjectsLoaded(int32 InReserveSize)
 	{
-		return ObjectsLoaded;
+		ObjectsLoaded.Reserve(InReserveSize);
+	}
+
+	int32 GetNumObjectsLoaded() const
+	{
+		return ObjectsLoaded.Num();
 	}
 
 	void AddDelayedLinkerClosePackage(class FLinkerLoad* InLinker)
diff --git a/Engine/Source/Runtime/Engine/Classes/Engine/BlueprintGeneratedClass.h b/Engine/Source/Runtime/Engine/Classes/Engine/BlueprintGeneratedClass.h
--- a/Engine/Source/Runtime/Engine/Classes/Engine/BlueprintGeneratedClass.h
+++ b/Engine/Source/Runtime/Engine/Classes/Engine/BlueprintGeneratedClass.h
@@ -694,7 +694,7 @@ public:
 	// UClass interface
 #if WITH_EDITOR
 	virtual UClass* GetAuthoritativeClass() override;
-	virtual void ConditionalRecompileClass() override;
+	virtual void ConditionalRecompileClass(FUObjectSerializeContext* InLoadContext) override;
 	virtual void FlushCompilationQueueForLevel() override;
 	virtual UObject* GetArchetypeForCDO() const override;
 #endif //WITH_EDITOR
diff --git a/Engine/Source/Runtime/Engine/Private/BlueprintGeneratedClass.cpp b/Engine/Source/Runtime/Engine/Private/BlueprintGeneratedClass.cpp
--- a/Engine/Source/Runtime/Engine/Private/BlueprintGeneratedClass.cpp
+++ b/Engine/Source/Runtime/Engine/Private/BlueprintGeneratedClass.cpp
@@ -216,11 +216,11 @@ struct FConditionalRecompileClassHepler
 extern UNREALED_API FSecondsCounterData BlueprintCompileAndLoadTimerData;
 extern COREUOBJECT_API bool GBlueprintUseCompilationManager;
 
-void UBlueprintGeneratedClass::ConditionalRecompileClass()
+void UBlueprintGeneratedClass::ConditionalRecompileClass(FUObjectSerializeContext* InLoadContext)
 {
 	if(GBlueprintUseCompilationManager)
 	{
-		FBlueprintCompilationManager::FlushCompilationQueue();
+		FBlueprintCompilationManager::FlushCompilationQueue(InLoadContext);
 		return;
 	}
 	
@@ -277,7 +277,7 @@ void UBlueprintGeneratedClass::FlushCompilationQueueForLevel()
 	{
 		if(Cast<ULevelScriptBlueprint>(ClassGeneratedBy))
 		{
-			FBlueprintCompilationManager::FlushCompilationQueue();
+			FBlueprintCompilationManager::FlushCompilationQueue(nullptr);
 		}
 	}
 }
diff --git a/Engine/Source/Runtime/Engine/Private/LevelStreaming.cpp b/Engine/Source/Runtime/Engine/Private/LevelStreaming.cpp
--- a/Engine/Source/Runtime/Engine/Private/LevelStreaming.cpp
+++ b/Engine/Source/Runtime/Engine/Private/LevelStreaming.cpp
@@ -291,7 +291,7 @@ void ULevelStreaming::PostLoad()
 			// Convert the FName reference to a TSoftObjectPtr, then broadcast that we loaded a reference
 			// so this reference is gathered by the cooker without having to resave the package.
 			SetWorldAssetByPackageName(PackageName_DEPRECATED);
-			WorldAsset.GetUniqueID().PostLoadPath();
+			WorldAsset.GetUniqueID().PostLoadPath(GetLinker());
 		}
 		else
 		{
