Suppress progress notifications in the editor during a deferred Blueprint re-indexing operation after an asset load.

#rnx

diff --git a/Engine/Source/Editor/Kismet/Private/FindInBlueprintManager.cpp b/Engine/Source/Editor/Kismet/Private/FindInBlueprintManager.cpp
--- a/Engine/Source/Editor/Kismet/Private/FindInBlueprintManager.cpp
+++ b/Engine/Source/Editor/Kismet/Private/FindInBlueprintManager.cpp
@@ -31,7 +31,6 @@
 #include "EdGraphSchema_K2.h"
 #include "K2Node_FunctionEntry.h"
 #include "EditorStyleSet.h"
-#include "BlueprintEditorModule.h"
 #include "BlueprintEditorSettings.h"
 #include "Framework/Docking/TabManager.h"
 #include "Widgets/Docking/SDockTab.h"
@@ -1090,20 +1089,19 @@ public:
 
 	struct FCacheParams
 	{
-		/** If TRUE, allow the operation to be cancelled by the user */
-		bool bAllowCancel = false;
-
-		/** If TRUE, show a progress notification (toast) while processing */
-		bool bShowProgress = true;
-
-		/** If TRUE, Blueprints will be checked out and resaved after being loaded */
-		bool bCheckOutAndSave = false;
+		/** Control flags */
+		EFiBCacheOpFlags OpFlags;
 
 		/** Callback for when assets are cached */
 		FOnAssetCached OnCached;
 
 		/** Callback for when caching is finished */
 		FSimpleDelegate OnFinished;
+
+		FCacheParams()
+			:OpFlags(EFiBCacheOpFlags::None)
+		{
+		}
 	};
 
 	FCacheAllBlueprintsTickableObject(TSet<FName> InAssets, const FCacheParams& InParams)
@@ -1113,12 +1111,13 @@ public:
 		, bIsCancelled(false)
 		, CacheParams(InParams)
 	{
-		if (CacheParams.bShowProgress)
+		if (EnumHasAnyFlags(CacheParams.OpFlags, EFiBCacheOpFlags::ShowProgress)
+			&& !EnumHasAnyFlags(CacheParams.OpFlags, EFiBCacheOpFlags::HideNotifications))
 		{
 			// Start the Blueprint indexing 'progress' notification
 			FNotificationInfo Info(LOCTEXT("BlueprintIndexMessage", "Indexing Blueprints..."));
 			Info.bFireAndForget = false;
-			if (CacheParams.bAllowCancel)
+			if (EnumHasAnyFlags(CacheParams.OpFlags, EFiBCacheOpFlags::AllowUserCancel))
 			{
 				Info.ButtonDetails.Add(FNotificationButtonInfo(
 					LOCTEXT("BlueprintIndexCancel", "Cancel"),
@@ -1209,7 +1208,7 @@ public:
 		if (!bIsStarted)
 		{
 			bIsStarted = true;
-			FFindInBlueprintSearchManager::Get().StartedCachingBlueprints();
+			FFindInBlueprintSearchManager::Get().StartedCachingBlueprints(CacheParams.OpFlags);
 		}
 	}
 
@@ -1245,10 +1244,12 @@ public:
 					bValidFilename = bIsWorldAsset ? FEditorFileUtils::IsValidMapFilename(FinalPackageFilename, ErrorMessage) : FPackageName::IsValidLongPackageName(FinalPackageFilename, false, &ErrorMessage);
 				}
 
+				const bool bCheckOutAndSave = EnumHasAnyFlags(CacheParams.OpFlags, EFiBCacheOpFlags::CheckOutAndSave);
+
 				bool bIsAssetReadOnlyOnDisk = IFileManager::Get().IsReadOnly(*FinalPackageFilename);
-				bool bFailedToCache = CacheParams.bCheckOutAndSave;
+				bool bFailedToCache = bCheckOutAndSave;
 
-				if (!bIsAssetReadOnlyOnDisk || !CacheParams.bCheckOutAndSave)
+				if (!bIsAssetReadOnlyOnDisk || !bCheckOutAndSave)
 				{
 					if (!FFindInBlueprintSearchManager::Get().IsUnindexedCacheInProgress())
 					{
@@ -1264,7 +1265,7 @@ public:
 					else
 					{
 						UObject* Asset = AssetData.GetAsset();
-						if (Asset && CacheParams.bCheckOutAndSave)
+						if (Asset && bCheckOutAndSave)
 						{
 							if (UBlueprint* BlueprintAsset = Cast<UBlueprint>(Asset))
 							{
@@ -1422,12 +1423,6 @@ FFindInBlueprintSearchManager::~FFindInBlueprintSearchManager()
 		HotReloadSupport.OnHotReload().RemoveAll(this);
 	}
 
-	if (FModuleManager::Get().IsModuleLoaded("Kismet"))
-	{
-		FBlueprintEditorModule& BlueprintEditorModule = FModuleManager::GetModuleChecked<FBlueprintEditorModule>("Kismet");
-		BlueprintEditorModule.OnBlueprintEditorOpened().RemoveAll(this);
-	}
-
 	// Shut down the global find results tab feature.
 	EnableGlobalFindResults(false);
 }
@@ -1466,10 +1461,6 @@ void FFindInBlueprintSearchManager::Initialize()
 	IHotReloadInterface& HotReloadSupport = FModuleManager::LoadModuleChecked<IHotReloadInterface>("HotReload");
 	HotReloadSupport.OnHotReload().AddRaw(this, &FFindInBlueprintSearchManager::OnHotReload);
 
-	// Register to be notified of a Blueprint editor open event
-	FBlueprintEditorModule& BlueprintEditorModule = FModuleManager::LoadModuleChecked<FBlueprintEditorModule>("Kismet");
-	BlueprintEditorModule.OnBlueprintEditorOpened().AddRaw(this, &FFindInBlueprintSearchManager::HandleBlueprintEditorOpened);
-
 	if(!GIsSavingPackage && AssetRegistryModule)
 	{
 		// Do an immediate load of the cache to catch any Blueprints that were discovered by the asset registry before we initialized.
@@ -1701,17 +1692,6 @@ void FFindInBlueprintSearchManager::OnHotReload(bool bWasTriggeredAutomatically)
 	CachedAssetClasses.Reset();
 }
 
-void FFindInBlueprintSearchManager::HandleBlueprintEditorOpened(EBlueprintType InBlueprintType)
-{
-	// Kick off an indexing operation if we have pending assets
-	if (!IsCacheInProgress() && PendingAssets.Num() > 0)
-	{
-		FFindInBlueprintCachingOptions CachingOptions;
-		CachingOptions.OpType = EFiBCacheOpType::CachePendingAssets;
-		CacheAllAssets(nullptr, CachingOptions);
-	}
-}
-
 FString FFindInBlueprintSearchManager::GatherBlueprintSearchMetadata(const UBlueprint* Blueprint)
 {	
 	FTemporarilyUseFriendlyNodeTitles TemporarilyUseFriendlyNodeTitles;
@@ -1803,13 +1783,14 @@ FString FFindInBlueprintSearchManager::GatherBlueprintSearchMetadata(const UBlue
 
 void FFindInBlueprintSearchManager::AddOrUpdateBlueprintSearchMetadata(UBlueprint* InBlueprint, bool bInForceReCache/* = false*/)
 {
-	// Do not try to gather any info from Blueprints who are loaded for diffing. It makes search all very strange and allows you to fully open those Blueprints.
-	if (InBlueprint->GetOutermost()->HasAnyPackageFlags(PKG_ForDiffing))
-	{
-		return;
-	}
-	
-	if (!bEnableGatheringData)
+	// No need to update the cache in the following cases:
+	//	a) Indexing is disabled.
+	//	b) The Blueprint is not yet fully loaded. This ensures that we don't make attempts to re-index before load completion.
+	//	c) The Blueprint was loaded for diffing. It makes search all very strange and allows you to fully open those Blueprints.
+	//	d) The Blueprint was loaded/copied for PIE. These assets are temporarily created for a session and don't need to be re-indexed.
+	if (!bEnableGatheringData
+		|| !InBlueprint->HasAnyFlags(RF_WasLoaded)
+		|| InBlueprint->GetOutermost()->HasAnyPackageFlags(PKG_ForDiffing | PKG_PlayInEditor))
 	{
 		return;
 	}
@@ -2192,8 +2173,17 @@ void FFindInBlueprintSearchManager::CacheAllAssets(TWeakPtr< SFindInBlueprints >
 				PendingAssets.Remove(InAssetName);
 			});
 
-			// Show the progress notification if we have more than one asset
-			CacheParams.bShowProgress = (PendingAssets.Num() > 1);
+			// Determine if PIE is active - in that case we're potentially streaming assets in at random intervals, so just hide the progress UI while re-indexing those assets
+			const bool bIsPIESimulating = (GEditor->bIsSimulatingInEditor || GEditor->PlayWorld);
+
+			// Display progress during a re-indexing operation only if we have multiple assets to process (e.g. avoid showing after compiling a single asset) and we're not in PIE
+			if ((PendingAssets.Num() > 1) && !bIsPIESimulating)
+			{
+				CacheParams.OpFlags |= EFiBCacheOpFlags::ShowProgress;
+			}
+
+			// Keep popup notifications hidden during this operation
+			CacheParams.OpFlags |= EFiBCacheOpFlags::HideNotifications;
 
 			// Keep track of which global FiB context started the operation (if any)
 			SourceCachingWidget = InSourceWidget;
@@ -2252,21 +2242,25 @@ void FFindInBlueprintSearchManager::CacheAllAssets(TWeakPtr< SFindInBlueprints >
 				TempUncachedAssets.Append(UnindexedAssets);
 				TempUncachedAssets.Append(BlueprintsToUpdate);
 
+				const bool bCheckOutAndSave = (ReturnValue == EAppReturnType::Yes);
 				FCacheAllBlueprintsTickableObject::FCacheParams CacheParams;
-				CacheParams.bAllowCancel = true;
-				CacheParams.bCheckOutAndSave = ReturnValue == EAppReturnType::Yes;
+				CacheParams.OpFlags = EFiBCacheOpFlags::ShowProgress | EFiBCacheOpFlags::AllowUserCancel;
+				if (bCheckOutAndSave)
+				{
+					CacheParams.OpFlags |= EFiBCacheOpFlags::CheckOutAndSave;
+				}
 				CacheParams.OnFinished = InOptions.OnFinished;
 				CachingObject = new FCacheAllBlueprintsTickableObject(MoveTemp(TempUncachedAssets), CacheParams);
 
 				const bool bIsSourceControlEnabled = ISourceControlModule::Get().IsEnabled();
-				if (!bIsSourceControlEnabled && CacheParams.bCheckOutAndSave)
+				if (!bIsSourceControlEnabled && bCheckOutAndSave)
 				{
 					// Offer to start up Source Control
-					ISourceControlModule::Get().ShowLoginDialog(FSourceControlLoginClosed::CreateRaw(this, &FFindInBlueprintSearchManager::OnCacheAllUnindexedAssets, CacheParams.bCheckOutAndSave), ELoginWindowMode::Modeless, EOnLoginWindowStartup::PreserveProvider);
+					ISourceControlModule::Get().ShowLoginDialog(FSourceControlLoginClosed::CreateRaw(this, &FFindInBlueprintSearchManager::OnCacheAllUnindexedAssets, bCheckOutAndSave), ELoginWindowMode::Modeless, EOnLoginWindowStartup::PreserveProvider);
 				}
 				else
 				{
-					OnCacheAllUnindexedAssets(bIsSourceControlEnabled, CacheParams.bCheckOutAndSave);
+					OnCacheAllUnindexedAssets(bIsSourceControlEnabled, bCheckOutAndSave);
 				}
 
 				SourceCachingWidget = InSourceWidget;
@@ -2338,14 +2332,14 @@ int32 FFindInBlueprintSearchManager::GetNumberUncachedAssets() const
 	return 0;
 }
 
-void FFindInBlueprintSearchManager::StartedCachingBlueprints()
+void FFindInBlueprintSearchManager::StartedCachingBlueprints(EFiBCacheOpFlags InCacheOpFlags)
 {
 	// Invoke the callback on any open global widgets
 	for (TWeakPtr<SFindInBlueprints> FindResultsPtr : GlobalFindResults)
 	{
 		if (FindResultsPtr.IsValid())
 		{
-			FindResultsPtr.Pin()->OnCacheStarted(CurrentCacheOpType);
+			FindResultsPtr.Pin()->OnCacheStarted(CurrentCacheOpType, InCacheOpFlags);
 		}
 	}
 }
@@ -2641,9 +2635,9 @@ void FFindInBlueprintSearchManager::Tick(float DeltaTime)
 		check(CachingObject);
 		CachingObject->Tick(DeltaTime);
 	}
-	else if (PendingAssets.Num() > 0 && GlobalFindResults.Num() > 0)
+	else if (PendingAssets.Num() > 0)
 	{
-		// Kick off an indexing operation if we have pending assets and at least one open global FiB context
+		// Kick off a re-indexing operation to update the cache
 		FFindInBlueprintCachingOptions CachingOptions;
 		CachingOptions.OpType = EFiBCacheOpType::CachePendingAssets;
 		CacheAllAssets(nullptr, CachingOptions);
@@ -2652,8 +2646,8 @@ void FFindInBlueprintSearchManager::Tick(float DeltaTime)
 
 bool FFindInBlueprintSearchManager::IsTickable() const
 {
-	// Tick only if we have an active caching operation or if we have pending assets and an open FiB context
-	return IsCacheInProgress() || (PendingAssets.Num() > 0 && GlobalFindResults.Num() > 0);
+	// Tick if we have an active caching operation or one or more pending assets
+	return IsCacheInProgress() || (PendingAssets.Num() > 0);
 }
 
 TStatId FFindInBlueprintSearchManager::GetStatId() const
diff --git a/Engine/Source/Editor/Kismet/Private/FindInBlueprints.cpp b/Engine/Source/Editor/Kismet/Private/FindInBlueprints.cpp
--- a/Engine/Source/Editor/Kismet/Private/FindInBlueprints.cpp
+++ b/Engine/Source/Editor/Kismet/Private/FindInBlueprints.cpp
@@ -594,7 +594,6 @@ void SFindInBlueprints::Construct( const FArguments& InArgs, TSharedPtr<FBluepri
 
 	HostTab = InArgs._ContainingTab;
 	bIsLocked = false;
-	bKeepCacheBarOpenOnCacheCompletion = false;
 
 	if (HostTab.IsValid())
 	{
@@ -723,11 +722,10 @@ void SFindInBlueprints::Construct( const FArguments& InArgs, TSharedPtr<FBluepri
 
 void SFindInBlueprints::ConditionallyAddCacheBar()
 {
-	FFindInBlueprintSearchManager& FindInBlueprintManager = FFindInBlueprintSearchManager::Get();
-
-	// Do not add a second cache bar and do not add it when there are no unindexed Blueprints
-	if(FindInBlueprintManager.GetNumberUncachedAssets() > 0 || FindInBlueprintManager.GetNumberUnindexedAssets() > 0 || FindInBlueprintManager.GetFailedToCacheCount() > 0)
+	// Do not add when it should not be visible
+	if(GetCachingBarVisibility() == EVisibility::Visible)
 	{
+		// Do not add a second cache bar
 		if(MainVerticalBox.IsValid() && !CacheBarSlot.IsValid())
 		{
 			// Create a single string of all the Blueprint paths that failed to cache, on separate lines
@@ -1329,8 +1327,9 @@ EVisibility SFindInBlueprints::GetCacheAllCancelButtonVisibility() const
 
 EVisibility SFindInBlueprints::GetCachingBarVisibility() const
 {
+	const bool bIsPIESimulating = (GEditor->bIsSimulatingInEditor || GEditor->PlayWorld);
 	FFindInBlueprintSearchManager& FindInBlueprintManager = FFindInBlueprintSearchManager::Get();
-	return (FindInBlueprintManager.GetNumberUncachedAssets() > 0 || FindInBlueprintManager.GetNumberUnindexedAssets() > 0 || FindInBlueprintManager.GetFailedToCacheCount())? EVisibility::Visible : EVisibility::Collapsed;
+	return (FindInBlueprintManager.GetNumberUncachedAssets() > 0 || (!bIsPIESimulating && (FindInBlueprintManager.GetNumberUnindexedAssets() > 0 || FindInBlueprintManager.GetFailedToCacheCount())))? EVisibility::Visible : EVisibility::Collapsed;
 }
 
 EVisibility SFindInBlueprints::GetCachingBlueprintNameVisiblity() const
@@ -1406,18 +1405,13 @@ FText SFindInBlueprints::GetCurrentCacheBlueprintName() const
 	return FText::FromName(FFindInBlueprintSearchManager::Get().GetCurrentCacheBlueprintName());
 }
 
-void SFindInBlueprints::OnCacheStarted(EFiBCacheOpType InOpType)
+void SFindInBlueprints::OnCacheStarted(EFiBCacheOpType InOpType, EFiBCacheOpFlags InOpFlags)
 {
 	const bool bIsCacheBarAdded = CacheBarSlot.IsValid();
-
-	// Keep the cache bar open after a re-indexing operation if we're currently showing it (post-search unindexed/out-of-date UI context)
-	bKeepCacheBarOpenOnCacheCompletion = (InOpType == EFiBCacheOpType::CachePendingAssets) && bIsCacheBarAdded;
-
-	// Show progress during a re-indexing operation only if we have multiple assets to process (e.g. avoid showing after compiling a single asset)
-	const bool bShouldShowProgress = InOpType != EFiBCacheOpType::CachePendingAssets || FFindInBlueprintSearchManager::Get().GetNumberPendingAssets() > 1;
+	const bool bShowProgress = EnumHasAnyFlags(InOpFlags, EFiBCacheOpFlags::ShowProgress);
 
 	// Ensure that the cache bar is visible to show progress
-	if (!bIsCacheBarAdded && bShouldShowProgress)
+	if (!bIsCacheBarAdded && bShowProgress)
 	{
 		ConditionallyAddCacheBar();
 	}
@@ -1425,28 +1419,17 @@ void SFindInBlueprints::OnCacheStarted(EFiBCacheOpType InOpType)
 
 void SFindInBlueprints::OnCacheComplete(EFiBCacheOpType InOpType)
 {
-	if (InOpType == EFiBCacheOpType::CacheUnindexedAssets)
+	TWeakPtr<SFindInBlueprints> SourceCachingWidgetPtr = FFindInBlueprintSearchManager::Get().GetSourceCachingWidget();
+	if (InOpType == EFiBCacheOpType::CacheUnindexedAssets
+		&& SourceCachingWidgetPtr.IsValid() && SourceCachingWidgetPtr.Pin() == SharedThis(this))
 	{
 		// Resubmit the last search, which will also remove the bar if needed
 		OnSearchTextCommitted(SearchTextField->GetText(), ETextCommit::OnEnter);
 	}
-	else if(!bKeepCacheBarOpenOnCacheCompletion)
+	else if (CacheBarSlot.IsValid())
 	{
-		const bool bIsCacheBarAdded = CacheBarSlot.IsValid();
-
-		if (bIsCacheBarAdded)
-		{
-			if (!bKeepCacheBarOpenOnCacheCompletion)
-			{
-				// Remove the cache bar
-				OnRemoveCacheBar();
-			}
-		}
-		else if(bKeepCacheBarOpenOnCacheCompletion)
-		{
-			// Restore the cache bar to its previous state (e.g. in case the user closed the progress bar while caching)
-			ConditionallyAddCacheBar();
-		}
+		// Remove the cache bar
+		OnRemoveCacheBar();
 	}
 }
 
diff --git a/Engine/Source/Editor/Kismet/Public/FindInBlueprintManager.h b/Engine/Source/Editor/Kismet/Public/FindInBlueprintManager.h
--- a/Engine/Source/Editor/Kismet/Public/FindInBlueprintManager.h
+++ b/Engine/Source/Editor/Kismet/Public/FindInBlueprintManager.h
@@ -185,12 +185,28 @@ struct KISMET_API FFiBMD
 };
 
 /** Which assets to index for caching */
-enum EFiBCacheOpType
+enum class EFiBCacheOpType
 {
 	CachePendingAssets,
 	CacheUnindexedAssets
 };
 
+/** Flags to control the UX while caching */
+enum class EFiBCacheOpFlags
+{
+	None = 0,
+	/** Whether to show progress */
+	ShowProgress = 1 << 0,
+	/** Whether to hide toast popups */
+	HideNotifications = 1 << 1,
+	/** Whether to allow users to cancel */
+	AllowUserCancel = 1 << 2,
+	/** Set if the user wants to check out and save (applies to unindexed caching only) */
+	CheckOutAndSave = 1 << 3,
+};
+
+ENUM_CLASS_FLAGS(EFiBCacheOpFlags);
+
 /** Options to configure the bulk caching task */
 struct FFindInBlueprintCachingOptions
 {
@@ -479,7 +495,7 @@ public:
 	bool HasCachingFailed() const { return FailedToCachePaths.Num() > 0; };
 
 	/** Callback to note that Blueprint caching is started */
-	void StartedCachingBlueprints();
+	void StartedCachingBlueprints(EFiBCacheOpFlags InCacheOpFlags);
 
 	/**
 	 * Callback to note that Blueprint caching is complete
@@ -497,9 +513,6 @@ public:
 	/** Returns a weak reference to the widget that initiated the current caching operation */
 	TWeakPtr<SFindInBlueprints> GetSourceCachingWidget() const { return SourceCachingWidget; }
 
-	/** Serializes an FString to memory and converts the memory into a string of hex characters */
-	static FString ConvertFStringToHexString(FString InValue);
-
 	/** Given a fully constructed Find-in-Blueprint FString of searchable data, will parse and construct a JsonObject */
 	static TSharedPtr< class FJsonObject > ConvertJsonStringToObject(bool bInIsVersioned, FString InJsonString, TMap<int32, FText>& OutFTextLookupTable);
 
@@ -546,9 +559,6 @@ private:
 	/** Callback hook from the Hot Reload manager that indicates that a module has been hot-reloaded */
 	void OnHotReload(bool bWasTriggeredAutomatically);
 
-	/** Callback from Kismet when a Blueprint asset is opened in the editor */
-	void HandleBlueprintEditorOpened(EBlueprintType InBlueprintType);
-
 	/** Helper to gathers the Blueprint's search metadata */
 	FString GatherBlueprintSearchMetadata(const UBlueprint* Blueprint);
 
diff --git a/Engine/Source/Editor/Kismet/Public/FindInBlueprints.h b/Engine/Source/Editor/Kismet/Public/FindInBlueprints.h
--- a/Engine/Source/Editor/Kismet/Public/FindInBlueprints.h
+++ b/Engine/Source/Editor/Kismet/Public/FindInBlueprints.h
@@ -218,7 +218,7 @@ public:
 	void MakeSearchQuery(FString InSearchString, bool bInIsFindWithinBlueprint, enum ESearchQueryFilter InSearchFilterForImaginaryDataReturn = ESearchQueryFilter::AllFilter, EFiBVersion InMinimiumVersionRequirement = EFiBVersion::FIB_VER_LATEST, FOnSearchComplete InOnSearchComplete = FOnSearchComplete());
 
 	/** Called when caching Blueprints is started */
-	void OnCacheStarted(EFiBCacheOpType InOpType);
+	void OnCacheStarted(EFiBCacheOpType InOpType, EFiBCacheOpFlags InOpFlags);
 	
 	/** Called when caching Blueprints is complete */
 	void OnCacheComplete(EFiBCacheOpType InOpType);
@@ -409,7 +409,4 @@ private:
 
 	/** True if current search should not be changed by an external source */
 	bool bIsLocked;
-
-	/** True if the cache bar should remain active after a caching operation */
-	bool bKeepCacheBarOpenOnCacheCompletion;
 };
