#DUPEMERGE : Merging using Dev-Anim<->Release4.21

fix crash when deleting state machine itself

#code review:Thomas.Sarkanen

diff --git a/Engine/Source/Editor/Kismet/Private/BlueprintEditor.cpp b/Engine/Source/Editor/Kismet/Private/BlueprintEditor.cpp
--- a/Engine/Source/Editor/Kismet/Private/BlueprintEditor.cpp
+++ b/Engine/Source/Editor/Kismet/Private/BlueprintEditor.cpp
@@ -136,6 +136,8 @@
 #include "BlueprintEditorModes.h"
 #include "BlueprintEditorSettings.h"
 #include "K2Node_SwitchString.h"
+#include "AnimGraphNode_StateMachineBase.h"
+#include "AnimationStateMachineGraph.h"
 
 #include "EngineAnalytics.h"
 #include "AnalyticsEventAttribute.h"
@@ -5314,6 +5316,14 @@ void FBlueprintEditor::DeleteSelectedNodes()
 						CloseDocumentTab(NodeGraph);
 					}
 				}
+				else
+				{
+					const UAnimGraphNode_StateMachineBase* SMNode = Cast<UAnimGraphNode_StateMachineBase>(Node);
+					if (SMNode)
+					{
+						CloseDocumentTab(SMNode->EditorStateMachineGraph);
+					}
+				}
 
 				UK2Node* K2Node = Cast<UK2Node>(Node);
 				if (K2Node != NULL && K2Node->NodeCausesStructuralBlueprintChange())
