#DUPEMERGE Merging using Dev-Anim<->Release4.21

Now remove all recurves graphs with this statemachine or state.

diff --git a/Engine/Source/Editor/Kismet/Private/BlueprintEditor.cpp b/Engine/Source/Editor/Kismet/Private/BlueprintEditor.cpp
--- a/Engine/Source/Editor/Kismet/Private/BlueprintEditor.cpp
+++ b/Engine/Source/Editor/Kismet/Private/BlueprintEditor.cpp
@@ -5299,6 +5299,22 @@ void FBlueprintEditor::DeleteSelectedNodes()
 
 	SetUISelectionState(NAME_None);
 
+	// this closes all the document that is outered by this node
+	// this is used by AnimBP statemachines/states that can create subgraph
+	auto CloseAllDocumentsTab = [this](const UEdGraphNode* InNode)
+	{
+		TArray<UObject*> NodesToClose;
+		GetObjectsWithOuter(InNode, NodesToClose);
+		for (UObject* Node : NodesToClose)
+		{
+			UEdGraph* NodeGraph = Cast<UEdGraph>(Node);
+			if (NodeGraph)
+			{
+				CloseDocumentTab(NodeGraph);
+			}
+		}
+	};
+
 	for (FGraphPanelSelectionSet::TConstIterator NodeIt( SelectedNodes ); NodeIt; ++NodeIt)
 	{
 		if (UEdGraphNode* Node = Cast<UEdGraphNode>(*NodeIt))
@@ -5310,18 +5326,14 @@ void FBlueprintEditor::DeleteSelectedNodes()
 				UAnimStateNodeBase* StateNode = Cast<UAnimStateNodeBase>(Node);
 				if (StateNode)
 				{
-					UEdGraph* NodeGraph = StateNode->GetBoundGraph();
-					if (NodeGraph)
-					{
-						CloseDocumentTab(NodeGraph);
-					}
+					CloseAllDocumentsTab(StateNode);
 				}
 				else
 				{
 					const UAnimGraphNode_StateMachineBase* SMNode = Cast<UAnimGraphNode_StateMachineBase>(Node);
 					if (SMNode)
 					{
-						CloseDocumentTab(SMNode->EditorStateMachineGraph);
+						CloseAllDocumentsTab(SMNode);
 					}
 				}
 
