Temp fix for a regression in PIE startup load time.

Change summary:
- Don't start a deferred Blueprint re-indexing operation unless at least one global FiB tab is open.

#rnx

diff --git a/Engine/Source/Editor/Kismet/Private/FindInBlueprintManager.cpp b/Engine/Source/Editor/Kismet/Private/FindInBlueprintManager.cpp
--- a/Engine/Source/Editor/Kismet/Private/FindInBlueprintManager.cpp
+++ b/Engine/Source/Editor/Kismet/Private/FindInBlueprintManager.cpp
@@ -2646,8 +2646,8 @@ void FFindInBlueprintSearchManager::Tick(float DeltaTime)
 
 bool FFindInBlueprintSearchManager::IsTickable() const
 {
-	// Tick if we have an active caching operation or one or more pending assets
-	return IsCacheInProgress() || (PendingAssets.Num() > 0);
+	// Tick only if we have an active caching operation or if we have pending assets and an open FiB context
+	return IsCacheInProgress() || (PendingAssets.Num() > 0 && GlobalFindResults.Num() > 0);
 }
 
 TStatId FFindInBlueprintSearchManager::GetStatId() const
