Fixed crash when reparenting a blueprint, caused by stale function pointers in skeleton class

diff --git a/Engine/Source/Editor/Kismet/Private/BlueprintCompilationManager.cpp b/Engine/Source/Editor/Kismet/Private/BlueprintCompilationManager.cpp
--- a/Engine/Source/Editor/Kismet/Private/BlueprintCompilationManager.cpp
+++ b/Engine/Source/Editor/Kismet/Private/BlueprintCompilationManager.cpp
@@ -712,10 +712,14 @@ void FBlueprintCompilationManagerImpl::FlushCompilationQueueImpl(bool bSuppressB
 					// CDO needs to be moved aside already:
 					ensure(SkeletonToRelink->ClassDefaultObject == nullptr);
 					ensure(!SkeletonToRelink->GetSuperClass()->HasAnyClassFlags(CLASS_NewerVersionExists));
-
+					
+					SkeletonToRelink->ClassConstructor = nullptr;
+					SkeletonToRelink->ClassVTableHelperCtorCaller = nullptr;
+					SkeletonToRelink->ClassAddReferencedObjects = nullptr;
 					SkeletonToRelink->Bind();
 					SkeletonToRelink->ClearFunctionMapsCaches();
 					SkeletonToRelink->StaticLink(true);
+					SkeletonToRelink->GetDefaultObject();
 				}
 
 				if(CompilerData.ShouldMarkUpToDateAfterSkeletonStage())
@@ -1728,7 +1732,7 @@ void FBlueprintCompilationManagerImpl::ReinstanceBatch(TArray<FReinstancingJob>&
 		ArchetypeReferencers.Add(CurrentReinstancer->ClassToReinstance->ClassGeneratedBy);
 		if(UBlueprint* BP = Cast<UBlueprint>(CurrentReinstancer->ClassToReinstance->ClassGeneratedBy))
 		{
-			// The only known way to cause this ensure to trip is to enqueue bluerpints for compilation
+			// The only known way to cause this ensure to trip is to enqueue blueprints for compilation
 			// while blueprints are already compiling:
 			if( ensure(BP->SkeletonGeneratedClass) )
 			{
