Fixed crashes when reloading a blueprint that has children types also loaded

diff --git a/Engine/Source/Editor/Kismet/Private/BlueprintCompilationManager.cpp b/Engine/Source/Editor/Kismet/Private/BlueprintCompilationManager.cpp
--- a/Engine/Source/Editor/Kismet/Private/BlueprintCompilationManager.cpp
+++ b/Engine/Source/Editor/Kismet/Private/BlueprintCompilationManager.cpp
@@ -95,6 +95,8 @@ struct FBlueprintCompilationManagerImpl : public FGCObject
 	bool HasBlueprintsToCompile() const;
 	bool IsGeneratedClassLayoutReady() const;
 	void GetDefaultValue(const UClass* ForClass, const UProperty* Property, FString& OutDefaultValueAsString) const;
+
+	static void ReparentHierarchies(const TMap<UClass*, UClass*>& OldClassToNewClass);
 	static void BuildDSOMap(UObject* OldObject, UObject* NewObject, TMap<UObject*, UObject*>& OutOldToNewDSO);
 	static void ReinstanceBatch(TArray<FReinstancingJob>& Reinstancers, TMap< UClass*, UClass* >& InOutOldToNewClassMap, FUObjectSerializeContext* InLoadContext);
 	static UClass* FastGenerateSkeletonClass(UBlueprint* BP, FKismetCompilerContext& CompilerContext, bool bIsSkeletonOnly);
@@ -1344,6 +1346,61 @@ void FBlueprintCompilationManagerImpl::GetDefaultValue(const UClass* ForClass, c
 	}
 }
 
+void FBlueprintCompilationManagerImpl::ReparentHierarchies(const TMap<UClass*, UClass*>& OldClassToNewClass)
+{
+	// something has decided to replace instances of a class. We need to update all the children of those types:
+	TSet<TPair<UClass*, UClass*>> ClassesToReparent;
+	for(const TPair<UClass*, UClass*>& OldToNewClass : OldClassToNewClass)
+	{
+		TArray<UClass*> DerivedClasses;
+		GetDerivedClasses(OldToNewClass.Key, DerivedClasses, false);
+
+		for(UClass* DerivedClass : DerivedClasses)
+		{
+			ClassesToReparent.Add(
+				TPair<UClass*, UClass*>(DerivedClass, OldToNewClass.Value)
+			);
+		}
+	}
+
+	// create reinstancers:
+	TArray<FReinstancingJob> Reinstancers;
+	for(const TPair<UClass*, UClass*> OldToNewParentClass : ClassesToReparent)
+	{
+		Reinstancers.Push( {
+			TSharedPtr<FBlueprintCompileReinstancer>(
+				new FBlueprintCompileReinstancer(
+					OldToNewParentClass.Key,
+					EBlueprintCompileReinstancerFlags::AutoInferSaveOnCompile | EBlueprintCompileReinstancerFlags::AvoidCDODuplication
+				)
+			),
+			TSharedPtr<FKismetCompilerContext>()
+		} );
+	}
+
+	// reparent blueprints, now that their cdo has been moved aside and the REINST_ type has been created for old instances:
+	TMap<UClass*, UClass*> OldClassToNewClassIncludingChildren = OldClassToNewClass;
+	for(const FReinstancingJob& ReinstancingJob : Reinstancers)
+	{
+		UClass* ClassToReinstance = ReinstancingJob.Reinstancer->ClassToReinstance;
+		ClassToReinstance->SetSuperStruct(OldClassToNewClass.FindChecked(ClassToReinstance->GetSuperClass()));
+		OldClassToNewClassIncludingChildren.Add(ReinstancingJob.Reinstancer->DuplicatedClass, ClassToReinstance);
+	}
+
+	// reparenting done, reinstance the hierarchy and update archetypes:
+	ReinstanceBatch(Reinstancers, OldClassToNewClassIncludingChildren, nullptr);
+
+	// reinstance instances - but only the classes we created
+	TMap<UClass*, UClass*> ClassesToReinstance;
+	for(const FReinstancingJob& ReinstancingJob : Reinstancers)
+	{
+		ClassesToReinstance.Add(ReinstancingJob.Reinstancer->DuplicatedClass, ReinstancingJob.Reinstancer->ClassToReinstance);
+	}
+	TGuardValue<bool> ReinstancingGuard(GIsReinstancing, true);
+	FBlueprintCompileReinstancer::BatchReplaceInstancesOfClass( ClassesToReinstance, /* bArchetypesAreUpToDate */ true );
+}
+
+
 void FBlueprintCompilationManagerImpl::BuildDSOMap(UObject* OldObject, UObject* NewObject, TMap<UObject*, UObject*>& OutOldToNewDSO)
 {
 	// IsDefaultObject() unfortunately cannot be relied upon for archetypes, so we explicitly search for the flag:
@@ -2474,5 +2531,11 @@ bool FBlueprintCompilationManager::GetDefaultValue(const UClass* ForClass, const
 	return true;
 }
 
+
+void FBlueprintCompilationManager::ReparentHierarchies(const TMap<UClass*, UClass*>& OldClassToNewClass)
+{
+	FBlueprintCompilationManagerImpl::ReparentHierarchies(OldClassToNewClass);
+}
+
 #undef LOCTEXT_NAMESPACE
 
diff --git a/Engine/Source/Editor/Kismet/Public/BlueprintCompilationManager.h b/Engine/Source/Editor/Kismet/Public/BlueprintCompilationManager.h
--- a/Engine/Source/Editor/Kismet/Public/BlueprintCompilationManager.h
+++ b/Engine/Source/Editor/Kismet/Public/BlueprintCompilationManager.h
@@ -74,6 +74,13 @@ struct KISMET_API FBlueprintCompilationManager
 	 * value from there
 	 */
 	static bool GetDefaultValue(const UClass* ForClass, const UProperty* Property, FString& OutDefaultValueAsString);
+
+	/**
+	 * Safely reparents all child classes of every Key in OldClassToNewClass to the class in 
+	 * the corresponding Value. Typically this means every child type will be reinstanced - although
+	 * reinstancing could be avoided when layouts match.
+	 */
+	static void ReparentHierarchies(const TMap<UClass*, UClass*>& OldClassToNewClass);
 private:
 	FBlueprintCompilationManager();
 };
diff --git a/Engine/Source/Editor/UnrealEd/Private/PackageTools.cpp b/Engine/Source/Editor/UnrealEd/Private/PackageTools.cpp
--- a/Engine/Source/Editor/UnrealEd/Private/PackageTools.cpp
+++ b/Engine/Source/Editor/UnrealEd/Private/PackageTools.cpp
@@ -2,6 +2,7 @@
 
 
 #include "PackageTools.h"
+#include "BlueprintCompilationManager.h"
 #include "UObject/PackageReload.h"
 #include "Misc/MessageDialog.h"
 #include "Misc/Paths.h"
@@ -846,6 +847,28 @@ UPackageTools::UPackageTools(const FObjectInitializer& ObjectInitializer)
 
 		if (InPackageReloadPhase == EPackageReloadPhase::OnPackageFixup)
 		{
+			TMap<UClass*, UClass*> OldClassToNewClass;
+			for (const auto& RepointedObjectPair : InPackageReloadedEvent->GetRepointedObjects())
+			{
+				UObject* OldObject = RepointedObjectPair.Key;
+				UObject* NewObject = RepointedObjectPair.Value;
+
+				if(OldObject && NewObject)
+				{
+					UClass* OldObjectAsClass = Cast<UClass>(OldObject);
+					if(OldObjectAsClass)
+					{
+						UClass* NewObjectAsClass = Cast<UClass>(NewObject);
+						if(ensureMsgf(NewObjectAsClass, TEXT("Class object replaced with non-class object: %s %s"), *(OldObject->GetName()), *(NewObject->GetName())))
+						{
+							OldClassToNewClass.Add(OldObjectAsClass, NewObjectAsClass);
+						}
+					}
+				}
+			}
+
+			FBlueprintCompilationManager::ReparentHierarchies(OldClassToNewClass);
+
 			for (const auto& RepointedObjectPair : InPackageReloadedEvent->GetRepointedObjects())
 			{
 				UObject* OldObject = RepointedObjectPair.Key;
diff --git a/Engine/Source/Runtime/CoreUObject/Private/UObject/PackageReload.cpp b/Engine/Source/Runtime/CoreUObject/Private/UObject/PackageReload.cpp
--- a/Engine/Source/Runtime/CoreUObject/Private/UObject/PackageReload.cpp
+++ b/Engine/Source/Runtime/CoreUObject/Private/UObject/PackageReload.cpp
@@ -566,6 +566,16 @@ void ReloadPackages(const TArrayView<FReloadPackageData>& InPackagesToReload, TA
 			{
 				UObject* PotentialReferencer = *ObjIter;
 
+				// Mutating the old versions of classes can result in us replacing the SuperStruct pointer, which results
+				// in class layout change and subsequently crashes because instances will not match this new class layout:
+				if(UClass* AsClass = Cast<UClass>(PotentialReferencer))
+				{
+					if(AsClass->HasAnyClassFlags(CLASS_NewerVersionExists))
+					{
+						continue;
+					}
+				}
+
 				FixingUpReferencesSlowTask.EnterProgressFrame(1.0f);
 
 				PackageReloadInternal::FReplaceObjectReferencesArchive ReplaceRefsArchive(PotentialReferencer, OldObjectToNewData, ExistingPackages.Refs, NewPackages.Refs);
