References to the old CDO that are referenced in other blueprints are now updated by the blueprint compilation manager

diff --git a/Engine/Source/Editor/Kismet/Private/BlueprintCompilationManager.cpp b/Engine/Source/Editor/Kismet/Private/BlueprintCompilationManager.cpp
--- a/Engine/Source/Editor/Kismet/Private/BlueprintCompilationManager.cpp
+++ b/Engine/Source/Editor/Kismet/Private/BlueprintCompilationManager.cpp
@@ -1644,6 +1644,15 @@ void FBlueprintCompilationManagerImpl::ReinstanceBatch(TArray<FReinstancingJob>&
 		UClass* OldClass = CurrentReinstancer->DuplicatedClass;
 		if(OldClass)
 		{
+			UClass* NewClass = CurrentReinstancer->ClassToReinstance;
+            if (NewClass && 
+				OldClass->ClassDefaultObject && 
+				NewClass->ClassDefaultObject &&
+				OldClass->ClassDefaultObject != NewClass->ClassDefaultObject)
+            {
+                OldArchetypeToNewArchetype.Add(OldClass->ClassDefaultObject, NewClass->ClassDefaultObject);
+            }
+
 			TArray<UObject*> ArchetypeObjects;
 			GetObjectsOfClass(OldClass, ArchetypeObjects, false);
 			
