Harden access to timeline data in the compilation manager

diff --git a/Engine/Source/Editor/Kismet/Private/BlueprintCompilationManager.cpp b/Engine/Source/Editor/Kismet/Private/BlueprintCompilationManager.cpp
--- a/Engine/Source/Editor/Kismet/Private/BlueprintCompilationManager.cpp
+++ b/Engine/Source/Editor/Kismet/Private/BlueprintCompilationManager.cpp
@@ -2207,19 +2207,22 @@ UClass* FBlueprintCompilationManagerImpl::FastGenerateSkeletonClass(UBlueprint*
 	
 	for (UTimelineTemplate* Timeline : BP->Timelines)
 	{
-		// If the timeline hasn't gone through post load that means that the cache isn't up to date, so force an update on it
-		if (Timeline->HasAllFlags(RF_NeedPostLoad) && Timeline->GetLinkerCustomVersion(FFortniteMainBranchObjectVersion::GUID) < FFortniteMainBranchObjectVersion::StoreTimelineNamesInTemplate)
+		if(Timeline)
 		{
-			FUpdateTimelineCachedNames::Execute(Timeline);
-		}
+			// If the timeline hasn't gone through post load that means that the cache isn't up to date, so force an update on it
+			if (Timeline->HasAllFlags(RF_NeedPostLoad) && Timeline->GetLinkerCustomVersion(FFortniteMainBranchObjectVersion::GUID) < FFortniteMainBranchObjectVersion::StoreTimelineNamesInTemplate)
+			{
+				FUpdateTimelineCachedNames::Execute(Timeline);
+			}
 
-		for (const FTTEventTrack& EventTrack : Timeline->EventTracks)
-		{
-			MakeEventFunction(EventTrack.GetFunctionName(), EFunctionFlags::FUNC_None, TArray<UEdGraphPin*>(), TArray< TSharedPtr<FUserPinInfo> >(), nullptr, false);
-		}
+			for (const FTTEventTrack& EventTrack : Timeline->EventTracks)
+			{
+				MakeEventFunction(EventTrack.GetFunctionName(), EFunctionFlags::FUNC_None, TArray<UEdGraphPin*>(), TArray< TSharedPtr<FUserPinInfo> >(), nullptr, false);
+			}
 		
-		MakeEventFunction(Timeline->GetUpdateFunctionName(), EFunctionFlags::FUNC_None, TArray<UEdGraphPin*>(), TArray< TSharedPtr<FUserPinInfo> >(), nullptr, false);
-		MakeEventFunction(Timeline->GetFinishedFunctionName(), EFunctionFlags::FUNC_None, TArray<UEdGraphPin*>(), TArray< TSharedPtr<FUserPinInfo> >(), nullptr, false);
+			MakeEventFunction(Timeline->GetUpdateFunctionName(), EFunctionFlags::FUNC_None, TArray<UEdGraphPin*>(), TArray< TSharedPtr<FUserPinInfo> >(), nullptr, false);
+			MakeEventFunction(Timeline->GetFinishedFunctionName(), EFunctionFlags::FUNC_None, TArray<UEdGraphPin*>(), TArray< TSharedPtr<FUserPinInfo> >(), nullptr, false);
+		}
 	}
 
 	{
