Fix new reparenting logic for blueprint reloading - it did not handle ParentClass being changed by blueprint reload
#rnx

diff --git a/Engine/Source/Editor/Kismet/Private/BlueprintCompilationManager.cpp b/Engine/Source/Editor/Kismet/Private/BlueprintCompilationManager.cpp
--- a/Engine/Source/Editor/Kismet/Private/BlueprintCompilationManager.cpp
+++ b/Engine/Source/Editor/Kismet/Private/BlueprintCompilationManager.cpp
@@ -529,6 +529,8 @@ void FBlueprintCompilationManagerImpl::FlushCompilationQueueImpl(bool bSuppressB
 			if(UClass* OldSkeletonClass = BP->SkeletonGeneratedClass)
 			{
 				TArray<UClass*> SkeletonClassesToReparentList;
+				// Has to be recursive gather of children because instances of a UClass will cache information about
+				// classes that are above their immediate parent (e.g. ClassConstructor):
 				GetDerivedClasses(OldSkeletonClass, SkeletonClassesToReparentList);
 		
 				for(UClass* ChildClass : SkeletonClassesToReparentList)
@@ -1349,23 +1351,37 @@ void FBlueprintCompilationManagerImpl::GetDefaultValue(const UClass* ForClass, c
 void FBlueprintCompilationManagerImpl::ReparentHierarchies(const TMap<UClass*, UClass*>& OldClassToNewClass)
 {
 	// something has decided to replace instances of a class. We need to update all the children of those types:
-	TSet<TPair<UClass*, UClass*>> ClassesToReparent;
+	TSet<TPair<UClass*, UClass*>> ClassesToReinstance;
 	for(const TPair<UClass*, UClass*>& OldToNewClass : OldClassToNewClass)
 	{
 		TArray<UClass*> DerivedClasses;
-		GetDerivedClasses(OldToNewClass.Key, DerivedClasses, false);
+		// Just like when compiling we have to gather all children, not just immediate children. This is so that we can
+		// update things like the ClassConstructor pointer in case it changed:
+		GetDerivedClasses(OldToNewClass.Key, DerivedClasses);
 
 		for(UClass* DerivedClass : DerivedClasses)
 		{
-			ClassesToReparent.Add(
-				TPair<UClass*, UClass*>(DerivedClass, OldToNewClass.Value)
-			);
+			if(DerivedClass->GetSuperClass() == OldToNewClass.Key)
+			{
+				// need to reparent, change the old parent class to the new one:
+				ClassesToReinstance.Add(
+					TPair<UClass*, UClass*>(DerivedClass, OldToNewClass.Value)
+				);
+			}
+			else
+			{
+				// just need to reinstance, parent class can remain the same as
+				// it is generally stable (outside of hotreload/asset reload):
+				ClassesToReinstance.Add(
+					TPair<UClass*, UClass*>(DerivedClass, DerivedClass->GetSuperClass())
+				);
+			}
 		}
 	}
 
 	// create reinstancers:
 	TArray<FReinstancingJob> Reinstancers;
-	for(const TPair<UClass*, UClass*> OldToNewParentClass : ClassesToReparent)
+	for(const TPair<UClass*, UClass*> OldToNewParentClass : ClassesToReinstance)
 	{
 		Reinstancers.Push( {
 			TSharedPtr<FBlueprintCompileReinstancer>(
@@ -1383,21 +1399,35 @@ void FBlueprintCompilationManagerImpl::ReparentHierarchies(const TMap<UClass*, U
 	for(const FReinstancingJob& ReinstancingJob : Reinstancers)
 	{
 		UClass* ClassToReinstance = ReinstancingJob.Reinstancer->ClassToReinstance;
-		ClassToReinstance->SetSuperStruct(OldClassToNewClass.FindChecked(ClassToReinstance->GetSuperClass()));
+		ClassToReinstance->ClassConstructor = nullptr;
+		ClassToReinstance->ClassVTableHelperCtorCaller = nullptr;
+		ClassToReinstance->ClassAddReferencedObjects = nullptr;
+
+		UClass* const* NewParent = OldClassToNewClass.Find(ClassToReinstance->GetSuperClass());
+		if(NewParent)
+		{
+			check(*NewParent);
+			ClassToReinstance->SetSuperStruct(*NewParent);
+		}
+
+		ClassToReinstance->Bind();
+		ClassToReinstance->ClearFunctionMapsCaches();
+		ClassToReinstance->StaticLink(true);
+
 		OldClassToNewClassIncludingChildren.Add(ReinstancingJob.Reinstancer->DuplicatedClass, ClassToReinstance);
 	}
 
 	// reparenting done, reinstance the hierarchy and update archetypes:
 	ReinstanceBatch(Reinstancers, OldClassToNewClassIncludingChildren, nullptr);
 
-	// reinstance instances - but only the classes we created
-	TMap<UClass*, UClass*> ClassesToReinstance;
+	// reinstance instances - but only the classes we created, rest will be handled by caller
+	TMap<UClass*, UClass*> OldClassToNewClassDerivedTypes;
 	for(const FReinstancingJob& ReinstancingJob : Reinstancers)
 	{
-		ClassesToReinstance.Add(ReinstancingJob.Reinstancer->DuplicatedClass, ReinstancingJob.Reinstancer->ClassToReinstance);
+		OldClassToNewClassDerivedTypes.Add(ReinstancingJob.Reinstancer->DuplicatedClass, ReinstancingJob.Reinstancer->ClassToReinstance);
 	}
 	TGuardValue<bool> ReinstancingGuard(GIsReinstancing, true);
-	FBlueprintCompileReinstancer::BatchReplaceInstancesOfClass( ClassesToReinstance, /* bArchetypesAreUpToDate */ true );
+	FBlueprintCompileReinstancer::BatchReplaceInstancesOfClass( OldClassToNewClassDerivedTypes, /* bArchetypesAreUpToDate */ true );
 }
 
 
