Fix for an editor crash on contextual promotion of an output pin on a GetClassDefaults() node to a new variable.

diff --git a/Engine/Source/Editor/Kismet/Private/BlueprintEditor.cpp b/Engine/Source/Editor/Kismet/Private/BlueprintEditor.cpp
--- a/Engine/Source/Editor/Kismet/Private/BlueprintEditor.cpp
+++ b/Engine/Source/Editor/Kismet/Private/BlueprintEditor.cpp
@@ -3717,6 +3717,7 @@ void FBlueprintEditor::AppendExtraCompilerResults(TSharedPtr<IMessageLogListing>
 
 void FBlueprintEditor::DoPromoteToVariable( UBlueprint* InBlueprint, UEdGraphPin* InTargetPin, bool bInToMemberVariable )
 {
+	FName PinName = InTargetPin->PinName;
 	UEdGraphNode* PinNode = InTargetPin->GetOwningNode();
 	check(PinNode);
 	UEdGraph* GraphObj = PinNode->GetGraph();
@@ -3750,6 +3751,9 @@ void FBlueprintEditor::DoPromoteToVariable( UBlueprint* InBlueprint, UEdGraphPin
 
 	if (bWasSuccessful)
 	{
+		// The owning node may have been reconstructed as a result of adding a new variable above, so ensure the pin reference is up-to-date.
+		InTargetPin = PinNode->FindPinChecked(PinName);
+
 		// Create the new setter node
 		FEdGraphSchemaAction_K2NewNode NodeInfo;
 
