Fixed compilation error due to delegate properties referencing function signatures on out of date classes

diff --git a/Engine/Source/Editor/Kismet/Private/BlueprintCompilationManager.cpp b/Engine/Source/Editor/Kismet/Private/BlueprintCompilationManager.cpp
--- a/Engine/Source/Editor/Kismet/Private/BlueprintCompilationManager.cpp
+++ b/Engine/Source/Editor/Kismet/Private/BlueprintCompilationManager.cpp
@@ -79,6 +79,7 @@
 #define VERIFY_NO_BAD_SKELETON_REFERENCES 0
 
 struct FReinstancingJob;
+struct FSkeletonFixupData;
 
 struct FBlueprintCompilationManagerImpl : public FGCObject
 {
@@ -99,7 +100,7 @@ struct FBlueprintCompilationManagerImpl : public FGCObject
 	static void ReparentHierarchies(const TMap<UClass*, UClass*>& OldClassToNewClass);
 	static void BuildDSOMap(UObject* OldObject, UObject* NewObject, TMap<UObject*, UObject*>& OutOldToNewDSO);
 	static void ReinstanceBatch(TArray<FReinstancingJob>& Reinstancers, TMap< UClass*, UClass* >& InOutOldToNewClassMap, FUObjectSerializeContext* InLoadContext);
-	static UClass* FastGenerateSkeletonClass(UBlueprint* BP, FKismetCompilerContext& CompilerContext, bool bIsSkeletonOnly);
+	static UClass* FastGenerateSkeletonClass(UBlueprint* BP, FKismetCompilerContext& CompilerContext, bool bIsSkeletonOnly, TArray<FSkeletonFixupData>& OutSkeletonFixupData);
 	static bool IsQueuedForCompilation(UBlueprint* BP);
 	
 	// Declaration of archive to fix up bytecode references of blueprints that are actively compiled:
@@ -291,6 +292,14 @@ enum class ECompilationManagerJobType
 	RelinkOnly,
 };
 
+// Currently only used to fix up delegate parameters on skeleton ufunctions, resolving the cyclical dependency,
+// could be augmented if similar cases arise:
+struct FSkeletonFixupData
+{
+	FSimpleMemberReference MemberReference;
+	UProperty* DelegateProperty;
+};
+
 struct FCompilerData
 {
 	explicit FCompilerData(UBlueprint* InBP, ECompilationManagerJobType InJobType, FCompilerResultsLog* InResultsLogOverride, EBlueprintCompileOptions UserOptions, bool bBytecodeOnly)
@@ -346,6 +355,7 @@ struct FCompilerData
 	TSharedPtr<FKismetCompilerContext> Compiler;
 	FKismetCompilerOptions InternalOptions;
 	TSharedPtr<FBlueprintCompileReinstancer> Reinstancer;
+	TArray<FSkeletonFixupData> SkeletonFixupData;
 
 	ECompilationManagerJobType JobType;
 	bool bPackageWasDirty;
@@ -679,7 +689,7 @@ void FBlueprintCompilationManagerImpl::FlushCompilationQueueImpl(bool bSuppressB
 						BlueprintsCompiledOrSkeletonCompiled->Add(BP);
 					}
 
-					BP->SkeletonGeneratedClass = FastGenerateSkeletonClass(BP, *(CompilerData.Compiler), CompilerData.IsSkeletonOnly());
+					BP->SkeletonGeneratedClass = FastGenerateSkeletonClass(BP, *(CompilerData.Compiler), CompilerData.IsSkeletonOnly(), CompilerData.SkeletonFixupData);
 					UBlueprintGeneratedClass* AuthoritativeClass = Cast<UBlueprintGeneratedClass>(BP->GeneratedClass);
 					if(AuthoritativeClass && bSkipUnneededDependencyCompilation)
 					{
@@ -738,6 +748,25 @@ void FBlueprintCompilationManagerImpl::FlushCompilationQueueImpl(bool bSuppressB
 				}
 			}
 
+			// Fix up delegate parameters on skeleton class UFunctions, as they have a direct reference to a UFunction*
+			// that may have been created as part of skeleton generation:
+			for (FCompilerData& CompilerData : CurrentlyCompilingBPs)
+			{
+				UBlueprint* BP = CompilerData.BP;
+				TArray< FSkeletonFixupData >& ParamsToFix = CompilerData.SkeletonFixupData;
+				for( const FSkeletonFixupData& FixupData : ParamsToFix )
+				{
+					if(UDelegateProperty* DelegateProperty = Cast<UDelegateProperty>(FixupData.DelegateProperty))
+					{
+						DelegateProperty->SignatureFunction = FMemberReference::ResolveSimpleMemberReference<UFunction>(FixupData.MemberReference, BP->SkeletonGeneratedClass);
+					}
+					else if(UMulticastDelegateProperty* MCDelegateProperty = Cast<UMulticastDelegateProperty>(FixupData.DelegateProperty))
+					{
+						MCDelegateProperty->SignatureFunction = FMemberReference::ResolveSimpleMemberReference<UFunction>(FixupData.MemberReference, BP->SkeletonGeneratedClass);
+					}
+				}
+			}
+
 			// Skip further compilation for blueprints that are being bytecode compiled as a dependency of something that has
 			// not had a change in its function parameters:
 			auto DependenciesAreCompiled = [&BlueprintsWithSignatureChanges](FCompilerData& Data)
@@ -1859,7 +1888,7 @@ void FBlueprintCompilationManagerImpl::ReinstanceBatch(TArray<FReinstancingJob>&
 	Notes to maintainers: any UObject created here and outered to the resulting class must be marked as transient
 	or you will create a cook error!
 */
-UClass* FBlueprintCompilationManagerImpl::FastGenerateSkeletonClass(UBlueprint* BP, FKismetCompilerContext& CompilerContext, bool bIsSkeletonOnly)
+UClass* FBlueprintCompilationManagerImpl::FastGenerateSkeletonClass(UBlueprint* BP, FKismetCompilerContext& CompilerContext, bool bIsSkeletonOnly, TArray<FSkeletonFixupData>& OutSkeletonFixupData)
 {
 	const UEdGraphSchema_K2* Schema = GetDefault<UEdGraphSchema_K2>();
 
@@ -1914,7 +1943,7 @@ UClass* FBlueprintCompilationManagerImpl::FastGenerateSkeletonClass(UBlueprint*
 	Ret->ClassGeneratedBy = BP;
 
 	// This is a version of PrecompileFunction that does not require 'terms' and graph cloning:
-	const auto MakeFunction = [Ret, ParentClass, Schema, BP, &MessageLog]
+	const auto MakeFunction = [Ret, ParentClass, Schema, BP, &MessageLog, &OutSkeletonFixupData]
 		(	FName FunctionNameFName, 
 			UField**& InCurrentFieldStorageLocation, 
 			UField**& InCurrentParamStorageLocation, 
@@ -2037,6 +2066,23 @@ UClass* FBlueprintCompilationManagerImpl::FastGenerateSkeletonClass(UBlueprint*
 							// ALWAYS pass array parameters as out params, so they're set up as passed by ref
 							Param->PropertyFlags |= CPF_OutParm;
 						}
+						// Delegate properties have a direct reference to a UFunction that we may currently be generating, so we're going
+						// to track them and fix them after all UFunctions have been generated. As you can tell we're tightly coupled
+						// to the implementation of CreatePropertyOnScope
+						else if( UDelegateProperty* DelegateProp = Cast<UDelegateProperty>(Param))
+						{
+							OutSkeletonFixupData.Add( {
+								Pin->PinType.PinSubCategoryMemberReference,
+								DelegateProp
+							} );
+						}
+						else if( UMulticastDelegateProperty* MCDelegateProp = Cast<UMulticastDelegateProperty>(Param))
+						{
+							OutSkeletonFixupData.Add( {
+								Pin->PinType.PinSubCategoryMemberReference,
+								MCDelegateProp
+							} );
+						}
 
 						*InCurrentParamStorageLocation = Param;
 						InCurrentParamStorageLocation = &Param->Next;
diff --git a/Engine/Source/Editor/KismetCompiler/Private/KismetCompilerMisc.cpp b/Engine/Source/Editor/KismetCompiler/Private/KismetCompilerMisc.cpp
--- a/Engine/Source/Editor/KismetCompiler/Private/KismetCompilerMisc.cpp
+++ b/Engine/Source/Editor/KismetCompiler/Private/KismetCompilerMisc.cpp
@@ -1073,6 +1073,9 @@ UProperty* FKismetCompilerUtilities::CreatePrimitiveProperty(UObject* PropertySc
 /** Creates a property named PropertyName of type PropertyType in the Scope or returns NULL if the type is unknown, but does *not* link that property in */
 UProperty* FKismetCompilerUtilities::CreatePropertyOnScope(UStruct* Scope, const FName& PropertyName, const FEdGraphPinType& Type, UClass* SelfClass, EPropertyFlags PropertyFlags, const UEdGraphSchema_K2* Schema, FCompilerResultsLog& MessageLog)
 {
+	// When creating properties that depend on other properties (e.g. UDelegateProperty/UMulticastDelegateProperty::SignatureFunction)
+	// you may need to update fixup logic in the compilation manager.
+		
 	const EObjectFlags ObjectFlags = RF_Public;
 
 	FName ValidatedPropertyName = PropertyName;
diff --git a/Engine/Source/Runtime/Engine/Classes/Engine/MemberReference.h b/Engine/Source/Runtime/Engine/Classes/Engine/MemberReference.h
--- a/Engine/Source/Runtime/Engine/Classes/Engine/MemberReference.h
+++ b/Engine/Source/Runtime/Engine/Classes/Engine/MemberReference.h
@@ -502,7 +502,7 @@ public:
 	}
 
 	template<class TFieldType>
-	static TFieldType* ResolveSimpleMemberReference(const FSimpleMemberReference& Reference)
+	static TFieldType* ResolveSimpleMemberReference(const FSimpleMemberReference& Reference, UClass* SelfScope = nullptr)
 	{
 		FMemberReference TempMemberReference;
 
@@ -511,11 +511,11 @@ public:
 		TempMemberReference.MemberGuid   = Reference.MemberGuid;
 		TempMemberReference.MemberParent = Reference.MemberParent;
 
-		auto Result = TempMemberReference.ResolveMember<TFieldType>();
+		auto Result = TempMemberReference.ResolveMember<TFieldType>(SelfScope);
 		if (!Result && (Name != Reference.MemberName))
 		{
 			TempMemberReference.MemberName = Reference.MemberName;
-			Result = TempMemberReference.ResolveMember<TFieldType>();
+			Result = TempMemberReference.ResolveMember<TFieldType>(SelfScope);
 		}
 
 		return Result;
