Code cleanup: Improved a check by turning it into a checkf so it displays the values of the components of its expression if it's hit.  Now if the check is hit, it will have enough detail in the message to make clear exactly what went wrong.

#OKForGitHub public

#UE4 #RNX

diff --git a/Engine/Source/Editor/Kismet/Private/FindInBlueprintManager.cpp b/Engine/Source/Editor/Kismet/Private/FindInBlueprintManager.cpp
--- a/Engine/Source/Editor/Kismet/Private/FindInBlueprintManager.cpp
+++ b/Engine/Source/Editor/Kismet/Private/FindInBlueprintManager.cpp
@@ -1958,7 +1958,12 @@ FString FFindInBlueprintSearchManager::QuerySingleBlueprint(UBlueprint* InBluepr
 
 		int32* ArrayIdx = SearchMap.Find(Key);
 		// This should always be true since we make sure to refresh the search data for this Blueprint when doing the search, unless we do not rebuild the searchable data
-		check((bInRebuildSearchData && ArrayIdx && *ArrayIdx < SearchArray.Num()) || !bInRebuildSearchData);
+		checkf((bInRebuildSearchData && ArrayIdx && *ArrayIdx < SearchArray.Num()) || !bInRebuildSearchData,
+			TEXT("bInRebuildSearchData: %s, ArrayIdx is %s:%d, SearchArray.Num():%d"),
+			bInRebuildSearchData ? TEXT("true") : TEXT("false"),
+			(ArrayIdx != nullptr) ? TEXT("Valid") : TEXT("Invalid"),
+			(ArrayIdx != nullptr) ? *ArrayIdx : INDEX_NONE,
+			SearchArray.Num());
 
 		if (ArrayIdx)
 		{
