Clear the instanced watch list before repopulating it when we hit a blueprint breakpoint.
Make sure we're only adding watches from each blueprint in the stack once.

diff --git a/Engine/Source/Editor/Kismet/Private/WatchPointViewer.cpp b/Engine/Source/Editor/Kismet/Private/WatchPointViewer.cpp
--- a/Engine/Source/Editor/Kismet/Private/WatchPointViewer.cpp
+++ b/Engine/Source/Editor/Kismet/Private/WatchPointViewer.cpp
@@ -765,8 +765,11 @@ void WatchViewer::UpdateInstancedWatchDisplay()
 {
 #if DO_BLUEPRINT_GUARD
 	{
+		Private_InstanceWatchSource.Reset();
 		const TArray<const FFrame*>& ScriptStack = FBlueprintExceptionTracker::Get().ScriptStack;
 
+		TSet<const UBlueprint*> SeenBlueprints;
+
 		for (const FFrame* ScriptFrame : ScriptStack)
 		{
 			UObject* BlueprintInstance = ScriptFrame ? ScriptFrame->Object : nullptr;
@@ -777,6 +780,13 @@ void WatchViewer::UpdateInstancedWatchDisplay()
 				continue;
 			}
 
+			// Only add watchpoints from each blueprint once
+			if (SeenBlueprints.Contains(BlueprintObj))
+			{
+				continue;
+			}
+			SeenBlueprints.Add(BlueprintObj);
+
 			FText BlueprintName = FText::FromString(BlueprintObj->GetName());
 
 			// Don't show info for the CDO
