Fixed copy paste error introduced when refactoring node alignment, caused align bottom not to function correctly for blueprints
#rnx

diff --git a/Engine/Source/Editor/Kismet/Private/BlueprintEditor.cpp b/Engine/Source/Editor/Kismet/Private/BlueprintEditor.cpp
--- a/Engine/Source/Editor/Kismet/Private/BlueprintEditor.cpp
+++ b/Engine/Source/Editor/Kismet/Private/BlueprintEditor.cpp
@@ -5214,7 +5214,7 @@ void FBlueprintEditor::OnAlignBottom()
 	TSharedPtr<SGraphEditor> FocusedGraphEd = FocusedGraphEdPtr.Pin();
 	if (FocusedGraphEd.IsValid())
 	{
-		FocusedGraphEd->OnAlignMiddle();
+		FocusedGraphEd->OnAlignBottom();
 	}
 }
 
