Fix a few cases where widgets were in two places at the same time, causing parent pointers to break

diff --git a/Engine/Source/Editor/Kismet/Private/SBlueprintEditorSelectedDebugObjectWidget.cpp b/Engine/Source/Editor/Kismet/Private/SBlueprintEditorSelectedDebugObjectWidget.cpp
--- a/Engine/Source/Editor/Kismet/Private/SBlueprintEditorSelectedDebugObjectWidget.cpp
+++ b/Engine/Source/Editor/Kismet/Private/SBlueprintEditorSelectedDebugObjectWidget.cpp
@@ -28,12 +28,9 @@ void SBlueprintEditorSelectedDebugObjectWidget::Construct(const FArguments& InAr
 {
 	BlueprintEditor = InBlueprintEditor;
 
-	const TSharedRef<SWidget> BrowseButton = PropertyCustomizationHelpers::MakeBrowseButton(FSimpleDelegate::CreateSP(this, &SBlueprintEditorSelectedDebugObjectWidget::SelectedDebugObject_OnClicked));
-	BrowseButton->SetVisibility(TAttribute<EVisibility>::Create(TAttribute<EVisibility>::FGetter::CreateSP(this, &SBlueprintEditorSelectedDebugObjectWidget::IsSelectDebugObjectButtonVisible)));
-	BrowseButton->SetToolTipText(LOCTEXT("DebugSelectActor", "Select this Actor in level"));
-
 	GenerateDebugWorldNames(false);
 	GenerateDebugObjectNames(false);
+
 	LastObjectObserved = DebugObjects[0];
 
 	DebugWorldsComboBox = SNew(STextComboBox)
@@ -62,73 +59,12 @@ void SBlueprintEditorSelectedDebugObjectWidget::Construct(const FArguments& InAr
 		.OnSelectionChanged(this, &SBlueprintEditorSelectedDebugObjectWidget::DebugObjectSelectionChanged)
 		.AddMetaData<FTagMetaData>(TEXT("SelectDebugObjectCobmo"));
 
-	TSharedRef<SWidget> DebugObjectSelectionWidget =
-		SNew(SHorizontalBox)
-		+ SHorizontalBox::Slot()
-		.AutoWidth()
-		[
-			DebugObjectsComboBox.ToSharedRef()
-		]
-	+ SHorizontalBox::Slot()
-		.AutoWidth()
-		.HAlign(HAlign_Right)
-		.VAlign(VAlign_Center)
-		.Padding(2.0f)
-		[
-			BrowseButton
-		];
-
 	ChildSlot
-		[
-			SNew(SLevelOfDetailBranchNode)
-			.UseLowDetailSlot(FMultiBoxSettings::UseSmallToolBarIcons)
-			.LowDetail()
-			[
-				// Horizontal Layout when using small icons
-				SNew(SHorizontalBox)
-				+ SHorizontalBox::Slot()
-				.AutoWidth()
-				[
-					DebugWorldsComboBox.ToSharedRef()
-				]
-
-				+ SHorizontalBox::Slot()
-					.AutoWidth()
-					[
-						DebugObjectSelectionWidget
-					]
-			]
-			.HighDetail()
-				[
-					SNew(SVerticalBox)
-					+ SVerticalBox::Slot()
-					.AutoHeight()
-					.VAlign(VAlign_Bottom)
-					[
-						// Vertical Layout when using normal size icons
-						SNew(SVerticalBox)
-						+ SVerticalBox::Slot()
-						.AutoHeight()
-						[
-							DebugWorldsComboBox.ToSharedRef()
-						]
-
-						+ SVerticalBox::Slot()
-							.AutoHeight()
-							[
-								DebugObjectSelectionWidget
-							]
-					]
-					+ SVerticalBox::Slot()
-						.AutoHeight()
-						.HAlign(HAlign_Center)
-						.Padding(2.0f)
-						[
-							SNew(STextBlock)
-							.Text(LOCTEXT("DebugSelectTitle", "Debug Filter"))
-						]
-				]
-		];
+	[
+		SNew(SLevelOfDetailBranchNode)
+		.UseLowDetailSlot(FMultiBoxSettings::UseSmallToolBarIcons)
+		.OnGetActiveDetailSlotContent(this, &SBlueprintEditorSelectedDebugObjectWidget::OnGetActiveDetailSlotContent)
+	];
 }
 
 void SBlueprintEditorSelectedDebugObjectWidget::Tick( const FGeometry& AllottedGeometry, const double InCurrentTime, const float InDeltaTime )
@@ -191,6 +127,76 @@ const FString& SBlueprintEditorSelectedDebugObjectWidget::GetDebugAllWorldsStrin
 	return NSLOCTEXT("BlueprintEditor", "DebugWorldNothingSelected", "All Worlds").ToString();
 }
 
+TSharedRef<SWidget> SBlueprintEditorSelectedDebugObjectWidget::OnGetActiveDetailSlotContent(bool bChangedToHighDetail)
+{
+
+	const TSharedRef<SWidget> BrowseButton = PropertyCustomizationHelpers::MakeBrowseButton(FSimpleDelegate::CreateSP(this, &SBlueprintEditorSelectedDebugObjectWidget::SelectedDebugObject_OnClicked));
+	BrowseButton->SetVisibility(TAttribute<EVisibility>::Create(TAttribute<EVisibility>::FGetter::CreateSP(this, &SBlueprintEditorSelectedDebugObjectWidget::IsSelectDebugObjectButtonVisible)));
+	BrowseButton->SetToolTipText(LOCTEXT("DebugSelectActor", "Select this Actor in level"));
+
+	TSharedRef<SWidget> DebugObjectSelectionWidget =
+		SNew(SHorizontalBox)
+		+ SHorizontalBox::Slot()
+		.AutoWidth()
+		[
+			DebugObjectsComboBox.ToSharedRef()
+		]
+		+ SHorizontalBox::Slot()
+		.AutoWidth()
+		.HAlign(HAlign_Right)
+		.VAlign(VAlign_Center)
+		.Padding(2.0f)
+		[
+			BrowseButton
+		];
+
+	if (!bChangedToHighDetail)
+	{
+		return
+			SNew(SHorizontalBox)
+			+ SHorizontalBox::Slot()
+			.AutoWidth()
+			[
+				DebugWorldsComboBox.ToSharedRef()
+			]
+			+ SHorizontalBox::Slot()
+			.AutoWidth()
+			[
+				DebugObjectSelectionWidget
+			];
+	}
+	else
+	{
+		return
+			SNew(SVerticalBox)
+			+ SVerticalBox::Slot()
+			.AutoHeight()
+			.VAlign(VAlign_Bottom)
+			[
+				// Vertical Layout when using normal size icons
+				SNew(SVerticalBox)
+				+ SVerticalBox::Slot()
+				.AutoHeight()
+				[
+					DebugWorldsComboBox.ToSharedRef()
+				]
+				+ SVerticalBox::Slot()
+				.AutoHeight()
+				[
+					DebugObjectSelectionWidget
+				]
+			]
+			+ SVerticalBox::Slot()
+			.AutoHeight()
+			.HAlign(HAlign_Center)
+			.Padding(2.0f)
+			[
+				SNew(STextBlock)
+				.Text(LOCTEXT("DebugSelectTitle", "Debug Filter"))
+			];
+	}
+}
+
 void SBlueprintEditorSelectedDebugObjectWidget::OnRefresh()
 {
 	if (GetBlueprintObj())
diff --git a/Engine/Source/Editor/Kismet/Private/SBlueprintEditorSelectedDebugObjectWidget.h b/Engine/Source/Editor/Kismet/Private/SBlueprintEditorSelectedDebugObjectWidget.h
--- a/Engine/Source/Editor/Kismet/Private/SBlueprintEditorSelectedDebugObjectWidget.h
+++ b/Engine/Source/Editor/Kismet/Private/SBlueprintEditorSelectedDebugObjectWidget.h
@@ -55,12 +55,12 @@ private:
 
 	EVisibility IsDebugWorldComboVisible() const;
 
-private:
 	/* Returns the string to indicate no debug object is selected */
 	const FString& GetNoDebugString() const;
 
 	const FString& GetDebugAllWorldsString() const;
 
+	TSharedRef<SWidget> OnGetActiveDetailSlotContent(bool bChangedToHighDetail);
 private:
 	/** Pointer back to the blueprint editor tool that owns us */
 	TWeakPtr<FBlueprintEditor> BlueprintEditor;
diff --git a/Engine/Source/Editor/KismetWidgets/Private/SLevelOfDetailBranchNode.cpp b/Engine/Source/Editor/KismetWidgets/Private/SLevelOfDetailBranchNode.cpp
--- a/Engine/Source/Editor/KismetWidgets/Private/SLevelOfDetailBranchNode.cpp
+++ b/Engine/Source/Editor/KismetWidgets/Private/SLevelOfDetailBranchNode.cpp
@@ -14,10 +14,21 @@ SLevelOfDetailBranchNode::SLevelOfDetailBranchNode()
 
 void SLevelOfDetailBranchNode::Construct(const FArguments& InArgs)
 {
+	OnGetActiveDetailSlotContent = InArgs._OnGetActiveDetailSlotContent;
 	ShowLowDetailAttr = InArgs._UseLowDetailSlot;
 	ChildSlotLowDetail = InArgs._LowDetail.Widget;
 	ChildSlotHighDetail = InArgs._HighDetail.Widget;
-	ChildSlot[ChildSlotHighDetail];
+
+	if (OnGetActiveDetailSlotContent.IsBound())
+	{
+		const int32 CurrentValue = ShowLowDetailAttr.Get() ? 1 : 0;
+		ChildSlot[OnGetActiveDetailSlotContent.Execute((CurrentValue != 0) ? false : true)];
+	}
+	else
+	{
+		ChildSlot[ChildSlotHighDetail];
+	}
+
 }
 
 void SLevelOfDetailBranchNode::Tick( const FGeometry& AllottedGeometry, const double InCurrentTime, const float InDeltaTime )
@@ -26,9 +37,11 @@ void SLevelOfDetailBranchNode::Tick( const FGeometry& AllottedGeometry, const do
 	if (CurrentValue != LastCachedValue)
 	{
 		LastCachedValue = CurrentValue;
+
 		ChildSlot
 		[
-			(CurrentValue != 0) ? ChildSlotLowDetail : ChildSlotHighDetail
+			OnGetActiveDetailSlotContent.IsBound() ? OnGetActiveDetailSlotContent.Execute((CurrentValue != 0) ? false : true) : SNullWidget::NullWidget
 		];
+
 	}
 }
diff --git a/Engine/Source/Editor/KismetWidgets/Public/SLevelOfDetailBranchNode.h b/Engine/Source/Editor/KismetWidgets/Public/SLevelOfDetailBranchNode.h
--- a/Engine/Source/Editor/KismetWidgets/Public/SLevelOfDetailBranchNode.h
+++ b/Engine/Source/Editor/KismetWidgets/Public/SLevelOfDetailBranchNode.h
@@ -11,6 +11,8 @@
 /////////////////////////////////////////////////////
 // SLevelOfDetailBranchNode
 
+DECLARE_DELEGATE_RetVal_OneParam(TSharedRef<SWidget>, FOnGetActiveDetailSlotContent, bool);
+
 class KISMETWIDGETS_API SLevelOfDetailBranchNode : public SCompoundWidget
 {
 	SLATE_BEGIN_ARGS(SLevelOfDetailBranchNode)
@@ -25,8 +27,19 @@ class KISMETWIDGETS_API SLevelOfDetailBranchNode : public SCompoundWidget
 
 		// The high-detail slot
 		SLATE_NAMED_SLOT(FArguments, HighDetail)
+
+		SLATE_EVENT(FOnGetActiveDetailSlotContent, OnGetActiveDetailSlotContent)
 	SLATE_END_ARGS()
 
+public:
+	SLevelOfDetailBranchNode();
+
+	void Construct(const FArguments& InArgs);
+
+	// SWidget interface
+	virtual void Tick( const FGeometry& AllottedGeometry, const double InCurrentTime, const float InDeltaTime ) override;
+	// End of SWidget interface
+
 protected:
 	// What kind of slot was shown last frame
 	int LastCachedValue;
@@ -40,12 +53,6 @@ protected:
 	// The high-detail child slot
 	TSharedRef<SWidget> ChildSlotHighDetail;
 
-public:
-	SLevelOfDetailBranchNode();
+	FOnGetActiveDetailSlotContent OnGetActiveDetailSlotContent;
 
-	void Construct(const FArguments& InArgs);
-
-	// SWidget interface
-	virtual void Tick( const FGeometry& AllottedGeometry, const double InCurrentTime, const float InDeltaTime ) override;
-	// End of SWidget interface
 };
diff --git a/Engine/Source/Runtime/Slate/Private/Framework/MultiBox/SMenuEntryBlock.cpp b/Engine/Source/Runtime/Slate/Private/Framework/MultiBox/SMenuEntryBlock.cpp
--- a/Engine/Source/Runtime/Slate/Private/Framework/MultiBox/SMenuEntryBlock.cpp
+++ b/Engine/Source/Runtime/Slate/Private/Framework/MultiBox/SMenuEntryBlock.cpp
@@ -932,14 +932,7 @@ void SMenuEntryBlock::BuildMultiBlockWidget(const ISlateStyle* StyleSet, const F
 	if(TutorialName != NAME_None)
 	{
 		TSharedRef<SWidget> ChildWidget = ChildSlot.GetWidget();
-		ChildSlot
-		[
-			SNew( SBox )
-			.AddMetaData<FTagMetaData>(FTagMetaData(TutorialName))
-			[
-				ChildWidget
-			]
-		];
+		ChildWidget->AddMetaData<FTagMetaData>(MakeShared<FTagMetaData>(TutorialName));
 	}
 
 	// Bind our widget's enabled state to whether or not our action can execute
