The editor assumed that renaming a component on creation was part of the creation process (from the undo perspective) and was not made part of any transaction so that undoing after giving the initial name would actually undo the creation. For multi-user, giving a new name, on creation, must be a transaction so that it get recorded and replicated. Fixed the issue by keeping the create transaction on-going until the initial name is given. The rename transaction gets folded under any on-going create. From the undo/redo buffer, this is still a single transaction, but now, the rename gets replicated in the multi-user context.

diff --git a/Engine/Source/Editor/Kismet/Private/SSCSEditor.cpp b/Engine/Source/Editor/Kismet/Private/SSCSEditor.cpp
--- a/Engine/Source/Editor/Kismet/Private/SSCSEditor.cpp
+++ b/Engine/Source/Editor/Kismet/Private/SSCSEditor.cpp
@@ -394,7 +394,6 @@ FReply FSCSRowDragDropOp::DroppedOnPanel(const TSharedRef< class SWidget >& Pane
 FSCSEditorTreeNode::FSCSEditorTreeNode(FSCSEditorTreeNode::ENodeType InNodeType)
 	: ComponentTemplatePtr(nullptr)
 	, NodeType(InNodeType)
-	, bNonTransactionalRename(false)
 	, FilterFlags((uint8)EFilteredState::Unknown)
 {
 }
@@ -905,14 +904,16 @@ void FSCSEditorTreeNode::RemoveChild(FSCSEditorTreeNodePtrType InChildNodePtr)
 	}
 }
 
-void FSCSEditorTreeNode::OnRequestRename(bool bTransactional)
+void FSCSEditorTreeNode::OnRequestRename(TUniquePtr<FScopedTransaction> InOngoingCreateTransaction)
 {
-	bNonTransactionalRename = !bTransactional;
+	OngoingCreateTransaction = MoveTemp(InOngoingCreateTransaction); // Take responsibility to end the 'create + give initial name' transaction.
 	RenameRequestedDelegate.ExecuteIfBound();
 }
 
 void FSCSEditorTreeNode::OnCompleteRename(const FText& InNewName)
 {
+	// If a 'create + give initial name' transaction exists, end it, the object is expected to have its initial name.
+	OngoingCreateTransaction.Reset();
 }
 
 //////////////////////////////////////////////////////////////////////////
@@ -1134,8 +1135,11 @@ void FSCSEditorTreeNodeInstanceAddedComponent::RemoveMeAsChild()
 
 void FSCSEditorTreeNodeInstanceAddedComponent::OnCompleteRename(const FText& InNewName)
 {
-	bool bIsNonTransactionalRename = GetAndClearNonTransactionalRenameFlag();
-	FScopedTransaction TransactionContext(LOCTEXT("RenameComponentVariable", "Rename Component Variable"), !bIsNonTransactionalRename);
+	// If the 'rename' was part of an ongoing component creation, ensure the transaction is ended when the local object goes out of scope. (Must complete after the rename transaction below)
+	TUniquePtr<FScopedTransaction> ScopedCreateTransaction(MoveTemp(OngoingCreateTransaction));
+
+	// If a 'create' transaction is opened, the rename will be folded into it and will be invisible to the 'undo' as create + give a name is really just one operation from the user point of view.
+	FScopedTransaction TransactionContext(LOCTEXT("RenameComponentVariable", "Rename Component Variable"));
 
 	UActorComponent* ComponentInstance = GetComponentTemplate();
 	if(ComponentInstance == nullptr)
@@ -1144,11 +1148,7 @@ void FSCSEditorTreeNodeInstanceAddedComponent::OnCompleteRename(const FText& InN
 	}
 
 	ERenameFlags RenameFlags = REN_DontCreateRedirectors;
-	if (bIsNonTransactionalRename)
-	{
-		RenameFlags |= REN_NonTransactional;
-	}
-	
+
 	// name collision could occur due to e.g. our archetype being updated and causing a conflict with our ComponentInstance:
 	FString NewNameAsString = InNewName.ToString();
 	if(StaticFindObject(UObject::StaticClass(), ComponentInstance->GetOuter(), *NewNameAsString) == nullptr)
@@ -1344,18 +1344,13 @@ UActorComponent* FSCSEditorTreeNode::FindComponentInstanceInActor(const AActor*
 
 void FSCSEditorTreeNodeComponent::OnCompleteRename(const FText& InNewName)
 {
-	FScopedTransaction* TransactionContext = NULL;
-	if (!GetAndClearNonTransactionalRenameFlag())
-	{
-		TransactionContext = new FScopedTransaction(LOCTEXT("RenameComponentVariable", "Rename Component Variable"));
-	}
+	// If the 'rename' was part of the creation process, we need to complete the creation transaction as the component has a user confirmed name. (Must complete after rename transaction below)
+	TUniquePtr<FScopedTransaction> ScopedOngoingCreateTransaction(MoveTemp(OngoingCreateTransaction));
 
-	FBlueprintEditorUtils::RenameComponentMemberVariable(GetBlueprint(), GetSCSNode(), FName(*InNewName.ToString()));
+	// If a 'create' transaction is opened, the rename will be folded into it and will be invisible to the 'undo' as 'create + give initial name' means creating an object from the user point of view.
+	FScopedTransaction RenameTransaction(LOCTEXT("RenameComponentVariable", "Rename Component Variable"));
 
-	if (TransactionContext)
-	{
-		delete TransactionContext;
-	}
+	FBlueprintEditorUtils::RenameComponentMemberVariable(GetBlueprint(), GetSCSNode(), FName(*InNewName.ToString()));
 }
 
 void FSCSEditorTreeNodeComponent::RemoveMeAsChild()
@@ -1416,6 +1411,9 @@ void FSCSEditorTreeNodeRootActor::OnCompleteRename(const FText& InNewName)
 		const FScopedTransaction Transaction(LOCTEXT("SCSEditorRenameActorTransaction", "Rename Actor"));
 		FActorLabelUtilities::RenameExistingActor(Actor, InNewName.ToString());
 	}
+
+	// Not expected to reach here with an ongoing create transaction, but if it does, end it.
+	OngoingCreateTransaction.Reset();
 }
 
 //////////////////////////////////////////////////////////////////////////
@@ -3456,7 +3454,7 @@ void SSCSEditor::Construct( const FArguments& InArgs )
 		);
 
 	CommandList->MapAction( FGenericCommands::Get().Rename,
-			FUIAction( FExecuteAction::CreateSP( this, &SSCSEditor::OnRenameComponent, true ), // true = transactional (i.e. undoable)
+			FUIAction( FExecuteAction::CreateSP( this, &SSCSEditor::OnRenameComponent),
 			FCanExecuteAction::CreateSP( this, &SSCSEditor::CanRenameComponent ) ) 
 		);
 
@@ -3748,7 +3746,7 @@ void SSCSEditor::OnLevelComponentRequestRename(const UActorComponent* InComponen
 	FSCSEditorTreeNodePtrType Node = GetNodeFromActorComponent(InComponent);
 	if (SelectedItems.Contains(Node) && CanRenameComponent())
 	{
-		OnRenameComponent(true);
+		OnRenameComponent();
 	}
 }
 
@@ -5119,7 +5117,9 @@ UActorComponent* SSCSEditor::AddNewComponent( UClass* NewComponentClass, UObject
 		return nullptr;
 	}
 
-	const FScopedTransaction Transaction( LOCTEXT("AddComponent", "Add Component") );
+	// Begin a transaction. The transaction will end when the component name will be provided/confirmed by the user.
+	check(!OngoingCreateTransaction.IsValid())
+	TUniquePtr<FScopedTransaction> AddTransaction = MakeUnique<FScopedTransaction>( LOCTEXT("AddComponent", "Add Component") );
 
 	UActorComponent* NewComponent = nullptr;
 	UActorComponent* ComponentTemplate = Cast<UActorComponent>(Asset);
@@ -5156,7 +5156,7 @@ UActorComponent* SSCSEditor::AddNewComponent( UClass* NewComponentClass, UObject
 		{
 			NewVariableName = *FComponentEditorUtils::GenerateValidVariableNameFromAsset(Asset, nullptr);
 		}
-		NewComponent = AddNewNode(Blueprint->SimpleConstructionScript->CreateNode(NewComponentClass, NewVariableName), Asset, bMarkBlueprintModified, bSetFocusToNewItem);
+		NewComponent = AddNewNode(MoveTemp(AddTransaction), Blueprint->SimpleConstructionScript->CreateNode(NewComponentClass, NewVariableName), Asset, bMarkBlueprintModified, bSetFocusToNewItem);
 
 		if (ComponentTemplate)
 		{
@@ -5179,7 +5179,7 @@ UActorComponent* SSCSEditor::AddNewComponent( UClass* NewComponentClass, UObject
 		if (ComponentTemplate)
 		{
 			// Create a duplicate of the provided template
-			NewComponent = AddNewNodeForInstancedComponent(FComponentEditorUtils::DuplicateComponent(ComponentTemplate), nullptr, bSetFocusToNewItem);
+			NewComponent = AddNewNodeForInstancedComponent(MoveTemp(AddTransaction), FComponentEditorUtils::DuplicateComponent(ComponentTemplate), nullptr, bSetFocusToNewItem);
 		}
 		else if (AActor* ActorInstance = GetActorContext())
 		{
@@ -5241,14 +5241,14 @@ UActorComponent* SSCSEditor::AddNewComponent( UClass* NewComponentClass, UObject
 			// Rerun construction scripts
 			ActorInstance->RerunConstructionScripts();
 
-			NewComponent = AddNewNodeForInstancedComponent(NewInstanceComponent, Asset, bSetFocusToNewItem);
+			NewComponent = AddNewNodeForInstancedComponent(MoveTemp(AddTransaction), NewInstanceComponent, Asset, bSetFocusToNewItem);
 		}
 	}
 
 	return NewComponent;
 }
 
-UActorComponent* SSCSEditor::AddNewNode(USCS_Node* NewNode, UObject* Asset, bool bMarkBlueprintModified, bool bSetFocusToNewItem)
+UActorComponent* SSCSEditor::AddNewNode(TUniquePtr<FScopedTransaction> InOngoingCreateTransaction, USCS_Node* NewNode, UObject* Asset, bool bMarkBlueprintModified, bool bSetFocusToNewItem)
 {
 	check(NewNode != nullptr);
 
@@ -5299,7 +5299,7 @@ UActorComponent* SSCSEditor::AddNewNode(USCS_Node* NewNode, UObject* Asset, bool
 	{
 		// Select and request a rename on the new component
 		SCSTreeWidget->SetSelection(NewNodePtr);
-		OnRenameComponent(false);
+		OnRenameComponent(MoveTemp(InOngoingCreateTransaction));
 	}
 
 	// Will call UpdateTree as part of OnBlueprintChanged handling
@@ -5315,7 +5315,7 @@ UActorComponent* SSCSEditor::AddNewNode(USCS_Node* NewNode, UObject* Asset, bool
 	return NewNode->ComponentTemplate;
 }
 
-UActorComponent* SSCSEditor::AddNewNodeForInstancedComponent(UActorComponent* NewInstanceComponent, UObject* Asset, bool bSetFocusToNewItem)
+UActorComponent* SSCSEditor::AddNewNodeForInstancedComponent(TUniquePtr<FScopedTransaction> InOngoingCreateTransaction, UActorComponent* NewInstanceComponent, UObject* Asset, bool bSetFocusToNewItem)
 {
 	check(NewInstanceComponent != nullptr);
 
@@ -5351,7 +5351,7 @@ UActorComponent* SSCSEditor::AddNewNodeForInstancedComponent(UActorComponent* Ne
 	{
 		// Select and request a rename on the new component
 		SCSTreeWidget->SetSelection(NewNodePtr);
-		OnRenameComponent(false);
+		OnRenameComponent(MoveTemp(InOngoingCreateTransaction));
 	}
 
 	UpdateTree(false);
@@ -5534,7 +5534,7 @@ void SSCSEditor::PasteNodes()
 			check(NewActorComponent);
 
 			// Create a new SCS node to contain the new component and add it to the tree
-			NewActorComponent = AddNewNode(Blueprint->SimpleConstructionScript->CreateNodeAndRenameComponent(NewActorComponent), nullptr, false, false);
+			NewActorComponent = AddNewNode(TUniquePtr<FScopedTransaction>(), Blueprint->SimpleConstructionScript->CreateNodeAndRenameComponent(NewActorComponent), nullptr, false, false);
 
 			if (NewActorComponent)
 			{
@@ -6175,7 +6175,7 @@ void SSCSEditor::OnItemScrolledIntoView( FSCSEditorTreeNodePtrType InItem, const
 		if(DeferredRenameRequest == ItemName)
 		{
 			DeferredRenameRequest = NAME_None;
-			InItem->OnRequestRename(bIsDeferredRenameRequestTransactional);
+			InItem->OnRequestRename(MoveTemp(OngoingCreateTransaction)); // Transfer responsibility to end the 'create + give initial name' transaction to the tree item if such transaction is ongoing.
 		}
 	}
 }
@@ -6186,7 +6186,12 @@ void SSCSEditor::HandleItemDoubleClicked(FSCSEditorTreeNodePtrType InItem)
 	OnItemDoubleClicked.ExecuteIfBound(InItem);
 }
 
-void SSCSEditor::OnRenameComponent(bool bTransactional)
+void SSCSEditor::OnRenameComponent()
+{
+	OnRenameComponent(nullptr); // null means that the rename is not part of the creation process (create + give initial name).
+}
+
+void SSCSEditor::OnRenameComponent(TUniquePtr<FScopedTransaction> InComponentCreateTransaction)
 {
 	TArray< FSCSEditorTreeNodePtrType > SelectedItems = SCSTreeWidget->GetSelectedItems();
 
@@ -6194,7 +6199,9 @@ void SSCSEditor::OnRenameComponent(bool bTransactional)
 	check(SelectedItems.Num() == 1);
 
 	DeferredRenameRequest = SelectedItems[0]->GetNodeID();
-	bIsDeferredRenameRequestTransactional = bTransactional;
+
+	check(!OngoingCreateTransaction.IsValid()); // If this fails, something in the chain of responsibility failed to end the previous transaction.
+	OngoingCreateTransaction = MoveTemp(InComponentCreateTransaction); // If a 'create + give initial name' transaction is ongoing, take responsibility of ending it until it's transfered to the selected item.
 
 	SCSTreeWidget->RequestScrollIntoView(SelectedItems[0]);
 }
diff --git a/Engine/Source/Editor/Kismet/Public/SSCSEditor.h b/Engine/Source/Editor/Kismet/Public/SSCSEditor.h
--- a/Engine/Source/Editor/Kismet/Public/SSCSEditor.h
+++ b/Engine/Source/Editor/Kismet/Public/SSCSEditor.h
@@ -27,6 +27,7 @@ class FMenuBuilder;
 class FSCSEditorTreeNode;
 class SSCSEditor;
 class UPrimitiveComponent;
+class FScopedTransaction;
 struct EventData;
 
 // SCS tree node pointer type
@@ -279,8 +280,11 @@ public:
 	 */
 	virtual bool CanRename() const { return false; }
 
-	/** Requests a rename on the component */
-	void OnRequestRename(bool bTransactional);
+	/**
+	 * Requests a rename on the component.
+	 * @param OngoingCreateTransaction The transaction scoping the node creation which will end once the node is named by the user or null if the rename is not part of a the creation process.
+	 */
+	void OnRequestRename(TUniquePtr<FScopedTransaction> OngoingCreateTransaction);
 
 	/** Renames the component */
 	virtual void OnCompleteRename(const FText& InNewName);
@@ -304,12 +308,8 @@ protected:
 	/** Used to update the EFilteredState::ChildMatches flag for parent nodes, when this item's filtration state has changed */
 	void ApplyFilteredStateToParent();
 	
-	bool GetAndClearNonTransactionalRenameFlag()
-	{
-		const bool bResult = bNonTransactionalRename;
-		bNonTransactionalRename = false;
-		return bResult;
-	}
+	// Scope the creation of a component which ends when the initial component 'name' is given/accepted by the user, which can be several frames after the component was actually created.
+	TUniquePtr<FScopedTransaction> OngoingCreateTransaction;
 
 	// Component template represented by this node, if it's a component node, otherwise invalid
 	TWeakObjectPtr<UActorComponent> ComponentTemplatePtr;
@@ -323,7 +323,6 @@ private:
 	TArray<FSCSEditorTreeNodePtrType> Children;
 
 	/** Handles rename requests */
-	bool bNonTransactionalRename;
 	FOnRenameRequested RenameRequestedDelegate;
 
 	enum EFilteredState
@@ -812,19 +811,21 @@ public:
 	UActorComponent* AddNewComponent(UClass* NewComponentClass, UObject* Asset, const bool bSkipMarkBlueprintModified = false, bool bSetFocusToNewItem = true);
 
 	/** Adds a new SCS Node to the component Table
+	   @param OngoingCreateTransaction (In) The transaction containing the creation of the node. The transaction will remain ongoing until the node gets its initial name from user.
 	   @param NewNode	(In) The SCS node to add
 	   @param Asset		(In) Optional asset to assign to the component
 	   @param bMarkBlueprintModified (In) Whether or not to mark the Blueprint as structurally modified
 	   @param bSetFocusToNewItem (In) Select the new item and activate the inline rename widget (default is true)
 	   @return The reference of the newly created ActorComponent */
-	UActorComponent* AddNewNode(USCS_Node* NewNode, UObject* Asset, bool bMarkBlueprintModified, bool bSetFocusToNewItem = true);
+	UActorComponent* AddNewNode(TUniquePtr<FScopedTransaction> OngoingCreateTransaction, USCS_Node* NewNode, UObject* Asset, bool bMarkBlueprintModified, bool bSetFocusToNewItem = true);
 
 	/** Adds a new component instance node to the component Table
+		@param OngoingCreateTransaction (In) The transaction containing the creation of the node. The transaction will remain ongoing until the node gets its initial name from user.
 		@param NewInstanceComponent	(In) The component being added to the actor instance
 		@param Asset (In) Optional asset to assign to the component
 		@param bSetFocusToNewItem (In) Select the new item and activate the inline rename widget (default is true)
 		@return The reference of the newly created ActorComponent */
-	UActorComponent* AddNewNodeForInstancedComponent(UActorComponent* NewInstanceComponent, UObject* Asset, bool bSetFocusToNewItem = true);
+	UActorComponent* AddNewNodeForInstancedComponent(TUniquePtr<FScopedTransaction> OngoingCreateTransaction, UActorComponent* NewInstanceComponent, UObject* Asset, bool bSetFocusToNewItem = true);
 	
 	/** Returns true if the specified component is currently selected */
 	bool IsComponentSelected(const UPrimitiveComponent* PrimComponent) const;
@@ -989,11 +990,18 @@ protected:
 
 	/** Checks to see if renaming is allowed on the selected component */
 	bool CanRenameComponent() const;
+
+	/**
+	 * Requests a rename on the selected component just after creation so that the user can provide the initial
+	 * component name (overwriting the default generated one), which is considered part of the creation process.
+	 * @param OngoingCreateTransaction The ongoing transaction started when the component was created.
+	 */
+	void OnRenameComponent(TUniquePtr<FScopedTransaction> OngoingCreateTransaction);
+
 	/**
-	 * Requests a rename on the selected component
-	 * @param bTransactional Whether or not the rename should be transactional (i.e. undoable)
+	 * Requests a rename on the selected component.
 	 */
-	void OnRenameComponent(bool bTransactional);
+	void OnRenameComponent();
 
 	/** Called when component objects are replaced following construction script execution */
 	void OnObjectsReplaced(const TMap<UObject*, UObject*>& OldToNewInstanceMap);
@@ -1134,8 +1142,8 @@ public:
 	/** Name of a node that has been requested to be renamed */
 	FName DeferredRenameRequest;
 
-	/** Whether or not the deferred rename request was flagged as transactional */
-	bool bIsDeferredRenameRequestTransactional;
+	/** Scope the creation of a component which ends when the initial component 'name' is given/accepted by the user, which can be several frames after the component was actually created. */
+	TUniquePtr<FScopedTransaction> OngoingCreateTransaction;
 
 	/** Attribute that provides access to the Actor context for which we are viewing/editing the SCS. */
 	TAttribute<class AActor*> ActorContext;
