Fixed compiler error introduce when fixing UE-70329. Reported because FScopedTransaction was an incomplete type. Include the file rather than forward declaring the type.

diff --git a/Engine/Source/Editor/Kismet/Public/SSCSEditor.h b/Engine/Source/Editor/Kismet/Public/SSCSEditor.h
--- a/Engine/Source/Editor/Kismet/Public/SSCSEditor.h
+++ b/Engine/Source/Editor/Kismet/Public/SSCSEditor.h
@@ -22,12 +22,12 @@
 #include "BlueprintEditor.h"
 #include "Widgets/SToolTip.h"
 #include "SComponentClassCombo.h"
+#include "ScopedTransaction.h"
 
 class FMenuBuilder;
 class FSCSEditorTreeNode;
 class SSCSEditor;
 class UPrimitiveComponent;
-class FScopedTransaction;
 struct EventData;
 
 // SCS tree node pointer type
