Fixed focused widget when opening the Override Function menu in the Blueprint Editor - this allows the Override Function menu to be searched

diff --git a/Engine/Source/Editor/Kismet/Private/SMyBlueprint.cpp b/Engine/Source/Editor/Kismet/Private/SMyBlueprint.cpp
--- a/Engine/Source/Editor/Kismet/Private/SMyBlueprint.cpp
+++ b/Engine/Source/Editor/Kismet/Private/SMyBlueprint.cpp
@@ -827,6 +827,12 @@ TSharedRef<SWidget> SMyBlueprint::OnGetFunctionListMenu()
 	BuildOverridableFunctionsMenu(MenuBuilder);
 
 	TSharedRef<SWidget> MenuWidget = MenuBuilder.MakeWidget();
+	
+	// force user focus onto the menu widget:
+	if(FunctionSectionButton.IsValid())
+	{
+		FunctionSectionButton->SetMenuContentWidgetToFocus(MenuWidget);
+	}
 
 	return MenuWidget;
 }
