Adding additional descriptive info to the debug world names shown in debugger dropdown

diff --git a/Engine/Source/Editor/Kismet/Private/SBlueprintEditorSelectedDebugObjectWidget.cpp b/Engine/Source/Editor/Kismet/Private/SBlueprintEditorSelectedDebugObjectWidget.cpp
--- a/Engine/Source/Editor/Kismet/Private/SBlueprintEditorSelectedDebugObjectWidget.cpp
+++ b/Engine/Source/Editor/Kismet/Private/SBlueprintEditorSelectedDebugObjectWidget.cpp
@@ -276,6 +276,14 @@ void SBlueprintEditorSelectedDebugObjectWidget::GenerateDebugWorldNames(bool bRe
 
 		if (!WorldName.IsEmpty())
 		{
+			if (FWorldContext* PieContext = GEngine->GetWorldContextFromWorld(TestWorld))
+			{
+				if (!PieContext->CustomDescription.IsEmpty())
+				{
+					WorldName += TEXT(" ") + PieContext->CustomDescription;
+				}
+			}
+
 			// DebugWorlds & DebugWorldNames need to be the same size (we expect
 			// an index in one to correspond to the other) - DebugWorldNames is
 			// what populates the dropdown, so it is the authority (if there's 
diff --git a/Engine/Source/Runtime/Engine/Classes/Engine/Engine.h b/Engine/Source/Runtime/Engine/Classes/Engine/Engine.h
--- a/Engine/Source/Runtime/Engine/Classes/Engine/Engine.h
+++ b/Engine/Source/Runtime/Engine/Classes/Engine/Engine.h
@@ -376,6 +376,9 @@ struct FWorldContext
 	/** Handle to this world context's audio device.*/
 	uint32 AudioDeviceHandle;
 
+	/** Custom description to be display in blueprint debugger UI */
+	FString CustomDescription;
+
 	/**************************************************************/
 
 	/** Outside pointers to CurrentWorld that should be kept in sync if current world changes  */
diff --git a/Engine/Source/Runtime/Engine/Private/UnrealEngine.cpp b/Engine/Source/Runtime/Engine/Private/UnrealEngine.cpp
--- a/Engine/Source/Runtime/Engine/Private/UnrealEngine.cpp
+++ b/Engine/Source/Runtime/Engine/Private/UnrealEngine.cpp
@@ -10108,9 +10108,7 @@ void DrawStatsHUD( UWorld* World, FViewport* Viewport, FCanvas* Canvas, UCanvas*
 				Y += 20;
 				Y += Canvas->DrawShadowedString( X, Y, TEXT("[Server]"), GEngine->GetSmallFont(), FLinearColor::Gray);
 				DrawDebugPropertiesForWorld(ServerWorld, X, Y);
-				break;
 			}
-
 		}
 	}
 
