Fix merge error when enum moved file

#rnx

diff --git a/Engine/Source/Editor/Kismet/Public/BlueprintEditor.h b/Engine/Source/Editor/Kismet/Public/BlueprintEditor.h
--- a/Engine/Source/Editor/Kismet/Public/BlueprintEditor.h
+++ b/Engine/Source/Editor/Kismet/Public/BlueprintEditor.h
@@ -52,6 +52,8 @@ namespace NodeSectionID
 	{
 		NONE = 0,
 		GRAPH,					// Graph
+		ANIMGRAPH,				// Anim Graph
+		ANIMLAYER,				// Anim Layer
 		FUNCTION,				// Functions
 		FUNCTION_OVERRIDABLE,	// Overridable functions
 		INTERFACE,				// Interface
