Integrate changes from //UE4//dev-framework-turtles [REDACTED_HANDLE]
- Refact of SSCSEditor which moves some members from the SSCSEditor class to its Actor Node class (less code)
- Handling of Separator nodes is now done inside the Actor Node class
- Now possible to create Component on Selected Component (will be parented if Selected Component is SceneComponent)
- JIRA UE-65144: Find Component by Name and Hierarchy instead of only by Name

diff --git a/Engine/Plugins/MovieScene/ActorSequence/Source/ActorSequenceEditor/Private/ActorSequenceEditorTabSummoner.cpp b/Engine/Plugins/MovieScene/ActorSequence/Source/ActorSequenceEditor/Private/ActorSequenceEditorTabSummoner.cpp
--- a/Engine/Plugins/MovieScene/ActorSequence/Source/ActorSequenceEditor/Private/ActorSequenceEditorTabSummoner.cpp
+++ b/Engine/Plugins/MovieScene/ActorSequence/Source/ActorSequenceEditor/Private/ActorSequenceEditorTabSummoner.cpp
@@ -386,7 +386,7 @@ public:
 			UBlueprint* Blueprint = BlueprintEditor->GetBlueprintObj();
 			if (Blueprint)
 			{
-				EditingComponent = SelectedNode->GetEditableComponentTemplate(Blueprint);
+				EditingComponent = SelectedNode->GetOrCreateEditableComponentTemplate(Blueprint);
 			}
 		}
 		else if (AActor* Actor = GetPreviewActor())
diff --git a/Engine/Source/Editor/Kismet/Private/BlueprintDetailsCustomization.cpp b/Engine/Source/Editor/Kismet/Private/BlueprintDetailsCustomization.cpp
--- a/Engine/Source/Editor/Kismet/Private/BlueprintDetailsCustomization.cpp
+++ b/Engine/Source/Editor/Kismet/Private/BlueprintDetailsCustomization.cpp
@@ -5887,7 +5887,7 @@ void FBlueprintComponentDetails::OnBrowseSocket()
 
 		if (ParentFNode.IsValid())
 		{
-			if (USceneComponent* ParentSceneComponent = Cast<USceneComponent>(ParentFNode->GetEditableComponentTemplate(Editor->GetBlueprint())))
+			if (USceneComponent* ParentSceneComponent = Cast<USceneComponent>(ParentFNode->GetOrCreateEditableComponentTemplate(Editor->GetBlueprint())))
 			{
 				if (ParentSceneComponent->HasAnySockets())
 				{
diff --git a/Engine/Source/Editor/Kismet/Private/BlueprintEditor.cpp b/Engine/Source/Editor/Kismet/Private/BlueprintEditor.cpp
--- a/Engine/Source/Editor/Kismet/Private/BlueprintEditor.cpp
+++ b/Engine/Source/Editor/Kismet/Private/BlueprintEditor.cpp
@@ -936,7 +936,7 @@ void FBlueprintEditor::OnSelectionUpdated(const TArray<FSCSEditorTreeNodePtrType
 				}
 				else
 				{
-					UActorComponent* EditableComponent = NodePtr->GetEditableComponentTemplate(GetBlueprintObj());
+					UActorComponent* EditableComponent = NodePtr->GetOrCreateEditableComponentTemplate(GetBlueprintObj());
 					if (EditableComponent)
 					{
 						InspectorTitle = FText::FromString(NodePtr->GetDisplayString());
diff --git a/Engine/Source/Editor/Kismet/Private/SCSDiff.cpp b/Engine/Source/Editor/Kismet/Private/SCSDiff.cpp
--- a/Engine/Source/Editor/Kismet/Private/SCSDiff.cpp
+++ b/Engine/Source/Editor/Kismet/Private/SCSDiff.cpp
@@ -73,9 +73,9 @@ TArray< FSCSResolvedIdentifier > FSCSDiff::GetDisplayedHierarchy() const
 {
 	TArray< FSCSResolvedIdentifier > Ret;
 
-	if( SCSEditor.IsValid() )
+	if( SCSEditor.IsValid() && SCSEditor->GetActorNode().IsValid())
 	{
-		const TArray<FSCSEditorTreeNodePtrType>& RootNodes = SCSEditor->GetRootComponentNodes();
+		const TArray<FSCSEditorTreeNodePtrType>& RootNodes = SCSEditor->GetActorNode()->GetComponentNodes();
 		for (int32 Iter = 0; Iter != RootNodes.Num(); ++Iter)
 		{
 			TArray< int32 > TreeAddress;
diff --git a/Engine/Source/Editor/Kismet/Private/SCSEditorViewportClient.cpp b/Engine/Source/Editor/Kismet/Private/SCSEditorViewportClient.cpp
--- a/Engine/Source/Editor/Kismet/Private/SCSEditorViewportClient.cpp
+++ b/Engine/Source/Editor/Kismet/Private/SCSEditorViewportClient.cpp
@@ -447,7 +447,7 @@ bool FSCSEditorViewportClient::InputWidgetDelta( FViewport* InViewport, EAxisLis
 						}
 						
 						USceneComponent* SceneComp = Cast<USceneComponent>(SelectedNodePtr->FindComponentInstanceInActor(PreviewActor));
-						USceneComponent* SelectedTemplate = Cast<USceneComponent>(SelectedNodePtr->GetEditableComponentTemplate(BlueprintEditor->GetBlueprintObj()));
+						USceneComponent* SelectedTemplate = Cast<USceneComponent>(SelectedNodePtr->GetOrCreateEditableComponentTemplate(BlueprintEditor->GetBlueprintObj()));
 						if(SceneComp != NULL && SelectedTemplate != NULL)
 						{
 							// Cache the current default values for propagation
@@ -588,31 +588,34 @@ FWidget::EWidgetMode FSCSEditorViewportClient::GetWidgetMode() const
 		if ( BluePrintEditor.IsValid() )
 		{
 			TArray<FSCSEditorTreeNodePtrType> SelectedNodes = BluePrintEditor->GetSelectedSCSEditorTreeNodes();
-			const TArray<FSCSEditorTreeNodePtrType>& RootNodes = BluePrintEditor->GetSCSEditor()->GetRootComponentNodes();
-
-			if (GUnrealEd->ComponentVisManager.IsActive() &&
-				GUnrealEd->ComponentVisManager.IsVisualizingArchetype())
-			{
-				// Component visualizer is active and editing the archetype
-				ReturnWidgetMode = WidgetMode;
-			}
-			else
+			if (BluePrintEditor->GetSCSEditor()->GetActorNode().IsValid())
 			{
-				// if the selected nodes array is empty, or only contains entries from the
-				// root nodes array, or isn't visible in the preview actor, then don't display a transform widget
-				for (int32 CurrentNodeIndex = 0; CurrentNodeIndex < SelectedNodes.Num(); CurrentNodeIndex++)
+				const TArray<FSCSEditorTreeNodePtrType>& RootNodes = BluePrintEditor->GetSCSEditor()->GetActorNode()->GetComponentNodes();
+
+				if (GUnrealEd->ComponentVisManager.IsActive() &&
+					GUnrealEd->ComponentVisManager.IsVisualizingArchetype())
 				{
-					FSCSEditorTreeNodePtrType CurrentNodePtr = SelectedNodes[CurrentNodeIndex];
-					if ((CurrentNodePtr.IsValid() &&
-						 ((!RootNodes.Contains(CurrentNodePtr) && !CurrentNodePtr->IsRootComponent()) ||
-							(Cast<UInstancedStaticMeshComponent>(CurrentNodePtr->GetComponentTemplate()) && // show widget if we are editing individual instances even if it is the root component
-							 CastChecked<UInstancedStaticMeshComponent>(CurrentNodePtr->FindComponentInstanceInActor(GetPreviewActor()))->SelectedInstances.Contains(true))) &&
-						 CurrentNodePtr->CanEditDefaults() &&
-						 CurrentNodePtr->FindComponentInstanceInActor(PreviewActor)))
+					// Component visualizer is active and editing the archetype
+					ReturnWidgetMode = WidgetMode;
+				}
+				else
+				{
+					// if the selected nodes array is empty, or only contains entries from the
+					// root nodes array, or isn't visible in the preview actor, then don't display a transform widget
+					for (int32 CurrentNodeIndex = 0; CurrentNodeIndex < SelectedNodes.Num(); CurrentNodeIndex++)
 					{
-						// a non-NULL, non-root item is selected, draw the widget
-						ReturnWidgetMode = WidgetMode;
-						break;
+						FSCSEditorTreeNodePtrType CurrentNodePtr = SelectedNodes[CurrentNodeIndex];
+						if ((CurrentNodePtr.IsValid() &&
+							((!RootNodes.Contains(CurrentNodePtr) && !CurrentNodePtr->IsRootComponent()) ||
+							(Cast<UInstancedStaticMeshComponent>(CurrentNodePtr->GetComponentTemplate()) && // show widget if we are editing individual instances even if it is the root component
+								CastChecked<UInstancedStaticMeshComponent>(CurrentNodePtr->FindComponentInstanceInActor(GetPreviewActor()))->SelectedInstances.Contains(true))) &&
+							CurrentNodePtr->CanEditDefaults() &&
+							CurrentNodePtr->FindComponentInstanceInActor(PreviewActor)))
+						{
+							// a non-NULL, non-root item is selected, draw the widget
+							ReturnWidgetMode = WidgetMode;
+							break;
+						}
 					}
 				}
 			}
@@ -927,7 +930,7 @@ void FSCSEditorViewportClient::BeginTransaction(const FText& Description)
 					}
 
 					// Modify template, any instances will be reconstructed as part of PostUndo:
-					UActorComponent* ComponentTemplate = Node->GetEditableComponentTemplate(PreviewBlueprint);
+					UActorComponent* ComponentTemplate = Node->GetOrCreateEditableComponentTemplate(PreviewBlueprint);
 					if (ComponentTemplate != nullptr)
 					{
 						ComponentTemplate->SetFlags(RF_Transactional);
diff --git a/Engine/Source/Editor/Kismet/Private/SSCSEditor.cpp b/Engine/Source/Editor/Kismet/Private/SSCSEditor.cpp
--- a/Engine/Source/Editor/Kismet/Private/SSCSEditor.cpp
+++ b/Engine/Source/Editor/Kismet/Private/SSCSEditor.cpp
@@ -71,6 +71,8 @@
 #include "Algo/Find.h"
 #include "ActorEditorUtils.h"
 
+#pragma optimize("", off)
+
 #define LOCTEXT_NAMESPACE "SSCSEditor"
 
 DEFINE_LOG_CATEGORY_STATIC(LogSCSEditor, Log, All);
@@ -433,7 +435,7 @@ class USCS_Node* FSCSEditorTreeNode::GetSCSNode() const
 	return nullptr;
 }
 
-UActorComponent* FSCSEditorTreeNode::GetEditableComponentTemplate(UBlueprint* ActualEditedBlueprint)
+UActorComponent* FSCSEditorTreeNode::GetOrCreateEditableComponentTemplate(UBlueprint* ActualEditedBlueprint)
 {
 	return nullptr;
 }
@@ -529,7 +531,7 @@ void FSCSEditorTreeNode::UpdateCachedFilterState(bool bMatchesFilter, bool bUpda
 
 void FSCSEditorTreeNode::RefreshCachedChildFilterState(bool bUpdateParent)
 {
-	const bool bCointainedMatch = !IsFlaggedForFiltration();
+	const bool bContainedMatch = !IsFlaggedForFiltration();
 
 	FilterFlags &= ~EFilteredState::ChildMatches;
 	for (FSCSEditorTreeNodePtrType Child : Children)
@@ -540,9 +542,9 @@ void FSCSEditorTreeNode::RefreshCachedChildFilterState(bool bUpdateParent)
 			break;
 		}
 	}
-	const bool bCointainsMatch = !IsFlaggedForFiltration();
+	const bool bContainsMatch = !IsFlaggedForFiltration();
 
-	const bool bStateChange = bCointainedMatch != bCointainsMatch;
+	const bool bStateChange = bContainedMatch != bContainsMatch;
 	if (bUpdateParent && bStateChange)
 	{
 		ApplyFilteredStateToParent();
@@ -789,7 +791,7 @@ FSCSEditorTreeNodePtrType FSCSEditorTreeNode::FactoryNodeFromComponent(UActorCom
 		}
 		else
 		{
-			return MakeShareable(new FSCSEditorTreeNodeInstancedInheritedComponent(Owner, InComponent->GetFName()));
+			return MakeShareable(new FSCSEditorTreeNodeInstancedInheritedComponent(Owner, InComponent));
 		}
 	}
 
@@ -866,7 +868,7 @@ FSCSEditorTreeNodePtrType FSCSEditorTreeNode::FindChild(const FName& InVariableO
 		for(int32 ChildIndex = 0; ChildIndex < Children.Num() && !Result.IsValid(); ++ChildIndex)
 		{
 			FName ItemName = Children[ChildIndex]->GetVariableName();
-			if(ItemName == NAME_None)
+			if(ItemName == NAME_None && Children[ChildIndex]->GetNodeType() == ComponentNode)
 			{
 				UActorComponent* ComponentTemplateOrInstance = Children[ChildIndex]->GetComponentTemplate();
 				check(ComponentTemplateOrInstance != nullptr);
@@ -987,22 +989,13 @@ FString FSCSEditorTreeNodeComponentBase::GetDisplayString() const
 //////////////////////////////////////////////////////////////////////////
 // FSCSEditorTreeNodeInstancedInheritedComponent
 
-FSCSEditorTreeNodeInstancedInheritedComponent::FSCSEditorTreeNodeInstancedInheritedComponent(AActor* Owner, FName InComponentName)
+FSCSEditorTreeNodeInstancedInheritedComponent::FSCSEditorTreeNodeInstancedInheritedComponent(AActor* Owner, UActorComponent* ComponentInstance)
 {
-	InstancedComponentName = InComponentName;
-	check(InstancedComponentName != NAME_None);	// ...otherwise IsRootActor() can return a false positive.
+	check(ComponentInstance != nullptr);
 
 	InstancedComponentOwnerPtr = Owner;
 
-	SetComponentTemplate(nullptr);
-	for (UActorComponent* ComponentInstance : Owner->GetComponents())
-	{
-		if (ComponentInstance && ComponentInstance->GetFName() == InstancedComponentName)
-		{
-			SetComponentTemplate(ComponentInstance);
-			break;
-		}
-	}
+	SetComponentTemplate(ComponentInstance);
 }
 
 bool FSCSEditorTreeNodeInstancedInheritedComponent::IsNative() const
@@ -1059,7 +1052,7 @@ FText FSCSEditorTreeNodeInstancedInheritedComponent::GetDisplayName() const
 	return FText::GetEmpty();
 }
 
-UActorComponent* FSCSEditorTreeNodeInstancedInheritedComponent::GetEditableComponentTemplate(UBlueprint* ActualEditedBlueprint)
+UActorComponent* FSCSEditorTreeNodeInstancedInheritedComponent::GetOrCreateEditableComponentTemplate(UBlueprint* ActualEditedBlueprint) 
 {
 	if (CanEditDefaults())
 	{
@@ -1118,7 +1111,7 @@ FText FSCSEditorTreeNodeInstanceAddedComponent::GetDisplayName() const
 	return FText::FromName(InstancedComponentName);
 }
 
-UActorComponent* FSCSEditorTreeNodeInstanceAddedComponent::GetEditableComponentTemplate(UBlueprint* ActualEditedBlueprint)
+UActorComponent* FSCSEditorTreeNodeInstanceAddedComponent::GetOrCreateEditableComponentTemplate(UBlueprint* ActualEditedBlueprint)
 {
 	return GetComponentTemplate();
 }
@@ -1276,7 +1269,7 @@ class USCS_Node* FSCSEditorTreeNodeComponent::GetSCSNode() const
 	return SCSNodePtr.Get();
 }
 
-UActorComponent* FSCSEditorTreeNodeComponent::GetEditableComponentTemplate(UBlueprint* ActualEditedBlueprint)
+UActorComponent* FSCSEditorTreeNodeComponent::GetOrCreateEditableComponentTemplate(UBlueprint* ActualEditedBlueprint)
 {
 	if (CanEditDefaults())
 	{
@@ -1327,15 +1320,9 @@ UActorComponent* FSCSEditorTreeNode::FindComponentInstanceInActor(const AActor*
 		}
 		else if (ComponentTemplate != NULL)
 		{
-			// Look for a native component instance with a name that matches the template name
-			for (UActorComponent* Component : InActor->GetComponents())
-			{
-				if (Component && Component->GetFName() == ComponentTemplate->GetFName())
-				{
-					ComponentInstance = Component;
-					break;
-				}
-			}
+			TInlineComponentArray<UActorComponent*> Components;
+			InActor->GetComponents(Components);
+			ComponentInstance = FComponentEditorUtils::FindMatchingComponent(ComponentTemplate, Components);
 		}
 	}
 
@@ -1399,6 +1386,84 @@ UActorComponent* FSCSEditorTreeNodeComponent::INTERNAL_GetOverridenComponentTemp
 
 //////////////////////////////////////////////////////////////////////////
 // FSCSEditorTreeNodeRootActor
+FSCSEditorTreeNodePtrType FSCSEditorTreeNodeRootActor::GetSceneRootNode() const
+{
+	return SceneRootNodePtr;
+}
+
+void FSCSEditorTreeNodeRootActor::SetSceneRootNode(FSCSEditorTreeNodePtrType NewSceneRootNode)
+{
+	if (SceneRootNodePtr.IsValid())
+	{
+		ComponentNodes.Remove(SceneRootNodePtr);
+	}
+
+	SceneRootNodePtr = NewSceneRootNode;
+
+	if (!ComponentNodes.Contains(SceneRootNodePtr))
+	{
+		ComponentNodes.Add(SceneRootNodePtr);
+	}
+}
+
+const TArray<FSCSEditorTreeNodePtrType>& FSCSEditorTreeNodeRootActor::GetComponentNodes() const
+{
+	return ComponentNodes;
+}
+
+void FSCSEditorTreeNodeRootActor::AddChild(FSCSEditorTreeNodePtrType InChildNodePtr)
+{
+	if (InChildNodePtr->GetNodeType() == FSCSEditorTreeNode::ComponentNode)
+	{
+		ComponentNodes.Add(InChildNodePtr);
+		USceneComponent* SceneComponent = Cast<USceneComponent>(InChildNodePtr->GetComponentTemplate());
+		if (!SceneRootNodePtr.IsValid() && SceneComponent != nullptr)
+		{
+			SetSceneRootNode(InChildNodePtr);
+		}
+
+		// Make sure separators are shown
+		if (SceneComponent != nullptr && !SceneComponentSeparatorNodePtr.IsValid())
+		{
+			SceneComponentSeparatorNodePtr = MakeShareable(new FSCSEditorTreeNodeSeparator());
+			Super::AddChild(SceneComponentSeparatorNodePtr);
+		}
+		else if (SceneComponent == nullptr && !NonSceneComponentSeparatorNodePtr.IsValid())
+		{
+			NonSceneComponentSeparatorNodePtr = MakeShareable(new FSCSEditorTreeNodeSeparator());
+			Super::AddChild(NonSceneComponentSeparatorNodePtr);
+		}
+	}
+	
+	Super::AddChild(InChildNodePtr);
+}
+
+void FSCSEditorTreeNodeRootActor::RemoveChild(FSCSEditorTreeNodePtrType InChildNodePtr)
+{
+	Super::RemoveChild(InChildNodePtr);
+
+	int32 indexOfFirstSceneComponent = ComponentNodes.IndexOfByPredicate([](const FSCSEditorTreeNodePtrType& NodePtr)
+	{
+		return NodePtr->GetNodeType() == FSCSEditorTreeNode::ComponentNode && Cast<USceneComponent>(NodePtr->GetComponentTemplate());
+	});
+
+	if (indexOfFirstSceneComponent == -1 && SceneComponentSeparatorNodePtr.IsValid())
+	{
+		Super::RemoveChild(SceneComponentSeparatorNodePtr);
+		SceneComponentSeparatorNodePtr = nullptr;
+	}
+	
+	int32 indexOffFirstNonSceneComponent = ComponentNodes.IndexOfByPredicate([](const FSCSEditorTreeNodePtrType& NodePtr)
+	{
+		return NodePtr->GetNodeType() == FSCSEditorTreeNode::ComponentNode && !Cast<USceneComponent>(NodePtr->GetComponentTemplate());
+	});
+
+	if (indexOffFirstNonSceneComponent == -1 && NonSceneComponentSeparatorNodePtr.IsValid())
+	{
+		Super::RemoveChild(NonSceneComponentSeparatorNodePtr);
+		NonSceneComponentSeparatorNodePtr = nullptr;
+	}
+}
 
 FName FSCSEditorTreeNodeRootActor::GetNodeID() const
 {
@@ -2089,7 +2154,7 @@ void SSCS_RowWidget::HandleOnDragEnter( const FDragDropEvent& DragDropEvent )
 
 		if (Message.IsEmpty())
 		{
-			FSCSEditorTreeNodePtrType SceneRootNodePtr = SCSEditor.Pin()->SceneRootNodePtr;
+			FSCSEditorTreeNodePtrType SceneRootNodePtr = SCSEditor.Pin()->GetSceneRootNode();
 			check(SceneRootNodePtr.IsValid());
 
 			FSCSEditorTreeNodePtrType NodePtr = GetNode();
@@ -2692,11 +2757,12 @@ void SSCS_RowWidget::OnDetachFromDropAction(const TArray<FSCSEditorTreeNodePtrTy
 			}
 
 			// Attach the dropped node to the current scene root node
-			check(SCSEditorPtr->SceneRootNodePtr.IsValid());
-			SCSEditorPtr->SceneRootNodePtr->AddChild(DroppedNodePtr);
+			FSCSEditorTreeNodePtrType SceneRootNodePtr = SCSEditorPtr->GetSceneRootNode();
+			check(SceneRootNodePtr.IsValid());
+			SceneRootNodePtr->AddChild(DroppedNodePtr);
 
 			// Attempt to locate a matching instance of the scene root component template in the Actor context that's being edited
-			USceneComponent* InstancedSceneRootComponent = Cast<USceneComponent>(SCSEditorPtr->SceneRootNodePtr->FindComponentInstanceInActor(PreviewActor));
+			USceneComponent* InstancedSceneRootComponent = Cast<USceneComponent>(SceneRootNodePtr->FindComponentInstanceInActor(PreviewActor));
 			if(SceneComponentTemplate && InstancedSceneRootComponent && InstancedSceneRootComponent->IsRegistered())
 			{
 				// If we find a match, calculate its new position relative to the scene root component instance in the preview scene
@@ -2751,8 +2817,9 @@ void SSCS_RowWidget::OnDetachFromDropAction(const TArray<FSCSEditorTreeNodePtrTy
 			NodePtr->RemoveChild(DroppedNodePtr);
 
 			// Attach the dropped node to the current scene root node
-			check(SCSEditorPtr->SceneRootNodePtr.IsValid());
-			SCSEditorPtr->SceneRootNodePtr->AddChild(DroppedNodePtr);
+			FSCSEditorTreeNodePtrType SceneRootNodePtr = SCSEditorPtr->GetSceneRootNode();
+			check(SceneRootNodePtr.IsValid());
+			SceneRootNodePtr->AddChild(DroppedNodePtr);
 		}
 	}
 	
@@ -2765,7 +2832,7 @@ void SSCS_RowWidget::OnMakeNewRootDropAction(FSCSEditorTreeNodePtrType DroppedNo
 	check(SCSEditorPtr.IsValid());
 
 	// Get the current scene root node
-	FSCSEditorTreeNodePtrType& SceneRootNodePtr = SCSEditorPtr->SceneRootNodePtr;
+	FSCSEditorTreeNodePtrType SceneRootNodePtr = SCSEditorPtr->GetSceneRootNode();
 
 	FSCSEditorTreeNodePtrType NodePtr = GetNode();
 
@@ -2876,6 +2943,7 @@ void SSCS_RowWidget::OnMakeNewRootDropAction(FSCSEditorTreeNodePtrType DroppedNo
 
 		// Set node we are dropping as new root
 		SceneRootNodePtr = DroppedNodePtr;
+		SCSEditorPtr->SetSceneRootNode(SceneRootNodePtr);
 
 		// Add dropped node to the SCS context
 		Blueprint->SimpleConstructionScript->AddNode(SceneRootNodePtr->GetSCSNode());
@@ -2910,6 +2978,7 @@ void SSCS_RowWidget::OnMakeNewRootDropAction(FSCSEditorTreeNodePtrType DroppedNo
 
 		// Set node we are dropping as new root
 		SceneRootNodePtr = DroppedNodePtr;
+		SCSEditorPtr->SetSceneRootNode(SceneRootNodePtr);
 
 		// Remove or re-parent the old root
 		if (OldSceneRootNodePtr.IsValid())
@@ -3020,7 +3089,7 @@ FString SSCS_RowWidget::GetDocumentationLink() const
 	check(SCSEditor.IsValid());
 
 	FSCSEditorTreeNodePtrType NodePtr = GetNode();
-	if ((NodePtr == SCSEditor.Pin()->SceneRootNodePtr) || NodePtr->IsInherited())
+	if ((NodePtr == SCSEditor.Pin()->GetSceneRootNode()) || NodePtr->IsInherited())
 	{
 		return TEXT("Shared/Editors/BlueprintEditor/ComponentsMode");
 	}
@@ -3033,7 +3102,7 @@ FString SSCS_RowWidget::GetDocumentationExcerptName() const
 	check(SCSEditor.IsValid());
 
 	FSCSEditorTreeNodePtrType NodePtr = GetNode();
-	if (NodePtr == SCSEditor.Pin()->SceneRootNodePtr)
+	if (NodePtr == SCSEditor.Pin()->GetSceneRootNode())
 	{
 		return TEXT("RootComponent");
 	}
@@ -3369,9 +3438,10 @@ FText SSCS_RowWidget_ActorRoot::GetActorMobilityText() const
 		{
 			USceneComponent* RootComponent = DefaultActor->GetRootComponent();
 
-			if ((RootComponent == nullptr) && (SCSEditorPtr->SceneRootNodePtr.IsValid()))
+			FSCSEditorTreeNodePtrType SceneRootNodePtr = SCSEditorPtr->GetSceneRootNode();
+			if ((RootComponent == nullptr) && SceneRootNodePtr.IsValid())
 			{
-				RootComponent = Cast<USceneComponent>(SCSEditorPtr->SceneRootNodePtr->GetComponentTemplate());
+				RootComponent = Cast<USceneComponent>(SceneRootNodePtr->GetComponentTemplate());
 			}
 
 			if (RootComponent != nullptr)
@@ -3428,7 +3498,6 @@ void SSCSEditor::Construct( const FArguments& InArgs )
 	OnItemDoubleClicked = InArgs._OnItemDoubleClicked;
 	OnHighlightPropertyInDetailsView = InArgs._OnHighlightPropertyInDetailsView;
 	bUpdatingSelection = false;
-	bHasAddedSceneAndBehaviorComponentSeparator = false;
 	bAllowTreeUpdates = true;
 	bIsDiffing = InArgs._IsDiffing;
 
@@ -3474,7 +3543,7 @@ void SSCSEditor::Construct( const FArguments& InArgs )
 	SCSTreeWidget = SNew(SSCSTreeType)
 		.ToolTipText(LOCTEXT("DropAssetToAddComponent", "Drop asset here to add a component."))
 		.SCSEditor(this)
-		.TreeItemsSource(&FilteredRootNodes)
+		.TreeItemsSource(&RootNodes)
 		.SelectionMode(ESelectionMode::Multi)
 		.OnGenerateRow(this, &SSCSEditor::MakeTableRowWidget)
 		.OnGetChildren(this, &SSCSEditor::OnGetChildrenForTree)
@@ -3754,7 +3823,10 @@ void SSCSEditor::OnLevelComponentRequestRename(const UActorComponent* InComponen
 
 void SSCSEditor::OnObjectsReplaced(const TMap<UObject*, UObject*>& OldToNewInstanceMap)
 {
-	ReplaceComponentReferencesInTree(GetRootComponentNodes(), OldToNewInstanceMap);
+	if (GetActorNode().IsValid())
+	{
+		ReplaceComponentReferencesInTree(GetActorNode()->GetComponentNodes(), OldToNewInstanceMap);
+	}
 }
 
 void SSCSEditor::ReplaceComponentReferencesInTree(const TArray<FSCSEditorTreeNodePtrType>& Nodes, const TMap<UObject*, UObject*>& OldToNewInstanceMap)
@@ -4123,7 +4195,7 @@ void SSCSEditor::OnDuplicateComponent()
 			{
 				// If we're duplicating the root then we're already a child of it so need to reparent, but we do need to reset the scale
 				// otherwise we'll end up with the square of the root's scale instead of being the same size.
-				if (OriginalNodePtr == SceneRootNodePtr)
+				if (OriginalNodePtr == GetSceneRootNode())
 				{
 					NewSceneComponent->RelativeScale3D = FVector(1.f);
 				}
@@ -4144,7 +4216,7 @@ void SSCSEditor::OnDuplicateComponent()
 						}
 
 						// Locate the duplicate node (as a child of the current scene root node), and switch it to be a child of the original node's parent
-						FSCSEditorTreeNodePtrType NewChildNodePtr = SceneRootNodePtr->FindChild(NewSceneComponent, true);
+						FSCSEditorTreeNodePtrType NewChildNodePtr = GetSceneRootNode()->FindChild(NewSceneComponent, true);
 						if (NewChildNodePtr.IsValid())
 						{
 							// Note: This method will handle removal from the scene root node as well
@@ -4457,7 +4529,7 @@ void SSCSEditor::UpdateTree(bool bRegenerateTreeNodes)
 	{
 		// Obtain the set of expandable tree nodes that are currently collapsed
 		TSet<FSCSEditorTreeNodePtrType> CollapsedTreeNodes;
-		GetCollapsedNodes(SceneRootNodePtr, CollapsedTreeNodes);
+		GetCollapsedNodes(GetSceneRootNode(), CollapsedTreeNodes);
 
 		// Obtain the list of selected items
 		TArray<FSCSEditorTreeNodePtrType> SelectedTreeNodes = SCSTreeWidget->GetSelectedItems();
@@ -4468,18 +4540,11 @@ void SSCSEditor::UpdateTree(bool bRegenerateTreeNodes)
 			SCSTreeWidget->ClearSelection();
 		}
 		RootNodes.Empty();
-		RootComponentNodes.Empty();
-
-		bHasAddedSceneAndBehaviorComponentSeparator = false;
-
-		// Reset the scene root node
-		SceneRootNodePtr.Reset();
-
+		
 		TSharedPtr<FSCSEditorTreeNode> ActorTreeNode = MakeShareable(new FSCSEditorTreeNodeRootActor(GetActorContext(),EditorMode == EComponentEditorMode::ActorInstance));
-
+		SCSTreeWidget->SetItemExpansion(ActorTreeNode, true);
 		RootNodes.Add(ActorTreeNode);
-		RootNodes.Add(MakeShareable(new FSCSEditorTreeNodeSeparator()));
-
+				
 		// Build the tree data source according to what mode we're in
 		if (EditorMode == EComponentEditorMode::BlueprintSCS)
 		{
@@ -4510,26 +4575,13 @@ void SSCSEditor::UpdateTree(bool bRegenerateTreeNodes)
 				if(RootComponent != nullptr)
 				{
 					Components.Remove(RootComponent);
-					AddTreeNodeFromComponent(RootComponent);
+					AddTreeNodeFromComponent(RootComponent, FindOrCreateParentForExistingComponent(RootComponent, GetActorNode()));
 				}
 				
 				for (UActorComponent* Component : Components)
 				{
-					if (USceneComponent* SceneComp = Cast<USceneComponent>(Component))
-					{
-						// Add the rest of the native base class SceneComponent hierarchy
-						AddTreeNodeFromComponent(SceneComp);
-					}
-					else
-					{
-						// Add native ActorComponent nodes that aren't SceneComponents
-						if (!bHasAddedSceneAndBehaviorComponentSeparator)
-						{
-							bHasAddedSceneAndBehaviorComponentSeparator = true;
-							RootNodes.Add(MakeShareable(new FSCSEditorTreeNodeSeparator()));
-						}
-						AddRootComponentTreeNode(Component);
-					}
+					// Add the rest of the native base class SceneComponent hierarchy
+					AddTreeNodeFromComponent(Component, FindOrCreateParentForExistingComponent(Component, GetActorNode()));
 				}
 			}
 
@@ -4559,7 +4611,7 @@ void SSCSEditor::UpdateTree(bool bRegenerateTreeNodes)
 						}
 						else
 						{
-							NewNodePtr = AddTreeNode(SCS_Node, SceneRootNodePtr, StackIndex > 0);
+							NewNodePtr = AddTreeNode(SCS_Node, ActorTreeNode, StackIndex > 0);
 						}
 
 						// Only necessary to do the following for inherited nodes (StackIndex > 0).
@@ -4567,19 +4619,19 @@ void SSCSEditor::UpdateTree(bool bRegenerateTreeNodes)
 						{
 							// This call creates ICH override templates for the current Blueprint. Without this, the parent node
 							// search above can fail when attempting to match an inherited node in the tree via component template.
-							NewNodePtr->GetEditableComponentTemplate(ParentBPStack[0]);
+							NewNodePtr->GetOrCreateEditableComponentTemplate(ParentBPStack[0]);
 							for (FSCSEditorTreeNodePtrType ChildNodePtr : NewNodePtr->GetChildren())
 							{
 								if (ensure(ChildNodePtr.IsValid()))
 								{
-									ChildNodePtr->GetEditableComponentTemplate(ParentBPStack[0]);
+									ChildNodePtr->GetOrCreateEditableComponentTemplate(ParentBPStack[0]);
 								}
 							}
 						}
 					}
 				}
 			}
-
+						
 			AActor* PreviewActorInstance = PreviewActor.Get();
 			if(PreviewActorInstance != nullptr && !GetDefault<UBlueprintEditorSettings>()->bHideConstructionScriptComponentsInDetailsView)
 			{
@@ -4590,15 +4642,7 @@ void SSCSEditor::UpdateTree(bool bRegenerateTreeNodes)
 				{
 					if(Component->CreationMethod == EComponentCreationMethod::UserConstructionScript)
 					{
-						USceneComponent* SceneComponent = Cast<USceneComponent>(Component);
-						if(SceneComponent != nullptr)
-						{
-							AddTreeNodeFromComponent(SceneComponent);
-						}
-						else
-						{
-							AddRootComponentTreeNode(Component);
-						}
+						AddTreeNodeFromComponent(Component, FindOrCreateParentForExistingComponent(Component, GetActorNode()));
 					}
 				}
 			}
@@ -4619,7 +4663,7 @@ void SSCSEditor::UpdateTree(bool bRegenerateTreeNodes)
 
 					// Recursively add any instanced children that are already attached through the root, and keep track of added
 					// instances. This will be a faster path than the loop below, because we create new parent tree nodes as we go.
-					FSCSEditorTreeNodePtrType NewParentNode = AddTreeNodeFromComponent(RootComponent);
+					FSCSEditorTreeNodePtrType NewParentNode = AddTreeNodeFromComponent(RootComponent, FindOrCreateParentForExistingComponent(RootComponent, GetActorNode()));
 					AddInstancedTreeNodesRecursive(RootComponent, NewParentNode, ComponentsToAdd);
 				}
 
@@ -4637,19 +4681,7 @@ void SSCSEditor::UpdateTree(bool bRegenerateTreeNodes)
 					USceneComponent* ParentSceneComp = SceneComp != nullptr ? SceneComp->GetAttachParent() : nullptr;
 					if (ShouldAddInstancedActorComponent(ActorComp, ParentSceneComp))
 					{
-						if (SceneComp != nullptr)
-						{
-							AddTreeNodeFromComponent(SceneComp);
-						}
-						else
-						{
-							if (!bHasAddedSceneAndBehaviorComponentSeparator)
-							{
-								bHasAddedSceneAndBehaviorComponentSeparator = true;
-								RootNodes.Add(MakeShareable(new FSCSEditorTreeNode(FSCSEditorTreeNode::SeparatorNode)));
-							}
-							AddRootComponentTreeNode(ActorComp);
-						}
+						AddTreeNodeFromComponent(ActorComp, FindOrCreateParentForExistingComponent(ActorComp, GetActorNode()));
 					}
 				}
 			}
@@ -4705,8 +4737,6 @@ void SSCSEditor::UpdateTree(bool bRegenerateTreeNodes)
 				SCSTreeWidget->RequestScrollIntoView(NodeToRenamePtr);
 			}
 		}
-
-		RebuildFilteredRootList();
 	}
 
 	// refresh widget
@@ -4911,29 +4941,30 @@ const TArray<FSCSEditorTreeNodePtrType>& SSCSEditor::GetRootNodes() const
 	return RootNodes;
 }
 
-TSharedPtr<FSCSEditorTreeNode> SSCSEditor::AddRootComponentTreeNode(UActorComponent* ActorComp)
+FSCSEditorActorNodePtrType SSCSEditor::GetActorNode() const
 {
-	TSharedPtr<FSCSEditorTreeNode> NewTreeNode;
-	if (RootTreeNode.IsValid())
+	if (RootNodes.Num() > 0)
 	{
-		NewTreeNode = RootTreeNode->AddChildFromComponent(ActorComp);
-		RefreshFilteredState(NewTreeNode, /*bRecursive =*/false);
+		return StaticCastSharedPtr<FSCSEditorTreeNodeRootActor>(RootNodes[0]);
 	}
-	else
-	{
-		NewTreeNode = FSCSEditorTreeNode::FactoryNodeFromComponent(ActorComp);
-		RootNodes.Add(NewTreeNode);
 
-		bool bIsFilteredOut = RefreshFilteredState(NewTreeNode, /*bRecursive =*/false);
-		if (!bIsFilteredOut)
-		{
-			FilteredRootNodes.Add(NewTreeNode);
-		}
+	return FSCSEditorActorNodePtrType();
+}
+
+FSCSEditorTreeNodePtrType SSCSEditor::GetSceneRootNode() const
+{
+	FSCSEditorActorNodePtrType ActorNode = GetActorNode();
+	if (ActorNode.IsValid())
+	{
+		return ActorNode->GetSceneRootNode();
 	}
 
-	RootComponentNodes.Add(NewTreeNode);
+	return FSCSEditorTreeNodePtrType();
+}
 
-	return NewTreeNode;
+void SSCSEditor::SetSceneRootNode(FSCSEditorTreeNodePtrType NewSceneRootNode)
+{
+	GetActorNode()->SetSceneRootNode(NewSceneRootNode);
 }
 
 class FComponentClassParentFilter : public IClassViewerFilter
@@ -5034,36 +5065,6 @@ UClass* SSCSEditor::CreateNewBPComponent(TSubclassOf<UActorComponent> ComponentC
 	return NewClass;
 }
 
-void SSCSEditor::RebuildFilteredRootList()
-{
- 	FilteredRootNodes.Empty(RootNodes.Num());
-
-	FSCSEditorTreeNodePtrType PendingSeparator;
-	for (const FSCSEditorTreeNodePtrType& Node : RootNodes)
-	{
-		switch (Node->GetNodeType())
-		{
-		case FSCSEditorTreeNode::ENodeType::ComponentNode:
-			if (Node->IsFlaggedForFiltration())
-			{	
-				break;
-			}
-		case FSCSEditorTreeNode::ENodeType::RootActorNode:
-			if (PendingSeparator.IsValid())
-			{
-				FilteredRootNodes.Add(PendingSeparator);
-				PendingSeparator.Reset();
-			}
-			FilteredRootNodes.Add(Node);
-			break;
-
-		case FSCSEditorTreeNode::ENodeType::SeparatorNode:
-			PendingSeparator = Node;
-			break;
-		}
-	}
-}
-
 void SSCSEditor::ClearSelection()
 {
 	if ( bUpdatingSelection == false )
@@ -5179,7 +5180,9 @@ UActorComponent* SSCSEditor::AddNewComponent( UClass* NewComponentClass, UObject
 		if (ComponentTemplate)
 		{
 			// Create a duplicate of the provided template
-			NewComponent = AddNewNodeForInstancedComponent(FComponentEditorUtils::DuplicateComponent(ComponentTemplate), nullptr, bSetFocusToNewItem);
+			UActorComponent* NewInstanceComponent = FComponentEditorUtils::DuplicateComponent(ComponentTemplate);
+			FSCSEditorTreeNodePtrType ParentNodePtr = FindParentForNewComponent(NewInstanceComponent);
+			NewComponent = AddNewNodeForInstancedComponent(NewInstanceComponent, ParentNodePtr, nullptr, bSetFocusToNewItem);
 		}
 		else if (AActor* ActorInstance = GetActorContext())
 		{
@@ -5203,16 +5206,24 @@ UActorComponent* SSCSEditor::AddNewComponent( UClass* NewComponentClass, UObject
 
 			// Construct the new component and attach as needed
 			UActorComponent* NewInstanceComponent = NewObject<UActorComponent>(ActorInstance, NewComponentClass, NewComponentName, RF_Transactional);
+			FSCSEditorTreeNodePtrType ParentNodePtr = FindParentForNewComponent(NewInstanceComponent);
+						
+			// Do Scene Attachment if this new Comnponent is a USceneComponent
 			if (USceneComponent* NewSceneComponent = Cast<USceneComponent>(NewInstanceComponent))
 			{
-				USceneComponent* RootComponent = ActorInstance->GetRootComponent();
-				if (RootComponent)
-				{
-					NewSceneComponent->AttachToComponent(RootComponent, FAttachmentTransformRules::KeepRelativeTransform);
+				if(ParentNodePtr->GetNodeType() == FSCSEditorTreeNode::RootActorNode)
+				{ 
+					ActorInstance->SetRootComponent(NewSceneComponent);
 				}
 				else
 				{
-					ActorInstance->SetRootComponent(NewSceneComponent);
+					USceneComponent* AttachTo = Cast<USceneComponent>(ParentNodePtr->GetComponentTemplate());
+					if (AttachTo == nullptr)
+					{
+						AttachTo = ActorInstance->GetRootComponent();
+					}
+					check(AttachTo != nullptr);
+					NewSceneComponent->AttachToComponent(AttachTo, FAttachmentTransformRules::KeepRelativeTransform);
 				}
 			}
 
@@ -5241,53 +5252,143 @@ UActorComponent* SSCSEditor::AddNewComponent( UClass* NewComponentClass, UObject
 			// Rerun construction scripts
 			ActorInstance->RerunConstructionScripts();
 
-			NewComponent = AddNewNodeForInstancedComponent(NewInstanceComponent, Asset, bSetFocusToNewItem);
+			NewComponent = AddNewNodeForInstancedComponent(NewInstanceComponent, ParentNodePtr, Asset, bSetFocusToNewItem);
 		}
 	}
 
 	return NewComponent;
 }
 
-UActorComponent* SSCSEditor::AddNewNode(USCS_Node* NewNode, UObject* Asset, bool bMarkBlueprintModified, bool bSetFocusToNewItem)
+FSCSEditorTreeNodePtrType SSCSEditor::FindOrCreateParentForExistingComponent(UActorComponent* InActorComponent, FSCSEditorActorNodePtrType ActorRootNode)
 {
-	check(NewNode != nullptr);
+	check(InActorComponent != nullptr);
 
-	if(Asset)
+	USceneComponent* SceneComponent = Cast<USceneComponent>(InActorComponent);
+	if (SceneComponent == nullptr)
 	{
-		FComponentAssetBrokerage::AssignAssetToComponent(NewNode->ComponentTemplate, Asset);
+		check(ActorRootNode.IsValid());
+		check(ActorRootNode->GetNodeType() == FSCSEditorTreeNode::RootActorNode);
+		return ActorRootNode;
 	}
 
-	FSCSEditorTreeNodePtrType NewNodePtr;
+	FSCSEditorTreeNodePtrType ParentNodePtr;
+	if (SceneComponent->GetAttachParent() != nullptr
+		&& (EditorMode != EComponentEditorMode::ActorInstance || SceneComponent->GetAttachParent()->GetOwner() == GetActorContext()))
+	{
+		// Attempt to find the parent node in the current tree
+		ParentNodePtr = FindTreeNode(SceneComponent->GetAttachParent());
+		if (!ParentNodePtr.IsValid())
+		{
+			// If the actual attach parent wasn't found, attempt to find its archetype.
+			// This handles the BP editor case where we might add UCS component nodes taken
+			// from the preview actor instance, which are not themselves template objects.
+			ParentNodePtr = FindTreeNode(Cast<USceneComponent>(SceneComponent->GetAttachParent()->GetArchetype()));
+			if (!ParentNodePtr.IsValid())
+			{
+				// Recursively add the parent node to the tree if it does not exist yet
+				ParentNodePtr = AddTreeNodeFromComponent(SceneComponent->GetAttachParent(), FindOrCreateParentForExistingComponent(SceneComponent->GetAttachParent(), ActorRootNode));
+			}
+		}
+	}
 
-	UBlueprint* Blueprint = GetBlueprint();
-	check(Blueprint != nullptr && Blueprint->SimpleConstructionScript != nullptr);
+	if (!ParentNodePtr.IsValid())
+	{
+		ParentNodePtr = ActorRootNode->GetSceneRootNode();
+	}
 
-	bool AttachToSceneRootNode = true;
-	if (USceneComponent* NewSceneComponent = Cast<USceneComponent>(NewNode->ComponentTemplate))
+	// Actor doesn't have a root component yet
+	if (!ParentNodePtr.IsValid())
 	{
-		// get currently selected component
-		TArray<FSCSEditorTreeNodePtrType> SelectedTreeNodes;
-		if (SCSTreeWidget.IsValid() && SCSTreeWidget->GetSelectedItems(SelectedTreeNodes) > 0)
+		ParentNodePtr = ActorRootNode; 
+	}
+
+	return ParentNodePtr;
+}
+
+FSCSEditorTreeNodePtrType SSCSEditor::FindParentForNewComponent(UActorComponent* NewComponent) const
+{
+	// Find Parent to attach to (depending on the new Node type).
+	FSCSEditorTreeNodePtrType TargetParentNode;
+	TArray<FSCSEditorTreeNodePtrType> SelectedTreeNodes;
+	if (SCSTreeWidget.IsValid() && SCSTreeWidget->GetSelectedItems(SelectedTreeNodes))
+	{
+		TargetParentNode = SelectedTreeNodes[0];
+	}
+
+	if (USceneComponent* NewSceneComponent = Cast<USceneComponent>(NewComponent))
+	{
+		if (TargetParentNode.IsValid())
 		{
-			FSCSEditorTreeNodePtrType FirstTreeNode = SelectedTreeNodes[0];
-			if (FirstTreeNode.IsValid() && FirstTreeNode->GetComponentTemplate())
+			if (TargetParentNode->GetNodeType() == FSCSEditorTreeNode::RootActorNode)
+			{
+				FSCSEditorActorNodePtrType TargetActorNode = StaticCastSharedPtr<FSCSEditorTreeNodeRootActor>(TargetParentNode);
+				if (TargetActorNode.IsValid())
+				{
+					TargetParentNode = TargetActorNode->GetSceneRootNode();
+					USceneComponent* CastTargetToSceneComponent = Cast<USceneComponent>(TargetParentNode->GetComponentTemplate());
+					if (CastTargetToSceneComponent == nullptr || !NewSceneComponent->CanAttachAsChild(CastTargetToSceneComponent, NAME_None))
+					{
+						TargetParentNode = GetSceneRootNode(); // Default to SceneRoot
+					}
+				}
+			}
+			else if(TargetParentNode->GetNodeType() == FSCSEditorTreeNode::ComponentNode)
 			{
-				USceneComponent* CastFirstTreeNode = Cast<USceneComponent>(FirstTreeNode->GetComponentTemplate());
-				if (CastFirstTreeNode && NewSceneComponent->CanAttachAsChild(CastFirstTreeNode, NAME_None))
+				USceneComponent* CastTargetToSceneComponent = Cast<USceneComponent>(TargetParentNode->GetComponentTemplate());
+				if (CastTargetToSceneComponent == nullptr || !NewSceneComponent->CanAttachAsChild(CastTargetToSceneComponent, NAME_None))
 				{
-					NewNodePtr = AddTreeNode(NewNode, FirstTreeNode, false);
-					AttachToSceneRootNode = false;
+					TargetParentNode = GetSceneRootNode(); // Default to SceneRoot
 				}
 			}
 		}
+		else
+		{
+			TargetParentNode = GetSceneRootNode();
+		}
+	}
+	else
+	{
+		if (TargetParentNode.IsValid())
+		{
+			while (TargetParentNode->GetNodeType() == FSCSEditorTreeNode::ComponentNode)
+			{
+				TargetParentNode = TargetParentNode->GetParent();
+			}
+		}
+		else
+		{
+			TargetParentNode = GetActorNode();
+		}
+
+		check(TargetParentNode->GetNodeType() == FSCSEditorTreeNode::RootActorNode);
 	}
 
-	if (AttachToSceneRootNode)
+	return TargetParentNode;
+}
+
+FSCSEditorTreeNodePtrType SSCSEditor::FindParentForNewNode(USCS_Node* NewNode) const
+{
+	return FindParentForNewComponent(NewNode->ComponentTemplate);
+}
+
+UActorComponent* SSCSEditor::AddNewNode(USCS_Node* NewNode, UObject* Asset, bool bMarkBlueprintModified, bool bSetFocusToNewItem)
+{
+	check(NewNode != nullptr);
+
+	if(Asset)
 	{
-		// Add the new node to the editor tree
-		NewNodePtr = AddTreeNode(NewNode, SceneRootNodePtr, false);
+		FComponentAssetBrokerage::AssignAssetToComponent(NewNode->ComponentTemplate, Asset);
 	}
 
+	FSCSEditorTreeNodePtrType NewNodePtr;
+	FSCSEditorTreeNodePtrType ParentNodePtr = FindParentForNewNode(NewNode);
+	
+	UBlueprint* Blueprint = GetBlueprint();
+	check(Blueprint != nullptr && Blueprint->SimpleConstructionScript != nullptr);
+
+	// Add the new node to the editor tree
+	NewNodePtr = AddTreeNode(NewNode, ParentNodePtr, /*bIsInheritedSCS=*/ false);
+	
 	// Potentially adjust variable names for any child blueprints
 	const FName VariableName = NewNode->GetVariableName();
 	if(VariableName != NAME_None)
@@ -5315,38 +5416,15 @@ UActorComponent* SSCSEditor::AddNewNode(USCS_Node* NewNode, UObject* Asset, bool
 	return NewNode->ComponentTemplate;
 }
 
-UActorComponent* SSCSEditor::AddNewNodeForInstancedComponent(UActorComponent* NewInstanceComponent, UObject* Asset, bool bSetFocusToNewItem)
+UActorComponent* SSCSEditor::AddNewNodeForInstancedComponent(UActorComponent* NewInstanceComponent, FSCSEditorTreeNodePtrType InParentNodePtr, UObject* Asset, bool bSetFocusToNewItem)
 {
 	check(NewInstanceComponent != nullptr);
 
 	FSCSEditorTreeNodePtrType NewNodePtr;
 
 	// Add the new node to the editor tree
-	USceneComponent* NewSceneComponent = Cast<USceneComponent>(NewInstanceComponent);
-	if(NewSceneComponent != nullptr)
-	{
-		NewNodePtr = AddTreeNodeFromComponent(NewSceneComponent);
-
-		// Remove the old scene root node if it's set to the default one
-		//if(SceneRootNodePtr.IsValid() && SceneRootNodePtr->IsDefaultSceneRoot())
-		//{
-		//	RemoveComponentNode(SceneRootNodePtr);
-		//	RootNodes.Remove( SceneRootNodePtr );
-		//	SceneRootNodePtr.Reset();
-		//}
-	}
-	else
-	{
-		// Make sure we've added the separator between scene and behavior components
-		if (!bHasAddedSceneAndBehaviorComponentSeparator)
-		{
-			bHasAddedSceneAndBehaviorComponentSeparator = true;
-			RootNodes.Add(MakeShareable(new FSCSEditorTreeNode(FSCSEditorTreeNode::SeparatorNode)));
-		}
-
-		NewNodePtr = AddRootComponentTreeNode(NewInstanceComponent);
-	}
-
+	NewNodePtr = AddTreeNodeFromComponent(NewInstanceComponent, InParentNodePtr);
+	
 	if(bSetFocusToNewItem)
 	{
 		// Select and request a rename on the new component
@@ -5496,6 +5574,7 @@ bool SSCSEditor::CanPasteNodes() const
 		return false;
 	}
 
+	FSCSEditorTreeNodePtrType SceneRootNodePtr = GetSceneRootNode();
 	return SceneRootNodePtr.IsValid() && FComponentEditorUtils::CanPasteComponents(Cast<USceneComponent>(SceneRootNodePtr->GetComponentTemplate()), SceneRootNodePtr->IsDefaultSceneRoot(), true);
 }
 
@@ -5856,124 +5935,73 @@ FSCSEditorTreeNodePtrType SSCSEditor::AddTreeNode(USCS_Node* InSCSNode, FSCSEdit
 {
 	FSCSEditorTreeNodePtrType NewNodePtr;
 
-	check(InSCSNode != NULL);
+	check(InSCSNode != nullptr && InParentNodePtr.IsValid());
 
 	// During diffs, ComponentTemplates can easily be null, so prevent these checks.
 	if (!bIsDiffing && InSCSNode->ComponentTemplate)
 	{
 		checkf(InSCSNode->ParentComponentOrVariableName == NAME_None
-			|| (!InSCSNode->bIsParentComponentNative && InParentNodePtr->GetSCSNode() != NULL && InParentNodePtr->GetSCSNode()->GetVariableName() == InSCSNode->ParentComponentOrVariableName)
-			|| (InSCSNode->bIsParentComponentNative && InParentNodePtr->GetComponentTemplate() != NULL && InParentNodePtr->GetComponentTemplate()->GetFName() == InSCSNode->ParentComponentOrVariableName),
+			|| (!InSCSNode->bIsParentComponentNative && InParentNodePtr->GetSCSNode() != nullptr && InParentNodePtr->GetSCSNode()->GetVariableName() == InSCSNode->ParentComponentOrVariableName)
+			|| (InSCSNode->bIsParentComponentNative && InParentNodePtr->GetComponentTemplate() != nullptr && InParentNodePtr->GetComponentTemplate()->GetFName() == InSCSNode->ParentComponentOrVariableName),
 			TEXT("Failed to add SCS node %s to tree:\n- bIsParentComponentNative=%d\n- Stored ParentComponentOrVariableName=%s\n- Actual ParentComponentOrVariableName=%s"),
 			*InSCSNode->GetVariableName().ToString(),
 			!!InSCSNode->bIsParentComponentNative,
 			*InSCSNode->ParentComponentOrVariableName.ToString(),
 			!InSCSNode->bIsParentComponentNative
-			? (InParentNodePtr->GetSCSNode() != NULL ? *InParentNodePtr->GetSCSNode()->GetVariableName().ToString() : TEXT("NULL"))
-			: (InParentNodePtr->GetComponentTemplate() != NULL ? *InParentNodePtr->GetComponentTemplate()->GetFName().ToString() : TEXT("NULL")));
+			? (InParentNodePtr->GetSCSNode() != nullptr ? *InParentNodePtr->GetSCSNode()->GetVariableName().ToString() : TEXT("NULL"))
+			: (InParentNodePtr->GetComponentTemplate() != nullptr ? *InParentNodePtr->GetComponentTemplate()->GetFName().ToString() : TEXT("NULL")));
 	}
 	
 	// Determine whether or not the given node is inherited from a parent Blueprint
 	USimpleConstructionScript* NodeSCS = InSCSNode->GetSCS();
+	
+	// do this first, because we need a FSCSEditorTreeNodePtrType for the new node
+	NewNodePtr = InParentNodePtr->AddChild(InSCSNode, bIsInheritedSCS);
+	RefreshFilteredState(NewNodePtr, /*bRecursive =*/false);
+	
 
-	if(InSCSNode->ComponentTemplate && InSCSNode->ComponentTemplate->IsA(USceneComponent::StaticClass()))
+	if( InSCSNode->ComponentTemplate && 
+		InSCSNode->ComponentTemplate->IsA(USceneComponent::StaticClass()) && 
+		InParentNodePtr->GetNodeType() == FSCSEditorTreeNode::ComponentNode)
 	{
-		FSCSEditorTreeNodePtrType ParentPtr = InParentNodePtr.IsValid() ? InParentNodePtr : SceneRootNodePtr;
-		if(ParentPtr.IsValid())
+		bool bParentIsEditorOnly = InParentNodePtr->GetComponentTemplate()->IsEditorOnly();
+		// if you can't nest this new node under the proposed parent (then swap the two)
+		if (bParentIsEditorOnly && !InSCSNode->ComponentTemplate->IsEditorOnly() && InParentNodePtr->CanReparent())
 		{
-			// do this first, because we need a FSCSEditorTreeNodePtrType for the new node
-			NewNodePtr = ParentPtr->AddChild(InSCSNode, bIsInheritedSCS);
-			RefreshFilteredState(NewNodePtr, /*bRecursive =*/false);
-
-			bool bParentIsEditorOnly = ParentPtr->GetComponentTemplate()->IsEditorOnly();
-			// if you can't nest this new node under the proposed parent (then swap the two)
-			if (bParentIsEditorOnly && !InSCSNode->ComponentTemplate->IsEditorOnly() && ParentPtr->CanReparent())
-			{
-				FSCSEditorTreeNodePtrType OldParentPtr = ParentPtr;
-				ParentPtr = OldParentPtr->GetParent();
+			FSCSEditorTreeNodePtrType OldParentPtr = InParentNodePtr;
+			InParentNodePtr = OldParentPtr->GetParent();
 
-				OldParentPtr->RemoveChild(NewNodePtr);
-				NodeSCS->RemoveNode(OldParentPtr->GetSCSNode());
+			OldParentPtr->RemoveChild(NewNodePtr);
+			NodeSCS->RemoveNode(OldParentPtr->GetSCSNode());
 
-				// if the grandparent node is invalid (assuming this means that the parent node was the scene-root)
-				if (!ParentPtr.IsValid())
-				{
-					check(OldParentPtr == SceneRootNodePtr);
-					SceneRootNodePtr = NewNodePtr;
-					NodeSCS->AddNode(SceneRootNodePtr->GetSCSNode());
-				}
-				else 
-				{
-					ParentPtr->AddChild(NewNodePtr);
-				}
-
-				// move the proposed parent in as a child to the new node
-				NewNodePtr->AddChild(OldParentPtr);
-			} // if bParentIsEditorOnly...
-
-			// Expand parent nodes by default
-			SCSTreeWidget->SetItemExpansion(ParentPtr, true);
-		}
-		//else, if !SceneRootNodePtr.IsValid(), make it the scene root node if it has not been set yet
-		else 
-		{
-			// Create a new root node
-			if (RootTreeNode.IsValid())
+			// if the grandparent node is invalid (assuming this means that the parent node was the scene-root)
+			if (!InParentNodePtr.IsValid())
 			{
-				NewNodePtr = RootTreeNode->AddChild(InSCSNode, bIsInheritedSCS);
+				check(OldParentPtr == GetSceneRootNode());
+				SetSceneRootNode(NewNodePtr);
+				NodeSCS->AddNode(NewNodePtr->GetSCSNode());
 			}
 			else
 			{
-				NewNodePtr = MakeShareable(new FSCSEditorTreeNodeComponent(InSCSNode, bIsInheritedSCS));
-				RootNodes.Add(NewNodePtr);
-
-				bool bIsFilteredOut = RefreshFilteredState(NewNodePtr, /*bRecursive =*/false);
-				if (!bIsFilteredOut)
-				{
-					FilteredRootNodes.Add(NewNodePtr);
-				}
+				InParentNodePtr->AddChild(NewNodePtr);
 			}
-			
-			NodeSCS->AddNode(InSCSNode);
-			
-			// Add it to the root set
-			RootComponentNodes.Insert(NewNodePtr, 0);
-
-			// Make it the scene root node
-			SceneRootNodePtr = NewNodePtr;
 
-			// Expand the scene root node by default
-			SCSTreeWidget->SetItemExpansion(SceneRootNodePtr, true);
-		}
+			// move the proposed parent in as a child to the new node
+			NewNodePtr->AddChild(OldParentPtr);
+		} // if bParentIsEditorOnly...
 	}
 	else
-	{
-		// If the given SCS node does not contain a scene component template, we create a new root node
-		if (RootTreeNode.IsValid())
-		{
-			NewNodePtr = RootTreeNode->AddChild(InSCSNode, bIsInheritedSCS);
-		}
-		else
-		{
-			NewNodePtr = MakeShareable(new FSCSEditorTreeNodeComponent(InSCSNode, bIsInheritedSCS));
-			RootNodes.Add(NewNodePtr);
-
-			bool bIsFilteredOut = RefreshFilteredState(NewNodePtr, /*bRecursive =*/false);
-			if (!bIsFilteredOut)
-			{
-				FilteredRootNodes.Add(NewNodePtr);
-			}
-		}
-
-		RootComponentNodes.Add(NewNodePtr);
-
+	{	
 		// If the SCS root node array does not already contain the given node, this will add it (this should only occur after node creation)
-		if(NodeSCS != NULL)
+		if(NodeSCS != nullptr)
 		{
 			NodeSCS->AddNode(InSCSNode);
 		}
 	}
 
+	// Expand parent nodes by default
+	SCSTreeWidget->SetItemExpansion(InParentNodePtr, true);
+
 	// Recursively add the given SCS node's child nodes
 	for (USCS_Node* ChildNode : InSCSNode->GetChildNodes())
 	{
@@ -5983,69 +6011,21 @@ FSCSEditorTreeNodePtrType SSCSEditor::AddTreeNode(USCS_Node* InSCSNode, FSCSEdit
 	return NewNodePtr;
 }
 
-FSCSEditorTreeNodePtrType SSCSEditor::AddTreeNodeFromComponent(USceneComponent* InSceneComponent, FSCSEditorTreeNodePtrType InParentTreeNode)
+FSCSEditorTreeNodePtrType SSCSEditor::AddTreeNodeFromComponent(UActorComponent* InActorComponent, FSCSEditorTreeNodePtrType InParentTreeNode)
 {
-	FSCSEditorTreeNodePtrType NewNodePtr;
+	check(InActorComponent != NULL);
+	ensure(!InActorComponent->IsPendingKill());
 
-	check(InSceneComponent != NULL);
-	ensure(!InSceneComponent->IsPendingKill());
-
-	// If the given component has a parent, and if we're not in "instance" mode OR the owner of the parent matches the Actor instance we're editing
-	if(InSceneComponent->GetAttachParent() != NULL
-		&& (EditorMode != EComponentEditorMode::ActorInstance || InSceneComponent->GetAttachParent()->GetOwner() == GetActorContext()))
+	FSCSEditorTreeNodePtrType NewNodePtr = InParentTreeNode->FindChild(InActorComponent);
+	if (!NewNodePtr.IsValid())
 	{
-		// Attempt to find the parent node in the current tree
-		FSCSEditorTreeNodePtrType ParentNodePtr;
-		if (InParentTreeNode.IsValid())
-		{
-			ParentNodePtr = InParentTreeNode;
-		}
-		else
-		{
-			ParentNodePtr = FindTreeNode(InSceneComponent->GetAttachParent());
-			if(!ParentNodePtr.IsValid())
-			{
-				// If the actual attach parent wasn't found, attempt to find its archetype.
-				// This handles the BP editor case where we might add UCS component nodes taken
-				// from the preview actor instance, which are not themselves template objects.
-				ParentNodePtr = FindTreeNode(Cast<USceneComponent>(InSceneComponent->GetAttachParent()->GetArchetype()));
-				if(!ParentNodePtr.IsValid())
-				{
-					// Recursively add the parent node to the tree if it does not exist yet
-					ParentNodePtr = AddTreeNodeFromComponent(InSceneComponent->GetAttachParent());
-				}
-			}
-		}
-
-		// Add a new tree node for the given scene component
-		check(ParentNodePtr.IsValid());
-		NewNodePtr = ParentNodePtr->AddChildFromComponent(InSceneComponent);
-		RefreshFilteredState(NewNodePtr, /*bRecursive =*/false);
-
-		// Expand parent nodes by default
-		SCSTreeWidget->SetItemExpansion(ParentNodePtr, true);
-	}
-	else
-	{
-		// Make it the scene root node if it has not been set yet
-		if(!SceneRootNodePtr.IsValid())
-		{
-			// Create a new root node
-			NewNodePtr = AddRootComponentTreeNode(InSceneComponent);
-
-			// Make it the scene root node
-			SceneRootNodePtr = NewNodePtr;
-
-			// Expand the scene root node by default
-			SCSTreeWidget->SetItemExpansion(SceneRootNodePtr, true);
-		}
-		else if (SceneRootNodePtr->GetComponentTemplate() != InSceneComponent)
-		{
-			NewNodePtr = SceneRootNodePtr->AddChildFromComponent(InSceneComponent);
-			RefreshFilteredState(NewNodePtr, /*bRecursive =*/false);
-		}
+		NewNodePtr = FSCSEditorTreeNode::FactoryNodeFromComponent(InActorComponent);
+		InParentTreeNode->AddChild(NewNodePtr);
+		RefreshFilteredState(NewNodePtr, false);
 	}
-
+		
+	SCSTreeWidget->SetItemExpansion(NewNodePtr, true);
+	
 	return NewNodePtr;
 }
 
@@ -6057,7 +6037,7 @@ FSCSEditorTreeNodePtrType SSCSEditor::FindTreeNode(const USCS_Node* InSCSNode, F
 		// Start at the scene root node if none was given
 		if(!InStartNodePtr.IsValid())
 		{
-			InStartNodePtr = SceneRootNodePtr;
+			InStartNodePtr = GetSceneRootNode();
 		}
 
 		if(InStartNodePtr.IsValid())
@@ -6093,17 +6073,17 @@ FSCSEditorTreeNodePtrType SSCSEditor::FindTreeNode(const UActorComponent* InComp
 		// Start at the scene root node if none was given
 		if(!InStartNodePtr.IsValid())
 		{
-			InStartNodePtr = SceneRootNodePtr;
+			InStartNodePtr = GetActorNode();
 		}
 
 		if(InStartNodePtr.IsValid())
 		{
 			// Check to see if the given component template matches the given tree node
 			// 
-			// For certain node types, GetEditableComponentTemplate() will handle retrieving 
+			// For certain node types, GetOrCreateEditableComponentTemplate() will handle retrieving 
 			// the "OverridenComponentTemplate" which may be what we're looking for in some 
 			// cases; if not, then we fall back to just checking GetComponentTemplate()
-			if (InStartNodePtr->GetEditableComponentTemplate(GetBlueprint()) == InComponent)
+			if (InStartNodePtr->GetOrCreateEditableComponentTemplate(GetBlueprint()) == InComponent)
 			{
 				NodePtr = InStartNodePtr;
 			}
@@ -6134,10 +6114,10 @@ FSCSEditorTreeNodePtrType SSCSEditor::FindTreeNode(const FName& InVariableOrInst
 	FSCSEditorTreeNodePtrType NodePtr;
 	if(InVariableOrInstanceName != NAME_None)
 	{
-		// Start at the scene root node if none was given
+		// Start at the root node if none was given
 		if(!InStartNodePtr.IsValid())
 		{
-			InStartNodePtr = SceneRootNodePtr;
+			InStartNodePtr = GetActorNode();
 		}
 
 		if(InStartNodePtr.IsValid())
@@ -6576,11 +6556,6 @@ FReply SSCSEditor::OnPromoteToBlueprintClicked()
 	return FReply::Handled();
 }
 
-const TArray<FSCSEditorTreeNodePtrType>& SSCSEditor::GetRootComponentNodes()
-{
-	return RootComponentNodes;
-}
-
 /** Returns the Actor context for which we are viewing/editing the SCS.  Can return null.  Should not be cached as it may change from frame to frame. */
 AActor* SSCSEditor::GetActorContext() const
 {
@@ -6637,36 +6612,28 @@ void SSCSEditor::OnFilterTextChanged(const FText& InFilterText)
 
 	bool bRootItemFilteredBackIn = false;
 	// iterate backwards so we select from the top down
-	for (int32 ComponentIndex = RootComponentNodes.Num() - 1; ComponentIndex >= 0; --ComponentIndex)
+	for (int32 ComponentIndex = RootNodes.Num() - 1; ComponentIndex >= 0; --ComponentIndex)
 	{
-		FSCSEditorTreeNodePtrType Component = RootComponentNodes[ComponentIndex];
+		FSCSEditorTreeNodePtrType Node = RootNodes[ComponentIndex];
 
-		const bool bWasFilteredOut = Component->IsFlaggedForFiltration();
-		bool bFilteredOut = RefreshFilteredState(Component, /*bRecursive =*/true);
+		const bool bWasFilteredOut = Node->IsFlaggedForFiltration();
+		bool bFilteredOut = RefreshFilteredState(Node, true);
 
 		if (!bFilteredOut)
 		{
 			if (!bIsFilterBlank)
 			{
-				NewSelection = OnFilterTextChanged_Inner::ExpandToFilteredChildren(this, Component);
+				NewSelection = OnFilterTextChanged_Inner::ExpandToFilteredChildren(this, Node);
 			}
 			bRootItemFilteredBackIn |= bWasFilteredOut;
 		}
-		else
-		{
-			FilteredRootNodes.Remove(Component);
-		}
 	}
 
 	if (NewSelection.IsValid() && !SCSTreeWidget->IsItemSelected(NewSelection))
 	{
 		SelectNode(NewSelection, /*IsCntrlDown =*/false);
 	}
-	
-	if (bRootItemFilteredBackIn)
-	{
-		RebuildFilteredRootList();
-	}
+		
 	UpdateTree(/*bRegenerateTreeNodes =*/false);
 }
 
diff --git a/Engine/Source/Editor/Kismet/Public/SSCSEditor.h b/Engine/Source/Editor/Kismet/Public/SSCSEditor.h
--- a/Engine/Source/Editor/Kismet/Public/SSCSEditor.h
+++ b/Engine/Source/Editor/Kismet/Public/SSCSEditor.h
@@ -30,7 +30,8 @@ class UPrimitiveComponent;
 struct EventData;
 
 // SCS tree node pointer type
-typedef TSharedPtr<class FSCSEditorTreeNode> FSCSEditorTreeNodePtrType;
+using FSCSEditorTreeNodePtrType = TSharedPtr<class FSCSEditorTreeNode>;
+using FSCSEditorActorNodePtrType = TSharedPtr<class FSCSEditorTreeNodeRootActor>;
 
 /**
  * FSCSEditorTreeNode
@@ -81,7 +82,7 @@ public:
 	 *
 	 * @return The component template that can be editable for actual class.
 	 */
-	virtual UActorComponent* GetEditableComponentTemplate(UBlueprint* ActualEditedBlueprint);
+	virtual UActorComponent* GetOrCreateEditableComponentTemplate(UBlueprint* ActualEditedBlueprint);
 	/**
 	 * Finds the component instance represented by this node contained within a given Actor instance.
 	 *
@@ -136,7 +137,7 @@ public:
 	 *
 	 * @param InChildNodePtr The node to add as a child node.
 	 */
-	void AddChild(FSCSEditorTreeNodePtrType InChildNodePtr);
+	virtual void AddChild(FSCSEditorTreeNodePtrType InChildNodePtr);
 
 	/**
 	 * Adds a child node for the given SCS node.
@@ -190,7 +191,7 @@ public:
 	 *
 	 * @param InChildNodePtr The child node to remove.
 	 */
-	void RemoveChild(FSCSEditorTreeNodePtrType InChildNodePtr);
+	virtual void RemoveChild(FSCSEditorTreeNodePtrType InChildNodePtr);
 
 	bool IsSceneComponent() const
 	{
@@ -289,7 +290,7 @@ public:
 	void SetRenameRequestedDelegate(FOnRenameRequested InRenameRequested) { RenameRequestedDelegate = InRenameRequested; }
 
 	/** Query that determines if this item should be filtered out or not */
-	bool IsFlaggedForFiltration() const 
+	virtual bool IsFlaggedForFiltration() const 
 	{
 		return ensureMsgf(FilterFlags != EFilteredState::Unknown, TEXT("Querying a bad filtration state.")) ? 
 			(FilterFlags & EFilteredState::FilteredInMask) == 0 : false; 
@@ -371,7 +372,7 @@ public:
 	 *
 	 * @param InComponentTemplate The component template represented by this object.
 	 */
-	FSCSEditorTreeNodeInstancedInheritedComponent(AActor* Owner, FName InComponentName);
+	FSCSEditorTreeNodeInstancedInheritedComponent(AActor* Owner, UActorComponent* InComponentTemplate);
 
 	// FSCSEditorTreeNode public interface
 	virtual bool IsNative() const override;
@@ -385,11 +386,10 @@ public:
 	//virtual FName GetVariableName() const override;
 	//virtual FString GetDisplayString() const override;
 	virtual FText GetDisplayName() const override;
-	virtual UActorComponent* GetEditableComponentTemplate(UBlueprint* ActualEditedBlueprint) override;
+	virtual UActorComponent* GetOrCreateEditableComponentTemplate(UBlueprint* ActualEditedBlueprint) override;
 	// End of FSCSEditorTreeNode public interface
 
 private:
-	FName InstancedComponentName;
 	TWeakObjectPtr<AActor> InstancedComponentOwnerPtr;
 };
 
@@ -417,7 +417,7 @@ public:
 	virtual FName GetVariableName() const override { return NAME_None; }
 	virtual FString GetDisplayString() const override;
 	virtual FText GetDisplayName() const override;
-	virtual UActorComponent* GetEditableComponentTemplate(UBlueprint* ActualEditedBlueprint) override;
+	virtual UActorComponent* GetOrCreateEditableComponentTemplate(UBlueprint* ActualEditedBlueprint) override;
 	virtual void OnCompleteRename(const FText& InNewName) override;
 	// End of FSCSEditorTreeNode public interface
 
@@ -465,7 +465,7 @@ public:
 	//virtual FString GetDisplayString() const override;
 	virtual FText GetDisplayName() const override;
 	virtual class USCS_Node* GetSCSNode() const override;
-	virtual UActorComponent* GetEditableComponentTemplate(UBlueprint* ActualEditedBlueprint) override;
+	virtual UActorComponent* GetOrCreateEditableComponentTemplate(UBlueprint* ActualEditedBlueprint) override;
 	virtual void OnCompleteRename(const FText& InNewName) override;
 	// End of FSCSEditorTreeNode public interface
 
@@ -485,7 +485,6 @@ private:
 	TWeakObjectPtr<class USCS_Node> SCSNodePtr;
 };
 
-
 class KISMET_API FSCSEditorTreeNodeRootActor : public FSCSEditorTreeNode
 {
 public:
@@ -496,15 +495,32 @@ public:
 	{
 	}
 
+	FSCSEditorTreeNodePtrType GetSceneRootNode() const;
+	void SetSceneRootNode(FSCSEditorTreeNodePtrType NewSceneRootNode);
+
+	/** Returns the set of root nodes */
+	const TArray<FSCSEditorTreeNodePtrType>& GetComponentNodes() const;
+
 	// FSCSEditorTreeNode public interface
 	virtual FName GetNodeID() const override;
 	virtual bool CanRename() const override { return bAllowRename; }
 	virtual void OnCompleteRename(const FText& InNewName) override;
+	virtual void AddChild(FSCSEditorTreeNodePtrType InChildNodePtr) override;
+	virtual void RemoveChild(FSCSEditorTreeNodePtrType InChildNodePtr) override;
 	// End of FSCSEditorTreeNode public interface
+protected:
+	using Super = FSCSEditorTreeNode;
 
 private:
+	
 	AActor* Actor;
 	bool bAllowRename;
+
+	FSCSEditorTreeNodePtrType SceneRootNodePtr;
+	/** Root set of components (contains the root scene component and any non-scene component nodes) */
+	TArray<FSCSEditorTreeNodePtrType> ComponentNodes;
+	FSCSEditorTreeNodePtrType SceneComponentSeparatorNodePtr;
+	FSCSEditorTreeNodePtrType NonSceneComponentSeparatorNodePtr;
 };
 
 class KISMET_API FSCSEditorTreeNodeSeparator : public FSCSEditorTreeNode
@@ -514,6 +530,8 @@ public:
 		: FSCSEditorTreeNode(FSCSEditorTreeNode::SeparatorNode)
 	{
 	}
+
+	virtual bool IsFlaggedForFiltration() const override { return false; }
 };
 
 //////////////////////////////////////////////////////////////////////////
@@ -804,6 +822,14 @@ public:
 	/** Returns true if editing is allowed */
 	bool IsEditingAllowed() const;
 
+	/** gets the actor root */
+	FSCSEditorActorNodePtrType GetActorNode() const;
+
+	/** get the root scene node */
+	FSCSEditorTreeNodePtrType GetSceneRootNode() const;
+
+	void SetSceneRootNode(FSCSEditorTreeNodePtrType NewSceneRootNode);
+
 	/** Adds a component to the SCS Table
 	   @param NewComponentClass				(In) The class to add
 	   @param Asset       					(In) Optional asset to assign to the component
@@ -821,10 +847,11 @@ public:
 
 	/** Adds a new component instance node to the component Table
 		@param NewInstanceComponent	(In) The component being added to the actor instance
+		@param InParentNodePtr (In) The node this component will be added to
 		@param Asset (In) Optional asset to assign to the component
 		@param bSetFocusToNewItem (In) Select the new item and activate the inline rename widget (default is true)
 		@return The reference of the newly created ActorComponent */
-	UActorComponent* AddNewNodeForInstancedComponent(UActorComponent* NewInstanceComponent, UObject* Asset, bool bSetFocusToNewItem = true);
+	UActorComponent* AddNewNodeForInstancedComponent(UActorComponent* NewInstanceComponent, FSCSEditorTreeNodePtrType InParentNodePtr, UObject* Asset, bool bSetFocusToNewItem = true);
 	
 	/** Returns true if the specified component is currently selected */
 	bool IsComponentSelected(const UPrimitiveComponent* PrimComponent) const;
@@ -959,10 +986,7 @@ public:
 
 	/** Provides access to the Blueprint context that's being edited */
 	class UBlueprint* GetBlueprint() const;
-
-	/** Returns the set of root nodes */
-	const TArray<FSCSEditorTreeNodePtrType>& GetRootComponentNodes();
-
+		
 	/** @return The current editor mode (editing live actors or editing blueprints) */
 	EComponentEditorMode::Type GetEditorMode() const { return EditorMode; }
 
@@ -976,6 +1000,10 @@ public:
 	FText GetFilterText() const;
 
 protected:
+	FSCSEditorTreeNodePtrType FindOrCreateParentForExistingComponent(UActorComponent* InActorComponent, FSCSEditorActorNodePtrType ActorRootNode);
+	FSCSEditorTreeNodePtrType FindParentForNewComponent(UActorComponent* NewComponent) const;
+	FSCSEditorTreeNodePtrType FindParentForNewNode(USCS_Node* NewNode) const;
+
 	FString GetSelectedClassText() const;
 
 	/** Add a component from the selection in the combo box */
@@ -1029,8 +1057,8 @@ protected:
 	/** Helper method to add a tree node for the given SCS node */
 	FSCSEditorTreeNodePtrType AddTreeNode(USCS_Node* InSCSNode, FSCSEditorTreeNodePtrType InParentNodePtr, const bool bIsInheritedSCS);
 
-	/** Helper method to add a tree node for the given scene component */
-	FSCSEditorTreeNodePtrType AddTreeNodeFromComponent(USceneComponent* InSceneComponent, FSCSEditorTreeNodePtrType InParentTreeNode = FSCSEditorTreeNodePtrType());
+	/** Helper method to add a tree node for the given actor component */
+	FSCSEditorTreeNodePtrType AddTreeNodeFromComponent(UActorComponent* InSceneComponent, FSCSEditorTreeNodePtrType InParentTreeNode = FSCSEditorTreeNodePtrType());
 	
 	/** Helper method to recursively add tree nodes for an already-attached instanced scene component hierarchy */
 	void AddInstancedTreeNodesRecursive(USceneComponent* Component, FSCSEditorTreeNodePtrType TreeNode, TSet<UActorComponent*>& ComponentsToAdd);
@@ -1085,10 +1113,7 @@ protected:
 
 	/** gets a root nodes of the tree */
 	const TArray<FSCSEditorTreeNodePtrType>& GetRootNodes() const;
-
-	/** Adds a root component tree node */
-	TSharedPtr<FSCSEditorTreeNode> AddRootComponentTreeNode(UActorComponent* ActorComp);
-
+				
 	/**
 	 * Creates a new C++ component from the specified class type
 	 * The user will be prompted to pick a new subclass name and code will be recompiled
@@ -1115,19 +1140,10 @@ protected:
 	 */
 	bool RefreshFilteredState(FSCSEditorTreeNodePtrType TreeNode, bool bRecursive);
 
-	/** 
-	 * Iterates the RootNodes list, and uses the cached filtered state to 
-	 * determine what items should be listed in the tree view. 
-	 */
-	void RebuildFilteredRootList();
-
 public:
 	/** Tree widget */
 	TSharedPtr<SSCSTreeType> SCSTreeWidget;
-
-	/** The node that represents the root component in the scene hierarchy */
-	FSCSEditorTreeNodePtrType SceneRootNodePtr;
-
+		
 	/** Command list for handling actions in the SSCSEditor */
 	TSharedPtr< FUICommandList > CommandList;
 
@@ -1163,16 +1179,7 @@ private:
 
 	/** Root set of tree */
 	TArray<FSCSEditorTreeNodePtrType> RootNodes;
-
-	/** Root set of components (contains the root scene component and any non-scene component nodes) */
-	TArray<FSCSEditorTreeNodePtrType> RootComponentNodes;
-
-	/** The list of nodes used for the UI (a filtered version of RootNodes) */
-	TArray<FSCSEditorTreeNodePtrType> FilteredRootNodes;
-
-	/* Root Tree Node (for scene components) */
-	TSharedPtr<FSCSEditorTreeNode> RootTreeNode;
-
+			
 	/* Root Tree Node*/
 	TSharedPtr<FExtender> ActorMenuExtender;
 
@@ -1181,10 +1188,7 @@ private:
 
 	/** Gate to prevent changing the selection while selection change is being broadcast. */
 	bool bUpdatingSelection;
-
-	/** true if we've added the separator between the scene and behavior components to the root nodes */
-	bool bHasAddedSceneAndBehaviorComponentSeparator;
-
+		
 	/** Controls whether or not to allow calls to UpdateTree() */
 	bool bAllowTreeUpdates;
 
diff --git a/Engine/Source/Editor/UnrealEd/Private/Kismet2/ComponentEditorUtils.cpp b/Engine/Source/Editor/UnrealEd/Private/Kismet2/ComponentEditorUtils.cpp
--- a/Engine/Source/Editor/UnrealEd/Private/Kismet2/ComponentEditorUtils.cpp
+++ b/Engine/Source/Editor/UnrealEd/Private/Kismet2/ComponentEditorUtils.cpp
@@ -956,6 +956,84 @@ void FComponentEditorUtils::FillComponentContextMenuOptions(FMenuBuilder& MenuBu
 	}
 }
 
+UActorComponent* FComponentEditorUtils::FindMatchingComponent(UActorComponent* ComponentInstance, const TInlineComponentArray<UActorComponent*>& ComponentList)
+{
+	if (ComponentInstance == nullptr)
+	{
+		return nullptr;
+	}
+
+	TInlineComponentArray<UActorComponent*> FoundComponents;
+	UActorComponent* LastFoundComponent = nullptr;
+	for (UActorComponent* Component : ComponentList)
+	{
+		// Early out on pointer match
+		if (ComponentInstance == Component)
+		{
+			return Component;
+		}
+
+		if (ComponentInstance->GetFName() == Component->GetFName())
+		{
+			FoundComponents.Add(Component);
+			LastFoundComponent = Component;
+		}
+	}
+
+	// No match or 1 match avoid sorting
+	if (FoundComponents.Num() <= 1)
+	{
+		return LastFoundComponent;
+	}
+
+	if (USceneComponent* CurrentSceneComponent = Cast<USceneComponent>(ComponentInstance))
+	{
+		// Sort by matching hierarchy
+		FoundComponents.Sort([&](const UActorComponent& ComponentA, const UActorComponent& ComponentB)
+		{
+			const USceneComponent* SceneComponentA = Cast<USceneComponent>(&ComponentA);
+			const USceneComponent* SceneComponentB = Cast<USceneComponent>(&ComponentB);
+			if (SceneComponentB == nullptr)
+			{
+				return true;
+			}
+			else if (SceneComponentA == nullptr)
+			{
+				return false;
+			}
+
+			const USceneComponent* AttachParentA = SceneComponentA->GetAttachParent();
+			const USceneComponent* AttachParentB = SceneComponentB->GetAttachParent();
+			const USceneComponent* CurrentParent = CurrentSceneComponent->GetAttachParent();
+			// No parents...
+			if (CurrentParent == nullptr)
+			{
+				return AttachParentA == nullptr;
+			}
+
+			bool MatchA = AttachParentA != nullptr && AttachParentA->GetFName() == CurrentParent->GetFName();
+			bool MatchB = AttachParentB != nullptr && AttachParentB->GetFName() == CurrentParent->GetFName();
+			while (MatchA && MatchB)
+			{
+				AttachParentA = AttachParentA->GetAttachParent();
+				AttachParentB = AttachParentB->GetAttachParent();
+				CurrentParent = CurrentParent->GetAttachParent();
+				if (CurrentParent == nullptr)
+				{
+					return AttachParentA == nullptr;
+				}
+
+				MatchA = AttachParentA != nullptr && AttachParentA->GetFName() == CurrentParent->GetFName();
+				MatchB = AttachParentB != nullptr && AttachParentB->GetFName() == CurrentParent->GetFName();
+			}
+
+			return MatchA;
+		});
+	}
+
+	return FoundComponents[0];
+}
+
 void FComponentEditorUtils::OnGoToComponentAssetInBrowser(UObject* Asset)
 {
 	TArray<UObject*> Objects;
diff --git a/Engine/Source/Editor/UnrealEd/Public/Kismet2/ComponentEditorUtils.h b/Engine/Source/Editor/UnrealEd/Public/Kismet2/ComponentEditorUtils.h
--- a/Engine/Source/Editor/UnrealEd/Public/Kismet2/ComponentEditorUtils.h
+++ b/Engine/Source/Editor/UnrealEd/Public/Kismet2/ComponentEditorUtils.h
@@ -244,6 +244,14 @@ public:
 	*/
 	static void FillComponentContextMenuOptions(FMenuBuilder& MenuBuilder, const TArray<UActorComponent*>& SelectedComponents);
 
+	/**
+	 * Tries to find a match for ComponentInstance in the ComponentList. First by name and then if multiple Components have a matching name try to match the SceneComponent hierarchy to find the best match.
+	 * @param ComponentInstance Component we are trying to match in the ComponentList
+	 * @param ComponentList List containing possible matches
+	 * @return Valid Component pointer if match was found. nullptr otherwise.
+	 */
+	static UActorComponent* FindMatchingComponent(UActorComponent* ComponentInstance, const TInlineComponentArray<UActorComponent*>& ComponentList);
+
 private:	
 	static USceneComponent* FindClosestParentInList(UActorComponent* ChildComponent, const TArray<UActorComponent*>& ComponentList);
 
diff --git a/Engine/Source/Runtime/Engine/Private/Blueprint.cpp b/Engine/Source/Runtime/Engine/Private/Blueprint.cpp
--- a/Engine/Source/Runtime/Engine/Private/Blueprint.cpp
+++ b/Engine/Source/Runtime/Engine/Private/Blueprint.cpp
@@ -24,6 +24,7 @@
 #include "Engine/SimpleConstructionScript.h"
 #include "Engine/SCS_Node.h"
 #include "Kismet2/BlueprintEditorUtils.h"
+#include "Kismet2/ComponentEditorUtils.h"
 #include "Kismet2/KismetEditorUtilities.h"
 #include "Kismet2/CompilerResultsLog.h"
 #include "Kismet2/StructureEditorUtils.h"
@@ -70,50 +71,21 @@ static void ConformNativeComponents(UBlueprint* Blueprint)
 			// native super-class)
 			TInlineComponentArray<UActorComponent*> NewNativeComponents;
 			NativeCDO->GetComponents(NewNativeComponents);
-
-			// utility lambda for finding named components in a supplied list
-			auto FindNamedComponentLambda = [](FName const ComponentName, TInlineComponentArray<UActorComponent*> const& ComponentList)->UActorComponent*
-			{
-				UActorComponent* FoundComponent = nullptr;
-				for (UActorComponent* Component : ComponentList)
-				{
-					if (Component->GetFName() == ComponentName)
-					{
-						FoundComponent = Component;
-						break;
-					}
-				}
-				return FoundComponent;
-			};
-
-			// utility lambda for finding matching components in the NewNativeComponents list
-			auto FindNativeComponentLambda = [&NewNativeComponents, &FindNamedComponentLambda](UActorComponent* BlueprintComponent)->UActorComponent*
-			{
-				UActorComponent* MatchingComponent = nullptr;
-				if (BlueprintComponent != nullptr)
-				{
-					FName const ComponentName = BlueprintComponent->GetFName();
-					MatchingComponent = FindNamedComponentLambda(ComponentName, NewNativeComponents);
-				}
-				return MatchingComponent;
-			};
-
+														   			
 			// loop through all components that this blueprint thinks come from its
 			// native super-class (last time it was saved)
 			for (UActorComponent* Component : OldNativeComponents)
 			{
-				// if we found this component also listed for the native class
-				if (UActorComponent* NativeComponent = FindNativeComponentLambda(Component))
+				if (UActorComponent* NativeComponent = FComponentEditorUtils::FindMatchingComponent(Component, NewNativeComponents))
 				{
-					USceneComponent* BlueprintSceneComponent = Cast<USceneComponent>(Component);
-					if (BlueprintSceneComponent == nullptr)
+					USceneComponent* SceneComponent = Cast<USceneComponent>(Component);
+					if (SceneComponent == nullptr)
 					{
 						// if this isn't a scene-component, then we don't care
 						// (we're looking to fixup scene-component parents)
 						continue;
 					}
-					UActorComponent* OldNativeParent = FindNativeComponentLambda(BlueprintSceneComponent->GetAttachParent());
-
+					USceneComponent* OldNativeParent = Cast<USceneComponent>(FComponentEditorUtils::FindMatchingComponent(SceneComponent->GetAttachParent(), NewNativeComponents));
 					USceneComponent* NativeSceneComponent = CastChecked<USceneComponent>(NativeComponent);
 					// if this native component has since been reparented, we need
 					// to make sure that this blueprint reflects that change
@@ -122,20 +94,21 @@ static void ConformNativeComponents(UBlueprint* Blueprint)
 						USceneComponent* NewParent = nullptr;
 						if (NativeSceneComponent->GetAttachParent() != nullptr)
 						{
-							NewParent = CastChecked<USceneComponent>(FindNamedComponentLambda(NativeSceneComponent->GetAttachParent()->GetFName(), OldNativeComponents));
+							NewParent = CastChecked<USceneComponent>(FComponentEditorUtils::FindMatchingComponent(NativeSceneComponent->GetAttachParent(), OldNativeComponents));
 						}
-						BlueprintSceneComponent->SetupAttachment(NewParent, BlueprintSceneComponent->GetAttachSocketName());
+						SceneComponent->SetupAttachment(NewParent, SceneComponent->GetAttachSocketName());
 					}
 				}
-				else // the component has been removed from the native class
+				else
 				{
-					// @TODO: I think we already handle removed native components elsewhere, so maybe we should error here?
-// 				BlueprintCDO->RemoveOwnedComponent(Component);
-// 
-// 				USimpleConstructionScript* BlueprintSCS = Blueprint->SimpleConstructionScript;
-// 				USCS_Node* ComponentNode = BlueprintSCS->CreateNode(Component, Component->GetFName());
-// 
-// 				BlueprintSCS->AddNode(ComponentNode);
+					// the component has been removed from the native class
+						// @TODO: I think we already handle removed native components elsewhere, so maybe we should error here?
+	// 				BlueprintCDO->RemoveOwnedComponent(Component);
+	// 
+	// 				USimpleConstructionScript* BlueprintSCS = Blueprint->SimpleConstructionScript;
+	// 				USCS_Node* ComponentNode = BlueprintSCS->CreateNode(Component, Component->GetFName());
+	// 
+	// 				BlueprintSCS->AddNode(ComponentNode);
 				}
 			}
 		}
