PR #5605: UserDefinedStruct's variables can now be marked SaveGame (Contributed by [REDACTED_PERSON])
Also fixed spelling on bDontEditOnInstance

diff --git a/Engine/Config/BaseEngine.ini b/Engine/Config/BaseEngine.ini
--- a/Engine/Config/BaseEngine.ini
+++ b/Engine/Config/BaseEngine.ini
@@ -654,6 +654,8 @@ bUseStreamingPause=false
 
 +ClassRedirects=(OldName="/Script/CoreUObject.MulticastDelegateProperty",NewName="/Script/CoreUObject.MulticastInlineDelegateProperty")
 
++PropertyRedirects=(OldName="StructVariableDescription.bDontEditoOnInstance",NewName="bDontEditOnInstance")
+
 [CoreUObject.Metadata]
 ; Note: UnrealHeaderTool should be rerun after making changes to MetadataRedirects to catch any keys specified in class headers
 +MetadataRedirects=(OldKey="K2Protected", NewKey="BlueprintProtected")
diff --git a/Engine/Source/Editor/Kismet/Private/UserDefinedStructureEditor.cpp b/Engine/Source/Editor/Kismet/Private/UserDefinedStructureEditor.cpp
--- a/Engine/Source/Editor/Kismet/Private/UserDefinedStructureEditor.cpp
+++ b/Engine/Source/Editor/Kismet/Private/UserDefinedStructureEditor.cpp
@@ -686,7 +686,7 @@ public:
 		{
 			if (const FStructVariableDescription* FieldDesc = StructureDetailsSP->FindStructureFieldByGuid(FieldGuid))
 			{
-				return !FieldDesc->bDontEditoOnInstance ? ECheckBoxState::Checked : ECheckBoxState::Unchecked;
+				return !FieldDesc->bDontEditOnInstance ? ECheckBoxState::Checked : ECheckBoxState::Unchecked;
 			}
 		}
 		return ECheckBoxState::Undetermined;
@@ -701,6 +701,28 @@ public:
 		}
 	}
 
+	ECheckBoxState OnGetSaveGameState() const
+	{
+		auto StructureDetailsSP = StructureDetails.Pin();
+		if (StructureDetailsSP.IsValid())
+		{
+			if (const FStructVariableDescription* FieldDesc = StructureDetailsSP->FindStructureFieldByGuid(FieldGuid))
+			{
+				return FieldDesc->bEnableSaveGame ? ECheckBoxState::Checked : ECheckBoxState::Unchecked;
+			}
+		}
+		return ECheckBoxState::Undetermined;
+	}
+
+	void OnSaveGameCommitted(ECheckBoxState InNewState)
+	{
+		auto StructureDetailsSP = StructureDetails.Pin();
+		if (StructureDetailsSP.IsValid() && (ECheckBoxState::Undetermined != InNewState))
+		{
+			FStructureEditorUtils::ChangeSaveGameEnabled(StructureDetailsSP->GetUserDefinedStruct(), FieldGuid, InNewState == ECheckBoxState::Checked);
+		}
+	}
+
 	// Multi-line text
 	EVisibility IsMultiLineTextOptionVisible() const
 	{
@@ -971,6 +993,21 @@ public:
 			.IsChecked(this, &FUserDefinedStructureFieldLayout::OnGetEditableOnBPInstanceState)
 		];
 
+		ChildrenBuilder.AddCustomRow(LOCTEXT("SaveGame", "Save Game"))
+		.NameContent()
+		[
+			SNew(STextBlock)
+			.Text(LOCTEXT("SaveGameText", "Save Game"))
+			.Font(IDetailLayoutBuilder::GetDetailFont())
+		]
+		.ValueContent()
+		[
+			SNew(SCheckBox)
+			.ToolTipText(LOCTEXT("SaveGame_Tooltip", "Should variable be serialized for saved games"))
+			.OnCheckStateChanged(this, &FUserDefinedStructureFieldLayout::OnSaveGameCommitted)
+			.IsChecked(this, &FUserDefinedStructureFieldLayout::OnGetSaveGameState)
+		];
+
 		ChildrenBuilder.AddCustomRow(LOCTEXT("MultiLineText", "Multi-line Text"))
 		.NameContent()
 		[
diff --git a/Engine/Source/Editor/KismetCompiler/Private/UserDefinedStructureCompilerUtils.cpp b/Engine/Source/Editor/KismetCompiler/Private/UserDefinedStructureCompilerUtils.cpp
--- a/Engine/Source/Editor/KismetCompiler/Private/UserDefinedStructureCompilerUtils.cpp
+++ b/Engine/Source/Editor/KismetCompiler/Private/UserDefinedStructureCompilerUtils.cpp
@@ -253,10 +253,14 @@ struct FUserDefinedStructureCompilerInner
 			}
 			
 			VarProperty->SetPropertyFlags(CPF_Edit | CPF_BlueprintVisible);
-			if (VarDesc.bDontEditoOnInstance)
+			if (VarDesc.bDontEditOnInstance)
 			{
 				VarProperty->SetPropertyFlags(CPF_DisableEditOnInstance);
 			}
+			if (VarDesc.bEnableSaveGame)
+			{
+				VarProperty->SetPropertyFlags(CPF_SaveGame);
+			}
 			if (VarDesc.bEnableMultiLineText)
 			{
 				VarProperty->SetMetaData("MultiLine", TEXT("true"));
diff --git a/Engine/Source/Editor/UnrealEd/Classes/UserDefinedStructure/UserDefinedStructEditorData.h b/Engine/Source/Editor/UnrealEd/Classes/UserDefinedStructure/UserDefinedStructEditorData.h
--- a/Engine/Source/Editor/UnrealEd/Classes/UserDefinedStructure/UserDefinedStructEditorData.h
+++ b/Engine/Source/Editor/UnrealEd/Classes/UserDefinedStructure/UserDefinedStructEditorData.h
@@ -64,7 +64,10 @@ struct FStructVariableDescription
 	uint8 bInvalidMember:1;
 
 	UPROPERTY()
-	uint8 bDontEditoOnInstance:1;
+	uint8 bDontEditOnInstance:1;
+
+	UPROPERTY()
+	uint8 bEnableSaveGame : 1;
 
 	UPROPERTY()
 	uint8 bEnableMultiLineText:1;
@@ -92,7 +95,8 @@ struct FStructVariableDescription
 		, bIsSet_DEPRECATED(false)
 		, bIsMap_DEPRECATED(false)
 		, bInvalidMember(false)
-		, bDontEditoOnInstance(false)
+		, bDontEditOnInstance(false)
+		, bEnableSaveGame(false)
 		, bEnableMultiLineText(false)
 		, bEnable3dWidget(false)
 	{ }
diff --git a/Engine/Source/Editor/UnrealEd/Private/Kismet2/StructureEditorUtils.cpp b/Engine/Source/Editor/UnrealEd/Private/Kismet2/StructureEditorUtils.cpp
--- a/Engine/Source/Editor/UnrealEd/Private/Kismet2/StructureEditorUtils.cpp
+++ b/Engine/Source/Editor/UnrealEd/Private/Kismet2/StructureEditorUtils.cpp
@@ -312,8 +312,6 @@ bool FStructureEditorUtils::AddVariable(UUserDefinedStruct* Struct, const FEdGra
 		NewVar.FriendlyName = DisplayName;
 		NewVar.SetPinType(VarType);
 		NewVar.VarGuid = Guid;
-		NewVar.bDontEditoOnInstance = false;
-		NewVar.bInvalidMember = false;
 		GetVarDesc(Struct).Add(NewVar);
 
 		OnStructureChanged(Struct, EStructureEditorChangeInfo::AddedVariable);
@@ -697,13 +695,29 @@ bool FStructureEditorUtils::ChangeEditableOnBPInstance(UUserDefinedStruct* Struc
 {
 	const UEdGraphSchema_K2* K2Schema = GetDefault<UEdGraphSchema_K2>();
 	FStructVariableDescription* VarDesc = GetVarDescByGuid(Struct, VarGuid);
-	const bool bNewDontEditoOnInstance = !bInIsEditable;
-	if (VarDesc && (bNewDontEditoOnInstance != VarDesc->bDontEditoOnInstance))
+	const bool bNewDontEditOnInstance = !bInIsEditable;
+	if (VarDesc && (bNewDontEditOnInstance != VarDesc->bDontEditOnInstance))
 	{
 		const FScopedTransaction Transaction(LOCTEXT("ChangeVariableOnBPInstance", "Change variable editable on BP instance"));
 		ModifyStructData(Struct);
 
-		VarDesc->bDontEditoOnInstance = bNewDontEditoOnInstance;
+		VarDesc->bDontEditOnInstance = bNewDontEditOnInstance;
+		OnStructureChanged(Struct);
+		return true;
+	}
+	return false;
+}
+
+bool FStructureEditorUtils::ChangeSaveGameEnabled(UUserDefinedStruct* Struct, FGuid VarGuid, bool bInSaveGame)
+{
+	const UEdGraphSchema_K2* K2Schema = GetDefault<UEdGraphSchema_K2>();
+	FStructVariableDescription* VarDesc = GetVarDescByGuid(Struct, VarGuid);
+	if (VarDesc && (bInSaveGame != VarDesc->bEnableSaveGame))
+	{
+		const FScopedTransaction Transaction(LOCTEXT("ChangeSaveGameOnVariable", "Change variable SaveGame flag"));
+		ModifyStructData(Struct);
+
+		VarDesc->bEnableSaveGame = bInSaveGame;
 		OnStructureChanged(Struct);
 		return true;
 	}
diff --git a/Engine/Source/Editor/UnrealEd/Public/Kismet2/StructureEditorUtils.h b/Engine/Source/Editor/UnrealEd/Public/Kismet2/StructureEditorUtils.h
--- a/Engine/Source/Editor/UnrealEd/Public/Kismet2/StructureEditorUtils.h
+++ b/Engine/Source/Editor/UnrealEd/Public/Kismet2/StructureEditorUtils.h
@@ -99,6 +99,8 @@ public:
 
 	static bool ChangeEditableOnBPInstance(UUserDefinedStruct* Struct, FGuid VarGuid, bool bInIsEditable);
 
+	static bool ChangeSaveGameEnabled(UUserDefinedStruct* Struct, FGuid VarGuid, bool bInSaveGame);
+
 	enum EMoveDirection
 	{
 		MD_Up,
