Hook up referencer name for editor FGCObjects that are likely to hold world references

diff --git a/Engine/Plugins/Experimental/PythonScriptPlugin/Source/PythonScriptPlugin/Private/PyReferenceCollector.cpp b/Engine/Plugins/Experimental/PythonScriptPlugin/Source/PythonScriptPlugin/Private/PyReferenceCollector.cpp
--- a/Engine/Plugins/Experimental/PythonScriptPlugin/Source/PythonScriptPlugin/Private/PyReferenceCollector.cpp
+++ b/Engine/Plugins/Experimental/PythonScriptPlugin/Source/PythonScriptPlugin/Private/PyReferenceCollector.cpp
@@ -37,6 +37,11 @@ void FPyReferenceCollector::AddReferencedObjects(FReferenceCollector& InCollecto
 	}
 }
 
+FString FPyReferenceCollector::GetReferencerName() const
+{
+	return TEXT("FPyReferenceCollector");
+}
+
 void FPyReferenceCollector::PurgeUnrealObjectReferences(const UObject* InObject, const bool bIncludeInnerObjects)
 {
 	PurgeUnrealObjectReferences(TArrayView<const UObject*>(&InObject, 1), bIncludeInnerObjects);
diff --git a/Engine/Plugins/Experimental/PythonScriptPlugin/Source/PythonScriptPlugin/Private/PyReferenceCollector.h b/Engine/Plugins/Experimental/PythonScriptPlugin/Source/PythonScriptPlugin/Private/PyReferenceCollector.h
--- a/Engine/Plugins/Experimental/PythonScriptPlugin/Source/PythonScriptPlugin/Private/PyReferenceCollector.h
+++ b/Engine/Plugins/Experimental/PythonScriptPlugin/Source/PythonScriptPlugin/Private/PyReferenceCollector.h
@@ -46,6 +46,7 @@ public:
 
 	//~ FGCObject interface
 	virtual void AddReferencedObjects(FReferenceCollector& InCollector) override;
+	virtual FString GetReferencerName() const override;
 
 	/** Utility function to ARO a delegate instance */
 	static void AddReferencedObjectsFromDelegate(FReferenceCollector& InCollector, FScriptDelegate& InDelegate);
diff --git a/Engine/Plugins/FX/Niagara/Source/Niagara/Private/NiagaraWorldManager.cpp b/Engine/Plugins/FX/Niagara/Source/Niagara/Private/NiagaraWorldManager.cpp
--- a/Engine/Plugins/FX/Niagara/Source/Niagara/Private/NiagaraWorldManager.cpp
+++ b/Engine/Plugins/FX/Niagara/Source/Niagara/Private/NiagaraWorldManager.cpp
@@ -76,6 +76,11 @@ void FNiagaraWorldManager::AddReferencedObjects(FReferenceCollector& Collector)
 	Collector.AddReferencedObjects(ParameterCollections);
 }
 
+FString FNiagaraWorldManager::GetReferencerName() const
+{
+	return TEXT("FNiagaraWorldManager");
+}
+
 UNiagaraParameterCollectionInstance* FNiagaraWorldManager::GetParameterCollection(UNiagaraParameterCollection* Collection)
 {
 	if (!Collection)
diff --git a/Engine/Plugins/FX/Niagara/Source/Niagara/Public/NiagaraWorldManager.h b/Engine/Plugins/FX/Niagara/Source/Niagara/Public/NiagaraWorldManager.h
--- a/Engine/Plugins/FX/Niagara/Source/Niagara/Public/NiagaraWorldManager.h
+++ b/Engine/Plugins/FX/Niagara/Source/Niagara/Public/NiagaraWorldManager.h
@@ -82,7 +82,8 @@ public:
 	static FNiagaraWorldManager* Get(UWorld* World);
 
 	//~ GCObject Interface
-	void AddReferencedObjects(FReferenceCollector& Collector);
+	virtual void AddReferencedObjects(FReferenceCollector& Collector) override;
+	virtual FString GetReferencerName() const override;
 	//~ GCObject Interface
 	
 	UNiagaraParameterCollectionInstance* GetParameterCollection(UNiagaraParameterCollection* Collection);
diff --git a/Engine/Source/Editor/BlueprintGraph/Private/BlueprintActionDatabase.cpp b/Engine/Source/Editor/BlueprintGraph/Private/BlueprintActionDatabase.cpp
--- a/Engine/Source/Editor/BlueprintGraph/Private/BlueprintActionDatabase.cpp
+++ b/Engine/Source/Editor/BlueprintGraph/Private/BlueprintActionDatabase.cpp
@@ -1098,6 +1098,11 @@ void FBlueprintActionDatabase::AddReferencedObjects(FReferenceCollector& Collect
 	}
 }
 
+FString FBlueprintActionDatabase::GetReferencerName() const
+{
+	return TEXT("FBlueprintActionDatabase");
+}
+
 int32 GBlueprintDatabasePrimingMaxPerFrame = 16;
 static FAutoConsoleVariableRef CVarBlueprintDatabasePrimingMaxPerFrame(
 	TEXT("bp.DatabasePrimingMaxPerFrame"),
diff --git a/Engine/Source/Editor/BlueprintGraph/Private/BlueprintNodeTemplateCache.cpp b/Engine/Source/Editor/BlueprintGraph/Private/BlueprintNodeTemplateCache.cpp
--- a/Engine/Source/Editor/BlueprintGraph/Private/BlueprintNodeTemplateCache.cpp
+++ b/Engine/Source/Editor/BlueprintGraph/Private/BlueprintNodeTemplateCache.cpp
@@ -477,6 +477,11 @@ void FBlueprintNodeTemplateCache::AddReferencedObjects(FReferenceCollector& Coll
 	Collector.AddReferencedObjects(TemplateOuters);
 }
 
+FString FBlueprintNodeTemplateCache::GetReferencerName() const
+{
+	return TEXT("FBlueprintNodeTemplateCache");
+}
+
 //------------------------------------------------------------------------------
 bool FBlueprintNodeTemplateCache::CacheBlueprintOuter(UBlueprint* Blueprint)
 {
diff --git a/Engine/Source/Editor/BlueprintGraph/Public/BlueprintActionDatabase.h b/Engine/Source/Editor/BlueprintGraph/Public/BlueprintActionDatabase.h
--- a/Engine/Source/Editor/BlueprintGraph/Public/BlueprintActionDatabase.h
+++ b/Engine/Source/Editor/BlueprintGraph/Public/BlueprintActionDatabase.h
@@ -52,6 +52,7 @@ public:
 
 	// FGCObject interface
 	virtual void AddReferencedObjects(FReferenceCollector& Collector) override;
+	virtual FString GetReferencerName() const override;
 	// End FGCObject interface
 
 	/**
diff --git a/Engine/Source/Editor/BlueprintGraph/Public/BlueprintNodeTemplateCache.h b/Engine/Source/Editor/BlueprintGraph/Public/BlueprintNodeTemplateCache.h
--- a/Engine/Source/Editor/BlueprintGraph/Public/BlueprintNodeTemplateCache.h
+++ b/Engine/Source/Editor/BlueprintGraph/Public/BlueprintNodeTemplateCache.h
@@ -82,6 +82,7 @@ public:
 
 	// FGCObject interface
 	virtual void AddReferencedObjects(FReferenceCollector& Collector) override;
+	virtual FString GetReferencerName() const override;
 	// End FGCObject interface
 
 private:
diff --git a/Engine/Source/Editor/Cascade/Private/Cascade.cpp b/Engine/Source/Editor/Cascade/Private/Cascade.cpp
--- a/Engine/Source/Editor/Cascade/Private/Cascade.cpp
+++ b/Engine/Source/Editor/Cascade/Private/Cascade.cpp
@@ -1530,6 +1530,11 @@ void FCascade::AddReferencedObjects(FReferenceCollector& Collector)
 	}
 }
 
+FString FCascade::GetReferencerName() const
+{
+	return TEXT("FCascade");
+}
+
 void FCascade::Tick(float DeltaTime)
 {
 	// This is a bit of a hack. In order to not tick all open Cascade editors (which tick through engine tick) even when not visible,
diff --git a/Engine/Source/Editor/Cascade/Private/Cascade.h b/Engine/Source/Editor/Cascade/Private/Cascade.h
--- a/Engine/Source/Editor/Cascade/Private/Cascade.h
+++ b/Engine/Source/Editor/Cascade/Private/Cascade.h
@@ -197,6 +197,7 @@ public:
 
 	/** FGCObject interface */
 	virtual void AddReferencedObjects(FReferenceCollector& Collector) override;
+	virtual FString GetReferencerName() const override;
 	
 	/** FTickableEditorObject interface */
 	virtual void Tick(float DeltaTime) override;
diff --git a/Engine/Source/Editor/GeometryMode/Private/EditorGeometry.cpp b/Engine/Source/Editor/GeometryMode/Private/EditorGeometry.cpp
--- a/Engine/Source/Editor/GeometryMode/Private/EditorGeometry.cpp
+++ b/Engine/Source/Editor/GeometryMode/Private/EditorGeometry.cpp
@@ -744,3 +744,8 @@ void FGeomObject::AddReferencedObjects( FReferenceCollector& Collector )
 {
 	Collector.AddReferencedObject( ActualBrush );
 }
+
+FString FGeomObject::GetReferencerName() const
+{
+	return TEXT("FGeomObject");
+}
diff --git a/Engine/Source/Editor/GeometryMode/Public/EditorGeometry.h b/Engine/Source/Editor/GeometryMode/Public/EditorGeometry.h
--- a/Engine/Source/Editor/GeometryMode/Public/EditorGeometry.h
+++ b/Engine/Source/Editor/GeometryMode/Public/EditorGeometry.h
@@ -304,6 +304,7 @@ public:
 	void ForceLastSelectionIndex( int32 InLastSelectionIndex ) { LastSelectionIndex = InLastSelectionIndex; }
 
 	virtual void AddReferencedObjects( FReferenceCollector& Collector ) override;
+	virtual FString GetReferencerName() const override;
 	
 	/**
 	 * Set the pivot position based on the 'highest' selected object (vertex/edge/poly) in the given selection array
diff --git a/Engine/Source/Editor/GraphEditor/Private/SGraphActionMenu.cpp b/Engine/Source/Editor/GraphEditor/Private/SGraphActionMenu.cpp
--- a/Engine/Source/Editor/GraphEditor/Private/SGraphActionMenu.cpp
+++ b/Engine/Source/Editor/GraphEditor/Private/SGraphActionMenu.cpp
@@ -1357,6 +1357,11 @@ void SGraphActionMenu::AddReferencedObjects( FReferenceCollector& Collector )
 	}
 }
 
+FString SGraphActionMenu::GetReferencerName() const
+{
+	return TEXT("SGraphActionMenu");
+}
+
 bool SGraphActionMenu::HandleSelection( TSharedPtr< FGraphActionNode > &InSelectedItem, ESelectInfo::Type InSelectionType )
 {
 	bool bResult = false;
diff --git a/Engine/Source/Editor/GraphEditor/Private/SGraphPanel.cpp b/Engine/Source/Editor/GraphEditor/Private/SGraphPanel.cpp
--- a/Engine/Source/Editor/GraphEditor/Private/SGraphPanel.cpp
+++ b/Engine/Source/Editor/GraphEditor/Private/SGraphPanel.cpp
@@ -1710,6 +1710,11 @@ void SGraphPanel::AddReferencedObjects(FReferenceCollector& Collector)
 	Collector.AddReferencedObject( GraphObjToDiff );
 }
 
+FString SGraphPanel::GetReferencerName() const
+{
+	return TEXT("SGraphPanel");
+}
+
 EActiveTimerReturnType SGraphPanel::InvalidatePerTick(double InCurrentTime, float InDeltaTime)
 {
 	// Invalidate the layout so it will redraw.
diff --git a/Engine/Source/Editor/GraphEditor/Public/SGraphActionMenu.h b/Engine/Source/Editor/GraphEditor/Public/SGraphActionMenu.h
--- a/Engine/Source/Editor/GraphEditor/Public/SGraphActionMenu.h
+++ b/Engine/Source/Editor/GraphEditor/Public/SGraphActionMenu.h
@@ -147,6 +147,7 @@ public:
 
 	// FGCObject override
 	virtual void AddReferencedObjects( FReferenceCollector& Collector ) override;
+	virtual FString GetReferencerName() const override;
 
 	/**
 	 * Refreshes the actions that this widget should display
diff --git a/Engine/Source/Editor/GraphEditor/Public/SGraphPanel.h b/Engine/Source/Editor/GraphEditor/Public/SGraphPanel.h
--- a/Engine/Source/Editor/GraphEditor/Public/SGraphPanel.h
+++ b/Engine/Source/Editor/GraphEditor/Public/SGraphPanel.h
@@ -119,6 +119,7 @@ public:
 
 	// FGCObject interface.
 	virtual void AddReferencedObjects( FReferenceCollector& Collector ) override;
+	virtual FString GetReferencerName() const override;
 	// End of FGCObject interface.
 
 	void ArrangeChildrenForContextMenuSummon(const FGeometry& AllottedGeometry, FArrangedChildren& ArrangedChildren) const;
diff --git a/Engine/Source/Editor/Kismet/Private/BlueprintCompilationManager.cpp b/Engine/Source/Editor/Kismet/Private/BlueprintCompilationManager.cpp
--- a/Engine/Source/Editor/Kismet/Private/BlueprintCompilationManager.cpp
+++ b/Engine/Source/Editor/Kismet/Private/BlueprintCompilationManager.cpp
@@ -88,6 +88,7 @@ struct FBlueprintCompilationManagerImpl : public FGCObject
 
 	// FGCObject:
 	virtual void AddReferencedObjects(FReferenceCollector& Collector);
+	virtual FString GetReferencerName() const override;
 
 	void QueueForCompilation(const FBPCompileRequest& CompileJob);
 	void CompileSynchronouslyImpl(const FBPCompileRequest& Request);
@@ -157,6 +158,11 @@ void FBlueprintCompilationManagerImpl::AddReferencedObjects(FReferenceCollector&
 	Collector.AddReferencedObjects(OldCDOs);
 }
 
+FString FBlueprintCompilationManagerImpl::GetReferencerName() const
+{
+	return TEXT("FBlueprintCompilationManagerImpl");
+}
+
 void FBlueprintCompilationManagerImpl::QueueForCompilation(const FBPCompileRequest& CompileJob)
 {
 	if(!CompileJob.BPToCompile->bQueuedForCompilation)
diff --git a/Engine/Source/Editor/Kismet/Private/BlueprintEditor.cpp b/Engine/Source/Editor/Kismet/Private/BlueprintEditor.cpp
--- a/Engine/Source/Editor/Kismet/Private/BlueprintEditor.cpp
+++ b/Engine/Source/Editor/Kismet/Private/BlueprintEditor.cpp
@@ -3590,6 +3590,11 @@ void FBlueprintEditor::AddReferencedObjects( FReferenceCollector& Collector )
 	}
 }
 
+FString FBlueprintEditor::GetReferencerName() const
+{
+	return TEXT("FBlueprintEditor");
+}
+
 bool FBlueprintEditor::IsNodeTitleVisible(const UEdGraphNode* Node, bool bRequestRename)
 {
 	TSharedPtr<SGraphEditor> GraphEditor;
diff --git a/Engine/Source/Editor/Kismet/Public/BlueprintEditor.h b/Engine/Source/Editor/Kismet/Public/BlueprintEditor.h
--- a/Engine/Source/Editor/Kismet/Public/BlueprintEditor.h
+++ b/Engine/Source/Editor/Kismet/Public/BlueprintEditor.h
@@ -182,6 +182,7 @@ public:
 
 	//~ Begin FGCObject Interface
 	virtual void AddReferencedObjects(FReferenceCollector& Collector) override;
+	virtual FString GetReferencerName() const override;
 	//~ End FGCObject Interface
 
 	//~ Begin IBlueprintEditor Interface
diff --git a/Engine/Source/Editor/LevelEditor/Private/SSurfaceProperties.cpp b/Engine/Source/Editor/LevelEditor/Private/SSurfaceProperties.cpp
--- a/Engine/Source/Editor/LevelEditor/Private/SSurfaceProperties.cpp
+++ b/Engine/Source/Editor/LevelEditor/Private/SSurfaceProperties.cpp
@@ -720,6 +720,11 @@ void SSurfaceProperties::AddReferencedObjects(FReferenceCollector& Collector)
 	}
 }
 
+FString SSurfaceProperties::GetReferencerName() const
+{
+	return TEXT("SSurfaceProperties");
+}
+
 TOptional<float> SSurfaceProperties::GetLightmapResolutionValue() const
 {
 	float LightMapScale = 0.0f;
diff --git a/Engine/Source/Editor/LevelEditor/Private/SSurfaceProperties.h b/Engine/Source/Editor/LevelEditor/Private/SSurfaceProperties.h
--- a/Engine/Source/Editor/LevelEditor/Private/SSurfaceProperties.h
+++ b/Engine/Source/Editor/LevelEditor/Private/SSurfaceProperties.h
@@ -54,6 +54,7 @@ private:
 	
 	// FGCObject Interface
 	virtual void AddReferencedObjects(FReferenceCollector& Collector) override;
+	virtual FString GetReferencerName() const override;
 
 	/** 
 	 * Creates Pan texture controls
diff --git a/Engine/Source/Editor/Sequencer/Private/Sequencer.cpp b/Engine/Source/Editor/Sequencer/Private/Sequencer.cpp
--- a/Engine/Source/Editor/Sequencer/Private/Sequencer.cpp
+++ b/Engine/Source/Editor/Sequencer/Private/Sequencer.cpp
@@ -3071,6 +3071,10 @@ void FSequencer::AddReferencedObjects( FReferenceCollector& Collector )
 	FMovieSceneRootEvaluationTemplateInstance::StaticStruct()->SerializeBin(Collector.GetVerySlowReferenceCollectorArchive(), &RootTemplateInstance);
 }
 
+FString FSequencer::GetReferencerName() const
+{
+	return TEXT("FSequencer");
+}
 
 void FSequencer::ResetPerMovieSceneData()
 {
@@ -9989,5 +9993,4 @@ void FSequencer::RecompileDirtyDirectors()
 	}
 }
 
-
 #undef LOCTEXT_NAMESPACE
diff --git a/Engine/Source/Editor/Sequencer/Private/Sequencer.h b/Engine/Source/Editor/Sequencer/Private/Sequencer.h
--- a/Engine/Source/Editor/Sequencer/Private/Sequencer.h
+++ b/Engine/Source/Editor/Sequencer/Private/Sequencer.h
@@ -630,8 +630,8 @@ public:
 public:
 
 	//~ FGCObject Interface
-
 	virtual void AddReferencedObjects(FReferenceCollector& Collector) override;
+	virtual FString GetReferencerName() const override;
 
 public:
 
diff --git a/Engine/Source/Editor/UnrealEd/Private/EditorViewportClient.cpp b/Engine/Source/Editor/UnrealEd/Private/EditorViewportClient.cpp
--- a/Engine/Source/Editor/UnrealEd/Private/EditorViewportClient.cpp
+++ b/Engine/Source/Editor/UnrealEd/Private/EditorViewportClient.cpp
@@ -3370,6 +3370,11 @@ void FEditorViewportClient::AddReferencedObjects( FReferenceCollector& Collector
 	}
 }
 
+FString FEditorViewportClient::GetReferencerName() const
+{
+	return TEXT("FEditorViewportClient");
+}
+
 void FEditorViewportClient::ProcessClick(class FSceneView& View, class HHitProxy* HitProxy, FKey Key, EInputEvent Event, uint32 HitX, uint32 HitY)
 {
 	const FViewportClick Click(&View, this, Key, Event, HitX, HitY);
diff --git a/Engine/Source/Editor/UnrealEd/Private/Toolkits/AssetEditorManager.cpp b/Engine/Source/Editor/UnrealEd/Private/Toolkits/AssetEditorManager.cpp
--- a/Engine/Source/Editor/UnrealEd/Private/Toolkits/AssetEditorManager.cpp
+++ b/Engine/Source/Editor/UnrealEd/Private/Toolkits/AssetEditorManager.cpp
@@ -137,6 +137,10 @@ void FAssetEditorManager::AddReferencedObjects( FReferenceCollector& Collector )
 	}
 }
 
+FString FAssetEditorManager::GetReferencerName() const
+{
+	return TEXT("FAssetEditorManager");
+}
 
 IAssetEditorInstance* FAssetEditorManager::FindEditorForAsset(UObject* Asset, bool bFocusIfOpen)
 {
diff --git a/Engine/Source/Editor/UnrealEd/Private/Toolkits/AssetEditorToolkit.cpp b/Engine/Source/Editor/UnrealEd/Private/Toolkits/AssetEditorToolkit.cpp
--- a/Engine/Source/Editor/UnrealEd/Private/Toolkits/AssetEditorToolkit.cpp
+++ b/Engine/Source/Editor/UnrealEd/Private/Toolkits/AssetEditorToolkit.cpp
@@ -1131,6 +1131,10 @@ void FAssetEditorToolkit::FGCEditingObjects::AddReferencedObjects(FReferenceColl
 	}
 }
 
+FString FAssetEditorToolkit::FGCEditingObjects::GetReferencerName() const
+{
+	return TEXT("FAssetEditorToolkit::FGCEditorObjects");
+}
 
 TSharedPtr<FExtender> FExtensibilityManager::GetAllExtenders()
 {
diff --git a/Engine/Source/Editor/UnrealEd/Public/EditorViewportClient.h b/Engine/Source/Editor/UnrealEd/Public/EditorViewportClient.h
--- a/Engine/Source/Editor/UnrealEd/Public/EditorViewportClient.h
+++ b/Engine/Source/Editor/UnrealEd/Public/EditorViewportClient.h
@@ -495,6 +495,7 @@ public:
 
 	/** FGCObject interface */
 	virtual void AddReferencedObjects( FReferenceCollector& Collector ) override;
+	virtual FString GetReferencerName() const override;
 
 	/**
 	 * Called when the user clicks in the viewport
diff --git a/Engine/Source/Editor/UnrealEd/Public/Toolkits/AssetEditorManager.h b/Engine/Source/Editor/UnrealEd/Public/Toolkits/AssetEditorManager.h
--- a/Engine/Source/Editor/UnrealEd/Public/Toolkits/AssetEditorManager.h
+++ b/Engine/Source/Editor/UnrealEd/Public/Toolkits/AssetEditorManager.h
@@ -117,6 +117,7 @@ public:
 
 	// FGCObject interface
 	virtual void AddReferencedObjects(FReferenceCollector& Collector) override;
+	virtual FString GetReferencerName() const override;
 
 	/** Close all open asset editors */
 	bool CloseAllAssetEditors();
diff --git a/Engine/Source/Editor/UnrealEd/Public/Toolkits/AssetEditorToolkit.h b/Engine/Source/Editor/UnrealEd/Public/Toolkits/AssetEditorToolkit.h
--- a/Engine/Source/Editor/UnrealEd/Public/Toolkits/AssetEditorToolkit.h
+++ b/Engine/Source/Editor/UnrealEd/Public/Toolkits/AssetEditorToolkit.h
@@ -338,6 +338,7 @@ private:
 
 		/** FGCObject interface */
 		virtual void AddReferencedObjects(FReferenceCollector& Collector) override;
+		virtual FString GetReferencerName() const override;
 
 	private:
 		FAssetEditorToolkit& OwnerToolkit;
