Allow Live Coding to work with KismetWidgets module
#rnx

diff --git a/Engine/Source/Editor/Kismet/Private/BlueprintEditorModule.cpp b/Engine/Source/Editor/Kismet/Private/BlueprintEditorModule.cpp
--- a/Engine/Source/Editor/Kismet/Private/BlueprintEditorModule.cpp
+++ b/Engine/Source/Editor/Kismet/Private/BlueprintEditorModule.cpp
@@ -29,6 +29,7 @@
 #include "BlueprintGraphPanelPinFactory.h"
 #include "WatchPointViewer.h"
 #include "KismetCompiler.h"
+#include "KismetWidgets.h"
 
 #define LOCTEXT_NAMESPACE "BlueprintEditor"
 
@@ -181,6 +182,8 @@ void FBlueprintEditorModule::StartupModule()
 		auto& MenuExtenders = LevelEditorModule.GetAllLevelViewportContextMenuExtenders();
 		MenuExtenders.Add(LevelViewportContextMenuBlueprintExtender);
 		LevelViewportContextMenuBlueprintExtenderDelegateHandle = MenuExtenders.Last().GetHandle();
+
+		FModuleManager::Get().LoadModuleChecked<FKismetWidgetsModule>("KismetWidgets");
 	}
 
 	FMessageLogModule& MessageLogModule = FModuleManager::LoadModuleChecked<FMessageLogModule>("MessageLog");
