Changing the parameter container type or pass by ref can now be undone/redone

diff --git a/Engine/Source/Editor/Kismet/Private/BlueprintDetailsCustomization.cpp b/Engine/Source/Editor/Kismet/Private/BlueprintDetailsCustomization.cpp
--- a/Engine/Source/Editor/Kismet/Private/BlueprintDetailsCustomization.cpp
+++ b/Engine/Source/Editor/Kismet/Private/BlueprintDetailsCustomization.cpp
@@ -3006,8 +3006,10 @@ ECheckBoxState FBlueprintGraphArgumentLayout::IsRefChecked() const
 
 void FBlueprintGraphArgumentLayout::OnRefCheckStateChanged(ECheckBoxState InState)
 {
+	const FScopedTransaction Transaction(LOCTEXT("ChangeByRef", "Change Pass By Reference"));
+
 	FEdGraphPinType PinType = OnGetPinInfo();
-	PinType.bIsReference = (InState == ECheckBoxState::Checked)? true : false;
+	PinType.bIsReference = (InState == ECheckBoxState::Checked) ? true : false;
 	
 	PinInfoChanged(PinType);
 }
@@ -3038,6 +3040,7 @@ void FBlueprintGraphArgumentLayout::PinInfoChanged(const FEdGraphPinType& PinTyp
 						});
 						if (UDPinPtr)
 						{
+							Node->Modify();
 							(*UDPinPtr)->PinType = PinType;
 
 							// Inputs flagged as pass-by-reference will also be flagged as 'const' here to conform to the expected native C++
diff --git a/Engine/Source/Editor/KismetWidgets/Private/SPinTypeSelector.cpp b/Engine/Source/Editor/KismetWidgets/Private/SPinTypeSelector.cpp
--- a/Engine/Source/Editor/KismetWidgets/Private/SPinTypeSelector.cpp
+++ b/Engine/Source/Editor/KismetWidgets/Private/SPinTypeSelector.cpp
@@ -531,6 +531,8 @@ void SPinTypeSelector::OnArrayStateToggled()
 
 void SPinTypeSelector::OnContainerTypeSelectionChanged( EPinContainerType PinContainerType)
 {
+	const FScopedTransaction Transaction(LOCTEXT("ChangeParam", "Change Parameter Type"));
+
 	FEdGraphPinType NewTargetPinType = TargetPinType.Get();
 	NewTargetPinType.ContainerType = PinContainerType;
 
