Add in-editor ability to deprecate Blueprint class variables, functions and custom events.

diff --git a/Engine/Source/Editor/BlueprintGraph/Classes/EdGraphSchema_K2.h b/Engine/Source/Editor/BlueprintGraph/Classes/EdGraphSchema_K2.h
--- a/Engine/Source/Editor/BlueprintGraph/Classes/EdGraphSchema_K2.h
+++ b/Engine/Source/Editor/BlueprintGraph/Classes/EdGraphSchema_K2.h
@@ -751,7 +751,7 @@ public:
 	 * @param	Function			The function to find a parent function for
 	 * @return	The UFunction parentfunction, if any.
 	 */
-	UFunction* GetCallableParentFunction(UFunction* Function) const;
+	static UFunction* GetCallableParentFunction(UFunction* Function);
 
 	/** Whether or not the specified actor is a valid target for bound events and literal references (in the right level, not a builder brush, etc */
 	bool IsActorValidForLevelScriptRefs(const AActor* TestActor, const UBlueprint* Blueprint) const;
diff --git a/Engine/Source/Editor/BlueprintGraph/Classes/K2Node_CustomEvent.h b/Engine/Source/Editor/BlueprintGraph/Classes/K2Node_CustomEvent.h
--- a/Engine/Source/Editor/BlueprintGraph/Classes/K2Node_CustomEvent.h
+++ b/Engine/Source/Editor/BlueprintGraph/Classes/K2Node_CustomEvent.h
@@ -19,6 +19,14 @@ class UK2Node_CustomEvent : public UK2Node_Event
 {
 	GENERATED_UCLASS_BODY()
 
+	/** Optional message to display when the event is deprecated */
+	UPROPERTY()
+	FString DeprecationMessage;
+
+	/** Specifies that usage of this event has been deprecated */
+	UPROPERTY()
+	bool bIsDeprecated;
+
 	/** Specifies that the event can be triggered in Editor */
 	UPROPERTY()
 	bool bCallInEditor;
diff --git a/Engine/Source/Editor/BlueprintGraph/Classes/K2Node_EditablePinBase.h b/Engine/Source/Editor/BlueprintGraph/Classes/K2Node_EditablePinBase.h
--- a/Engine/Source/Editor/BlueprintGraph/Classes/K2Node_EditablePinBase.h
+++ b/Engine/Source/Editor/BlueprintGraph/Classes/K2Node_EditablePinBase.h
@@ -59,6 +59,12 @@ struct FKismetUserDeclaredFunctionMetadata
 	UPROPERTY()
 	FLinearColor InstanceTitleColor;
 
+	UPROPERTY()
+	FString DeprecationMessage;
+
+	UPROPERTY()
+	bool bIsDeprecated;
+
 	UPROPERTY()
 	bool bCallInEditor;
 
@@ -69,6 +75,7 @@ struct FKismetUserDeclaredFunctionMetadata
 public:
 	FKismetUserDeclaredFunctionMetadata()
 		: InstanceTitleColor(FLinearColor::White)
+		, bIsDeprecated(false)
 		, bCallInEditor(false)
 		, HasLatentFunctions(INDEX_NONE)
 	{
diff --git a/Engine/Source/Editor/BlueprintGraph/Private/EdGraphSchema_K2.cpp b/Engine/Source/Editor/BlueprintGraph/Private/EdGraphSchema_K2.cpp
--- a/Engine/Source/Editor/BlueprintGraph/Private/EdGraphSchema_K2.cpp
+++ b/Engine/Source/Editor/BlueprintGraph/Private/EdGraphSchema_K2.cpp
@@ -886,7 +886,7 @@ bool UEdGraphSchema_K2::CanFunctionBeUsedInGraph(const UClass* InClass, const UF
 	return false;
 }
 
-UFunction* UEdGraphSchema_K2::GetCallableParentFunction(UFunction* Function) const
+UFunction* UEdGraphSchema_K2::GetCallableParentFunction(UFunction* Function)
 {
 	if( Function && Cast<UClass>(Function->GetOuter()) )
 	{
diff --git a/Engine/Source/Editor/BlueprintGraph/Private/K2Node_CustomEvent.cpp b/Engine/Source/Editor/BlueprintGraph/Private/K2Node_CustomEvent.cpp
--- a/Engine/Source/Editor/BlueprintGraph/Private/K2Node_CustomEvent.cpp
+++ b/Engine/Source/Editor/BlueprintGraph/Private/K2Node_CustomEvent.cpp
@@ -107,6 +107,7 @@ UK2Node_CustomEvent::UK2Node_CustomEvent(const FObjectInitializer& ObjectInitial
 	bOverrideFunction = false;
 	bIsEditable = true;
 	bCanRenameNode = true;
+	bIsDeprecated = false;
 	bCallInEditor = false;
 }
 
diff --git a/Engine/Source/Editor/Kismet/Private/BlueprintCompilationManager.cpp b/Engine/Source/Editor/Kismet/Private/BlueprintCompilationManager.cpp
--- a/Engine/Source/Editor/Kismet/Private/BlueprintCompilationManager.cpp
+++ b/Engine/Source/Editor/Kismet/Private/BlueprintCompilationManager.cpp
@@ -2212,7 +2212,7 @@ UClass* FBlueprintCompilationManagerImpl::FastGenerateSkeletonClass(UBlueprint*
 	UField** CurrentFieldStorageLocation = &Ret->Children;
 	
 	// Helper function for making UFunctions generated for 'event' nodes, e.g. custom event and timelines
-	const auto MakeEventFunction = [&CurrentFieldStorageLocation, MakeFunction, Schema]( FName InName, EFunctionFlags ExtraFnFlags, const TArray<UEdGraphPin*>& InputPins, const TArray< TSharedPtr<FUserPinInfo> >& UserPins, UFunction* InSourceFN, bool bInCallInEditor )
+	const auto MakeEventFunction = [&CurrentFieldStorageLocation, MakeFunction, Schema]( FName InName, EFunctionFlags ExtraFnFlags, const TArray<UEdGraphPin*>& InputPins, const TArray< TSharedPtr<FUserPinInfo> >& UserPins, UFunction* InSourceFN, bool bInCallInEditor, bool bIsDeprecated, const FString& DeprecationMessage )
 	{
 		UField** CurrentParamStorageLocation = nullptr;
 
@@ -2232,9 +2232,18 @@ UClass* FBlueprintCompilationManagerImpl::FastGenerateSkeletonClass(UBlueprint*
 		{
 			FKismetCompilerContext::SetDefaultInputValueMetaData(NewFunction, UserPins);
 
-			if(bInCallInEditor)
+			if (bIsDeprecated)
 			{
-				NewFunction->SetMetaData(FBlueprintMetadata::MD_CallInEditor, TEXT( "true" ));
+				NewFunction->SetMetaData(FBlueprintMetadata::MD_DeprecatedFunction, TEXT("true"));
+				if (!DeprecationMessage.IsEmpty())
+				{
+					NewFunction->SetMetaData(FBlueprintMetadata::MD_DeprecationMessage, *DeprecationMessage);
+				}
+			}
+
+			if (bInCallInEditor)
+			{
+				NewFunction->SetMetaData(FBlueprintMetadata::MD_CallInEditor, TEXT("true"));
 			}
 
 			NewFunction->Bind();
@@ -2271,10 +2280,17 @@ UClass* FBlueprintCompilationManagerImpl::FastGenerateSkeletonClass(UBlueprint*
 		Graph->GetNodesOfClass(EventNodes);
 		for( UK2Node_Event* Event : EventNodes )
 		{
+			FString DeprecationMessage;
+			bool bIsDeprecated = false;
 			bool bCallInEditor = false;
 			if(UK2Node_CustomEvent* CustomEvent = Cast<UK2Node_CustomEvent>(Event))
 			{
 				bCallInEditor = CustomEvent->bCallInEditor;
+				bIsDeprecated = CustomEvent->bIsDeprecated;
+				if (bIsDeprecated)
+				{
+					DeprecationMessage = CustomEvent->DeprecationMessage;
+				}
 			}
 			MakeEventFunction(
 				CompilerContext.GetEventStubFunctionName(Event), 
@@ -2282,7 +2298,9 @@ UClass* FBlueprintCompilationManagerImpl::FastGenerateSkeletonClass(UBlueprint*
 				Event->Pins, 
 				Event->UserDefinedPins,
 				Event->FindEventSignatureFunction(),
-				bCallInEditor
+				bCallInEditor,
+				bIsDeprecated,
+				DeprecationMessage
 			);
 		}
 	}
@@ -2299,11 +2317,11 @@ UClass* FBlueprintCompilationManagerImpl::FastGenerateSkeletonClass(UBlueprint*
 
 			for (const FTTEventTrack& EventTrack : Timeline->EventTracks)
 			{
-				MakeEventFunction(EventTrack.GetFunctionName(), EFunctionFlags::FUNC_None, TArray<UEdGraphPin*>(), TArray< TSharedPtr<FUserPinInfo> >(), nullptr, false);
+				MakeEventFunction(EventTrack.GetFunctionName(), EFunctionFlags::FUNC_None, TArray<UEdGraphPin*>(), TArray< TSharedPtr<FUserPinInfo> >(), nullptr, false, false, FString());
 			}
 		
-			MakeEventFunction(Timeline->GetUpdateFunctionName(), EFunctionFlags::FUNC_None, TArray<UEdGraphPin*>(), TArray< TSharedPtr<FUserPinInfo> >(), nullptr, false);
-			MakeEventFunction(Timeline->GetFinishedFunctionName(), EFunctionFlags::FUNC_None, TArray<UEdGraphPin*>(), TArray< TSharedPtr<FUserPinInfo> >(), nullptr, false);
+			MakeEventFunction(Timeline->GetUpdateFunctionName(), EFunctionFlags::FUNC_None, TArray<UEdGraphPin*>(), TArray< TSharedPtr<FUserPinInfo> >(), nullptr, false, false, FString());
+			MakeEventFunction(Timeline->GetFinishedFunctionName(), EFunctionFlags::FUNC_None, TArray<UEdGraphPin*>(), TArray< TSharedPtr<FUserPinInfo> >(), nullptr, false, false, FString());
 		}
 	}
 
diff --git a/Engine/Source/Editor/Kismet/Private/BlueprintDetailsCustomization.cpp b/Engine/Source/Editor/Kismet/Private/BlueprintDetailsCustomization.cpp
--- a/Engine/Source/Editor/Kismet/Private/BlueprintDetailsCustomization.cpp
+++ b/Engine/Source/Editor/Kismet/Private/BlueprintDetailsCustomization.cpp
@@ -981,6 +981,48 @@ void FBlueprintVarActionDetails::CustomizeDetails( IDetailLayoutBuilder& DetailL
 				.ToolTip(MultilineTooltip)
 			];
 
+		TSharedPtr<SToolTip> DeprecatedTooltip = IDocumentation::Get()->CreateToolTip(LOCTEXT("VariableDeprecated_Tooltip", "Deprecate usage of this variable. Any nodes that reference it will produce a compiler warning indicating that it should be removed or replaced."), nullptr, DocLink, TEXT("Deprecated"));
+
+		Category.AddCustomRow(LOCTEXT("VariableDeprecated", "Deprecated"), true)
+			.Visibility(TAttribute<EVisibility>(this, &FBlueprintVarActionDetails::GetDeprecatedVisibility))
+			.NameContent()
+			[
+				SNew(STextBlock)
+				.ToolTip(DeprecatedTooltip)
+				.Text(LOCTEXT("VariableDeprecated", "Deprecated"))
+				.Font(DetailFontInfo)
+			]
+			.ValueContent()
+			[
+				SNew(SCheckBox)
+				.IsChecked(this, &FBlueprintVarActionDetails::OnGetDeprecatedCheckboxState)
+				.OnCheckStateChanged(this, &FBlueprintVarActionDetails::OnDeprecatedChanged)
+				.IsEnabled(IsVariableInBlueprint())
+				.ToolTip(DeprecatedTooltip)
+			];
+
+		TSharedPtr<SToolTip> DeprecationMessageTooltip = IDocumentation::Get()->CreateToolTip(LOCTEXT("VariableDeprecationMessage_Tooltip", "Optional message to include with the deprecation compiler warning. For example: \'X is no longer being used. Please replace with Y.\'"), nullptr, DocLink, TEXT("DeprecationMessage"));
+
+		Category.AddCustomRow(LOCTEXT("VariableDeprecationMessageLabel", "Deprecation Message"), true)
+			.Visibility(TAttribute<EVisibility>(this, &FBlueprintVarActionDetails::GetDeprecatedVisibility))
+			.IsEnabled(TAttribute<bool>(this, &FBlueprintVarActionDetails::IsVariableDeprecated))
+			.NameContent()
+			[
+				SNew(STextBlock)
+				.ToolTip(DeprecationMessageTooltip)
+				.Text(LOCTEXT("VariableDeprecationMessageLabel", "Deprecation Message"))
+				.Font(DetailFontInfo)
+			]
+			.ValueContent()
+			[
+				SNew(SEditableTextBox)
+				.Text(this, &FBlueprintVarActionDetails::GetDeprecationMessageText)
+				.OnTextCommitted(this, &FBlueprintVarActionDetails::OnDeprecationMessageTextCommitted, CachedVariableName)
+				.IsEnabled(IsVariableInBlueprint())
+				.ToolTipText(this, &FBlueprintVarActionDetails::GetDeprecationMessageText)
+				.Font(IDetailLayoutBuilder::GetDetailFont())
+			];
+
 		TSharedPtr<SToolTip> PropertyFlagsTooltip = IDocumentation::Get()->CreateToolTip(LOCTEXT("DefinedPropertyFlags_Tooltip", "List of defined flags for this property"), NULL, DocLink, TEXT("PropertyFlags"));
 
 		Category.AddCustomRow(LOCTEXT("DefinedPropertyFlags", "Defined Property Flags"), true)
@@ -2505,6 +2547,60 @@ void FBlueprintVarActionDetails::OnMultilineChanged(ECheckBoxState InNewState)
 	}
 }
 
+EVisibility FBlueprintVarActionDetails::GetDeprecatedVisibility() const
+{
+	if (UProperty* VariableProperty = CachedVariableProperty.Get())
+	{
+		if (IsABlueprintVariable(VariableProperty) && IsAUserVariable(VariableProperty))
+		{
+			return EVisibility::Visible;
+		}
+	}
+	return EVisibility::Collapsed;
+}
+
+ECheckBoxState FBlueprintVarActionDetails::OnGetDeprecatedCheckboxState() const
+{
+	return IsVariableDeprecated() ? ECheckBoxState::Checked : ECheckBoxState::Unchecked;
+}
+
+void FBlueprintVarActionDetails::OnDeprecatedChanged(ECheckBoxState InNewState)
+{
+	UProperty* Property = CachedVariableProperty.Get();
+	if (Property)
+	{
+		const bool bDeprecatedFlag = (InNewState == ECheckBoxState::Checked);
+		FBlueprintEditorUtils::SetVariableDeprecatedFlag(GetBlueprintObj(), Property->GetFName(), bDeprecatedFlag);
+	}
+}
+
+FText FBlueprintVarActionDetails::GetDeprecationMessageText() const
+{
+	FName VarName = CachedVariableName;
+	if (VarName != NAME_None)
+	{
+		if (UBlueprint* OwnerBlueprint = GetPropertyOwnerBlueprint())
+		{
+			FString Result;
+			FBlueprintEditorUtils::GetBlueprintVariableMetaData(GetPropertyOwnerBlueprint(), VarName, GetLocalVariableScope(CachedVariableProperty.Get()), FBlueprintMetadata::MD_DeprecationMessage, Result);
+			return FText::FromString(Result);
+		}
+	}
+	return FText();
+}
+
+void FBlueprintVarActionDetails::OnDeprecationMessageTextCommitted(const FText& NewText, ETextCommit::Type InTextCommit, FName VarName)
+{
+	if (NewText.IsEmpty())
+	{
+		FBlueprintEditorUtils::RemoveBlueprintVariableMetaData(GetBlueprintObj(), VarName, GetLocalVariableScope(CachedVariableProperty.Get()), FBlueprintMetadata::MD_DeprecationMessage);
+	}
+	else
+	{
+		FBlueprintEditorUtils::SetBlueprintVariableMetaData(GetBlueprintObj(), VarName, GetLocalVariableScope(CachedVariableProperty.Get()), FBlueprintMetadata::MD_DeprecationMessage, NewText.ToString());
+	}
+}
+
 EVisibility FBlueprintVarActionDetails::IsTooltipEditVisible() const
 {
 	UProperty* VariableProperty = CachedVariableProperty.Get();
@@ -2577,6 +2673,13 @@ bool FBlueprintVarActionDetails::IsVariableInheritedByBlueprint() const
 	return GetBlueprintObj()->SkeletonGeneratedClass->IsChildOf(PropertyOwnerClass);
 }
 
+bool FBlueprintVarActionDetails::IsVariableDeprecated() const
+{
+	UProperty* Property = CachedVariableProperty.Get();
+
+	return Property && Property->HasAnyPropertyFlags(CPF_Deprecated);
+}
+
 static FDetailWidgetRow& AddRow( TArray<TSharedRef<FDetailWidgetRow> >& OutChildRows )
 {
 	TSharedRef<FDetailWidgetRow> NewRow( new FDetailWidgetRow );
@@ -3451,6 +3554,67 @@ void FBlueprintGraphActionDetails::CustomizeDetails( IDetailLayoutBuilder& Detai
 			];
 		}
 
+		const bool bShowDeprecated = bHasAGraph || IsCustomEvent();
+		if (bShowDeprecated)
+		{
+			FFormatNamedArguments DeprecationTooltipFormatArgs;
+			if (bHasAGraph)
+			{
+				DeprecationTooltipFormatArgs.Add(TEXT("FunctionOrCustomEvent"), LOCTEXT("FunctionOrEvent_Function", "function"));
+			}
+			else
+			{
+				DeprecationTooltipFormatArgs.Add(TEXT("FunctionOrCustomEvent"), LOCTEXT("FunctionOrEvent_CustomEvent", "custom event"));
+			}
+
+			bool bIsOverride = false;
+			FFunctionFromNodeHelper FunctionFromNode(FunctionEntryNode);
+			if (FunctionFromNode.Function)
+			{
+				bIsOverride = (UEdGraphSchema_K2::GetCallableParentFunction(FunctionFromNode.Function) != nullptr);
+			}
+
+			const FText DeprecatedTooltipText = FText::Format(LOCTEXT("DeprecatedFunction_Tooltip", "Deprecate usage of this {FunctionOrCustomEvent}. Any nodes that reference it will produce a compiler warning indicating that it should be removed or replaced."), DeprecationTooltipFormatArgs);
+
+			Category.AddCustomRow(LOCTEXT("DeprecatedFunction", "Deprecated"), true)
+			.NameContent()
+			[
+				SNew(STextBlock)
+				.Text(LOCTEXT("DeprecatedFunction", "Deprecated"))
+				.ToolTipText(DeprecatedTooltipText)
+				.Font(IDetailLayoutBuilder::GetDetailFont())
+			]
+			.ValueContent()
+			[
+				SNew(SCheckBox)
+				.IsChecked(this, &FBlueprintGraphActionDetails::OnGetDeprecatedCheckboxState)
+				.ToolTipText(DeprecatedTooltipText)
+				.OnCheckStateChanged(this, &FBlueprintGraphActionDetails::OnDeprecatedChanged)
+				.IsEnabled(!bIsOverride)
+			];
+
+			const FText DeprecationMessageTooltipText = LOCTEXT("DeprecationMessage_Tooltip", "Optional message to include with the deprecation compiler warning. For example: \'X is no longer being used. Please replace with Y.\'");
+
+			Category.AddCustomRow(LOCTEXT("DeprecationMessage", "Deprecation Message"), true)
+			.IsEnabled(TAttribute<bool>(this, &FBlueprintGraphActionDetails::IsFunctionDeprecated))
+			.NameContent()
+			[
+				SNew(STextBlock)
+				.Text(LOCTEXT("DeprecationMessage", "Deprecation Message"))
+				.ToolTipText(DeprecationMessageTooltipText)
+				.Font(IDetailLayoutBuilder::GetDetailFont())
+			]
+			.ValueContent()
+			[
+				SNew(SEditableTextBox)
+				.Text(this, &FBlueprintGraphActionDetails::GetDeprecationMessageText)
+				.OnTextCommitted(this, &FBlueprintGraphActionDetails::OnDeprecationMessageTextCommitted)
+				.IsEnabled(!bIsOverride)
+				.ToolTipText(this, &FBlueprintGraphActionDetails::GetDeprecationMessageText)
+				.Font(IDetailLayoutBuilder::GetDetailFont())
+			];
+		}
+
 		IDetailCategoryBuilder& InputsCategory = DetailLayout.EditCategory("Inputs", LOCTEXT("FunctionDetailsInputs", "Inputs"));
 		
 		TSharedRef<FBlueprintGraphArgumentGroupLayout> InputArgumentGroup =
@@ -3736,6 +3900,103 @@ void FBlueprintGraphActionDetails::OnEditorCallableEventModified( const ECheckBo
 	}
 }
 
+bool FBlueprintGraphActionDetails::IsFunctionDeprecated() const
+{
+	bool bIsDeprecated = false;
+
+	UK2Node_EditablePinBase* Node = FunctionEntryNodePtr.Get();
+	if (Node != nullptr)
+	{
+		UK2Node_CustomEvent* CustomEventNode = Cast<UK2Node_CustomEvent>(Node);
+		if (CustomEventNode != nullptr)
+		{
+			bIsDeprecated = CustomEventNode->bIsDeprecated;
+		}
+		else
+		{
+			UK2Node_FunctionEntry* FunctionEntryNode = Cast<UK2Node_FunctionEntry>(Node);
+			if (FunctionEntryNode != nullptr)
+			{
+				bIsDeprecated = FunctionEntryNode->MetaData.bIsDeprecated;
+			}
+		}
+	}
+
+	return bIsDeprecated;
+}
+
+ECheckBoxState FBlueprintGraphActionDetails::OnGetDeprecatedCheckboxState() const
+{
+	return IsFunctionDeprecated() ? ECheckBoxState::Checked : ECheckBoxState::Unchecked;
+}
+
+void FBlueprintGraphActionDetails::OnDeprecatedChanged(ECheckBoxState InNewState)
+{
+	const bool bIsDeprecated = (InNewState == ECheckBoxState::Checked);
+
+	UK2Node_EditablePinBase* Node = FunctionEntryNodePtr.Get();
+	if (Node != nullptr)
+	{
+		UK2Node_CustomEvent* CustomEventNode = Cast<UK2Node_CustomEvent>(Node);
+		if (CustomEventNode != nullptr)
+		{
+			CustomEventNode->bIsDeprecated = bIsDeprecated;
+		}
+		else
+		{
+			CastChecked<UK2Node_FunctionEntry>(Node)->MetaData.bIsDeprecated = bIsDeprecated;
+		}
+
+		FBlueprintEditorUtils::MarkBlueprintAsStructurallyModified(Node->GetBlueprint());
+	}
+}
+
+FText FBlueprintGraphActionDetails::GetDeprecationMessageText() const
+{
+	FText DeprecationMessage;
+
+	UK2Node_EditablePinBase* Node = FunctionEntryNodePtr.Get();
+	if (Node != nullptr)
+	{
+		UK2Node_CustomEvent* CustomEventNode = Cast<UK2Node_CustomEvent>(Node);
+		if (CustomEventNode != nullptr)
+		{
+			DeprecationMessage = FText::FromString(CustomEventNode->DeprecationMessage);
+		}
+		else
+		{
+			UK2Node_FunctionEntry* FunctionEntryNode = Cast<UK2Node_FunctionEntry>(Node);
+			if (FunctionEntryNode != nullptr)
+			{
+				DeprecationMessage = FText::FromString(FunctionEntryNode->MetaData.DeprecationMessage);
+			}
+		}
+	}
+
+	return DeprecationMessage;
+}
+
+void FBlueprintGraphActionDetails::OnDeprecationMessageTextCommitted(const FText& NewText, ETextCommit::Type InTextCommit)
+{
+	const FString DeprecationMessage = NewText.ToString();
+
+	UK2Node_EditablePinBase* Node = FunctionEntryNodePtr.Get();
+	if (Node != nullptr)
+	{
+		UK2Node_CustomEvent* CustomEventNode = Cast<UK2Node_CustomEvent>(Node);
+		if (CustomEventNode != nullptr)
+		{
+			CustomEventNode->DeprecationMessage = DeprecationMessage;
+		}
+		else
+		{
+			CastChecked<UK2Node_FunctionEntry>(Node)->MetaData.DeprecationMessage = DeprecationMessage;
+		}
+
+		FBlueprintEditorUtils::MarkBlueprintAsStructurallyModified(Node->GetBlueprint());
+	}
+}
+
 UMulticastDelegateProperty* FBlueprintDelegateActionDetails::GetDelegateProperty() const
 {
 	if (MyBlueprint.IsValid())
diff --git a/Engine/Source/Editor/Kismet/Private/BlueprintDetailsCustomization.h b/Engine/Source/Editor/Kismet/Private/BlueprintDetailsCustomization.h
--- a/Engine/Source/Editor/Kismet/Private/BlueprintDetailsCustomization.h
+++ b/Engine/Source/Editor/Kismet/Private/BlueprintDetailsCustomization.h
@@ -224,6 +224,13 @@ private:
 	ECheckBoxState OnGetMultilineCheckboxState() const;
 	void OnMultilineChanged(ECheckBoxState InNewState);
 
+	EVisibility GetDeprecatedVisibility() const;
+	ECheckBoxState OnGetDeprecatedCheckboxState() const;
+	void OnDeprecatedChanged(ECheckBoxState InNewState);
+
+	FText GetDeprecationMessageText() const;
+	void OnDeprecationMessageTextCommitted(const FText& NewText, ETextCommit::Type InTextCommit, FName VarName);
+
 	/** Refresh the property flags list */
 	void RefreshPropertyFlags();
 
@@ -244,6 +251,9 @@ private:
 
 	/** Returns TRUE if the Variable is inherited by the current Blueprint */
 	bool IsVariableInheritedByBlueprint() const;
+
+	/** Returns TRUE if the Variable is marked as deprecated */
+	bool IsVariableDeprecated() const;
 private:
 	/** Pointer back to my parent tab */
 	TWeakPtr<SMyBlueprint> MyBlueprint;
@@ -609,6 +619,12 @@ private:
 	/** Enables/Disables selected event as editor callable  */
 	void OnEditorCallableEventModified( const ECheckBoxState NewCheckedState ) const;
 
+	bool IsFunctionDeprecated() const;
+	ECheckBoxState OnGetDeprecatedCheckboxState() const;
+	void OnDeprecatedChanged(ECheckBoxState InNewState);
+
+	FText GetDeprecationMessageText() const;
+	void OnDeprecationMessageTextCommitted(const FText& NewText, ETextCommit::Type InTextCommit);
 	
 	FReply OnAddNewOutputClicked();
 
diff --git a/Engine/Source/Editor/KismetCompiler/Private/KismetCompiler.cpp b/Engine/Source/Editor/KismetCompiler/Private/KismetCompiler.cpp
--- a/Engine/Source/Editor/KismetCompiler/Private/KismetCompiler.cpp
+++ b/Engine/Source/Editor/KismetCompiler/Private/KismetCompiler.cpp
@@ -1701,6 +1701,16 @@ void FKismetCompilerContext::PrecompileFunction(FKismetFunctionContext& Context,
 			Context.Function->SetMetaData(FBlueprintMetadata::MD_CallInEditor, TEXT( "true" ));
 		}
 
+		// Set appropriate metadata if the function is deprecated
+		if (FunctionMetaData.bIsDeprecated)
+		{
+			Context.Function->SetMetaData(FBlueprintMetadata::MD_DeprecatedFunction, TEXT("true"));
+			if (!FunctionMetaData.DeprecationMessage.IsEmpty())
+			{
+				Context.Function->SetMetaData(FBlueprintMetadata::MD_DeprecationMessage, *FunctionMetaData.DeprecationMessage);
+			}
+		}
+
 		// Set the required function flags
 		if (Context.CanBeCalledByKismet())
 		{
@@ -2233,10 +2243,22 @@ void FKismetCompilerContext::SetCalculatedMetaDataAndFlags(UFunction* Function,
 	{
 		Function->SetMetaData(FBlueprintMetadata::MD_Tooltip, *EntryNode->MetaData.ToolTip.ToString());
 	}
+
 	if (EntryNode->MetaData.bCallInEditor)
 	{
-		Function->SetMetaData(FBlueprintMetadata::MD_CallInEditor, TEXT( "true" ));
+		Function->SetMetaData(FBlueprintMetadata::MD_CallInEditor, TEXT("true"));
 	}
+
+	if (EntryNode->MetaData.bIsDeprecated)
+	{
+		Function->SetMetaData(FBlueprintMetadata::MD_DeprecatedFunction, TEXT("true"));
+		
+		if (!EntryNode->MetaData.DeprecationMessage.IsEmpty())
+		{
+			Function->SetMetaData(FBlueprintMetadata::MD_DeprecationMessage, *(EntryNode->MetaData.DeprecationMessage));
+		}
+	}
+	
 	if (UEdGraphPin* WorldContextPin = EntryNode->GetAutoWorldContextPin())
 	{
 		Function->SetMetaData(FBlueprintMetadata::MD_WorldContext, *WorldContextPin->PinName.ToString());
@@ -2864,6 +2886,9 @@ void FKismetCompilerContext::CreateFunctionStubForEvent(UK2Node_Event* SrcEventN
 		StubContext.MarkAsNetFunction(SrcCustomEventNode->GetNetFlags());
 		// Synchronize the entry node call in editor value with the entry point
 		EntryNode->MetaData.bCallInEditor = SrcCustomEventNode->bCallInEditor;
+		// Synchronize the node deprecation state with the entry point
+		EntryNode->MetaData.bIsDeprecated = SrcCustomEventNode->bIsDeprecated;
+		EntryNode->MetaData.DeprecationMessage = SrcCustomEventNode->DeprecationMessage;
 	}
 	EntryNode->AllocateDefaultPins();
 
diff --git a/Engine/Source/Editor/UnrealEd/Private/Kismet2/BlueprintEditorUtils.cpp b/Engine/Source/Editor/UnrealEd/Private/Kismet2/BlueprintEditorUtils.cpp
--- a/Engine/Source/Editor/UnrealEd/Private/Kismet2/BlueprintEditorUtils.cpp
+++ b/Engine/Source/Editor/UnrealEd/Private/Kismet2/BlueprintEditorUtils.cpp
@@ -3626,6 +3626,25 @@ void FBlueprintEditorUtils::SetVariableAdvancedDisplayFlag(UBlueprint* InBluepri
 	FBlueprintEditorUtils::MarkBlueprintAsStructurallyModified(InBlueprint);
 }
 
+void FBlueprintEditorUtils::SetVariableDeprecatedFlag(UBlueprint* InBlueprint, const FName& InVarName, const bool bInIsDeprecated)
+{
+	const int32 VarIndex = FBlueprintEditorUtils::FindNewVariableIndex(InBlueprint, InVarName);
+
+	if (VarIndex != INDEX_NONE)
+	{
+		if (bInIsDeprecated)
+		{
+			InBlueprint->NewVariables[VarIndex].PropertyFlags |= CPF_Deprecated;
+		}
+		else
+		{
+			InBlueprint->NewVariables[VarIndex].PropertyFlags &= ~CPF_Deprecated;
+		}
+	}
+
+	FBlueprintEditorUtils::MarkBlueprintAsStructurallyModified(InBlueprint);
+}
+
 struct FMetaDataDependencyHelper
 {
 	static void OnChange(UBlueprint* Blueprint, FName MetaDataKey)
diff --git a/Engine/Source/Editor/UnrealEd/Public/Kismet2/BlueprintEditorUtils.h b/Engine/Source/Editor/UnrealEd/Public/Kismet2/BlueprintEditorUtils.h
--- a/Engine/Source/Editor/UnrealEd/Public/Kismet2/BlueprintEditorUtils.h
+++ b/Engine/Source/Editor/UnrealEd/Public/Kismet2/BlueprintEditorUtils.h
@@ -1018,6 +1018,14 @@ public:
 	 */
 	static void SetVariableAdvancedDisplayFlag(UBlueprint* InBlueprint, const FName& InVarName, const bool bInIsAdvancedDisplay);
 
+	/**
+	 * Sets the Deprecated flag on the variable with the specified name
+	 *
+	 * @param	InVarName				Name of the var to set the flag on
+	 * @param	bInIsDeprecated			The new value to set the bitflag to
+	 */
+	static void SetVariableDeprecatedFlag(UBlueprint* InBlueprint, const FName& InVarName, const bool bInIsDeprecated);
+
 	/** Sets a metadata key/value on the specified variable
 	 *
 	 * @param Blueprint				The Blueprint to find the variable in
