PR #5188: Make Blueprint function parameters case sensitive to renaming (Contributed by [REDACTED_PERSON])

diff --git a/Engine/Source/Editor/Kismet/Private/BlueprintDetailsCustomization.cpp b/Engine/Source/Editor/Kismet/Private/BlueprintDetailsCustomization.cpp
--- a/Engine/Source/Editor/Kismet/Private/BlueprintDetailsCustomization.cpp
+++ b/Engine/Source/Editor/Kismet/Private/BlueprintDetailsCustomization.cpp
@@ -3076,7 +3076,7 @@ void FBlueprintGraphArgumentLayout::OnArgNameTextCommitted(const FText& NewText,
 	{
 		const FName OldName = ParamItemPtr.Pin()->PinName;
 		const FString& NewName = NewText.ToString();
-		if (OldName.ToString() != NewName)
+		if (!OldName.ToString().Equals(NewName))
 		{
 			GraphActionDetailsPtr.Pin()->OnPinRenamed(TargetNode, OldName, NewName);
 		}
