Fix crash on attempt to deprecate a macro graph (not currently supported).

#rnx

diff --git a/Engine/Source/Editor/Kismet/Private/BlueprintDetailsCustomization.cpp b/Engine/Source/Editor/Kismet/Private/BlueprintDetailsCustomization.cpp
--- a/Engine/Source/Editor/Kismet/Private/BlueprintDetailsCustomization.cpp
+++ b/Engine/Source/Editor/Kismet/Private/BlueprintDetailsCustomization.cpp
@@ -3202,6 +3202,9 @@ void FBlueprintGraphActionDetails::CustomizeDetails( IDetailLayoutBuilder& Detai
 
 	if (FunctionEntryNode && FunctionEntryNode->IsEditable())
 	{
+		const bool bIsCustomEvent = IsCustomEvent();
+		const bool bIsFunctionGraph = FunctionEntryNode->IsA<UK2Node_FunctionEntry>();
+
 		IDetailCategoryBuilder& Category = DetailLayout.EditCategory("Graph", LOCTEXT("FunctionDetailsGraph", "Graph"));
 		if (bHasAGraph)
 		{
@@ -3399,7 +3402,7 @@ void FBlueprintGraphActionDetails::CustomizeDetails( IDetailLayoutBuilder& Detai
 			}
 		}
 
-		if (IsCustomEvent())
+		if (bIsCustomEvent)
 		{
 			/** A collection of static utility callbacks to provide the custom-event details ui with */
 			struct LocalCustomEventUtils
@@ -3524,7 +3527,7 @@ void FBlueprintGraphActionDetails::CustomizeDetails( IDetailLayoutBuilder& Detai
 				]
 			];
 		}
-		const bool bShowCallInEditor = IsCustomEvent() || FBlueprintEditorUtils::IsEditorUtilityBlueprint( GetBlueprintObj() ) || (FunctionEntryNode && FunctionEntryNode->IsEditable());
+		const bool bShowCallInEditor = bIsCustomEvent || FBlueprintEditorUtils::IsEditorUtilityBlueprint( GetBlueprintObj() ) || (FunctionEntryNode && FunctionEntryNode->IsEditable());
 		if( bShowCallInEditor )
 		{
 			Category.AddCustomRow( LOCTEXT( "EditorCallable", "Call In Editor" ) )
@@ -3554,11 +3557,11 @@ void FBlueprintGraphActionDetails::CustomizeDetails( IDetailLayoutBuilder& Detai
 			];
 		}
 
-		const bool bShowDeprecated = bHasAGraph || IsCustomEvent();
+		const bool bShowDeprecated = bIsFunctionGraph || bIsCustomEvent;
 		if (bShowDeprecated)
 		{
 			FFormatNamedArguments DeprecationTooltipFormatArgs;
-			if (bHasAGraph)
+			if (bIsFunctionGraph)
 			{
 				DeprecationTooltipFormatArgs.Add(TEXT("FunctionOrCustomEvent"), LOCTEXT("FunctionOrEvent_Function", "function"));
 			}
