[NO NON-PII COMMIT MESSAGE]

diff --git a/Engine/Source/Editor/Kismet/Private/DiffUtils.cpp b/Engine/Source/Editor/Kismet/Private/DiffUtils.cpp
--- a/Engine/Source/Editor/Kismet/Private/DiffUtils.cpp
+++ b/Engine/Source/Editor/Kismet/Private/DiffUtils.cpp
@@ -644,10 +644,20 @@ TArray<FPropertySoftPath> DiffUtils::GetVisiblePropertiesInOrderDeclared(const U
 					NewPath.Push(PropertyIt->GetFName());
 					if (const UObjectProperty* ObjectProperty = Cast<UObjectProperty>(*PropertyIt))
 					{
+						bool bShouldRecurse = false;
 						const UObject* const* BaseObject = reinterpret_cast<const UObject* const*>( ObjectProperty->ContainerPtrToValuePtr<void>(ForObj) );
 						if (BaseObject && *BaseObject)
 						{
-							Ret.Append( GetVisiblePropertiesInOrderDeclared(*BaseObject, NewPath) );
+							// See if this is an object in the package, if so recurse into it
+							if ((*BaseObject)->IsIn(ForObj->GetOutermost()))
+							{
+								bShouldRecurse = true;
+								Ret.Append(GetVisiblePropertiesInOrderDeclared(*BaseObject, NewPath));
+							}
+						}
+						if (!bShouldRecurse)
+						{
+							Ret.Push(NewPath);
 						}
 					}
 					else
diff --git a/Engine/Source/Editor/Kismet/Private/SCSDiff.cpp b/Engine/Source/Editor/Kismet/Private/SCSDiff.cpp
--- a/Engine/Source/Editor/Kismet/Private/SCSDiff.cpp
+++ b/Engine/Source/Editor/Kismet/Private/SCSDiff.cpp
@@ -18,6 +18,9 @@ FSCSDiff::FSCSDiff(const UBlueprint* InBlueprint)
 		return;
 	}
 
+	// Need to pull off const because the ICH access functions theoretically could modify it
+	Blueprint = const_cast<UBlueprint*>(InBlueprint);
+
 	Inspector = SNew(SKismetInspector)
 		.HideNameArea(true)
 		.ViewIdentifier(FName("BlueprintInspector"))
@@ -55,16 +58,16 @@ TSharedRef< SWidget > FSCSDiff::TreeWidget()
 	return ContainerWidget.ToSharedRef();
 }
 
-void GetDisplayedHierarchyRecursive(TArray< int32 >& TreeAddress, const FSCSEditorTreeNode& Node, TArray< FSCSResolvedIdentifier >& OutResult)
+void GetDisplayedHierarchyRecursive(UBlueprint* Blueprint, TArray< int32 >& TreeAddress, const FSCSEditorTreeNode& Node, TArray< FSCSResolvedIdentifier >& OutResult)
 {
 	FSCSIdentifier Identifier = { Node.GetVariableName(), TreeAddress };
-	FSCSResolvedIdentifier ResolvedIdentifier = { Identifier, Node.GetComponentTemplate() };
+	FSCSResolvedIdentifier ResolvedIdentifier = { Identifier, Node.GetOrCreateEditableComponentTemplate(Blueprint) };
 	OutResult.Push(ResolvedIdentifier);
 	const auto& Children = Node.GetChildren();
 	for (int32 Iter = 0; Iter != Children.Num(); ++Iter)
 	{
 		TreeAddress.Push(Iter);
-		GetDisplayedHierarchyRecursive(TreeAddress, *Children[Iter], OutResult);
+		GetDisplayedHierarchyRecursive(Blueprint, TreeAddress, *Children[Iter], OutResult);
 		TreeAddress.Pop();
 	}
 }
@@ -80,7 +83,7 @@ TArray< FSCSResolvedIdentifier > FSCSDiff::GetDisplayedHierarchy() const
 		{
 			TArray< int32 > TreeAddress;
 			TreeAddress.Push(Iter);
-			GetDisplayedHierarchyRecursive(TreeAddress, *RootNodes[Iter], Ret);
+			GetDisplayedHierarchyRecursive(Blueprint, TreeAddress, *RootNodes[Iter], Ret);
 		}
 	}
 
@@ -98,7 +101,7 @@ void FSCSDiff::OnSCSEditorUpdateSelectionFromNodes(const TArray<FSCSEditorTreeNo
 		if(NodePtr.IsValid() && NodePtr->CanEditDefaults())
 		{
 			InspectorTitle = FText::FromString(NodePtr->GetDisplayString());
-			InspectorObjects.Add(NodePtr->GetComponentTemplate());
+			InspectorObjects.Add(NodePtr->GetOrCreateEditableComponentTemplate(Blueprint));
 		}
 	}
 
diff --git a/Engine/Source/Editor/Kismet/Private/SSCSEditor.cpp b/Engine/Source/Editor/Kismet/Private/SSCSEditor.cpp
--- a/Engine/Source/Editor/Kismet/Private/SSCSEditor.cpp
+++ b/Engine/Source/Editor/Kismet/Private/SSCSEditor.cpp
@@ -432,7 +432,7 @@ class USCS_Node* FSCSEditorTreeNode::GetSCSNode() const
 	return nullptr;
 }
 
-UActorComponent* FSCSEditorTreeNode::GetOrCreateEditableComponentTemplate(UBlueprint* ActualEditedBlueprint)
+UActorComponent* FSCSEditorTreeNode::GetOrCreateEditableComponentTemplate(UBlueprint* ActualEditedBlueprint) const
 {
 	return nullptr;
 }
@@ -1051,7 +1051,7 @@ FText FSCSEditorTreeNodeInstancedInheritedComponent::GetDisplayName() const
 	return FText::GetEmpty();
 }
 
-UActorComponent* FSCSEditorTreeNodeInstancedInheritedComponent::GetOrCreateEditableComponentTemplate(UBlueprint* ActualEditedBlueprint) 
+UActorComponent* FSCSEditorTreeNodeInstancedInheritedComponent::GetOrCreateEditableComponentTemplate(UBlueprint* ActualEditedBlueprint) const
 {
 	if (CanEditDefaults())
 	{
@@ -1110,7 +1110,7 @@ FText FSCSEditorTreeNodeInstanceAddedComponent::GetDisplayName() const
 	return FText::FromName(InstancedComponentName);
 }
 
-UActorComponent* FSCSEditorTreeNodeInstanceAddedComponent::GetOrCreateEditableComponentTemplate(UBlueprint* ActualEditedBlueprint)
+UActorComponent* FSCSEditorTreeNodeInstanceAddedComponent::GetOrCreateEditableComponentTemplate(UBlueprint* ActualEditedBlueprint) const
 {
 	return GetComponentTemplate();
 }
@@ -1267,7 +1267,7 @@ class USCS_Node* FSCSEditorTreeNodeComponent::GetSCSNode() const
 	return SCSNodePtr.Get();
 }
 
-UActorComponent* FSCSEditorTreeNodeComponent::GetOrCreateEditableComponentTemplate(UBlueprint* ActualEditedBlueprint)
+UActorComponent* FSCSEditorTreeNodeComponent::GetOrCreateEditableComponentTemplate(UBlueprint* ActualEditedBlueprint) const
 {
 	if (CanEditDefaults())
 	{
diff --git a/Engine/Source/Editor/Kismet/Public/SCSDiff.h b/Engine/Source/Editor/Kismet/Public/SCSDiff.h
--- a/Engine/Source/Editor/Kismet/Public/SCSDiff.h
+++ b/Engine/Source/Editor/Kismet/Public/SCSDiff.h
@@ -27,4 +27,7 @@ private:
 	TSharedPtr< class SWidget > ContainerWidget;
 	TSharedPtr< class SSCSEditor > SCSEditor;
 	TSharedPtr< class SKismetInspector > Inspector;
+
+	/** Blueprint we are inspecting */
+	UBlueprint* Blueprint;
 };
diff --git a/Engine/Source/Editor/Kismet/Public/SSCSEditor.h b/Engine/Source/Editor/Kismet/Public/SSCSEditor.h
--- a/Engine/Source/Editor/Kismet/Public/SSCSEditor.h
+++ b/Engine/Source/Editor/Kismet/Public/SSCSEditor.h
@@ -83,7 +83,7 @@ public:
 	 *
 	 * @return The component template that can be editable for actual class.
 	 */
-	virtual UActorComponent* GetOrCreateEditableComponentTemplate(UBlueprint* ActualEditedBlueprint);
+	virtual UActorComponent* GetOrCreateEditableComponentTemplate(UBlueprint* ActualEditedBlueprint) const;
 	/**
 	 * Finds the component instance represented by this node contained within a given Actor instance.
 	 *
@@ -109,6 +109,7 @@ public:
 	 *								have been reinstanced following construction script execution).
 	 *
 	 * @note	Deliberately non-virtual, for performance reasons.
+	 * @warning This will not return the right component for components overridden by the inherited component handler, you need to call GetOrCreateEditableComponentTemplate instead
 	 * @return	The component template or instance represented by this node, if it's a component node.
 	 */
 	UActorComponent* GetComponentTemplate(bool bEvenIfPendingKill = false) const;
@@ -385,7 +386,7 @@ public:
 	//virtual FName GetVariableName() const override;
 	//virtual FString GetDisplayString() const override;
 	virtual FText GetDisplayName() const override;
-	virtual UActorComponent* GetOrCreateEditableComponentTemplate(UBlueprint* ActualEditedBlueprint) override;
+	virtual UActorComponent* GetOrCreateEditableComponentTemplate(UBlueprint* ActualEditedBlueprint) const override;
 	// End of FSCSEditorTreeNode public interface
 
 private:
@@ -416,7 +417,7 @@ public:
 	virtual FName GetVariableName() const override { return NAME_None; }
 	virtual FString GetDisplayString() const override;
 	virtual FText GetDisplayName() const override;
-	virtual UActorComponent* GetOrCreateEditableComponentTemplate(UBlueprint* ActualEditedBlueprint) override;
+	virtual UActorComponent* GetOrCreateEditableComponentTemplate(UBlueprint* ActualEditedBlueprint) const override;
 	virtual void OnCompleteRename(const FText& InNewName) override;
 	// End of FSCSEditorTreeNode public interface
 
@@ -464,7 +465,7 @@ public:
 	//virtual FString GetDisplayString() const override;
 	virtual FText GetDisplayName() const override;
 	virtual class USCS_Node* GetSCSNode() const override;
-	virtual UActorComponent* GetOrCreateEditableComponentTemplate(UBlueprint* ActualEditedBlueprint) override;
+	virtual UActorComponent* GetOrCreateEditableComponentTemplate(UBlueprint* ActualEditedBlueprint) const override;
 	virtual void OnCompleteRename(const FText& InNewName) override;
 	// End of FSCSEditorTreeNode public interface
 
