Build fix

diff --git a/Engine/Source/Editor/Kismet/Private/DiffUtils.cpp b/Engine/Source/Editor/Kismet/Private/DiffUtils.cpp
--- a/Engine/Source/Editor/Kismet/Private/DiffUtils.cpp
+++ b/Engine/Source/Editor/Kismet/Private/DiffUtils.cpp
@@ -2,6 +2,7 @@
 
 #include "DiffUtils.h"
 #include "UObject/PropertyPortFlags.h"
+#include "UObject/Package.h"
 #include "Widgets/Images/SImage.h"
 #include "EditorStyleSet.h"
 #include "ISourceControlProvider.h"
