[NO NON-PII COMMIT MESSAGE]

diff --git a/Engine/Source/Editor/Kismet/Private/DiffUtils.cpp b/Engine/Source/Editor/Kismet/Private/DiffUtils.cpp
--- a/Engine/Source/Editor/Kismet/Private/DiffUtils.cpp
+++ b/Engine/Source/Editor/Kismet/Private/DiffUtils.cpp
@@ -864,6 +864,48 @@ TSharedPtr<FBlueprintDifferenceTreeEntry> FBlueprintDifferenceTreeEntry::CreateC
 	));
 }
 
+TSharedPtr<FBlueprintDifferenceTreeEntry> FBlueprintDifferenceTreeEntry::CreateClassSettingsCategoryEntry(FOnDiffEntryFocused FocusCallback, const TArray< TSharedPtr<FBlueprintDifferenceTreeEntry> >& Children, bool bHasDifferences)
+{
+	const auto CreateSettingsRootEntry = [](FLinearColor Color) -> TSharedRef<SWidget>
+	{
+		return SNew(STextBlock)
+			.ToolTipText(NSLOCTEXT("FBlueprintDifferenceTreeEntry", "SettingsTooltip", "The list of changes made to class settings"))
+			.ColorAndOpacity(Color)
+			.Text(NSLOCTEXT("FBlueprintDifferenceTreeEntry", "SettingsLabel", "Class Settings"));
+	};
+
+	return TSharedPtr<FBlueprintDifferenceTreeEntry>(new FBlueprintDifferenceTreeEntry(
+		FocusCallback
+		, FGenerateDiffEntryWidget::CreateStatic(CreateSettingsRootEntry, DiffViewUtils::LookupColor(bHasDifferences))
+		, Children
+	));
+}
+
+TSharedPtr<FBlueprintDifferenceTreeEntry> FBlueprintDifferenceTreeEntry::CreateClassSettingsCategoryEntryForMerge(FOnDiffEntryFocused FocusCallback, const TArray< TSharedPtr<FBlueprintDifferenceTreeEntry> >& Children, bool bHasRemoteDifferences, bool bHasLocalDifferences, bool bHasConflicts)
+{
+	const auto CreateSettingsRootEntry = [](bool bInHasRemoteDifferences, bool bInHasLocalDifferences, bool bInHasConflicts) -> TSharedRef<SWidget>
+	{
+		const FLinearColor BaseColor = DiffViewUtils::LookupColor(bInHasRemoteDifferences || bInHasLocalDifferences, bInHasConflicts);
+		return  SNew(SHorizontalBox)
+			+ SHorizontalBox::Slot()
+			[
+				SNew(STextBlock)
+				.ToolTipText(NSLOCTEXT("FBlueprintDifferenceTreeEntry", "SettingsTooltip", "The list of changes made to class settings"))
+			.ColorAndOpacity(BaseColor)
+			.Text(NSLOCTEXT("FBlueprintDifferenceTreeEntry", "SettingsLabel", "Class Settings"))
+			]
+		+ DiffViewUtils::Box(true, DiffViewUtils::LookupColor(bInHasRemoteDifferences, bInHasConflicts))
+			+ DiffViewUtils::Box(true, BaseColor)
+			+ DiffViewUtils::Box(true, DiffViewUtils::LookupColor(bInHasLocalDifferences, bInHasConflicts));
+	};
+
+	return TSharedPtr<FBlueprintDifferenceTreeEntry>(new FBlueprintDifferenceTreeEntry(
+		FocusCallback
+		, FGenerateDiffEntryWidget::CreateStatic(CreateSettingsRootEntry, bHasRemoteDifferences, bHasLocalDifferences, bHasConflicts)
+		, Children
+	));
+}
+
 TSharedRef< STreeView<TSharedPtr< FBlueprintDifferenceTreeEntry > > > DiffTreeView::CreateTreeView(TArray< TSharedPtr<FBlueprintDifferenceTreeEntry> >* DifferencesList)
 {
 	const auto RowGenerator = [](TSharedPtr< FBlueprintDifferenceTreeEntry > Entry, const TSharedRef<STableViewBase>& Owner) -> TSharedRef< ITableRow >
diff --git a/Engine/Source/Editor/Kismet/Private/SBlueprintDiff.cpp b/Engine/Source/Editor/Kismet/Private/SBlueprintDiff.cpp
--- a/Engine/Source/Editor/Kismet/Private/SBlueprintDiff.cpp
+++ b/Engine/Source/Editor/Kismet/Private/SBlueprintDiff.cpp
@@ -41,6 +41,7 @@ typedef TMap< FName, const UProperty* > FNamePropertyMap;
 
 const FName DiffMyBluerpintTabId = FName(TEXT("DiffMyBluerpintTab"));
 const FName DiffGraphTabId = FName(TEXT("DiffGraphTab"));
+const FName ClassSettingsMode = FName(TEXT("ClassSettingsMode"));
 
 DECLARE_DELEGATE(FOnSCSDiffControlChanged);
 
@@ -136,41 +137,40 @@ FSCSDiffControl::FSCSDiffControl(
 	));
 }
 
-DECLARE_DELEGATE(FOnCDODiffControlChanged);
+/** Generic wrapper around a details view, this does not actually fill out OutTreeEntries */
+DECLARE_DELEGATE(FOnDetailsDiffControlChanged);
 
-class FCDODiffControl	: public TSharedFromThis<FCDODiffControl>
-						, public IDiffControl
+class FDetailsDiffControl : public IDiffControl
 {
 public:
-	FCDODiffControl( const UObject* InOldCDO
-					, const UObject* InNewCDO
+	FDetailsDiffControl( const UObject* InOldObject
+					, const UObject* InNewObject
 					, TArray< TSharedPtr<FBlueprintDifferenceTreeEntry> >& OutTreeEntries
 					, TArray< TSharedPtr<FBlueprintDifferenceTreeEntry> >& OutRealDifferences
-					, FOnCDODiffControlChanged SelectionCallback );
+					, FOnDetailsDiffControlChanged SelectionCallback );
 
 	TSharedRef<SWidget> OldDetailsWidget() { return OldDetails.DetailsWidget(); }
 	TSharedRef<SWidget> NewDetailsWidget() { return NewDetails.DetailsWidget(); }
 
-	virtual ~FCDODiffControl() {}
+	virtual ~FDetailsDiffControl() {}
 
-private:
+protected:
 	void HighlightDifference(const FPropertySoftPath& PropertyName);
-	void UpdateDisplayedDifferences(TArray< FSingleObjectDiffEntry > const& InDifferences, TArray< TSharedPtr<FBlueprintDifferenceTreeEntry> >& OutTreeEntries, FOnCDODiffControlChanged SelectionCallback);
 
 	FDetailsDiff OldDetails;
 	FDetailsDiff NewDetails;
 
-	int CurrentDifference;
+	TArray< TSharedPtr<FBlueprintDifferenceTreeEntry> > Children;
 };
 
-FCDODiffControl::FCDODiffControl( 
-		const UObject* InOldCDO
-		, const UObject* InNewCDO
+FDetailsDiffControl::FDetailsDiffControl(
+		const UObject* InOldObject
+		, const UObject* InNewObject
 		, TArray< TSharedPtr<FBlueprintDifferenceTreeEntry> >& OutTreeEntries
 		, TArray< TSharedPtr<FBlueprintDifferenceTreeEntry> >& OutRealDifferences
-		, FOnCDODiffControlChanged SelectionCallback )
-	: OldDetails(InOldCDO, FDetailsDiff::FOnDisplayedPropertiesChanged() )
-	, NewDetails(InNewCDO, FDetailsDiff::FOnDisplayedPropertiesChanged() )
+		, FOnDetailsDiffControlChanged SelectionCallback )
+	: OldDetails(InOldObject, FDetailsDiff::FOnDisplayedPropertiesChanged() )
+	, NewDetails(InNewObject, FDetailsDiff::FOnDisplayedPropertiesChanged() )
 {
 	TArray< FSingleObjectDiffEntry > DifferingProperties;
 	OldDetails.DiffAgainst(NewDetails, DifferingProperties);
@@ -260,61 +260,185 @@ FCDODiffControl::FCDODiffControl(
 		}
 	}
 
-	const auto CreateCDODifferenceWidget = [](FSingleObjectDiffEntry DiffEntry, FText ObjectName) -> TSharedRef<SWidget>
+	const auto CreateDifferenceWidget = [](FSingleObjectDiffEntry DiffEntry, FText ObjectName) -> TSharedRef<SWidget>
 	{
 		return SNew(STextBlock)
 			.Text(DiffViewUtils::PropertyDiffMessage(DiffEntry, ObjectName))
 			.ColorAndOpacity(DiffViewUtils::Differs());
 	};
 
-	const auto FocusDetailsDifferenceEntry = [](FPropertySoftPath Identifier, FCDODiffControl* Control, FOnCDODiffControlChanged InSelectionCallback)
+	const auto FocusDetailsDifferenceEntry = [](FPropertySoftPath Identifier, FDetailsDiffControl* Control, FOnDetailsDiffControlChanged InSelectionCallback)
 	{
 		// This allows the owning control to focus the correct tab (or do whatever else it likes):
 		InSelectionCallback.ExecuteIfBound();
 		Control->HighlightDifference(Identifier);
 	};
 
-	TArray< TSharedPtr<FBlueprintDifferenceTreeEntry> > Children;
-
 	for (auto Difference : OrderedProperties)
 	{
 		auto Entry = TSharedPtr<FBlueprintDifferenceTreeEntry>(
 			new FBlueprintDifferenceTreeEntry(
 			FOnDiffEntryFocused::CreateStatic(FocusDetailsDifferenceEntry, Difference->Identifier, this, SelectionCallback)
-			, FGenerateDiffEntryWidget::CreateStatic(CreateCDODifferenceWidget, *Difference, RightRevision)
+			, FGenerateDiffEntryWidget::CreateStatic(CreateDifferenceWidget, *Difference, RightRevision)
 			, TArray< TSharedPtr<FBlueprintDifferenceTreeEntry> >()
 			)
 		);
 		Children.Push(Entry);
 		OutRealDifferences.Push(Entry);
 	}
+}
 
-	const bool bHasDiffferences = Children.Num() != 0;
-	FLinearColor Color = bHasDiffferences ? DiffViewUtils::Differs() : DiffViewUtils::Identical();
-	if (!bHasDiffferences)
+void FDetailsDiffControl::HighlightDifference(const FPropertySoftPath& PropertyName)
+{
+	OldDetails.HighlightProperty(PropertyName);
+	NewDetails.HighlightProperty(PropertyName);
+}
+
+/** Override for CDO special case */
+class FCDODiffControl : public TSharedFromThis<FCDODiffControl>, public FDetailsDiffControl
+{
+public:
+	FCDODiffControl(const UObject* InOldObject
+		, const UObject* InNewObject
+		, TArray< TSharedPtr<FBlueprintDifferenceTreeEntry> >& OutTreeEntries
+		, TArray< TSharedPtr<FBlueprintDifferenceTreeEntry> >& OutRealDifferences
+		, FOnDetailsDiffControlChanged SelectionCallback)
+		: FDetailsDiffControl(InOldObject, InNewObject, OutTreeEntries, OutRealDifferences, SelectionCallback)
 	{
-		// make one child informing the user that there are no differences:
-		Children.Push(FBlueprintDifferenceTreeEntry::NoDifferencesEntry() );
+		// Parent constructor fills out Children
+		const bool bHasDiffferences = Children.Num() != 0;
+		if (!bHasDiffferences)
+		{
+			// make one child informing the user that there are no differences:
+			Children.Push(FBlueprintDifferenceTreeEntry::NoDifferencesEntry());
+		}
+
+		const auto ForwardSelection = [](FOnDetailsDiffControlChanged InSelectionCallback)
+		{
+			// This allows the owning control to focus the correct tab (or do whatever else it likes):
+			InSelectionCallback.ExecuteIfBound();
+		};
+
+		OutTreeEntries.Push(FBlueprintDifferenceTreeEntry::CreateDefaultsCategoryEntry(
+			FOnDiffEntryFocused::CreateStatic(ForwardSelection, SelectionCallback)
+			, Children
+			, bHasDiffferences
+		));
 	}
+};
 
-	const auto ForwardSelection = [](FOnCDODiffControlChanged InSelectionCallback)
+/** Override for class defaults */
+class FClassDefaultsDiffControl : public TSharedFromThis<FClassDefaultsDiffControl>, public FDetailsDiffControl
+{
+public:
+	FClassDefaultsDiffControl(const UObject* InOldObject
+		, const UObject* InNewObject
+		, TArray< TSharedPtr<FBlueprintDifferenceTreeEntry> >& OutTreeEntries
+		, TArray< TSharedPtr<FBlueprintDifferenceTreeEntry> >& OutRealDifferences
+		, FOnDetailsDiffControlChanged SelectionCallback)
+		: FDetailsDiffControl(InOldObject, InNewObject, OutTreeEntries, OutRealDifferences, SelectionCallback)
 	{
-		// This allows the owning control to focus the correct tab (or do whatever else it likes):
-		InSelectionCallback.ExecuteIfBound();
-	};
+		// Check for parent class and interfaces here
+		const UBlueprint* OldBlueprint = Cast<UBlueprint>(InOldObject);
+		const UBlueprint* NewBlueprint = Cast<UBlueprint>(InNewObject);
 
-	OutTreeEntries.Push(FBlueprintDifferenceTreeEntry::CreateDefaultsCategoryEntry(
-		FOnDiffEntryFocused::CreateStatic(ForwardSelection, SelectionCallback)
-		, Children
-		, bHasDiffferences
-	));
-}
+		const auto GenerateDiffWidget = [](FText DiffText) -> TSharedRef<SWidget>
+		{
+			return SNew(STextBlock)
+				.ColorAndOpacity(DiffViewUtils::Differs())
+				.Text(DiffText)
+				.ToolTipText(DiffText);
+		};
 
-void FCDODiffControl::HighlightDifference(const FPropertySoftPath& PropertyName)
-{
-	OldDetails.HighlightProperty(PropertyName);
-	NewDetails.HighlightProperty(PropertyName);
-}
+		const auto FocusSettings = [](FOnDetailsDiffControlChanged InSelectionCallback)
+		{
+			InSelectionCallback.ExecuteIfBound();
+		};
+
+		if (OldBlueprint && NewBlueprint)
+		{
+			if (OldBlueprint->ParentClass != NewBlueprint->ParentClass)
+			{
+				FText DiffText = FText::Format(LOCTEXT("ParentChanged", "Parent Class changed from {0} to {1}"), FText::FromString(OldBlueprint->ParentClass->GetName()), FText::FromString(NewBlueprint->ParentClass->GetName()));
+
+				auto Entry = TSharedPtr<FBlueprintDifferenceTreeEntry>(new FBlueprintDifferenceTreeEntry(
+					FOnDiffEntryFocused::CreateStatic(FocusSettings, SelectionCallback)
+					, FGenerateDiffEntryWidget::CreateStatic(GenerateDiffWidget, DiffText)
+					, TArray< TSharedPtr<FBlueprintDifferenceTreeEntry> >()
+				));
+
+				Children.Push(Entry);
+				OutRealDifferences.Push(Entry);
+			}
+
+			FString OldInterfaces, NewInterfaces;
+			for (const FBPInterfaceDescription& Desc : OldBlueprint->ImplementedInterfaces)
+			{
+				if (!OldInterfaces.IsEmpty())
+				{
+					OldInterfaces += TEXT(", ");
+				}
+				OldInterfaces += GetNameSafe(*Desc.Interface);
+			}
+
+			for (const FBPInterfaceDescription& Desc : NewBlueprint->ImplementedInterfaces)
+			{
+				if (!NewInterfaces.IsEmpty())
+				{
+					NewInterfaces += TEXT(", ");
+				}
+				NewInterfaces += GetNameSafe(*Desc.Interface);
+			}
+			
+			if (OldInterfaces != NewInterfaces)
+			{
+				FText DiffText = FText::Format(LOCTEXT("InterfacesChanged", "Interfaces changed from '{0}' to '{1}'"), FText::FromString(OldInterfaces), FText::FromString(NewInterfaces));
+
+				auto Entry = TSharedPtr<FBlueprintDifferenceTreeEntry>(new FBlueprintDifferenceTreeEntry(
+					FOnDiffEntryFocused::CreateStatic(FocusSettings, SelectionCallback)
+					, FGenerateDiffEntryWidget::CreateStatic(GenerateDiffWidget, DiffText)
+					, TArray< TSharedPtr<FBlueprintDifferenceTreeEntry> >()
+				));
+
+				Children.Push(Entry);
+				OutRealDifferences.Push(Entry);
+			}
+
+			if (OldBlueprint->SupportsNativization() != NewBlueprint->SupportsNativization())
+			{
+				FText DiffText = FText::Format(LOCTEXT("NativizationChanged", "Nativization changed from {0} to {1}"), FText::AsNumber(OldBlueprint->SupportsNativization()), FText::AsNumber(NewBlueprint->SupportsNativization()));
+
+				auto Entry = TSharedPtr<FBlueprintDifferenceTreeEntry>(new FBlueprintDifferenceTreeEntry(
+					FOnDiffEntryFocused::CreateStatic(FocusSettings, SelectionCallback)
+					, FGenerateDiffEntryWidget::CreateStatic(GenerateDiffWidget, DiffText)
+					, TArray< TSharedPtr<FBlueprintDifferenceTreeEntry> >()
+				));
+
+				Children.Push(Entry);
+				OutRealDifferences.Push(Entry);
+			}
+		}
+
+		const bool bHasDiffferences = Children.Num() != 0;
+		if (!bHasDiffferences)
+		{
+			// make one child informing the user that there are no differences:
+			Children.Push(FBlueprintDifferenceTreeEntry::NoDifferencesEntry());
+		}
+
+		const auto ForwardSelection = [](FOnDetailsDiffControlChanged InSelectionCallback)
+		{
+			// This allows the owning control to focus the correct tab (or do whatever else it likes):
+			InSelectionCallback.ExecuteIfBound();
+		};
+
+		OutTreeEntries.Push(FBlueprintDifferenceTreeEntry::CreateClassSettingsCategoryEntry(
+			FOnDiffEntryFocused::CreateStatic(ForwardSelection, SelectionCallback)
+			, Children
+			, bHasDiffferences
+		));
+	}
+};
 
 /*List item that entry for a graph*/
 struct KISMET_API FListItemGraphToDiff	: public TSharedFromThis<FListItemGraphToDiff>
@@ -1314,6 +1438,7 @@ void SBlueprintDiff::GenerateDifferencesList()
 	// Unfortunately we can't perform the diff until the UI is generated, the primary reason for this is that
 	// details customizations determine what is actually editable:
 	DefaultsPanel = GenerateDefaultsPanel();
+	ClassSettingsPanel = GenerateClassSettingsPanel();
 	if (bHasComponents)
 	{
 		ComponentsPanel = GenerateComponentsPanel();
@@ -1433,7 +1558,40 @@ SBlueprintDiff::FDiffControl SBlueprintDiff::GenerateDefaultsPanel()
 	const UObject* A = DiffUtils::GetCDO(PanelOld.Blueprint);
 	const UObject* B = DiffUtils::GetCDO(PanelNew.Blueprint);
 
-	auto NewDiffControl = TSharedPtr<FCDODiffControl>(new FCDODiffControl(A, B, MasterDifferencesList, RealDifferences, FOnCDODiffControlChanged::CreateRaw(this, &SBlueprintDiff::SetCurrentMode, FBlueprintEditorApplicationModes::BlueprintDefaultsMode)));
+	auto NewDiffControl = TSharedPtr<FCDODiffControl>(new FCDODiffControl(A, B, MasterDifferencesList, RealDifferences, FOnDetailsDiffControlChanged::CreateRaw(this, &SBlueprintDiff::SetCurrentMode, FBlueprintEditorApplicationModes::BlueprintDefaultsMode)));
+
+	SBlueprintDiff::FDiffControl Ret;
+	//Splitter for left and right blueprint. Current convention is for the local (probably newer?) blueprint to be on the right:
+	Ret.DiffControl = NewDiffControl;
+	Ret.Widget = SNew(SSplitter)
+		+ SSplitter::Slot()
+		.Value(0.5f)
+		[
+			SNew(SBorder)
+			.VAlign(VAlign_Fill)
+			[
+				NewDiffControl->OldDetailsWidget()
+			]
+		]
+	+ SSplitter::Slot()
+		.Value(0.5f)
+		[
+			SNew(SBorder)
+			.VAlign(VAlign_Fill)
+			[
+				NewDiffControl->NewDetailsWidget()
+			]
+		];
+
+	return Ret;
+}
+
+SBlueprintDiff::FDiffControl SBlueprintDiff::GenerateClassSettingsPanel()
+{
+	const UObject* A = PanelOld.Blueprint;
+	const UObject* B = PanelNew.Blueprint;
+
+	auto NewDiffControl = TSharedPtr<FClassDefaultsDiffControl>(new FClassDefaultsDiffControl(A, B, MasterDifferencesList, RealDifferences, FOnDetailsDiffControlChanged::CreateRaw(this, &SBlueprintDiff::SetCurrentMode, ClassSettingsMode)));
 
 	SBlueprintDiff::FDiffControl Ret;
 	//Splitter for left and right blueprint. Current convention is for the local (probably newer?) blueprint to be on the right:
@@ -1511,6 +1669,11 @@ void SBlueprintDiff::SetCurrentMode(FName NewMode)
 		DiffControl = DefaultsPanel.DiffControl;
 		ModeContents->SetContent( DefaultsPanel.Widget.ToSharedRef() );
 	}
+	else if (NewMode == ClassSettingsMode)
+	{
+		DiffControl = ClassSettingsPanel.DiffControl;
+		ModeContents->SetContent(ClassSettingsPanel.Widget.ToSharedRef());
+	}
 	else if (NewMode == FBlueprintEditorApplicationModes::BlueprintComponentsMode)
 	{
 		DiffControl = ComponentsPanel.DiffControl;
diff --git a/Engine/Source/Editor/Kismet/Public/DiffUtils.h b/Engine/Source/Editor/Kismet/Public/DiffUtils.h
--- a/Engine/Source/Editor/Kismet/Public/DiffUtils.h
+++ b/Engine/Source/Editor/Kismet/Public/DiffUtils.h
@@ -240,6 +240,8 @@ public:
 	KISMET_API static TSharedPtr<FBlueprintDifferenceTreeEntry> CreateDefaultsCategoryEntryForMerge(FOnDiffEntryFocused FocusCallback, const TArray< TSharedPtr<FBlueprintDifferenceTreeEntry> >& Children, bool bHasRemoteDifferences, bool bHasLocalDifferences, bool bHasConflicts );
 	KISMET_API static TSharedPtr<FBlueprintDifferenceTreeEntry> CreateComponentsCategoryEntry(FOnDiffEntryFocused FocusCallback, const TArray< TSharedPtr<FBlueprintDifferenceTreeEntry> >& Children, bool bHasDifferences);
 	KISMET_API static TSharedPtr<FBlueprintDifferenceTreeEntry> CreateComponentsCategoryEntryForMerge(FOnDiffEntryFocused FocusCallback, const TArray< TSharedPtr<FBlueprintDifferenceTreeEntry> >& Children, bool bHasRemoteDifferences, bool bHasLocalDifferences, bool bHasConflicts);
+	KISMET_API static TSharedPtr<FBlueprintDifferenceTreeEntry> CreateClassSettingsCategoryEntry(FOnDiffEntryFocused FocusCallback, const TArray< TSharedPtr<FBlueprintDifferenceTreeEntry> >& Children, bool bHasDifferences);
+	KISMET_API static TSharedPtr<FBlueprintDifferenceTreeEntry> CreateClassSettingsCategoryEntryForMerge(FOnDiffEntryFocused FocusCallback, const TArray< TSharedPtr<FBlueprintDifferenceTreeEntry> >& Children, bool bHasRemoteDifferences, bool bHasLocalDifferences, bool bHasConflicts);
 	
 	FOnDiffEntryFocused OnFocus;
 	FGenerateDiffEntryWidget GenerateWidget;
diff --git a/Engine/Source/Editor/Kismet/Public/SBlueprintDiff.h b/Engine/Source/Editor/Kismet/Public/SBlueprintDiff.h
--- a/Engine/Source/Editor/Kismet/Public/SBlueprintDiff.h
+++ b/Engine/Source/Editor/Kismet/Public/SBlueprintDiff.h
@@ -208,6 +208,7 @@ protected:
 
 	FDiffControl GenerateGraphPanel();
 	FDiffControl GenerateDefaultsPanel();
+	FDiffControl GenerateClassSettingsPanel();
 	FDiffControl GenerateComponentsPanel();
 
 	/** Accessor and event handler for toggling between diff view modes (defaults, components, graph view, interface, macro): */
@@ -245,6 +246,7 @@ protected:
 	/** Stored references to widgets used to display various parts of a blueprint: */
 	FDiffControl GraphPanel;
 	FDiffControl DefaultsPanel;
+	FDiffControl ClassSettingsPanel;
 	FDiffControl ComponentsPanel;
 
 	/** A pointer to the window holding this */
