- Fix ensure on ActorNode when clearing filter

#rnx

diff --git a/Engine/Source/Editor/Kismet/Private/SSCSEditor.cpp b/Engine/Source/Editor/Kismet/Private/SSCSEditor.cpp
--- a/Engine/Source/Editor/Kismet/Private/SSCSEditor.cpp
+++ b/Engine/Source/Editor/Kismet/Private/SSCSEditor.cpp
@@ -4540,6 +4540,7 @@ void SSCSEditor::UpdateTree(bool bRegenerateTreeNodes)
 		RootNodes.Empty();
 
 		TSharedPtr<FSCSEditorTreeNode> ActorTreeNode = MakeShareable(new FSCSEditorTreeNodeRootActor(GetActorContext(),EditorMode == EComponentEditorMode::ActorInstance));
+		RefreshFilteredState(ActorTreeNode, false);
 		SCSTreeWidget->SetItemExpansion(ActorTreeNode, true);
 		RootNodes.Add(ActorTreeNode);
 
