Don't create a child actor component when dragging and dropping an actor class on to the components panel if that actor class is not placeable.

diff --git a/Engine/Source/Editor/Kismet/Private/SSCSEditor.cpp b/Engine/Source/Editor/Kismet/Private/SSCSEditor.cpp
--- a/Engine/Source/Editor/Kismet/Private/SSCSEditor.cpp
+++ b/Engine/Source/Editor/Kismet/Private/SSCSEditor.cpp
@@ -211,7 +211,7 @@ FReply SSCSEditor::TryHandleAssetDragDropOperation(const FDragDropEvent& DragDro
 						bMarkBlueprintAsModified = true;
 					}
 				}
-				else if ((PotentialActorClass != nullptr) && !PotentialActorClass->HasAnyClassFlags(CLASS_Deprecated | CLASS_Abstract | CLASS_NewerVersionExists))
+				else if ((PotentialActorClass != nullptr) && !PotentialActorClass->HasAnyClassFlags(CLASS_Deprecated | CLASS_Abstract | CLASS_NewerVersionExists | CLASS_NotPlaceable))
 				{
 					AddNewComponent(UChildActorComponent::StaticClass(), PotentialActorClass, true, bSetFocusToNewItem );
 					bMarkBlueprintAsModified = true;
