Several other layout changes, including moving the diff labels to the header section

diff --git a/Engine/Source/Developer/Merge/Private/SMergeDetailsView.cpp b/Engine/Source/Developer/Merge/Private/SMergeDetailsView.cpp
--- a/Engine/Source/Developer/Merge/Private/SMergeDetailsView.cpp
+++ b/Engine/Source/Developer/Merge/Private/SMergeDetailsView.cpp
@@ -242,7 +242,11 @@ void SMergeDetailsView::Construct(const FArguments InArgs
 		InSelectionCallback.ExecuteIfBound();
 	};
 
-	TSharedPtr<FBlueprintDifferenceTreeEntry> Category = FBlueprintDifferenceTreeEntry::CreateDefaultsCategoryEntryForMerge(FOnDiffEntryFocused::CreateStatic(ForwardSelection, SelectionCallback), Children, RemoteDifferences.Num() != 0, LocalDifferences.Num() != 0, bAnyConflict);
+	TSharedPtr<FBlueprintDifferenceTreeEntry> Category = FBlueprintDifferenceTreeEntry::CreateCategoryEntryForMerge(
+		NSLOCTEXT("FBlueprintDifferenceTreeEntry", "DefaultsLabel", "Defaults"),
+		NSLOCTEXT("FBlueprintDifferenceTreeEntry", "DefaultsTooltip", "The list of changes made in the Defaults panel"),
+		FOnDiffEntryFocused::CreateStatic(ForwardSelection, SelectionCallback), 
+		Children, RemoteDifferences.Num() != 0, LocalDifferences.Num() != 0, bAnyConflict);
 	OutTreeEntries.Push(Category);
 
 	CurrentDifference = -1;
diff --git a/Engine/Source/Developer/Merge/Private/SMergeGraphView.cpp b/Engine/Source/Developer/Merge/Private/SMergeGraphView.cpp
--- a/Engine/Source/Developer/Merge/Private/SMergeGraphView.cpp
+++ b/Engine/Source/Developer/Merge/Private/SMergeGraphView.cpp
@@ -584,7 +584,7 @@ TSharedRef<SDockTab> SMergeGraphView::CreateMyBlueprintsViews(const FSpawnTabArg
 	{
 		PanelContainer->AddSlot()
 		[
-			Panel.GenerateMyBlueprintPanel()
+			Panel.GenerateMyBlueprintWidget()
 		];
 	}
 
diff --git a/Engine/Source/Developer/Merge/Private/SMergeTreeView.cpp b/Engine/Source/Developer/Merge/Private/SMergeTreeView.cpp
--- a/Engine/Source/Developer/Merge/Private/SMergeTreeView.cpp
+++ b/Engine/Source/Developer/Merge/Private/SMergeTreeView.cpp
@@ -243,12 +243,14 @@ void SMergeTreeView::Construct(const FArguments InArgs
 		Children.Push( FBlueprintDifferenceTreeEntry::NoDifferencesEntry() );
 	}
 
-	TSharedPtr<FBlueprintDifferenceTreeEntry> Category = FBlueprintDifferenceTreeEntry::CreateComponentsCategoryEntryForMerge(
-		FOnDiffEntryFocused::CreateStatic(ForwardSelection, SelectionCallback)
-		, Children
-		, RemoteDifferingProperties.Entries.Num() != 0
-		, LocalDifferingProperties.Entries.Num() != 0
-		, bAnyConflict);
+	TSharedPtr<FBlueprintDifferenceTreeEntry> Category = FBlueprintDifferenceTreeEntry::CreateCategoryEntryForMerge(
+		NSLOCTEXT("FBlueprintDifferenceTreeEntry", "SCSLabel", "Components"),
+		NSLOCTEXT("FBlueprintDifferenceTreeEntry", "SCSTooltip", "The list of changes made in the Components panel"),
+		FOnDiffEntryFocused::CreateStatic(ForwardSelection, SelectionCallback),
+		Children,
+		RemoteDifferingProperties.Entries.Num() != 0,
+		LocalDifferingProperties.Entries.Num() != 0,
+		bAnyConflict);
 	OutTreeEntries.Push(Category);
 
 	ChildSlot[
diff --git a/Engine/Source/Editor/Kismet/Private/DetailsDiff.cpp b/Engine/Source/Editor/Kismet/Private/DetailsDiff.cpp
--- a/Engine/Source/Editor/Kismet/Private/DetailsDiff.cpp
+++ b/Engine/Source/Editor/Kismet/Private/DetailsDiff.cpp
@@ -14,6 +14,7 @@ FDetailsDiff::FDetailsDiff(const UObject* InObject, FOnDisplayedPropertiesChange
 {
 	FDetailsViewArgs DetailsViewArgs;
 	DetailsViewArgs.bShowDifferingPropertiesOption = true;
+	DetailsViewArgs.NameAreaSettings = FDetailsViewArgs::HideNameArea;
 
 	FPropertyEditorModule& EditModule = FModuleManager::Get().GetModuleChecked<FPropertyEditorModule>("PropertyEditor");
 	DetailsView = EditModule.CreateDetailView(DetailsViewArgs);
diff --git a/Engine/Source/Editor/Kismet/Private/DiffUtils.cpp b/Engine/Source/Editor/Kismet/Private/DiffUtils.cpp
--- a/Engine/Source/Editor/Kismet/Private/DiffUtils.cpp
+++ b/Engine/Source/Editor/Kismet/Private/DiffUtils.cpp
@@ -779,36 +779,35 @@ TSharedPtr<FBlueprintDifferenceTreeEntry> FBlueprintDifferenceTreeEntry::WidgetB
 		, Children
 		));
 }
-
-TSharedPtr<FBlueprintDifferenceTreeEntry> FBlueprintDifferenceTreeEntry::CreateDefaultsCategoryEntry(FOnDiffEntryFocused FocusCallback, const TArray< TSharedPtr<FBlueprintDifferenceTreeEntry> >& Children, bool bHasDifferences)
+TSharedPtr<FBlueprintDifferenceTreeEntry> FBlueprintDifferenceTreeEntry::CreateCategoryEntry(const FText& LabelText, const FText& ToolTipText, FOnDiffEntryFocused FocusCallback, const TArray< TSharedPtr<FBlueprintDifferenceTreeEntry> >& Children, bool bHasDifferences)
 {
-	const auto CreateDefaultsRootEntry = [](FLinearColor Color) -> TSharedRef<SWidget>
+	const auto CreateDefaultsRootEntry = [](FText LabelText, FText ToolTipText, FLinearColor Color) -> TSharedRef<SWidget>
 	{
 		return SNew(STextBlock)
-			.ToolTipText(NSLOCTEXT("FBlueprintDifferenceTreeEntry", "DefaultsTooltip", "The list of changes made in the Defaults panel"))
+			.ToolTipText(ToolTipText)
 			.ColorAndOpacity(Color)
-			.Text(NSLOCTEXT("FBlueprintDifferenceTreeEntry", "DefaultsLabel", "Defaults"));
+			.Text(LabelText);
 	};
 
 	return TSharedPtr<FBlueprintDifferenceTreeEntry>(new FBlueprintDifferenceTreeEntry(
 		FocusCallback
-		, FGenerateDiffEntryWidget::CreateStatic(CreateDefaultsRootEntry, DiffViewUtils::LookupColor(bHasDifferences) )
+		, FGenerateDiffEntryWidget::CreateStatic(CreateDefaultsRootEntry, LabelText, ToolTipText, DiffViewUtils::LookupColor(bHasDifferences))
 		, Children
 	));
 }
 
-TSharedPtr<FBlueprintDifferenceTreeEntry> FBlueprintDifferenceTreeEntry::CreateDefaultsCategoryEntryForMerge(FOnDiffEntryFocused FocusCallback, const TArray< TSharedPtr<FBlueprintDifferenceTreeEntry> >& Children, bool bHasRemoteDifferences, bool bHasLocalDifferences, bool bHasConflicts)
+TSharedPtr<FBlueprintDifferenceTreeEntry> FBlueprintDifferenceTreeEntry::CreateCategoryEntryForMerge(const FText& LabelText, const FText& ToolTipText, FOnDiffEntryFocused FocusCallback, const TArray< TSharedPtr<FBlueprintDifferenceTreeEntry> >& Children, bool bHasRemoteDifferences, bool bHasLocalDifferences, bool bHasConflicts)
 {
-	const auto CreateDefaultsRootEntry = [](bool bInHasRemoteDifferences, bool bInHasLocalDifferences, bool bInHasConflicts) -> TSharedRef<SWidget>
+	const auto CreateDefaultsRootEntry = [](FText LabelText, FText ToolTipText, bool bInHasRemoteDifferences, bool bInHasLocalDifferences, bool bInHasConflicts) -> TSharedRef<SWidget>
 	{
 		const FLinearColor BaseColor = DiffViewUtils::LookupColor(bInHasRemoteDifferences || bInHasLocalDifferences, bInHasConflicts);
 		return SNew(SHorizontalBox)
 			+ SHorizontalBox::Slot()
 			[
 				SNew(STextBlock)
-				.ToolTipText(NSLOCTEXT("FBlueprintDifferenceTreeEntry", "DefaultsTooltip", "The list of changes made in the Defaults panel"))
+				.ToolTipText(ToolTipText)
 				.ColorAndOpacity(BaseColor)
-				.Text(NSLOCTEXT("FBlueprintDifferenceTreeEntry", "DefaultsLabel", "Defaults"))
+				.Text(LabelText)
 			]
 			+ DiffViewUtils::Box(true, DiffViewUtils::LookupColor(bInHasRemoteDifferences, bInHasConflicts))
 			+ DiffViewUtils::Box(true, BaseColor)
@@ -817,91 +816,7 @@ TSharedPtr<FBlueprintDifferenceTreeEntry> FBlueprintDifferenceTreeEntry::CreateD
 
 	return TSharedPtr<FBlueprintDifferenceTreeEntry>(new FBlueprintDifferenceTreeEntry(
 		FocusCallback
-		, FGenerateDiffEntryWidget::CreateStatic(CreateDefaultsRootEntry, bHasRemoteDifferences, bHasLocalDifferences, bHasConflicts)
-		, Children
-		));
-}
-
-TSharedPtr<FBlueprintDifferenceTreeEntry> FBlueprintDifferenceTreeEntry::CreateComponentsCategoryEntry(FOnDiffEntryFocused FocusCallback, const TArray< TSharedPtr<FBlueprintDifferenceTreeEntry> >& Children, bool bHasDifferences)
-{
-	const auto CreateComponentsRootEntry = [](FLinearColor Color) -> TSharedRef<SWidget>
-	{
-		return SNew(STextBlock)
-			.ToolTipText(NSLOCTEXT("FBlueprintDifferenceTreeEntry", "SCSTooltip", "The list of changes made in the Components panel"))
-			.ColorAndOpacity(Color)
-			.Text(NSLOCTEXT("FBlueprintDifferenceTreeEntry", "SCSLabel", "Components"));
-	};
-
-	return TSharedPtr<FBlueprintDifferenceTreeEntry>(new FBlueprintDifferenceTreeEntry(
-		FocusCallback
-		, FGenerateDiffEntryWidget::CreateStatic(CreateComponentsRootEntry, DiffViewUtils::LookupColor(bHasDifferences))
-		, Children
-		));
-}
-
-TSharedPtr<FBlueprintDifferenceTreeEntry> FBlueprintDifferenceTreeEntry::CreateComponentsCategoryEntryForMerge(FOnDiffEntryFocused FocusCallback, const TArray< TSharedPtr<FBlueprintDifferenceTreeEntry> >& Children, bool bHasRemoteDifferences, bool bHasLocalDifferences, bool bHasConflicts)
-{
-	const auto CreateComponentsRootEntry = [](bool bInHasRemoteDifferences, bool bInHasLocalDifferences, bool bInHasConflicts) -> TSharedRef<SWidget>
-	{
-		const FLinearColor BaseColor = DiffViewUtils::LookupColor(bInHasRemoteDifferences || bInHasLocalDifferences, bInHasConflicts);
-		return  SNew(SHorizontalBox)
-			+ SHorizontalBox::Slot()
-			[
-				SNew(STextBlock)
-				.ToolTipText(NSLOCTEXT("FBlueprintDifferenceTreeEntry", "SCSTooltip", "The list of changes made in the Components panel"))
-				.ColorAndOpacity(BaseColor)
-				.Text(NSLOCTEXT("FBlueprintDifferenceTreeEntry", "SCSLabel", "Components"))
-			]
-			+ DiffViewUtils::Box(true, DiffViewUtils::LookupColor(bInHasRemoteDifferences, bInHasConflicts))
-			+ DiffViewUtils::Box(true, BaseColor)
-			+ DiffViewUtils::Box(true, DiffViewUtils::LookupColor(bInHasLocalDifferences, bInHasConflicts));
-	};
-
-	return TSharedPtr<FBlueprintDifferenceTreeEntry>(new FBlueprintDifferenceTreeEntry(
-		FocusCallback
-		, FGenerateDiffEntryWidget::CreateStatic(CreateComponentsRootEntry, bHasRemoteDifferences, bHasLocalDifferences, bHasConflicts)
-		, Children
-	));
-}
-
-TSharedPtr<FBlueprintDifferenceTreeEntry> FBlueprintDifferenceTreeEntry::CreateClassSettingsCategoryEntry(FOnDiffEntryFocused FocusCallback, const TArray< TSharedPtr<FBlueprintDifferenceTreeEntry> >& Children, bool bHasDifferences)
-{
-	const auto CreateSettingsRootEntry = [](FLinearColor Color) -> TSharedRef<SWidget>
-	{
-		return SNew(STextBlock)
-			.ToolTipText(NSLOCTEXT("FBlueprintDifferenceTreeEntry", "SettingsTooltip", "The list of changes made to class settings"))
-			.ColorAndOpacity(Color)
-			.Text(NSLOCTEXT("FBlueprintDifferenceTreeEntry", "SettingsLabel", "Class Settings"));
-	};
-
-	return TSharedPtr<FBlueprintDifferenceTreeEntry>(new FBlueprintDifferenceTreeEntry(
-		FocusCallback
-		, FGenerateDiffEntryWidget::CreateStatic(CreateSettingsRootEntry, DiffViewUtils::LookupColor(bHasDifferences))
-		, Children
-	));
-}
-
-TSharedPtr<FBlueprintDifferenceTreeEntry> FBlueprintDifferenceTreeEntry::CreateClassSettingsCategoryEntryForMerge(FOnDiffEntryFocused FocusCallback, const TArray< TSharedPtr<FBlueprintDifferenceTreeEntry> >& Children, bool bHasRemoteDifferences, bool bHasLocalDifferences, bool bHasConflicts)
-{
-	const auto CreateSettingsRootEntry = [](bool bInHasRemoteDifferences, bool bInHasLocalDifferences, bool bInHasConflicts) -> TSharedRef<SWidget>
-	{
-		const FLinearColor BaseColor = DiffViewUtils::LookupColor(bInHasRemoteDifferences || bInHasLocalDifferences, bInHasConflicts);
-		return  SNew(SHorizontalBox)
-			+ SHorizontalBox::Slot()
-			[
-				SNew(STextBlock)
-				.ToolTipText(NSLOCTEXT("FBlueprintDifferenceTreeEntry", "SettingsTooltip", "The list of changes made to class settings"))
-			.ColorAndOpacity(BaseColor)
-			.Text(NSLOCTEXT("FBlueprintDifferenceTreeEntry", "SettingsLabel", "Class Settings"))
-			]
-		+ DiffViewUtils::Box(true, DiffViewUtils::LookupColor(bInHasRemoteDifferences, bInHasConflicts))
-			+ DiffViewUtils::Box(true, BaseColor)
-			+ DiffViewUtils::Box(true, DiffViewUtils::LookupColor(bInHasLocalDifferences, bInHasConflicts));
-	};
-
-	return TSharedPtr<FBlueprintDifferenceTreeEntry>(new FBlueprintDifferenceTreeEntry(
-		FocusCallback
-		, FGenerateDiffEntryWidget::CreateStatic(CreateSettingsRootEntry, bHasRemoteDifferences, bHasLocalDifferences, bHasConflicts)
+		, FGenerateDiffEntryWidget::CreateStatic(CreateDefaultsRootEntry, LabelText, ToolTipText, bHasRemoteDifferences, bHasLocalDifferences, bHasConflicts)
 		, Children
 	));
 }
@@ -1109,19 +1024,37 @@ FText DiffViewUtils::GetPanelLabel(const UBlueprint* Blueprint, const FRevisionI
 				, FText::FromString(Revision.Date.ToString(TEXT("%m/%d/%Y"))));		
 		}
 
-		return FText::Format( NSLOCTEXT("DiffViewUtils", "RevisionLabel", "{0}\n{1}\n{2}")
-			, Label
-			, FText::FromString( Blueprint->GetName() )
-			, RevisionData );
+		if (Label.IsEmpty())
+		{
+			return FText::Format(NSLOCTEXT("DiffViewUtils", "RevisionLabelTwoLines", "{0}\n{1}")
+				, FText::FromString(Blueprint->GetName())
+				, RevisionData);
+		}
+		else
+		{
+			return FText::Format(NSLOCTEXT("DiffViewUtils", "RevisionLabel", "{0}\n{1}\n{2}")
+				, Label
+				, FText::FromString(Blueprint->GetName())
+				, RevisionData);
+		}
 	}
 	else
 	{
 		if( Blueprint )
 		{
-			return FText::Format( NSLOCTEXT("DiffViewUtils", "RevisionLabel", "{0}\n{1}\n{2}")
-				, Label
-				, FText::FromString( Blueprint->GetName() )
-				, NSLOCTEXT("DiffViewUtils", "LocalRevisionLabel", "Local Revision" ));
+			if (Label.IsEmpty())
+			{
+				return FText::Format(NSLOCTEXT("DiffViewUtils", "RevisionLabelTwoLines", "{0}\n{1}")
+					, FText::FromString(Blueprint->GetName())
+					, NSLOCTEXT("DiffViewUtils", "LocalRevisionLabel", "Local Revision"));
+			}
+			else
+			{
+				return FText::Format(NSLOCTEXT("DiffViewUtils", "RevisionLabel", "{0}\n{1}\n{2}")
+					, Label
+					, FText::FromString(Blueprint->GetName())
+					, NSLOCTEXT("DiffViewUtils", "LocalRevisionLabel", "Local Revision"));
+			}
 		}
 
 		return NSLOCTEXT("DiffViewUtils", "NoBlueprint", "None" );
diff --git a/Engine/Source/Editor/Kismet/Private/SBlueprintDiff.cpp b/Engine/Source/Editor/Kismet/Private/SBlueprintDiff.cpp
--- a/Engine/Source/Editor/Kismet/Private/SBlueprintDiff.cpp
+++ b/Engine/Source/Editor/Kismet/Private/SBlueprintDiff.cpp
@@ -8,16 +8,15 @@
 #include "Widgets/Layout/SSpacer.h"
 #include "Framework/MultiBox/MultiBoxDefs.h"
 #include "Framework/MultiBox/MultiBoxBuilder.h"
-#include "Framework/Docking/TabManager.h"
 #include "EditorStyleSet.h"
 #include "Animation/AnimBlueprint.h"
 #include "K2Node_MathExpression.h"
 #include "Kismet2/BlueprintEditorUtils.h"
+#include "Kismet2/KismetEditorUtilities.h"
 #include "BlueprintEditorModes.h"
 #include "DetailsDiff.h"
 #include "EdGraphUtilities.h"
 #include "GraphDiffControl.h"
-#include "Widgets/Docking/SDockTab.h"
 #include "SMyBlueprint.h"
 #include "SCSDiff.h"
 #include "WorkflowOrientedApp/SModeWidget.h"
@@ -35,20 +34,100 @@ public:
 	virtual ~IDiffControl() {}
 };
 
+DECLARE_DELEGATE(FOnDiffControlChanged);
+
 FText RightRevision = LOCTEXT("OlderRevisionIdentifier", "Right Revision");
 
 typedef TMap< FName, const UProperty* > FNamePropertyMap;
 
-const FName DiffMyBluerpintTabId = FName(TEXT("DiffMyBluerpintTab"));
-const FName DiffGraphTabId = FName(TEXT("DiffGraphTab"));
+const FName MyBlueprintMode = FName(TEXT("MyBlueprintMode"));
+const FName DefaultsMode = FName(TEXT("DefaultsMode"));
 const FName ClassSettingsMode = FName(TEXT("ClassSettingsMode"));
+const FName ComponentsMode = FName(TEXT("ComponentsMode"));
+const FName GraphMode = FName(TEXT("GraphMode"));
+
+/** Shows all differences for the blueprint structure itself that aren't picked up elsewhere */
+class FMyBlueprintDiffControl : public TSharedFromThis<FMyBlueprintDiffControl>, public IDiffControl
+{
+public:
+	FMyBlueprintDiffControl(
+		const UBlueprint* InOldBlueprint
+		, const UBlueprint* InNewBlueprint
+		, TArray< TSharedPtr<FBlueprintDifferenceTreeEntry> >& OutTreeEntries
+		, TArray< TSharedPtr<FBlueprintDifferenceTreeEntry> >& OutRealDifferences
+		, FOnDiffControlChanged SelectionCallback
+	);
+
+	virtual ~FMyBlueprintDiffControl() {}
+};
+
+FMyBlueprintDiffControl::FMyBlueprintDiffControl(
+	const UBlueprint* InOldBlueprint
+	, const UBlueprint* InNewBlueprint
+	, TArray< TSharedPtr<FBlueprintDifferenceTreeEntry> >& OutTreeEntries
+	, TArray< TSharedPtr<FBlueprintDifferenceTreeEntry> >& OutRealDifferences
+	, FOnDiffControlChanged SelectionCallback
+)
+{
+	TArray< TSharedPtr<FBlueprintDifferenceTreeEntry> > Children;
+
+	/*const auto FocusSCSDifferenceEntry = [](FSCSDiffEntry Entry, FOnDiffControlChanged InSelectionCallback, FMyBlueprintDiffControl* Owner)
+	{
+		InSelectionCallback.ExecuteIfBound();
+		if (Entry.TreeIdentifier.Name != NAME_None)
+		{
+			Owner->OldSCS.HighlightProperty(Entry.TreeIdentifier.Name, FPropertyPath());
+			Owner->NewSCS.HighlightProperty(Entry.TreeIdentifier.Name, FPropertyPath());
+		}
+	};
+
+	const auto CreateSCSDifferenceWidget = [](FSCSDiffEntry Entry, FText ObjectName) -> TSharedRef<SWidget>
+	{
+		return SNew(STextBlock)
+			.Text(DiffViewUtils::SCSDiffMessage(Entry, ObjectName))
+			.ColorAndOpacity(DiffViewUtils::Differs());
+	};
 
-DECLARE_DELEGATE(FOnSCSDiffControlChanged);
+	for (auto Difference : DifferingProperties.Entries)
+	{
+		auto Entry = TSharedPtr<FBlueprintDifferenceTreeEntry>(
+			new FBlueprintDifferenceTreeEntry(
+				FOnDiffEntryFocused::CreateStatic(FocusSCSDifferenceEntry, Difference, SelectionCallback, this)
+				, FGenerateDiffEntryWidget::CreateStatic(CreateSCSDifferenceWidget, Difference, RightRevision)
+				, TArray< TSharedPtr<FBlueprintDifferenceTreeEntry> >()
+			)
+			);
+		Children.Push(Entry);
+		OutRealDifferences.Push(Entry);
+	}*/
+
+	const bool bHasDiffferences = Children.Num() != 0;
+	if (!bHasDiffferences)
+	{
+		// make one child informing the user that there are no differences:
+		Children.Push(FBlueprintDifferenceTreeEntry::NoDifferencesEntry());
+	}
+
+	const auto ForwardSelection = [](FOnDiffControlChanged InSelectionCallback)
+	{
+		// This allows the owning control to focus the correct tab (or do whatever else it likes):
+		InSelectionCallback.ExecuteIfBound();
+	};
+
+	OutTreeEntries.Push(FBlueprintDifferenceTreeEntry::CreateCategoryEntry(
+		NSLOCTEXT("FBlueprintDifferenceTreeEntry", "MyBlueprintLabel", "My Blueprint"),
+		NSLOCTEXT("FBlueprintDifferenceTreeEntry", "MyBlueprintTooltip", "The list of changes made to blueprint structure in the My Blueprint panel"),
+		FOnDiffEntryFocused::CreateStatic(ForwardSelection, SelectionCallback),
+		Children,
+		bHasDiffferences
+	));
+}
 
-// Each difference in the tree will either be a tree node that is added in one Blueprint 
-// or a tree node and an FName of a property that has been added or edited in one Blueprint
-class FSCSDiffControl	: public TSharedFromThis< FSCSDiffControl >
-						, public IDiffControl
+/** 
+ * Each difference in the tree will either be a tree node that is added in one Blueprint 
+ * or a tree node and an FName of a property that has been added or edited in one Blueprint
+ */
+class FSCSDiffControl : public TSharedFromThis<FSCSDiffControl>, public IDiffControl
 {
 public:
 	FSCSDiffControl(
@@ -56,7 +135,7 @@ public:
 					, const UBlueprint* InNewBlueprint
 					, TArray< TSharedPtr<FBlueprintDifferenceTreeEntry> >& OutTreeEntries
 					, TArray< TSharedPtr<FBlueprintDifferenceTreeEntry> >& OutRealDifferences
-					, FOnSCSDiffControlChanged SelectionCallback
+					, FOnDiffControlChanged SelectionCallback
 					);
 
 
@@ -76,7 +155,7 @@ FSCSDiffControl::FSCSDiffControl(
 		, const UBlueprint* InNewBlueprint
 		, TArray< TSharedPtr<FBlueprintDifferenceTreeEntry> >& OutTreeEntries
 		, TArray< TSharedPtr<FBlueprintDifferenceTreeEntry> >& OutRealDifferences
-		, FOnSCSDiffControlChanged SelectionCallback
+		, FOnDiffControlChanged SelectionCallback
 	)
 	: DifferingProperties()
 	, OldSCS( InOldBlueprint )
@@ -86,7 +165,7 @@ FSCSDiffControl::FSCSDiffControl(
 	TArray< FSCSResolvedIdentifier > NewHierarchy = NewSCS.GetDisplayedHierarchy();
 	DiffUtils::CompareUnrelatedSCS(InOldBlueprint, OldHierarchy, InNewBlueprint, NewHierarchy, DifferingProperties);
 
-	const auto FocusSCSDifferenceEntry = [](FSCSDiffEntry Entry, FOnSCSDiffControlChanged InSelectionCallback, FSCSDiffControl* Owner)
+	const auto FocusSCSDifferenceEntry = [](FSCSDiffEntry Entry, FOnDiffControlChanged InSelectionCallback, FSCSDiffControl* Owner)
 	{
 		InSelectionCallback.ExecuteIfBound();
 		if (Entry.TreeIdentifier.Name != NAME_None)
@@ -124,13 +203,15 @@ FSCSDiffControl::FSCSDiffControl(
 		Children.Push(FBlueprintDifferenceTreeEntry::NoDifferencesEntry());
 	}
 
-	const auto ForwardSelection = [](FOnSCSDiffControlChanged InSelectionCallback)
+	const auto ForwardSelection = [](FOnDiffControlChanged InSelectionCallback)
 	{
 		// This allows the owning control to focus the correct tab (or do whatever else it likes):
 		InSelectionCallback.ExecuteIfBound();
 	};
 
-	OutTreeEntries.Push(FBlueprintDifferenceTreeEntry::CreateComponentsCategoryEntry(
+	OutTreeEntries.Push(FBlueprintDifferenceTreeEntry::CreateCategoryEntry(
+		NSLOCTEXT("FBlueprintDifferenceTreeEntry", "SCSLabel", "Components"),
+		NSLOCTEXT("FBlueprintDifferenceTreeEntry", "SCSTooltip", "The list of changes made in the Components panel"),
 		FOnDiffEntryFocused::CreateStatic(ForwardSelection, SelectionCallback),
 		Children,
 		bHasDiffferences
@@ -138,8 +219,6 @@ FSCSDiffControl::FSCSDiffControl(
 }
 
 /** Generic wrapper around a details view, this does not actually fill out OutTreeEntries */
-DECLARE_DELEGATE(FOnDetailsDiffControlChanged);
-
 class FDetailsDiffControl : public IDiffControl
 {
 public:
@@ -147,7 +226,7 @@ public:
 					, const UObject* InNewObject
 					, TArray< TSharedPtr<FBlueprintDifferenceTreeEntry> >& OutTreeEntries
 					, TArray< TSharedPtr<FBlueprintDifferenceTreeEntry> >& OutRealDifferences
-					, FOnDetailsDiffControlChanged SelectionCallback );
+					, FOnDiffControlChanged SelectionCallback );
 
 	TSharedRef<SWidget> OldDetailsWidget() { return OldDetails.DetailsWidget(); }
 	TSharedRef<SWidget> NewDetailsWidget() { return NewDetails.DetailsWidget(); }
@@ -168,7 +247,7 @@ FDetailsDiffControl::FDetailsDiffControl(
 		, const UObject* InNewObject
 		, TArray< TSharedPtr<FBlueprintDifferenceTreeEntry> >& OutTreeEntries
 		, TArray< TSharedPtr<FBlueprintDifferenceTreeEntry> >& OutRealDifferences
-		, FOnDetailsDiffControlChanged SelectionCallback )
+		, FOnDiffControlChanged SelectionCallback )
 	: OldDetails(InOldObject, FDetailsDiff::FOnDisplayedPropertiesChanged() )
 	, NewDetails(InNewObject, FDetailsDiff::FOnDisplayedPropertiesChanged() )
 {
@@ -267,7 +346,7 @@ FDetailsDiffControl::FDetailsDiffControl(
 			.ColorAndOpacity(DiffViewUtils::Differs());
 	};
 
-	const auto FocusDetailsDifferenceEntry = [](FPropertySoftPath Identifier, FDetailsDiffControl* Control, FOnDetailsDiffControlChanged InSelectionCallback)
+	const auto FocusDetailsDifferenceEntry = [](FPropertySoftPath Identifier, FDetailsDiffControl* Control, FOnDiffControlChanged InSelectionCallback)
 	{
 		// This allows the owning control to focus the correct tab (or do whatever else it likes):
 		InSelectionCallback.ExecuteIfBound();
@@ -302,7 +381,7 @@ public:
 		, const UObject* InNewObject
 		, TArray< TSharedPtr<FBlueprintDifferenceTreeEntry> >& OutTreeEntries
 		, TArray< TSharedPtr<FBlueprintDifferenceTreeEntry> >& OutRealDifferences
-		, FOnDetailsDiffControlChanged SelectionCallback)
+		, FOnDiffControlChanged SelectionCallback)
 		: FDetailsDiffControl(InOldObject, InNewObject, OutTreeEntries, OutRealDifferences, SelectionCallback)
 	{
 		// Parent constructor fills out Children
@@ -313,16 +392,18 @@ public:
 			Children.Push(FBlueprintDifferenceTreeEntry::NoDifferencesEntry());
 		}
 
-		const auto ForwardSelection = [](FOnDetailsDiffControlChanged InSelectionCallback)
+		const auto ForwardSelection = [](FOnDiffControlChanged InSelectionCallback)
 		{
 			// This allows the owning control to focus the correct tab (or do whatever else it likes):
 			InSelectionCallback.ExecuteIfBound();
 		};
 
-		OutTreeEntries.Push(FBlueprintDifferenceTreeEntry::CreateDefaultsCategoryEntry(
-			FOnDiffEntryFocused::CreateStatic(ForwardSelection, SelectionCallback)
-			, Children
-			, bHasDiffferences
+		OutTreeEntries.Push(FBlueprintDifferenceTreeEntry::CreateCategoryEntry(
+			NSLOCTEXT("FBlueprintDifferenceTreeEntry", "DefaultsLabel", "Defaults"),
+			NSLOCTEXT("FBlueprintDifferenceTreeEntry", "DefaultsTooltip", "The list of changes made in the Defaults panel"),
+			FOnDiffEntryFocused::CreateStatic(ForwardSelection, SelectionCallback),
+			Children,
+			bHasDiffferences
 		));
 	}
 };
@@ -335,7 +416,7 @@ public:
 		, const UObject* InNewObject
 		, TArray< TSharedPtr<FBlueprintDifferenceTreeEntry> >& OutTreeEntries
 		, TArray< TSharedPtr<FBlueprintDifferenceTreeEntry> >& OutRealDifferences
-		, FOnDetailsDiffControlChanged SelectionCallback)
+		, FOnDiffControlChanged SelectionCallback)
 		: FDetailsDiffControl(InOldObject, InNewObject, OutTreeEntries, OutRealDifferences, SelectionCallback)
 	{
 		// Check for parent class and interfaces here
@@ -350,7 +431,7 @@ public:
 				.ToolTipText(DiffText);
 		};
 
-		const auto FocusSettings = [](FOnDetailsDiffControlChanged InSelectionCallback)
+		const auto FocusSettings = [](FOnDiffControlChanged InSelectionCallback)
 		{
 			InSelectionCallback.ExecuteIfBound();
 		};
@@ -426,16 +507,18 @@ public:
 			Children.Push(FBlueprintDifferenceTreeEntry::NoDifferencesEntry());
 		}
 
-		const auto ForwardSelection = [](FOnDetailsDiffControlChanged InSelectionCallback)
+		const auto ForwardSelection = [](FOnDiffControlChanged InSelectionCallback)
 		{
 			// This allows the owning control to focus the correct tab (or do whatever else it likes):
 			InSelectionCallback.ExecuteIfBound();
 		};
 
-		OutTreeEntries.Push(FBlueprintDifferenceTreeEntry::CreateClassSettingsCategoryEntry(
-			FOnDiffEntryFocused::CreateStatic(ForwardSelection, SelectionCallback)
-			, Children
-			, bHasDiffferences
+		OutTreeEntries.Push(FBlueprintDifferenceTreeEntry::CreateCategoryEntry(
+			NSLOCTEXT("FBlueprintDifferenceTreeEntry", "SettingsLabel", "Class Settings"),
+			NSLOCTEXT("FBlueprintDifferenceTreeEntry", "SettingsTooltip", "The list of changes made in the Class Settings panel"),
+			FOnDiffEntryFocused::CreateStatic(ForwardSelection, SelectionCallback),
+			Children,
+			bHasDiffferences
 		));
 	}
 };
@@ -796,31 +879,45 @@ bool DiffWidgetUtils::HasPrevDifference(SListView< TSharedPtr< struct FDiffSingl
 BEGIN_SLATE_FUNCTION_BUILD_OPTIMIZATION
 void SBlueprintDiff::Construct( const FArguments& InArgs)
 {
-	const TSharedRef<SDockTab> MajorTab = SNew(SDockTab)
-		.TabRole(ETabRole::MajorTab);
-	TabManager = FGlobalTabmanager::Get()->NewTabManager(MajorTab);
-
-	TabManager->RegisterTabSpawner(DiffGraphTabId,
-		FOnSpawnTab::CreateRaw(this, &SBlueprintDiff::CreateGraphDiffViews ))
-		.SetDisplayName(NSLOCTEXT("SBlueprintDiff", "GraphsTabTitle", "Graphs"))
-		.SetTooltipText(NSLOCTEXT("SBlueprintDiff", "GraphsTooltipText", "Differences in the various graphs present in the blueprint"));
-
-	TabManager->RegisterTabSpawner(DiffMyBluerpintTabId,
-		FOnSpawnTab::CreateRaw(this, &SBlueprintDiff::CreateMyBlueprintsViews))
-		.SetDisplayName(NSLOCTEXT("SBlueprintDiff", "MyBlueprintTabTitle", "My Blueprint"))
-		.SetTooltipText(NSLOCTEXT("SBlueprintDiff", "MyBlueprintTooltipText", "Differences in the 'My Blueprints' attributes of the blueprint"));
-
 	check(InArgs._BlueprintOld && InArgs._BlueprintNew);
 	PanelOld.Blueprint = InArgs._BlueprintOld;
 	PanelNew.Blueprint = InArgs._BlueprintNew;
 	PanelOld.RevisionInfo = InArgs._OldRevision;
 	PanelNew.RevisionInfo = InArgs._NewRevision;
 
+	// Create a skeleton if we don't have one, this is true for revision history diffs
+	if (!PanelOld.Blueprint->SkeletonGeneratedClass)
+	{
+		FKismetEditorUtilities::GenerateBlueprintSkeleton(const_cast<UBlueprint*>(PanelOld.Blueprint));
+	}
+	
+	if (!PanelNew.Blueprint->SkeletonGeneratedClass)
+	{
+		FKismetEditorUtilities::GenerateBlueprintSkeleton(const_cast<UBlueprint*>(PanelNew.Blueprint));
+	}
+
 	// sometimes we want to clearly identify the assets being diffed (when it's
 	// not the same asset in each panel)
 	PanelOld.bShowAssetName = InArgs._ShowAssetNames;
 	PanelNew.bShowAssetName = InArgs._ShowAssetNames;
 
+	// SMyBlueprint needs to be created *before* the KismetInspector or the diffs are generated, because the KismetInspector's customizations
+	// need a reference to the SMyBlueprint widget that is controlling them...
+	const auto CreateInspector = [](TSharedPtr<SMyBlueprint> InMyBlueprint) {
+		return SNew(SKismetInspector)
+			.HideNameArea(true)
+			.ViewIdentifier(FName("BlueprintInspector"))
+			.MyBlueprintWidget(InMyBlueprint)
+			.IsPropertyEditingEnabledDelegate(FIsPropertyEditingEnabled::CreateStatic([] { return false; }));
+	};
+
+	PanelOld.GenerateMyBlueprintWidget();
+	PanelOld.DetailsView = CreateInspector(PanelOld.MyBlueprint);
+	PanelOld.MyBlueprint->SetInspector(PanelOld.DetailsView);
+	PanelNew.GenerateMyBlueprintWidget();
+	PanelNew.DetailsView = CreateInspector(PanelNew.MyBlueprint);
+	PanelNew.MyBlueprint->SetInspector(PanelNew.DetailsView);
+
 	bLockViews = true;
 
 	if (InArgs._ParentWindow.IsValid())
@@ -860,26 +957,44 @@ void SBlueprintDiff::Construct( const FArguments& InArgs)
 		, TAttribute<FSlateIcon>(this, &SBlueprintDiff::GetLockViewImage)
 	);
 
-	GraphPanel = GenerateGraphPanel();
-
 	GenerateDifferencesList();
 
 	const auto TextBlock = [](FText Text) -> TSharedRef<SWidget>
 	{
-		return SNew(STextBlock)
+		return SNew(SBox)
+		.Padding(FMargin(4.0f,10.0f))
+		.VAlign(VAlign_Center)
+		.HAlign(HAlign_Left)
+		[
+			SNew(STextBlock)
 			.Visibility(EVisibility::HitTestInvisible)
-			.TextStyle(FEditorStyle::Get(), "GraphPreview.CornerText")
-			.Text(Text);
+			.TextStyle(FEditorStyle::Get(), "DetailsView.CategoryTextStyle")
+			.Text(Text)
+		];
 	};
 
-	TSharedRef<SWidget> Overlay = SNew(SHorizontalBox)
-		+ SHorizontalBox::Slot()
+	TSharedRef<SWidget> Overlay = 
+		SNew(SSplitter)
+		+ SSplitter::Slot()
+		.Value(.2f)
 		[
-			TextBlock(DiffViewUtils::GetPanelLabel(PanelOld.Blueprint, PanelOld.RevisionInfo, FText()))
+			SNew(SBox)
 		]
-	+ SHorizontalBox::Slot()
+		+ SSplitter::Slot()
+		.Value(.8f)
 		[
-			TextBlock(DiffViewUtils::GetPanelLabel(PanelNew.Blueprint, PanelNew.RevisionInfo, FText()))
+			SNew(SSplitter)
+			.PhysicalSplitterHandleSize(10.0f)
+			+ SSplitter::Slot()
+			.Value(.5f)
+			[
+				TextBlock(DiffViewUtils::GetPanelLabel(PanelOld.Blueprint, PanelOld.RevisionInfo, FText()))
+			]
+			+ SSplitter::Slot()
+			.Value(.5f)
+			[
+				TextBlock(DiffViewUtils::GetPanelLabel(PanelNew.Blueprint, PanelNew.RevisionInfo, FText()))
+			]
 		];
 
 	this->ChildSlot
@@ -887,55 +1002,54 @@ void SBlueprintDiff::Construct( const FArguments& InArgs)
 			SNew(SBorder)
 			.BorderImage(FEditorStyle::GetBrush( "Docking.Tab", ".ContentAreaBrush" ))
 			[
-				SNew(SVerticalBox)
-				+SVerticalBox::Slot()
-				.AutoHeight()
-				.Padding(0.0f, 2.0f, 0.0f, 2.0f)
-				[
-					SNew(SHorizontalBox)
-					+ SHorizontalBox::Slot()
-					.Padding(4.f)
-					.AutoWidth()
-					[
-						ToolbarBuilder.MakeWidget()
-					]
-					+ SHorizontalBox::Slot()
-					[
-						SNew(SSpacer)
-					]
-				]
-				+ SVerticalBox::Slot()
+				SNew(SOverlay)
+				+ SOverlay::Slot()
 				[
-					SNew(SSplitter)
-					+ SSplitter::Slot()
-					.Value(.2f)
+					SNew(SVerticalBox)
+					+SVerticalBox::Slot()
+					.AutoHeight()
+					.Padding(0.0f, 2.0f, 0.0f, 2.0f)
 					[
-						SNew(SBorder)
-						.BorderImage(FEditorStyle::GetBrush("ToolPanel.GroupBorder"))
+						SNew(SHorizontalBox)
+						+ SHorizontalBox::Slot()
+						.Padding(4.f)
+						.AutoWidth()
 						[
-							DifferencesTreeView.ToSharedRef()
+							ToolbarBuilder.MakeWidget()
+						]
+						+ SHorizontalBox::Slot()
+						[
+							SNew(SSpacer)
 						]
 					]
-					+ SSplitter::Slot()
-					.Value(.8f)
+					+ SVerticalBox::Slot()
 					[
-						SNew(SOverlay)
-						+ SOverlay::Slot()
+						SNew(SSplitter)
+						+ SSplitter::Slot()
+						.Value(.2f)
 						[
-							SAssignNew(ModeContents, SBox)
+							SNew(SBorder)
+							.BorderImage(FEditorStyle::GetBrush("ToolPanel.GroupBorder"))
+							[
+								DifferencesTreeView.ToSharedRef()
+							]
 						]
-						+ SOverlay::Slot()
-						.VAlign(VAlign_Bottom)
+						+ SSplitter::Slot()
+						.Value(.8f)
 						[
-							Overlay
+							SAssignNew(ModeContents, SBox)
 						]
 					]
 				]
+				+ SOverlay::Slot()
+				.VAlign(VAlign_Top)
+				[
+					Overlay
+				]
 			]
-			
 		];
 
-	SetCurrentMode( FBlueprintEditorApplicationModes::StandardBlueprintEditorMode );
+	SetCurrentMode(MyBlueprintMode);
 
 	// Bind to blueprint changed events as they may be real in memory blueprints that will be modified
 	const_cast<UBlueprint*>(PanelNew.Blueprint)->OnChanged().AddSP(this, &SBlueprintDiff::OnBlueprintChanged);
@@ -1080,46 +1194,6 @@ bool SBlueprintDiff::HasPrevDiff() const
 	return DiffTreeView::HasPrevDifference(DifferencesTreeView.ToSharedRef(), RealDifferences);
 }
 
-TSharedRef<SDockTab> SBlueprintDiff::CreateGraphDiffViews(const FSpawnTabArgs& Args)
-{
-	return SNew(SDockTab)
-	[
-		SNew(SSplitter) 
-		+ SSplitter::Slot()
-		[
-			SAssignNew(PanelOld.GraphEditorBorder, SBox)
-			.VAlign(VAlign_Fill)
-			[
-				DefaultEmptyPanel()
-			]
-		]
-		+ SSplitter::Slot()
-		[
-			SAssignNew(PanelNew.GraphEditorBorder, SBox)
-			.VAlign(VAlign_Fill)
-			[
-				DefaultEmptyPanel()
-			]
-		]
-	];
-}
-
-TSharedRef<SDockTab> SBlueprintDiff::CreateMyBlueprintsViews(const FSpawnTabArgs& Args)
-{
-	return SNew(SDockTab)
-	[
-		SNew(SSplitter)
-		+ SSplitter::Slot()
-		[
-			PanelOld.GenerateMyBlueprintPanel()
-		]
-		+ SSplitter::Slot()
-		[
-			PanelNew.GenerateMyBlueprintPanel()
-		]
-	];
-}
-
 FListItemGraphToDiff* SBlueprintDiff::FindGraphToDiffEntry(const FString& GraphPath)
 {
 	for( const auto& Graph : Graphs )
@@ -1266,7 +1340,7 @@ void FDiffPanel::GeneratePanel(UEdGraph* Graph, UEdGraph* GraphToDiff )
 	GraphEditorBorder->SetContent(Widget.ToSharedRef());
 }
 
-TSharedRef<SWidget> FDiffPanel::GenerateMyBlueprintPanel()
+TSharedRef<class SWidget> FDiffPanel::GenerateMyBlueprintWidget()
 {
 	return SAssignNew(MyBlueprint, SMyBlueprint, TWeakPtr<FBlueprintEditor>(), Blueprint);
 }
@@ -1352,7 +1426,7 @@ FDiffPanel& SBlueprintDiff::GetDiffPanelForNode(UEdGraphNode& Node)
 
 void SBlueprintDiff::HandleGraphChanged( const FString& GraphPath )
 {
-	SetCurrentMode(FBlueprintEditorApplicationModes::StandardBlueprintEditorMode);
+	SetCurrentMode(GraphMode);
 	
 	TArray<UEdGraph*> GraphsOld, GraphsNew;
 	PanelOld.Blueprint->GetAllGraphs(GraphsOld);
@@ -1435,13 +1509,14 @@ void SBlueprintDiff::GenerateDifferencesList()
 		bHasComponents = false;
 	}
 
-	// Unfortunately we can't perform the diff until the UI is generated, the primary reason for this is that
-	// details customizations determine what is actually editable:
-	DefaultsPanel = GenerateDefaultsPanel();
-	ClassSettingsPanel = GenerateClassSettingsPanel();
+	// Now that we have done the diffs, create the panel widgets
+	ModePanels.Add(MyBlueprintMode, GenerateMyBlueprintPanel());
+	ModePanels.Add(GraphMode, GenerateGraphPanel());
+	ModePanels.Add(DefaultsMode, GenerateDefaultsPanel());
+	ModePanels.Add(ClassSettingsMode, GenerateClassSettingsPanel());
 	if (bHasComponents)
 	{
-		ComponentsPanel = GenerateComponentsPanel();
+		ModePanels.Add(ComponentsMode, GenerateComponentsPanel());
 	}
 
 	for (const auto& Graph : Graphs)
@@ -1479,74 +1554,110 @@ void SBlueprintDiff::GenerateDifferencesList()
 	DifferencesTreeView = DiffTreeView::CreateTreeView(&MasterDifferencesList);
 }
 
-SBlueprintDiff::FDiffControl SBlueprintDiff::GenerateGraphPanel()
-{
-	const TSharedRef<FTabManager::FLayout> DefaultLayout = FTabManager::NewLayout("BlueprintDiff_Layout_v1")
-	->AddArea
-	(
-		FTabManager::NewPrimaryArea()
-		->Split
-		(
-			FTabManager::NewStack()
-			->AddTab(DiffMyBluerpintTabId, ETabState::OpenedTab)
-			->AddTab(DiffGraphTabId, ETabState::OpenedTab)
-		)
-	);
 
-	// SMyBlueprint needs to be created *before* the KismetInspector, because the KismetInspector's customizations
-	// need a reference to the SMyBlueprint widget that is controlling them...
-	TSharedRef<SWidget> TabControl = TabManager->RestoreFrom(DefaultLayout, TSharedPtr<SWindow>()).ToSharedRef();
-
-	const auto CreateInspector = [](TSharedPtr<SMyBlueprint> InMyBlueprint) {
-		return SNew( SKismetInspector)
-			.HideNameArea(true)
-			.ViewIdentifier(FName("BlueprintInspector"))
-			.MyBlueprintWidget(InMyBlueprint)
-			.IsPropertyEditingEnabledDelegate(FIsPropertyEditingEnabled::CreateStatic([] { return false; }));
-	};
+SBlueprintDiff::FDiffControl SBlueprintDiff::GenerateMyBlueprintPanel()
+{
+	auto NewDiffControl = TSharedPtr<FMyBlueprintDiffControl>(new FMyBlueprintDiffControl(PanelOld.Blueprint, PanelNew.Blueprint, MasterDifferencesList, RealDifferences, FOnDiffControlChanged::CreateRaw(this, &SBlueprintDiff::SetCurrentMode, MyBlueprintMode)));
 
 	SBlueprintDiff::FDiffControl Ret;
 
-	PanelOld.DetailsView = CreateInspector(PanelOld.MyBlueprint);
-	PanelOld.MyBlueprint->SetInspector(PanelOld.DetailsView);
-	PanelNew.DetailsView = CreateInspector(PanelNew.MyBlueprint);
-	PanelNew.MyBlueprint->SetInspector(PanelNew.DetailsView);
-
-	Ret.Widget = SNew(SBorder)
-	.BorderImage(FEditorStyle::GetBrush("ToolPanel.GroupBorder"))
-	.Content()
+	Ret.DiffControl = NewDiffControl;
+	Ret.Widget = SNew(SVerticalBox)
+	+ SVerticalBox::Slot()
+	.FillHeight(1.f)
 	[
-		SNew(SVerticalBox)
-		+ SVerticalBox::Slot()
-		.FillHeight(1.f)
+		SNew(SHorizontalBox)
+		+ SHorizontalBox::Slot()
+		.FillWidth(1.f)
 		[
-			SNew(SHorizontalBox)
-			+ SHorizontalBox::Slot()
-			.FillWidth(1.f)
+			//diff window
+			SNew(SSplitter)
+			.Orientation(Orient_Vertical)
+			+SSplitter::Slot()
+			.Value(.8f)
+			[
+				SNew(SSplitter)
+				.PhysicalSplitterHandleSize(10.0f)
+				+ SSplitter::Slot()
+				[
+					PanelOld.MyBlueprint.ToSharedRef()
+				]
+				+ SSplitter::Slot()
+				[
+					PanelNew.MyBlueprint.ToSharedRef()
+				]
+			]
+			+ SSplitter::Slot()
+			.Value(.2f)
 			[
-				//diff window
 				SNew(SSplitter)
-				.Orientation(Orient_Vertical)
+				.PhysicalSplitterHandleSize(10.0f)
 				+SSplitter::Slot()
-				.Value(.8f)
 				[
-					// graph and my blueprint views:
-					TabControl
+					PanelOld.DetailsView.ToSharedRef()
 				]
 				+ SSplitter::Slot()
-				.Value(.2f)
 				[
-					SNew( SSplitter )
-					+SSplitter::Slot()
+					PanelNew.DetailsView.ToSharedRef()
+				]
+			]
+		]
+	];
+
+	return Ret;
+}
+
+SBlueprintDiff::FDiffControl SBlueprintDiff::GenerateGraphPanel()
+{
+	SBlueprintDiff::FDiffControl Ret;
+
+	Ret.Widget = SNew(SVerticalBox)
+	+ SVerticalBox::Slot()
+	.FillHeight(1.f)
+	[
+		SNew(SHorizontalBox)
+		+ SHorizontalBox::Slot()
+		.FillWidth(1.f)
+		[
+			//diff window
+			SNew(SSplitter)
+			.Orientation(Orient_Vertical)
+			+SSplitter::Slot()
+			.Value(.8f)
+			[
+				SNew(SSplitter)
+				.PhysicalSplitterHandleSize(10.0f)
+				+ SSplitter::Slot()
+				[
+					SAssignNew(PanelOld.GraphEditorBorder, SBox)
+					.VAlign(VAlign_Fill)
 					[
-						PanelOld.DetailsView.ToSharedRef()
+						DefaultEmptyPanel()
 					]
-					+ SSplitter::Slot()
+				]
+				+ SSplitter::Slot()
+				[
+					SAssignNew(PanelNew.GraphEditorBorder, SBox)
+					.VAlign(VAlign_Fill)
 					[
-						PanelNew.DetailsView.ToSharedRef()
+						DefaultEmptyPanel()
 					]
 				]
 			]
+			+ SSplitter::Slot()
+			.Value(.2f)
+			[
+				SNew(SSplitter)
+				.PhysicalSplitterHandleSize(10.0f)
+				+SSplitter::Slot()
+				[
+					PanelOld.DetailsView.ToSharedRef()
+				]
+				+ SSplitter::Slot()
+				[
+					PanelNew.DetailsView.ToSharedRef()
+				]
+			]
 		]
 	];
 
@@ -1558,29 +1669,22 @@ SBlueprintDiff::FDiffControl SBlueprintDiff::GenerateDefaultsPanel()
 	const UObject* A = DiffUtils::GetCDO(PanelOld.Blueprint);
 	const UObject* B = DiffUtils::GetCDO(PanelNew.Blueprint);
 
-	auto NewDiffControl = TSharedPtr<FCDODiffControl>(new FCDODiffControl(A, B, MasterDifferencesList, RealDifferences, FOnDetailsDiffControlChanged::CreateRaw(this, &SBlueprintDiff::SetCurrentMode, FBlueprintEditorApplicationModes::BlueprintDefaultsMode)));
+	auto NewDiffControl = TSharedPtr<FCDODiffControl>(new FCDODiffControl(A, B, MasterDifferencesList, RealDifferences, FOnDiffControlChanged::CreateRaw(this, &SBlueprintDiff::SetCurrentMode, DefaultsMode)));
 
 	SBlueprintDiff::FDiffControl Ret;
 	//Splitter for left and right blueprint. Current convention is for the local (probably newer?) blueprint to be on the right:
 	Ret.DiffControl = NewDiffControl;
 	Ret.Widget = SNew(SSplitter)
+		.PhysicalSplitterHandleSize(10.0f)
 		+ SSplitter::Slot()
 		.Value(0.5f)
 		[
-			SNew(SBorder)
-			.VAlign(VAlign_Fill)
-			[
-				NewDiffControl->OldDetailsWidget()
-			]
+			NewDiffControl->OldDetailsWidget()
 		]
-	+ SSplitter::Slot()
+		+ SSplitter::Slot()
 		.Value(0.5f)
 		[
-			SNew(SBorder)
-			.VAlign(VAlign_Fill)
-			[
-				NewDiffControl->NewDetailsWidget()
-			]
+			NewDiffControl->NewDetailsWidget()
 		];
 
 	return Ret;
@@ -1591,29 +1695,22 @@ SBlueprintDiff::FDiffControl SBlueprintDiff::GenerateClassSettingsPanel()
 	const UObject* A = PanelOld.Blueprint;
 	const UObject* B = PanelNew.Blueprint;
 
-	auto NewDiffControl = TSharedPtr<FClassDefaultsDiffControl>(new FClassDefaultsDiffControl(A, B, MasterDifferencesList, RealDifferences, FOnDetailsDiffControlChanged::CreateRaw(this, &SBlueprintDiff::SetCurrentMode, ClassSettingsMode)));
+	auto NewDiffControl = TSharedPtr<FClassDefaultsDiffControl>(new FClassDefaultsDiffControl(A, B, MasterDifferencesList, RealDifferences, FOnDiffControlChanged::CreateRaw(this, &SBlueprintDiff::SetCurrentMode, ClassSettingsMode)));
 
 	SBlueprintDiff::FDiffControl Ret;
 	//Splitter for left and right blueprint. Current convention is for the local (probably newer?) blueprint to be on the right:
 	Ret.DiffControl = NewDiffControl;
 	Ret.Widget = SNew(SSplitter)
+		.PhysicalSplitterHandleSize(10.0f)
 		+ SSplitter::Slot()
 		.Value(0.5f)
 		[
-			SNew(SBorder)
-			.VAlign(VAlign_Fill)
-			[
-				NewDiffControl->OldDetailsWidget()
-			]
+			NewDiffControl->OldDetailsWidget()
 		]
-	+ SSplitter::Slot()
+		+ SSplitter::Slot()
 		.Value(0.5f)
 		[
-			SNew(SBorder)
-			.VAlign(VAlign_Fill)
-			[
-				NewDiffControl->NewDetailsWidget()
-			]
+			NewDiffControl->NewDetailsWidget()
 		];
 
 	return Ret;
@@ -1624,26 +1721,19 @@ SBlueprintDiff::FDiffControl SBlueprintDiff::GenerateComponentsPanel()
 	SBlueprintDiff::FDiffControl Ret;
 
 	//Splitter for left and right blueprint. Current convention is for the local (probably newer?) blueprint to be on the right:
-	auto NewDiffControl = TSharedPtr<FSCSDiffControl>(new FSCSDiffControl(PanelOld.Blueprint, PanelNew.Blueprint, MasterDifferencesList, RealDifferences, FOnSCSDiffControlChanged::CreateRaw(this, &SBlueprintDiff::SetCurrentMode, FBlueprintEditorApplicationModes::BlueprintComponentsMode) ) );
+	auto NewDiffControl = TSharedPtr<FSCSDiffControl>(new FSCSDiffControl(PanelOld.Blueprint, PanelNew.Blueprint, MasterDifferencesList, RealDifferences, FOnDiffControlChanged::CreateRaw(this, &SBlueprintDiff::SetCurrentMode, ComponentsMode) ) );
 	Ret.DiffControl = NewDiffControl;
 	Ret.Widget = SNew(SSplitter)
+		.PhysicalSplitterHandleSize(10.0f)
 		+ SSplitter::Slot()
 		.Value(0.5f)
 		[
-			SNew(SBorder)
-			.VAlign(VAlign_Fill)
-			[
-				NewDiffControl->OldTreeWidget()
-			]
+			NewDiffControl->OldTreeWidget()
 		]
 		+ SSplitter::Slot()
 		.Value(0.5f)
 		[
-			SNew(SBorder)
-			.VAlign(VAlign_Fill)
-			[
-				NewDiffControl->NewTreeWidget()
-			]
+			NewDiffControl->NewTreeWidget()
 		];
 
 	return Ret;
@@ -1658,26 +1748,16 @@ void SBlueprintDiff::SetCurrentMode(FName NewMode)
 
 	CurrentMode = NewMode;
 
-	DiffControl = TSharedPtr< IDiffControl >();
-	if( NewMode == FBlueprintEditorApplicationModes::StandardBlueprintEditorMode)
-	{
-		DiffControl = GraphPanel.DiffControl;
-		ModeContents->SetContent( GraphPanel.Widget.ToSharedRef() );
-	}
-	else if( NewMode == FBlueprintEditorApplicationModes::BlueprintDefaultsMode )
-	{
-		DiffControl = DefaultsPanel.DiffControl;
-		ModeContents->SetContent( DefaultsPanel.Widget.ToSharedRef() );
-	}
-	else if (NewMode == ClassSettingsMode)
-	{
-		DiffControl = ClassSettingsPanel.DiffControl;
-		ModeContents->SetContent(ClassSettingsPanel.Widget.ToSharedRef());
-	}
-	else if (NewMode == FBlueprintEditorApplicationModes::BlueprintComponentsMode)
+	FDiffControl* FoundControl = ModePanels.Find(NewMode);
+
+	if (FoundControl)
 	{
-		DiffControl = ComponentsPanel.DiffControl;
-		ModeContents->SetContent( ComponentsPanel.Widget.ToSharedRef() );
+		// Reset inspector view
+		PanelOld.DetailsView->ShowDetailsForObjects(TArray<UObject*>());
+		PanelNew.DetailsView->ShowDetailsForObjects(TArray<UObject*>());
+
+		DiffControl = FoundControl->DiffControl;
+		ModeContents->SetContent(FoundControl->Widget.ToSharedRef());
 	}
 	else
 	{
diff --git a/Engine/Source/Editor/Kismet/Public/DiffUtils.h b/Engine/Source/Editor/Kismet/Public/DiffUtils.h
--- a/Engine/Source/Editor/Kismet/Public/DiffUtils.h
+++ b/Engine/Source/Editor/Kismet/Public/DiffUtils.h
@@ -235,13 +235,9 @@ public:
 	KISMET_API static TSharedPtr<FBlueprintDifferenceTreeEntry> NoDifferencesEntry();
 	KISMET_API static TSharedPtr<FBlueprintDifferenceTreeEntry> AnimBlueprintEntry();
 	KISMET_API static TSharedPtr<FBlueprintDifferenceTreeEntry> WidgetBlueprintEntry();
-	/** The FBlueprintDifferenceTreeEntry used to label the defaults category: */
-	KISMET_API static TSharedPtr<FBlueprintDifferenceTreeEntry> CreateDefaultsCategoryEntry(FOnDiffEntryFocused FocusCallback, const TArray< TSharedPtr<FBlueprintDifferenceTreeEntry> >& Children, bool bHasDifferences );
-	KISMET_API static TSharedPtr<FBlueprintDifferenceTreeEntry> CreateDefaultsCategoryEntryForMerge(FOnDiffEntryFocused FocusCallback, const TArray< TSharedPtr<FBlueprintDifferenceTreeEntry> >& Children, bool bHasRemoteDifferences, bool bHasLocalDifferences, bool bHasConflicts );
-	KISMET_API static TSharedPtr<FBlueprintDifferenceTreeEntry> CreateComponentsCategoryEntry(FOnDiffEntryFocused FocusCallback, const TArray< TSharedPtr<FBlueprintDifferenceTreeEntry> >& Children, bool bHasDifferences);
-	KISMET_API static TSharedPtr<FBlueprintDifferenceTreeEntry> CreateComponentsCategoryEntryForMerge(FOnDiffEntryFocused FocusCallback, const TArray< TSharedPtr<FBlueprintDifferenceTreeEntry> >& Children, bool bHasRemoteDifferences, bool bHasLocalDifferences, bool bHasConflicts);
-	KISMET_API static TSharedPtr<FBlueprintDifferenceTreeEntry> CreateClassSettingsCategoryEntry(FOnDiffEntryFocused FocusCallback, const TArray< TSharedPtr<FBlueprintDifferenceTreeEntry> >& Children, bool bHasDifferences);
-	KISMET_API static TSharedPtr<FBlueprintDifferenceTreeEntry> CreateClassSettingsCategoryEntryForMerge(FOnDiffEntryFocused FocusCallback, const TArray< TSharedPtr<FBlueprintDifferenceTreeEntry> >& Children, bool bHasRemoteDifferences, bool bHasLocalDifferences, bool bHasConflicts);
+	/** The FBlueprintDifferenceTreeEntry used to label categories: */
+	KISMET_API static TSharedPtr<FBlueprintDifferenceTreeEntry> CreateCategoryEntry(const FText& LabelText, const FText& ToolTipText, FOnDiffEntryFocused FocusCallback, const TArray< TSharedPtr<FBlueprintDifferenceTreeEntry> >& Children, bool bHasDifferences);
+	KISMET_API static TSharedPtr<FBlueprintDifferenceTreeEntry> CreateCategoryEntryForMerge(const FText& LabelText, const FText& ToolTipText, FOnDiffEntryFocused FocusCallback, const TArray< TSharedPtr<FBlueprintDifferenceTreeEntry> >& Children, bool bHasRemoteDifferences, bool bHasLocalDifferences, bool bHasConflicts);
 	
 	FOnDiffEntryFocused OnFocus;
 	FGenerateDiffEntryWidget GenerateWidget;
diff --git a/Engine/Source/Editor/Kismet/Public/SBlueprintDiff.h b/Engine/Source/Editor/Kismet/Public/SBlueprintDiff.h
--- a/Engine/Source/Editor/Kismet/Public/SBlueprintDiff.h
+++ b/Engine/Source/Editor/Kismet/Public/SBlueprintDiff.h
@@ -55,7 +55,7 @@ struct KISMET_API FDiffPanel
 	void GeneratePanel(UEdGraph* Graph, UEdGraph* GraphToDiff);
 
 	/* Generate the 'MyBlueprint' widget, which is private to this module */
-	TSharedRef<SWidget> GenerateMyBlueprintPanel();
+	TSharedRef<class SWidget> GenerateMyBlueprintWidget();
 
 	/* Called when user hits keyboard shortcut to copy nodes*/
 	void CopySelectedNodes();
@@ -139,10 +139,6 @@ protected:
 	bool HasNextDiff() const;
 	bool HasPrevDiff() const;
 
-	/** Spawns the tabs that contain the Graph views and blueprints views respectively: */
-	TSharedRef<SDockTab> CreateGraphDiffViews( const FSpawnTabArgs& Args );
-	TSharedRef<SDockTab> CreateMyBlueprintsViews( const FSpawnTabArgs& Args );
-
 	typedef TSharedPtr<struct FListItemGraphToDiff>	FGraphToDiff;
 	typedef SListView<FGraphToDiff >	SListViewType;
 
@@ -184,7 +180,7 @@ protected:
 
 	/** Event handler that updates the graph view when user selects a new graph */
 	void HandleGraphChanged( const FString& GraphPath );
-
+	
 	/** Function used to generate the list of differences and the widgets needed to calculate that list */
 	void GenerateDifferencesList();
 
@@ -198,7 +194,7 @@ protected:
 	{
 		FDiffControl()
 		: Widget()
-		, DiffControl(NULL)
+		, DiffControl(nullptr)
 		{
 		}
 
@@ -206,6 +202,7 @@ protected:
 		TSharedPtr< class IDiffControl > DiffControl;
 	};
 
+	FDiffControl GenerateMyBlueprintPanel();
 	FDiffControl GenerateGraphPanel();
 	FDiffControl GenerateDefaultsPanel();
 	FDiffControl GenerateClassSettingsPanel();
@@ -243,11 +240,8 @@ protected:
 	/** Tree view that displays the differences, cached for the buttons that iterate the differences: */
 	TSharedPtr< STreeView< TSharedPtr< FBlueprintDifferenceTreeEntry > > > DifferencesTreeView;
 
-	/** Stored references to widgets used to display various parts of a blueprint: */
-	FDiffControl GraphPanel;
-	FDiffControl DefaultsPanel;
-	FDiffControl ClassSettingsPanel;
-	FDiffControl ComponentsPanel;
+	/** Stored references to widgets used to display various parts of a blueprint, from the mode name */
+	TMap<FName, FDiffControl> ModePanels;
 
 	/** A pointer to the window holding this */
 	TWeakPtr<SWindow> WeakParentWindow;
