[NO NON-PII COMMIT MESSAGE]

diff --git a/Engine/Source/Editor/Kismet/Private/SBlueprintDiff.cpp b/Engine/Source/Editor/Kismet/Private/SBlueprintDiff.cpp
--- a/Engine/Source/Editor/Kismet/Private/SBlueprintDiff.cpp
+++ b/Engine/Source/Editor/Kismet/Private/SBlueprintDiff.cpp
@@ -101,8 +101,8 @@ FMyBlueprintDiffControl::FMyBlueprintDiffControl(
 		OutRealDifferences.Push(Entry);
 	}*/
 
-	const bool bHasDiffferences = Children.Num() != 0;
-	if (!bHasDiffferences)
+	const bool bHasDifferences = Children.Num() != 0;
+	if (!bHasDifferences)
 	{
 		// make one child informing the user that there are no differences:
 		Children.Push(FBlueprintDifferenceTreeEntry::NoDifferencesEntry());
@@ -119,7 +119,7 @@ FMyBlueprintDiffControl::FMyBlueprintDiffControl(
 		NSLOCTEXT("FBlueprintDifferenceTreeEntry", "MyBlueprintTooltip", "The list of changes made to blueprint structure in the My Blueprint panel"),
 		FOnDiffEntryFocused::CreateStatic(ForwardSelection, SelectionCallback),
 		Children,
-		bHasDiffferences
+		bHasDifferences
 	));
 }
 
@@ -196,8 +196,8 @@ FSCSDiffControl::FSCSDiffControl(
 		OutRealDifferences.Push(Entry);
 	}
 
-	const bool bHasDiffferences = Children.Num() != 0;
-	if (!bHasDiffferences)
+	const bool bHasDifferences = Children.Num() != 0;
+	if (!bHasDifferences)
 	{
 		// make one child informing the user that there are no differences:
 		Children.Push(FBlueprintDifferenceTreeEntry::NoDifferencesEntry());
@@ -214,7 +214,7 @@ FSCSDiffControl::FSCSDiffControl(
 		NSLOCTEXT("FBlueprintDifferenceTreeEntry", "SCSTooltip", "The list of changes made in the Components panel"),
 		FOnDiffEntryFocused::CreateStatic(ForwardSelection, SelectionCallback),
 		Children,
-		bHasDiffferences
+		bHasDifferences
 	));
 }
 
@@ -385,8 +385,8 @@ public:
 		: FDetailsDiffControl(InOldObject, InNewObject, OutTreeEntries, OutRealDifferences, SelectionCallback)
 	{
 		// Parent constructor fills out Children
-		const bool bHasDiffferences = Children.Num() != 0;
-		if (!bHasDiffferences)
+		const bool bHasDifferences = Children.Num() != 0;
+		if (!bHasDifferences)
 		{
 			// make one child informing the user that there are no differences:
 			Children.Push(FBlueprintDifferenceTreeEntry::NoDifferencesEntry());
@@ -403,7 +403,7 @@ public:
 			NSLOCTEXT("FBlueprintDifferenceTreeEntry", "DefaultsTooltip", "The list of changes made in the Defaults panel"),
 			FOnDiffEntryFocused::CreateStatic(ForwardSelection, SelectionCallback),
 			Children,
-			bHasDiffferences
+			bHasDifferences
 		));
 	}
 };
@@ -500,8 +500,8 @@ public:
 			}
 		}
 
-		const bool bHasDiffferences = Children.Num() != 0;
-		if (!bHasDiffferences)
+		const bool bHasDifferences = Children.Num() != 0;
+		if (!bHasDifferences)
 		{
 			// make one child informing the user that there are no differences:
 			Children.Push(FBlueprintDifferenceTreeEntry::NoDifferencesEntry());
@@ -518,7 +518,7 @@ public:
 			NSLOCTEXT("FBlueprintDifferenceTreeEntry", "SettingsTooltip", "The list of changes made in the Class Settings panel"),
 			FOnDiffEntryFocused::CreateStatic(ForwardSelection, SelectionCallback),
 			Children,
-			bHasDiffferences
+			bHasDifferences
 		));
 	}
 };
@@ -975,6 +975,7 @@ void SBlueprintDiff::Construct( const FArguments& InArgs)
 
 	TSharedRef<SWidget> Overlay = 
 		SNew(SSplitter)
+		.Visibility(EVisibility::HitTestInvisible)
 		+ SSplitter::Slot()
 		.Value(.2f)
 		[
