Implement diffing for Widget tree default values
Refactored SBlueprintDiff significantly to avoid code duplication and add comments
Change DiffUtils to not recurse into instanced objects unless they are EditInlineNew and not a component

diff --git a/Engine/Source/Developer/Merge/Private/SMergeGraphView.cpp b/Engine/Source/Developer/Merge/Private/SMergeGraphView.cpp
--- a/Engine/Source/Developer/Merge/Private/SMergeGraphView.cpp
+++ b/Engine/Source/Developer/Merge/Private/SMergeGraphView.cpp
@@ -214,7 +214,7 @@ static TArray< FMergeGraphEntry > GenerateDiffListItems(const FBlueprintRevPair&
 
 						FMergeGraphRowEntry NewEntry = {
 							Label
-							, Difference.OwningGraphPath
+							, Difference.OwningObjectPath
 							, ConflictingDifference ? (*ConflictingDifference)->Node2 : nullptr /*UEdGraphNode* LocalNode*/
 							, Difference.Node1 /*UEdGraphNode* BaseNode*/
 							, Difference.Node2 /*UEdGraphNode* RemoteNode*/
@@ -239,7 +239,7 @@ static TArray< FMergeGraphEntry > GenerateDiffListItems(const FBlueprintRevPair&
 						{
 							FMergeGraphRowEntry NewEntry = {
 								Difference.DisplayString
-								, Difference.OwningGraphPath
+								, Difference.OwningObjectPath
 								, Difference.Node2 /*UEdGraphNode* LocalNode*/
 								, Difference.Node1 /*UEdGraphNode* BaseNode*/
 								, nullptr
@@ -563,7 +563,7 @@ TSharedRef<SDockTab> SMergeGraphView::CreateGraphDiffViews(const FSpawnTabArgs&
 	{
 		PanelContainer->AddSlot()
 		[
-			SAssignNew(Panel.GraphEditorBorder, SBox)
+			SAssignNew(Panel.GraphEditorBox, SBox)
 			.VAlign(VAlign_Fill)
 			[
 				SBlueprintDiff::DefaultEmptyPanel()
diff --git a/Engine/Source/Editor/GraphEditor/Private/GraphDiffControl.cpp b/Engine/Source/Editor/GraphEditor/Private/GraphDiffControl.cpp
--- a/Engine/Source/Editor/GraphEditor/Private/GraphDiffControl.cpp
+++ b/Engine/Source/Editor/GraphEditor/Private/GraphDiffControl.cpp
@@ -710,7 +710,7 @@ bool FGraphDiffControl::DiffGraphs(UEdGraph* const LhsGraph, UEdGraph* const Rhs
 	const FString GraphPath = LhsGraph ? GetGraphPath(LhsGraph) : GetGraphPath(RhsGraph);
 	for( FDiffSingleResult& Entry : DiffsOut )
 	{
-		Entry.OwningGraphPath = GraphPath;
+		Entry.OwningObjectPath = GraphPath;
 	}
 
 	return bFoundDifferences;
diff --git a/Engine/Source/Editor/GraphEditor/Public/DiffResults.h b/Engine/Source/Editor/GraphEditor/Public/DiffResults.h
deleted file mode 100644
diff --git a/Engine/Source/Editor/Kismet/Private/DetailsDiff.cpp b/Engine/Source/Editor/Kismet/Private/DetailsDiff.cpp
--- a/Engine/Source/Editor/Kismet/Private/DetailsDiff.cpp
+++ b/Engine/Source/Editor/Kismet/Private/DetailsDiff.cpp
@@ -63,7 +63,7 @@ TArray<FPropertySoftPath> FDetailsDiff::GetDisplayedProperties() const
 	return Ret;
 }
 
-void FDetailsDiff::DiffAgainst(const FDetailsDiff& Newer, TArray< FSingleObjectDiffEntry > &OutDifferences) const
+void FDetailsDiff::DiffAgainst(const FDetailsDiff& Newer, TArray< FSingleObjectDiffEntry > &OutDifferences, bool bSortByDisplayOrder) const
 {
 	TSharedPtr< class IDetailsView > OldDetailsView = DetailsView;
 	TSharedPtr< class IDetailsView > NewDetailsView = Newer.DetailsView;
@@ -116,4 +116,101 @@ void FDetailsDiff::DiffAgainst(const FDetailsDiff& Newer, TArray< FSingleObjectD
 			}
 		}
 	}
+
+	if (bSortByDisplayOrder)
+	{
+		TArray<FSingleObjectDiffEntry> AllDifferingProperties = OutDifferences;
+		OutDifferences.Reset();
+
+		// OrderedProperties will contain differences in the order they are displayed:
+		TArray< const FSingleObjectDiffEntry* > OrderedProperties;
+
+		// create differing properties list based on what is displayed by the old properties..
+		TArray<FPropertySoftPath> SoftOldProperties = GetDisplayedProperties();
+		TArray<FPropertySoftPath> SoftNewProperties = Newer.GetDisplayedProperties();
+
+		const auto FindAndPushDiff = [&OrderedProperties, &AllDifferingProperties](const FPropertySoftPath& PropertyIdentifier) -> bool
+		{
+			bool bDiffers = false;
+			for (const auto& Difference : AllDifferingProperties)
+			{
+				if (Difference.Identifier == PropertyIdentifier)
+				{
+					bDiffers = true;
+					// if there are any nested differences associated with PropertyIdentifier, add those
+					// as well:
+					OrderedProperties.AddUnique(&Difference);
+				}
+				else if (Difference.Identifier.IsSubPropertyMatch(PropertyIdentifier))
+				{
+					bDiffers = true;
+					OrderedProperties.AddUnique(&Difference);
+				}
+			}
+			return bDiffers;
+		};
+
+		// zip the two sets of properties, zip iterators are not trivial to write in c++,
+		// so this procedural stuff will have to do:
+		int IterOld = 0;
+		int IterNew = 0;
+		while (IterOld < SoftOldProperties.Num() || IterNew < SoftNewProperties.Num())
+		{
+			const bool bOldIterValid = IterOld < SoftOldProperties.Num();
+			const bool bNewIterValid = IterNew < SoftNewProperties.Num();
+
+			// We've reached the end of the new list, but still have properties in the old list.
+			// Continue over the old list to catch any remaining diffs.
+			if (bOldIterValid && !bNewIterValid)
+			{
+				FindAndPushDiff(SoftOldProperties[IterOld]);
+				++IterOld;
+			}
+			// We've reached the end of the old list, but still have properties in the new list.
+			// Continue over the new list to catch any remaining diffs.
+			else if (!bOldIterValid && bNewIterValid)
+			{
+				FindAndPushDiff(SoftNewProperties[IterNew]);
+				++IterNew;
+			}
+			else
+			{
+				// If both properties have the same path, check to ensure the property hasn't changed.
+				if (SoftOldProperties[IterOld] == SoftNewProperties[IterNew])
+				{
+					FindAndPushDiff(SoftOldProperties[IterOld]);
+					++IterNew;
+					++IterOld;
+				}
+				else
+				{
+					// If the old property is different, add it to the list and increment the old iter.
+					// This indicates the property was removed.
+					if (FindAndPushDiff(SoftOldProperties[IterOld]))
+					{
+						++IterOld;
+					}
+					// If the new property is different, add it to the list and increment the new iter.
+					// This indicates the property was added.
+					else if (FindAndPushDiff(SoftNewProperties[IterNew]))
+					{
+						++IterNew;
+					}
+					// Neither property was different.
+					// This indicates the iterators were just out of step from a previous addition or removal.
+					else
+					{
+						++IterOld;
+						++IterNew;
+					}
+				}
+			}
+		}
+
+		// Readd to OutDifferences
+		for (const FSingleObjectDiffEntry* Difference : OrderedProperties)
+		{
+			OutDifferences.Add(*Difference);
+		}
+	}
 }
diff --git a/Engine/Source/Editor/Kismet/Private/DiffUtils.cpp b/Engine/Source/Editor/Kismet/Private/DiffUtils.cpp
--- a/Engine/Source/Editor/Kismet/Private/DiffUtils.cpp
+++ b/Engine/Source/Editor/Kismet/Private/DiffUtils.cpp
@@ -570,21 +570,25 @@ static void IdenticalHelper(const UProperty* AProperty, const UProperty* BProper
 			return;
 		}
 
-		// dig into the objects if they are in the same package as our initial object:
 		const UObjectProperty* BPropAsObject = CastChecked<UObjectProperty>(BProperty);
 
 		const UObject* A = *((const UObject* const*)AValue);
 		const UObject* B = *((const UObject* const*)BValue);
 
+		// dig into the objects if they are in the same package as our initial object:
 		if(BPropAsObject->HasAnyPropertyFlags(CPF_InstancedReference) && APropAsObject->HasAnyPropertyFlags(CPF_InstancedReference) && A && B && A->GetClass() == B->GetClass())
 		{
-			// dive into the object to find actual differences:
 			const UClass* AClass = A->GetClass(); // BClass and AClass are identical!
 
-			for (TFieldIterator<UProperty> PropertyIt(AClass); PropertyIt; ++PropertyIt)
+			// We only want to recurse if this is EditInlineNew and not a component
+			// Other instanced refs are likely to form a type-specific web so recursion doesn't make sense and won't be displayed in the details pane
+			if (AClass->HasAnyClassFlags(CLASS_EditInlineNew) && !AClass->IsChildOf(UActorComponent::StaticClass()))
 			{
-				const UProperty* ClassProp = *PropertyIt;
-				IdenticalHelper(ClassProp, ClassProp, ClassProp->ContainerPtrToValuePtr<void>(A, 0), ClassProp->ContainerPtrToValuePtr<void>(B, 0), FPropertySoftPath(RootPath, ClassProp), DifferingSubProperties);
+				for (TFieldIterator<UProperty> PropertyIt(AClass); PropertyIt; ++PropertyIt)
+				{
+					const UProperty* ClassProp = *PropertyIt;
+					IdenticalHelper(ClassProp, ClassProp, ClassProp->ContainerPtrToValuePtr<void>(A, 0), ClassProp->ContainerPtrToValuePtr<void>(B, 0), FPropertySoftPath(RootPath, ClassProp), DifferingSubProperties);
+				}
 			}
 		}
 		else
@@ -712,73 +716,24 @@ TSharedPtr<FBlueprintDifferenceTreeEntry> FBlueprintDifferenceTreeEntry::NoDiffe
 	) );
 }
 
-TSharedPtr<FBlueprintDifferenceTreeEntry> FBlueprintDifferenceTreeEntry::AnimBlueprintEntry()
+TSharedPtr<FBlueprintDifferenceTreeEntry> FBlueprintDifferenceTreeEntry::UnknownDifferencesEntry()
 {
-	// For now, a widget and a short message explaining that differences in the AnimGraph are
-	// not detected by the diff tool:
+	// Warn about there being unknown differences
 	const auto GenerateWidget = []() -> TSharedRef<SWidget>
 	{
 		return SNew(STextBlock)
 			.ColorAndOpacity(FLinearColor(.7f, .7f, .7f))
 			.TextStyle(FEditorStyle::Get(), TEXT("BlueprintDif.ItalicText"))
-			.Text(NSLOCTEXT("FBlueprintDifferenceTreeEntry", "AnimBlueprintsNotSupported", "Warning: Detecting differences in Animation Blueprint specific data is not yet supported..."));
-	};
-
-	TArray< TSharedPtr<FBlueprintDifferenceTreeEntry> > Children;
-	Children.Emplace( TSharedPtr<FBlueprintDifferenceTreeEntry>(new FBlueprintDifferenceTreeEntry(
-		FOnDiffEntryFocused()
-		, FGenerateDiffEntryWidget::CreateStatic(GenerateWidget)
-		, TArray< TSharedPtr<FBlueprintDifferenceTreeEntry> >()
-		)) );
-
-	const auto CreateAnimGraphRootEntry = []() -> TSharedRef<SWidget>
-	{
-		return SNew(STextBlock)
-			.ToolTipText(NSLOCTEXT("FBlueprintDifferenceTreeEntry", "AnimGraphTooltip", "Detecting differences in Animation Blueprint specific data is not yet supported"))
-			.ColorAndOpacity(DiffViewUtils::LookupColor(true))
-			.Text(NSLOCTEXT("FBlueprintDifferenceTreeEntry", "AnimGraphLabel", "Animation Blueprint"));
+			.Text(NSLOCTEXT("FBlueprintDifferenceTreeEntry", "BlueprintTypeNotSupported", "Warning: Detecting differences in this Blueprint type specific data is not yet supported..."));
 	};
 
 	return TSharedPtr<FBlueprintDifferenceTreeEntry>(new FBlueprintDifferenceTreeEntry(
-		FOnDiffEntryFocused()
-		, FGenerateDiffEntryWidget::CreateStatic(CreateAnimGraphRootEntry)
-		, Children
-		));
-}
-
-TSharedPtr<FBlueprintDifferenceTreeEntry> FBlueprintDifferenceTreeEntry::WidgetBlueprintEntry()
-{
-	// For now, a widget and a short message explaining that differences in the WidgetTree are
-	// not detected by the diff tool:
-	const auto GenerateWidget = []() -> TSharedRef<SWidget>
-	{
-		return SNew(STextBlock)
-			.ColorAndOpacity(FLinearColor(.7f, .7f, .7f))
-			.TextStyle(FEditorStyle::Get(), TEXT("BlueprintDif.ItalicText"))
-			.Text(NSLOCTEXT("FBlueprintDifferenceTreeEntry", "WidgetTreeNotSupported", "Warning: Detecting differences in Widget Blueprint specific data is not yet supported..."));
-	};
-
-	TArray< TSharedPtr<FBlueprintDifferenceTreeEntry> > Children;
-	Children.Emplace(TSharedPtr<FBlueprintDifferenceTreeEntry>(new FBlueprintDifferenceTreeEntry(
 		FOnDiffEntryFocused()
 		, FGenerateDiffEntryWidget::CreateStatic(GenerateWidget)
 		, TArray< TSharedPtr<FBlueprintDifferenceTreeEntry> >()
-		)));
-
-	const auto CreateWidgetTreeRootEntry = []() -> TSharedRef<SWidget>
-	{
-		return SNew(STextBlock)
-			.ToolTipText(NSLOCTEXT("FBlueprintDifferenceTreeEntry", "WidgetTreeTooltip", "Detecting differences in Widget Blueprint specific data is not yet supported"))
-			.ColorAndOpacity(DiffViewUtils::LookupColor(true))
-			.Text(NSLOCTEXT("FBlueprintDifferenceTreeEntry", "WidgetTreeLabel", "Widget Blueprint"));
-	};
-
-	return TSharedPtr<FBlueprintDifferenceTreeEntry>(new FBlueprintDifferenceTreeEntry(
-		FOnDiffEntryFocused()
-		, FGenerateDiffEntryWidget::CreateStatic(CreateWidgetTreeRootEntry)
-		, Children
-		));
+	));
 }
+
 TSharedPtr<FBlueprintDifferenceTreeEntry> FBlueprintDifferenceTreeEntry::CreateCategoryEntry(const FText& LabelText, const FText& ToolTipText, FOnDiffEntryFocused FocusCallback, const TArray< TSharedPtr<FBlueprintDifferenceTreeEntry> >& Children, bool bHasDifferences)
 {
 	const auto CreateDefaultsRootEntry = [](FText LabelText, FText ToolTipText, FLinearColor Color) -> TSharedRef<SWidget>
diff --git a/Engine/Source/Editor/Kismet/Private/SBlueprintDiff.cpp b/Engine/Source/Editor/Kismet/Private/SBlueprintDiff.cpp
--- a/Engine/Source/Editor/Kismet/Private/SBlueprintDiff.cpp
+++ b/Engine/Source/Editor/Kismet/Private/SBlueprintDiff.cpp
@@ -32,96 +32,96 @@ class IDiffControl
 {
 public:
 	virtual ~IDiffControl() {}
-};
 
-DECLARE_DELEGATE(FOnDiffControlChanged);
+	/** Adds widgets to the tree of differences to show */
+	virtual void GenerateTreeEntries(TArray< TSharedPtr<FBlueprintDifferenceTreeEntry> >& OutTreeEntries, TArray< TSharedPtr<FBlueprintDifferenceTreeEntry> >& OutRealDifferences) = 0;
+};
 
 FText RightRevision = LOCTEXT("OlderRevisionIdentifier", "Right Revision");
 
 typedef TMap< FName, const UProperty* > FNamePropertyMap;
 
+const FName BlueprintTypeMode = FName(TEXT("BlueprintTypeMode"));
 const FName MyBlueprintMode = FName(TEXT("MyBlueprintMode"));
 const FName DefaultsMode = FName(TEXT("DefaultsMode"));
 const FName ClassSettingsMode = FName(TEXT("ClassSettingsMode"));
 const FName ComponentsMode = FName(TEXT("ComponentsMode"));
 const FName GraphMode = FName(TEXT("GraphMode"));
 
+TSharedRef<SWidget>	FDiffResultItem::GenerateWidget() const
+{
+	FText ToolTip = Result.ToolTip;
+	FLinearColor Color = Result.DisplayColor;
+	FText Text = Result.DisplayString;
+	if (Text.IsEmpty())
+	{
+		Text = LOCTEXT("DIF_UnknownDiff", "Unknown Diff");
+		ToolTip = LOCTEXT("DIF_Confused", "There is an unspecified difference");
+	}
+	return SNew(STextBlock)
+		.ToolTipText(ToolTip)
+		.ColorAndOpacity(Color)
+		.Text(Text);
+}
+
 /** Shows all differences for the blueprint structure itself that aren't picked up elsewhere */
 class FMyBlueprintDiffControl : public TSharedFromThis<FMyBlueprintDiffControl>, public IDiffControl
 {
 public:
-	FMyBlueprintDiffControl(
-		const UBlueprint* InOldBlueprint
-		, const UBlueprint* InNewBlueprint
-		, TArray< TSharedPtr<FBlueprintDifferenceTreeEntry> >& OutTreeEntries
-		, TArray< TSharedPtr<FBlueprintDifferenceTreeEntry> >& OutRealDifferences
-		, FOnDiffControlChanged SelectionCallback
-	);
+	FMyBlueprintDiffControl(const UBlueprint* InOldBlueprint, const UBlueprint* InNewBlueprint, FOnDiffEntryFocused InSelectionCallback)
+		: SelectionCallback(InSelectionCallback)
+	{
+	}
 
-	virtual ~FMyBlueprintDiffControl() {}
-};
+	virtual void GenerateTreeEntries(TArray< TSharedPtr<FBlueprintDifferenceTreeEntry> >& OutTreeEntries, TArray< TSharedPtr<FBlueprintDifferenceTreeEntry> >& OutRealDifferences) override
+	{
+		TArray< TSharedPtr<FBlueprintDifferenceTreeEntry> > Children;
 
-FMyBlueprintDiffControl::FMyBlueprintDiffControl(
-	const UBlueprint* InOldBlueprint
-	, const UBlueprint* InNewBlueprint
-	, TArray< TSharedPtr<FBlueprintDifferenceTreeEntry> >& OutTreeEntries
-	, TArray< TSharedPtr<FBlueprintDifferenceTreeEntry> >& OutRealDifferences
-	, FOnDiffControlChanged SelectionCallback
-)
-{
-	TArray< TSharedPtr<FBlueprintDifferenceTreeEntry> > Children;
+		/*const auto FocusSCSDifferenceEntry = [](FSCSDiffEntry Entry, FOnDiffEntryFocused InSelectionCallback, FMyBlueprintDiffControl* Owner)
+		{
+			InSelectionCallback.ExecuteIfBound();
+			if (Entry.TreeIdentifier.Name != NAME_None)
+			{
+				Owner->OldSCS.HighlightProperty(Entry.TreeIdentifier.Name, FPropertyPath());
+				Owner->NewSCS.HighlightProperty(Entry.TreeIdentifier.Name, FPropertyPath());
+			}
+		};
 
-	/*const auto FocusSCSDifferenceEntry = [](FSCSDiffEntry Entry, FOnDiffControlChanged InSelectionCallback, FMyBlueprintDiffControl* Owner)
-	{
-		InSelectionCallback.ExecuteIfBound();
-		if (Entry.TreeIdentifier.Name != NAME_None)
+		const auto CreateSCSDifferenceWidget = [](FSCSDiffEntry Entry, FText ObjectName) -> TSharedRef<SWidget>
 		{
-			Owner->OldSCS.HighlightProperty(Entry.TreeIdentifier.Name, FPropertyPath());
-			Owner->NewSCS.HighlightProperty(Entry.TreeIdentifier.Name, FPropertyPath());
-		}
-	};
+			return SNew(STextBlock)
+				.Text(DiffViewUtils::SCSDiffMessage(Entry, ObjectName))
+				.ColorAndOpacity(DiffViewUtils::Differs());
+		};
 
-	const auto CreateSCSDifferenceWidget = [](FSCSDiffEntry Entry, FText ObjectName) -> TSharedRef<SWidget>
-	{
-		return SNew(STextBlock)
-			.Text(DiffViewUtils::SCSDiffMessage(Entry, ObjectName))
-			.ColorAndOpacity(DiffViewUtils::Differs());
-	};
+		for (auto Difference : DifferingProperties.Entries)
+		{
+			TSharedPtr<FBlueprintDifferenceTreeEntry> Entry = MakeShared<FBlueprintDifferenceTreeEntry>(
+					FOnDiffEntryFocused::CreateStatic(FocusSCSDifferenceEntry, Difference, SelectionCallback, this),
+					FGenerateDiffEntryWidget::CreateStatic(CreateSCSDifferenceWidget, Difference, RightRevision));
+			Children.Push(Entry);
+			OutRealDifferences.Push(Entry);
+		}*/
 
-	for (auto Difference : DifferingProperties.Entries)
-	{
-		auto Entry = TSharedPtr<FBlueprintDifferenceTreeEntry>(
-			new FBlueprintDifferenceTreeEntry(
-				FOnDiffEntryFocused::CreateStatic(FocusSCSDifferenceEntry, Difference, SelectionCallback, this)
-				, FGenerateDiffEntryWidget::CreateStatic(CreateSCSDifferenceWidget, Difference, RightRevision)
-				, TArray< TSharedPtr<FBlueprintDifferenceTreeEntry> >()
-			)
-			);
-		Children.Push(Entry);
-		OutRealDifferences.Push(Entry);
-	}*/
+		const bool bHasDifferences = Children.Num() != 0;
+		if (!bHasDifferences)
+		{
+			// make one child informing the user that there are no differences:
+			Children.Push(FBlueprintDifferenceTreeEntry::NoDifferencesEntry());
+		}
 
-	const bool bHasDifferences = Children.Num() != 0;
-	if (!bHasDifferences)
-	{
-		// make one child informing the user that there are no differences:
-		Children.Push(FBlueprintDifferenceTreeEntry::NoDifferencesEntry());
+		OutTreeEntries.Push(FBlueprintDifferenceTreeEntry::CreateCategoryEntry(
+			NSLOCTEXT("FBlueprintDifferenceTreeEntry", "MyBlueprintLabel", "My Blueprint"),
+			NSLOCTEXT("FBlueprintDifferenceTreeEntry", "MyBlueprintTooltip", "The list of changes made to blueprint structure in the My Blueprint panel"),
+			SelectionCallback,
+			Children,
+			bHasDifferences
+		));
 	}
 
-	const auto ForwardSelection = [](FOnDiffControlChanged InSelectionCallback)
-	{
-		// This allows the owning control to focus the correct tab (or do whatever else it likes):
-		InSelectionCallback.ExecuteIfBound();
-	};
-
-	OutTreeEntries.Push(FBlueprintDifferenceTreeEntry::CreateCategoryEntry(
-		NSLOCTEXT("FBlueprintDifferenceTreeEntry", "MyBlueprintLabel", "My Blueprint"),
-		NSLOCTEXT("FBlueprintDifferenceTreeEntry", "MyBlueprintTooltip", "The list of changes made to blueprint structure in the My Blueprint panel"),
-		FOnDiffEntryFocused::CreateStatic(ForwardSelection, SelectionCallback),
-		Children,
-		bHasDifferences
-	));
-}
+private:
+	FOnDiffEntryFocused SelectionCallback;
+};
 
 /** 
  * Each difference in the tree will either be a tree node that is added in one Blueprint 
@@ -130,261 +130,136 @@ FMyBlueprintDiffControl::FMyBlueprintDiffControl(
 class FSCSDiffControl : public TSharedFromThis<FSCSDiffControl>, public IDiffControl
 {
 public:
-	FSCSDiffControl(
-					const UBlueprint* InOldBlueprint
-					, const UBlueprint* InNewBlueprint
-					, TArray< TSharedPtr<FBlueprintDifferenceTreeEntry> >& OutTreeEntries
-					, TArray< TSharedPtr<FBlueprintDifferenceTreeEntry> >& OutRealDifferences
-					, FOnDiffControlChanged SelectionCallback
-					);
+	FSCSDiffControl(const UBlueprint* InOldBlueprint, const UBlueprint* InNewBlueprint, FOnDiffEntryFocused InSelectionCallback)
+		: SelectionCallback(InSelectionCallback)
+		, OldSCS(InOldBlueprint)
+		, NewSCS(InNewBlueprint)
+	{
+	}
+
+	virtual void GenerateTreeEntries(TArray< TSharedPtr<FBlueprintDifferenceTreeEntry> >& OutTreeEntries, TArray< TSharedPtr<FBlueprintDifferenceTreeEntry> >& OutRealDifferences) override
+	{
+		TArray< FSCSResolvedIdentifier > OldHierarchy = OldSCS.GetDisplayedHierarchy();
+		TArray< FSCSResolvedIdentifier > NewHierarchy = NewSCS.GetDisplayedHierarchy();
+		DiffUtils::CompareUnrelatedSCS(OldSCS.GetBlueprint(), OldHierarchy, NewSCS.GetBlueprint(), NewHierarchy, DifferingProperties);
+
+		const auto FocusSCSDifferenceEntry = [](FSCSDiffEntry Entry, FOnDiffEntryFocused InSelectionCallback, FSCSDiffControl* Owner)
+		{
+			InSelectionCallback.ExecuteIfBound();
+			if (Entry.TreeIdentifier.Name != NAME_None)
+			{
+				Owner->OldSCS.HighlightProperty(Entry.TreeIdentifier.Name, FPropertyPath());
+				Owner->NewSCS.HighlightProperty(Entry.TreeIdentifier.Name, FPropertyPath());
+			}
+		};
+
+		const auto CreateSCSDifferenceWidget = [](FSCSDiffEntry Entry, FText ObjectName) -> TSharedRef<SWidget>
+		{
+			return SNew(STextBlock)
+				.Text(DiffViewUtils::SCSDiffMessage(Entry, ObjectName))
+				.ColorAndOpacity(DiffViewUtils::Differs());
+		};
+
+		TArray< TSharedPtr<FBlueprintDifferenceTreeEntry> > Children;
+		for (auto Difference : DifferingProperties.Entries)
+		{
+			TSharedPtr<FBlueprintDifferenceTreeEntry> Entry = MakeShared<FBlueprintDifferenceTreeEntry>(
+				FOnDiffEntryFocused::CreateStatic(FocusSCSDifferenceEntry, Difference, SelectionCallback, this),
+				FGenerateDiffEntryWidget::CreateStatic(CreateSCSDifferenceWidget, Difference, RightRevision));
+			Children.Push(Entry);
+			OutRealDifferences.Push(Entry);
+		}
 
+		const bool bHasDifferences = Children.Num() != 0;
+		if (!bHasDifferences)
+		{
+			// make one child informing the user that there are no differences:
+			Children.Push(FBlueprintDifferenceTreeEntry::NoDifferencesEntry());
+		}
+
+		OutTreeEntries.Push(FBlueprintDifferenceTreeEntry::CreateCategoryEntry(
+			NSLOCTEXT("FBlueprintDifferenceTreeEntry", "SCSLabel", "Components"),
+			NSLOCTEXT("FBlueprintDifferenceTreeEntry", "SCSTooltip", "The list of changes made in the Components panel"),
+			SelectionCallback,
+			Children,
+			bHasDifferences
+		));
+	}
 
 	TSharedRef<SWidget> OldTreeWidget() { return OldSCS.TreeWidget(); }
 	TSharedRef<SWidget> NewTreeWidget() { return NewSCS.TreeWidget(); }
 
-	virtual ~FSCSDiffControl() {}
 private:
+	FOnDiffEntryFocused SelectionCallback;
 	FSCSDiffRoot DifferingProperties;
 
 	FSCSDiff OldSCS;
 	FSCSDiff NewSCS;
 };
 
-FSCSDiffControl::FSCSDiffControl(
-		const UBlueprint* InOldBlueprint
-		, const UBlueprint* InNewBlueprint
-		, TArray< TSharedPtr<FBlueprintDifferenceTreeEntry> >& OutTreeEntries
-		, TArray< TSharedPtr<FBlueprintDifferenceTreeEntry> >& OutRealDifferences
-		, FOnDiffControlChanged SelectionCallback
-	)
-	: DifferingProperties()
-	, OldSCS( InOldBlueprint )
-	, NewSCS( InNewBlueprint )
-{
-	TArray< FSCSResolvedIdentifier > OldHierarchy = OldSCS.GetDisplayedHierarchy();
-	TArray< FSCSResolvedIdentifier > NewHierarchy = NewSCS.GetDisplayedHierarchy();
-	DiffUtils::CompareUnrelatedSCS(InOldBlueprint, OldHierarchy, InNewBlueprint, NewHierarchy, DifferingProperties);
-
-	const auto FocusSCSDifferenceEntry = [](FSCSDiffEntry Entry, FOnDiffControlChanged InSelectionCallback, FSCSDiffControl* Owner)
-	{
-		InSelectionCallback.ExecuteIfBound();
-		if (Entry.TreeIdentifier.Name != NAME_None)
-		{
-			Owner->OldSCS.HighlightProperty(Entry.TreeIdentifier.Name, FPropertyPath());
-			Owner->NewSCS.HighlightProperty(Entry.TreeIdentifier.Name, FPropertyPath());
-		}
-	};
-
-	const auto CreateSCSDifferenceWidget = [](FSCSDiffEntry Entry, FText ObjectName) -> TSharedRef<SWidget>
-	{
-		return SNew(STextBlock)
-			.Text(DiffViewUtils::SCSDiffMessage(Entry,ObjectName))
-			.ColorAndOpacity(DiffViewUtils::Differs());
-	};
-
-	TArray< TSharedPtr<FBlueprintDifferenceTreeEntry> > Children;
-	for( auto Difference : DifferingProperties.Entries )
+/** Generic wrapper around a details view, this does not actually fill out OutTreeEntries */
+class FDetailsDiffControl : public TSharedFromThis<FDetailsDiffControl>, public IDiffControl
+{
+public:
+	FDetailsDiffControl(const UObject* InOldObject, const UObject* InNewObject, FOnDiffEntryFocused InSelectionCallback)
+		: SelectionCallback(InSelectionCallback)
+		, OldDetails(InOldObject, FDetailsDiff::FOnDisplayedPropertiesChanged())
+		, NewDetails(InNewObject, FDetailsDiff::FOnDisplayedPropertiesChanged())
 	{
-		auto Entry = TSharedPtr<FBlueprintDifferenceTreeEntry>(
-			new FBlueprintDifferenceTreeEntry(
-				FOnDiffEntryFocused::CreateStatic(FocusSCSDifferenceEntry, Difference, SelectionCallback, this)
-				, FGenerateDiffEntryWidget::CreateStatic(CreateSCSDifferenceWidget, Difference, RightRevision)
-				, TArray< TSharedPtr<FBlueprintDifferenceTreeEntry> >()
-			)
-		);
-		Children.Push(Entry);
-		OutRealDifferences.Push(Entry);
+		OldDetails.DiffAgainst(NewDetails, DifferingProperties, true);
 	}
 
-	const bool bHasDifferences = Children.Num() != 0;
-	if (!bHasDifferences)
+	virtual void GenerateTreeEntries(TArray< TSharedPtr<FBlueprintDifferenceTreeEntry> >& OutTreeEntries, TArray< TSharedPtr<FBlueprintDifferenceTreeEntry> >& OutRealDifferences) override
 	{
-		// make one child informing the user that there are no differences:
-		Children.Push(FBlueprintDifferenceTreeEntry::NoDifferencesEntry());
+		for (const FSingleObjectDiffEntry& Difference : DifferingProperties)
+		{
+			TSharedPtr<FBlueprintDifferenceTreeEntry> Entry = MakeShared<FBlueprintDifferenceTreeEntry>(
+				FOnDiffEntryFocused::CreateSP(AsShared(), &FDetailsDiffControl::OnSelectDiffEntry, Difference.Identifier),
+				FGenerateDiffEntryWidget::CreateSP(AsShared(), &FDetailsDiffControl::CreateDifferenceWidget, Difference, RightRevision));
+			Children.Push(Entry);
+			OutRealDifferences.Push(Entry);
+		}
 	}
 
-	const auto ForwardSelection = [](FOnDiffControlChanged InSelectionCallback)
-	{
-		// This allows the owning control to focus the correct tab (or do whatever else it likes):
-		InSelectionCallback.ExecuteIfBound();
-	};
-
-	OutTreeEntries.Push(FBlueprintDifferenceTreeEntry::CreateCategoryEntry(
-		NSLOCTEXT("FBlueprintDifferenceTreeEntry", "SCSLabel", "Components"),
-		NSLOCTEXT("FBlueprintDifferenceTreeEntry", "SCSTooltip", "The list of changes made in the Components panel"),
-		FOnDiffEntryFocused::CreateStatic(ForwardSelection, SelectionCallback),
-		Children,
-		bHasDifferences
-	));
-}
-
-/** Generic wrapper around a details view, this does not actually fill out OutTreeEntries */
-class FDetailsDiffControl : public IDiffControl
-{
-public:
-	FDetailsDiffControl( const UObject* InOldObject
-					, const UObject* InNewObject
-					, TArray< TSharedPtr<FBlueprintDifferenceTreeEntry> >& OutTreeEntries
-					, TArray< TSharedPtr<FBlueprintDifferenceTreeEntry> >& OutRealDifferences
-					, FOnDiffControlChanged SelectionCallback );
-
 	TSharedRef<SWidget> OldDetailsWidget() { return OldDetails.DetailsWidget(); }
 	TSharedRef<SWidget> NewDetailsWidget() { return NewDetails.DetailsWidget(); }
 
-	virtual ~FDetailsDiffControl() {}
-
 protected:
-	void HighlightDifference(const FPropertySoftPath& PropertyName);
-
-	FDetailsDiff OldDetails;
-	FDetailsDiff NewDetails;
-
-	TArray< TSharedPtr<FBlueprintDifferenceTreeEntry> > Children;
-};
-
-FDetailsDiffControl::FDetailsDiffControl(
-		const UObject* InOldObject
-		, const UObject* InNewObject
-		, TArray< TSharedPtr<FBlueprintDifferenceTreeEntry> >& OutTreeEntries
-		, TArray< TSharedPtr<FBlueprintDifferenceTreeEntry> >& OutRealDifferences
-		, FOnDiffControlChanged SelectionCallback )
-	: OldDetails(InOldObject, FDetailsDiff::FOnDisplayedPropertiesChanged() )
-	, NewDetails(InNewObject, FDetailsDiff::FOnDisplayedPropertiesChanged() )
-{
-	TArray< FSingleObjectDiffEntry > DifferingProperties;
-	OldDetails.DiffAgainst(NewDetails, DifferingProperties);
-
-	// OrderedProperties will contain differences in the order they are displayed:
-	TArray< const FSingleObjectDiffEntry* > OrderedProperties;
-
-	// create differing properties list based on what is displayed by the old properties..
-	TArray<FPropertySoftPath> OldProperties = OldDetails.GetDisplayedProperties();
-	TArray<FPropertySoftPath> NewProperties = NewDetails.GetDisplayedProperties();
-
-	const auto FindAndPushDiff = [&OrderedProperties, &DifferingProperties](const FPropertySoftPath& PropertyIdentifier) -> bool
+	virtual void OnSelectDiffEntry(FPropertySoftPath PropertyName)
 	{
-		bool bDiffers = false;
-		for (const auto& Difference : DifferingProperties)
-		{
-			if (Difference.Identifier == PropertyIdentifier)
-			{
-				bDiffers = true;
-				// if there are any nested differences associated with PropertyIdentifier, add those
-				// as well:
-				OrderedProperties.AddUnique(&Difference);
-			}
-			else if (Difference.Identifier.IsSubPropertyMatch(PropertyIdentifier))
-			{
-				bDiffers = true;
-				OrderedProperties.AddUnique(&Difference);
-			}
-		}
-		return bDiffers;
-	};
-
-	// zip the two sets of properties, zip iterators are not trivial to write in c++,
-	// so this procedural stuff will have to do:
-	int IterOld = 0;
-	int IterNew = 0;
-	while (IterOld < OldProperties.Num() || IterNew < NewProperties.Num())
-	{
-		const bool bOldIterValid = IterOld < OldProperties.Num();
-		const bool bNewIterValid = IterNew < NewProperties.Num();
-
-		// We've reached the end of the new list, but still have properties in the old list.
-		// Continue over the old list to catch any remaining diffs.
-		if (bOldIterValid && !bNewIterValid)
-		{
-			FindAndPushDiff(OldProperties[IterOld]);
-			++IterOld;
-		}
-		// We've reached the end of the old list, but still have properties in the new list.
-		// Continue over the new list to catch any remaining diffs.
-		else if (!bOldIterValid && bNewIterValid)
-		{
-			FindAndPushDiff(NewProperties[IterNew]);
-			++IterNew;
-		}
-		else
-		{
-			// If both properties have the same path, check to ensure the property hasn't changed.
-			if (OldProperties[IterOld] == NewProperties[IterNew])
-			{
-				FindAndPushDiff(OldProperties[IterOld]);
-				++IterNew;
-				++IterOld;
-			}
-			else
-			{
-				// If the old property is different, add it to the list and increment the old iter.
-				// This indicates the property was removed.
-				if (FindAndPushDiff(OldProperties[IterOld]))
-				{
-					++IterOld;
-				}
-				// If the new property is different, add it to the list and increment the new iter.
-				// This indicates the property was added.
-				else if (FindAndPushDiff(NewProperties[IterNew]))
-				{
-					++IterNew;
-				}
-				// Neither property was different.
-				// This indicates the iterators were just out of step from a previous addition or removal.
-				else
-				{
-					++IterOld;
-					++IterNew;
-				}
-			}
-		}
+		SelectionCallback.ExecuteIfBound();
+		OldDetails.HighlightProperty(PropertyName);
+		NewDetails.HighlightProperty(PropertyName);
 	}
 
-	const auto CreateDifferenceWidget = [](FSingleObjectDiffEntry DiffEntry, FText ObjectName) -> TSharedRef<SWidget>
+	virtual TSharedRef<SWidget> CreateDifferenceWidget(FSingleObjectDiffEntry DiffEntry, FText ObjectName)
 	{
 		return SNew(STextBlock)
 			.Text(DiffViewUtils::PropertyDiffMessage(DiffEntry, ObjectName))
 			.ColorAndOpacity(DiffViewUtils::Differs());
-	};
-
-	const auto FocusDetailsDifferenceEntry = [](FPropertySoftPath Identifier, FDetailsDiffControl* Control, FOnDiffControlChanged InSelectionCallback)
-	{
-		// This allows the owning control to focus the correct tab (or do whatever else it likes):
-		InSelectionCallback.ExecuteIfBound();
-		Control->HighlightDifference(Identifier);
-	};
-
-	for (auto Difference : OrderedProperties)
-	{
-		auto Entry = TSharedPtr<FBlueprintDifferenceTreeEntry>(
-			new FBlueprintDifferenceTreeEntry(
-			FOnDiffEntryFocused::CreateStatic(FocusDetailsDifferenceEntry, Difference->Identifier, this, SelectionCallback)
-			, FGenerateDiffEntryWidget::CreateStatic(CreateDifferenceWidget, *Difference, RightRevision)
-			, TArray< TSharedPtr<FBlueprintDifferenceTreeEntry> >()
-			)
-		);
-		Children.Push(Entry);
-		OutRealDifferences.Push(Entry);
 	}
-}
 
-void FDetailsDiffControl::HighlightDifference(const FPropertySoftPath& PropertyName)
-{
-	OldDetails.HighlightProperty(PropertyName);
-	NewDetails.HighlightProperty(PropertyName);
-}
+	FOnDiffEntryFocused SelectionCallback;
+	FDetailsDiff OldDetails;
+	FDetailsDiff NewDetails;
+
+	TArray<FSingleObjectDiffEntry> DifferingProperties;
+	TArray< TSharedPtr<FBlueprintDifferenceTreeEntry> > Children;
+};
 
 /** Override for CDO special case */
 class FCDODiffControl : public TSharedFromThis<FCDODiffControl>, public FDetailsDiffControl
 {
 public:
-	FCDODiffControl(const UObject* InOldObject
-		, const UObject* InNewObject
-		, TArray< TSharedPtr<FBlueprintDifferenceTreeEntry> >& OutTreeEntries
-		, TArray< TSharedPtr<FBlueprintDifferenceTreeEntry> >& OutRealDifferences
-		, FOnDiffControlChanged SelectionCallback)
-		: FDetailsDiffControl(InOldObject, InNewObject, OutTreeEntries, OutRealDifferences, SelectionCallback)
-	{
-		// Parent constructor fills out Children
+	FCDODiffControl(const UObject* InOldObject, const UObject* InNewObject, FOnDiffEntryFocused InSelectionCallback)
+		: FDetailsDiffControl(InOldObject, InNewObject, InSelectionCallback)
+	{
+	}
+	
+	virtual void GenerateTreeEntries(TArray< TSharedPtr<FBlueprintDifferenceTreeEntry> >& OutTreeEntries, TArray< TSharedPtr<FBlueprintDifferenceTreeEntry> >& OutRealDifferences) override
+	{
+		FDetailsDiffControl::GenerateTreeEntries(OutTreeEntries, OutRealDifferences);
+
 		const bool bHasDifferences = Children.Num() != 0;
 		if (!bHasDifferences)
 		{
@@ -392,36 +267,32 @@ public:
 			Children.Push(FBlueprintDifferenceTreeEntry::NoDifferencesEntry());
 		}
 
-		const auto ForwardSelection = [](FOnDiffControlChanged InSelectionCallback)
-		{
-			// This allows the owning control to focus the correct tab (or do whatever else it likes):
-			InSelectionCallback.ExecuteIfBound();
-		};
-
 		OutTreeEntries.Push(FBlueprintDifferenceTreeEntry::CreateCategoryEntry(
 			NSLOCTEXT("FBlueprintDifferenceTreeEntry", "DefaultsLabel", "Defaults"),
 			NSLOCTEXT("FBlueprintDifferenceTreeEntry", "DefaultsTooltip", "The list of changes made in the Defaults panel"),
-			FOnDiffEntryFocused::CreateStatic(ForwardSelection, SelectionCallback),
+			SelectionCallback,
 			Children,
 			bHasDifferences
 		));
 	}
 };
 
-/** Override for class defaults */
-class FClassDefaultsDiffControl : public TSharedFromThis<FClassDefaultsDiffControl>, public FDetailsDiffControl
+/** Override for class class settings */
+class FClassSettingsDiffControl : public TSharedFromThis<FClassSettingsDiffControl>, public FDetailsDiffControl
 {
 public:
-	FClassDefaultsDiffControl(const UObject* InOldObject
-		, const UObject* InNewObject
-		, TArray< TSharedPtr<FBlueprintDifferenceTreeEntry> >& OutTreeEntries
-		, TArray< TSharedPtr<FBlueprintDifferenceTreeEntry> >& OutRealDifferences
-		, FOnDiffControlChanged SelectionCallback)
-		: FDetailsDiffControl(InOldObject, InNewObject, OutTreeEntries, OutRealDifferences, SelectionCallback)
+	FClassSettingsDiffControl(const UObject* InOldObject, const UObject* InNewObject, FOnDiffEntryFocused InSelectionCallback)
+		: FDetailsDiffControl(InOldObject, InNewObject, InSelectionCallback)
+	{
+	}
+
+	virtual void GenerateTreeEntries(TArray< TSharedPtr<FBlueprintDifferenceTreeEntry> >& OutTreeEntries, TArray< TSharedPtr<FBlueprintDifferenceTreeEntry> >& OutRealDifferences) override
 	{
+		FDetailsDiffControl::GenerateTreeEntries(OutTreeEntries, OutRealDifferences);
+
 		// Check for parent class and interfaces here
-		const UBlueprint* OldBlueprint = Cast<UBlueprint>(InOldObject);
-		const UBlueprint* NewBlueprint = Cast<UBlueprint>(InNewObject);
+		const UBlueprint* OldBlueprint = Cast<UBlueprint>(OldDetails.GetDisplayedObject());
+		const UBlueprint* NewBlueprint = Cast<UBlueprint>(OldDetails.GetDisplayedObject());
 
 		const auto GenerateDiffWidget = [](FText DiffText) -> TSharedRef<SWidget>
 		{
@@ -431,22 +302,15 @@ public:
 				.ToolTipText(DiffText);
 		};
 
-		const auto FocusSettings = [](FOnDiffControlChanged InSelectionCallback)
-		{
-			InSelectionCallback.ExecuteIfBound();
-		};
-
 		if (OldBlueprint && NewBlueprint)
 		{
 			if (OldBlueprint->ParentClass != NewBlueprint->ParentClass)
 			{
 				FText DiffText = FText::Format(LOCTEXT("ParentChanged", "Parent Class changed from {0} to {1}"), FText::FromString(OldBlueprint->ParentClass->GetName()), FText::FromString(NewBlueprint->ParentClass->GetName()));
 
-				auto Entry = TSharedPtr<FBlueprintDifferenceTreeEntry>(new FBlueprintDifferenceTreeEntry(
-					FOnDiffEntryFocused::CreateStatic(FocusSettings, SelectionCallback)
-					, FGenerateDiffEntryWidget::CreateStatic(GenerateDiffWidget, DiffText)
-					, TArray< TSharedPtr<FBlueprintDifferenceTreeEntry> >()
-				));
+				TSharedPtr<FBlueprintDifferenceTreeEntry> Entry = MakeShared<FBlueprintDifferenceTreeEntry>(
+					SelectionCallback,
+					FGenerateDiffEntryWidget::CreateStatic(GenerateDiffWidget, DiffText));
 
 				Children.Push(Entry);
 				OutRealDifferences.Push(Entry);
@@ -475,11 +339,9 @@ public:
 			{
 				FText DiffText = FText::Format(LOCTEXT("InterfacesChanged", "Interfaces changed from '{0}' to '{1}'"), FText::FromString(OldInterfaces), FText::FromString(NewInterfaces));
 
-				auto Entry = TSharedPtr<FBlueprintDifferenceTreeEntry>(new FBlueprintDifferenceTreeEntry(
-					FOnDiffEntryFocused::CreateStatic(FocusSettings, SelectionCallback)
-					, FGenerateDiffEntryWidget::CreateStatic(GenerateDiffWidget, DiffText)
-					, TArray< TSharedPtr<FBlueprintDifferenceTreeEntry> >()
-				));
+				TSharedPtr<FBlueprintDifferenceTreeEntry> Entry = MakeShared<FBlueprintDifferenceTreeEntry>(
+					SelectionCallback,
+					FGenerateDiffEntryWidget::CreateStatic(GenerateDiffWidget, DiffText));
 
 				Children.Push(Entry);
 				OutRealDifferences.Push(Entry);
@@ -489,11 +351,9 @@ public:
 			{
 				FText DiffText = FText::Format(LOCTEXT("NativizationChanged", "Nativization changed from {0} to {1}"), FText::AsNumber(OldBlueprint->SupportsNativization()), FText::AsNumber(NewBlueprint->SupportsNativization()));
 
-				auto Entry = TSharedPtr<FBlueprintDifferenceTreeEntry>(new FBlueprintDifferenceTreeEntry(
-					FOnDiffEntryFocused::CreateStatic(FocusSettings, SelectionCallback)
-					, FGenerateDiffEntryWidget::CreateStatic(GenerateDiffWidget, DiffText)
-					, TArray< TSharedPtr<FBlueprintDifferenceTreeEntry> >()
-				));
+				TSharedPtr<FBlueprintDifferenceTreeEntry> Entry = MakeShared<FBlueprintDifferenceTreeEntry>(
+					SelectionCallback,
+					FGenerateDiffEntryWidget::CreateStatic(GenerateDiffWidget, DiffText));
 
 				Children.Push(Entry);
 				OutRealDifferences.Push(Entry);
@@ -507,121 +367,293 @@ public:
 			Children.Push(FBlueprintDifferenceTreeEntry::NoDifferencesEntry());
 		}
 
-		const auto ForwardSelection = [](FOnDiffControlChanged InSelectionCallback)
-		{
-			// This allows the owning control to focus the correct tab (or do whatever else it likes):
-			InSelectionCallback.ExecuteIfBound();
-		};
-
 		OutTreeEntries.Push(FBlueprintDifferenceTreeEntry::CreateCategoryEntry(
 			NSLOCTEXT("FBlueprintDifferenceTreeEntry", "SettingsLabel", "Class Settings"),
 			NSLOCTEXT("FBlueprintDifferenceTreeEntry", "SettingsTooltip", "The list of changes made in the Class Settings panel"),
-			FOnDiffEntryFocused::CreateStatic(ForwardSelection, SelectionCallback),
+			SelectionCallback,
 			Children,
 			bHasDifferences
 		));
 	}
 };
 
-/*List item that entry for a graph*/
-struct KISMET_API FListItemGraphToDiff	: public TSharedFromThis<FListItemGraphToDiff>
+/** Diff control to handle finding type-specific differences */
+struct FBlueprintTypeDiffControl : public TSharedFromThis<FBlueprintTypeDiffControl>, public IDiffControl
 {
-	FListItemGraphToDiff(class SBlueprintDiff* Diff, class UEdGraph* GraphOld, class UEdGraph* GraphNew, const FRevisionInfo& RevisionOld, const FRevisionInfo& RevisionNew);
-	virtual ~FListItemGraphToDiff();
+	struct FSubObjectDiff
+	{
+		FDiffSingleResult SourceResult;
+		FDetailsDiff OldDetails;
+		FDetailsDiff NewDetails;
+		TArray<TSharedPtr<FBlueprintDifferenceTreeEntry>> Diffs;
+
+		FSubObjectDiff(const FDiffSingleResult& InSourceResult, const UObject* OldObject, const UObject* NewObject)
+			: SourceResult(InSourceResult)
+			, OldDetails(OldObject, FDetailsDiff::FOnDisplayedPropertiesChanged())
+			, NewDetails(NewObject, FDetailsDiff::FOnDisplayedPropertiesChanged())
+		{}
+	};
 
-	/*Generate Widget for list item*/
-	TSharedRef<SWidget> GenerateWidget();
+	FBlueprintTypeDiffControl(const UBlueprint* InBlueprintOld, const UBlueprint* InBlueprintNew, FOnDiffEntryFocused InSelectionCallback)
+		: BlueprintOld(InBlueprintOld), BlueprintNew(InBlueprintNew), SelectionCallback(InSelectionCallback), bDiffSucceeded(false)
+	{
+		check(InBlueprintNew && InBlueprintOld);
+	}
 
-	/*Get tooltip for list item */
-	FText   GetToolTip();
+	/** Generate difference tree widgets */
+	virtual void GenerateTreeEntries(TArray< TSharedPtr<FBlueprintDifferenceTreeEntry> >& OutTreeEntries, TArray< TSharedPtr<FBlueprintDifferenceTreeEntry> >& OutRealDifferences) override;
 
-	/*Get old(left) graph*/
-	UEdGraph* GetGraphOld()const{ return GraphOld; }
+	/** The old blueprint (left) */
+	const UBlueprint* BlueprintOld;
 
-	/*Get new(right) graph*/
-	UEdGraph* GetGraphNew()const{ return GraphNew; }
+	/** The new blueprint(right) */
+	const UBlueprint* BlueprintNew;
+
+	/** Boxes that will display the details diffs */
+	TSharedPtr<SBox> OldDetailsBox;
+	TSharedPtr<SBox> NewDetailsBox;
 
-	/** Source for list view */
-	TArray<TSharedPtr<struct FDiffResultItem>> DiffListSource;
 private:
-	/*Diff widget*/
-	class SBlueprintDiff* Diff;
+	/** Generate Widget for top category */
+	TSharedRef<SWidget> GenerateCategoryWidget(bool bHasRealDiffs);
 
-	/*The old graph(left)*/
-	class UEdGraph*	GraphOld;
+	/** Build up the Diff Source Array*/
+	void BuildDiffSourceArray();
 
-	/*The new graph(right)*/
-	class UEdGraph* GraphNew;
+	/** Handle selecting a diff */
+	void OnSelectSubobjectDiff(FPropertySoftPath Identifier, TSharedPtr<FSubObjectDiff> SubObjectDiff);
 
-	/*Description of Old and new graph*/
-	FRevisionInfo	RevisionOld, RevisionNew;
+	/** List of objects with differences */
+	TArray<TSharedPtr<FSubObjectDiff>> SubObjectDiffs;
 
-	//////////////////////////////////////////////////////////////////////////
-	// Diff list
-	//////////////////////////////////////////////////////////////////////////
+	/** Source for list view */
+	TArray<TSharedPtr<FDiffResultItem>> DiffListSource;
 
-	typedef TSharedPtr<struct FDiffResultItem>	FSharedDiffOnGraph;
-	typedef SListView<FSharedDiffOnGraph >		SListViewType;
+	/** Selection callback */
+	FOnDiffEntryFocused SelectionCallback;
 
-public:
+	/** Did diff generation succeed? */
+	bool bDiffSucceeded;
+};
+
+TSharedRef<SWidget> FBlueprintTypeDiffControl::GenerateCategoryWidget(bool bHasRealDiffs)
+{
+	FLinearColor Color = FLinearColor::White;
+
+	if (bHasRealDiffs)
+	{
+		Color = DiffViewUtils::Differs();
+	}
+
+	FText Label = BlueprintNew->GetClass()->GetDisplayNameText();
+
+	return SNew(SHorizontalBox)
+		+ SHorizontalBox::Slot()
+		[
+			SNew(STextBlock)
+			.ColorAndOpacity(Color)
+		.Text(Label)
+		];
+}
+
+void FBlueprintTypeDiffControl::OnSelectSubobjectDiff(FPropertySoftPath Identifier, TSharedPtr<FSubObjectDiff> SubObjectDiff)
+{
+	// This allows the owning control to focus the correct tab (or do whatever else it likes):
+	SelectionCallback.ExecuteIfBound();
+
+	if (SubObjectDiff.IsValid())
+	{
+		SubObjectDiff->OldDetails.HighlightProperty(Identifier);
+		SubObjectDiff->NewDetails.HighlightProperty(Identifier);
+
+		OldDetailsBox->SetContent(SubObjectDiff->OldDetails.DetailsWidget());
+		NewDetailsBox->SetContent(SubObjectDiff->NewDetails.DetailsWidget());
+	}
+}
+
+void FBlueprintTypeDiffControl::BuildDiffSourceArray()
+{
+	TArray<FDiffSingleResult> BlueprintDiffResults;
+	FDiffResults BlueprintDiffs(&BlueprintDiffResults);
+	if (BlueprintNew->FindDiffs(BlueprintOld, BlueprintDiffs))
+	{
+		bDiffSucceeded = true;
+
+		// Add manual diffs
+		for (const FDiffSingleResult& CurrentDiff : BlueprintDiffResults)
+		{
+			if (CurrentDiff.Diff == EDiffType::OBJECT_REQUEST_DIFF)
+			{
+				// Turn into a subobject diff
+
+				// Invert order, we want old then new
+				TSharedPtr<FSubObjectDiff> SubObjectDiff = MakeShared<FSubObjectDiff>(CurrentDiff, CurrentDiff.Object2, CurrentDiff.Object1);
+
+				TArray<FSingleObjectDiffEntry> DifferingProperties;
+				SubObjectDiff->OldDetails.DiffAgainst(SubObjectDiff->NewDetails, DifferingProperties, true);
+
+				if (DifferingProperties.Num() > 0)
+				{
+					// Actual differences, so add to tree
+					SubObjectDiffs.Add(SubObjectDiff);
+
+					const auto CreateDifferenceWidget = [](FSingleObjectDiffEntry DiffEntry, FText ObjectName) -> TSharedRef<SWidget>
+					{
+						return SNew(STextBlock)
+							.Text(DiffViewUtils::PropertyDiffMessage(DiffEntry, ObjectName))
+							.ColorAndOpacity(DiffViewUtils::Differs());
+					};
+
+					for (const FSingleObjectDiffEntry& Difference : DifferingProperties)
+					{
+						TSharedPtr<FBlueprintDifferenceTreeEntry> Entry = MakeShared<FBlueprintDifferenceTreeEntry>(
+							FOnDiffEntryFocused::CreateSP(AsShared(), &FBlueprintTypeDiffControl::OnSelectSubobjectDiff, Difference.Identifier, SubObjectDiff),
+							FGenerateDiffEntryWidget::CreateStatic(CreateDifferenceWidget, Difference, RightRevision));
+						SubObjectDiff->Diffs.Push(Entry);
+					}
+				}
+			}
+			else
+			{
+				DiffListSource.Add(MakeShared<FDiffResultItem>(CurrentDiff));
+			}
+		}
+
+		struct SortDiff
+		{
+			bool operator () (const TSharedPtr<FDiffResultItem>& A, const TSharedPtr<FDiffResultItem>& B) const
+			{
+				return A->Result.Diff < B->Result.Diff;
+			}
+		};
+
+		Sort(DiffListSource.GetData(), DiffListSource.Num(), SortDiff());
+	}
+}
+
+void FBlueprintTypeDiffControl::GenerateTreeEntries(TArray< TSharedPtr<FBlueprintDifferenceTreeEntry> >& OutTreeEntries, TArray< TSharedPtr<FBlueprintDifferenceTreeEntry> >& OutRealDifferences)
+{
+	BuildDiffSourceArray();
+
+	TArray< TSharedPtr<FBlueprintDifferenceTreeEntry> > Children;
+	
+	bool bHasRealChange = false;
+
+	// First add manual diffs in main category
+	for (const TSharedPtr<FDiffResultItem>& Difference : DiffListSource)
+	{
+		TSharedPtr<FBlueprintDifferenceTreeEntry> ChildEntry = MakeShared<FBlueprintDifferenceTreeEntry>(
+			SelectionCallback,
+			FGenerateDiffEntryWidget::CreateSP(Difference.ToSharedRef(), &FDiffResultItem::GenerateWidget));
+		Children.Push(ChildEntry);
+		OutRealDifferences.Push(ChildEntry);
+
+		if (Difference->Result.IsRealDifference())
+		{
+			bHasRealChange = true;
+		}
+	}
+
+	if (Children.Num() == 0)
+	{
+		// Make one child informing the user that there are no differences, or that it is unknown
+		if (bDiffSucceeded)
+		{
+			Children.Push(FBlueprintDifferenceTreeEntry::NoDifferencesEntry());
+		}
+		else
+		{
+			Children.Push(FBlueprintDifferenceTreeEntry::UnknownDifferencesEntry());
+		}
+	}
+
+	TSharedPtr<FBlueprintDifferenceTreeEntry> CategoryEntry = MakeShared<FBlueprintDifferenceTreeEntry>(
+			SelectionCallback,
+			FGenerateDiffEntryWidget::CreateSP(AsShared(), &FBlueprintTypeDiffControl::GenerateCategoryWidget, bHasRealChange),
+			Children);
+	OutTreeEntries.Push(CategoryEntry);
+
+	// Now add subobject diffs, one category per object
+	for (const TSharedPtr<FSubObjectDiff> SubObjectDiff : SubObjectDiffs)
+	{
+		Children.Reset();
+
+		Children.Append(SubObjectDiff->Diffs);
+		OutRealDifferences.Append(SubObjectDiff->Diffs);
+
+		TSharedPtr<FBlueprintDifferenceTreeEntry> SubObjectEntry = FBlueprintDifferenceTreeEntry::CreateCategoryEntry(
+			SubObjectDiff->SourceResult.DisplayString,
+			SubObjectDiff->SourceResult.ToolTip,
+			FOnDiffEntryFocused::CreateSP(AsShared(), &FBlueprintTypeDiffControl::OnSelectSubobjectDiff, FPropertySoftPath(), SubObjectDiff),
+			Children,
+			true);
+
+		OutTreeEntries.Push(SubObjectEntry);
+	}
+}
+
+/** Category list item for a graph*/
+struct FGraphToDiff	: public TSharedFromThis<FGraphToDiff>, IDiffControl
+{
+	FGraphToDiff(SBlueprintDiff* DiffWidget, UEdGraph* GraphOld, UEdGraph* GraphNew, const FRevisionInfo& RevisionOld, const FRevisionInfo& RevisionNew);
+	virtual ~FGraphToDiff();
+
+	/** Add widgets to the differences tree */
+	virtual void GenerateTreeEntries(TArray< TSharedPtr<FBlueprintDifferenceTreeEntry> >& OutTreeEntries, TArray< TSharedPtr<FBlueprintDifferenceTreeEntry> >& OutRealDifferences) override;
+
+	/** Get old(left) graph*/
+	UEdGraph* GetGraphOld() const { return GraphOld; }
+
+	/** Get new(right) graph*/
+	UEdGraph* GetGraphNew() const { return GraphNew; }
+
+	/** Source for list view */
+	TArray<TSharedPtr<FDiffResultItem>> DiffListSource;
+
+private:
+	/** Get tooltip for category */
+	FText GetToolTip();
+
+	/** Generate Widget for category list */
+	TSharedRef<SWidget> GenerateCategoryWidget();
 
 	/** Called when the Newer Graph is modified*/
 	void OnGraphChanged(const FEdGraphEditAction& Action);
 
-	/** Generate list of differences*/
-	TSharedRef<SWidget> GenerateDiffListWidget();
-
 	/** Build up the Diff Source Array*/
 	void BuildDiffSourceArray();
 
-	/** Called when user clicks on a new graph list item */
-	void OnSelectionChanged(FSharedDiffOnGraph Item, ESelectInfo::Type SelectionType);
+	/** Diff widget */
+	class SBlueprintDiff* DiffWidget;
 
-private:
-	/** Get Index of the current diff that is selected */
-	int32 GetCurrentDiffIndex() const;
+	/** The old graph(left)*/
+	UEdGraph* GraphOld;
 
-	/* Called when a new row is being generated */
-	TSharedRef<ITableRow> OnGenerateRow(FSharedDiffOnGraph ParamItem, const TSharedRef<STableViewBase>& OwnerTable);
+	/** The new graph(right)*/
+	UEdGraph* GraphNew;
 
-	/** ListView of differences */
-	TSharedPtr<SListViewType> DiffList;
+	/** Description of Old and new graph*/
+	FRevisionInfo	RevisionOld, RevisionNew;
 
 	/** Handle to the registered OnGraphChanged delegate. */
 	FDelegateHandle OnGraphChangedDelegateHandle;
 };
 
-TSharedRef<SWidget>	FDiffResultItem::GenerateWidget() const
-{
-	FText ToolTip = Result.ToolTip;
-	FLinearColor Color = Result.DisplayColor;
-	FText Text = Result.DisplayString;
-	if (Text.IsEmpty())
-	{
-		Text = LOCTEXT("DIF_UnknownDiff", "Unknown Diff");
-		ToolTip = LOCTEXT("DIF_Confused", "There is an unspecified difference");
-	}
-	return SNew(STextBlock)
-		.ToolTipText(ToolTip)
-		.ColorAndOpacity(Color)
-		.Text(Text);
-}
-
-FListItemGraphToDiff::FListItemGraphToDiff( class SBlueprintDiff* InDiff, class UEdGraph* InGraphOld, class UEdGraph* InGraphNew, const FRevisionInfo& InRevisionOld, const FRevisionInfo& InRevisionNew )
-	: Diff(InDiff), GraphOld(InGraphOld), GraphNew(InGraphNew), RevisionOld(InRevisionOld), RevisionNew(InRevisionNew)
+FGraphToDiff::FGraphToDiff(SBlueprintDiff* InDiffWidget, UEdGraph* InGraphOld, UEdGraph* InGraphNew, const FRevisionInfo& InRevisionOld, const FRevisionInfo& InRevisionNew)
+	: DiffWidget(InDiffWidget), GraphOld(InGraphOld), GraphNew(InGraphNew), RevisionOld(InRevisionOld), RevisionNew(InRevisionNew)
 {
 	check(InGraphOld || InGraphNew); //one of them needs to exist
 
 	//need to know when it is modified
 	if(InGraphNew)
 	{
-		OnGraphChangedDelegateHandle = InGraphNew->AddOnGraphChangedHandler( FOnGraphChanged::FDelegate::CreateRaw(this, &FListItemGraphToDiff::OnGraphChanged));
+		OnGraphChangedDelegateHandle = InGraphNew->AddOnGraphChangedHandler( FOnGraphChanged::FDelegate::CreateRaw(this, &FGraphToDiff::OnGraphChanged));
 	}
 
 	BuildDiffSourceArray();
 }
 
-FListItemGraphToDiff::~FListItemGraphToDiff()
+FGraphToDiff::~FGraphToDiff()
 {
 	if(GraphNew)
 	{
@@ -629,7 +661,7 @@ FListItemGraphToDiff::~FListItemGraphToDiff()
 	}
 }
 
-TSharedRef<SWidget> FListItemGraphToDiff::GenerateWidget() 
+TSharedRef<SWidget> FGraphToDiff::GenerateCategoryWidget()
 {
 	const UEdGraph* Graph = GraphOld ? GraphOld : GraphNew;
 	check(Graph);
@@ -662,81 +694,17 @@ TSharedRef<SWidget> FListItemGraphToDiff::GenerateWidget()
 		SNew(STextBlock)
 		.ColorAndOpacity(Color)
 		.Text(GraphName)
+		.ToolTipText(GetToolTip())
 	]
 	+ DiffViewUtils::Box( GraphOld != nullptr, Color )
 	+ DiffViewUtils::Box( GraphNew != nullptr, Color );
 }
 
-TSharedRef<SWidget> FListItemGraphToDiff::GenerateDiffListWidget()
-{
-	if(DiffListSource.Num() > 0)
-	{
-		TSharedPtr<SListViewType> DiffListRef;
-		TSharedRef<SHorizontalBox> Result =	SNew(SHorizontalBox)
-			+SHorizontalBox::Slot()
-			.FillWidth(1.f)
-			.MaxWidth(350.f)
-			[
-				SNew(SVerticalBox)
-				+SVerticalBox::Slot()
-				.Padding(0.f)
-				.AutoHeight()
-				[
-					SNew(SBorder)
-					.BorderImage(FEditorStyle::GetBrush("PropertyWindow.CategoryBackground"))
-					.Padding(FMargin(2.0f))
-					.ForegroundColor(FEditorStyle::GetColor("PropertyWindow.CategoryForeground"))
-					.ToolTipText(LOCTEXT("BlueprintDifDifferencesToolTip", "List of differences found between revisions, click to select"))
-					.HAlign(HAlign_Center)
-					[
-						SNew(STextBlock)
-						.Text(LOCTEXT("RevisionDifferences", "Revision Differences"))
-					]
-				]
-				+SVerticalBox::Slot()
-				.Padding(1.f)
-				.FillHeight(1.f)
-				[
-					SAssignNew(DiffList, SListViewType)
-					.ItemHeight(24)
-					.ListItemsSource( &DiffListSource )
-					.OnGenerateRow( this, &FListItemGraphToDiff::OnGenerateRow )
-					.SelectionMode( ESelectionMode::Single )
-					.OnSelectionChanged(this, &FListItemGraphToDiff::OnSelectionChanged)
-				]
-			];
-		return Result;
-	}
-	else
-	{
-		return SNew(SBorder).Visibility(EVisibility::Hidden);
-	}
-		
-	
-}
-
-TSharedRef<ITableRow> FListItemGraphToDiff::OnGenerateRow( FSharedDiffOnGraph ParamItem, const TSharedRef<STableViewBase>& OwnerTable )
-{
-	return	SNew( STableRow< FSharedDiffOnGraph >, OwnerTable )
-		.Content()
-		[
-			ParamItem->GenerateWidget()
-		];
-}
-
-void FListItemGraphToDiff::OnSelectionChanged( FSharedDiffOnGraph Item, ESelectInfo::Type SelectionType )
+FText FGraphToDiff::GetToolTip()
 {
-	if(Item.IsValid())
+	if (GraphOld && GraphNew)
 	{
-		Diff->OnDiffListSelectionChanged(Item);
-	}
-}
-
-FText FListItemGraphToDiff::GetToolTip() 
-{
-	if ( GraphOld && GraphNew )
-	{
-		if ( DiffListSource.Num() > 0 )
+		if (DiffListSource.Num() > 0)
 		{
 			return LOCTEXT("ContainsDifferences", "Revisions are different");
 		}
@@ -752,16 +720,41 @@ FText FListItemGraphToDiff::GetToolTip()
 		const FRevisionInfo& Revision = GraphNew ? RevisionOld : RevisionNew;
 		FText RevisionText = LOCTEXT("CurrentRevision", "Current Revision");
 
-		if ( !Revision.Revision.IsEmpty() )
+		if (!Revision.Revision.IsEmpty())
 		{
-			RevisionText = FText::Format( LOCTEXT("Revision Number", "Revision {0}") , FText::FromString( Revision.Revision ) );
+			RevisionText = FText::Format(LOCTEXT("Revision Number", "Revision {0}"), FText::FromString(Revision.Revision));
 		}
 
-		return FText::Format( LOCTEXT("MissingGraph", "Graph '{0}' missing from {1}"), FText::FromString( GoodGraph->GetName() ), RevisionText );
+		return FText::Format(LOCTEXT("MissingGraph", "Graph '{0}' missing from {1}"), FText::FromString(GoodGraph->GetName()), RevisionText);
 	}
 }
 
-void FListItemGraphToDiff::BuildDiffSourceArray()
+void FGraphToDiff::GenerateTreeEntries(TArray< TSharedPtr<FBlueprintDifferenceTreeEntry> >& OutTreeEntries, TArray< TSharedPtr<FBlueprintDifferenceTreeEntry> >& OutRealDifferences)
+{
+	TArray< TSharedPtr<FBlueprintDifferenceTreeEntry> > Children;
+	for (const TSharedPtr<FDiffResultItem>& Difference : DiffListSource)
+	{
+		TSharedPtr<FBlueprintDifferenceTreeEntry> ChildEntry = MakeShared<FBlueprintDifferenceTreeEntry>(
+				FOnDiffEntryFocused::CreateRaw(DiffWidget, &SBlueprintDiff::OnDiffListSelectionChanged, Difference),
+				FGenerateDiffEntryWidget::CreateSP(Difference.ToSharedRef(), &FDiffResultItem::GenerateWidget));
+		Children.Push(ChildEntry);
+		OutRealDifferences.Push(ChildEntry);
+	}
+
+	if (Children.Num() == 0)
+	{
+		// make one child informing the user that there are no differences:
+		Children.Push(FBlueprintDifferenceTreeEntry::NoDifferencesEntry());
+	}
+
+	TSharedPtr<FBlueprintDifferenceTreeEntry> Entry = MakeShared<FBlueprintDifferenceTreeEntry>(
+			FOnDiffEntryFocused::CreateRaw(DiffWidget, &SBlueprintDiff::OnGraphSelectionChanged, TSharedPtr<FGraphToDiff>(AsShared()), ESelectInfo::Direct),
+			FGenerateDiffEntryWidget::CreateSP(AsShared(), &FGraphToDiff::GenerateCategoryWidget),
+			Children);
+	OutTreeEntries.Push(Entry);
+}
+
+void FGraphToDiff::BuildDiffSourceArray()
 {
 	TArray<FDiffSingleResult> FoundDiffs;
 	FGraphDiffControl::DiffGraphs(GraphOld, GraphNew, FoundDiffs);
@@ -769,12 +762,12 @@ void FListItemGraphToDiff::BuildDiffSourceArray()
 	DiffListSource.Empty();
 	for (auto DiffIt(FoundDiffs.CreateIterator()); DiffIt; ++DiffIt)
 	{
-		DiffListSource.Add(FSharedDiffOnGraph(new FDiffResultItem(*DiffIt)));
+		DiffListSource.Add(MakeShared<FDiffResultItem>(*DiffIt));
 	}
 
 	struct SortDiff
 	{
-		bool operator () (const FSharedDiffOnGraph& A, const FSharedDiffOnGraph& B) const
+		bool operator () (const TSharedPtr<FDiffResultItem>& A, const TSharedPtr<FDiffResultItem>& B) const
 		{
 			return A->Result.Diff < B->Result.Diff;
 		}
@@ -783,29 +776,9 @@ void FListItemGraphToDiff::BuildDiffSourceArray()
 	Sort(DiffListSource.GetData(), DiffListSource.Num(), SortDiff());
 }
 
-int32 FListItemGraphToDiff::GetCurrentDiffIndex() const
-{
-	if ( DiffList.IsValid() )
-	{
-		auto Selected = DiffList->GetSelectedItems();
-		if(Selected.Num() == 1)
-		{	
-			int32 Index = 0;
-			for(auto It(DiffListSource.CreateConstIterator());It;++It,Index++)
-			{
-				if(*It == Selected[0])
-				{
-					return Index;
-				}
-			}
-		}
-	}
-	return -1;
-}
-
-void FListItemGraphToDiff::OnGraphChanged( const FEdGraphEditAction& Action )
+void FGraphToDiff::OnGraphChanged( const FEdGraphEditAction& Action )
 {
-	Diff->OnGraphChanged(this);
+	DiffWidget->OnGraphChanged(this);
 }
 
 FDiffPanel::FDiffPanel()
@@ -825,7 +798,7 @@ void FDiffPanel::InitializeDiffPanel()
 	MyBlueprint->SetInspector(DetailsView);
 }
 
-int32 GetCurrentIndex( SListView< TSharedPtr< struct FDiffSingleResult> > const& ListView, const TArray< TSharedPtr< FDiffSingleResult > >& ListViewSource )
+int32 GetCurrentIndex( SListView< TSharedPtr< FDiffSingleResult> > const& ListView, const TArray< TSharedPtr< FDiffSingleResult > >& ListViewSource )
 {
 	const auto& Selected = ListView.GetSelectedItems();
 	if (Selected.Num() == 1)
@@ -864,13 +837,13 @@ void DiffWidgetUtils::SelectPrevRow(SListView< TSharedPtr< FDiffSingleResult> >&
 	ListView.SetSelection(ListViewSource[CurrentIndex - 1]);
 }
 
-bool DiffWidgetUtils::HasNextDifference(SListView< TSharedPtr< struct FDiffSingleResult> >& ListView, const TArray< TSharedPtr< struct FDiffSingleResult > >& ListViewSource)
+bool DiffWidgetUtils::HasNextDifference(SListView< TSharedPtr< FDiffSingleResult> >& ListView, const TArray< TSharedPtr< FDiffSingleResult > >& ListViewSource)
 {
 	int32 CurrentIndex = GetCurrentIndex(ListView, ListViewSource);
 	return ListViewSource.IsValidIndex(CurrentIndex+1);
 }
 
-bool DiffWidgetUtils::HasPrevDifference(SListView< TSharedPtr< struct FDiffSingleResult> >& ListView, const TArray< TSharedPtr< struct FDiffSingleResult > >& ListViewSource)
+bool DiffWidgetUtils::HasPrevDifference(SListView< TSharedPtr< FDiffSingleResult> >& ListView, const TArray< TSharedPtr< FDiffSingleResult > >& ListViewSource)
 {
 	int32 CurrentIndex = GetCurrentIndex(ListView, ListViewSource);
 	return ListViewSource.IsValidIndex(CurrentIndex - 1);
@@ -1085,22 +1058,12 @@ void SBlueprintDiff::OnCloseAssetEditor(UObject* Asset, EAssetEditorCloseReason
 	}
 }
 
-TSharedRef<ITableRow> SBlueprintDiff::OnGenerateRow( FGraphToDiff ParamItem, const TSharedRef<STableViewBase>& OwnerTable )
+void SBlueprintDiff::CreateGraphEntry( UEdGraph* GraphOld, UEdGraph* GraphNew )
 {
-	return	SNew( STableRow< FGraphToDiff >, OwnerTable )
-			.ToolTipText(ParamItem->GetToolTip())
-			.Content()
-			[
-				ParamItem->GenerateWidget()
-			];
-}
-
-void SBlueprintDiff::CreateGraphEntry( class UEdGraph* GraphOld, class UEdGraph* GraphNew )
-{
-	Graphs.Add(FGraphToDiff(new FListItemGraphToDiff(this, GraphOld, GraphNew, PanelOld.RevisionInfo, PanelNew.RevisionInfo)));
+	Graphs.Add(MakeShared<FGraphToDiff>(this, GraphOld, GraphNew, PanelOld.RevisionInfo, PanelNew.RevisionInfo));
 }
 
-void SBlueprintDiff::OnSelectionChanged( FGraphToDiff Item, ESelectInfo::Type SelectionType) 
+void SBlueprintDiff::OnGraphSelectionChanged(TSharedPtr<FGraphToDiff> Item, ESelectInfo::Type SelectionType)
 {
 	if(!Item.IsValid())
 	{
@@ -1111,7 +1074,7 @@ void SBlueprintDiff::OnSelectionChanged( FGraphToDiff Item, ESelectInfo::Type Se
 
 }
 
-void SBlueprintDiff::OnGraphChanged(FListItemGraphToDiff* Diff)
+void SBlueprintDiff::OnGraphChanged(FGraphToDiff* Diff)
 {
 	if(PanelNew.GraphEditor.IsValid() && PanelNew.GraphEditor.Pin()->GetCurrentGraph() == Diff->GetGraphNew())
 	{
@@ -1143,7 +1106,7 @@ TSharedRef<SWidget> SBlueprintDiff::DefaultEmptyPanel()
 		];
 }
 
-TSharedPtr<SWindow> SBlueprintDiff::CreateDiffWindow(FText WindowTitle, UBlueprint* OldBlueprint, UBlueprint* NewBlueprint, const struct FRevisionInfo& OldRevision, const struct FRevisionInfo& NewRevision)
+TSharedPtr<SWindow> SBlueprintDiff::CreateDiffWindow(FText WindowTitle, UBlueprint* OldBlueprint, UBlueprint* NewBlueprint, const FRevisionInfo& OldRevision, const FRevisionInfo& NewRevision)
 {
 	// sometimes we're comparing different revisions of one single asset (other 
 	// times we're comparing two completely separate assets altogether)
@@ -1195,7 +1158,7 @@ bool SBlueprintDiff::HasPrevDiff() const
 	return DiffTreeView::HasPrevDifference(DifferencesTreeView.ToSharedRef(), RealDifferences);
 }
 
-FListItemGraphToDiff* SBlueprintDiff::FindGraphToDiffEntry(const FString& GraphPath)
+FGraphToDiff* SBlueprintDiff::FindGraphToDiffEntry(const FString& GraphPath)
 {
 	for( const auto& Graph : Graphs )
 	{
@@ -1208,7 +1171,7 @@ FListItemGraphToDiff* SBlueprintDiff::FindGraphToDiffEntry(const FString& GraphP
 	return nullptr;
 }
 
-void SBlueprintDiff::FocusOnGraphRevisions( FListItemGraphToDiff* Diff )
+void SBlueprintDiff::FocusOnGraphRevisions( FGraphToDiff* Diff )
 {
 	UEdGraph* Graph = Diff->GetGraphOld() ? Diff->GetGraphOld() : Diff->GetGraphNew();
 
@@ -1221,8 +1184,8 @@ void SBlueprintDiff::FocusOnGraphRevisions( FListItemGraphToDiff* Diff )
 
 void SBlueprintDiff::OnDiffListSelectionChanged(TSharedPtr<FDiffResultItem> TheDiff )
 {
-	check( !TheDiff->Result.OwningGraphPath.IsEmpty() );
-	FocusOnGraphRevisions( FindGraphToDiffEntry( TheDiff->Result.OwningGraphPath) );
+	check( !TheDiff->Result.OwningObjectPath.IsEmpty() );
+	FocusOnGraphRevisions( FindGraphToDiffEntry( TheDiff->Result.OwningObjectPath) );
 	FDiffSingleResult Result = TheDiff->Result;
 
 	const auto SafeClearSelection = []( TWeakPtr<SGraphEditor> GraphEditor )
@@ -1316,7 +1279,7 @@ void FDiffPanel::GeneratePanel(UEdGraph* Graph, UEdGraph* GraphToDiff )
 
 		if ( !GraphEditorCommands.IsValid() )
 		{
-			GraphEditorCommands = TSharedPtr<FUICommandList>( new FUICommandList() );
+			GraphEditorCommands = MakeShared<FUICommandList>();
 
 			GraphEditorCommands->MapAction( FGenericCommands::Get().Copy,
 				FExecuteAction::CreateRaw( this, &FDiffPanel::CopySelectedNodes ),
@@ -1338,10 +1301,10 @@ void FDiffPanel::GeneratePanel(UEdGraph* Graph, UEdGraph* GraphToDiff )
 		Widget = Editor;
 	}
 
-	GraphEditorBorder->SetContent(Widget.ToSharedRef());
+	GraphEditorBox->SetContent(Widget.ToSharedRef());
 }
 
-TSharedRef<class SWidget> FDiffPanel::GenerateMyBlueprintWidget()
+TSharedRef<SWidget> FDiffPanel::GenerateMyBlueprintWidget()
 {
 	return SAssignNew(MyBlueprint, SMyBlueprint, TWeakPtr<FBlueprintEditor>(), Blueprint);
 }
@@ -1498,16 +1461,17 @@ void SBlueprintDiff::GenerateDifferencesList()
 		}
 	}
 
-	bool bHasComponents = true;
-	if (const UAnimBlueprint* AnimBP = Cast<UAnimBlueprint>(PanelOld.Blueprint))
+	bool bHasComponents = false;
+	UClass* BlueprintClass = PanelOld.Blueprint->GeneratedClass;
+	if (BlueprintClass->IsChildOf<AActor>())
 	{
-		MasterDifferencesList.Push(FBlueprintDifferenceTreeEntry::AnimBlueprintEntry());
-		bHasComponents = false;
+		bHasComponents = true;
 	}
-	else if (const UWidgetBlueprint* WidgetBP = Cast<UWidgetBlueprint>(PanelOld.Blueprint))
+
+	// If this isn't a normal blueprint type, add the type panel
+	if (PanelOld.Blueprint->GetClass() != UBlueprint::StaticClass())
 	{
-		MasterDifferencesList.Push(FBlueprintDifferenceTreeEntry::WidgetBlueprintEntry());
-		bHasComponents = false;
+		ModePanels.Add(BlueprintTypeMode, GenerateBlueprintTypePanel());
 	}
 
 	// Now that we have done the diffs, create the panel widgets
@@ -1520,45 +1484,50 @@ void SBlueprintDiff::GenerateDifferencesList()
 		ModePanels.Add(ComponentsMode, GenerateComponentsPanel());
 	}
 
-	for (const auto& Graph : Graphs)
+	for (const TSharedPtr<FGraphToDiff>& Graph : Graphs)
 	{
-		TArray< TSharedPtr<FBlueprintDifferenceTreeEntry> > Children;
-		for (const auto& Difference : Graph->DiffListSource)
-		{
-			auto ChildEntry = TSharedPtr<FBlueprintDifferenceTreeEntry>(
-				new FBlueprintDifferenceTreeEntry(
-				FOnDiffEntryFocused::CreateRaw(this, &SBlueprintDiff::OnDiffListSelectionChanged, Difference)
-				, FGenerateDiffEntryWidget::CreateSP(Difference.ToSharedRef(), &FDiffResultItem::GenerateWidget)
-				, TArray< TSharedPtr<FBlueprintDifferenceTreeEntry> >()
-				)
-				);
-			Children.Push(ChildEntry);
-			RealDifferences.Push(ChildEntry);
-			(void)ChildEntry;
-		}
-
-		if (Children.Num() == 0)
-		{
-			// make one child informing the user that there are no differences:
-			Children.Push(FBlueprintDifferenceTreeEntry::NoDifferencesEntry());
-		}
-
-		auto Entry = TSharedPtr<FBlueprintDifferenceTreeEntry>(
-			new FBlueprintDifferenceTreeEntry(
-			FOnDiffEntryFocused::CreateRaw(this, &SBlueprintDiff::OnSelectionChanged, Graph, ESelectInfo::Direct)
-			, FGenerateDiffEntryWidget::CreateSP(Graph.ToSharedRef(), &FListItemGraphToDiff::GenerateWidget)
-			, Children)
-			);
-		MasterDifferencesList.Push(Entry);
+		Graph->GenerateTreeEntries(MasterDifferencesList, RealDifferences);
 	}
 
 	DifferencesTreeView = DiffTreeView::CreateTreeView(&MasterDifferencesList);
 }
 
+SBlueprintDiff::FDiffControl SBlueprintDiff::GenerateBlueprintTypePanel()
+{
+	TSharedPtr<FBlueprintTypeDiffControl> NewDiffControl = MakeShared<FBlueprintTypeDiffControl>(PanelOld.Blueprint, PanelNew.Blueprint, FOnDiffEntryFocused::CreateRaw(this, &SBlueprintDiff::SetCurrentMode, BlueprintTypeMode));
+	NewDiffControl->GenerateTreeEntries(MasterDifferencesList, RealDifferences);
+
+	SBlueprintDiff::FDiffControl Ret;
+	//Splitter for left and right blueprint. Current convention is for the local (probably newer?) blueprint to be on the right:
+	Ret.DiffControl = NewDiffControl;
+	Ret.Widget = SNew(SSplitter)
+		.PhysicalSplitterHandleSize(10.0f)
+		+ SSplitter::Slot()
+		.Value(0.5f)
+		[
+			SAssignNew(NewDiffControl->OldDetailsBox, SBox)
+			.VAlign(VAlign_Fill)
+			[
+				DefaultEmptyPanel()
+			]
+		]
+		+ SSplitter::Slot()
+		.Value(0.5f)
+		[
+			SAssignNew(NewDiffControl->NewDetailsBox, SBox)
+			.VAlign(VAlign_Fill)
+			[
+				DefaultEmptyPanel()
+			]
+		];
+
+	return Ret;
+}
 
 SBlueprintDiff::FDiffControl SBlueprintDiff::GenerateMyBlueprintPanel()
 {
-	auto NewDiffControl = TSharedPtr<FMyBlueprintDiffControl>(new FMyBlueprintDiffControl(PanelOld.Blueprint, PanelNew.Blueprint, MasterDifferencesList, RealDifferences, FOnDiffControlChanged::CreateRaw(this, &SBlueprintDiff::SetCurrentMode, MyBlueprintMode)));
+	TSharedPtr<FMyBlueprintDiffControl> NewDiffControl = MakeShared<FMyBlueprintDiffControl>(PanelOld.Blueprint, PanelNew.Blueprint, FOnDiffEntryFocused::CreateRaw(this, &SBlueprintDiff::SetCurrentMode, MyBlueprintMode));
+	NewDiffControl->GenerateTreeEntries(MasterDifferencesList, RealDifferences);
 
 	SBlueprintDiff::FDiffControl Ret;
 
@@ -1630,7 +1599,7 @@ SBlueprintDiff::FDiffControl SBlueprintDiff::GenerateGraphPanel()
 				.PhysicalSplitterHandleSize(10.0f)
 				+ SSplitter::Slot()
 				[
-					SAssignNew(PanelOld.GraphEditorBorder, SBox)
+					SAssignNew(PanelOld.GraphEditorBox, SBox)
 					.VAlign(VAlign_Fill)
 					[
 						DefaultEmptyPanel()
@@ -1638,7 +1607,7 @@ SBlueprintDiff::FDiffControl SBlueprintDiff::GenerateGraphPanel()
 				]
 				+ SSplitter::Slot()
 				[
-					SAssignNew(PanelNew.GraphEditorBorder, SBox)
+					SAssignNew(PanelNew.GraphEditorBox, SBox)
 					.VAlign(VAlign_Fill)
 					[
 						DefaultEmptyPanel()
@@ -1670,10 +1639,10 @@ SBlueprintDiff::FDiffControl SBlueprintDiff::GenerateDefaultsPanel()
 	const UObject* A = DiffUtils::GetCDO(PanelOld.Blueprint);
 	const UObject* B = DiffUtils::GetCDO(PanelNew.Blueprint);
 
-	auto NewDiffControl = TSharedPtr<FCDODiffControl>(new FCDODiffControl(A, B, MasterDifferencesList, RealDifferences, FOnDiffControlChanged::CreateRaw(this, &SBlueprintDiff::SetCurrentMode, DefaultsMode)));
+	TSharedPtr<FCDODiffControl> NewDiffControl = MakeShared<FCDODiffControl>(A, B, FOnDiffEntryFocused::CreateRaw(this, &SBlueprintDiff::SetCurrentMode, DefaultsMode));
+	NewDiffControl->GenerateTreeEntries(MasterDifferencesList, RealDifferences);
 
 	SBlueprintDiff::FDiffControl Ret;
-	//Splitter for left and right blueprint. Current convention is for the local (probably newer?) blueprint to be on the right:
 	Ret.DiffControl = NewDiffControl;
 	Ret.Widget = SNew(SSplitter)
 		.PhysicalSplitterHandleSize(10.0f)
@@ -1693,13 +1662,10 @@ SBlueprintDiff::FDiffControl SBlueprintDiff::GenerateDefaultsPanel()
 
 SBlueprintDiff::FDiffControl SBlueprintDiff::GenerateClassSettingsPanel()
 {
-	const UObject* A = PanelOld.Blueprint;
-	const UObject* B = PanelNew.Blueprint;
-
-	auto NewDiffControl = TSharedPtr<FClassDefaultsDiffControl>(new FClassDefaultsDiffControl(A, B, MasterDifferencesList, RealDifferences, FOnDiffControlChanged::CreateRaw(this, &SBlueprintDiff::SetCurrentMode, ClassSettingsMode)));
+	TSharedPtr<FClassSettingsDiffControl> NewDiffControl = MakeShared<FClassSettingsDiffControl>(PanelOld.Blueprint, PanelNew.Blueprint, FOnDiffEntryFocused::CreateRaw(this, &SBlueprintDiff::SetCurrentMode, ClassSettingsMode));
+	NewDiffControl->GenerateTreeEntries(MasterDifferencesList, RealDifferences);
 
 	SBlueprintDiff::FDiffControl Ret;
-	//Splitter for left and right blueprint. Current convention is for the local (probably newer?) blueprint to be on the right:
 	Ret.DiffControl = NewDiffControl;
 	Ret.Widget = SNew(SSplitter)
 		.PhysicalSplitterHandleSize(10.0f)
@@ -1719,10 +1685,10 @@ SBlueprintDiff::FDiffControl SBlueprintDiff::GenerateClassSettingsPanel()
 
 SBlueprintDiff::FDiffControl SBlueprintDiff::GenerateComponentsPanel()
 {
-	SBlueprintDiff::FDiffControl Ret;
+	TSharedPtr<FSCSDiffControl> NewDiffControl = MakeShared<FSCSDiffControl>(PanelOld.Blueprint, PanelNew.Blueprint, FOnDiffEntryFocused::CreateRaw(this, &SBlueprintDiff::SetCurrentMode, ComponentsMode));
+	NewDiffControl->GenerateTreeEntries(MasterDifferencesList, RealDifferences);
 
-	//Splitter for left and right blueprint. Current convention is for the local (probably newer?) blueprint to be on the right:
-	auto NewDiffControl = TSharedPtr<FSCSDiffControl>(new FSCSDiffControl(PanelOld.Blueprint, PanelNew.Blueprint, MasterDifferencesList, RealDifferences, FOnDiffControlChanged::CreateRaw(this, &SBlueprintDiff::SetCurrentMode, ComponentsMode) ) );
+	SBlueprintDiff::FDiffControl Ret;
 	Ret.DiffControl = NewDiffControl;
 	Ret.Widget = SNew(SSplitter)
 		.PhysicalSplitterHandleSize(10.0f)
@@ -1766,7 +1732,7 @@ void SBlueprintDiff::SetCurrentMode(FName NewMode)
 	}
 }
 
-bool SBlueprintDiff::IsGraphDiffNeeded(class UEdGraph* InGraph) const
+bool SBlueprintDiff::IsGraphDiffNeeded(UEdGraph* InGraph) const
 {
 	// Do not worry about graphs that are contained in MathExpression nodes, they are recreated each compile
 	return !InGraph->GetOuter()->IsA<UK2Node_MathExpression>();
diff --git a/Engine/Source/Editor/Kismet/Public/DetailsDiff.h b/Engine/Source/Editor/Kismet/Public/DetailsDiff.h
--- a/Engine/Source/Editor/Kismet/Public/DetailsDiff.h
+++ b/Engine/Source/Editor/Kismet/Public/DetailsDiff.h
@@ -8,6 +8,7 @@
 
 class IDetailsView;
 
+/** Struct to handle showing details for an object and provide an interface for listing all differences */
 class KISMET_API FDetailsDiff
 {
 public:
@@ -16,11 +17,20 @@ public:
 	FDetailsDiff( const UObject* InObject, FOnDisplayedPropertiesChanged InOnDisplayedPropertiesChanged );
 	~FDetailsDiff();
 
+	/** Attempt to highlight the property with the given path, may not always succeed */
 	void HighlightProperty( const FPropertySoftPath& PropertyName );
+
+	/** Returns actual widget that is used to display details */
 	TSharedRef< SWidget > DetailsWidget();
+
+	/** Returns object being displayed */
+	const UObject* GetDisplayedObject() const { return DisplayedObject; }
+
+	/** Returns a list of all properties that would be diffed */
 	TArray<FPropertySoftPath> GetDisplayedProperties() const;
 
-	void DiffAgainst(const FDetailsDiff& Newer, TArray< FSingleObjectDiffEntry > &OutDifferences) const;
+	/** Perform a diff against another view, ordering either by display order or by remove/add/change */
+	void DiffAgainst(const FDetailsDiff& Newer, TArray<FSingleObjectDiffEntry>& OutDifferences, bool bSortByDisplayOrder = false) const;
 
 private:
 	void HandlePropertiesChanged();
diff --git a/Engine/Source/Editor/Kismet/Public/DiffUtils.h b/Engine/Source/Editor/Kismet/Public/DiffUtils.h
--- a/Engine/Source/Editor/Kismet/Public/DiffUtils.h
+++ b/Engine/Source/Editor/Kismet/Public/DiffUtils.h
@@ -223,7 +223,7 @@ DECLARE_DELEGATE_RetVal(TSharedRef<SWidget>, FGenerateDiffEntryWidget);
 class FBlueprintDifferenceTreeEntry
 {
 public:
-	FBlueprintDifferenceTreeEntry( FOnDiffEntryFocused InOnFocus, FGenerateDiffEntryWidget InGenerateWidget, TArray< TSharedPtr<FBlueprintDifferenceTreeEntry> > InChildren ) 
+	FBlueprintDifferenceTreeEntry(FOnDiffEntryFocused InOnFocus, FGenerateDiffEntryWidget InGenerateWidget, TArray< TSharedPtr<FBlueprintDifferenceTreeEntry> > InChildren = TArray< TSharedPtr<FBlueprintDifferenceTreeEntry> >())
 		: OnFocus(InOnFocus)
 		, GenerateWidget(InGenerateWidget)
 		, Children(InChildren) 
@@ -231,12 +231,16 @@ public:
 		check( InGenerateWidget.IsBound() );
 	}
 
-	/** The FBlueprintDifferenceTreeEntry used to display a message to the user explaining that there are no differences: */
+	/** Displays message to user saying there are no differences */
 	KISMET_API static TSharedPtr<FBlueprintDifferenceTreeEntry> NoDifferencesEntry();
-	KISMET_API static TSharedPtr<FBlueprintDifferenceTreeEntry> AnimBlueprintEntry();
-	KISMET_API static TSharedPtr<FBlueprintDifferenceTreeEntry> WidgetBlueprintEntry();
-	/** The FBlueprintDifferenceTreeEntry used to label categories: */
+
+	/** Displays message to user warning that there may be undetected differences */
+	KISMET_API static TSharedPtr<FBlueprintDifferenceTreeEntry> UnknownDifferencesEntry();
+
+	/** Create category message for the diff UI */
 	KISMET_API static TSharedPtr<FBlueprintDifferenceTreeEntry> CreateCategoryEntry(const FText& LabelText, const FText& ToolTipText, FOnDiffEntryFocused FocusCallback, const TArray< TSharedPtr<FBlueprintDifferenceTreeEntry> >& Children, bool bHasDifferences);
+
+	/** Create category message for the merge UI */
 	KISMET_API static TSharedPtr<FBlueprintDifferenceTreeEntry> CreateCategoryEntryForMerge(const FText& LabelText, const FText& ToolTipText, FOnDiffEntryFocused FocusCallback, const TArray< TSharedPtr<FBlueprintDifferenceTreeEntry> >& Children, bool bHasRemoteDifferences, bool bHasLocalDifferences, bool bHasConflicts);
 	
 	FOnDiffEntryFocused OnFocus;
diff --git a/Engine/Source/Editor/Kismet/Public/SBlueprintDiff.h b/Engine/Source/Editor/Kismet/Public/SBlueprintDiff.h
--- a/Engine/Source/Editor/Kismet/Public/SBlueprintDiff.h
+++ b/Engine/Source/Editor/Kismet/Public/SBlueprintDiff.h
@@ -11,7 +11,7 @@
 #include "Widgets/Views/SListView.h"
 #include "GraphEditor.h"
 #include "DiffUtils.h"
-#include "Editor/GraphEditor/Public/DiffResults.h"
+#include "DiffResults.h"
 #include "SKismetInspector.h"
 #include "Developer/AssetTools/Public/IAssetTypeActions.h"
 
@@ -20,7 +20,7 @@ class FTabManager;
 class IDiffControl;
 class SMyBlueprint;
 class UEdGraph;
-struct FListItemGraphToDiff;
+struct FGraphToDiff;
 enum class EAssetEditorCloseReason : uint8;
 
 /** Individual Diff item shown in the list of diffs */
@@ -43,7 +43,7 @@ namespace DiffWidgetUtils
 	KISMET_API bool HasPrevDifference(SListView< TSharedPtr< struct FDiffSingleResult> >& ListView, const TArray< TSharedPtr< struct FDiffSingleResult > >& ListViewSource);
 }
 
-/*panel used to display the blueprint*/
+/** Panel used to display the blueprint */
 struct KISMET_API FDiffPanel
 {
 	FDiffPanel();
@@ -51,50 +51,50 @@ struct KISMET_API FDiffPanel
 	/** Initializes the panel, can be moved into constructor if diff and merge clients are made more uniform: */
 	void InitializeDiffPanel();
 
-	/* Generate this panel based on the specified graph */
+	/** Generate this panel based on the specified graph */
 	void GeneratePanel(UEdGraph* Graph, UEdGraph* GraphToDiff);
 
-	/* Generate the 'MyBlueprint' widget, which is private to this module */
+	/** Generate the 'MyBlueprint' widget, which is private to this module */
 	TSharedRef<class SWidget> GenerateMyBlueprintWidget();
 
-	/* Called when user hits keyboard shortcut to copy nodes*/
+	/** Called when user hits keyboard shortcut to copy nodes */
 	void CopySelectedNodes();
 
-	/*Gets whatever nodes are selected in the Graph Editor*/
+	/** Gets whatever nodes are selected in the Graph Editor */
 	FGraphPanelSelectionSet GetSelectedNodes() const;
 
-	/*Can user copy any of the selected nodes?*/
+	/** Can user copy any of the selected nodes? */
 	bool CanCopyNodes() const;
 
-	/*Functions used to focus/find a particular change in a diff result*/
+	/** Functions used to focus/find a particular change in a diff result */
 	void FocusDiff(UEdGraphPin& Pin);
 	void FocusDiff(UEdGraphNode& Node);
 
-	/*The blueprint that owns the graph we are showing*/
-	const class UBlueprint*				Blueprint;
+	/** The blueprint that owns the graph we are showing */
+	const UBlueprint*				Blueprint;
 
-	/*The border around the graph editor, used to change the content when new graphs are set */
-	TSharedPtr<SBox>				GraphEditorBorder;
+	/** The box around the graph editor, used to change the content when new graphs are set */
+	TSharedPtr<SBox>				GraphEditorBox;
 
-	/* The border around the my blueprint panel, used to regenerate the panel when the new graphs are set */
+	/** The actual my blueprint panel, used to regenerate the panel when the new graphs are set */
 	TSharedPtr<class SMyBlueprint>	MyBlueprint;
 
-	/*The box around the the details view associated with the graph editor */
+	/** The details view associated with the graph editor */
 	TSharedPtr<class SKismetInspector>	DetailsView;
 
-	/*The graph editor which does the work of displaying the graph*/
+	/** The graph editor which does the work of displaying the graph */
 	TWeakPtr<class SGraphEditor>	GraphEditor;
 
-	/*Revision information for this blueprint */
+	/** Revision information for this blueprint */
 	FRevisionInfo					RevisionInfo;
 
-	/*A name identifying which asset this panel is displaying */
+	/** True if we should show a name identifying which asset this panel is displaying */
 	bool							bShowAssetName;
 
-	/*The panel stores the last pin that was focused on by the user, so that it can clear the visual style when selection changes*/
+	/** The panel stores the last pin that was focused on by the user, so that it can clear the visual style when selection changes */
 	UEdGraphPin*					LastFocusedPin;
 private:
-	/*Command list for this diff panel*/
+	/** Command list for this diff panel */
 	TSharedPtr<FUICommandList> GraphEditorCommands;
 };
 
@@ -117,11 +117,17 @@ public:
 	virtual ~SBlueprintDiff();
 
 	/** Called when a new Graph is clicked on by user */
-	void OnGraphChanged(struct FListItemGraphToDiff* Diff);
+	void OnGraphChanged(FGraphToDiff* Diff);
 
 	/** Called when blueprint is modified */
 	void OnBlueprintChanged(UBlueprint* InBlueprint);
 
+	/** Called when user clicks on a new graph list item */
+	void OnGraphSelectionChanged(TSharedPtr<FGraphToDiff> Item, ESelectInfo::Type SelectionType);
+
+	/** Called when user clicks on an entry in the listview of differences */
+	void OnDiffListSelectionChanged(TSharedPtr<struct FDiffResultItem> TheDiff);
+
 	/** Helper function for generating an empty widget */
 	static TSharedRef<SWidget> DefaultEmptyPanel();
 
@@ -139,41 +145,29 @@ protected:
 	bool HasNextDiff() const;
 	bool HasPrevDiff() const;
 
-	typedef TSharedPtr<struct FListItemGraphToDiff>	FGraphToDiff;
-	typedef SListView<FGraphToDiff >	SListViewType;
-
-	/** Find the FListItemGraphToDiff that displays the graph with GraphPath relative path */
-	FListItemGraphToDiff* FindGraphToDiffEntry(const FString& GraphPath);
+	/** Find the FGraphToDiff that displays the graph with GraphPath relative path */
+	FGraphToDiff* FindGraphToDiffEntry(const FString& GraphPath);
 
 	/** Bring these revisions of graph into focus on main display*/
-	void FocusOnGraphRevisions( struct FListItemGraphToDiff* Diff);
+	void FocusOnGraphRevisions(FGraphToDiff* Diff);
 
-	/*Create a list item entry graph that exists in at least one of the blueprints */
+	/** Create a list item entry graph that exists in at least one of the blueprints */
 	void CreateGraphEntry(class UEdGraph* GraphOld, class UEdGraph* GraphNew);
-
-	/* Called when a new row is being generated */
-	TSharedRef<ITableRow> OnGenerateRow(FGraphToDiff ParamItem, const TSharedRef<STableViewBase>& OwnerTable );
-
-	/*Called when user clicks on a new graph list item */
-	void OnSelectionChanged(FGraphToDiff Item, ESelectInfo::Type SelectionType);
-
-	/** Called when user clicks on an entry in the listview of differences */
-	void OnDiffListSelectionChanged(TSharedPtr<struct FDiffResultItem> TheDiff );
 		
 	/** Disable the focus on a particular pin */
 	void DisablePinDiffFocus();
 
-	/*User toggles the option to lock the views between the two blueprints */
+	/** User toggles the option to lock the views between the two blueprints */
 	void OnToggleLockView();
 
-	/*Reset the graph editor, called when user switches graphs to display*/
+	/** Reset the graph editor, called when user switches graphs to display*/
 	void ResetGraphEditors();
 
-	/*Get the image to show for the toggle lock option*/
+	/** Get the image to show for the toggle lock option*/
 	FSlateIcon GetLockViewImage() const;
 
-	/* This buffer stores the currently displayed results */
-	TArray< FGraphToDiff> Graphs;
+	/** List of graphs to diff, are added to panel last */
+	TArray<TSharedPtr<FGraphToDiff>> Graphs;
 
 	/** Get Graph editor associated with this Graph */
 	FDiffPanel& GetDiffPanelForNode(UEdGraphNode& Node);
@@ -202,6 +196,7 @@ protected:
 		TSharedPtr< class IDiffControl > DiffControl;
 	};
 
+	FDiffControl GenerateBlueprintTypePanel();
 	FDiffControl GenerateMyBlueprintPanel();
 	FDiffControl GenerateGraphPanel();
 	FDiffControl GenerateDefaultsPanel();
diff --git a/Engine/Source/Editor/Kismet/Public/SCSDiff.h b/Engine/Source/Editor/Kismet/Public/SCSDiff.h
--- a/Engine/Source/Editor/Kismet/Public/SCSDiff.h
+++ b/Engine/Source/Editor/Kismet/Public/SCSDiff.h
@@ -9,6 +9,7 @@ class FSCSEditorTreeNode;
 class SKismetInspector;
 class SSCSEditor;
 
+/** Struct to support diffing the component tree for a blueprint */
 class KISMET_API FSCSDiff
 {
 public:
@@ -19,6 +20,8 @@ public:
 
 	TArray< FSCSResolvedIdentifier > GetDisplayedHierarchy() const;
 
+	const UBlueprint* GetBlueprint() const { return Blueprint; }
+
 protected:
 	void OnSCSEditorUpdateSelectionFromNodes(const TArray< TSharedPtr<class FSCSEditorTreeNode> >& SelectedNodes);
 	void OnSCSEditorHighlightPropertyInDetailsView(const class FPropertyPath& InPropertyPath);
diff --git a/Engine/Source/Editor/UMGEditor/Private/WidgetBlueprint.cpp b/Engine/Source/Editor/UMGEditor/Private/WidgetBlueprint.cpp
--- a/Engine/Source/Editor/UMGEditor/Private/WidgetBlueprint.cpp
+++ b/Engine/Source/Editor/UMGEditor/Private/WidgetBlueprint.cpp
@@ -25,6 +25,7 @@
 #if WITH_EDITOR
 #include "Interfaces/ITargetPlatform.h"
 #include "Modules/ModuleManager.h"
+#include "DiffResults.h"
 #endif
 #include "Kismet2/BlueprintEditorUtils.h"
 #include "K2Node_CallFunction.h"
@@ -606,6 +607,122 @@ void UWidgetBlueprint::NotifyGraphRenamed(class UEdGraph* Graph, FName OldName,
 	});
 }
 
+bool UWidgetBlueprint::FindDiffs(const UBlueprint* OtherBlueprint, FDiffResults& Results) const
+{
+	const UWidgetBlueprint* OtherWidgetBP = Cast<UWidgetBlueprint>(OtherBlueprint);
+	if (!OtherWidgetBP)
+	{
+		return false;
+	}
+
+	// Look for all widget instances in both, add shared ones to ObjectsToDiff and add notes for add/remove
+	TMap<FString, UWidget*> WidgetMap;
+	TMap<FString, UWidget*> OtherWidgetMap;
+
+	WidgetTree->ForEachWidget([&](UWidget* Widget) 
+	{
+		FString WidgetPath = Widget->GetPathName(this);
+		WidgetMap.Add(WidgetPath, Widget);
+	});
+
+	OtherWidgetBP->WidgetTree->ForEachWidget([&](UWidget* Widget)
+	{
+		FString WidgetPath = Widget->GetPathName(OtherWidgetBP);
+		OtherWidgetMap.Add(WidgetPath, Widget);
+	});
+
+	for (TPair<FString, UWidget*> Pair : WidgetMap)
+	{
+		UWidget** FoundOtherWidget = OtherWidgetMap.Find(Pair.Key);
+		UWidget* Widget = Pair.Value;
+
+		if (FoundOtherWidget)
+		{
+			// Add to general object diff map
+			FDiffSingleResult Diff;
+			Diff.Diff = EDiffType::OBJECT_REQUEST_DIFF;
+			Diff.Object1 = Widget;
+			Diff.Object2 = *FoundOtherWidget;
+			Diff.OwningObjectPath = Pair.Key;
+
+			if (Results.CanStoreResults())
+			{
+				FFormatNamedArguments Args;
+				Args.Add(TEXT("WidgetTitle"), Widget->GetLabelTextWithMetadata());
+				Args.Add(TEXT("WidgetPath"), FText::FromString(Pair.Key));
+				Diff.ToolTip = FText::Format(LOCTEXT("DIF_RequestWidgetTooltip", "Widget {WidgetTitle}\nPath: {WidgetPath}"), Args);
+				Diff.DisplayString = FText::Format(LOCTEXT("DIF_RequestWidgetLabel", "Widget {WidgetTitle}"), Args);
+				Diff.DisplayColor = FLinearColor(1.f, 0.4f, 0.4f);
+			}
+
+			Results.Add(Diff);
+
+			// TODO: Add some type specific diffing if needed
+		}
+		else
+		{
+			// This is newly added
+			FDiffSingleResult Diff;
+			Diff.Diff = EDiffType::OBJECT_ADDED;
+			Diff.Object1 = Widget;
+			Diff.OwningObjectPath = Pair.Key;
+
+			if (Results.CanStoreResults())
+			{
+				FFormatNamedArguments Args;
+				Args.Add(TEXT("WidgetTitle"), Widget->GetLabelTextWithMetadata());
+				Args.Add(TEXT("WidgetPath"), FText::FromString(Pair.Key));
+				Diff.ToolTip = FText::Format(LOCTEXT("DIF_AddedWidgetTooltip", "Added Widget {WidgetTitle}\nPath: {WidgetPath}"), Args);
+				Diff.DisplayString = FText::Format(LOCTEXT("DIF_AddedWidgetLabel", "Added Widget {WidgetTitle}"), Args);
+				Diff.DisplayColor = FLinearColor(1.f, 0.4f, 0.4f);
+			}
+
+			Results.Add(Diff);
+		}
+	}
+
+	for (TPair<FString, UWidget*> Pair : OtherWidgetMap)
+	{
+		UWidget** FoundMyWidget = WidgetMap.Find(Pair.Key);
+		UWidget* OtherWidget = Pair.Value;
+
+		if (!FoundMyWidget)
+		{
+			// This is newly added
+			FDiffSingleResult Diff;
+			Diff.Diff = EDiffType::OBJECT_REMOVED;
+			Diff.Object1 = OtherWidget;
+			Diff.OwningObjectPath = Pair.Key;
+
+			if (Results.CanStoreResults())
+			{
+				FFormatNamedArguments Args;
+				Args.Add(TEXT("WidgetTitle"), OtherWidget->GetLabelTextWithMetadata());
+				Args.Add(TEXT("WidgetPath"), FText::FromString(Pair.Key));
+				Diff.ToolTip = FText::Format(LOCTEXT("DIF_RemovedWidgetTooltip", "Removed Widget {WidgetTitle}\nPath:{WidgetPath}"), Args);
+				Diff.DisplayString = FText::Format(LOCTEXT("DIF_RemovedWidgetLabel", "Removed Widget {WidgetTitle}"), Args);
+				Diff.DisplayColor = FLinearColor(1.f, 0.4f, 0.4f);
+			}
+
+			Results.Add(Diff);
+		}
+	}
+
+	// Add info warning
+	if (Results.CanStoreResults())
+	{
+		FDiffSingleResult Diff;
+		Diff.Diff = EDiffType::INFO_MESSAGE;
+		Diff.DisplayColor = FLinearColor(.7f, .7f, .7f);
+		Diff.ToolTip = LOCTEXT("DIF_WidgetWarningMessage", "Warning: May be missing some widget-specific differences");
+		Diff.DisplayString = Diff.ToolTip;
+
+		Results.Add(Diff);
+	}
+
+	return true;
+}
+
 #endif
 
 void UWidgetBlueprint::Serialize(FArchive& Ar)
diff --git a/Engine/Source/Editor/UMGEditor/Public/WidgetBlueprint.h b/Engine/Source/Editor/UMGEditor/Public/WidgetBlueprint.h
--- a/Engine/Source/Editor/UMGEditor/Public/WidgetBlueprint.h
+++ b/Engine/Source/Editor/UMGEditor/Public/WidgetBlueprint.h
@@ -275,6 +275,7 @@ public:
 #if WITH_EDITOR
 	virtual void GetAssetRegistryTags(TArray<FAssetRegistryTag>& OutTags) const override;
 	virtual void NotifyGraphRenamed(class UEdGraph* Graph, FName OldName, FName NewName) override;
+	virtual bool FindDiffs(const UBlueprint* OtherBlueprint, FDiffResults& Results) const override;
 #endif
 
 	virtual void Serialize(FArchive& Ar) override;
diff --git a/Engine/Source/Runtime/Engine/Classes/Animation/AnimBlueprint.h b/Engine/Source/Runtime/Engine/Classes/Animation/AnimBlueprint.h
--- a/Engine/Source/Runtime/Engine/Classes/Animation/AnimBlueprint.h
+++ b/Engine/Source/Runtime/Engine/Classes/Animation/AnimBlueprint.h
@@ -147,6 +147,7 @@ class ENGINE_API UAnimBlueprint : public UBlueprint, public IInterface_PreviewMe
 	}
 
 	virtual void PostLoad() override;
+	virtual bool FindDiffs(const UBlueprint* OtherBlueprint, FDiffResults& Results) const override;
 
 protected:
 	// Broadcast when an override is changed, allowing derived blueprints to be updated
diff --git a/Engine/Source/Runtime/Engine/Classes/Engine/Blueprint.h b/Engine/Source/Runtime/Engine/Classes/Engine/Blueprint.h
--- a/Engine/Source/Runtime/Engine/Classes/Engine/Blueprint.h
+++ b/Engine/Source/Runtime/Engine/Classes/Engine/Blueprint.h
@@ -21,6 +21,7 @@ class UActorComponent;
 class UEdGraph;
 class UInheritableComponentHandler;
 class FBlueprintActionDatabaseRegistrar;
+struct FDiffResults;
 
 /**
  * Enumerates states a blueprint can be in.
@@ -849,6 +850,16 @@ public:
 	/** Returns Valid if this object has data validation rules set up for it and the data for this object is valid. Returns Invalid if it does not pass the rules. Returns NotValidated if no rules are set for this object. */
 	virtual EDataValidationResult IsDataValid(TArray<FText>& ValidationErrors) override;
 
+	/** 
+	 * Fills in a list of differences between this blueprint and another blueprint.
+	 * Default blueprints are handled by SBlueprintDiff, this should be overridden for specific blueprint types.
+	 *
+	 * @param OtherBlueprint	Other blueprint to compare this to, should be the same type
+	 * @param Results			List of diff results to fill in with type-specific differences
+	 * @return					True if these blueprints were checked for specific differences, false if they are not comparable
+	 */
+	virtual bool FindDiffs(const UBlueprint* OtherBlueprint, FDiffResults& Results) const;
+
 #endif	//#if WITH_EDITOR
 
 	//~ Begin UObject Interface
diff --git a/Engine/Source/Runtime/Engine/Private/AnimBlueprint.cpp b/Engine/Source/Runtime/Engine/Private/AnimBlueprint.cpp
--- a/Engine/Source/Runtime/Engine/Private/AnimBlueprint.cpp
+++ b/Engine/Source/Runtime/Engine/Private/AnimBlueprint.cpp
@@ -180,6 +180,19 @@ bool UAnimBlueprint::CanRecompileWhilePlayingInEditor() const
 {
 	return true;
 }
+
+bool UAnimBlueprint::FindDiffs(const UBlueprint* OtherBlueprint, FDiffResults& Results) const
+{
+	const UAnimBlueprint* OtherAnimBP = Cast<UAnimBlueprint>(OtherBlueprint);
+	if (!OtherAnimBP)
+	{
+		return false;
+	}
+
+	// Anim BPs should diff correctly, as all the info is stored in graphs or the parent
+	return true;
+}
+
 #endif
 
 USkeletalMesh* UAnimBlueprint::GetPreviewMesh(bool bFindIfNotSet/*=false*/)
diff --git a/Engine/Source/Runtime/Engine/Private/Blueprint.cpp b/Engine/Source/Runtime/Engine/Private/Blueprint.cpp
--- a/Engine/Source/Runtime/Engine/Private/Blueprint.cpp
+++ b/Engine/Source/Runtime/Engine/Private/Blueprint.cpp
@@ -1910,6 +1910,11 @@ EDataValidationResult UBlueprint::IsDataValid(TArray<FText>& ValidationErrors)
 	return GeneratedClass ? GeneratedClass->GetDefaultObject()->IsDataValid(ValidationErrors) : EDataValidationResult::Invalid;
 }
 
+bool UBlueprint::FindDiffs(const UBlueprint* OtherBlueprint, FDiffResults& Results) const
+{
+	return false;
+}
+
 FName UBlueprint::GetFunctionNameFromClassByGuid(const UClass* InClass, const FGuid FunctionGuid)
 {
 	return FBlueprintEditorUtils::GetFunctionNameFromClassByGuid(InClass, FunctionGuid);
diff --git a/Engine/Source/Runtime/Engine/Private/EdGraph/EdGraphNode.cpp b/Engine/Source/Runtime/Engine/Private/EdGraph/EdGraphNode.cpp
--- a/Engine/Source/Runtime/Engine/Private/EdGraph/EdGraphNode.cpp
+++ b/Engine/Source/Runtime/Engine/Private/EdGraph/EdGraphNode.cpp
@@ -14,7 +14,7 @@
 #include "UObject/PropertyPortFlags.h"
 #include "ScopedTransaction.h"
 #include "FindInBlueprintManager.h"
-#include "Editor/GraphEditor/Public/DiffResults.h"
+#include "DiffResults.h"
 #endif
 
 #define LOCTEXT_NAMESPACE "EdGraph"
diff --git a/Engine/Source/Runtime/Engine/Public/DiffResults.h b/Engine/Source/Runtime/Engine/Public/DiffResults.h
new file mode 100644
--- /dev/null
+++ b/Engine/Source/Runtime/Engine/Public/DiffResults.h
@@ -0,0 +1,167 @@
+// Copyright 1998-2019 Epic Games, Inc. All Rights Reserved.
+
+#pragma once
+
+#include "CoreMinimal.h"
+
+/** Differences found within a graph or object */
+namespace EDiffType
+{
+	/** Differences are presented to the user in the order listed here, so put less important differences lower down */
+	enum Type
+	{
+		// No actual difference
+		NO_DIFFERENCE,
+		
+		// Object does not exist in current version
+		OBJECT_REMOVED,
+
+		// Object was added to current version
+		OBJECT_ADDED,
+
+		// An individual property differs
+		OBJECT_PROPERTY,
+
+		// Requests a full diff to be done on two objects, this will expand into other diffs
+		OBJECT_REQUEST_DIFF,
+
+		NODE_REMOVED,
+		NODE_ADDED,
+		PIN_LINKEDTO_NUM_DEC,
+		PIN_LINKEDTO_NUM_INC,
+		PIN_DEFAULT_VALUE,
+		PIN_TYPE_CATEGORY,
+		PIN_TYPE_SUBCATEGORY,
+		PIN_TYPE_SUBCATEGORY_OBJECT,
+		PIN_TYPE_IS_ARRAY,
+		PIN_TYPE_IS_REF,
+		PIN_LINKEDTO_NODE,
+		NODE_MOVED,
+		TIMELINE_LENGTH,
+		TIMELINE_AUTOPLAY,
+		TIMELINE_LOOP,
+		TIMELINE_IGNOREDILATION,
+		TIMELINE_NUM_TRACKS,
+		TIMELINE_TRACK_MODIFIED,
+		NODE_PIN_COUNT,
+		NODE_COMMENT,
+		NODE_PROPERTY,
+
+		// Informational message, does't count as a real diff
+		INFO_MESSAGE
+	};
+}
+
+/** Result of a single difference found on graph or object */
+struct FDiffSingleResult
+{
+	FDiffSingleResult()
+	{
+		Diff = EDiffType::NO_DIFFERENCE;
+		Node1 = nullptr;
+		Node2 = nullptr;
+		Pin1 = nullptr;
+		Pin2 = nullptr;
+		Object1 = nullptr;
+		Object2 = nullptr;
+		DisplayColor = FLinearColor::White;
+	}
+
+	/** The type of diff */
+	EDiffType::Type Diff;
+
+	/** The first node involved in diff */
+	class UEdGraphNode* Node1;
+
+	/** The second node involved in diff */
+	class UEdGraphNode* Node2;
+
+	/** The first pin involved in diff */
+	class UEdGraphPin* Pin1;
+
+	/** The second pin involved in diff */
+	class UEdGraphPin* Pin2;
+
+	/** First top-level object involved in a diff */
+	UObject* Object1;
+
+	/** Second top-level object involved in a diff */
+	UObject* Object2;
+
+	/** String describing the error to the user */
+	FText DisplayString;
+
+	/** Optional tooltip containing more information */
+	FText ToolTip; 
+
+	/** User can override color to use for display string */
+	FLinearColor DisplayColor;
+
+	/** Path string of graph, relative to blueprint/asset root */
+	FString OwningObjectPath;
+
+	/** Returns true if this is a confirmed difference */
+	FORCEINLINE bool IsRealDifference() const 
+	{
+		switch (Diff)
+		{
+		case EDiffType::NO_DIFFERENCE:
+		case EDiffType::OBJECT_REQUEST_DIFF:
+		case EDiffType::INFO_MESSAGE:
+			return false;
+		default:
+			return true;
+		}
+	}
+};
+
+FORCEINLINE bool operator==( const FDiffSingleResult& LHS, const FDiffSingleResult& RHS )
+{
+	return	LHS.Diff == RHS.Diff &&
+			LHS.Node1 == RHS.Node1 &&
+			LHS.Node2 == RHS.Node2 &&
+			LHS.Pin1 == RHS.Pin1 &&
+			LHS.Pin2 == RHS.Pin2 &&
+			LHS.Object1 == RHS.Object1 &&
+			LHS.Object2 == RHS.Object2 &&
+			LHS.DisplayString.ToString() == RHS.DisplayString.ToString() &&
+			LHS.ToolTip.ToString() == RHS.ToolTip.ToString() &&
+			LHS.DisplayColor == RHS.DisplayColor;
+}
+
+/** Collects the Diffs found for a node/object */
+struct FDiffResults
+{
+	FDiffResults(TArray<FDiffSingleResult>* InResultArray): ResultArray(InResultArray), bHasFoundDiffs(false)
+	{}
+
+	/** Add a diff that was found */
+	void Add(const FDiffSingleResult& Result)
+	{
+		if(Result.Diff != EDiffType::NO_DIFFERENCE)
+		{
+			bHasFoundDiffs = true;
+			if(ResultArray)
+			{
+				ResultArray->Add(Result);
+			}
+		}
+	}
+
+	/** Test if it can store results*/
+	bool CanStoreResults() const
+	{
+		return ResultArray != nullptr;
+	}
+
+	/** Get the number of diffs found*/
+	int32 Num() const { return ResultArray ? ResultArray->Num() : 0;}
+
+	/** True if diffs were found */
+	bool HasFoundDiffs() const { return bHasFoundDiffs; }
+
+private:
+	/** Pointer to optional array, passed in by user to store results in */
+	TArray<FDiffSingleResult>* ResultArray;
+	bool bHasFoundDiffs;
+};
