CIS fix

diff --git a/Engine/Source/Editor/Kismet/Private/DiffUtils.cpp b/Engine/Source/Editor/Kismet/Private/DiffUtils.cpp
--- a/Engine/Source/Editor/Kismet/Private/DiffUtils.cpp
+++ b/Engine/Source/Editor/Kismet/Private/DiffUtils.cpp
@@ -12,6 +12,7 @@
 #include "Engine/Blueprint.h"
 #include "IAssetTypeActions.h"
 #include "ObjectEditorUtils.h"
+#include "Components/ActorComponent.h"
 
 namespace UE4DiffUtils_Private
 {
