[NO NON-PII COMMIT MESSAGE]

diff --git a/Engine/Source/Editor/Kismet/Private/SBlueprintDiff.cpp b/Engine/Source/Editor/Kismet/Private/SBlueprintDiff.cpp
--- a/Engine/Source/Editor/Kismet/Private/SBlueprintDiff.cpp
+++ b/Engine/Source/Editor/Kismet/Private/SBlueprintDiff.cpp
@@ -161,7 +161,7 @@ public:
 		};
 
 		TArray< TSharedPtr<FBlueprintDifferenceTreeEntry> > Children;
-		for (auto Difference : DifferingProperties.Entries)
+		for (const FSCSDiffEntry& Difference : DifferingProperties.Entries)
 		{
 			TSharedPtr<FBlueprintDifferenceTreeEntry> Entry = MakeShared<FBlueprintDifferenceTreeEntry>(
 				FOnDiffEntryFocused::CreateStatic(FocusSCSDifferenceEntry, Difference, SelectionCallback, this),
@@ -248,7 +248,7 @@ protected:
 };
 
 /** Override for CDO special case */
-class FCDODiffControl : public TSharedFromThis<FCDODiffControl>, public FDetailsDiffControl
+class FCDODiffControl : public FDetailsDiffControl
 {
 public:
 	FCDODiffControl(const UObject* InOldObject, const UObject* InNewObject, FOnDiffEntryFocused InSelectionCallback)
@@ -278,7 +278,7 @@ public:
 };
 
 /** Override for class class settings */
-class FClassSettingsDiffControl : public TSharedFromThis<FClassSettingsDiffControl>, public FDetailsDiffControl
+class FClassSettingsDiffControl : public FDetailsDiffControl
 {
 public:
 	FClassSettingsDiffControl(const UObject* InOldObject, const UObject* InNewObject, FOnDiffEntryFocused InSelectionCallback)
@@ -715,7 +715,7 @@ FText FGraphToDiff::GetToolTip()
 	}
 	else
 	{
-		auto GoodGraph = GraphOld ? GraphOld : GraphNew;
+		UEdGraph* GoodGraph = GraphOld ? GraphOld : GraphNew;
 		check(GoodGraph);
 		const FRevisionInfo& Revision = GraphNew ? RevisionOld : RevisionNew;
 		FText RevisionText = LOCTEXT("CurrentRevision", "Current Revision");
