[NO NON-PII COMMIT MESSAGE]

diff --git a/Engine/Source/Editor/Kismet/Private/DiffUtils.cpp b/Engine/Source/Editor/Kismet/Private/DiffUtils.cpp
--- a/Engine/Source/Editor/Kismet/Private/DiffUtils.cpp
+++ b/Engine/Source/Editor/Kismet/Private/DiffUtils.cpp
@@ -34,25 +34,30 @@ namespace UE4DiffUtils_Private
 		return nullptr;
 	}
 
-	FPropertySoftPathSet GetPropertyNameSet(const UObject* ForObj)
+	FPropertySoftPathSet GetPropertyNameSet(const UStruct* ForStruct)
 	{
-		return FPropertySoftPathSet(DiffUtils::GetVisiblePropertiesInOrderDeclared(ForObj));
+		return FPropertySoftPathSet(DiffUtils::GetVisiblePropertiesInOrderDeclared(ForStruct));
 	}
 }
 
 FResolvedProperty FPropertySoftPath::Resolve(const UObject* Object) const
+{
+	return Resolve(Object->GetClass(), Object);
+}
+
+FResolvedProperty FPropertySoftPath::Resolve(const UStruct* Struct, const void* StructData) const
 {
 	// dig into the object, finding nested objects, etc:
-	const void* CurrentBlock = Object;
-	const UStruct* NextClass = Object->GetClass();
+	const void* CurrentBlock = StructData;
+	const UStruct* NextClass = Struct;
 	const void* NextBlock = CurrentBlock;
 	const UProperty* Property = nullptr;
 
-	for( int32 i = 0; i < PropertyChain.Num(); ++i )
+	for (int32 i = 0; i < PropertyChain.Num(); ++i)
 	{
 		CurrentBlock = NextBlock;
-		const UProperty* NextProperty = UE4DiffUtils_Private::Resolve(NextClass, PropertyChain[i]);
-		if( NextProperty )
+		const UProperty* NextProperty = UE4DiffUtils_Private::Resolve(NextClass, PropertyChain[i].PropertyName);
+		if (NextProperty)
 		{
 			Property = NextProperty;
 			if (const UObjectProperty* ObjectProperty = Cast<UObjectProperty>(Property))
@@ -103,11 +108,11 @@ FPropertyPath FPropertySoftPath::ResolvePath(const UObject* Object) const
 		}
 	};
 
-	auto TryReadIndex = [](const TArray<FName>& LocalPropertyChain, int32& OutIndex) -> int32
+	auto TryReadIndex = [](const TArray<FChainElement>& LocalPropertyChain, int32& OutIndex) -> int32
 	{
 		if(OutIndex + 1 < LocalPropertyChain.Num())
 		{
-			FString AsString = LocalPropertyChain[OutIndex + 1].ToString();
+			FString AsString = LocalPropertyChain[OutIndex + 1].DisplayString;
 			if(AsString.IsNumeric())
 			{
 				++OutIndex;
@@ -123,7 +128,7 @@ FPropertyPath FPropertySoftPath::ResolvePath(const UObject* Object) const
 	FPropertyPath Ret;
 	for( int32 I = 0; I < PropertyChain.Num(); ++I )
 	{
-		FName PropertyIdentifier = PropertyChain[I];
+		FName PropertyIdentifier = PropertyChain[I].PropertyName;
 		UProperty* ResolvedProperty = UE4DiffUtils_Private::Resolve(ContainerStruct, PropertyIdentifier);
 
 		FPropertyInfo Info(ResolvedProperty, INDEX_NONE);
@@ -188,7 +193,7 @@ FPropertyPath FPropertySoftPath::ResolvePath(const UObject* Object) const
 				// we have an index, but are we looking into a key or value? Peek ahead to find out:
 				if(ensure(I + 1 < PropertyChain.Num()))
 				{
-					if(PropertyChain[I+1] == MapProperty->KeyProp->GetFName())
+					if(PropertyChain[I+1].PropertyName == MapProperty->KeyProp->GetFName())
 					{
 						++I;
 
@@ -197,7 +202,7 @@ FPropertyPath FPropertySoftPath::ResolvePath(const UObject* Object) const
 						FPropertyInfo MakKeyInfo(MapProperty->KeyProp, RealIndex);
 						Ret.AddProperty(MakKeyInfo);
 					}
-					else if(ensure( PropertyChain[I+1] == MapProperty->ValueProp->GetFName() ))
+					else if(ensure( PropertyChain[I+1].PropertyName == MapProperty->ValueProp->GetFName() ))
 					{	
 						++I;
 
@@ -247,9 +252,9 @@ FPropertyPath FPropertySoftPath::ResolvePath(const UObject* Object) const
 FString FPropertySoftPath::ToDisplayName() const
 {
 	FString Ret;
-	for( FName Property : PropertyChain )
+	for( FChainElement Element : PropertyChain )
 	{
-		FString PropertyAsString = Property.ToString();
+		FString PropertyAsString = Element.DisplayString;
 		if(Ret.IsEmpty())
 		{
 			Ret.Append(PropertyAsString);
@@ -280,23 +285,23 @@ const UObject* DiffUtils::GetCDO(const UBlueprint* ForBlueprint)
 	return ForBlueprint->GeneratedClass->ClassDefaultObject;
 }
 
-void DiffUtils::CompareUnrelatedObjects(const UObject* A, const UObject* B, TArray<FSingleObjectDiffEntry>& OutDifferingProperties)
+void DiffUtils::CompareUnrelatedStructs(const UStruct* StructA, const void* A, const UStruct* StructB, const void* B, TArray<FSingleObjectDiffEntry>& OutDifferingProperties)
 {
-	FPropertySoftPathSet PropertiesInA = UE4DiffUtils_Private::GetPropertyNameSet(A);
-	FPropertySoftPathSet PropertiesInB = UE4DiffUtils_Private::GetPropertyNameSet(B);
+	FPropertySoftPathSet PropertiesInA = UE4DiffUtils_Private::GetPropertyNameSet(StructA);
+	FPropertySoftPathSet PropertiesInB = UE4DiffUtils_Private::GetPropertyNameSet(StructB);
 
 	// any properties in A that aren't in B are differing:
 	auto AddedToA = PropertiesInA.Difference(PropertiesInB).Array();
-	for( const auto& Entry : AddedToA )
+	for (const auto& Entry : AddedToA)
 	{
-		OutDifferingProperties.Push(FSingleObjectDiffEntry( Entry, EPropertyDiffType::PropertyAddedToA ));
+		OutDifferingProperties.Push(FSingleObjectDiffEntry(Entry, EPropertyDiffType::PropertyAddedToA));
 	}
 
 	// and the converse:
 	auto AddedToB = PropertiesInB.Difference(PropertiesInA).Array();
 	for (const auto& Entry : AddedToB)
 	{
-		OutDifferingProperties.Push(FSingleObjectDiffEntry( Entry, EPropertyDiffType::PropertyAddedToB ));
+		OutDifferingProperties.Push(FSingleObjectDiffEntry(Entry, EPropertyDiffType::PropertyAddedToB));
 	}
 
 	// for properties in common, dig out the uproperties and determine if they're identical:
@@ -305,8 +310,8 @@ void DiffUtils::CompareUnrelatedObjects(const UObject* A, const UObject* B, TArr
 		FPropertySoftPathSet Common = PropertiesInA.Intersect(PropertiesInB);
 		for (const auto& PropertyName : Common)
 		{
-			FResolvedProperty AProp = PropertyName.Resolve(A);
-			FResolvedProperty BProp = PropertyName.Resolve(B);
+			FResolvedProperty AProp = PropertyName.Resolve(StructA, A);
+			FResolvedProperty BProp = PropertyName.Resolve(StructB, B);
 
 			check(AProp != FResolvedProperty() && BProp != FResolvedProperty());
 			TArray<FPropertySoftPath> DifferingSubProperties;
@@ -321,6 +326,14 @@ void DiffUtils::CompareUnrelatedObjects(const UObject* A, const UObject* B, TArr
 	}
 }
 
+void DiffUtils::CompareUnrelatedObjects(const UObject* A, const UObject* B, TArray<FSingleObjectDiffEntry>& OutDifferingProperties)
+{
+	if (A && B)
+	{
+		return CompareUnrelatedStructs(A->GetClass(), A, B->GetClass(), B, OutDifferingProperties);
+	}
+}
+
 void DiffUtils::CompareUnrelatedSCS(const UBlueprint* Old, const TArray< FSCSResolvedIdentifier >& OldHierarchy, const UBlueprint* New, const TArray< FSCSResolvedIdentifier >& NewHierarchy, FSCSDiffRoot& OutDifferingEntries )
 {
 	const auto FindEntry = [](TArray< FSCSResolvedIdentifier > const& InArray, const FSCSIdentifier* Value) -> const FSCSResolvedIdentifier*
@@ -582,7 +595,7 @@ static void IdenticalHelper(const UProperty* AProperty, const UProperty* BProper
 			const UClass* AClass = A->GetClass(); // BClass and AClass are identical!
 
 			// We only want to recurse if this is EditInlineNew and not a component
-			// Other instanced refs are likely to form a type-specific web so recursion doesn't make sense and won't be displayed in the details pane
+			// Other instanced refs are likely to form a type-specific web so recursion doesn't make sense and won't be displayed properly in the details pane
 			if (AClass->HasAnyClassFlags(CLASS_EditInlineNew) && !AClass->IsChildOf(UActorComponent::StaticClass()))
 			{
 				for (TFieldIterator<UProperty> PropertyIt(AClass); PropertyIt; ++PropertyIt)
@@ -591,6 +604,11 @@ static void IdenticalHelper(const UProperty* AProperty, const UProperty* BProper
 					IdenticalHelper(ClassProp, ClassProp, ClassProp->ContainerPtrToValuePtr<void>(A, 0), ClassProp->ContainerPtrToValuePtr<void>(B, 0), FPropertySoftPath(RootPath, ClassProp), DifferingSubProperties);
 				}
 			}
+			else if (A->GetFName() != B->GetFName())
+			{
+				// If the names don't match, report that as a difference as the object was likely changed
+				DifferingSubProperties.Push(RootPath);
+			}
 		}
 		else
 		{
@@ -632,44 +650,22 @@ bool DiffUtils::Identical(const FResolvedProperty& AProp, const FResolvedPropert
 	return DifferingProperties.Num() == 0;
 }
 
-TArray<FPropertySoftPath> DiffUtils::GetVisiblePropertiesInOrderDeclared(const UObject* ForObj, const TArray<FName>& Scope /*= TArray<FName>()*/)
+TArray<FPropertySoftPath> DiffUtils::GetVisiblePropertiesInOrderDeclared(const UStruct* ForStruct, const FPropertySoftPath& Scope /*= TArray<FName>()*/)
 {
 	TArray<FPropertySoftPath> Ret;
-	if (ForObj)
+	if (ForStruct)
 	{
-		const UClass* Class = ForObj->GetClass();
-		TSet<FString> HiddenCategories = FEditorCategoryUtils::GetHiddenCategories(Class);
-		for (TFieldIterator<UProperty> PropertyIt(Class); PropertyIt; ++PropertyIt)
+		TSet<FString> HiddenCategories = FEditorCategoryUtils::GetHiddenCategories(ForStruct);
+		for (TFieldIterator<UProperty> PropertyIt(ForStruct); PropertyIt; ++PropertyIt)
 		{
 			FName CategoryName = FObjectEditorUtils::GetCategoryFName(*PropertyIt);
 			if (!HiddenCategories.Contains(CategoryName.ToString()))
 			{
 				if (PropertyIt->PropertyFlags&CPF_Edit)
 				{
-					TArray<FName> NewPath(Scope);
-					NewPath.Push(PropertyIt->GetFName());
-					if (const UObjectProperty* ObjectProperty = Cast<UObjectProperty>(*PropertyIt))
-					{
-						bool bShouldRecurse = false;
-						const UObject* const* BaseObject = reinterpret_cast<const UObject* const*>( ObjectProperty->ContainerPtrToValuePtr<void>(ForObj) );
-						if (BaseObject && *BaseObject)
-						{
-							// See if this is an object in the package, if so recurse into it
-							if ((*BaseObject)->IsIn(ForObj->GetOutermost()))
-							{
-								bShouldRecurse = true;
-								Ret.Append(GetVisiblePropertiesInOrderDeclared(*BaseObject, NewPath));
-							}
-						}
-						if (!bShouldRecurse)
-						{
-							Ret.Push(NewPath);
-						}
-					}
-					else
-					{
-						Ret.Push(NewPath);
-					}
+					// We don't need to recurse into objects/structs as those will be picked up in the Identical check later
+					FPropertySoftPath NewPath(Scope, *PropertyIt);
+					Ret.Push(NewPath);
 				}
 			}
 		}
