Add Slot property diffing for UMG Widgets

diff --git a/Engine/Source/Editor/BlueprintGraph/Classes/K2Node_Event.h b/Engine/Source/Editor/BlueprintGraph/Classes/K2Node_Event.h
--- a/Engine/Source/Editor/BlueprintGraph/Classes/K2Node_Event.h
+++ b/Engine/Source/Editor/BlueprintGraph/Classes/K2Node_Event.h
@@ -71,6 +71,7 @@ class UK2Node_Event : public UK2Node_EditablePinBase, public IK2Node_EventNodeIn
 	BLUEPRINTGRAPH_API virtual UObject* GetJumpTargetForDoubleClick() const override;
 	BLUEPRINTGRAPH_API virtual FSlateIcon GetIconAndTint(FLinearColor& OutColor) const override;
 	BLUEPRINTGRAPH_API virtual FString GetFindReferenceSearchString() const override;
+	BLUEPRINTGRAPH_API virtual void FindDiffs(UEdGraphNode* OtherNode, struct FDiffResults& Results) override;
 	//~ End UEdGraphNode Interface
 
 	//~ Begin UK2Node Interface
diff --git a/Engine/Source/Editor/BlueprintGraph/Classes/K2Node_FunctionEntry.h b/Engine/Source/Editor/BlueprintGraph/Classes/K2Node_FunctionEntry.h
--- a/Engine/Source/Editor/BlueprintGraph/Classes/K2Node_FunctionEntry.h
+++ b/Engine/Source/Editor/BlueprintGraph/Classes/K2Node_FunctionEntry.h
@@ -43,6 +43,7 @@ class UK2Node_FunctionEntry : public UK2Node_FunctionTerminator
 	virtual bool IsDeprecated() const override;
 	virtual FString GetDeprecationMessage() const override;
 	virtual FText GetTooltipText() const override;
+	virtual void FindDiffs(UEdGraphNode* OtherNode, struct FDiffResults& Results) override;
 	//~ End UEdGraphNode Interface
 
 	//~ Begin UK2Node Interface
diff --git a/Engine/Source/Editor/BlueprintGraph/Private/K2Node_Event.cpp b/Engine/Source/Editor/BlueprintGraph/Private/K2Node_Event.cpp
--- a/Engine/Source/Editor/BlueprintGraph/Private/K2Node_Event.cpp
+++ b/Engine/Source/Editor/BlueprintGraph/Private/K2Node_Event.cpp
@@ -16,6 +16,7 @@
 #include "KismetCompilerMisc.h"
 #include "KismetCompiler.h"
 #include "EventEntryHandler.h"
+#include "DiffResults.h"
 
 const FName UK2Node_Event::DelegateOutputName(TEXT("OutputDelegate"));
 
@@ -838,6 +839,27 @@ FString UK2Node_Event::GetFindReferenceSearchString() const
 	return FunctionName;
 }
 
+void UK2Node_Event::FindDiffs(UEdGraphNode* OtherNode, struct FDiffResults& Results)
+{
+	Super::FindDiffs(OtherNode, Results);
+	UK2Node_Event* OtherFunction = Cast<UK2Node_Event>(OtherNode);
+
+	if (OtherFunction)
+	{
+		if (FunctionFlags != OtherFunction->FunctionFlags)
+		{
+			FDiffSingleResult Diff;
+			Diff.Diff = EDiffType::NODE_PROPERTY;
+			Diff.Node1 = this;
+			Diff.Node2 = OtherNode;
+			Diff.DisplayString = LOCTEXT("DIF_EventFlags", "Event flags have changed");
+			Diff.DisplayColor = FLinearColor(0.25f, 0.71f, 0.85f);
+
+			Results.Add(Diff);
+		}
+	}
+}
+
 bool UK2Node_Event::AreEventNodesIdentical(const UK2Node_Event* InNodeA, const UK2Node_Event* InNodeB)
 {
 	return InNodeA->EventReference.GetMemberName() == InNodeB->EventReference.GetMemberName()
diff --git a/Engine/Source/Editor/BlueprintGraph/Private/K2Node_FunctionEntry.cpp b/Engine/Source/Editor/BlueprintGraph/Private/K2Node_FunctionEntry.cpp
--- a/Engine/Source/Editor/BlueprintGraph/Private/K2Node_FunctionEntry.cpp
+++ b/Engine/Source/Editor/BlueprintGraph/Private/K2Node_FunctionEntry.cpp
@@ -24,6 +24,7 @@
 #include "KismetCompilerMisc.h"
 #include "KismetCompiler.h"
 #include "Misc/OutputDeviceNull.h"
+#include "DiffResults.h"
 
 #define LOCTEXT_NAMESPACE "K2Node_FunctionEntry"
 
@@ -491,6 +492,61 @@ FText UK2Node_FunctionEntry::GetTooltipText() const
 	return Super::GetTooltipText();
 }
 
+void UK2Node_FunctionEntry::FindDiffs(UEdGraphNode* OtherNode, struct FDiffResults& Results)
+{
+	Super::FindDiffs(OtherNode, Results);
+	UK2Node_FunctionEntry* OtherFunction = Cast<UK2Node_FunctionEntry>(OtherNode);
+
+	if (OtherFunction)
+	{
+		if (ExtraFlags != OtherFunction->ExtraFlags)
+		{
+			FDiffSingleResult Diff;
+			Diff.Diff = EDiffType::NODE_PROPERTY;
+			Diff.Node1 = this;
+			Diff.Node2 = OtherNode;
+			Diff.DisplayString = LOCTEXT("DIF_FunctionFlags", "Function flags have changed");
+			Diff.DisplayColor = FLinearColor(0.25f, 0.71f, 0.85f);
+
+			Results.Add(Diff);
+		}
+
+		if (!FKismetUserDeclaredFunctionMetadata::StaticStruct()->CompareScriptStruct(&MetaData, &OtherFunction->MetaData, 0))
+		{
+			FDiffSingleResult Diff;
+			Diff.Diff = EDiffType::NODE_PROPERTY;
+			Diff.Node1 = this;
+			Diff.Node2 = OtherNode;
+			Diff.DisplayString = LOCTEXT("DIF_FunctionMetadata", "Function metadata has changed");
+			Diff.DisplayColor = FLinearColor(0.25f, 0.71f, 0.85f);
+
+			Results.Add(Diff);
+		}
+
+		bool bLocalVarsDiffer = (LocalVariables.Num() != OtherFunction->LocalVariables.Num());
+
+		for (int32 i = 0; i < LocalVariables.Num() && !bLocalVarsDiffer; i++)
+		{
+			if (!FBPVariableDescription::StaticStruct()->CompareScriptStruct(&LocalVariables[i], &OtherFunction->LocalVariables[i], 0))
+			{
+				bLocalVarsDiffer = true;
+			}
+		}
+
+		if (bLocalVarsDiffer)
+		{
+			FDiffSingleResult Diff;
+			Diff.Diff = EDiffType::NODE_PROPERTY;
+			Diff.Node1 = this;
+			Diff.Node2 = OtherNode;
+			Diff.DisplayString = LOCTEXT("DIF_FunctionLocalVariables", "Function local variables have changed");
+			Diff.DisplayColor = FLinearColor(0.25f, 0.71f, 0.85f);
+
+			Results.Add(Diff);
+		}
+	}
+}
+
 int32 UK2Node_FunctionEntry::GetFunctionFlags() const
 {
 	int32 ReturnFlags = 0;
diff --git a/Engine/Source/Editor/GraphEditor/Private/GraphDiffControl.cpp b/Engine/Source/Editor/GraphEditor/Private/GraphDiffControl.cpp
--- a/Engine/Source/Editor/GraphEditor/Private/GraphDiffControl.cpp
+++ b/Engine/Source/Editor/GraphEditor/Private/GraphDiffControl.cpp
@@ -473,6 +473,8 @@ static bool ArePinsDifferent(const TArray<UEdGraphPin*>& OldPins, TArray<UEdGrap
 		UEdGraphPin* OldPin = OldPins[i];
 		UEdGraphPin* NewPin = NewPins[i];
 
+		const UEdGraphSchema* Schema = OldPin->GetSchema();
+
 		if(IsPinTypeDifferent(OldPin->PinType, NewPin->PinType))
 		{
 			DiffR_PinTypeChanged(Results, NewPin, OldPin);
@@ -488,7 +490,7 @@ static bool ArePinsDifferent(const TArray<UEdGraphPin*>& OldPins, TArray<UEdGrap
 			KEEP_GOING_IF_RESULTS()
 		}
 
-		if(NewPin->LinkedTo.Num() == 0 && (NewPin->GetDefaultAsString() != OldPin->GetDefaultAsString()))
+		if(NewPin->LinkedTo.Num() == 0 && Schema && !Schema->DoesDefaultValueMatch(*OldPin, NewPin->GetDefaultAsString()))
 		{
 			DiffR_PinDefaultValueChanged(Results, NewPin, OldPin); //note: some issues with how floating point is stored as string format(0.0 vs 0.00) can cause false diffs
 			KEEP_GOING_IF_RESULTS()
diff --git a/Engine/Source/Editor/Kismet/Private/SBlueprintDiff.cpp b/Engine/Source/Editor/Kismet/Private/SBlueprintDiff.cpp
--- a/Engine/Source/Editor/Kismet/Private/SBlueprintDiff.cpp
+++ b/Engine/Source/Editor/Kismet/Private/SBlueprintDiff.cpp
@@ -64,12 +64,28 @@ TSharedRef<SWidget>	FDiffResultItem::GenerateWidget() const
 		.Text(Text);
 }
 
+static TSharedRef<SWidget> GenerateObjectDiffWidget(FSingleObjectDiffEntry DiffEntry, FText ObjectName)
+{
+	return SNew(STextBlock)
+		.Text(DiffViewUtils::PropertyDiffMessage(DiffEntry, ObjectName))
+		.ToolTipText(DiffViewUtils::PropertyDiffMessage(DiffEntry, ObjectName))
+		.ColorAndOpacity(DiffViewUtils::Differs());
+}
+
+static TSharedRef<SWidget> GenerateSimpleDiffWidget(FText DiffText)
+{
+	return SNew(STextBlock)
+		.Text(DiffText)
+		.ToolTipText(DiffText)
+		.ColorAndOpacity(DiffViewUtils::Differs());
+};
+
 /** Shows all differences for the blueprint structure itself that aren't picked up elsewhere */
 class FMyBlueprintDiffControl : public TSharedFromThis<FMyBlueprintDiffControl>, public IDiffControl
 {
 public:
 	FMyBlueprintDiffControl(const UBlueprint* InOldBlueprint, const UBlueprint* InNewBlueprint, FOnDiffEntryFocused InSelectionCallback)
-		: SelectionCallback(InSelectionCallback)
+		: SelectionCallback(InSelectionCallback), OldBlueprint(InOldBlueprint), NewBlueprint(InNewBlueprint)
 	{
 	}
 
@@ -77,32 +93,65 @@ public:
 	{
 		TArray< TSharedPtr<FBlueprintDifferenceTreeEntry> > Children;
 
-		/*const auto FocusSCSDifferenceEntry = [](FSCSDiffEntry Entry, FOnDiffEntryFocused InSelectionCallback, FMyBlueprintDiffControl* Owner)
+		if (OldBlueprint && NewBlueprint)
 		{
-			InSelectionCallback.ExecuteIfBound();
-			if (Entry.TreeIdentifier.Name != NAME_None)
+			for (TFieldIterator<UProperty> PropertyIt(OldBlueprint->SkeletonGeneratedClass); PropertyIt; ++PropertyIt)
 			{
-				Owner->OldSCS.HighlightProperty(Entry.TreeIdentifier.Name, FPropertyPath());
-				Owner->NewSCS.HighlightProperty(Entry.TreeIdentifier.Name, FPropertyPath());
+				UProperty* OldProperty = *PropertyIt;
+				UProperty* NewProperty = NewBlueprint->SkeletonGeneratedClass->FindPropertyByName(OldProperty->GetFName());
+
+				FText PropertyText = FText::FromString(OldProperty->GetAuthoredName());
+
+				if (NewProperty)
+				{
+					const int32 OldVarIndex = FBlueprintEditorUtils::FindNewVariableIndex(OldBlueprint, OldProperty->GetFName());
+					const int32 NewVarIndex = FBlueprintEditorUtils::FindNewVariableIndex(NewBlueprint, OldProperty->GetFName());
+
+					if (OldVarIndex != INDEX_NONE && NewVarIndex != INDEX_NONE)
+					{
+						TArray<FSingleObjectDiffEntry> DifferingProperties;
+						DiffUtils::CompareUnrelatedStructs(FBPVariableDescription::StaticStruct(), &OldBlueprint->NewVariables[OldVarIndex], FBPVariableDescription::StaticStruct(), &NewBlueprint->NewVariables[NewVarIndex], DifferingProperties);
+						for (const FSingleObjectDiffEntry& Difference : DifferingProperties)
+						{
+							TSharedPtr<FBlueprintDifferenceTreeEntry> Entry = MakeShared<FBlueprintDifferenceTreeEntry>(
+								SelectionCallback,
+								FGenerateDiffEntryWidget::CreateStatic(&GenerateObjectDiffWidget, Difference, PropertyText));
+							Children.Push(Entry);
+							OutRealDifferences.Push(Entry);
+						}
+					}	
+				}
+				else
+				{
+					FText DiffText = FText::Format(LOCTEXT("VariableRemoved", "Removed Variable {0}"), PropertyText);
+
+					TSharedPtr<FBlueprintDifferenceTreeEntry> Entry = MakeShared<FBlueprintDifferenceTreeEntry>(
+						SelectionCallback,
+						FGenerateDiffEntryWidget::CreateStatic(&GenerateSimpleDiffWidget, DiffText));
+
+					Children.Push(Entry);
+					OutRealDifferences.Push(Entry);
+				}
 			}
-		};
 
-		const auto CreateSCSDifferenceWidget = [](FSCSDiffEntry Entry, FText ObjectName) -> TSharedRef<SWidget>
-		{
-			return SNew(STextBlock)
-				.Text(DiffViewUtils::SCSDiffMessage(Entry, ObjectName))
-				.ColorAndOpacity(DiffViewUtils::Differs());
-		};
+			for (TFieldIterator<UProperty> PropertyIt(NewBlueprint->SkeletonGeneratedClass); PropertyIt; ++PropertyIt)
+			{
+				UProperty* NewProperty = *PropertyIt;
+				UProperty* OldProperty = OldBlueprint->SkeletonGeneratedClass->FindPropertyByName(NewProperty->GetFName());
 
-		for (auto Difference : DifferingProperties.Entries)
-		{
-			TSharedPtr<FBlueprintDifferenceTreeEntry> Entry = MakeShared<FBlueprintDifferenceTreeEntry>(
-					FOnDiffEntryFocused::CreateStatic(FocusSCSDifferenceEntry, Difference, SelectionCallback, this),
-					FGenerateDiffEntryWidget::CreateStatic(CreateSCSDifferenceWidget, Difference, RightRevision));
-			Children.Push(Entry);
-			OutRealDifferences.Push(Entry);
-		}*/
+				if (!OldProperty)
+				{
+					FText DiffText = FText::Format(LOCTEXT("VariableAdded", "Added Variable {0}"), FText::FromString(NewProperty->GetAuthoredName()));
+
+					TSharedPtr<FBlueprintDifferenceTreeEntry> Entry = MakeShared<FBlueprintDifferenceTreeEntry>(
+						SelectionCallback,
+						FGenerateDiffEntryWidget::CreateStatic(&GenerateSimpleDiffWidget, DiffText));
 
+					Children.Push(Entry);
+					OutRealDifferences.Push(Entry);
+				}
+			}
+		}
 		const bool bHasDifferences = Children.Num() != 0;
 		if (!bHasDifferences)
 		{
@@ -121,6 +170,8 @@ public:
 
 private:
 	FOnDiffEntryFocused SelectionCallback;
+	const UBlueprint* OldBlueprint;
+	const UBlueprint* NewBlueprint;
 };
 
 /** 
@@ -215,7 +266,7 @@ public:
 		{
 			TSharedPtr<FBlueprintDifferenceTreeEntry> Entry = MakeShared<FBlueprintDifferenceTreeEntry>(
 				FOnDiffEntryFocused::CreateSP(AsShared(), &FDetailsDiffControl::OnSelectDiffEntry, Difference.Identifier),
-				FGenerateDiffEntryWidget::CreateSP(AsShared(), &FDetailsDiffControl::CreateDifferenceWidget, Difference, RightRevision));
+				FGenerateDiffEntryWidget::CreateStatic(&GenerateObjectDiffWidget, Difference, RightRevision));
 			Children.Push(Entry);
 			OutRealDifferences.Push(Entry);
 		}
@@ -232,13 +283,6 @@ protected:
 		NewDetails.HighlightProperty(PropertyName);
 	}
 
-	virtual TSharedRef<SWidget> CreateDifferenceWidget(FSingleObjectDiffEntry DiffEntry, FText ObjectName)
-	{
-		return SNew(STextBlock)
-			.Text(DiffViewUtils::PropertyDiffMessage(DiffEntry, ObjectName))
-			.ColorAndOpacity(DiffViewUtils::Differs());
-	}
-
 	FOnDiffEntryFocused SelectionCallback;
 	FDetailsDiff OldDetails;
 	FDetailsDiff NewDetails;
@@ -294,14 +338,6 @@ public:
 		const UBlueprint* OldBlueprint = Cast<UBlueprint>(OldDetails.GetDisplayedObject());
 		const UBlueprint* NewBlueprint = Cast<UBlueprint>(OldDetails.GetDisplayedObject());
 
-		const auto GenerateDiffWidget = [](FText DiffText) -> TSharedRef<SWidget>
-		{
-			return SNew(STextBlock)
-				.ColorAndOpacity(DiffViewUtils::Differs())
-				.Text(DiffText)
-				.ToolTipText(DiffText);
-		};
-
 		if (OldBlueprint && NewBlueprint)
 		{
 			if (OldBlueprint->ParentClass != NewBlueprint->ParentClass)
@@ -310,7 +346,7 @@ public:
 
 				TSharedPtr<FBlueprintDifferenceTreeEntry> Entry = MakeShared<FBlueprintDifferenceTreeEntry>(
 					SelectionCallback,
-					FGenerateDiffEntryWidget::CreateStatic(GenerateDiffWidget, DiffText));
+					FGenerateDiffEntryWidget::CreateStatic(&GenerateSimpleDiffWidget, DiffText));
 
 				Children.Push(Entry);
 				OutRealDifferences.Push(Entry);
@@ -341,7 +377,7 @@ public:
 
 				TSharedPtr<FBlueprintDifferenceTreeEntry> Entry = MakeShared<FBlueprintDifferenceTreeEntry>(
 					SelectionCallback,
-					FGenerateDiffEntryWidget::CreateStatic(GenerateDiffWidget, DiffText));
+					FGenerateDiffEntryWidget::CreateStatic(&GenerateSimpleDiffWidget, DiffText));
 
 				Children.Push(Entry);
 				OutRealDifferences.Push(Entry);
@@ -353,7 +389,7 @@ public:
 
 				TSharedPtr<FBlueprintDifferenceTreeEntry> Entry = MakeShared<FBlueprintDifferenceTreeEntry>(
 					SelectionCallback,
-					FGenerateDiffEntryWidget::CreateStatic(GenerateDiffWidget, DiffText));
+					FGenerateDiffEntryWidget::CreateStatic(&GenerateSimpleDiffWidget, DiffText));
 
 				Children.Push(Entry);
 				OutRealDifferences.Push(Entry);
@@ -497,18 +533,11 @@ void FBlueprintTypeDiffControl::BuildDiffSourceArray()
 					// Actual differences, so add to tree
 					SubObjectDiffs.Add(SubObjectDiff);
 
-					const auto CreateDifferenceWidget = [](FSingleObjectDiffEntry DiffEntry, FText ObjectName) -> TSharedRef<SWidget>
-					{
-						return SNew(STextBlock)
-							.Text(DiffViewUtils::PropertyDiffMessage(DiffEntry, ObjectName))
-							.ColorAndOpacity(DiffViewUtils::Differs());
-					};
-
 					for (const FSingleObjectDiffEntry& Difference : DifferingProperties)
 					{
 						TSharedPtr<FBlueprintDifferenceTreeEntry> Entry = MakeShared<FBlueprintDifferenceTreeEntry>(
 							FOnDiffEntryFocused::CreateSP(AsShared(), &FBlueprintTypeDiffControl::OnSelectSubobjectDiff, Difference.Identifier, SubObjectDiff),
-							FGenerateDiffEntryWidget::CreateStatic(CreateDifferenceWidget, Difference, RightRevision));
+							FGenerateDiffEntryWidget::CreateStatic(&GenerateObjectDiffWidget, Difference, RightRevision));
 						SubObjectDiff->Diffs.Push(Entry);
 					}
 				}
@@ -760,9 +789,9 @@ void FGraphToDiff::BuildDiffSourceArray()
 	FGraphDiffControl::DiffGraphs(GraphOld, GraphNew, FoundDiffs);
 
 	DiffListSource.Empty();
-	for (auto DiffIt(FoundDiffs.CreateIterator()); DiffIt; ++DiffIt)
+	for (const FDiffSingleResult& Diff : FoundDiffs)
 	{
-		DiffListSource.Add(MakeShared<FDiffResultItem>(*DiffIt));
+		DiffListSource.Add(MakeShared<FDiffResultItem>(Diff));
 	}
 
 	struct SortDiff
@@ -798,15 +827,15 @@ void FDiffPanel::InitializeDiffPanel()
 	MyBlueprint->SetInspector(DetailsView);
 }
 
-int32 GetCurrentIndex( SListView< TSharedPtr< FDiffSingleResult> > const& ListView, const TArray< TSharedPtr< FDiffSingleResult > >& ListViewSource )
+static int32 GetCurrentIndex( SListView< TSharedPtr< FDiffSingleResult> > const& ListView, const TArray< TSharedPtr< FDiffSingleResult > >& ListViewSource )
 {
-	const auto& Selected = ListView.GetSelectedItems();
+	const TArray< TSharedPtr<FDiffSingleResult> >& Selected = ListView.GetSelectedItems();
 	if (Selected.Num() == 1)
 	{
 		int32 Index = 0;
-		for (auto It(ListViewSource.CreateConstIterator()); It; ++It, Index++)
+		for (const TSharedPtr<FDiffSingleResult>& Diff : ListViewSource)
 		{
-			if (*It == Selected[0])
+			if (Diff == Selected[0])
 			{
 				return Index;
 			}
@@ -1160,7 +1189,7 @@ bool SBlueprintDiff::HasPrevDiff() const
 
 FGraphToDiff* SBlueprintDiff::FindGraphToDiffEntry(const FString& GraphPath)
 {
-	for( const auto& Graph : Graphs )
+	for(const TSharedPtr<FGraphToDiff>& Graph : Graphs)
 	{
 		FString SearchGraphPath = Graph->GetGraphOld() ? FGraphDiffControl::GetGraphPath(Graph->GetGraphOld()) : FGraphDiffControl::GetGraphPath(Graph->GetGraphNew());
 		if (SearchGraphPath.Equals(GraphPath, ESearchCase::CaseSensitive))
@@ -1190,7 +1219,7 @@ void SBlueprintDiff::OnDiffListSelectionChanged(TSharedPtr<FDiffResultItem> TheD
 
 	const auto SafeClearSelection = []( TWeakPtr<SGraphEditor> GraphEditor )
 	{
-		auto GraphEditorPtr = GraphEditor.Pin();
+		TSharedPtr<SGraphEditor> GraphEditorPtr = GraphEditor.Pin();
 		if( GraphEditorPtr.IsValid())
 		{
 			GraphEditorPtr->ClearSelectionSet();
@@ -1290,7 +1319,7 @@ void FDiffPanel::GeneratePanel(UEdGraph* Graph, UEdGraph* GraphToDiff )
 		MyBlueprint->SetFocusedGraph(Graph);
 		MyBlueprint->Refresh();
 
-		auto Editor = SNew(SGraphEditor)
+		TSharedRef<SGraphEditor> Editor = SNew(SGraphEditor)
 			.AdditionalCommands(GraphEditorCommands)
 			.GraphToEdit(Graph)
 			.GraphToDiff(GraphToDiff)
@@ -1373,12 +1402,12 @@ void FDiffPanel::FocusDiff(UEdGraphNode& Node)
 
 FDiffPanel& SBlueprintDiff::GetDiffPanelForNode(UEdGraphNode& Node)
 {
-	auto OldGraphEditorPtr = PanelOld.GraphEditor.Pin();
+	TSharedPtr<SGraphEditor> OldGraphEditorPtr = PanelOld.GraphEditor.Pin();
 	if (OldGraphEditorPtr.IsValid() && Node.GetGraph() == OldGraphEditorPtr->GetCurrentGraph())
 	{
 		return PanelOld;
 	}
-	auto NewGraphEditorPtr = PanelNew.GraphEditor.Pin();
+	TSharedPtr<SGraphEditor> NewGraphEditorPtr = PanelNew.GraphEditor.Pin();
 	if (NewGraphEditorPtr.IsValid() && Node.GetGraph() == NewGraphEditorPtr->GetCurrentGraph())
 	{
 		return PanelNew;
@@ -1430,17 +1459,17 @@ void SBlueprintDiff::GenerateDifferencesList()
 	PanelNew.Blueprint->GetAllGraphs(GraphsNew);
 
 	//Add Graphs that exist in both blueprints, or in blueprint 1 only
-	for (auto It(GraphsOld.CreateConstIterator()); It; It++)
+	for (UEdGraph* GraphOld : GraphsOld)
 	{
-		UEdGraph* GraphOld = *It;
 		UEdGraph* GraphNew = nullptr;
-		for (auto It2(GraphsNew.CreateIterator()); It2; It2++)
+		for (UEdGraph*& TestGraph : GraphsNew)
 		{
-			UEdGraph* TestGraph = *It2;
 			if (TestGraph && GraphOld->GetName() == TestGraph->GetName())
 			{
 				GraphNew = TestGraph;
-				*It2 = nullptr;
+
+				// Null reference inside array
+				TestGraph = nullptr;
 				break;
 			}
 		}
@@ -1452,9 +1481,8 @@ void SBlueprintDiff::GenerateDifferencesList()
 	}
 
 	//Add graphs that only exist in 2nd(new) blueprint
-	for (auto It2(GraphsNew.CreateIterator()); It2; It2++)
+	for (UEdGraph* GraphNew : GraphsNew)
 	{
-		UEdGraph* GraphNew = *It2;
 		if (GraphNew != nullptr && IsGraphDiffNeeded(GraphNew))
 		{
 			CreateGraphEntry(nullptr, GraphNew);
diff --git a/Engine/Source/Editor/Kismet/Public/DiffUtils.h b/Engine/Source/Editor/Kismet/Public/DiffUtils.h
--- a/Engine/Source/Editor/Kismet/Public/DiffUtils.h
+++ b/Engine/Source/Editor/Kismet/Public/DiffUtils.h
@@ -47,22 +47,25 @@ struct FPropertySoftPath
 	}
 
 	FPropertySoftPath(TArray<FName> InPropertyChain)
-		: PropertyChain(InPropertyChain)
 	{
+		for (FName PropertyName : InPropertyChain)
+		{
+			PropertyChain.Push(FChainElement(PropertyName));
+		}
 	}
 
 	FPropertySoftPath( FPropertyPath InPropertyPath )
 	{
 		for( int32 i = 0, end = InPropertyPath.GetNumProperties(); i != end; ++i )
 		{
-			PropertyChain.Push( InPropertyPath.GetPropertyInfo(i).Property->GetFName() );
+			PropertyChain.Push(FChainElement(InPropertyPath.GetPropertyInfo(i).Property.Get()));
 		}
 	}
 
 	FPropertySoftPath(const FPropertySoftPath& SubPropertyPath, const UProperty* LeafProperty)
 		: PropertyChain(SubPropertyPath.PropertyChain)
 	{
-		PropertyChain.Push(LeafProperty->GetFName());
+		PropertyChain.Push(FChainElement(LeafProperty));
 	}
 
 	FPropertySoftPath(const FPropertySoftPath& SubPropertyPath, int32 ContainerIndex)
@@ -71,9 +74,10 @@ struct FPropertySoftPath
 		PropertyChain.Push(FName(*FString::FromInt(ContainerIndex)));
 	}
 
-	FResolvedProperty Resolve(const UObject* Object) const;
-	FPropertyPath ResolvePath(const UObject* Object) const;
-	FString ToDisplayName() const;
+	KISMET_API FResolvedProperty Resolve(const UObject* Object) const;
+	KISMET_API FResolvedProperty Resolve(const UStruct* Struct, const void* StructData) const;
+	KISMET_API FPropertyPath ResolvePath(const UObject* Object) const;
+	KISMET_API FString ToDisplayName() const;
 
 	inline bool IsSubPropertyMatch(const FPropertySoftPath& PotentialBasePropertyPath) const
 	{
@@ -103,8 +107,46 @@ struct FPropertySoftPath
 		return !(*this == RHS);
 	}
 private:
+	struct FChainElement
+	{
+		// FName of the property
+		FName PropertyName;
+
+		// Display string of the property
+		FString DisplayString;
+
+		FChainElement(FName InPropertyName, const FString& InDisplayString = FString())
+			: PropertyName(InPropertyName)
+			, DisplayString(InDisplayString)
+		{
+			if (DisplayString.IsEmpty())
+			{
+				DisplayString = PropertyName.ToString();
+			}
+		}
+
+		FChainElement(const UProperty* Property)
+		{
+			if (Property)
+			{
+				PropertyName = Property->GetFName();
+				DisplayString = Property->GetAuthoredName();
+			}
+		}
+
+		inline bool operator==(FChainElement const& RHS) const
+		{
+			return PropertyName == RHS.PropertyName;
+		}
+
+		inline bool operator!=(FChainElement const& RHS) const
+		{
+			return !(*this == RHS);
+		}
+	};
+
 	friend uint32 GetTypeHash( FPropertySoftPath const& Path );
-	TArray<FName> PropertyChain;
+	TArray<FChainElement> PropertyChain;
 };
 
 struct FSCSIdentifier
@@ -132,9 +174,9 @@ FORCEINLINE bool operator!=(const FSCSIdentifier& A, const FSCSIdentifier& B)
 FORCEINLINE uint32 GetTypeHash( FPropertySoftPath const& Path )
 {
 	uint32 Ret = 0;
-	for( const auto& ProperytName : Path.PropertyChain )
+	for( const FPropertySoftPath::FChainElement& PropertyElement : Path.PropertyChain )
 	{
-		Ret = Ret ^ GetTypeHash(ProperytName);
+		Ret = Ret ^ GetTypeHash(PropertyElement.PropertyName);
 	}
 	return Ret;
 }
@@ -208,10 +250,11 @@ struct FSCSDiffRoot
 namespace DiffUtils
 {
 	KISMET_API const UObject* GetCDO(const UBlueprint* ForBlueprint);
+	KISMET_API void CompareUnrelatedStructs(const UStruct* StructA, const void* A, const UStruct* StructB, const void* B, TArray<FSingleObjectDiffEntry>& OutDifferingProperties);
 	KISMET_API void CompareUnrelatedObjects(const UObject* A, const UObject* B, TArray<FSingleObjectDiffEntry>& OutDifferingProperties);
 	KISMET_API void CompareUnrelatedSCS(const UBlueprint* Old, const TArray< FSCSResolvedIdentifier >& OldHierarchy, const UBlueprint* New, const TArray< FSCSResolvedIdentifier >& NewHierarchy, FSCSDiffRoot& OutDifferingEntries );
 	KISMET_API bool Identical(const FResolvedProperty& AProp, const FResolvedProperty& BProp, const FPropertySoftPath& RootPath, TArray<FPropertySoftPath>& DifferingProperties); 
-	TArray<FPropertySoftPath> GetVisiblePropertiesInOrderDeclared(const UObject* ForObj, const TArray<FName>& Scope = TArray<FName>());
+	KISMET_API TArray<FPropertySoftPath> GetVisiblePropertiesInOrderDeclared(const UStruct* ForStruct, const FPropertySoftPath& Scope = FPropertySoftPath());
 
 	KISMET_API TArray<FPropertyPath> ResolveAll(const UObject* Object, const TArray<FPropertySoftPath>& InSoftProperties);
 	KISMET_API TArray<FPropertyPath> ResolveAll(const UObject* Object, const TArray<FSingleObjectDiffEntry>& InDifferences);
diff --git a/Engine/Source/Editor/UMGEditor/Private/WidgetBlueprint.cpp b/Engine/Source/Editor/UMGEditor/Private/WidgetBlueprint.cpp
--- a/Engine/Source/Editor/UMGEditor/Private/WidgetBlueprint.cpp
+++ b/Engine/Source/Editor/UMGEditor/Private/WidgetBlueprint.cpp
@@ -638,26 +638,43 @@ bool UWidgetBlueprint::FindDiffs(const UBlueprint* OtherBlueprint, FDiffResults&
 
 		if (FoundOtherWidget)
 		{
-			// Add to general object diff map
-			FDiffSingleResult Diff;
-			Diff.Diff = EDiffType::OBJECT_REQUEST_DIFF;
-			Diff.Object1 = Widget;
-			Diff.Object2 = *FoundOtherWidget;
-			Diff.OwningObjectPath = Pair.Key;
-
 			if (Results.CanStoreResults())
 			{
+				// Add to general object diff map
+				FDiffSingleResult Diff;
+				Diff.Diff = EDiffType::OBJECT_REQUEST_DIFF;
+				Diff.Object1 = Widget;
+				Diff.Object2 = *FoundOtherWidget;
+				Diff.OwningObjectPath = Pair.Key;			
 				FFormatNamedArguments Args;
 				Args.Add(TEXT("WidgetTitle"), Widget->GetLabelTextWithMetadata());
 				Args.Add(TEXT("WidgetPath"), FText::FromString(Pair.Key));
 				Diff.ToolTip = FText::Format(LOCTEXT("DIF_RequestWidgetTooltip", "Widget {WidgetTitle}\nPath: {WidgetPath}"), Args);
 				Diff.DisplayString = FText::Format(LOCTEXT("DIF_RequestWidgetLabel", "Widget {WidgetTitle}"), Args);
 				Diff.DisplayColor = FLinearColor(1.f, 0.4f, 0.4f);
-			}
 
-			Results.Add(Diff);
+				Results.Add(Diff);
+
+				UPanelSlot* Slot = Widget->Slot;
+				UPanelSlot* OtherSlot = (*FoundOtherWidget)->Slot;
 
-			// TODO: Add some type specific diffing if needed
+				if (Slot && OtherSlot)
+				{
+					FDiffSingleResult SlotDiff;
+					SlotDiff.Diff = EDiffType::OBJECT_REQUEST_DIFF;
+					SlotDiff.Object1 = Slot;
+					SlotDiff.Object2 = OtherSlot;
+					SlotDiff.OwningObjectPath = Pair.Key;
+					FFormatNamedArguments SlotArgs;
+					SlotArgs.Add(TEXT("WidgetTitle"), Widget->GetLabelTextWithMetadata());
+					SlotArgs.Add(TEXT("WidgetPath"), FText::FromString(Pair.Key));
+					SlotDiff.ToolTip = FText::Format(LOCTEXT("DIF_RequestSlotTooltip", "Slot for {WidgetTitle}\nPath: {WidgetPath}"), SlotArgs);
+					SlotDiff.DisplayString = FText::Format(LOCTEXT("DIF_RequestSlotLabel", "Slot for {WidgetTitle}"), SlotArgs);
+					SlotDiff.DisplayColor = FLinearColor(1.f, 0.4f, 0.4f);
+
+					Results.Add(SlotDiff);
+				}
+			}
 		}
 		else
 		{
@@ -714,7 +731,7 @@ bool UWidgetBlueprint::FindDiffs(const UBlueprint* OtherBlueprint, FDiffResults&
 		FDiffSingleResult Diff;
 		Diff.Diff = EDiffType::INFO_MESSAGE;
 		Diff.DisplayColor = FLinearColor(.7f, .7f, .7f);
-		Diff.ToolTip = LOCTEXT("DIF_WidgetWarningMessage", "Warning: May be missing some widget-specific differences");
+		Diff.ToolTip = LOCTEXT("DIF_WidgetWarningMessage", "Warning: This may be missing changes to Animations and Bindings");
 		Diff.DisplayString = Diff.ToolTip;
 
 		Results.Add(Diff);
