Fix incorrect comment numbering.

#rnx

diff --git a/Engine/Source/Editor/Kismet/Private/BlueprintCompilationManager.cpp b/Engine/Source/Editor/Kismet/Private/BlueprintCompilationManager.cpp
--- a/Engine/Source/Editor/Kismet/Private/BlueprintCompilationManager.cpp
+++ b/Engine/Source/Editor/Kismet/Private/BlueprintCompilationManager.cpp
@@ -1066,7 +1066,7 @@ void FBlueprintCompilationManagerImpl::FlushCompilationQueueImpl(bool bSuppressB
 			}
 			else
 			{
-				// default value propagation occurrs below:
+				// default value propagation occurs below:
 				if(BPGC)
 				{
 					if( BPGC->ClassDefaultObject && 
@@ -1173,7 +1173,7 @@ void FBlueprintCompilationManagerImpl::FlushCompilationQueueImpl(bool bSuppressB
 			}
 		}
 
-		// STAGE XV: CLEAR TEMPORARY FLAGS
+		// STAGE XVI: CLEAR TEMPORARY FLAGS
 		for (FCompilerData& CompilerData : CurrentlyCompilingBPs)
 		{
 			UBlueprint* BP = CompilerData.BP;
