[NO NON-PII COMMIT MESSAGE]

diff --git a/Engine/Source/Editor/Kismet/Private/SBlueprintDiff.cpp b/Engine/Source/Editor/Kismet/Private/SBlueprintDiff.cpp
--- a/Engine/Source/Editor/Kismet/Private/SBlueprintDiff.cpp
+++ b/Engine/Source/Editor/Kismet/Private/SBlueprintDiff.cpp
@@ -903,23 +903,6 @@ void SBlueprintDiff::Construct( const FArguments& InArgs)
 	PanelOld.bShowAssetName = InArgs._ShowAssetNames;
 	PanelNew.bShowAssetName = InArgs._ShowAssetNames;
 
-	// SMyBlueprint needs to be created *before* the KismetInspector or the diffs are generated, because the KismetInspector's customizations
-	// need a reference to the SMyBlueprint widget that is controlling them...
-	const auto CreateInspector = [](TSharedPtr<SMyBlueprint> InMyBlueprint) {
-		return SNew(SKismetInspector)
-			.HideNameArea(true)
-			.ViewIdentifier(FName("BlueprintInspector"))
-			.MyBlueprintWidget(InMyBlueprint)
-			.IsPropertyEditingEnabledDelegate(FIsPropertyEditingEnabled::CreateStatic([] { return false; }));
-	};
-
-	PanelOld.GenerateMyBlueprintWidget();
-	PanelOld.DetailsView = CreateInspector(PanelOld.MyBlueprint);
-	PanelOld.MyBlueprint->SetInspector(PanelOld.DetailsView);
-	PanelNew.GenerateMyBlueprintWidget();
-	PanelNew.DetailsView = CreateInspector(PanelNew.MyBlueprint);
-	PanelNew.MyBlueprint->SetInspector(PanelNew.DetailsView);
-
 	bLockViews = true;
 
 	if (InArgs._ParentWindow.IsValid())
@@ -959,6 +942,8 @@ void SBlueprintDiff::Construct( const FArguments& InArgs)
 		, TAttribute<FSlateIcon>(this, &SBlueprintDiff::GetLockViewImage)
 	);
 
+	DifferencesTreeView = DiffTreeView::CreateTreeView(&MasterDifferencesList);
+
 	GenerateDifferencesList();
 
 	const auto TextBlock = [](FText Text) -> TSharedRef<SWidget>
@@ -1113,13 +1098,11 @@ void SBlueprintDiff::OnGraphChanged(FGraphToDiff* Diff)
 
 void SBlueprintDiff::OnBlueprintChanged(UBlueprint* InBlueprint)
 {
-	if (InBlueprint == PanelOld.Blueprint && PanelOld.GraphEditor.IsValid())
+	if (InBlueprint == PanelOld.Blueprint || InBlueprint == PanelNew.Blueprint)
 	{
-		PanelOld.GraphEditor.Pin()->NotifyGraphChanged();
-	}
-	if (InBlueprint == PanelNew.Blueprint && PanelNew.GraphEditor.IsValid())
-	{
-		PanelNew.GraphEditor.Pin()->NotifyGraphChanged();
+		// After a BP has changed significantly, we need to regenerate the UI and set back to initial UI to avoid crashes
+		GenerateDifferencesList();
+		SetCurrentMode(MyBlueprintMode);
 	}
 }
 
@@ -1452,7 +1435,26 @@ void SBlueprintDiff::HandleGraphChanged( const FString& GraphPath )
 void SBlueprintDiff::GenerateDifferencesList()
 {
 	MasterDifferencesList.Empty();
+	RealDifferences.Empty();
 	Graphs.Empty();
+	ModePanels.Empty();
+
+	// SMyBlueprint needs to be created *before* the KismetInspector or the diffs are generated, because the KismetInspector's customizations
+	// need a reference to the SMyBlueprint widget that is controlling them...
+	const auto CreateInspector = [](TSharedPtr<SMyBlueprint> InMyBlueprint) {
+		return SNew(SKismetInspector)
+			.HideNameArea(true)
+			.ViewIdentifier(FName("BlueprintInspector"))
+			.MyBlueprintWidget(InMyBlueprint)
+			.IsPropertyEditingEnabledDelegate(FIsPropertyEditingEnabled::CreateStatic([] { return false; }));
+	};
+
+	PanelOld.GenerateMyBlueprintWidget();
+	PanelOld.DetailsView = CreateInspector(PanelOld.MyBlueprint);
+	PanelOld.MyBlueprint->SetInspector(PanelOld.DetailsView);
+	PanelNew.GenerateMyBlueprintWidget();
+	PanelNew.DetailsView = CreateInspector(PanelNew.MyBlueprint);
+	PanelNew.MyBlueprint->SetInspector(PanelNew.DetailsView);
 
 	TArray<UEdGraph*> GraphsOld, GraphsNew;
 	PanelOld.Blueprint->GetAllGraphs(GraphsOld);
@@ -1517,7 +1519,7 @@ void SBlueprintDiff::GenerateDifferencesList()
 		Graph->GenerateTreeEntries(MasterDifferencesList, RealDifferences);
 	}
 
-	DifferencesTreeView = DiffTreeView::CreateTreeView(&MasterDifferencesList);
+	DifferencesTreeView->RequestTreeRefresh();
 }
 
 SBlueprintDiff::FDiffControl SBlueprintDiff::GenerateBlueprintTypePanel()
@@ -1751,7 +1753,6 @@ void SBlueprintDiff::SetCurrentMode(FName NewMode)
 		PanelOld.DetailsView->ShowDetailsForObjects(TArray<UObject*>());
 		PanelNew.DetailsView->ShowDetailsForObjects(TArray<UObject*>());
 
-		DiffControl = FoundControl->DiffControl;
 		ModeContents->SetContent(FoundControl->Widget.ToSharedRef());
 	}
 	else
diff --git a/Engine/Source/Editor/Kismet/Public/SBlueprintDiff.h b/Engine/Source/Editor/Kismet/Public/SBlueprintDiff.h
--- a/Engine/Source/Editor/Kismet/Public/SBlueprintDiff.h
+++ b/Engine/Source/Editor/Kismet/Public/SBlueprintDiff.h
@@ -220,9 +220,6 @@ protected:
 
 	friend struct FListItemGraphToDiff;
 
-	/** Helper class for highlighting diffs in different types of controls (graph view, details view, etc) */
-	TSharedPtr< class IDiffControl > DiffControl;
-
 	/** We can't use the global tab manager because we need to instance the diff control, so we have our own tab manager: */
 	TSharedPtr<FTabManager> TabManager;
 
