[NO NON-PII COMMIT MESSAGE]

diff --git a/Engine/Source/Editor/Kismet/Private/SBlueprintDiff.cpp b/Engine/Source/Editor/Kismet/Private/SBlueprintDiff.cpp
--- a/Engine/Source/Editor/Kismet/Private/SBlueprintDiff.cpp
+++ b/Engine/Source/Editor/Kismet/Private/SBlueprintDiff.cpp
@@ -336,7 +336,7 @@ public:
 
 		// Check for parent class and interfaces here
 		const UBlueprint* OldBlueprint = Cast<UBlueprint>(OldDetails.GetDisplayedObject());
-		const UBlueprint* NewBlueprint = Cast<UBlueprint>(OldDetails.GetDisplayedObject());
+		const UBlueprint* NewBlueprint = Cast<UBlueprint>(NewDetails.GetDisplayedObject());
 
 		if (OldBlueprint && NewBlueprint)
 		{
